#ue4 - CharacterMovement networking tweaks for bandwidth and accuracy.

- Acceleration is rounded to the 1st decimal place in network games for more client/server consistency and less bandwidth usage. Exposed through RoundAcceleration() function.
- CanCombineMove(): angle between acceleration vectors is tunable (and lowered from 8 to 5 degrees).
- Custom classes can decide to avoid delaying a move if CanDelaySendingMove() returns false.

diff --git a/Engine/Source/Runtime/Engine/Classes/GameFramework/CharacterMovementComponent.h b/Engine/Source/Runtime/Engine/Classes/GameFramework/CharacterMovementComponent.h
--- a/Engine/Source/Runtime/Engine/Classes/GameFramework/CharacterMovementComponent.h
+++ b/Engine/Source/Runtime/Engine/Classes/GameFramework/CharacterMovementComponent.h
@@ -11,8 +11,6 @@
 #include "WorldCollision.h"
 #include "CharacterMovementComponent.generated.h"
 
-struct FVector_NetQuantize100;
-struct FVector_NetQuantizeNormal;
 class FDebugDisplayInfo;
 
 /** Data about the floor for walking movement, used by CharacterMovementComponent. */
@@ -101,6 +99,9 @@ struct FCharacterMovementComponentPreClothTickFunction : public FTickFunction
 	virtual FString DiagnosticMessage() override;
 };
 
+/** Shared pointer for easy memory management of FSavedMove_Character, for accumulating and replaying network moves. */
+typedef TSharedPtr<class FSavedMove_Character> FSavedMovePtr;
+
 
 //=============================================================================
 /**
@@ -1484,6 +1485,9 @@ public:
 	 */
 	virtual float GetNetworkSafeRandomAngleDegrees() const;
 
+	/** Round acceleration, for better consistency and lower bandwidth in networked games. */
+	virtual FVector RoundAcceleration(FVector InAccel) const;
+
 	//--------------------------------
 	// INetworkPredictionInterface implementation
 
@@ -1557,6 +1561,9 @@ protected:
 	/** Unpack compressed flags from a saved move and set state accordingly. See FSavedMove_Character. */
 	virtual void UpdateFromCompressedFlags(uint8 Flags);
 
+	/** Return true if it is OK to delay sending this player movement to the server, in order to conserve bandwidth. */
+	virtual bool CanDelaySendingMove(const FSavedMovePtr& NewMove);
+
 	/** Ticks the characters pose and accumulates root motion */
 	void TickCharacterPose(float DeltaTime);
 
@@ -1585,15 +1592,15 @@ public:
 
 	/** Replicated function sent by client to server - contains client movement and view info. */
 	UFUNCTION(unreliable, server, WithValidation)
-	virtual void ServerMove(float TimeStamp, FVector_NetQuantize100 InAccel, FVector_NetQuantize100 ClientLoc, uint8 CompressedMoveFlags, uint8 ClientRoll, uint32 View, UPrimitiveComponent* ClientMovementBase, FName ClientBaseBoneName, uint8 ClientMovementMode);
+	virtual void ServerMove(float TimeStamp, FVector_NetQuantize10 InAccel, FVector_NetQuantize100 ClientLoc, uint8 CompressedMoveFlags, uint8 ClientRoll, uint32 View, UPrimitiveComponent* ClientMovementBase, FName ClientBaseBoneName, uint8 ClientMovementMode);
 
 	/** Replicated function sent by client to server - contains client movement and view info for two moves. */
 	UFUNCTION(unreliable, server, WithValidation)
-	virtual void ServerMoveDual(float TimeStamp0, FVector_NetQuantize100 InAccel0, uint8 PendingFlags, uint32 View0, float TimeStamp, FVector_NetQuantize100 InAccel, FVector_NetQuantize100 ClientLoc, uint8 NewFlags, uint8 ClientRoll, uint32 View, UPrimitiveComponent* ClientMovementBase, FName ClientBaseBoneName, uint8 ClientMovementMode);
+	virtual void ServerMoveDual(float TimeStamp0, FVector_NetQuantize10 InAccel0, uint8 PendingFlags, uint32 View0, float TimeStamp, FVector_NetQuantize10 InAccel, FVector_NetQuantize100 ClientLoc, uint8 NewFlags, uint8 ClientRoll, uint32 View, UPrimitiveComponent* ClientMovementBase, FName ClientBaseBoneName, uint8 ClientMovementMode);
 	
 	/* Resending an (important) old move. Process it if not already processed. */
 	UFUNCTION(unreliable, server, WithValidation)
-	virtual void ServerMoveOld(float OldTimeStamp, FVector_NetQuantize100 OldAccel, uint8 OldMoveFlags);
+	virtual void ServerMoveOld(float OldTimeStamp, FVector_NetQuantize10 OldAccel, uint8 OldMoveFlags);
 	
 	/** If no client adjustment is needed after processing received ServerMove(), ack the good move so client can remove it from SavedMoves */
 	UFUNCTION(unreliable, client)
@@ -1682,8 +1689,6 @@ public:
 };
 
 
-typedef TSharedPtr<class FSavedMove_Character> FSavedMovePtr;
-
 /** FSavedMove_Character represents a saved move on the client that has been sent to the server and might need to be played back. */
 class ENGINE_API FSavedMove_Character
 {
@@ -1739,10 +1744,12 @@ public:
 	float RootMotionTrackPosition;
 	FRootMotionMovementParams RootMotionMovement;
 
-	/** threshold for deciding this is an "important" move based on DP with last acked acceleration. */
+	/** Threshold for deciding this is an "important" move based on DP with last acked acceleration. */
 	float AccelDotThreshold;    
 	/** Threshold for deciding is this is an important move because acceleration magnitude has changed too much */
 	float AccelMagThreshold;	
+	/** Threshold for deciding if we can combine two moves, true if cosine of angle between them is <= this. */
+	float AccelDotThresholdCombine;
 	
 	/** Clear saved move properties, so it can be re-used. */
 	virtual void Clear();
diff --git a/Engine/Source/Runtime/Engine/Private/CharacterMovementComponent.cpp b/Engine/Source/Runtime/Engine/Private/CharacterMovementComponent.cpp
--- a/Engine/Source/Runtime/Engine/Private/CharacterMovementComponent.cpp
+++ b/Engine/Source/Runtime/Engine/Private/CharacterMovementComponent.cpp
@@ -5048,6 +5048,16 @@ FVector UCharacterMovementComponent::ScaleInputAcceleration(const FVector& Input
 }
 
 
+FVector UCharacterMovementComponent::RoundAcceleration(FVector InAccel) const
+{
+	// Match FVector_NetQuantize10 (1 decimal place of precision).
+	InAccel.X = FMath::RoundToFloat(InAccel.X * 10.f) / 10.f;
+	InAccel.Y = FMath::RoundToFloat(InAccel.Y * 10.f) / 10.f;
+	InAccel.Z = FMath::RoundToFloat(InAccel.Z * 10.f) / 10.f;
+	return InAccel;
+}
+
+
 float UCharacterMovementComponent::ComputeAnalogInputModifier() const
 {
 	const float MaxAccel = GetMaxAcceleration();
@@ -5440,6 +5450,7 @@ void UCharacterMovementComponent::ReplicateMoveToServer(float DeltaTime, const F
 
 	// Perform the move locally
 	Acceleration = NewMove->Acceleration;
+	AnalogInputModifier = ComputeAnalogInputModifier(); // recompute since acceleration may have changed.
 	CharacterOwner->ClientRootMotionParams.Clear();
 	PerformMovement(NewMove->DeltaTime);
 
@@ -5448,9 +5459,9 @@ void UCharacterMovementComponent::ReplicateMoveToServer(float DeltaTime, const F
 	// Add NewMove to the list
 	ClientData->SavedMoves.Push(NewMove);
 
-	const bool bCanCombine = (CVarNetEnableMoveCombining.GetValueOnGameThread() != 0);
+	const bool bCanDelayMove = (CVarNetEnableMoveCombining.GetValueOnGameThread() != 0) && CanDelaySendingMove(NewMove);
 
-	if (bCanCombine && ClientData->PendingMove.IsValid() == false)
+	if (bCanDelayMove && ClientData->PendingMove.IsValid() == false)
 	{
 		// Decide whether to hold off on move
 		// send moves more frequently in small games where server isn't likely to be saturated
@@ -5472,6 +5483,7 @@ void UCharacterMovementComponent::ReplicateMoveToServer(float DeltaTime, const F
 
 		if ((GetWorld()->TimeSeconds - ClientData->ClientUpdateTime) * CharacterOwner->GetWorldSettings()->GetEffectiveTimeDilation() < NetMoveDelta)
 		{
+			// Delay sending this move.
 			ClientData->PendingMove = NewMove;
 			return;
 		}
@@ -5573,7 +5585,7 @@ void UCharacterMovementComponent::CallServerMove
 void UCharacterMovementComponent::ServerMoveOld_Implementation
 	(
 	float OldTimeStamp,
-	FVector_NetQuantize100 OldAccel,
+	FVector_NetQuantize10 OldAccel,
 	uint8 OldMoveFlags
 	)
 {
@@ -5601,11 +5613,11 @@ void UCharacterMovementComponent::ServerMoveOld_Implementation
 
 void UCharacterMovementComponent::ServerMoveDual_Implementation(
 	float TimeStamp0,
-	FVector_NetQuantize100 InAccel0,
+	FVector_NetQuantize10 InAccel0,
 	uint8 PendingFlags,
 	uint32 View0,
 	float TimeStamp,
-	FVector_NetQuantize100 InAccel,
+	FVector_NetQuantize10 InAccel,
 	FVector_NetQuantize100 ClientLoc,
 	uint8 NewFlags,
 	uint8 ClientRoll,
@@ -5653,7 +5665,7 @@ bool UCharacterMovementComponent::VerifyClientTimeStamp(float TimeStamp, FNetwor
 
 void UCharacterMovementComponent::ServerMove_Implementation(
 	float TimeStamp,
-	FVector_NetQuantize100 InAccel,
+	FVector_NetQuantize10 InAccel,
 	FVector_NetQuantize100 ClientLoc,
 	uint8 MoveFlags,
 	uint8 ClientRoll,
@@ -5827,17 +5839,17 @@ void UCharacterMovementComponent::ServerMoveHandleClientError(float TimeStamp, f
 }
 
 
-bool UCharacterMovementComponent::ServerMove_Validate(float TimeStamp, FVector_NetQuantize100 InAccel, FVector_NetQuantize100 ClientLoc, uint8 MoveFlags, uint8 ClientRoll, uint32 View, UPrimitiveComponent* ClientMovementBase, FName ClientBaseBoneName, uint8 ClientMovementMode)
+bool UCharacterMovementComponent::ServerMove_Validate(float TimeStamp, FVector_NetQuantize10 InAccel, FVector_NetQuantize100 ClientLoc, uint8 MoveFlags, uint8 ClientRoll, uint32 View, UPrimitiveComponent* ClientMovementBase, FName ClientBaseBoneName, uint8 ClientMovementMode)
 {
 	return true;
 }
 
-bool UCharacterMovementComponent::ServerMoveDual_Validate(float TimeStamp0, FVector_NetQuantize100 InAccel0, uint8 PendingFlags, uint32 View0, float TimeStamp, FVector_NetQuantize100 InAccel, FVector_NetQuantize100 ClientLoc, uint8 NewFlags, uint8 ClientRoll, uint32 View, UPrimitiveComponent* ClientMovementBase, FName ClientBaseBoneName, uint8 ClientMovementMode)
+bool UCharacterMovementComponent::ServerMoveDual_Validate(float TimeStamp0, FVector_NetQuantize10 InAccel0, uint8 PendingFlags, uint32 View0, float TimeStamp, FVector_NetQuantize10 InAccel, FVector_NetQuantize100 ClientLoc, uint8 NewFlags, uint8 ClientRoll, uint32 View, UPrimitiveComponent* ClientMovementBase, FName ClientBaseBoneName, uint8 ClientMovementMode)
 {
 	return true;
 }
 
-bool UCharacterMovementComponent::ServerMoveOld_Validate(float OldTimeStamp, FVector_NetQuantize100 OldAccel, uint8 OldMoveFlags)
+bool UCharacterMovementComponent::ServerMoveOld_Validate(float OldTimeStamp, FVector_NetQuantize10 OldAccel, uint8 OldMoveFlags)
 {
 	return true;
 }
@@ -6604,6 +6616,7 @@ FSavedMove_Character::FSavedMove_Character()
 {
 	AccelMagThreshold = 1.f;
 	AccelDotThreshold = 0.9f;
+	AccelDotThresholdCombine = 0.996f; // approx 5 degrees.
 }
 
 FSavedMove_Character::~FSavedMove_Character()
@@ -6658,7 +6671,10 @@ void FSavedMove_Character::SetMoveFor(ACharacter* Character, float InDeltaTime,
 
 	AccelMag = NewAccel.Size();
 	AccelNormal = (AccelMag > SMALL_NUMBER ? NewAccel / AccelMag : FVector::ZeroVector);
-	Acceleration = NewAccel;
+	
+	// Round value, so that client and server match exactly (and so we can send with less bandwidth). This rounded value is copied back to the client in ReplicateMoveToServer.
+	// This is done after the AccelMag and AccelNormal are computed above, because those are only used client-side for combining move logic and need to remain accurate.
+	Acceleration = Character->GetCharacterMovement()->RoundAcceleration(NewAccel);
 
 	bPressedJump = Character->bPressedJump;
 	JumpKeyHoldTime = Character->JumpKeyHoldTime;
@@ -6760,6 +6776,11 @@ FVector FSavedMove_Character::GetRevertedLocation() const
 	return StartLocation;
 }
 
+bool UCharacterMovementComponent::CanDelaySendingMove(const FSavedMovePtr& NewMove)
+{
+	return true;
+}
+
 bool FSavedMove_Character::CanCombineWith(const FSavedMovePtr& NewMove, ACharacter* Character, float MaxDelta) const
 {
 	if (bForceNoCombine || NewMove->bForceNoCombine)
@@ -6800,7 +6821,7 @@ bool FSavedMove_Character::CanCombineWith(const FSavedMovePtr& NewMove, ACharact
 			&& MovementMode == NewMove->MovementMode
 			&& StartCapsuleRadius == NewMove->StartCapsuleRadius
 			&& StartCapsuleHalfHeight == NewMove->StartCapsuleHalfHeight
-			&& ((AccelNormal | NewMove->AccelNormal) > 0.99f)
+			&& FVector::Coincident(AccelNormal, NewMove->AccelNormal, AccelDotThresholdCombine)
 			&& StartBaseRotation.Equals(NewMove->StartBaseRotation) // only if base hasn't rotated
 			&& (CustomTimeDilation == NewMove->CustomTimeDilation);
 	}
