NavMesh: Avoid race condition over tile cache data

diff --git a/Engine/Source/Runtime/Engine/Private/AI/Navigation/RecastNavMeshGenerator.cpp b/Engine/Source/Runtime/Engine/Private/AI/Navigation/RecastNavMeshGenerator.cpp
--- a/Engine/Source/Runtime/Engine/Private/AI/Navigation/RecastNavMeshGenerator.cpp
+++ b/Engine/Source/Runtime/Engine/Private/AI/Navigation/RecastNavMeshGenerator.cpp
@@ -1406,7 +1406,7 @@ static FBox CalculateTileBounds(int32 X, int32 Y, const FVector& NavMeshOrigin,
 //----------------------------------------------------------------------//
 
 FRecastTileGenerator::FRecastTileGenerator(
-	const FRecastNavMeshGenerator* ParentGenerator,
+	FRecastNavMeshGenerator* ParentGenerator,
 	const int32 X, 
 	const int32 Y,
 	TArray<FBox> DirtyAreas
@@ -1445,21 +1445,16 @@ FRecastTileGenerator::FRecastTileGenerator(
 			}
 		}
 	}
+
+	// Take ownership of tile cache data if it exist 
+	CompressedLayers = ParentGenerator->TakeIntermediateLayersData(FIntPoint(TileX, TileY));
+
+	// We have to regenerate layers data in case geometry is changed or tile cache is missing
+	bRegenerateCompressedLayers = (DirtyAreas.Num() == 0 || CompressedLayers.Num() == 0);
 	
 	// Gather geometry for tile if it inside navigable bounds
-	bRegenerateCompressedLayers = true;
 	if (InclusionBounds.Num())
 	{
-		// In case we don't need to regenerate tile from scratch
-		// get compressed layers data from cache
-		if (DirtyAreas.Num() > 0)
-		{
-			CompressedLayers = ParentGenerator->GetIntermediateLayersData(FIntPoint(TileX, TileY));
-		}
-	
-		// We have to regenerate layers data in case it's missing
-		bRegenerateCompressedLayers = (CompressedLayers.Num() == 0);
-
 		if (!bRegenerateCompressedLayers)
 		{
 			// Mark layers that needs to be updated
@@ -3213,17 +3208,11 @@ FBox FRecastNavMeshGenerator::GrowBoundingBox(const FBox& BBox, bool bIncludeAge
 	return FBox(BBox.Min - BBoxGrowOffsetBoth - BBoxGrowOffsetMin, BBox.Max + BBoxGrowOffsetBoth);
 }
 
-TArray<FNavMeshTileData> FRecastNavMeshGenerator::GetIntermediateLayersData(FIntPoint GridCoord) const
+TArray<FNavMeshTileData> FRecastNavMeshGenerator::TakeIntermediateLayersData(FIntPoint GridCoord)
 {
-	const TArray<FNavMeshTileData>* ExistingData = IntermediateLayerDataMap.Find(GridCoord);
-	if (ExistingData)
-	{
-		return *ExistingData;
-	}
-	else
-	{
-		return TArray<FNavMeshTileData>();
-	}
+	TArray<FNavMeshTileData> Result;
+	IntermediateLayerDataMap.RemoveAndCopyValue(GridCoord, Result);
+	return Result;
 }
 
 static bool IntercestBounds(const FBox& TestBox, const TNavStatArray<FBox>& Bounds)
@@ -3515,10 +3504,6 @@ TArray<uint32> FRecastNavMeshGenerator::ProcessTileTasks(const int32 NumTasksToS
 				{
 					IntermediateLayerDataMap.Add(Element.Coord, ComressedLayers);
 				}
-				else
-				{
-					IntermediateLayerDataMap.Remove(Element.Coord);
-				}
 			}
 
 			// Destroy tile generator task
diff --git a/Engine/Source/Runtime/Engine/Public/AI/Navigation/RecastNavMeshGenerator.h b/Engine/Source/Runtime/Engine/Public/AI/Navigation/RecastNavMeshGenerator.h
--- a/Engine/Source/Runtime/Engine/Public/AI/Navigation/RecastNavMeshGenerator.h
+++ b/Engine/Source/Runtime/Engine/Public/AI/Navigation/RecastNavMeshGenerator.h
@@ -112,7 +112,7 @@ class ENGINE_API FRecastTileGenerator : public FNonAbandonableTask
 {
 public:
 	FRecastTileGenerator(
-		const FRecastNavMeshGenerator* ParentGenerator,
+		FRecastNavMeshGenerator* ParentGenerator,
 		const int32 X, 
 		const int32 Y,
 		TArray<FBox> DirtyAreas
@@ -346,8 +346,8 @@ public:
 
 	FBox GrowBoundingBox(const FBox& BBox, bool bIncludeAgentHeight) const;
 
-	/**  */
-	TArray<FNavMeshTileData> GetIntermediateLayersData(FIntPoint GridCoord) const;
+	/** Transfers ownership if tile cache data to the caller */
+	TArray<FNavMeshTileData> TakeIntermediateLayersData(FIntPoint GridCoord);
 
 #if !(UE_BUILD_SHIPPING || UE_BUILD_TEST)
 	virtual void ExportNavigationData(const FString& FileName) const;
