Constructor chaining is a C++11 feature we don't support yet.

diff --git a/Engine/Source/Runtime/Engine/Private/NiagaraSimulation.cpp b/Engine/Source/Runtime/Engine/Private/NiagaraSimulation.cpp
--- a/Engine/Source/Runtime/Engine/Private/NiagaraSimulation.cpp
+++ b/Engine/Source/Runtime/Engine/Private/NiagaraSimulation.cpp
@@ -16,8 +16,13 @@ FNiagaraSimulation::FNiagaraSimulation(FNiagaraEmitterProperties *InProps)
 }
 
 FNiagaraSimulation::FNiagaraSimulation(FNiagaraEmitterProperties *InProps, ERHIFeatureLevel::Type InFeatureLevel)
-:FNiagaraSimulation(InProps)
+	: Age(0.0f)
+	, bIsEnabled(true)
+	, SpawnRemainder(0.0f)
+	, CachedBounds(ForceInit)
+	, EffectRenderer(nullptr)
 {
+	Props = InProps;
 	SetRenderModuleType(InProps->RenderModuleType, InFeatureLevel);
 }
 
