Automatic world origin rebasing only on XY plane

diff --git a/Engine/Source/Runtime/Engine/Private/CheatManager.cpp b/Engine/Source/Runtime/Engine/Private/CheatManager.cpp
--- a/Engine/Source/Runtime/Engine/Private/CheatManager.cpp
+++ b/Engine/Source/Runtime/Engine/Private/CheatManager.cpp
@@ -1063,6 +1063,9 @@ void UCheatManager::SetWorldOrigin()
 		ViewLocation = MyPlayerController->GetPawn()->GetActorLocation();
 	}
 	
+	// Consider only XY plane
+	ViewLocation.Z = 0;
+
 	FIntVector NewOrigin = FIntVector(ViewLocation.X, ViewLocation.Y, ViewLocation.Z) + World->OriginLocation;
 	World->RequestNewWorldOrigin(NewOrigin);
 }
diff --git a/Engine/Source/Runtime/Engine/Private/World.cpp b/Engine/Source/Runtime/Engine/Private/World.cpp
--- a/Engine/Source/Runtime/Engine/Private/World.cpp
+++ b/Engine/Source/Runtime/Engine/Private/World.cpp
@@ -2356,6 +2356,8 @@ void UWorld::EvaluateWorldOriginLocation(const FSceneViewFamily& ViewFamily)
 		}
 
 		CentroidLocation/= ViewFamily.Views.Num();
+		// Consider only XY plane
+		CentroidLocation.Z = 0.f;
 	
 		// Request to shift world in case current view is quite far from current origin
 		if (CentroidLocation.Size() > HALF_WORLD_MAX1*0.5f)
