Engine private PCH cleanup.

diff --git a/Engine/Source/Editor/Cascade/Private/Cascade.cpp b/Engine/Source/Editor/Cascade/Private/Cascade.cpp
--- a/Engine/Source/Editor/Cascade/Private/Cascade.cpp
+++ b/Engine/Source/Editor/Cascade/Private/Cascade.cpp
@@ -40,6 +40,7 @@
 #include "SNumericEntryBox.h"
 #include "STextEntryPopup.h"
 #include "GenericCommands.h"
+#include "Components/VectorFieldComponent.h"
 
 static const FName Cascade_PreviewViewportTab("Cascade_PreviewViewport");
 static const FName Cascade_EmmitterCanvasTab("Cascade_EmitterCanvas");
diff --git a/Engine/Source/Editor/Cascade/Private/Cascade.h b/Engine/Source/Editor/Cascade/Private/Cascade.h
--- a/Engine/Source/Editor/Cascade/Private/Cascade.h
+++ b/Engine/Source/Editor/Cascade/Private/Cascade.h
@@ -13,6 +13,7 @@ class SCascadePreviewViewport;
 class SCascadeEmitterCanvas;
 class UCascadeParticleSystemComponent;
 class UCascadeConfiguration;
+class UVectorFieldComponent;
 
 
 DECLARE_LOG_CATEGORY_EXTERN(LogCascade, Log, All);
diff --git a/Engine/Source/Editor/Cascade/Private/CascadePreviewViewportClient.cpp b/Engine/Source/Editor/Cascade/Private/CascadePreviewViewportClient.cpp
--- a/Engine/Source/Editor/Cascade/Private/CascadePreviewViewportClient.cpp
+++ b/Engine/Source/Editor/Cascade/Private/CascadePreviewViewportClient.cpp
@@ -13,6 +13,7 @@
 #include "Particles/ParticleModuleRequired.h"
 #include "Particles/ParticleSpriteEmitter.h"
 #include "PhysicsPublic.h"
+#include "Components/VectorFieldComponent.h"
 
 #define LOCTEXT_NAMESPACE "CascadeViewportClient"
 
diff --git a/Engine/Source/Editor/ComponentVisualizers/Private/ComponentVisualizers.cpp b/Engine/Source/Editor/ComponentVisualizers/Private/ComponentVisualizers.cpp
--- a/Engine/Source/Editor/ComponentVisualizers/Private/ComponentVisualizers.cpp
+++ b/Engine/Source/Editor/ComponentVisualizers/Private/ComponentVisualizers.cpp
@@ -17,6 +17,14 @@
 #include "SensingComponentVisualizer.h"
 #include "SpringComponentVisualizer.h"
 #include "PrimitiveComponentVisualizer.h"
+#include "Components/PointLightComponent.h"
+#include "Components/SpotLightComponent.h"
+#include "Components/AudioComponent.h"
+#include "PhysicsEngine/RadialForceComponent.h"
+#include "PhysicsEngine/PhysicsConstraintComponent.h"
+#include "GameFramework/SpringArmComponent.h"
+#include "Components/SplineComponent.h"
+#include "Components/PrimitiveComponent.h"
 
 IMPLEMENT_MODULE( FComponentVisualizersModule, ComponentVisualizers );
 
diff --git a/Engine/Source/Editor/CurveAssetEditor/Private/CurveAssetEditor.cpp b/Engine/Source/Editor/CurveAssetEditor/Private/CurveAssetEditor.cpp
--- a/Engine/Source/Editor/CurveAssetEditor/Private/CurveAssetEditor.cpp
+++ b/Engine/Source/Editor/CurveAssetEditor/Private/CurveAssetEditor.cpp
@@ -8,6 +8,8 @@
 //#include "Toolkits/IToolkitHost.h"
 #include "WorkspaceMenuStructureModule.h"
 #include "SDockTab.h"
+#include "Curves/CurveVector.h"
+#include "Curves/CurveLinearColor.h"
 
 #define LOCTEXT_NAMESPACE "CurveAssetEditor"
 
diff --git a/Engine/Source/Editor/DetailCustomizations/Private/BrushDetails.cpp b/Engine/Source/Editor/DetailCustomizations/Private/BrushDetails.cpp
--- a/Engine/Source/Editor/DetailCustomizations/Private/BrushDetails.cpp
+++ b/Engine/Source/Editor/DetailCustomizations/Private/BrushDetails.cpp
@@ -10,6 +10,7 @@
 #include "AssetSelection.h"
 #include "ClassIconFinder.h"
 #include "ScopedTransaction.h"
+#include "Engine/StaticMeshActor.h"
 
 #define LOCTEXT_NAMESPACE "BrushDetails"
 
diff --git a/Engine/Source/Editor/DetailCustomizations/Private/CurveColorCustomization.cpp b/Engine/Source/Editor/DetailCustomizations/Private/CurveColorCustomization.cpp
--- a/Engine/Source/Editor/DetailCustomizations/Private/CurveColorCustomization.cpp
+++ b/Engine/Source/Editor/DetailCustomizations/Private/CurveColorCustomization.cpp
@@ -7,6 +7,7 @@
 #include "MiniCurveEditor.h"
 #include "AssetRegistryModule.h"
 #include "SCurveEditor.h"
+#include "Curves/CurveLinearColor.h"
 
 #define LOCTEXT_NAMESPACE "CurveColorCustomization"
 
diff --git a/Engine/Source/Editor/DetailCustomizations/Private/DialogueStructsCustomizations.cpp b/Engine/Source/Editor/DetailCustomizations/Private/DialogueStructsCustomizations.cpp
--- a/Engine/Source/Editor/DetailCustomizations/Private/DialogueStructsCustomizations.cpp
+++ b/Engine/Source/Editor/DetailCustomizations/Private/DialogueStructsCustomizations.cpp
@@ -5,6 +5,7 @@
 #include "DialogueStructsCustomizations.h"
 #include "DialogueWaveWidgets.h"
 #include "AssetThumbnail.h"
+#include "Sound/DialogueWave.h"
 
 #define LOCTEXT_NAMESPACE "DialogueWaveDetails"
 
diff --git a/Engine/Source/Editor/DetailCustomizations/Private/DialogueWaveDetails.cpp b/Engine/Source/Editor/DetailCustomizations/Private/DialogueWaveDetails.cpp
--- a/Engine/Source/Editor/DetailCustomizations/Private/DialogueWaveDetails.cpp
+++ b/Engine/Source/Editor/DetailCustomizations/Private/DialogueWaveDetails.cpp
@@ -4,6 +4,7 @@
 #include "SoundDefinitions.h"
 #include "DialogueWaveDetails.h"
 #include "DialogueWaveWidgets.h"
+#include "Sound/DialogueWave.h"
 
 #define LOCTEXT_NAMESPACE "DialogueWaveDetails"
 
diff --git a/Engine/Source/Editor/DetailCustomizations/Private/DialogueWaveWidgets.cpp b/Engine/Source/Editor/DetailCustomizations/Private/DialogueWaveWidgets.cpp
--- a/Engine/Source/Editor/DetailCustomizations/Private/DialogueWaveWidgets.cpp
+++ b/Engine/Source/Editor/DetailCustomizations/Private/DialogueWaveWidgets.cpp
@@ -6,6 +6,7 @@
 #include "AssetThumbnail.h"
 #include "SAssetDropTarget.h"
 #include "AssetRegistryModule.h"
+#include "Sound/DialogueVoice.h"
 
 #define LOCTEXT_NAMESPACE "DialogueWaveDetails"
 
diff --git a/Engine/Source/Editor/DistCurveEditor/Private/CurveEditorSharedData.cpp b/Engine/Source/Editor/DistCurveEditor/Private/CurveEditorSharedData.cpp
--- a/Engine/Source/Editor/DistCurveEditor/Private/CurveEditorSharedData.cpp
+++ b/Engine/Source/Editor/DistCurveEditor/Private/CurveEditorSharedData.cpp
@@ -4,7 +4,7 @@
 #include "SDistributionCurveEditor.h"
 #include "CurveEditorSharedData.h"
 #include "CurveEditorHitProxies.h"
-
+#include "Engine/InterpCurveEdSetup.h"
 
 FCurveEditorSharedData::FCurveEditorSharedData(UInterpCurveEdSetup* InEdSetup)
 	: LabelEntryHeight(36)
diff --git a/Engine/Source/Editor/DistCurveEditor/Private/CurveEditorViewportClient.cpp b/Engine/Source/Editor/DistCurveEditor/Private/CurveEditorViewportClient.cpp
--- a/Engine/Source/Editor/DistCurveEditor/Private/CurveEditorViewportClient.cpp
+++ b/Engine/Source/Editor/DistCurveEditor/Private/CurveEditorViewportClient.cpp
@@ -8,6 +8,7 @@
 #include "SCurveEditorViewport.h"
 #include "CurveEditorViewportClient.h"
 #include "Runtime/Engine/Public/Slate/SceneViewport.h"
+#include "Engine/InterpCurveEdSetup.h"
 
 FCurveEditorViewportClient::FCurveEditorViewportClient(TWeakPtr<SDistributionCurveEditor> InCurveEditor, TWeakPtr<SCurveEditorViewport> InCurveEditorViewport)
 	: CurveEditorPtr(InCurveEditor)
diff --git a/Engine/Source/Editor/DistCurveEditor/Private/CurveEditorViewportClient.h b/Engine/Source/Editor/DistCurveEditor/Private/CurveEditorViewportClient.h
--- a/Engine/Source/Editor/DistCurveEditor/Private/CurveEditorViewportClient.h
+++ b/Engine/Source/Editor/DistCurveEditor/Private/CurveEditorViewportClient.h
@@ -5,7 +5,7 @@
 class SDistributionCurveEditor;
 class SCurveEditorViewport;
 class FCurveEditorSharedData;
-
+struct FCurveEdEntry;
 
 /*-----------------------------------------------------------------------------
    FCurveEditorViewportClient
diff --git a/Engine/Source/Editor/GraphEditor/Private/ConnectionDrawingPolicy.cpp b/Engine/Source/Editor/GraphEditor/Private/ConnectionDrawingPolicy.cpp
--- a/Engine/Source/Editor/GraphEditor/Private/ConnectionDrawingPolicy.cpp
+++ b/Engine/Source/Editor/GraphEditor/Private/ConnectionDrawingPolicy.cpp
@@ -15,6 +15,7 @@
 #include "Sound/SoundNode.h"
 #include "K2Node_Knot.h"
 #include "BlueprintEditorSettings.h"
+#include "Sound/SoundCue.h"
 
 DEFINE_LOG_CATEGORY_STATIC(LogConnectionDrawingPolicy, Log, All);
 
diff --git a/Engine/Source/Editor/Kismet/Private/BlueprintEditorTabFactories.cpp b/Engine/Source/Editor/Kismet/Private/BlueprintEditorTabFactories.cpp
--- a/Engine/Source/Editor/Kismet/Private/BlueprintEditorTabFactories.cpp
+++ b/Engine/Source/Editor/Kismet/Private/BlueprintEditorTabFactories.cpp
@@ -14,6 +14,8 @@
 #include "SBlueprintPalette.h"
 #include "FindInBlueprints.h"
 #include "SMyBlueprint.h"
+#include "Engine/Blueprint.h"
+#include "Engine/TimelineTemplate.h"
 
 
 
diff --git a/Engine/Source/Editor/Kismet/Private/DiffUtils.cpp b/Engine/Source/Editor/Kismet/Private/DiffUtils.cpp
--- a/Engine/Source/Editor/Kismet/Private/DiffUtils.cpp
+++ b/Engine/Source/Editor/Kismet/Private/DiffUtils.cpp
@@ -5,6 +5,8 @@
 #include "DiffUtils.h"
 #include "EditorCategoryUtils.h"
 #include "ObjectEditorUtils.h"
+#include "Engine/SCS_Node.h"
+#include "Engine/SimpleConstructionScript.h"
 
 static const UProperty* Resolve( const UStruct* Class, FName PropertyName )
 {
diff --git a/Engine/Source/Editor/Kismet/Private/SBlueprintPalette.cpp b/Engine/Source/Editor/Kismet/Private/SBlueprintPalette.cpp
--- a/Engine/Source/Editor/Kismet/Private/SBlueprintPalette.cpp
+++ b/Engine/Source/Editor/Kismet/Private/SBlueprintPalette.cpp
@@ -33,6 +33,7 @@
 #include "TutorialMetaData.h"
 #include "BlueprintEditorSettings.h" // for bShowActionMenuItemSignatures
 #include "SInlineEditableTextBlock.h"
+#include "Engine/SimpleConstructionScript.h"
 
 #define LOCTEXT_NAMESPACE "BlueprintPalette"
 
diff --git a/Engine/Source/Editor/Kismet/Private/SCSEditorViewportClient.cpp b/Engine/Source/Editor/Kismet/Private/SCSEditorViewportClient.cpp
--- a/Engine/Source/Editor/Kismet/Private/SCSEditorViewportClient.cpp
+++ b/Engine/Source/Editor/Kismet/Private/SCSEditorViewportClient.cpp
@@ -16,6 +16,7 @@
 #include "ISCSEditorCustomization.h"
 #include "ComponentVisualizer.h"
 #include "InstancedFoliage.h"
+#include "Engine/SimpleConstructionScript.h"
 
 DEFINE_LOG_CATEGORY_STATIC(LogSCSEditorViewport, Log, All);
 
diff --git a/Engine/Source/Editor/Kismet/Private/SMyBlueprint.cpp b/Engine/Source/Editor/Kismet/Private/SMyBlueprint.cpp
--- a/Engine/Source/Editor/Kismet/Private/SMyBlueprint.cpp
+++ b/Engine/Source/Editor/Kismet/Private/SMyBlueprint.cpp
@@ -38,6 +38,7 @@
 #include "EditorCategoryUtils.h"
 #include "SSearchBox.h"
 #include "GenericCommands.h"
+#include "Components/TimelineComponent.h"
 
 #define LOCTEXT_NAMESPACE "MyBlueprint"
 
diff --git a/Engine/Source/Editor/Kismet/Private/STimelineEditor.cpp b/Engine/Source/Editor/Kismet/Private/STimelineEditor.cpp
--- a/Engine/Source/Editor/Kismet/Private/STimelineEditor.cpp
+++ b/Engine/Source/Editor/Kismet/Private/STimelineEditor.cpp
@@ -17,6 +17,7 @@
 #include "GenericCommands.h"
 #include "SNotificationList.h"
 #include "NotificationManager.h"
+#include "Engine/TimelineTemplate.h"
 
 #define LOCTEXT_NAMESPACE "STimelineEditor"
 
diff --git a/Engine/Source/Editor/Kismet/Private/STimelineEditor.h b/Engine/Source/Editor/Kismet/Private/STimelineEditor.h
--- a/Engine/Source/Editor/Kismet/Private/STimelineEditor.h
+++ b/Engine/Source/Editor/Kismet/Private/STimelineEditor.h
@@ -4,6 +4,8 @@
 
 #include "BlueprintEditor.h"
 
+struct FTTTrackBase;
+
 //////////////////////////////////////////////////////////////////////////
 // FTimelineEdTrack
 
diff --git a/Engine/Source/Editor/UnrealEd/Private/Analytics/EditorAnalytics.cpp b/Engine/Source/Editor/UnrealEd/Private/Analytics/EditorAnalytics.cpp
--- a/Engine/Source/Editor/UnrealEd/Private/Analytics/EditorAnalytics.cpp
+++ b/Engine/Source/Editor/UnrealEd/Private/Analytics/EditorAnalytics.cpp
@@ -7,6 +7,7 @@
 #include "Runtime/Analytics/Analytics/Public/Interfaces/IAnalyticsProvider.h"
 
 #include "EditorAnalytics.h"
+#include "GeneralProjectSettings.h"
 
 void FEditorAnalytics::ReportBuildRequirementsFailure(FString EventName, FString PlatformName, bool bHasCode, int32 Requirements)
 {
diff --git a/Engine/Source/Editor/UnrealEd/Private/Animation/EditorAnimCompositeSegment.cpp b/Engine/Source/Editor/UnrealEd/Private/Animation/EditorAnimCompositeSegment.cpp
--- a/Engine/Source/Editor/UnrealEd/Private/Animation/EditorAnimCompositeSegment.cpp
+++ b/Engine/Source/Editor/UnrealEd/Private/Animation/EditorAnimCompositeSegment.cpp
@@ -5,6 +5,7 @@
 =============================================================================*/ 
 
 #include "UnrealEd.h"
+#include "Animation/AnimComposite.h"
 
 
 
diff --git a/Engine/Source/Editor/UnrealEd/Private/BSPOps.cpp b/Engine/Source/Editor/UnrealEd/Private/BSPOps.cpp
--- a/Engine/Source/Editor/UnrealEd/Private/BSPOps.cpp
+++ b/Engine/Source/Editor/UnrealEd/Private/BSPOps.cpp
@@ -3,6 +3,8 @@
 
 #include "UnrealEd.h"
 #include "BSPOps.h"
+#include "Engine/Polys.h"
+#include "Components/BrushComponent.h"
 
 DEFINE_LOG_CATEGORY_STATIC(LogBSPOps, Log, All);
 
diff --git a/Engine/Source/Editor/UnrealEd/Private/Commandlets/DumpBlueprintsInfoCommandlet.cpp b/Engine/Source/Editor/UnrealEd/Private/Commandlets/DumpBlueprintsInfoCommandlet.cpp
--- a/Engine/Source/Editor/UnrealEd/Private/Commandlets/DumpBlueprintsInfoCommandlet.cpp
+++ b/Engine/Source/Editor/UnrealEd/Private/Commandlets/DumpBlueprintsInfoCommandlet.cpp
@@ -29,6 +29,7 @@
 #include "PackageName.h"
 #include "PackageHelperFunctions.h" // for NORMALIZE_ExcludeMapPackages
 #include "K2Node_FunctionEntry.h"
+#include "Engine/LevelScriptActor.h"
 
 DEFINE_LOG_CATEGORY_STATIC(LogBlueprintInfoDump, Log, All);
 
diff --git a/Engine/Source/Editor/UnrealEd/Private/Commandlets/FixupRedirectsCommandlet.cpp b/Engine/Source/Editor/UnrealEd/Private/Commandlets/FixupRedirectsCommandlet.cpp
--- a/Engine/Source/Editor/UnrealEd/Private/Commandlets/FixupRedirectsCommandlet.cpp
+++ b/Engine/Source/Editor/UnrealEd/Private/Commandlets/FixupRedirectsCommandlet.cpp
@@ -6,6 +6,7 @@
 
 #include "UnrealEd.h"
 #include "ISourceControlModule.h"
+#include "Engine/UserDefinedEnum.h"
 
 DEFINE_LOG_CATEGORY_STATIC(LogFixupRedirectsCommandlet, Log, All);
 
diff --git a/Engine/Source/Editor/UnrealEd/Private/ComponentAssetBroker.cpp b/Engine/Source/Editor/UnrealEd/Private/ComponentAssetBroker.cpp
--- a/Engine/Source/Editor/UnrealEd/Private/ComponentAssetBroker.cpp
+++ b/Engine/Source/Editor/UnrealEd/Private/ComponentAssetBroker.cpp
@@ -9,6 +9,8 @@
 #include "Particles/ParticleSystemComponent.h"
 
 #include "Engine/DestructibleMesh.h"
+#include "Components/DestructibleComponent.h"
+#include "Sound/SoundBase.h"
 
 //////////////////////////////////////////////////////////////////////////
 // FStaticMeshComponentBroker
diff --git a/Engine/Source/Editor/UnrealEd/Private/CubemapUnwrapUtils.cpp b/Engine/Source/Editor/UnrealEd/Private/CubemapUnwrapUtils.cpp
--- a/Engine/Source/Editor/UnrealEd/Private/CubemapUnwrapUtils.cpp
+++ b/Engine/Source/Editor/UnrealEd/Private/CubemapUnwrapUtils.cpp
@@ -8,6 +8,9 @@
 #include "ShaderParameterUtils.h"
 #include "CubemapUnwrapUtils.h" 
 #include "RHIStaticStates.h"
+#include "Engine/TextureRenderTarget2D.h"
+#include "Engine/TextureCube.h"
+#include "Engine/TextureRenderTargetCube.h"
 
 IMPLEMENT_SHADER_TYPE(,FCubemapTexturePropertiesVS,TEXT("SimpleElementVertexShader"),TEXT("Main"),SF_Vertex);
 IMPLEMENT_SHADER_TYPE(template<>,FCubemapTexturePropertiesPS<false>,TEXT("SimpleElementPixelShader"),TEXT("CubemapTextureProperties"),SF_Pixel);
diff --git a/Engine/Source/Editor/UnrealEd/Private/DebugToolExec.cpp b/Engine/Source/Editor/UnrealEd/Private/DebugToolExec.cpp
--- a/Engine/Source/Editor/UnrealEd/Private/DebugToolExec.cpp
+++ b/Engine/Source/Editor/UnrealEd/Private/DebugToolExec.cpp
@@ -5,6 +5,7 @@
 #include "DebugToolExec.h"
 #include "Editor/PropertyEditor/Public/PropertyEditorModule.h"
 #include "Editor/PropertyEditor/Public/IDetailsView.h"
+#include "Engine/GameEngine.h"
 
 
 /**
diff --git a/Engine/Source/Editor/UnrealEd/Private/Dialogs/DlgReferenceTree.cpp b/Engine/Source/Editor/UnrealEd/Private/Dialogs/DlgReferenceTree.cpp
--- a/Engine/Source/Editor/UnrealEd/Private/Dialogs/DlgReferenceTree.cpp
+++ b/Engine/Source/Editor/UnrealEd/Private/Dialogs/DlgReferenceTree.cpp
@@ -6,6 +6,7 @@
 #include "Toolkits/AssetEditorManager.h"
 #include "ObjectTools.h"
 #include "AssetSelection.h"
+#include "Engine/Polys.h"
 
 FArchiveGenerateReferenceGraph::FArchiveGenerateReferenceGraph( FReferenceGraph& OutGraph ) 
 	: CurrentObject(NULL),
diff --git a/Engine/Source/Editor/UnrealEd/Private/Dialogs/DlgSoundWaveOptions.cpp b/Engine/Source/Editor/UnrealEd/Private/Dialogs/DlgSoundWaveOptions.cpp
--- a/Engine/Source/Editor/UnrealEd/Private/Dialogs/DlgSoundWaveOptions.cpp
+++ b/Engine/Source/Editor/UnrealEd/Private/Dialogs/DlgSoundWaveOptions.cpp
@@ -7,6 +7,7 @@
 #include "Dialogs/DlgSoundWaveOptions.h"
 #include "SoundPreviewThread.h"
 #include "SlateBasics.h"
+#include "Sound/SoundWave.h"
 
 /** 
  * Info used to setup the rows of the sound quality previewer
diff --git a/Engine/Source/Editor/UnrealEd/Private/Dialogs/SMeshProxyDialog.cpp b/Engine/Source/Editor/UnrealEd/Private/Dialogs/SMeshProxyDialog.cpp
--- a/Engine/Source/Editor/UnrealEd/Private/Dialogs/SMeshProxyDialog.cpp
+++ b/Engine/Source/Editor/UnrealEd/Private/Dialogs/SMeshProxyDialog.cpp
@@ -9,6 +9,7 @@
 #include "Editor/LevelEditor/Public/LevelEditor.h"
 #include "SNumericEntryBox.h"
 #include "STextComboBox.h"
+#include "Engine/StaticMeshActor.h"
 
 #define LOCTEXT_NAMESPACE "SMeshProxyDialog"
 
diff --git a/Engine/Source/Editor/UnrealEd/Private/DragTool_FrustumSelect.cpp b/Engine/Source/Editor/UnrealEd/Private/DragTool_FrustumSelect.cpp
--- a/Engine/Source/Editor/UnrealEd/Private/DragTool_FrustumSelect.cpp
+++ b/Engine/Source/Editor/UnrealEd/Private/DragTool_FrustumSelect.cpp
@@ -7,6 +7,8 @@
 #include "ActorEditorUtils.h"
 #include "ScopedTransaction.h"
 #include "HModel.h"
+#include "Components/DrawSphereComponent.h"
+#include "AI/Navigation/NavLinkRenderingComponent.h"
 
 ///////////////////////////////////////////////////////////////////////////////
 //
diff --git a/Engine/Source/Editor/UnrealEd/Private/Editor.cpp b/Engine/Source/Editor/UnrealEd/Private/Editor.cpp
--- a/Engine/Source/Editor/UnrealEd/Private/Editor.cpp
+++ b/Engine/Source/Editor/UnrealEd/Private/Editor.cpp
@@ -91,6 +91,8 @@
 #include "HotReloadInterface.h"
 #include "SNotificationList.h"
 #include "NotificationManager.h"
+#include "Engine/GameEngine.h"
+#include "Engine/TextureRenderTarget2D.h"
 
 DEFINE_LOG_CATEGORY_STATIC(LogEditor, Log, All);
 
diff --git a/Engine/Source/Editor/UnrealEd/Private/EditorActor.cpp b/Engine/Source/Editor/UnrealEd/Private/EditorActor.cpp
--- a/Engine/Source/Editor/UnrealEd/Private/EditorActor.cpp
+++ b/Engine/Source/Editor/UnrealEd/Private/EditorActor.cpp
@@ -18,6 +18,7 @@
 #include "Particles/Emitter.h"
 #include "Particles/ParticleSystemComponent.h"
 #include "UnrealExporter.h"
+#include "Engine/Polys.h"
 
 #define LOCTEXT_NAMESPACE "UnrealEd.EditorActor"
 
diff --git a/Engine/Source/Editor/UnrealEd/Private/EditorLevelUtils.cpp b/Engine/Source/Editor/UnrealEd/Private/EditorLevelUtils.cpp
--- a/Engine/Source/Editor/UnrealEd/Private/EditorLevelUtils.cpp
+++ b/Engine/Source/Editor/UnrealEd/Private/EditorLevelUtils.cpp
@@ -20,6 +20,7 @@
 #include "PackageTools.h"
 
 #include "Runtime/AssetRegistry/Public/AssetRegistryModule.h"
+#include "Engine/LevelStreamingVolume.h"
 
 DEFINE_LOG_CATEGORY(LogLevelTools);
 
diff --git a/Engine/Source/Editor/UnrealEd/Private/Factories/CSVImportFactory.cpp b/Engine/Source/Editor/UnrealEd/Private/Factories/CSVImportFactory.cpp
--- a/Engine/Source/Editor/UnrealEd/Private/Factories/CSVImportFactory.cpp
+++ b/Engine/Source/Editor/UnrealEd/Private/Factories/CSVImportFactory.cpp
@@ -6,6 +6,7 @@
 #include "ModuleManager.h"
 #include "DirectoryWatcherModule.h"
 #include "../../../DataTableEditor/Public/IDataTableEditor.h"
+#include "Curves/CurveVector.h"
 DEFINE_LOG_CATEGORY(LogCSVImportFactory);
 
 #define LOCTEXT_NAMESPACE "CSVImportFactory"
diff --git a/Engine/Source/Editor/UnrealEd/Private/Kismet2/BlueprintEditorUtils.cpp b/Engine/Source/Editor/UnrealEd/Private/Kismet2/BlueprintEditorUtils.cpp
--- a/Engine/Source/Editor/UnrealEd/Private/Kismet2/BlueprintEditorUtils.cpp
+++ b/Engine/Source/Editor/UnrealEd/Private/Kismet2/BlueprintEditorUtils.cpp
@@ -38,6 +38,10 @@
 #include "Runtime/Engine/Classes/Engine/UserDefinedStruct.h"
 #include "UnrealExporter.h"
 
+#include "Engine/SimpleConstructionScript.h"
+#include "Engine/SCS_Node.h"
+#include "Engine/TimelineTemplate.h"
+
 #define LOCTEXT_NAMESPACE "Blueprint"
 
 DEFINE_LOG_CATEGORY(LogBlueprintDebug);
diff --git a/Engine/Source/Editor/UnrealEd/Private/Kismet2/ComponentEditorUtils.cpp b/Engine/Source/Editor/UnrealEd/Private/Kismet2/ComponentEditorUtils.cpp
--- a/Engine/Source/Editor/UnrealEd/Private/Kismet2/ComponentEditorUtils.cpp
+++ b/Engine/Source/Editor/UnrealEd/Private/Kismet2/ComponentEditorUtils.cpp
@@ -4,6 +4,8 @@
 #include "Editor/UnrealEd/Public/Kismet2/ComponentEditorUtils.h"
 #include "Engine/Blueprint.h"
 #include "Runtime/Engine/Classes/Components/SceneComponent.h"
+#include "Engine/SCS_Node.h"
+#include "Engine/SimpleConstructionScript.h"
 
 USceneComponent* FComponentEditorUtils::GetSceneComponent( UObject* Object, UObject* SubObject /*= NULL*/ )
 {
diff --git a/Engine/Source/Editor/UnrealEd/Private/Kismet2/DebuggerCommands.cpp b/Engine/Source/Editor/UnrealEd/Private/Kismet2/DebuggerCommands.cpp
--- a/Engine/Source/Editor/UnrealEd/Private/Kismet2/DebuggerCommands.cpp
+++ b/Engine/Source/Editor/UnrealEd/Private/Kismet2/DebuggerCommands.cpp
@@ -26,6 +26,8 @@
 
 #include "EditorAnalytics.h"
 #include "MessageLog.h"
+#include "GameFramework/PlayerStart.h"
+#include "Components/CapsuleComponent.h"
 
 
 #define LOCTEXT_NAMESPACE "DebuggerCommands"
diff --git a/Engine/Source/Editor/UnrealEd/Private/Kismet2/StructureEditorUtils.cpp b/Engine/Source/Editor/UnrealEd/Private/Kismet2/StructureEditorUtils.cpp
--- a/Engine/Source/Editor/UnrealEd/Private/Kismet2/StructureEditorUtils.cpp
+++ b/Engine/Source/Editor/UnrealEd/Private/Kismet2/StructureEditorUtils.cpp
@@ -12,6 +12,7 @@
 #include "Editor/KismetCompiler/Public/KismetCompilerModule.h"
 #include "Toolkits/AssetEditorManager.h"
 #include "Editor/DataTableEditor/Public/IDataTableEditor.h"
+#include "Engine/UserDefinedStruct.h"
 
 #define LOCTEXT_NAMESPACE "Structure"
 
diff --git a/Engine/Source/Editor/UnrealEd/Private/Layers/Layers.cpp b/Engine/Source/Editor/UnrealEd/Private/Layers/Layers.cpp
--- a/Engine/Source/Editor/UnrealEd/Private/Layers/Layers.cpp
+++ b/Engine/Source/Editor/UnrealEd/Private/Layers/Layers.cpp
@@ -4,6 +4,7 @@
 
 #include "Layers.h"
 #include "ActorEditorUtils.h"
+#include "Layers/Layer.h"
 
 
 FLayers::FLayers( const TWeakObjectPtr< UEditorEngine >& InEditor )
diff --git a/Engine/Source/Editor/UnrealEd/Private/LevelViewportClickHandlers.cpp b/Engine/Source/Editor/UnrealEd/Private/LevelViewportClickHandlers.cpp
--- a/Engine/Source/Editor/UnrealEd/Private/LevelViewportClickHandlers.cpp
+++ b/Engine/Source/Editor/UnrealEd/Private/LevelViewportClickHandlers.cpp
@@ -12,6 +12,7 @@
 #include "MessageLog.h"
 #include "ActorEditorUtils.h"
 #include "Editor/ActorPositioning.h"
+#include "Engine/PointLight.h"
 
 #define LOCTEXT_NAMESPACE "ClickHandlers"
 
diff --git a/Engine/Source/Editor/UnrealEd/Private/Lightmass/LightmassRender.cpp b/Engine/Source/Editor/UnrealEd/Private/Lightmass/LightmassRender.cpp
--- a/Engine/Source/Editor/UnrealEd/Private/Lightmass/LightmassRender.cpp
+++ b/Engine/Source/Editor/UnrealEd/Private/Lightmass/LightmassRender.cpp
@@ -11,6 +11,7 @@
 #include "LightmassRender.h"
 #include "MaterialCompiler.h"
 #include "LightMap.h"
+#include "Engine/TextureRenderTarget2D.h"
 
 DEFINE_LOG_CATEGORY_STATIC(LogLightmassRender, Error, All);
 
diff --git a/Engine/Source/Editor/UnrealEd/Private/MeshPaintRendering.cpp b/Engine/Source/Editor/UnrealEd/Private/MeshPaintRendering.cpp
--- a/Engine/Source/Editor/UnrealEd/Private/MeshPaintRendering.cpp
+++ b/Engine/Source/Editor/UnrealEd/Private/MeshPaintRendering.cpp
@@ -10,6 +10,7 @@
 #include "ShaderParameters.h"
 #include "ShaderParameterUtils.h"
 #include "RHIStaticStates.h"
+#include "Engine/TextureRenderTarget2D.h"
 
 namespace MeshPaintRendering
 {
diff --git a/Engine/Source/Editor/UnrealEd/Private/PhysicsManipulationMode.cpp b/Engine/Source/Editor/UnrealEd/Private/PhysicsManipulationMode.cpp
--- a/Engine/Source/Editor/UnrealEd/Private/PhysicsManipulationMode.cpp
+++ b/Engine/Source/Editor/UnrealEd/Private/PhysicsManipulationMode.cpp
@@ -1,7 +1,7 @@
 // Copyright 1998-2014 Epic Games, Inc. All Rights Reserved.
 
 #include "UnrealEd.h"
-//#include "EnginePhysicsClasses.h"
+#include "PhysicsEngine/PhysicsHandleComponent.h"
 
 DEFINE_LOG_CATEGORY_STATIC(LogEditorPhysMode, Log, All);
 
diff --git a/Engine/Source/Editor/UnrealEd/Private/SnappingUtils.cpp b/Engine/Source/Editor/UnrealEd/Private/SnappingUtils.cpp
--- a/Engine/Source/Editor/UnrealEd/Private/SnappingUtils.cpp
+++ b/Engine/Source/Editor/UnrealEd/Private/SnappingUtils.cpp
@@ -7,6 +7,7 @@
 #include "Editor/LevelEditor/Public/LevelEditorActions.h"
 #include "VertexSnapping.h"
 #include "ViewportSnappingModule.h"
+#include "GameFramework/PhysicsVolume.h"
 
 //////////////////////////////////////////////////////////////////////////
 // FEditorViewportSnapping
diff --git a/Engine/Source/Editor/UnrealEd/Private/SoundCueGraphNode.cpp b/Engine/Source/Editor/UnrealEd/Private/SoundCueGraphNode.cpp
--- a/Engine/Source/Editor/UnrealEd/Private/SoundCueGraphNode.cpp
+++ b/Engine/Source/Editor/UnrealEd/Private/SoundCueGraphNode.cpp
@@ -11,6 +11,7 @@
 #include "SoundCueGraphEditorCommands.h"
 #include "ScopedTransaction.h"
 #include "GenericCommands.h"
+#include "Sound/SoundCue.h"
 
 #define LOCTEXT_NAMESPACE "SoundCueGraphNode"
 
diff --git a/Engine/Source/Editor/UnrealEd/Private/SoundPreviewThread.cpp b/Engine/Source/Editor/UnrealEd/Private/SoundPreviewThread.cpp
--- a/Engine/Source/Editor/UnrealEd/Private/SoundPreviewThread.cpp
+++ b/Engine/Source/Editor/UnrealEd/Private/SoundPreviewThread.cpp
@@ -4,6 +4,7 @@
 #include "UnrealEd.h"
 #include "SoundPreviewThread.h"
 #include "SoundDefinitions.h"
+#include "Sound/SoundWave.h"
 
 /**
  *
diff --git a/Engine/Source/Editor/UnrealEd/Private/StaticLightingSystem/StaticLightingDebug.cpp b/Engine/Source/Editor/UnrealEd/Private/StaticLightingSystem/StaticLightingDebug.cpp
--- a/Engine/Source/Editor/UnrealEd/Private/StaticLightingSystem/StaticLightingDebug.cpp
+++ b/Engine/Source/Editor/UnrealEd/Private/StaticLightingSystem/StaticLightingDebug.cpp
@@ -8,6 +8,7 @@
 #include "StaticMeshResources.h"
 #include "StaticLightingPrivate.h"
 #include "LightMap.h"
+#include "Components/ModelComponent.h"
 
 /** Information about the texel that is selected */
 FSelectedLightmapSample GCurrentSelectedLightmapSample;
diff --git a/Engine/Source/Editor/UnrealEd/Private/StaticLightingSystem/StaticLightingExport.cpp b/Engine/Source/Editor/UnrealEd/Private/StaticLightingSystem/StaticLightingExport.cpp
--- a/Engine/Source/Editor/UnrealEd/Private/StaticLightingSystem/StaticLightingExport.cpp
+++ b/Engine/Source/Editor/UnrealEd/Private/StaticLightingSystem/StaticLightingExport.cpp
@@ -14,6 +14,7 @@
 #include "LandscapeComponent.h"
 
 #include "Lightmass/Lightmass.h"
+#include "Lightmass/LightmappedSurfaceCollection.h"
 
 // Doxygen cannot parse these definitions correctly since the declaration is in a header from a different module
 #if !UE_BUILD_DOCS
diff --git a/Engine/Source/Editor/UnrealEd/Private/StaticLightingSystem/StaticLightingSystem.cpp b/Engine/Source/Editor/UnrealEd/Private/StaticLightingSystem/StaticLightingSystem.cpp
--- a/Engine/Source/Editor/UnrealEd/Private/StaticLightingSystem/StaticLightingSystem.cpp
+++ b/Engine/Source/Editor/UnrealEd/Private/StaticLightingSystem/StaticLightingSystem.cpp
@@ -34,6 +34,7 @@ DEFINE_LOG_CATEGORY(LogStaticLightingSystem);
 #include "MessageLog.h"
 #include "SNotificationList.h"
 #include "NotificationManager.h"
+#include "Engine/GeneratedMeshAreaLight.h"
 
 #define LOCTEXT_NAMESPACE "StaticLightingSystem"
 
diff --git a/Engine/Source/Editor/UnrealEd/Private/ThumbnailHelpers.cpp b/Engine/Source/Editor/UnrealEd/Private/ThumbnailHelpers.cpp
--- a/Engine/Source/Editor/UnrealEd/Private/ThumbnailHelpers.cpp
+++ b/Engine/Source/Editor/UnrealEd/Private/ThumbnailHelpers.cpp
@@ -7,6 +7,7 @@
 #include "Particles/ParticleSystem.h"
 #include "Particles/ParticleSystemComponent.h"
 #include "ContentStreaming.h"
+#include "Components/DirectionalLightComponent.h"
 
 /*
 ***************************************************************
diff --git a/Engine/Source/Editor/UnrealEd/Private/ThumbnailRendering/LevelThumbnailRenderer.cpp b/Engine/Source/Editor/UnrealEd/Private/ThumbnailRendering/LevelThumbnailRenderer.cpp
--- a/Engine/Source/Editor/UnrealEd/Private/ThumbnailRendering/LevelThumbnailRenderer.cpp
+++ b/Engine/Source/Editor/UnrealEd/Private/ThumbnailRendering/LevelThumbnailRenderer.cpp
@@ -3,6 +3,7 @@
 #include "UnrealEd.h"
 #include "EngineModule.h"
 #include "RendererInterface.h"
+#include "Engine/LevelBounds.h"
 
 ULevelThumbnailRenderer::ULevelThumbnailRenderer(const FObjectInitializer& ObjectInitializer)
 	: Super(ObjectInitializer)
diff --git a/Engine/Source/Editor/UnrealEd/Private/ThumbnailRendering/SubsurfaceProfileRenderer.cpp b/Engine/Source/Editor/UnrealEd/Private/ThumbnailRendering/SubsurfaceProfileRenderer.cpp
--- a/Engine/Source/Editor/UnrealEd/Private/ThumbnailRendering/SubsurfaceProfileRenderer.cpp
+++ b/Engine/Source/Editor/UnrealEd/Private/ThumbnailRendering/SubsurfaceProfileRenderer.cpp
@@ -1,6 +1,7 @@
 // Copyright 1998-2014 Epic Games, Inc. All Rights Reserved.
 
 #include "UnrealEd.h"
+#include "Engine/SubsurfaceProfile.h"
 
 USubsurfaceProfileRenderer::USubsurfaceProfileRenderer(const FObjectInitializer& ObjectInitializer)
 	: Super(ObjectInitializer)
diff --git a/Engine/Source/Editor/UnrealEd/Private/ThumbnailRendering/TextureCubeThumbnailRenderer.cpp b/Engine/Source/Editor/UnrealEd/Private/ThumbnailRendering/TextureCubeThumbnailRenderer.cpp
--- a/Engine/Source/Editor/UnrealEd/Private/ThumbnailRendering/TextureCubeThumbnailRenderer.cpp
+++ b/Engine/Source/Editor/UnrealEd/Private/ThumbnailRendering/TextureCubeThumbnailRenderer.cpp
@@ -1,6 +1,7 @@
 // Copyright 1998-2014 Epic Games, Inc. All Rights Reserved.
 
 #include "UnrealEd.h"
+#include "Engine/TextureCube.h"
 
 UTextureCubeThumbnailRenderer::UTextureCubeThumbnailRenderer(const FObjectInitializer& ObjectInitializer)
 	: Super(ObjectInitializer)
diff --git a/Engine/Source/Editor/UnrealEd/Private/ThumbnailRendering/WorldThumbnailRenderer.cpp b/Engine/Source/Editor/UnrealEd/Private/ThumbnailRendering/WorldThumbnailRenderer.cpp
--- a/Engine/Source/Editor/UnrealEd/Private/ThumbnailRendering/WorldThumbnailRenderer.cpp
+++ b/Engine/Source/Editor/UnrealEd/Private/ThumbnailRendering/WorldThumbnailRenderer.cpp
@@ -4,6 +4,8 @@
 #include "EngineModule.h"
 #include "ContentStreaming.h"
 #include "RendererInterface.h"
+#include "Engine/LevelBounds.h"
+
 UWorldThumbnailRenderer::UWorldThumbnailRenderer(const FObjectInitializer& ObjectInitializer)
 	: Super(ObjectInitializer)
 {
diff --git a/Engine/Source/Editor/UnrealEd/Private/VertexSnapping.cpp b/Engine/Source/Editor/UnrealEd/Private/VertexSnapping.cpp
--- a/Engine/Source/Editor/UnrealEd/Private/VertexSnapping.cpp
+++ b/Engine/Source/Editor/UnrealEd/Private/VertexSnapping.cpp
@@ -3,6 +3,7 @@
 #include "UnrealEd.h"
 #include "VertexSnapping.h"
 #include "StaticMeshResources.h"
+#include "Components/BrushComponent.h"
 
 namespace VertexSnappingConstants
 {
diff --git a/Engine/Source/Editor/UnrealEd/Public/Dialogs/DlgReferenceTree.h b/Engine/Source/Editor/UnrealEd/Public/Dialogs/DlgReferenceTree.h
--- a/Engine/Source/Editor/UnrealEd/Public/Dialogs/DlgReferenceTree.h
+++ b/Engine/Source/Editor/UnrealEd/Public/Dialogs/DlgReferenceTree.h
@@ -4,6 +4,8 @@
 #ifndef __DLGREFERENCETREE_H__
 #define __DLGREFERENCETREE_H__
 
+#include "Engine/Polys.h"
+
 /**
  * A struct representing a node in the reference graph
  */
diff --git a/Engine/Source/Editor/UnrealEd/Public/Kismet2/BlueprintEditorUtils.h b/Engine/Source/Editor/UnrealEd/Public/Kismet2/BlueprintEditorUtils.h
--- a/Engine/Source/Editor/UnrealEd/Public/Kismet2/BlueprintEditorUtils.h
+++ b/Engine/Source/Editor/UnrealEd/Public/Kismet2/BlueprintEditorUtils.h
@@ -6,6 +6,8 @@
 #include "EdGraph/EdGraph.h"
 #include "EdGraphSchema_K2.h"
 
+class USCS_Node;
+
 /** 
   * Flags describing how to handle graph removal
   */
diff --git a/Engine/Source/Editor/UnrealEd/Public/ThumbnailHelpers.h b/Engine/Source/Editor/UnrealEd/Public/ThumbnailHelpers.h
--- a/Engine/Source/Editor/UnrealEd/Public/ThumbnailHelpers.h
+++ b/Engine/Source/Editor/UnrealEd/Public/ThumbnailHelpers.h
@@ -6,6 +6,7 @@
 #include "PreviewScene.h"
 
 class FSceneViewFamily;
+class USCS_Node;
 
 class UNREALED_API FThumbnailPreviewScene : public FPreviewScene
 {
diff --git a/Engine/Source/Runtime/Engine/Classes/AI/Navigation/NavLinkDefinition.h b/Engine/Source/Runtime/Engine/Classes/AI/Navigation/NavLinkDefinition.h
--- a/Engine/Source/Runtime/Engine/Classes/AI/Navigation/NavLinkDefinition.h
+++ b/Engine/Source/Runtime/Engine/Classes/AI/Navigation/NavLinkDefinition.h
@@ -2,6 +2,7 @@
 
 #pragma once
 
+#include "NavAreas/NavArea.h"
 #include "NavLinkDefinition.generated.h"
 
 UENUM()
diff --git a/Engine/Source/Runtime/Engine/Classes/AI/Navigation/NavigationTestingActor.h b/Engine/Source/Runtime/Engine/Classes/AI/Navigation/NavigationTestingActor.h
--- a/Engine/Source/Runtime/Engine/Classes/AI/Navigation/NavigationTestingActor.h
+++ b/Engine/Source/Runtime/Engine/Classes/AI/Navigation/NavigationTestingActor.h
@@ -2,6 +2,7 @@
 
 #pragma once
 #include "Tickable.h"
+#include "NavPathObserverInterface.h"
 #include "NavigationTestingActor.generated.h"
 
 struct FNavTestTickHelper : FTickableGameObject
diff --git a/Engine/Source/Runtime/Engine/Classes/Camera/CameraModifier_CameraShake.h b/Engine/Source/Runtime/Engine/Classes/Camera/CameraModifier_CameraShake.h
--- a/Engine/Source/Runtime/Engine/Classes/Camera/CameraModifier_CameraShake.h
+++ b/Engine/Source/Runtime/Engine/Classes/Camera/CameraModifier_CameraShake.h
@@ -8,6 +8,8 @@
 #include "Camera/CameraModifier.h"
 #include "CameraModifier_CameraShake.generated.h"
 
+struct FFOscillator;
+
 USTRUCT()
 struct ENGINE_API FCameraShakeInstance
 {
diff --git a/Engine/Source/Runtime/Engine/Classes/Camera/PlayerCameraManager.h b/Engine/Source/Runtime/Engine/Classes/Camera/PlayerCameraManager.h
--- a/Engine/Source/Runtime/Engine/Classes/Camera/PlayerCameraManager.h
+++ b/Engine/Source/Runtime/Engine/Classes/Camera/PlayerCameraManager.h
@@ -3,6 +3,7 @@
 #pragma once
 #include "GameFramework/Actor.h"
 #include "CameraTypes.h"
+#include "CameraShake.h"
 #include "PlayerCameraManager.generated.h"
 
 class APawn;
@@ -420,7 +421,7 @@ public:
 	 * @param Falloff - Affects falloff of effect as it nears OuterRadius
 	 * @param bOrientShakeTowardsEpicenter - Changes the rotation of shake to point towards epicenter instead of forward
 	 */
-	static void PlayWorldCameraShake(UWorld* InWorld, TSubclassOf<class UCameraShake> Shake, FVector Epicenter, float InnerRadius, float OuterRadius, float Falloff, bool bOrientShakeTowardsEpicenter = false );
+	static void PlayWorldCameraShake(UWorld* InWorld, TSubclassOf<UCameraShake> Shake, FVector Epicenter, float InnerRadius, float OuterRadius, float Falloff, bool bOrientShakeTowardsEpicenter = false );
 
 protected:
 	/** Internal.  Returns intensity scalar in the range [0..1] for a shake originating at Epicenter. */
diff --git a/Engine/Source/Runtime/Engine/Classes/Components/BoxReflectionCaptureComponent.h b/Engine/Source/Runtime/Engine/Classes/Components/BoxReflectionCaptureComponent.h
--- a/Engine/Source/Runtime/Engine/Classes/Components/BoxReflectionCaptureComponent.h
+++ b/Engine/Source/Runtime/Engine/Classes/Components/BoxReflectionCaptureComponent.h
@@ -2,6 +2,7 @@
 
 
 #pragma once
+#include "ReflectionCaptureComponent.h"
 #include "BoxReflectionCaptureComponent.generated.h"
 
 	// -> will be exported to EngineDecalClasses.h
diff --git a/Engine/Source/Runtime/Engine/Classes/Components/ChildActorComponent.h b/Engine/Source/Runtime/Engine/Classes/Components/ChildActorComponent.h
--- a/Engine/Source/Runtime/Engine/Classes/Components/ChildActorComponent.h
+++ b/Engine/Source/Runtime/Engine/Classes/Components/ChildActorComponent.h
@@ -2,6 +2,7 @@
 
 
 #pragma once
+#include "SceneComponent.h"
 #include "ChildActorComponent.generated.h"
 
 class FChildActorComponentInstanceData;
diff --git a/Engine/Source/Runtime/Engine/Classes/Components/DirectionalLightComponent.h b/Engine/Source/Runtime/Engine/Classes/Components/DirectionalLightComponent.h
--- a/Engine/Source/Runtime/Engine/Classes/Components/DirectionalLightComponent.h
+++ b/Engine/Source/Runtime/Engine/Classes/Components/DirectionalLightComponent.h
@@ -4,6 +4,7 @@
 
 #include "Engine/EngineTypes.h"
 
+#include "LightComponent.h"
 #include "DirectionalLightComponent.generated.h"
 
 /**
diff --git a/Engine/Source/Runtime/Engine/Classes/Components/DrawSphereComponent.h b/Engine/Source/Runtime/Engine/Classes/Components/DrawSphereComponent.h
--- a/Engine/Source/Runtime/Engine/Classes/Components/DrawSphereComponent.h
+++ b/Engine/Source/Runtime/Engine/Classes/Components/DrawSphereComponent.h
@@ -2,6 +2,7 @@
 
 
 #pragma once
+#include "SphereComponent.h"
 #include "DrawSphereComponent.generated.h"
 
 /** 
diff --git a/Engine/Source/Runtime/Engine/Classes/Components/PlaneReflectionCaptureComponent.h b/Engine/Source/Runtime/Engine/Classes/Components/PlaneReflectionCaptureComponent.h
--- a/Engine/Source/Runtime/Engine/Classes/Components/PlaneReflectionCaptureComponent.h
+++ b/Engine/Source/Runtime/Engine/Classes/Components/PlaneReflectionCaptureComponent.h
@@ -2,6 +2,7 @@
 
 
 #pragma once
+#include "ReflectionCaptureComponent.h"
 #include "PlaneReflectionCaptureComponent.generated.h"
 
 	// -> will be exported to EngineDecalClasses.h
diff --git a/Engine/Source/Runtime/Engine/Classes/Components/PointLightComponent.h b/Engine/Source/Runtime/Engine/Classes/Components/PointLightComponent.h
--- a/Engine/Source/Runtime/Engine/Classes/Components/PointLightComponent.h
+++ b/Engine/Source/Runtime/Engine/Classes/Components/PointLightComponent.h
@@ -2,6 +2,7 @@
 
 
 #pragma once
+#include "LightComponent.h"
 #include "PointLightComponent.generated.h"
 
 /**
diff --git a/Engine/Source/Runtime/Engine/Classes/Components/SpotLightComponent.h b/Engine/Source/Runtime/Engine/Classes/Components/SpotLightComponent.h
--- a/Engine/Source/Runtime/Engine/Classes/Components/SpotLightComponent.h
+++ b/Engine/Source/Runtime/Engine/Classes/Components/SpotLightComponent.h
@@ -1,7 +1,7 @@
 // Copyright 1998-2014 Epic Games, Inc. All Rights Reserved.
 
 #pragma once
-
+#include "PointLightComponent.h"
 #include "SpotLightComponent.generated.h"
 
 /**
diff --git a/Engine/Source/Runtime/Engine/Classes/Components/VectorFieldComponent.h b/Engine/Source/Runtime/Engine/Classes/Components/VectorFieldComponent.h
--- a/Engine/Source/Runtime/Engine/Classes/Components/VectorFieldComponent.h
+++ b/Engine/Source/Runtime/Engine/Classes/Components/VectorFieldComponent.h
@@ -3,6 +3,7 @@
 
 
 #pragma once
+#include "PrimitiveComponent.h"
 #include "VectorFieldComponent.generated.h"
 
 /** 
diff --git a/Engine/Source/Runtime/Engine/Classes/Debug/LogVisualizerCameraController.h b/Engine/Source/Runtime/Engine/Classes/Debug/LogVisualizerCameraController.h
--- a/Engine/Source/Runtime/Engine/Classes/Debug/LogVisualizerCameraController.h
+++ b/Engine/Source/Runtime/Engine/Classes/Debug/LogVisualizerCameraController.h
@@ -2,6 +2,7 @@
 
 #pragma once
 
+#include "Engine/DebugCameraController.h"
 #include "LogVisualizerCameraController.generated.h"
 
 UCLASS(config=Input, hidedropdown)
diff --git a/Engine/Source/Runtime/Engine/Classes/Engine/BoxReflectionCapture.h b/Engine/Source/Runtime/Engine/Classes/Engine/BoxReflectionCapture.h
--- a/Engine/Source/Runtime/Engine/Classes/Engine/BoxReflectionCapture.h
+++ b/Engine/Source/Runtime/Engine/Classes/Engine/BoxReflectionCapture.h
@@ -1,6 +1,7 @@
 // Copyright 1998-2014 Epic Games, Inc. All Rights Reserved.
 
 #pragma once
+#include "ReflectionCapture.h"
 #include "BoxReflectionCapture.generated.h"
 
 /** 
diff --git a/Engine/Source/Runtime/Engine/Classes/Engine/CanvasRenderTarget2D.h b/Engine/Source/Runtime/Engine/Classes/Engine/CanvasRenderTarget2D.h
--- a/Engine/Source/Runtime/Engine/Classes/Engine/CanvasRenderTarget2D.h
+++ b/Engine/Source/Runtime/Engine/Classes/Engine/CanvasRenderTarget2D.h
@@ -2,6 +2,7 @@
 
 #pragma once
 
+#include "TextureRenderTarget2D.h"
 #include "CanvasRenderTarget2D.generated.h"
 
 class UCanvas;
diff --git a/Engine/Source/Runtime/Engine/Classes/Engine/Console.h b/Engine/Source/Runtime/Engine/Classes/Engine/Console.h
--- a/Engine/Source/Runtime/Engine/Classes/Engine/Console.h
+++ b/Engine/Source/Runtime/Engine/Classes/Engine/Console.h
@@ -4,7 +4,7 @@
 // Console - A quick little command line console that accepts most commands.
 //=============================================================================
 #pragma once
-
+#include "ConsoleSettings.h"
 #include "Console.generated.h"
 
 
diff --git a/Engine/Source/Runtime/Engine/Classes/Engine/EngineTypes.h b/Engine/Source/Runtime/Engine/Classes/Engine/EngineTypes.h
--- a/Engine/Source/Runtime/Engine/Classes/Engine/EngineTypes.h
+++ b/Engine/Source/Runtime/Engine/Classes/Engine/EngineTypes.h
@@ -7,6 +7,7 @@
  *	use in multiple files where the enum can't be mapped to a specific file.
  */
 #include "NetSerialization.h"
+#include "GameFramework/DamageType.h"
 #include "EngineTypes.generated.h"
 
 UENUM()
diff --git a/Engine/Source/Runtime/Engine/Classes/Engine/Light.h b/Engine/Source/Runtime/Engine/Classes/Engine/Light.h
--- a/Engine/Source/Runtime/Engine/Classes/Engine/Light.h
+++ b/Engine/Source/Runtime/Engine/Classes/Engine/Light.h
@@ -2,6 +2,7 @@
 
 #pragma once
 
+#include "Components/LightComponent.h"
 #include "Light.generated.h"
 
 UCLASS(Abstract, ClassGroup=Lights, hidecategories=(Input,Collision,Replication), showcategories=("Input|MouseInput", "Input|TouchInput"), ConversionRoot, meta=(ChildCanTick))
diff --git a/Engine/Source/Runtime/Engine/Classes/Engine/PlayerStartPIE.h b/Engine/Source/Runtime/Engine/Classes/Engine/PlayerStartPIE.h
--- a/Engine/Source/Runtime/Engine/Classes/Engine/PlayerStartPIE.h
+++ b/Engine/Source/Runtime/Engine/Classes/Engine/PlayerStartPIE.h
@@ -5,6 +5,7 @@
 //=============================================================================
 
 #pragma once
+#include "GameFramework/PlayerStart.h"
 #include "PlayerStartPIE.generated.h"
 
 UCLASS(notplaceable,MinimalAPI)
diff --git a/Engine/Source/Runtime/Engine/Classes/Engine/ShadowMapTexture2D.h b/Engine/Source/Runtime/Engine/Classes/Engine/ShadowMapTexture2D.h
--- a/Engine/Source/Runtime/Engine/Classes/Engine/ShadowMapTexture2D.h
+++ b/Engine/Source/Runtime/Engine/Classes/Engine/ShadowMapTexture2D.h
@@ -1,6 +1,7 @@
 // Copyright 1998-2014 Epic Games, Inc. All Rights Reserved.
 
 #pragma once
+#include "Texture2D.h"
 #include "ShadowMapTexture2D.generated.h"
 
 UCLASS(MinimalAPI)
diff --git a/Engine/Source/Runtime/Engine/Classes/Engine/TextureRenderTarget2D.h b/Engine/Source/Runtime/Engine/Classes/Engine/TextureRenderTarget2D.h
--- a/Engine/Source/Runtime/Engine/Classes/Engine/TextureRenderTarget2D.h
+++ b/Engine/Source/Runtime/Engine/Classes/Engine/TextureRenderTarget2D.h
@@ -2,6 +2,7 @@
 
 #pragma once
 #include "RenderUtils.h"
+#include "TextureRenderTarget.h"
 #include "TextureRenderTarget2D.generated.h"
 
 /**
diff --git a/Engine/Source/Runtime/Engine/Classes/Engine/TextureRenderTargetCube.h b/Engine/Source/Runtime/Engine/Classes/Engine/TextureRenderTargetCube.h
--- a/Engine/Source/Runtime/Engine/Classes/Engine/TextureRenderTargetCube.h
+++ b/Engine/Source/Runtime/Engine/Classes/Engine/TextureRenderTargetCube.h
@@ -1,6 +1,7 @@
 // Copyright 1998-2014 Epic Games, Inc. All Rights Reserved.
 
 #pragma once
+#include "TextureRenderTarget.h"
 #include "TextureRenderTargetCube.generated.h"
 
 /**
diff --git a/Engine/Source/Runtime/Engine/Classes/GameFramework/EngineMessage.h b/Engine/Source/Runtime/Engine/Classes/GameFramework/EngineMessage.h
--- a/Engine/Source/Runtime/Engine/Classes/GameFramework/EngineMessage.h
+++ b/Engine/Source/Runtime/Engine/Classes/GameFramework/EngineMessage.h
@@ -1,6 +1,7 @@
 // Copyright 1998-2014 Epic Games, Inc. All Rights Reserved.
 
 #pragma once
+#include "LocalMessage.h"
 #include "EngineMessage.generated.h"
 
 UCLASS(abstract, MinimalAPI)
diff --git a/Engine/Source/Runtime/Engine/Classes/GameFramework/PlayerStart.h b/Engine/Source/Runtime/Engine/Classes/GameFramework/PlayerStart.h
--- a/Engine/Source/Runtime/Engine/Classes/GameFramework/PlayerStart.h
+++ b/Engine/Source/Runtime/Engine/Classes/GameFramework/PlayerStart.h
@@ -5,6 +5,7 @@
 //=============================================================================
 
 #pragma once
+#include "Engine/NavigationObjectBase.h"
 #include "PlayerStart.generated.h"
 
 UCLASS(ClassGroup=Common, hidecategories=Collision)
diff --git a/Engine/Source/Runtime/Engine/Classes/Particles/ParticleEmitter.h b/Engine/Source/Runtime/Engine/Classes/Particles/ParticleEmitter.h
--- a/Engine/Source/Runtime/Engine/Classes/Particles/ParticleEmitter.h
+++ b/Engine/Source/Runtime/Engine/Classes/Particles/ParticleEmitter.h
@@ -8,8 +8,11 @@
 #pragma once
 #include "Components/SceneComponent.h"
 #include "ParticleEmitterInstances.h"
+#include "Scalability.h"
 #include "ParticleEmitter.generated.h"
 
+class UInterpCurveEdSetup;
+
 //=============================================================================
 //	Burst emissions
 //=============================================================================
diff --git a/Engine/Source/Runtime/Engine/Classes/Particles/ParticleLODLevel.h b/Engine/Source/Runtime/Engine/Classes/Particles/ParticleLODLevel.h
--- a/Engine/Source/Runtime/Engine/Classes/Particles/ParticleLODLevel.h
+++ b/Engine/Source/Runtime/Engine/Classes/Particles/ParticleLODLevel.h
@@ -8,6 +8,8 @@
 #pragma once
 #include "ParticleLODLevel.generated.h"
 
+class UInterpCurveEdSetup;
+
 UCLASS(collapsecategories, hidecategories=Object, editinlinenew, MinimalAPI)
 class UParticleLODLevel : public UObject
 {
diff --git a/Engine/Source/Runtime/Engine/Classes/Particles/ParticleModule.h b/Engine/Source/Runtime/Engine/Classes/Particles/ParticleModule.h
--- a/Engine/Source/Runtime/Engine/Classes/Particles/ParticleModule.h
+++ b/Engine/Source/Runtime/Engine/Classes/Particles/ParticleModule.h
@@ -5,8 +5,10 @@
 #include "ParticleHelper.h"
 #include "ParticleModule.generated.h"
 
+class UInterpCurveEdSetup;
 class UParticleEmitter;
 class UParticleLODLevel;
+struct FCurveEdEntry;
 struct FParticleEmitterInstance;
 
 /** ModuleType
diff --git a/Engine/Source/Runtime/Engine/Classes/Particles/ParticleModuleRequired.h b/Engine/Source/Runtime/Engine/Classes/Particles/ParticleModuleRequired.h
--- a/Engine/Source/Runtime/Engine/Classes/Particles/ParticleModuleRequired.h
+++ b/Engine/Source/Runtime/Engine/Classes/Particles/ParticleModuleRequired.h
@@ -7,6 +7,8 @@
 #include "Particles/ParticleSpriteEmitter.h"
 #include "ParticleModuleRequired.generated.h"
 
+class UInterpCurveEdSetup;
+
 UENUM()
 enum EParticleSortMode
 {
diff --git a/Engine/Source/Runtime/Engine/Classes/PhysicsEngine/PhysicsThruster.h b/Engine/Source/Runtime/Engine/Classes/PhysicsEngine/PhysicsThruster.h
--- a/Engine/Source/Runtime/Engine/Classes/PhysicsEngine/PhysicsThruster.h
+++ b/Engine/Source/Runtime/Engine/Classes/PhysicsEngine/PhysicsThruster.h
@@ -3,6 +3,7 @@
 
 
 #pragma once
+#include "RigidBodyBase.h"
 #include "PhysicsThruster.generated.h"
 
 /** 
diff --git a/Engine/Source/Runtime/Engine/Classes/PhysicsEngine/RadialForceActor.h b/Engine/Source/Runtime/Engine/Classes/PhysicsEngine/RadialForceActor.h
--- a/Engine/Source/Runtime/Engine/Classes/PhysicsEngine/RadialForceActor.h
+++ b/Engine/Source/Runtime/Engine/Classes/PhysicsEngine/RadialForceActor.h
@@ -2,6 +2,7 @@
 
 
 #pragma once
+#include "RigidBodyBase.h"
 #include "RadialForceActor.generated.h"
 
 UCLASS(MinimalAPI, hidecategories=(Collision, Input), showcategories=("Input|MouseInput", "Input|TouchInput"))
diff --git a/Engine/Source/Runtime/Engine/Private/AI/Navigation/NavLinkProxy.cpp b/Engine/Source/Runtime/Engine/Private/AI/Navigation/NavLinkProxy.cpp
--- a/Engine/Source/Runtime/Engine/Private/AI/Navigation/NavLinkProxy.cpp
+++ b/Engine/Source/Runtime/Engine/Private/AI/Navigation/NavLinkProxy.cpp
@@ -9,6 +9,7 @@
 #include "AI/Navigation/NavLinkCustomComponent.h"
 #include "AI/Navigation/NavLinkProxy.h"
 #include "AI/Navigation/NavLinkRenderingComponent.h"
+#include "NavigationSystemHelpers.h"
 
 ANavLinkProxy::ANavLinkProxy(const FObjectInitializer& ObjectInitializer) : Super(ObjectInitializer)
 {
diff --git a/Engine/Source/Runtime/Engine/Private/AI/Navigation/NavLinkRenderingComponent.cpp b/Engine/Source/Runtime/Engine/Private/AI/Navigation/NavLinkRenderingComponent.cpp
--- a/Engine/Source/Runtime/Engine/Private/AI/Navigation/NavLinkRenderingComponent.cpp
+++ b/Engine/Source/Runtime/Engine/Private/AI/Navigation/NavLinkRenderingComponent.cpp
@@ -7,6 +7,8 @@
 #include "ShowFlags.h"
 #include "ConvexVolume.h"
 #endif
+#include "AI/Navigation/NavLinkHostInterface.h"
+#include "AI/Navigation/NavAreas/NavArea.h"
 
 //----------------------------------------------------------------------//
 // UNavLinkRenderingComponent
diff --git a/Engine/Source/Runtime/Engine/Private/AI/Navigation/NavTestRenderingComponent.cpp b/Engine/Source/Runtime/Engine/Private/AI/Navigation/NavTestRenderingComponent.cpp
--- a/Engine/Source/Runtime/Engine/Private/AI/Navigation/NavTestRenderingComponent.cpp
+++ b/Engine/Source/Runtime/Engine/Private/AI/Navigation/NavTestRenderingComponent.cpp
@@ -4,6 +4,8 @@
 #include "DebugRenderSceneProxy.h"
 #include "AI/Navigation/NavTestRenderingComponent.h"
 #include "DynamicMeshBuilder.h"
+#include "AI/Navigation/NavigationTestingActor.h"
+#include "Debug/DebugDrawService.h"
 
 static const FColor NavMeshRenderColor_OpenSet(255,128,0,255);
 static const FColor NavMeshRenderColor_ClosedSet(255,196,0,255);
diff --git a/Engine/Source/Runtime/Engine/Private/AI/Navigation/NavigationModifier.cpp b/Engine/Source/Runtime/Engine/Private/AI/Navigation/NavigationModifier.cpp
--- a/Engine/Source/Runtime/Engine/Private/AI/Navigation/NavigationModifier.cpp
+++ b/Engine/Source/Runtime/Engine/Private/AI/Navigation/NavigationModifier.cpp
@@ -11,6 +11,7 @@
 #include "AI/Navigation/NavLinkDefinition.h"
 #include "AI/Navigation/NavLinkTrivial.h"
 #include "AI/Navigation/NavAreas/NavAreaMeta.h"
+#include "Components/BrushComponent.h"
 
 // if square distance between two points is less than this the those points
 // will be considered identical when calculating convex hull
diff --git a/Engine/Source/Runtime/Engine/Private/AI/Navigation/NavigationPath.cpp b/Engine/Source/Runtime/Engine/Private/AI/Navigation/NavigationPath.cpp
--- a/Engine/Source/Runtime/Engine/Private/AI/Navigation/NavigationPath.cpp
+++ b/Engine/Source/Runtime/Engine/Private/AI/Navigation/NavigationPath.cpp
@@ -5,6 +5,7 @@
 #include "AI/Navigation/NavigationTypes.h"
 #include "AI/Navigation/RecastNavMesh.h"
 #include "AI/Navigation/NavigationPath.h"
+#include "Debug/DebugDrawService.h"
 
 #define DEBUG_DRAW_OFFSET 0
 #define PATH_OFFSET_KEEP_VISIBLE_POINTS 1
diff --git a/Engine/Source/Runtime/Engine/Private/AI/Navigation/NavigationSystemHelpers.cpp b/Engine/Source/Runtime/Engine/Private/AI/Navigation/NavigationSystemHelpers.cpp
--- a/Engine/Source/Runtime/Engine/Private/AI/Navigation/NavigationSystemHelpers.cpp
+++ b/Engine/Source/Runtime/Engine/Private/AI/Navigation/NavigationSystemHelpers.cpp
@@ -8,6 +8,7 @@
 #endif // WITH_RECAST
 #include "AI/NavigationSystemHelpers.h"
 #include "VisualLogger/VisualLogger.h"
+#include "AI/Navigation/NavCollision.h"
 
 namespace NavigationHelper
 {
diff --git a/Engine/Source/Runtime/Engine/Private/AI/Navigation/NavigationTestingActor.cpp b/Engine/Source/Runtime/Engine/Private/AI/Navigation/NavigationTestingActor.cpp
--- a/Engine/Source/Runtime/Engine/Private/AI/Navigation/NavigationTestingActor.cpp
+++ b/Engine/Source/Runtime/Engine/Private/AI/Navigation/NavigationTestingActor.cpp
@@ -6,6 +6,7 @@
 #endif
 #include "AI/Navigation/NavigationTestingActor.h"
 #include "AI/Navigation/NavTestRenderingComponent.h"
+#include "Components/CapsuleComponent.h"
 
 void FNavTestTickHelper::Tick(float DeltaTime)
 {
diff --git a/Engine/Source/Runtime/Engine/Private/AI/Navigation/RecastNavMeshGenerator.cpp b/Engine/Source/Runtime/Engine/Private/AI/Navigation/RecastNavMeshGenerator.cpp
--- a/Engine/Source/Runtime/Engine/Private/AI/Navigation/RecastNavMeshGenerator.cpp
+++ b/Engine/Source/Runtime/Engine/Private/AI/Navigation/RecastNavMeshGenerator.cpp
@@ -22,6 +22,7 @@
 #include "DetourTileCache.h"
 #include "DetourTileCacheBuilder.h"
 #include "RecastHelpers.h"
+#include "NavigationSystemHelpers.h"
 
 #define SEAMLESS_REBUILDING_ENABLED 1
 
diff --git a/Engine/Source/Runtime/Engine/Private/Actor.cpp b/Engine/Source/Runtime/Engine/Private/Actor.cpp
--- a/Engine/Source/Runtime/Engine/Private/Actor.cpp
+++ b/Engine/Source/Runtime/Engine/Private/Actor.cpp
@@ -18,6 +18,13 @@
 #include "VisualLogger/VisualLogger.h"
 #include "Animation/AnimInstance.h"
 #include "Engine/DemoNetDriver.h"
+#include "GameFramework/Pawn.h"
+#include "Components/PawnNoiseEmitterComponent.h"
+#include "Components/TimelineComponent.h"
+#include "Camera/CameraComponent.h"
+#include "GameFramework/DamageType.h"
+#include "Kismet/GameplayStatics.h"
+#include "PhysicalMaterials/PhysicalMaterial.h"
 
 //DEFINE_LOG_CATEGORY_STATIC(LogActor, Log, All);
 DEFINE_LOG_CATEGORY(LogActor);
diff --git a/Engine/Source/Runtime/Engine/Private/ActorEditor.cpp b/Engine/Source/Runtime/Engine/Private/ActorEditor.cpp
--- a/Engine/Source/Runtime/Engine/Private/ActorEditor.cpp
+++ b/Engine/Source/Runtime/Engine/Private/ActorEditor.cpp
@@ -9,6 +9,7 @@
 #include "MapErrors.h"
 #include "Foliage/InstancedFoliageActor.h"
 #include "Engine/LevelBounds.h"
+#include "Components/ChildActorComponent.h"
 
 #if WITH_EDITOR
 
diff --git a/Engine/Source/Runtime/Engine/Private/Animation/AnimNotifyState_TimedParticleEffect.cpp b/Engine/Source/Runtime/Engine/Private/Animation/AnimNotifyState_TimedParticleEffect.cpp
--- a/Engine/Source/Runtime/Engine/Private/Animation/AnimNotifyState_TimedParticleEffect.cpp
+++ b/Engine/Source/Runtime/Engine/Private/Animation/AnimNotifyState_TimedParticleEffect.cpp
@@ -3,6 +3,7 @@
 #include "EnginePrivate.h"
 #include "Animation/AnimNotifies/AnimNotifyState_TimedParticleEffect.h"
 #include "Particles/ParticleSystemComponent.h"
+#include "Kismet/GameplayStatics.h"
 
 UAnimNotifyState_TimedParticleEffect::UAnimNotifyState_TimedParticleEffect(const FObjectInitializer& ObjectInitializer)
 : Super(ObjectInitializer)
diff --git a/Engine/Source/Runtime/Engine/Private/Animation/AnimSequence.cpp b/Engine/Source/Runtime/Engine/Private/Animation/AnimSequence.cpp
--- a/Engine/Source/Runtime/Engine/Private/Animation/AnimSequence.cpp
+++ b/Engine/Source/Runtime/Engine/Private/Animation/AnimSequence.cpp
@@ -14,6 +14,7 @@
 #include "Animation/AnimNotifies/AnimNotify.h"
 #include "Animation/AnimNotifies/AnimNotifyState.h"
 #include "Animation/Rig.h"
+#include "EditorFramework/AssetImportData.h"
 
 #define USE_SLERP 0
 
diff --git a/Engine/Source/Runtime/Engine/Private/Animation/AnimationUtils.cpp b/Engine/Source/Runtime/Engine/Private/Animation/AnimationUtils.cpp
--- a/Engine/Source/Runtime/Engine/Private/Animation/AnimationUtils.cpp
+++ b/Engine/Source/Runtime/Engine/Private/Animation/AnimationUtils.cpp
@@ -14,6 +14,7 @@
 #include "AnimationUtils.h"
 #include "AnimationCompression.h"
 #include "AnimEncoding.h"
+#include "Engine/SkeletalMeshSocket.h"
 
 /** Array to keep track of SkeletalMeshes we have built metadata for, and log out the results just once. */
 //static TArray<USkeleton*> UniqueSkeletonsMetadataArray;
diff --git a/Engine/Source/Runtime/Engine/Private/BlueprintGeneratedClass.cpp b/Engine/Source/Runtime/Engine/Private/BlueprintGeneratedClass.cpp
--- a/Engine/Source/Runtime/Engine/Private/BlueprintGeneratedClass.cpp
+++ b/Engine/Source/Runtime/Engine/Private/BlueprintGeneratedClass.cpp
@@ -3,6 +3,7 @@
 #include "EnginePrivate.h"
 #include "BlueprintUtilities.h"
 #include "Engine/InputDelegateBinding.h"
+#include "Engine/TimelineTemplate.h"
 
 #if WITH_EDITOR
 #include "BlueprintEditorUtils.h"
diff --git a/Engine/Source/Runtime/Engine/Private/BoxComponent.cpp b/Engine/Source/Runtime/Engine/Private/BoxComponent.cpp
--- a/Engine/Source/Runtime/Engine/Private/BoxComponent.cpp
+++ b/Engine/Source/Runtime/Engine/Private/BoxComponent.cpp
@@ -2,6 +2,7 @@
 
 
 #include "EnginePrivate.h"
+#include "Components/BoxComponent.h"
 
 UBoxComponent::UBoxComponent(const FObjectInitializer& ObjectInitializer)
 	: Super(ObjectInitializer)
diff --git a/Engine/Source/Runtime/Engine/Private/Brush.cpp b/Engine/Source/Runtime/Engine/Private/Brush.cpp
--- a/Engine/Source/Runtime/Engine/Private/Brush.cpp
+++ b/Engine/Source/Runtime/Engine/Private/Brush.cpp
@@ -5,10 +5,12 @@
 =============================================================================*/
 
 #include "EnginePrivate.h"
+#include "Engine/Polys.h"
 #include "Net/UnrealNetwork.h"
 #include "Model.h"
 #include "Engine/BrushBuilder.h"
 #include "ActorEditorUtils.h"
+#include "Components/BrushComponent.h"
 
 #if WITH_EDITOR
 /** Define static delegate */
diff --git a/Engine/Source/Runtime/Engine/Private/BrushComponent.cpp b/Engine/Source/Runtime/Engine/Private/BrushComponent.cpp
--- a/Engine/Source/Runtime/Engine/Private/BrushComponent.cpp
+++ b/Engine/Source/Runtime/Engine/Private/BrushComponent.cpp
@@ -5,6 +5,8 @@
 =============================================================================*/
 
 #include "EnginePrivate.h"
+#include "Components/BrushComponent.h"
+#include "Engine/Polys.h"
 #include "Model.h"
 #include "LevelUtils.h"
 #if WITH_EDITOR
diff --git a/Engine/Source/Runtime/Engine/Private/Camera/CameraActor.cpp b/Engine/Source/Runtime/Engine/Private/Camera/CameraActor.cpp
--- a/Engine/Source/Runtime/Engine/Private/Camera/CameraActor.cpp
+++ b/Engine/Source/Runtime/Engine/Private/Camera/CameraActor.cpp
@@ -1,6 +1,7 @@
 // Copyright 1998-2014 Epic Games, Inc. All Rights Reserved.
 
 #include "EnginePrivate.h"
+#include "Camera/CameraComponent.h"
 #include "Camera/CameraActor.h"
 #include "Camera/CameraAnim.h"
 #include "Net/UnrealNetwork.h"
diff --git a/Engine/Source/Runtime/Engine/Private/Camera/CameraComponent.cpp b/Engine/Source/Runtime/Engine/Private/Camera/CameraComponent.cpp
--- a/Engine/Source/Runtime/Engine/Private/Camera/CameraComponent.cpp
+++ b/Engine/Source/Runtime/Engine/Private/Camera/CameraComponent.cpp
@@ -4,6 +4,8 @@
 #include "MessageLog.h"
 #include "UObjectToken.h"
 #include "MapErrors.h"
+#include "Camera/CameraComponent.h"
+#include "Components/DrawFrustumComponent.h"
 
 #define LOCTEXT_NAMESPACE "CameraComponent"
 
diff --git a/Engine/Source/Runtime/Engine/Private/Camera/CameraModifier.cpp b/Engine/Source/Runtime/Engine/Private/Camera/CameraModifier.cpp
--- a/Engine/Source/Runtime/Engine/Private/Camera/CameraModifier.cpp
+++ b/Engine/Source/Runtime/Engine/Private/Camera/CameraModifier.cpp
@@ -4,6 +4,7 @@
 #include "ParticleDefinitions.h"
 
 #include "Net/UnrealNetwork.h"
+#include "Camera/CameraModifier.h"
 
 //////////////////////////////////////////////////////////////////////////
 
diff --git a/Engine/Source/Runtime/Engine/Private/Camera/CameraModifier_CameraShake.cpp b/Engine/Source/Runtime/Engine/Private/Camera/CameraModifier_CameraShake.cpp
--- a/Engine/Source/Runtime/Engine/Private/Camera/CameraModifier_CameraShake.cpp
+++ b/Engine/Source/Runtime/Engine/Private/Camera/CameraModifier_CameraShake.cpp
@@ -2,6 +2,7 @@
 
 #include "EnginePrivate.h"
 #include "Camera/CameraAnimInst.h"
+#include "Camera/CameraModifier_CameraShake.h"
 
 //////////////////////////////////////////////////////////////////////////
 // UCameraModifier_CameraShake
diff --git a/Engine/Source/Runtime/Engine/Private/CapsuleComponent.cpp b/Engine/Source/Runtime/Engine/Private/CapsuleComponent.cpp
--- a/Engine/Source/Runtime/Engine/Private/CapsuleComponent.cpp
+++ b/Engine/Source/Runtime/Engine/Private/CapsuleComponent.cpp
@@ -2,6 +2,7 @@
 
 
 #include "EnginePrivate.h"
+#include "Components/CapsuleComponent.h"
 
 
 UCapsuleComponent::UCapsuleComponent(const FObjectInitializer& ObjectInitializer)
diff --git a/Engine/Source/Runtime/Engine/Private/Character.cpp b/Engine/Source/Runtime/Engine/Private/Character.cpp
--- a/Engine/Source/Runtime/Engine/Private/Character.cpp
+++ b/Engine/Source/Runtime/Engine/Private/Character.cpp
@@ -14,6 +14,8 @@
 #include "Animation/AnimInstance.h"
 #include "Animation/AnimMontage.h"
 #include "Foliage/InstancedFoliageActor.h"
+#include "Components/CapsuleComponent.h"
+#include "GameFramework/DamageType.h"
 
 DEFINE_LOG_CATEGORY_STATIC(LogCharacter, Log, All);
 DEFINE_LOG_CATEGORY_STATIC(LogAvatar, Log, All);
diff --git a/Engine/Source/Runtime/Engine/Private/CharacterMovementComponent.cpp b/Engine/Source/Runtime/Engine/Private/CharacterMovementComponent.cpp
--- a/Engine/Source/Runtime/Engine/Private/CharacterMovementComponent.cpp
+++ b/Engine/Source/Runtime/Engine/Private/CharacterMovementComponent.cpp
@@ -10,12 +10,17 @@
 #include "GameFramework/GameNetworkManager.h"
 #include "GameFramework/Character.h"
 #include "GameFramework/CharacterMovementComponent.h"
+#include "GameFramework/GameState.h"
 #include "Components/PrimitiveComponent.h"
 #include "Animation/AnimMontage.h"
 #include "PhysicsEngine/DestructibleActor.h"
 
 // @todo this is here only due to circular dependency to AIModule. To be removed
 #include "Navigation/PathFollowingComponent.h"
+#include "AI/Navigation/AvoidanceManager.h"
+#include "Components/CapsuleComponent.h"
+#include "Components/BrushComponent.h"
+#include "Components/DestructibleComponent.h"
 
 DEFINE_LOG_CATEGORY_STATIC(LogCharacterMovement, Log, All);
 
diff --git a/Engine/Source/Runtime/Engine/Private/ChildActorComponent.cpp b/Engine/Source/Runtime/Engine/Private/ChildActorComponent.cpp
--- a/Engine/Source/Runtime/Engine/Private/ChildActorComponent.cpp
+++ b/Engine/Source/Runtime/Engine/Private/ChildActorComponent.cpp
@@ -2,6 +2,7 @@
 
 #include "EnginePrivate.h"
 #include "ComponentInstanceDataCache.h"
+#include "Components/ChildActorComponent.h"
 
 DEFINE_LOG_CATEGORY_STATIC(LogChildActorComponent, Warning, All);
 
diff --git a/Engine/Source/Runtime/Engine/Private/Collision/CollisionConversions.cpp b/Engine/Source/Runtime/Engine/Private/Collision/CollisionConversions.cpp
--- a/Engine/Source/Runtime/Engine/Private/Collision/CollisionConversions.cpp
+++ b/Engine/Source/Runtime/Engine/Private/Collision/CollisionConversions.cpp
@@ -7,6 +7,8 @@
 #include "PhysXCollision.h"
 #include "CollisionDebugDrawing.h"
 #include "CollisionConversions.h"
+#include "Components/DestructibleComponent.h"
+#include "PhysicalMaterials/PhysicalMaterial.h"
 
 // Sentinel for invalid query results.
 static const PxQueryHit InvalidQueryHit;
diff --git a/Engine/Source/Runtime/Engine/Private/Collision/PhysXCollision.cpp b/Engine/Source/Runtime/Engine/Private/Collision/PhysXCollision.cpp
--- a/Engine/Source/Runtime/Engine/Private/Collision/PhysXCollision.cpp
+++ b/Engine/Source/Runtime/Engine/Private/Collision/PhysXCollision.cpp
@@ -4,6 +4,8 @@
 #include "PhysicsPublic.h"
 #include "Collision.h"
 #include "CollisionDebugDrawingPublic.h"
+#include "PhysicsEngine/PhysicsSettings.h"
+#include "PhysicalMaterials/PhysicalMaterial.h"
 
 #if WITH_BOX2D
 #include "../PhysicsEngine2D/Box2DIntegration.h"
diff --git a/Engine/Source/Runtime/Engine/Private/ContentStreaming.cpp b/Engine/Source/Runtime/Engine/Private/ContentStreaming.cpp
--- a/Engine/Source/Runtime/Engine/Private/ContentStreaming.cpp
+++ b/Engine/Source/Runtime/Engine/Private/ContentStreaming.cpp
@@ -7,6 +7,8 @@
 #include "EnginePrivate.h"
 #include "GenericPlatformMemoryPoolStats.h"
 #include "AudioStreaming.h"
+#include "Engine/LightMapTexture2D.h"
+#include "Engine/ShadowMapTexture2D.h"
 
 DEFINE_LOG_CATEGORY_STATIC(LogContentStreaming, Log, All);
 
diff --git a/Engine/Source/Runtime/Engine/Private/Controller.cpp b/Engine/Source/Runtime/Engine/Private/Controller.cpp
--- a/Engine/Source/Runtime/Engine/Private/Controller.cpp
+++ b/Engine/Source/Runtime/Engine/Private/Controller.cpp
@@ -17,6 +17,9 @@
 #include "Navigation/PathFollowingComponent.h"
 #include "Navigation/NavigationComponent.h"
 
+#include "GameFramework/PlayerState.h"
+#include "Components/CapsuleComponent.h"
+
 DEFINE_LOG_CATEGORY(LogPath);
 
 
diff --git a/Engine/Source/Runtime/Engine/Private/CurveLinearColor.cpp b/Engine/Source/Runtime/Engine/Private/CurveLinearColor.cpp
--- a/Engine/Source/Runtime/Engine/Private/CurveLinearColor.cpp
+++ b/Engine/Source/Runtime/Engine/Private/CurveLinearColor.cpp
@@ -5,6 +5,7 @@
 =============================================================================*/
 
 #include "EnginePrivate.h"
+#include "Curves/CurveLinearColor.h"
 
 FLinearColor FRuntimeCurveLinearColor::GetLinearColorValue(float InTime) const
 {
diff --git a/Engine/Source/Runtime/Engine/Private/CurveVector.cpp b/Engine/Source/Runtime/Engine/Private/CurveVector.cpp
--- a/Engine/Source/Runtime/Engine/Private/CurveVector.cpp
+++ b/Engine/Source/Runtime/Engine/Private/CurveVector.cpp
@@ -5,6 +5,7 @@
 =============================================================================*/
 
 #include "EnginePrivate.h"
+#include "Curves/CurveVector.h"
 
 
 UCurveVector::UCurveVector(const FObjectInitializer& ObjectInitializer)
diff --git a/Engine/Source/Runtime/Engine/Private/DamageType.cpp b/Engine/Source/Runtime/Engine/Private/DamageType.cpp
--- a/Engine/Source/Runtime/Engine/Private/DamageType.cpp
+++ b/Engine/Source/Runtime/Engine/Private/DamageType.cpp
@@ -1,6 +1,7 @@
 // Copyright 1998-2014 Epic Games, Inc. All Rights Reserved.
 
 #include "EnginePrivate.h"
+#include "GameFramework/DamageType.h"
 
 UDamageType::UDamageType(const FObjectInitializer& ObjectInitializer)
 	: Super(ObjectInitializer)
diff --git a/Engine/Source/Runtime/Engine/Private/Debug/DebugDrawService.cpp b/Engine/Source/Runtime/Engine/Private/Debug/DebugDrawService.cpp
--- a/Engine/Source/Runtime/Engine/Private/Debug/DebugDrawService.cpp
+++ b/Engine/Source/Runtime/Engine/Private/Debug/DebugDrawService.cpp
@@ -1,6 +1,8 @@
 // Copyright 1998-2014 Epic Games, Inc. All Rights Reserved.
 
 #include "EnginePrivate.h"
+#include "DebugRenderSceneProxy.h"
+#include "Debug/DebugDrawService.h"
 
 TArray<TArray<FDebugDrawDelegate> > UDebugDrawService::Delegates;
 FEngineShowFlags UDebugDrawService::ObservedFlags(ESFIM_Editor);
diff --git a/Engine/Source/Runtime/Engine/Private/Debug/LogVisualizerCameraController.cpp b/Engine/Source/Runtime/Engine/Private/Debug/LogVisualizerCameraController.cpp
--- a/Engine/Source/Runtime/Engine/Private/Debug/LogVisualizerCameraController.cpp
+++ b/Engine/Source/Runtime/Engine/Private/Debug/LogVisualizerCameraController.cpp
@@ -4,6 +4,7 @@
 #include "Runtime/Engine/Classes/Debug/LogVisualizerCameraController.h"
 #include "Runtime/Engine/Classes/Debug/LogVisualizerHUD.h"
 #include "Runtime/Engine/Classes/Kismet/GameplayStatics.h"
+#include "GameFramework/PlayerInput.h"
 
 
 //----------------------------------------------------------------------//
diff --git a/Engine/Source/Runtime/Engine/Private/DebugCameraController.cpp b/Engine/Source/Runtime/Engine/Private/DebugCameraController.cpp
--- a/Engine/Source/Runtime/Engine/Private/DebugCameraController.cpp
+++ b/Engine/Source/Runtime/Engine/Private/DebugCameraController.cpp
@@ -11,6 +11,9 @@
 #include "GameFramework/SpectatorPawnMovement.h"
 #include "Engine/DebugCameraHUD.h"
 #include "LightMap.h"
+#include "Engine/DebugCameraController.h"
+#include "Components/DrawFrustumComponent.h"
+#include "GameFramework/PlayerInput.h"
 
 
 /** The currently selected actor. */
diff --git a/Engine/Source/Runtime/Engine/Private/DebugCameraHUD.cpp b/Engine/Source/Runtime/Engine/Private/DebugCameraHUD.cpp
--- a/Engine/Source/Runtime/Engine/Private/DebugCameraHUD.cpp
+++ b/Engine/Source/Runtime/Engine/Private/DebugCameraHUD.cpp
@@ -8,6 +8,7 @@
 #include "GameFramework/PawnMovementComponent.h"
 #include "GameFramework/SpectatorPawn.h"
 #include "Engine/DebugCameraHUD.h"
+#include "Engine/DebugCameraController.h"
 
 #if !(UE_BUILD_SHIPPING || UE_BUILD_TEST)
 static TAutoConsoleVariable<int32> CVarDebugCameraTraceComplex(
diff --git a/Engine/Source/Runtime/Engine/Private/DebugRenderSceneProxy.cpp b/Engine/Source/Runtime/Engine/Private/DebugRenderSceneProxy.cpp
--- a/Engine/Source/Runtime/Engine/Private/DebugRenderSceneProxy.cpp
+++ b/Engine/Source/Runtime/Engine/Private/DebugRenderSceneProxy.cpp
@@ -8,6 +8,7 @@
 
 #include "EnginePrivate.h"
 #include "DebugRenderSceneProxy.h"
+#include "Debug/DebugDrawService.h"
 
 // FPrimitiveSceneProxy interface.
 
diff --git a/Engine/Source/Runtime/Engine/Private/DecalActor.cpp b/Engine/Source/Runtime/Engine/Private/DecalActor.cpp
--- a/Engine/Source/Runtime/Engine/Private/DecalActor.cpp
+++ b/Engine/Source/Runtime/Engine/Private/DecalActor.cpp
@@ -2,6 +2,8 @@
 
 #include "EnginePrivate.h"
 #include "Engine/DecalActor.h"
+#include "Components/DecalComponent.h"
+#include "Components/BoxComponent.h"
 
 
 #if WITH_EDITOR
diff --git a/Engine/Source/Runtime/Engine/Private/DecalComponent.cpp b/Engine/Source/Runtime/Engine/Private/DecalComponent.cpp
--- a/Engine/Source/Runtime/Engine/Private/DecalComponent.cpp
+++ b/Engine/Source/Runtime/Engine/Private/DecalComponent.cpp
@@ -6,6 +6,7 @@
 
 #include "EnginePrivate.h"
 #include "LevelUtils.h"
+#include "Components/DecalComponent.h"
 
 FDeferredDecalProxy::FDeferredDecalProxy(const UDecalComponent* InComponent)
 {
diff --git a/Engine/Source/Runtime/Engine/Private/DefaultPawn.cpp b/Engine/Source/Runtime/Engine/Private/DefaultPawn.cpp
--- a/Engine/Source/Runtime/Engine/Private/DefaultPawn.cpp
+++ b/Engine/Source/Runtime/Engine/Private/DefaultPawn.cpp
@@ -4,6 +4,8 @@
 #include "EnginePrivate.h"
 #include "GameFramework/DefaultPawn.h"
 #include "GameFramework/FloatingPawnMovement.h"
+#include "Components/SphereComponent.h"
+#include "GameFramework/PlayerInput.h"
 
 FName ADefaultPawn::MovementComponentName(TEXT("MovementComponent0"));
 FName ADefaultPawn::CollisionComponentName(TEXT("CollisionComponent0"));
diff --git a/Engine/Source/Runtime/Engine/Private/DemoNetDriver.cpp b/Engine/Source/Runtime/Engine/Private/DemoNetDriver.cpp
--- a/Engine/Source/Runtime/Engine/Private/DemoNetDriver.cpp
+++ b/Engine/Source/Runtime/Engine/Private/DemoNetDriver.cpp
@@ -11,6 +11,8 @@
 #include "Engine/ActorChannel.h"
 #include "RepLayout.h"
 #include "GameFramework/SpectatorPawn.h"
+#include "Engine/LevelStreamingKismet.h"
+#include "GameFramework/PlayerStart.h"
 
 DEFINE_LOG_CATEGORY_STATIC( LogDemo, Log, All );
 
diff --git a/Engine/Source/Runtime/Engine/Private/DestructibleActor.cpp b/Engine/Source/Runtime/Engine/Private/DestructibleActor.cpp
--- a/Engine/Source/Runtime/Engine/Private/DestructibleActor.cpp
+++ b/Engine/Source/Runtime/Engine/Private/DestructibleActor.cpp
@@ -8,6 +8,7 @@
 #include "EnginePrivate.h"
 #include "PhysicsEngine/PhysXSupport.h"
 #include "PhysicsEngine/DestructibleActor.h"
+#include "Components/DestructibleComponent.h"
 
 ADestructibleActor::ADestructibleActor(const FObjectInitializer& ObjectInitializer)
 	: Super(ObjectInitializer)
diff --git a/Engine/Source/Runtime/Engine/Private/DestructibleComponent.cpp b/Engine/Source/Runtime/Engine/Private/DestructibleComponent.cpp
--- a/Engine/Source/Runtime/Engine/Private/DestructibleComponent.cpp
+++ b/Engine/Source/Runtime/Engine/Private/DestructibleComponent.cpp
@@ -15,6 +15,11 @@
 #include "Particles/ParticleSystemComponent.h"
 #include "PhysicsEngine/DestructibleActor.h"
 #include "Engine/DestructibleMesh.h"
+#include "Components/DestructibleComponent.h"
+#include "NavigationSystemHelpers.h"
+#include "PhysicalMaterials/PhysicalMaterial.h"
+#include "Kismet/GameplayStatics.h"
+#include "GameFramework/DamageType.h"
 
 UDestructibleComponent::UDestructibleComponent(const FObjectInitializer& ObjectInitializer)
 	: Super(ObjectInitializer)
diff --git a/Engine/Source/Runtime/Engine/Private/DirectionalLightComponent.cpp b/Engine/Source/Runtime/Engine/Private/DirectionalLightComponent.cpp
--- a/Engine/Source/Runtime/Engine/Private/DirectionalLightComponent.cpp
+++ b/Engine/Source/Runtime/Engine/Private/DirectionalLightComponent.cpp
@@ -6,6 +6,7 @@
 
 #include "EnginePrivate.h"
 #include "Engine/DirectionalLight.h"
+#include "Components/DirectionalLightComponent.h"
 
 static float GMaxCSMRadiusToAllowPerObjectShadows = 8000;
 static FAutoConsoleVariableRef CVarMaxCSMRadiusToAllowPerObjectShadows(
diff --git a/Engine/Source/Runtime/Engine/Private/Distributions.cpp b/Engine/Source/Runtime/Engine/Private/Distributions.cpp
--- a/Engine/Source/Runtime/Engine/Private/Distributions.cpp
+++ b/Engine/Source/Runtime/Engine/Private/Distributions.cpp
@@ -14,6 +14,9 @@
 #include "Sound/SoundNode.h"
 
 #include "Particles/ParticleModule.h"
+#include "Distributions/DistributionFloatUniformCurve.h"
+#include "Distributions/DistributionFloatConstantCurve.h"
+#include "Distributions/DistributionFloatParameterBase.h"
 
 // Moving UDistributions to PostInitProps to not be default sub-objects:
 // Small enough value to be rounded to 0.0 in the editor 
diff --git a/Engine/Source/Runtime/Engine/Private/DrawDebugHelpers.cpp b/Engine/Source/Runtime/Engine/Private/DrawDebugHelpers.cpp
--- a/Engine/Source/Runtime/Engine/Private/DrawDebugHelpers.cpp
+++ b/Engine/Source/Runtime/Engine/Private/DrawDebugHelpers.cpp
@@ -5,6 +5,7 @@
  */
 
 #include "EnginePrivate.h"
+#include "Components/LineBatchComponent.h"
 #include "GameFramework/HUD.h"
 
 void FlushPersistentDebugLines( const UWorld* InWorld )
diff --git a/Engine/Source/Runtime/Engine/Private/DrawFrustumComponent.cpp b/Engine/Source/Runtime/Engine/Private/DrawFrustumComponent.cpp
--- a/Engine/Source/Runtime/Engine/Private/DrawFrustumComponent.cpp
+++ b/Engine/Source/Runtime/Engine/Private/DrawFrustumComponent.cpp
@@ -6,6 +6,7 @@
 
 #include "EnginePrivate.h"
 #include "LevelUtils.h"
+#include "Components/DrawFrustumComponent.h"
 
 
 /** Represents a draw frustum to the scene manager. */
diff --git a/Engine/Source/Runtime/Engine/Private/EditorFramework/AssetImportData.cpp b/Engine/Source/Runtime/Engine/Private/EditorFramework/AssetImportData.cpp
--- a/Engine/Source/Runtime/Engine/Private/EditorFramework/AssetImportData.cpp
+++ b/Engine/Source/Runtime/Engine/Private/EditorFramework/AssetImportData.cpp
@@ -1,6 +1,7 @@
 // Copyright 1998-2014 Epic Games, Inc. All Rights Reserved.
 
 #include "EnginePrivate.h"
+#include "EditorFramework/AssetImportData.h"
 
 UAssetImportData::UAssetImportData(const FObjectInitializer& ObjectInitializer)
 	: Super(ObjectInitializer)
diff --git a/Engine/Source/Runtime/Engine/Private/EngineMessage.cpp b/Engine/Source/Runtime/Engine/Private/EngineMessage.cpp
--- a/Engine/Source/Runtime/Engine/Private/EngineMessage.cpp
+++ b/Engine/Source/Runtime/Engine/Private/EngineMessage.cpp
@@ -3,6 +3,8 @@
 #include "EnginePrivate.h"
 #include "Engine/Console.h"
 #include "GameFramework/EngineMessage.h"
+#include "GameFramework/LocalMessage.h"
+#include "GameFramework/PlayerState.h"
 
 UEngineMessage::UEngineMessage(const FObjectInitializer& ObjectInitializer) :
 	Super(ObjectInitializer)
diff --git a/Engine/Source/Runtime/Engine/Private/ErrorChecking.cpp b/Engine/Source/Runtime/Engine/Private/ErrorChecking.cpp
--- a/Engine/Source/Runtime/Engine/Private/ErrorChecking.cpp
+++ b/Engine/Source/Runtime/Engine/Private/ErrorChecking.cpp
@@ -15,6 +15,8 @@
 #include "GameFramework/DefaultPhysicsVolume.h"
 #include "Engine/LevelStreamingVolume.h"
 #include "Engine/Light.h"
+#include "Engine/Note.h"
+#include "Engine/Polys.h"
 
 #define LOCTEXT_NAMESPACE "ErrorChecking"
 
diff --git a/Engine/Source/Runtime/Engine/Private/GameEngine.cpp b/Engine/Source/Runtime/Engine/Private/GameEngine.cpp
--- a/Engine/Source/Runtime/Engine/Private/GameEngine.cpp
+++ b/Engine/Source/Runtime/Engine/Private/GameEngine.cpp
@@ -31,6 +31,10 @@
 #include "RendererInterface.h"
 #include "HotReloadInterface.h"
 #include "SDPIScaler.h"
+#include "Components/SkyLightComponent.h"
+#include "Components/ReflectionCaptureComponent.h"
+#include "Engine/GameEngine.h"
+#include "GameFramework/GameUserSettings.h"
 
 DEFINE_LOG_CATEGORY_STATIC(LogEngine, Log, All);
 
diff --git a/Engine/Source/Runtime/Engine/Private/GameMode.cpp b/Engine/Source/Runtime/Engine/Private/GameMode.cpp
--- a/Engine/Source/Runtime/Engine/Private/GameMode.cpp
+++ b/Engine/Source/Runtime/Engine/Private/GameMode.cpp
@@ -14,6 +14,11 @@
 #include "GameFramework/SpectatorPawn.h"
 #include "Engine/PlayerStartPIE.h"
 #include "GameFramework/EngineMessage.h"
+#include "GameFramework/GameState.h"
+#include "GameFramework/PlayerState.h"
+#include "GameFramework/GameSession.h"
+#include "GameFramework/CheatManager.h"
+#include "Engine/ChildConnection.h"
 
 DEFINE_LOG_CATEGORY_STATIC(LogGameMode, Log, All);
 
diff --git a/Engine/Source/Runtime/Engine/Private/GameSession.cpp b/Engine/Source/Runtime/Engine/Private/GameSession.cpp
--- a/Engine/Source/Runtime/Engine/Private/GameSession.cpp
+++ b/Engine/Source/Runtime/Engine/Private/GameSession.cpp
@@ -8,6 +8,8 @@
 #include "Net/UnrealNetwork.h"
 #include "OnlineSubsystemUtils.h"
 #include "Kismet/GameplayStatics.h"
+#include "GameFramework/PlayerState.h"
+#include "GameFramework/GameSession.h"
 
 DEFINE_LOG_CATEGORY_STATIC(LogGameSession, Log, All);
 
diff --git a/Engine/Source/Runtime/Engine/Private/GameState.cpp b/Engine/Source/Runtime/Engine/Private/GameState.cpp
--- a/Engine/Source/Runtime/Engine/Private/GameState.cpp
+++ b/Engine/Source/Runtime/Engine/Private/GameState.cpp
@@ -6,6 +6,8 @@
 
 #include "EnginePrivate.h"
 #include "Net/UnrealNetwork.h"
+#include "GameFramework/GameState.h"
+#include "GameFramework/PlayerState.h"
 
 DEFINE_LOG_CATEGORY_STATIC(LogGameState, Log, All);
 
diff --git a/Engine/Source/Runtime/Engine/Private/GameUserSettings.cpp b/Engine/Source/Runtime/Engine/Private/GameUserSettings.cpp
--- a/Engine/Source/Runtime/Engine/Private/GameUserSettings.cpp
+++ b/Engine/Source/Runtime/Engine/Private/GameUserSettings.cpp
@@ -4,6 +4,8 @@
 #include "SlateBasics.h"
 #include "AudioDevice.h"
 #include "Scalability.h"
+#include "GameFramework/GameUserSettings.h"
+#include "Engine/GameEngine.h"
 
 extern EWindowMode::Type GetWindowModeType(EWindowMode::Type WindowMode);
 
diff --git a/Engine/Source/Runtime/Engine/Private/GameViewportClient.cpp b/Engine/Source/Runtime/Engine/Private/GameViewportClient.cpp
--- a/Engine/Source/Runtime/Engine/Private/GameViewportClient.cpp
+++ b/Engine/Source/Runtime/Engine/Private/GameViewportClient.cpp
@@ -23,6 +23,12 @@
 #include "Runtime/GameLiveStreaming/Public/IGameLiveStreaming.h"
 #include "BufferVisualizationData.h"
 #include "RendererInterface.h"
+#include "GameFramework/InputSettings.h"
+#include "Components/LineBatchComponent.h"
+#include "Debug/DebugDrawService.h"
+#include "Components/BrushComponent.h"
+#include "Engine/GameEngine.h"
+#include "GameFramework/GameUserSettings.h"
 
 /** This variable allows forcing full screen of the first player controller viewport, even if there are multiple controllers plugged in and no cinematic playing. */
 bool GForceFullscreen = false;
diff --git a/Engine/Source/Runtime/Engine/Private/GameplayStatics.cpp b/Engine/Source/Runtime/Engine/Private/GameplayStatics.cpp
--- a/Engine/Source/Runtime/Engine/Private/GameplayStatics.cpp
+++ b/Engine/Source/Runtime/Engine/Private/GameplayStatics.cpp
@@ -14,6 +14,9 @@
 #include "Sound/SoundCue.h"
 #include "Engine/GameInstance.h"
 #include "GameFramework/SaveGame.h"
+#include "GameFramework/DamageType.h"
+#include "PhysicalMaterials/PhysicalMaterial.h"
+#include "Components/DecalComponent.h"
 
 //////////////////////////////////////////////////////////////////////////
 // UGameplayStatics
diff --git a/Engine/Source/Runtime/Engine/Private/GestureRecognizer.cpp b/Engine/Source/Runtime/Engine/Private/GestureRecognizer.cpp
--- a/Engine/Source/Runtime/Engine/Private/GestureRecognizer.cpp
+++ b/Engine/Source/Runtime/Engine/Private/GestureRecognizer.cpp
@@ -5,6 +5,8 @@
 =============================================================================*/
 
 #include "EnginePrivate.h"
+#include "GestureRecognizer.h"
+#include "GameFramework/PlayerInput.h"
 
 void FGestureRecognizer::DetectGestures(const FVector (&Touches)[EKeys::NUM_TOUCH_KEYS], UPlayerInput* PlayerInput, float DeltaTime)
 {
diff --git a/Engine/Source/Runtime/Engine/Private/HUD.cpp b/Engine/Source/Runtime/Engine/Private/HUD.cpp
--- a/Engine/Source/Runtime/Engine/Private/HUD.cpp
+++ b/Engine/Source/Runtime/Engine/Private/HUD.cpp
@@ -5,6 +5,7 @@
 =============================================================================*/
 
 #include "EnginePrivate.h"
+#include "Components/LineBatchComponent.h"
 #include "GameFramework/HUD.h"
 #include "Net/UnrealNetwork.h"
 #include "MessageLog.h"
diff --git a/Engine/Source/Runtime/Engine/Private/HardwareInfo.cpp b/Engine/Source/Runtime/Engine/Private/HardwareInfo.cpp
--- a/Engine/Source/Runtime/Engine/Private/HardwareInfo.cpp
+++ b/Engine/Source/Runtime/Engine/Private/HardwareInfo.cpp
@@ -7,6 +7,7 @@
 
 
 #include "EnginePrivate.h"
+#include "HardwareInfo.h"
 
 
 static TMap< FName, FString > HardwareDetailsMap;
diff --git a/Engine/Source/Runtime/Engine/Private/InGameAdvertising.cpp b/Engine/Source/Runtime/Engine/Private/InGameAdvertising.cpp
--- a/Engine/Source/Runtime/Engine/Private/InGameAdvertising.cpp
+++ b/Engine/Source/Runtime/Engine/Private/InGameAdvertising.cpp
@@ -5,6 +5,7 @@
 =============================================================================*/ 
 
 #include "EnginePrivate.h"
+#include "Engine/InGameAdManager.h"
 
 UInGameAdManager::UInGameAdManager(const FObjectInitializer& ObjectInitializer)
 	: Super(ObjectInitializer)
diff --git a/Engine/Source/Runtime/Engine/Private/InstancedStaticMesh.cpp b/Engine/Source/Runtime/Engine/Private/InstancedStaticMesh.cpp
--- a/Engine/Source/Runtime/Engine/Private/InstancedStaticMesh.cpp
+++ b/Engine/Source/Runtime/Engine/Private/InstancedStaticMesh.cpp
@@ -48,6 +48,9 @@
 #include "ShadowMap.h"
 #include "Logging/MessageLog.h"
 #endif
+#include "NavigationSystemHelpers.h"
+#include "AI/Navigation/NavCollision.h"
+#include "Components/InstancedStaticMeshComponent.h"
 
 // This must match the maximum a user could specify in the material (see 
 // FHLSLMaterialTranslator::TextureCoordinate), otherwise the material will attempt 
diff --git a/Engine/Source/Runtime/Engine/Private/InterpCurveEdSetup.cpp b/Engine/Source/Runtime/Engine/Private/InterpCurveEdSetup.cpp
--- a/Engine/Source/Runtime/Engine/Private/InterpCurveEdSetup.cpp
+++ b/Engine/Source/Runtime/Engine/Private/InterpCurveEdSetup.cpp
@@ -6,6 +6,7 @@
 
 #include "EnginePrivate.h"
 #include "Matinee/InterpTrack.h"
+#include "Engine/InterpCurveEdSetup.h"
 
 
 void UInterpCurveEdSetup::PostLoad()
diff --git a/Engine/Source/Runtime/Engine/Private/KismetMaterialLibrary.cpp b/Engine/Source/Runtime/Engine/Private/KismetMaterialLibrary.cpp
--- a/Engine/Source/Runtime/Engine/Private/KismetMaterialLibrary.cpp
+++ b/Engine/Source/Runtime/Engine/Private/KismetMaterialLibrary.cpp
@@ -4,6 +4,8 @@
 #include "Kismet/KismetMaterialLibrary.h"
 #include "MessageLog.h"
 #include "UObjectToken.h"
+#include "Materials/MaterialParameterCollection.h"
+#include "Materials/MaterialParameterCollectionInstance.h"
 
 //////////////////////////////////////////////////////////////////////////
 // UKismetMaterialLibrary
diff --git a/Engine/Source/Runtime/Engine/Private/KismetSystemLibrary.cpp b/Engine/Source/Runtime/Engine/Private/KismetSystemLibrary.cpp
--- a/Engine/Source/Runtime/Engine/Private/KismetSystemLibrary.cpp
+++ b/Engine/Source/Runtime/Engine/Private/KismetSystemLibrary.cpp
@@ -9,6 +9,8 @@
 #include "Advertising.h"
 #include "Online.h"
 #include "Camera/CameraActor.h"
+#include "Camera/CameraComponent.h"
+#include "Kismet/GameplayStatics.h"
 
 //////////////////////////////////////////////////////////////////////////
 // UKismetSystemLibrary
diff --git a/Engine/Source/Runtime/Engine/Private/Level.cpp b/Engine/Source/Runtime/Engine/Private/Level.cpp
--- a/Engine/Source/Runtime/Engine/Private/Level.cpp
+++ b/Engine/Source/Runtime/Engine/Private/Level.cpp
@@ -28,7 +28,11 @@ Level.cpp: Level-related functions
 #include "TargetPlatform.h"
 #include "ContentStreaming.h"
 #include "Foliage/InstancedFoliageActor.h"
-
+#include "Engine/NavigationObjectBase.h"
+#include "Engine/ShadowMapTexture2D.h"
+#include "Components/ModelComponent.h"
+#include "Engine/LightMapTexture2D.h"
+#include "Components/LineBatchComponent.h"
 DEFINE_LOG_CATEGORY(LogLevel);
 
 /*-----------------------------------------------------------------------------
diff --git a/Engine/Source/Runtime/Engine/Private/LevelActor.cpp b/Engine/Source/Runtime/Engine/Private/LevelActor.cpp
--- a/Engine/Source/Runtime/Engine/Private/LevelActor.cpp
+++ b/Engine/Source/Runtime/Engine/Private/LevelActor.cpp
@@ -10,6 +10,8 @@
 #if WITH_PHYSX
 	#include "PhysicsEngine/PhysXSupport.h"
 #endif // WITH_PHYSX
+#include "Components/SphereComponent.h"
+#include "Components/CapsuleComponent.h"
 
 // CVars
 static TAutoConsoleVariable<float> CVarEncroachEpsilon(
diff --git a/Engine/Source/Runtime/Engine/Private/LevelBounds.cpp b/Engine/Source/Runtime/Engine/Private/LevelBounds.cpp
--- a/Engine/Source/Runtime/Engine/Private/LevelBounds.cpp
+++ b/Engine/Source/Runtime/Engine/Private/LevelBounds.cpp
@@ -2,6 +2,7 @@
 
 #include "EnginePrivate.h"
 #include "Engine/LevelBounds.h"
+#include "Components/BoxComponent.h"
 
 // Default size of the box (scale)
 static const FVector DefaultLevelSize = FVector(1000.f);
diff --git a/Engine/Source/Runtime/Engine/Private/LevelStreaming.cpp b/Engine/Source/Runtime/Engine/Private/LevelStreaming.cpp
--- a/Engine/Source/Runtime/Engine/Private/LevelStreaming.cpp
+++ b/Engine/Source/Runtime/Engine/Private/LevelStreaming.cpp
@@ -12,6 +12,8 @@
 	#include "SNotificationList.h"
 	#include "NotificationManager.h"
 #endif
+#include "Engine/LevelStreamingKismet.h"
+#include "Components/BrushComponent.h"
 
 DEFINE_LOG_CATEGORY_STATIC(LogLevelStreaming, Log, All);
 
diff --git a/Engine/Source/Runtime/Engine/Private/Light.cpp b/Engine/Source/Runtime/Engine/Private/Light.cpp
--- a/Engine/Source/Runtime/Engine/Private/Light.cpp
+++ b/Engine/Source/Runtime/Engine/Private/Light.cpp
@@ -9,6 +9,8 @@
 #include "Components/ArrowComponent.h"
 #include "LightingBuildOptions.h"
 #include "Net/UnrealNetwork.h"
+#include "Components/DirectionalLightComponent.h"
+#include "Components/PointLightComponent.h"
 
 ALight::ALight(const FObjectInitializer& ObjectInitializer)
 	: Super(ObjectInitializer)
diff --git a/Engine/Source/Runtime/Engine/Private/LightComponent.cpp b/Engine/Source/Runtime/Engine/Private/LightComponent.cpp
--- a/Engine/Source/Runtime/Engine/Private/LightComponent.cpp
+++ b/Engine/Source/Runtime/Engine/Private/LightComponent.cpp
@@ -13,6 +13,9 @@
 #include "ComponentInstanceDataCache.h"
 #include "TargetPlatform.h"
 #include "ComponentReregisterContext.h"
+#include "Components/PointLightComponent.h"
+#include "Components/LightComponent.h"
+#include "Components/DirectionalLightComponent.h"
 
 void FStaticShadowDepthMap::InitRHI()
 {
diff --git a/Engine/Source/Runtime/Engine/Private/LightMap.cpp b/Engine/Source/Runtime/Engine/Private/LightMap.cpp
--- a/Engine/Source/Runtime/Engine/Private/LightMap.cpp
+++ b/Engine/Source/Runtime/Engine/Private/LightMap.cpp
@@ -8,6 +8,7 @@
 #include "TargetPlatform.h"
 #include "StaticLighting.h"
 #include "LightMap.h"
+#include "Components/InstancedStaticMeshComponent.h"
 
 DEFINE_LOG_CATEGORY_STATIC(LogLightMap, Log, All);
 
diff --git a/Engine/Source/Runtime/Engine/Private/LocalMessage.cpp b/Engine/Source/Runtime/Engine/Private/LocalMessage.cpp
--- a/Engine/Source/Runtime/Engine/Private/LocalMessage.cpp
+++ b/Engine/Source/Runtime/Engine/Private/LocalMessage.cpp
@@ -1,6 +1,7 @@
 // Copyright 1998-2014 Epic Games, Inc. All Rights Reserved.
 
 #include "EnginePrivate.h"
+#include "GameFramework/LocalMessage.h"
 
 FClientReceiveData::FClientReceiveData()
 	: LocalPC(NULL)
diff --git a/Engine/Source/Runtime/Engine/Private/LocalPlayer.cpp b/Engine/Source/Runtime/Engine/Private/LocalPlayer.cpp
--- a/Engine/Source/Runtime/Engine/Private/LocalPlayer.cpp
+++ b/Engine/Source/Runtime/Engine/Private/LocalPlayer.cpp
@@ -19,6 +19,7 @@
 #include "SceneViewExtension.h"
 #include "DataChannel.h"
 #include "GameFramework/OnlineSession.h"
+#include "GameFramework/PlayerInput.h"
 
 DEFINE_LOG_CATEGORY(LogPlayerManagement);
 DEFINE_LOG_CATEGORY_STATIC(LogEngine, Log, All);
diff --git a/Engine/Source/Runtime/Engine/Private/MaterialBillboardComponent.cpp b/Engine/Source/Runtime/Engine/Private/MaterialBillboardComponent.cpp
--- a/Engine/Source/Runtime/Engine/Private/MaterialBillboardComponent.cpp
+++ b/Engine/Source/Runtime/Engine/Private/MaterialBillboardComponent.cpp
@@ -3,6 +3,7 @@
 #include "EnginePrivate.h"
 #include "LevelUtils.h"
 #include "LocalVertexFactory.h"
+#include "Components/MaterialBillboardComponent.h"
 
 /** A material sprite vertex. */
 struct FMaterialSpriteVertex
diff --git a/Engine/Source/Runtime/Engine/Private/Materials/HLSLMaterialTranslator.h b/Engine/Source/Runtime/Engine/Private/Materials/HLSLMaterialTranslator.h
--- a/Engine/Source/Runtime/Engine/Private/Materials/HLSLMaterialTranslator.h
+++ b/Engine/Source/Runtime/Engine/Private/Materials/HLSLMaterialTranslator.h
@@ -21,6 +21,7 @@
 #include "MaterialCompiler.h"
 #include "MaterialUniformExpressions.h"
 #include "ParameterCollection.h"
+#include "Materials/MaterialParameterCollection.h"
 
 /** @return the number of components in a vector type. */
 uint32 GetNumComponents(EMaterialValueType Type)
diff --git a/Engine/Source/Runtime/Engine/Private/Materials/MaterialInstance.cpp b/Engine/Source/Runtime/Engine/Private/Materials/MaterialInstance.cpp
--- a/Engine/Source/Runtime/Engine/Private/Materials/MaterialInstance.cpp
+++ b/Engine/Source/Runtime/Engine/Private/Materials/MaterialInstance.cpp
@@ -14,6 +14,7 @@
 #include "MaterialShaderType.h"
 #include "TargetPlatform.h"
 #include "Engine/Font.h"
+#include "Engine/SubsurfaceProfile.h"
 
 /**
  * Cache uniform expressions for the given material.
diff --git a/Engine/Source/Runtime/Engine/Private/Materials/MaterialUniformExpressions.cpp b/Engine/Source/Runtime/Engine/Private/Materials/MaterialUniformExpressions.cpp
--- a/Engine/Source/Runtime/Engine/Private/Materials/MaterialUniformExpressions.cpp
+++ b/Engine/Source/Runtime/Engine/Private/Materials/MaterialUniformExpressions.cpp
@@ -6,6 +6,7 @@
 
 #include "EnginePrivate.h"
 #include "MaterialUniformExpressions.h"
+#include "Materials/MaterialParameterCollection.h"
 
 TLinkedList<FMaterialUniformExpressionType*>*& FMaterialUniformExpressionType::GetTypeList()
 {
diff --git a/Engine/Source/Runtime/Engine/Private/Materials/ParameterCollection.cpp b/Engine/Source/Runtime/Engine/Private/Materials/ParameterCollection.cpp
--- a/Engine/Source/Runtime/Engine/Private/Materials/ParameterCollection.cpp
+++ b/Engine/Source/Runtime/Engine/Private/Materials/ParameterCollection.cpp
@@ -2,6 +2,8 @@
 
 #include "EnginePrivate.h"
 #include "ParameterCollection.h"
+#include "Materials/MaterialParameterCollectionInstance.h"
+#include "Materials/MaterialParameterCollection.h"
 
 UMaterialParameterCollection::UMaterialParameterCollection(const FObjectInitializer& ObjectInitializer)
 	: Super(ObjectInitializer)
diff --git a/Engine/Source/Runtime/Engine/Private/MovementComponent.cpp b/Engine/Source/Runtime/Engine/Private/MovementComponent.cpp
--- a/Engine/Source/Runtime/Engine/Private/MovementComponent.cpp
+++ b/Engine/Source/Runtime/Engine/Private/MovementComponent.cpp
@@ -6,6 +6,7 @@
 #include "GameFramework/PhysicsVolume.h"
 #include "Components/SceneComponent.h"
 #include "MessageLog.h"
+#include "PhysicsEngine/PhysicsSettings.h"
 
 #define LOCTEXT_NAMESPACE "MovementComponent"
 DEFINE_LOG_CATEGORY_STATIC(LogMovement, Log, All);
diff --git a/Engine/Source/Runtime/Engine/Private/NavigationObjectBase.cpp b/Engine/Source/Runtime/Engine/Private/NavigationObjectBase.cpp
--- a/Engine/Source/Runtime/Engine/Private/NavigationObjectBase.cpp
+++ b/Engine/Source/Runtime/Engine/Private/NavigationObjectBase.cpp
@@ -2,6 +2,8 @@
 
 #include "EnginePrivate.h"
 #include "GameFramework/PhysicsVolume.h"
+#include "Engine/NavigationObjectBase.h"
+#include "Components/CapsuleComponent.h"
 
 DEFINE_LOG_CATEGORY_STATIC(LogNavigationPoint, Log, All);
 
diff --git a/Engine/Source/Runtime/Engine/Private/NetworkDriver.cpp b/Engine/Source/Runtime/Engine/Private/NetworkDriver.cpp
--- a/Engine/Source/Runtime/Engine/Private/NetworkDriver.cpp
+++ b/Engine/Source/Runtime/Engine/Private/NetworkDriver.cpp
@@ -16,6 +16,7 @@
 #include "NetworkingDistanceConstants.h"
 #include "DataChannel.h"
 #include "Engine/PackageMapClient.h"
+#include "GameFramework/PlayerState.h"
 
 // Default net driver stats
 DEFINE_STAT(STAT_Ping);
diff --git a/Engine/Source/Runtime/Engine/Private/NetworkPrediction.cpp b/Engine/Source/Runtime/Engine/Private/NetworkPrediction.cpp
--- a/Engine/Source/Runtime/Engine/Private/NetworkPrediction.cpp
+++ b/Engine/Source/Runtime/Engine/Private/NetworkPrediction.cpp
@@ -1,6 +1,7 @@
 // Copyright 1998-2014 Epic Games, Inc. All Rights Reserved.
 
 #include "EnginePrivate.h"
+#include "Interfaces/NetworkPredictionInterface.h"
 
 
 UNetworkPredictionInterface::UNetworkPredictionInterface(const FObjectInitializer& ObjectInitializer)
diff --git a/Engine/Source/Runtime/Engine/Private/Note.cpp b/Engine/Source/Runtime/Engine/Private/Note.cpp
--- a/Engine/Source/Runtime/Engine/Private/Note.cpp
+++ b/Engine/Source/Runtime/Engine/Private/Note.cpp
@@ -1,6 +1,7 @@
 // Copyright 1998-2014 Epic Games, Inc. All Rights Reserved.
 
 #include "EnginePrivate.h"
+#include "Engine/Note.h"
 
 ANote::ANote(const FObjectInitializer& ObjectInitializer)
 	: Super(ObjectInitializer)
diff --git a/Engine/Source/Runtime/Engine/Private/Particles/ParticleBeamModules.cpp b/Engine/Source/Runtime/Engine/Private/Particles/ParticleBeamModules.cpp
--- a/Engine/Source/Runtime/Engine/Private/Particles/ParticleBeamModules.cpp
+++ b/Engine/Source/Runtime/Engine/Private/Particles/ParticleBeamModules.cpp
@@ -16,6 +16,8 @@
 #include "Particles/TypeData/ParticleModuleTypeDataBeam2.h"
 #include "Particles/ParticleSystem.h"
 #include "Particles/ParticleSystemComponent.h"
+#include "Engine/InterpCurveEdSetup.h"
+#include "Distributions/DistributionFloatConstantCurve.h"
 
 /*-----------------------------------------------------------------------------
 	Abstract base modules used for categorization.
diff --git a/Engine/Source/Runtime/Engine/Private/Particles/ParticleComponents.cpp b/Engine/Source/Runtime/Engine/Private/Particles/ParticleComponents.cpp
--- a/Engine/Source/Runtime/Engine/Private/Particles/ParticleComponents.cpp
+++ b/Engine/Source/Runtime/Engine/Private/Particles/ParticleComponents.cpp
@@ -48,6 +48,9 @@
 #include "Particles/ParticleSystem.h"
 #include "Particles/ParticleSystemComponent.h"
 #include "Particles/ParticleSystemReplay.h"
+#include "Distributions/DistributionFloatConstantCurve.h"
+#include "Engine/InterpCurveEdSetup.h"
+#include "GameFramework/GameState.h"
 
 
 #define LOCTEXT_NAMESPACE "ParticleComponents"
diff --git a/Engine/Source/Runtime/Engine/Private/Particles/ParticleModules.cpp b/Engine/Source/Runtime/Engine/Private/Particles/ParticleModules.cpp
--- a/Engine/Source/Runtime/Engine/Private/Particles/ParticleModules.cpp
+++ b/Engine/Source/Runtime/Engine/Private/Particles/ParticleModules.cpp
@@ -74,6 +74,9 @@
 #include "Particles/ParticleSpriteEmitter.h"
 #include "Particles/ParticleSystem.h"
 #include "Particles/ParticleSystemComponent.h"
+#include "Distributions/DistributionFloatUniformCurve.h"
+#include "Engine/InterpCurveEdSetup.h"
+#include "Distributions/DistributionFloatConstantCurve.h"
 
 /*-----------------------------------------------------------------------------
 	Abstract base modules used for categorization.
diff --git a/Engine/Source/Runtime/Engine/Private/Particles/ParticleModules_Color.cpp b/Engine/Source/Runtime/Engine/Private/Particles/ParticleModules_Color.cpp
--- a/Engine/Source/Runtime/Engine/Private/Particles/ParticleModules_Color.cpp
+++ b/Engine/Source/Runtime/Engine/Private/Particles/ParticleModules_Color.cpp
@@ -19,6 +19,8 @@
 #include "Particles/ParticleSystem.h"
 #include "Particles/ParticleSystemComponent.h"
 #include "Distributions/DistributionVectorConstantCurve.h"
+#include "Engine/InterpCurveEdSetup.h"
+#include "Distributions/DistributionFloatConstantCurve.h"
 
 UParticleModuleColorBase::UParticleModuleColorBase(const FObjectInitializer& ObjectInitializer)
 	: Super(ObjectInitializer)
diff --git a/Engine/Source/Runtime/Engine/Private/Particles/ParticleModules_Location.cpp b/Engine/Source/Runtime/Engine/Private/Particles/ParticleModules_Location.cpp
--- a/Engine/Source/Runtime/Engine/Private/Particles/ParticleModules_Location.cpp
+++ b/Engine/Source/Runtime/Engine/Private/Particles/ParticleModules_Location.cpp
@@ -31,7 +31,7 @@
 #include "Particles/ParticleSpriteEmitter.h"
 #include "Particles/ParticleSystemComponent.h"
 #include "Animation/SkeletalMeshActor.h"
-
+#include "Engine/SkeletalMeshSocket.h"
 UParticleModuleLocationBase::UParticleModuleLocationBase(const FObjectInitializer& ObjectInitializer)
 	: Super(ObjectInitializer)
 {
diff --git a/Engine/Source/Runtime/Engine/Private/Particles/ParticleModules_Spawn.cpp b/Engine/Source/Runtime/Engine/Private/Particles/ParticleModules_Spawn.cpp
--- a/Engine/Source/Runtime/Engine/Private/Particles/ParticleModules_Spawn.cpp
+++ b/Engine/Source/Runtime/Engine/Private/Particles/ParticleModules_Spawn.cpp
@@ -11,6 +11,7 @@
 #include "Particles/Spawn/ParticleModuleSpawnPerUnit.h"
 #include "Particles/ParticleEmitter.h"
 #include "Particles/ParticleSystemComponent.h"
+#include "Distributions/DistributionFloatConstantCurve.h"
 
 UParticleModuleSpawnBase::UParticleModuleSpawnBase(const FObjectInitializer& ObjectInitializer)
 	: Super(ObjectInitializer)
diff --git a/Engine/Source/Runtime/Engine/Private/Particles/ParticleTrail2EmitterInstance.cpp b/Engine/Source/Runtime/Engine/Private/Particles/ParticleTrail2EmitterInstance.cpp
--- a/Engine/Source/Runtime/Engine/Private/Particles/ParticleTrail2EmitterInstance.cpp
+++ b/Engine/Source/Runtime/Engine/Private/Particles/ParticleTrail2EmitterInstance.cpp
@@ -20,6 +20,7 @@
 #include "Particles/ParticleSpriteEmitter.h"
 #include "Particles/ParticleSystemComponent.h"
 #include "Particles/ParticleModuleRequired.h"
+#include "Scalability.h"
 /** trail stats */
 
 
diff --git a/Engine/Source/Runtime/Engine/Private/Pawn.cpp b/Engine/Source/Runtime/Engine/Private/Pawn.cpp
--- a/Engine/Source/Runtime/Engine/Private/Pawn.cpp
+++ b/Engine/Source/Runtime/Engine/Private/Pawn.cpp
@@ -16,6 +16,11 @@
 #include "NetworkingDistanceConstants.h"
 #include "VisualLogger/VisualLogger.h"
 #include "Engine/InputDelegateBinding.h"
+#include "GameFramework/DamageType.h"
+#include "Interfaces/NetworkPredictionInterface.h"
+#include "Kismet/GameplayStatics.h"
+#include "GameFramework/PlayerState.h"
+#include "Components/PawnNoiseEmitterComponent.h"
 
 DEFINE_LOG_CATEGORY(LogDamage);
 DEFINE_LOG_CATEGORY_STATIC(LogPawn, Warning, All);
diff --git a/Engine/Source/Runtime/Engine/Private/PawnNoiseEmitterComponent.cpp b/Engine/Source/Runtime/Engine/Private/PawnNoiseEmitterComponent.cpp
--- a/Engine/Source/Runtime/Engine/Private/PawnNoiseEmitterComponent.cpp
+++ b/Engine/Source/Runtime/Engine/Private/PawnNoiseEmitterComponent.cpp
@@ -2,6 +2,7 @@
 
 
 #include "EnginePrivate.h"
+#include "Components/PawnNoiseEmitterComponent.h"
 
 
 
