Engine private PCH cleanup.

diff --git a/Engine/Source/Runtime/Engine/Private/PendingNetGame.cpp b/Engine/Source/Runtime/Engine/Private/PendingNetGame.cpp
--- a/Engine/Source/Runtime/Engine/Private/PendingNetGame.cpp
+++ b/Engine/Source/Runtime/Engine/Private/PendingNetGame.cpp
@@ -8,6 +8,7 @@
 #include "Net/UnrealNetwork.h"
 #include "Net/NetworkProfiler.h"
 #include "Net/DataChannel.h"
+#include "GeneralProjectSettings.h"
 
 UPendingNetGame::UPendingNetGame( const FObjectInitializer& ObjectInitializer, const FURL& InURL )
 	: Super(ObjectInitializer)
diff --git a/Engine/Source/Runtime/Engine/Private/PhysicsEngine/BodyInstance.cpp b/Engine/Source/Runtime/Engine/Private/PhysicsEngine/BodyInstance.cpp
--- a/Engine/Source/Runtime/Engine/Private/PhysicsEngine/BodyInstance.cpp
+++ b/Engine/Source/Runtime/Engine/Private/PhysicsEngine/BodyInstance.cpp
@@ -19,6 +19,10 @@
 #include "PhysicsEngine/BodySetup2D.h"
 #include "PhysicsEngine/AggregateGeometry2D.h"
 #endif	//WITH_BOX2D
+#include "Components/ModelComponent.h"
+#include "Components/BrushComponent.h"
+#include "PhysicalMaterials/PhysicalMaterial.h"
+#include "PhysicsEngine/PhysicsSettings.h"
 
 ////////////////////////////////////////////////////////////////////////////
 // FCollisionResponse
diff --git a/Engine/Source/Runtime/Engine/Private/PhysicsEngine/BodySetup.cpp b/Engine/Source/Runtime/Engine/Private/PhysicsEngine/BodySetup.cpp
--- a/Engine/Source/Runtime/Engine/Private/PhysicsEngine/BodySetup.cpp
+++ b/Engine/Source/Runtime/Engine/Private/PhysicsEngine/BodySetup.cpp
@@ -15,6 +15,7 @@
 
 
 #include "PhysDerivedData.h"
+#include "PhysicalMaterials/PhysicalMaterial.h"
 
 #if WITH_PHYSX
 	namespace
diff --git a/Engine/Source/Runtime/Engine/Private/PhysicsEngine/PhysAnim.cpp b/Engine/Source/Runtime/Engine/Private/PhysicsEngine/PhysAnim.cpp
--- a/Engine/Source/Runtime/Engine/Private/PhysicsEngine/PhysAnim.cpp
+++ b/Engine/Source/Runtime/Engine/Private/PhysicsEngine/PhysAnim.cpp
@@ -7,10 +7,12 @@
 #include "EnginePrivate.h"
 #include "AnimTree.h"
 #include "SkeletalRenderPublic.h"
-
+#include "Components/LineBatchComponent.h"
 #if WITH_PHYSX
 	#include "PhysXSupport.h"
 #endif // WITH_PHYSX
+#include "PhysicsEngine/PhysicsConstraintTemplate.h"
+#include "PhysicsEngine/PhysicsAsset.h"
 
 
 /** Used for drawing pre-phys skeleton if bShowPrePhysBones is true */
diff --git a/Engine/Source/Runtime/Engine/Private/PhysicsEngine/PhysDrawing.cpp b/Engine/Source/Runtime/Engine/Private/PhysicsEngine/PhysDrawing.cpp
--- a/Engine/Source/Runtime/Engine/Private/PhysicsEngine/PhysDrawing.cpp
+++ b/Engine/Source/Runtime/Engine/Private/PhysicsEngine/PhysDrawing.cpp
@@ -4,6 +4,8 @@
 #include "PhysicsPublic.h"
 #include "PhysXSupport.h"
 #include "DynamicMeshBuilder.h"
+#include "PhysicsEngine/PhysicsConstraintTemplate.h"
+#include "PhysicsEngine/PhysicsAsset.h"
 
 static const int32 DrawCollisionSides = 16;
 static const int32 DrawConeLimitSides = 40;
diff --git a/Engine/Source/Runtime/Engine/Private/PhysicsEngine/PhysLevel.cpp b/Engine/Source/Runtime/Engine/Private/PhysicsEngine/PhysLevel.cpp
--- a/Engine/Source/Runtime/Engine/Private/PhysicsEngine/PhysLevel.cpp
+++ b/Engine/Source/Runtime/Engine/Private/PhysicsEngine/PhysLevel.cpp
@@ -12,6 +12,7 @@
 #if WITH_BOX2D
 	#include "../PhysicsEngine2D/Box2DIntegration.h"
 #endif
+#include "PhysicsEngine/PhysicsSettings.h"
 
 #ifndef APEX_STATICALLY_LINKED
 	#define APEX_STATICALLY_LINKED	0
diff --git a/Engine/Source/Runtime/Engine/Private/PhysicsEngine/PhysScene.cpp b/Engine/Source/Runtime/Engine/Private/PhysicsEngine/PhysScene.cpp
--- a/Engine/Source/Runtime/Engine/Private/PhysicsEngine/PhysScene.cpp
+++ b/Engine/Source/Runtime/Engine/Private/PhysicsEngine/PhysScene.cpp
@@ -16,7 +16,10 @@
 #endif
 
 #include "PhysSubstepTasks.h"	//needed even if not substepping, contains common utility class for PhysX
-
+#include "PhysicsEngine/PhysicsCollisionHandler.h"
+#include "Components/DestructibleComponent.h"
+#include "Components/LineBatchComponent.h"
+#include "PhysicsEngine/PhysicsSettings.h"
 
 #define USE_ADAPTIVE_FORCES_FOR_ASYNC_SCENE			1
 #define USE_SPECIAL_FRICTION_MODEL_FOR_ASYNC_SCENE	0
diff --git a/Engine/Source/Runtime/Engine/Private/PhysicsEngine/PhysXSupport.cpp b/Engine/Source/Runtime/Engine/Private/PhysicsEngine/PhysXSupport.cpp
--- a/Engine/Source/Runtime/Engine/Private/PhysicsEngine/PhysXSupport.cpp
+++ b/Engine/Source/Runtime/Engine/Private/PhysicsEngine/PhysXSupport.cpp
@@ -10,6 +10,8 @@
 #if WITH_PHYSX
 
 #include "PhysXSupport.h"
+#include "Components/DestructibleComponent.h"
+#include "PhysicsEngine/PhysicsAsset.h"
 
 PxFoundation*			GPhysXFoundation = NULL;
 PxProfileZoneManager*	GPhysXProfileZoneManager = NULL;
diff --git a/Engine/Source/Runtime/Engine/Private/PhysicsEngine/PhysicalMaterial.cpp b/Engine/Source/Runtime/Engine/Private/PhysicsEngine/PhysicalMaterial.cpp
--- a/Engine/Source/Runtime/Engine/Private/PhysicsEngine/PhysicalMaterial.cpp
+++ b/Engine/Source/Runtime/Engine/Private/PhysicsEngine/PhysicalMaterial.cpp
@@ -12,6 +12,7 @@
 	#include "PhysXSupport.h"
 	#include "../Vehicles/PhysXVehicleManager.h"
 #endif // WITH_PHYSX
+#include "PhysicalMaterials/PhysicalMaterial.h"
 
 UDEPRECATED_PhysicalMaterialPropertyBase::UDEPRECATED_PhysicalMaterialPropertyBase(const FObjectInitializer& ObjectInitializer)
 	: Super(ObjectInitializer)
diff --git a/Engine/Source/Runtime/Engine/Private/PhysicsEngine/PhysicsAsset.cpp b/Engine/Source/Runtime/Engine/Private/PhysicsEngine/PhysicsAsset.cpp
--- a/Engine/Source/Runtime/Engine/Private/PhysicsEngine/PhysicsAsset.cpp
+++ b/Engine/Source/Runtime/Engine/Private/PhysicsEngine/PhysicsAsset.cpp
@@ -5,11 +5,13 @@
 =============================================================================*/ 
 
 #include "EnginePrivate.h"
+#include "PhysicsEngine/PhysicsConstraintTemplate.h"
 
 
 #if WITH_PHYSX
 	#include "PhysXSupport.h"
 #endif // WITH_PHYSX
+#include "PhysicsEngine/PhysicsAsset.h"
 
 
 ///////////////////////////////////////	
diff --git a/Engine/Source/Runtime/Engine/Private/PhysicsEngine/PhysicsCollisionHandler.cpp b/Engine/Source/Runtime/Engine/Private/PhysicsEngine/PhysicsCollisionHandler.cpp
--- a/Engine/Source/Runtime/Engine/Private/PhysicsEngine/PhysicsCollisionHandler.cpp
+++ b/Engine/Source/Runtime/Engine/Private/PhysicsEngine/PhysicsCollisionHandler.cpp
@@ -2,6 +2,8 @@
 
 #include "EnginePrivate.h"
 #include "PhysicsPublic.h"
+#include "PhysicsEngine/PhysicsCollisionHandler.h"
+#include "Kismet/GameplayStatics.h"
 
 UPhysicsCollisionHandler::UPhysicsCollisionHandler(const FObjectInitializer& ObjectInitializer)
 	: Super(ObjectInitializer)
diff --git a/Engine/Source/Runtime/Engine/Private/PhysicsEngine/PhysicsConstraintActor.cpp b/Engine/Source/Runtime/Engine/Private/PhysicsEngine/PhysicsConstraintActor.cpp
--- a/Engine/Source/Runtime/Engine/Private/PhysicsEngine/PhysicsConstraintActor.cpp
+++ b/Engine/Source/Runtime/Engine/Private/PhysicsEngine/PhysicsConstraintActor.cpp
@@ -2,6 +2,7 @@
 
 #include "EnginePrivate.h"
 #include "PhysXSupport.h"
+#include "PhysicsEngine/PhysicsConstraintActor.h"
 
 APhysicsConstraintActor::APhysicsConstraintActor(const FObjectInitializer& ObjectInitializer)
 	: Super(ObjectInitializer)
diff --git a/Engine/Source/Runtime/Engine/Private/PhysicsEngine/PhysicsSettings.cpp b/Engine/Source/Runtime/Engine/Private/PhysicsEngine/PhysicsSettings.cpp
--- a/Engine/Source/Runtime/Engine/Private/PhysicsEngine/PhysicsSettings.cpp
+++ b/Engine/Source/Runtime/Engine/Private/PhysicsEngine/PhysicsSettings.cpp
@@ -2,6 +2,8 @@
 
 #include "EnginePrivate.h"
 #include "GameFramework/MovementComponent.h"
+#include "PhysicsEngine/PhysicsSettings.h"
+#include "PhysicalMaterials/PhysicalMaterial.h"
 
 UPhysicsSettings::UPhysicsSettings(const FObjectInitializer& ObjectInitializer)
 	: Super(ObjectInitializer)
diff --git a/Engine/Source/Runtime/Engine/Private/PhysicsEngine/RadialForceComponent.cpp b/Engine/Source/Runtime/Engine/Private/PhysicsEngine/RadialForceComponent.cpp
--- a/Engine/Source/Runtime/Engine/Private/PhysicsEngine/RadialForceComponent.cpp
+++ b/Engine/Source/Runtime/Engine/Private/PhysicsEngine/RadialForceComponent.cpp
@@ -5,6 +5,7 @@
 #include "Net/UnrealNetwork.h"
 #include "PhysicsEngine/RadialForceComponent.h"
 #include "PhysicsEngine/RadialForceActor.h"
+#include "Components/DestructibleComponent.h"
 
 //////////////////////////////////////////////////////////////////////////
 // RADIALFORCECOMPONENT
diff --git a/Engine/Source/Runtime/Engine/Private/PhysicsVolume.cpp b/Engine/Source/Runtime/Engine/Private/PhysicsVolume.cpp
--- a/Engine/Source/Runtime/Engine/Private/PhysicsVolume.cpp
+++ b/Engine/Source/Runtime/Engine/Private/PhysicsVolume.cpp
@@ -1,7 +1,9 @@
 // Copyright 1998-2014 Epic Games, Inc. All Rights Reserved.
 
 #include "EnginePrivate.h"
+#include "Components/BrushComponent.h"
 #include "GameFramework/PhysicsVolume.h"
+#include "PhysicsEngine/PhysicsSettings.h"
 
 APhysicsVolume::APhysicsVolume(const FObjectInitializer& ObjectInitializer)
 	: Super(ObjectInitializer)
diff --git a/Engine/Source/Runtime/Engine/Private/Player.cpp b/Engine/Source/Runtime/Engine/Private/Player.cpp
--- a/Engine/Source/Runtime/Engine/Private/Player.cpp
+++ b/Engine/Source/Runtime/Engine/Private/Player.cpp
@@ -9,11 +9,15 @@
 #include "SubtitleManager.h"
 #include "Net/UnrealNetwork.h"
 #include "GameFramework/HUD.h"
+#include "GameFramework/PlayerInput.h"
 #include "Engine/GameInstance.h"
 
 #include "RenderCore.h"
 #include "ColorList.h"
 #include "SlateBasics.h"
+#include "GameFramework/CheatManager.h"
+#include "GameFramework/Pawn.h"
+#include "GameFramework/GameState.h"
 
 //////////////////////////////////////////////////////////////////////////
 // UPlayer
diff --git a/Engine/Source/Runtime/Engine/Private/PlayerCameraManager.cpp b/Engine/Source/Runtime/Engine/Private/PlayerCameraManager.cpp
--- a/Engine/Source/Runtime/Engine/Private/PlayerCameraManager.cpp
+++ b/Engine/Source/Runtime/Engine/Private/PlayerCameraManager.cpp
@@ -9,6 +9,9 @@
 #include "Camera/CameraActor.h"
 #include "Camera/CameraAnim.h"
 #include "Camera/CameraAnimInst.h"
+#include "Camera/CameraComponent.h"
+#include "Camera/CameraModifier_CameraShake.h"
+#include "GameFramework/PlayerState.h"
 
 DEFINE_LOG_CATEGORY_STATIC(LogPlayerCameraManager, Log, All);
 
diff --git a/Engine/Source/Runtime/Engine/Private/PlayerController.cpp b/Engine/Source/Runtime/Engine/Private/PlayerController.cpp
--- a/Engine/Source/Runtime/Engine/Private/PlayerController.cpp
+++ b/Engine/Source/Runtime/Engine/Private/PlayerController.cpp
@@ -28,6 +28,13 @@
 #include "GenericPlatform/IForceFeedbackSystem.h"
 #include "Engine/InputDelegateBinding.h"
 #include "SVirtualJoystick.h"
+#include "GameFramework/LocalMessage.h"
+#include "GameFramework/PlayerStart.h"
+#include "GameFramework/CheatManager.h"
+#include "GameFramework/InputSettings.h"
+#include "GameFramework/PlayerState.h"
+#include "GameFramework/GameState.h"
+#include "Engine/ChildConnection.h"
 
 DEFINE_LOG_CATEGORY(LogPlayerController);
 
diff --git a/Engine/Source/Runtime/Engine/Private/PlayerMuteList.cpp b/Engine/Source/Runtime/Engine/Private/PlayerMuteList.cpp
--- a/Engine/Source/Runtime/Engine/Private/PlayerMuteList.cpp
+++ b/Engine/Source/Runtime/Engine/Private/PlayerMuteList.cpp
@@ -2,6 +2,8 @@
 
 #include "EnginePrivate.h"
 #include "OnlineSubsystemUtils.h"
+#include "GameFramework/GameSession.h"
+#include "GameFramework/PlayerState.h"
 
 static inline void AddIdToMuteList(TArray< TSharedRef<class FUniqueNetId> >& MuteList, const TSharedPtr<FUniqueNetId>& UniqueIdToAdd)
 {
diff --git a/Engine/Source/Runtime/Engine/Private/PlayerStart.cpp b/Engine/Source/Runtime/Engine/Private/PlayerStart.cpp
--- a/Engine/Source/Runtime/Engine/Private/PlayerStart.cpp
+++ b/Engine/Source/Runtime/Engine/Private/PlayerStart.cpp
@@ -1,6 +1,8 @@
 // Copyright 1998-2014 Epic Games, Inc. All Rights Reserved.
 
 #include "EnginePrivate.h"
+#include "Components/CapsuleComponent.h"
+#include "GameFramework/PlayerStart.h"
 
 APlayerStart::APlayerStart(const FObjectInitializer& ObjectInitializer)
 	: Super(ObjectInitializer)
diff --git a/Engine/Source/Runtime/Engine/Private/PlayerState.cpp b/Engine/Source/Runtime/Engine/Private/PlayerState.cpp
--- a/Engine/Source/Runtime/Engine/Private/PlayerState.cpp
+++ b/Engine/Source/Runtime/Engine/Private/PlayerState.cpp
@@ -8,6 +8,8 @@
 #include "GameFramework/EngineMessage.h"
 #include "Net/UnrealNetwork.h"
 #include "OnlineSubsystemUtils.h"
+#include "GameFramework/GameState.h"
+#include "GameFramework/PlayerState.h"
 
 APlayerState::APlayerState(const FObjectInitializer& ObjectInitializer)
 	: Super(ObjectInitializer
diff --git a/Engine/Source/Runtime/Engine/Private/PointLightSceneProxy.h b/Engine/Source/Runtime/Engine/Private/PointLightSceneProxy.h
--- a/Engine/Source/Runtime/Engine/Private/PointLightSceneProxy.h
+++ b/Engine/Source/Runtime/Engine/Private/PointLightSceneProxy.h
@@ -6,6 +6,7 @@
 
 #ifndef __PointLightSceneProxy_H__
 #define __PointLightSceneProxy_H__
+#include "Components/PointLightComponent.h"
 
 /**
  * Compute the screen bounds of a point light along one axis.
diff --git a/Engine/Source/Runtime/Engine/Private/PreviewScene.cpp b/Engine/Source/Runtime/Engine/Private/PreviewScene.cpp
--- a/Engine/Source/Runtime/Engine/Private/PreviewScene.cpp
+++ b/Engine/Source/Runtime/Engine/Private/PreviewScene.cpp
@@ -7,6 +7,8 @@
 #include "EnginePrivate.h"
 #include "SoundDefinitions.h"
 #include "PreviewScene.h"
+#include "Components/DirectionalLightComponent.h"
+#include "Components/LineBatchComponent.h"
 
 FPreviewScene::FPreviewScene(FPreviewScene::ConstructionValues CVS)
 	: PreviewWorld(NULL)
diff --git a/Engine/Source/Runtime/Engine/Private/PrimitiveComponent.cpp b/Engine/Source/Runtime/Engine/Private/PrimitiveComponent.cpp
--- a/Engine/Source/Runtime/Engine/Private/PrimitiveComponent.cpp
+++ b/Engine/Source/Runtime/Engine/Private/PrimitiveComponent.cpp
@@ -25,6 +25,8 @@
 #include "UObjectToken.h"
 #include "MapErrors.h"
 #include "CollisionDebugDrawingPublic.h"
+#include "GameFramework/CheatManager.h"
+#include "GameFramework/DamageType.h"
 
 #define LOCTEXT_NAMESPACE "PrimitiveComponent"
 
diff --git a/Engine/Source/Runtime/Engine/Private/ReflectionCaptureComponent.cpp b/Engine/Source/Runtime/Engine/Private/ReflectionCaptureComponent.cpp
--- a/Engine/Source/Runtime/Engine/Private/ReflectionCaptureComponent.cpp
+++ b/Engine/Source/Runtime/Engine/Private/ReflectionCaptureComponent.cpp
@@ -18,6 +18,9 @@
 #include "Engine/PlaneReflectionCapture.h"
 #include "Engine/BoxReflectionCapture.h"
 #include "Components/PlaneReflectionCaptureComponent.h"
+#include "Components/ReflectionCaptureComponent.h"
+#include "Components/BoxComponent.h"
+#include "Components/SkyLightComponent.h"
 
 /** 
  * Size of all reflection captures.
diff --git a/Engine/Source/Runtime/Engine/Private/RigidBodyBase.cpp b/Engine/Source/Runtime/Engine/Private/RigidBodyBase.cpp
--- a/Engine/Source/Runtime/Engine/Private/RigidBodyBase.cpp
+++ b/Engine/Source/Runtime/Engine/Private/RigidBodyBase.cpp
@@ -1,6 +1,7 @@
 // Copyright 1998-2014 Epic Games, Inc. All Rights Reserved.
 
 #include "EnginePrivate.h"
+#include "PhysicsEngine/RigidBodyBase.h"
 
 ARigidBodyBase::ARigidBodyBase(const FObjectInitializer& ObjectInitializer)
 	: Super(ObjectInitializer)
diff --git a/Engine/Source/Runtime/Engine/Private/SceneCaptureComponent.cpp b/Engine/Source/Runtime/Engine/Private/SceneCaptureComponent.cpp
--- a/Engine/Source/Runtime/Engine/Private/SceneCaptureComponent.cpp
+++ b/Engine/Source/Runtime/Engine/Private/SceneCaptureComponent.cpp
@@ -11,6 +11,7 @@
 #include "Components/SceneCaptureComponent2D.h"
 #include "Engine/SceneCaptureCube.h"
 #include "Components/SceneCaptureComponentCube.h"
+#include "Components/DrawFrustumComponent.h"
 
 ASceneCapture::ASceneCapture(const FObjectInitializer& ObjectInitializer)
 	: Super(ObjectInitializer)
diff --git a/Engine/Source/Runtime/Engine/Private/SceneView.cpp b/Engine/Source/Runtime/Engine/Private/SceneView.cpp
--- a/Engine/Source/Runtime/Engine/Private/SceneView.cpp
+++ b/Engine/Source/Runtime/Engine/Private/SceneView.cpp
@@ -14,6 +14,7 @@
 #include "RendererInterface.h"
 #include "BufferVisualizationData.h"
 #include "Interfaces/Interface_PostProcessVolume.h"
+#include "Engine/TextureCube.h"
 
 DEFINE_LOG_CATEGORY(LogBufferVisualization);
 
diff --git a/Engine/Source/Runtime/Engine/Private/ScriptPlatformInterface.cpp b/Engine/Source/Runtime/Engine/Private/ScriptPlatformInterface.cpp
--- a/Engine/Source/Runtime/Engine/Private/ScriptPlatformInterface.cpp
+++ b/Engine/Source/Runtime/Engine/Private/ScriptPlatformInterface.cpp
@@ -9,6 +9,7 @@
 #include "Engine/MicroTransactionBase.h"
 #include "Engine/PlatformInterfaceWebResponse.h"
 #include "Engine/TwitterIntegrationBase.h"
+#include "Engine/InGameAdManager.h"
 
 DEFINE_LOG_CATEGORY_STATIC(LogScriptPlatformInterface, Log, All);
 
diff --git a/Engine/Source/Runtime/Engine/Private/ShadowMap.cpp b/Engine/Source/Runtime/Engine/Private/ShadowMap.cpp
--- a/Engine/Source/Runtime/Engine/Private/ShadowMap.cpp
+++ b/Engine/Source/Runtime/Engine/Private/ShadowMap.cpp
@@ -6,6 +6,9 @@
 #include "TargetPlatform.h"
 #include "LightMap.h"
 #include "ShadowMap.h"
+#include "Engine/ShadowMapTexture2D.h"
+#include "Components/InstancedStaticMeshComponent.h"
+#include "Components/LightComponent.h"
 
 #if WITH_EDITOR
 	// NOTE: We're only counting the top-level mip-map for the following variables.
diff --git a/Engine/Source/Runtime/Engine/Private/SkeletalMesh.cpp b/Engine/Source/Runtime/Engine/Private/SkeletalMesh.cpp
--- a/Engine/Source/Runtime/Engine/Private/SkeletalMesh.cpp
+++ b/Engine/Source/Runtime/Engine/Private/SkeletalMesh.cpp
@@ -25,6 +25,7 @@
 #include "Animation/AnimBlueprint.h"
 #include "Animation/AnimBlueprintGeneratedClass.h"
 #include "ComponentReregisterContext.h"
+#include "PhysicsEngine/PhysicsAsset.h"
 
 #if WITH_EDITOR
 #include "MeshUtilities.h"
@@ -37,6 +38,9 @@
 #include "PhysicsEngine/PhysXSupport.h"
 #endif// #if WITH_APEX
 
+#include "EditorFramework/AssetImportData.h"
+#include "Engine/SkeletalMeshSocket.h"
+#include "Components/BrushComponent.h"
 #define LOCTEXT_NAMESPACE "SkeltalMesh"
 
 DEFINE_LOG_CATEGORY(LogSkeletalMesh);
diff --git a/Engine/Source/Runtime/Engine/Private/SkeletalMeshComponent.cpp b/Engine/Source/Runtime/Engine/Private/SkeletalMeshComponent.cpp
--- a/Engine/Source/Runtime/Engine/Private/SkeletalMeshComponent.cpp
+++ b/Engine/Source/Runtime/Engine/Private/SkeletalMeshComponent.cpp
@@ -23,6 +23,8 @@
 #include "Collision.h"
 #include "ConvexVolume.h"
 #endif
+#include "PhysicsEngine/PhysicsSettings.h"
+#include "PhysicsEngine/PhysicsAsset.h"
 
 TAutoConsoleVariable<int32> CVarUseParallelAnimationEvaluation(TEXT("a.ParallelAnimEvaluation"), 0, TEXT("If 1, animation evaluation will be run across the task graph system. If 0, evaluation will run purely on the game thread"));
 
diff --git a/Engine/Source/Runtime/Engine/Private/SkeletalMeshMerge.cpp b/Engine/Source/Runtime/Engine/Private/SkeletalMeshMerge.cpp
--- a/Engine/Source/Runtime/Engine/Private/SkeletalMeshMerge.cpp
+++ b/Engine/Source/Runtime/Engine/Private/SkeletalMeshMerge.cpp
@@ -7,6 +7,7 @@
 #include "EnginePrivate.h"
 #include "RawIndexBuffer.h"
 #include "SkeletalMeshMerge.h"
+#include "Engine/SkeletalMeshSocket.h"
 
 /*-----------------------------------------------------------------------------
 	FSkeletalMeshMerge
diff --git a/Engine/Source/Runtime/Engine/Private/SkinnedMeshComponent.cpp b/Engine/Source/Runtime/Engine/Private/SkinnedMeshComponent.cpp
--- a/Engine/Source/Runtime/Engine/Private/SkinnedMeshComponent.cpp
+++ b/Engine/Source/Runtime/Engine/Private/SkinnedMeshComponent.cpp
@@ -13,6 +13,8 @@
 #include "AnimTree.h"
 #include "Animation/VertexAnim/MorphTarget.h"
 #include "ComponentReregisterContext.h"
+#include "Engine/SkeletalMeshSocket.h"
+#include "PhysicsEngine/PhysicsAsset.h"
 
 DEFINE_LOG_CATEGORY_STATIC(LogSkinnedMeshComp, Log, All);
 
diff --git a/Engine/Source/Runtime/Engine/Private/SkyLightComponent.cpp b/Engine/Source/Runtime/Engine/Private/SkyLightComponent.cpp
--- a/Engine/Source/Runtime/Engine/Private/SkyLightComponent.cpp
+++ b/Engine/Source/Runtime/Engine/Private/SkyLightComponent.cpp
@@ -15,6 +15,7 @@
 #include "MapErrors.h"
 #include "ComponentInstanceDataCache.h"
 #include "ShaderCompiler.h"
+#include "Components/SkyLightComponent.h"
 
 #define LOCTEXT_NAMESPACE "SkyLightComponent"
 
diff --git a/Engine/Source/Runtime/Engine/Private/Slate/SlateGameResources.cpp b/Engine/Source/Runtime/Engine/Private/Slate/SlateGameResources.cpp
--- a/Engine/Source/Runtime/Engine/Private/Slate/SlateGameResources.cpp
+++ b/Engine/Source/Runtime/Engine/Private/Slate/SlateGameResources.cpp
@@ -6,6 +6,8 @@
 #include "AssetRegistryModule.h"
 #include "SlateGameResources.h"
 #include "MessageLog.h"
+#include "Curves/CurveVector.h"
+#include "Curves/CurveLinearColor.h"
 
 TSharedRef<FSlateGameResources> FSlateGameResources::New( const FName& InStyleSetName, const FString& ScopeToDirectory, const FString& InBasePath )
 {
diff --git a/Engine/Source/Runtime/Engine/Private/SoundNodeEnveloper.cpp b/Engine/Source/Runtime/Engine/Private/SoundNodeEnveloper.cpp
--- a/Engine/Source/Runtime/Engine/Private/SoundNodeEnveloper.cpp
+++ b/Engine/Source/Runtime/Engine/Private/SoundNodeEnveloper.cpp
@@ -4,6 +4,7 @@
 #include "EnginePrivate.h"
 #include "SoundDefinitions.h"
 #include "Sound/SoundNodeEnveloper.h"
+#include "Distributions/DistributionFloatConstantCurve.h"
 
 /*-----------------------------------------------------------------------------
 USoundNodeEnveloper implementation.
diff --git a/Engine/Source/Runtime/Engine/Private/SpectatorPawn.cpp b/Engine/Source/Runtime/Engine/Private/SpectatorPawn.cpp
--- a/Engine/Source/Runtime/Engine/Private/SpectatorPawn.cpp
+++ b/Engine/Source/Runtime/Engine/Private/SpectatorPawn.cpp
@@ -4,6 +4,7 @@
 #include "EnginePrivate.h"
 #include "GameFramework/SpectatorPawn.h"
 #include "GameFramework/SpectatorPawnMovement.h"
+#include "Components/SphereComponent.h"
 
 ASpectatorPawn::ASpectatorPawn(const FObjectInitializer& ObjectInitializer)
 	: Super(ObjectInitializer
diff --git a/Engine/Source/Runtime/Engine/Private/SphereComponent.cpp b/Engine/Source/Runtime/Engine/Private/SphereComponent.cpp
--- a/Engine/Source/Runtime/Engine/Private/SphereComponent.cpp
+++ b/Engine/Source/Runtime/Engine/Private/SphereComponent.cpp
@@ -2,6 +2,7 @@
 
 
 #include "EnginePrivate.h"
+#include "Components/SphereComponent.h"
 
 USphereComponent::USphereComponent(const FObjectInitializer& ObjectInitializer)
 	: Super(ObjectInitializer)
diff --git a/Engine/Source/Runtime/Engine/Private/SplineMeshComponent.cpp b/Engine/Source/Runtime/Engine/Private/SplineMeshComponent.cpp
--- a/Engine/Source/Runtime/Engine/Private/SplineMeshComponent.cpp
+++ b/Engine/Source/Runtime/Engine/Private/SplineMeshComponent.cpp
@@ -4,6 +4,8 @@
 #include "Components/SplineMeshComponent.h"
 #include "SplineMeshSceneProxy.h"
 #include "ShaderParameterUtils.h"
+#include "NavigationSystemHelpers.h"
+#include "AI/Navigation/NavCollision.h"
 
 
 //////////////////////////////////////////////////////////////////////////
diff --git a/Engine/Source/Runtime/Engine/Private/SpotLight.cpp b/Engine/Source/Runtime/Engine/Private/SpotLight.cpp
--- a/Engine/Source/Runtime/Engine/Private/SpotLight.cpp
+++ b/Engine/Source/Runtime/Engine/Private/SpotLight.cpp
@@ -2,6 +2,7 @@
 
 #include "EnginePrivate.h"
 #include "Engine/SpotLight.h"
+#include "Components/SpotLightComponent.h"
 
 ASpotLight::ASpotLight(const FObjectInitializer& ObjectInitializer)
 	: Super(ObjectInitializer.SetDefaultSubobjectClass<USpotLightComponent>(TEXT("LightComponent0")))
diff --git a/Engine/Source/Runtime/Engine/Private/SpotLightComponent.cpp b/Engine/Source/Runtime/Engine/Private/SpotLightComponent.cpp
--- a/Engine/Source/Runtime/Engine/Private/SpotLightComponent.cpp
+++ b/Engine/Source/Runtime/Engine/Private/SpotLightComponent.cpp
@@ -6,6 +6,7 @@
 
 #include "EnginePrivate.h"
 #include "PointLightSceneProxy.h"
+#include "Components/SpotLightComponent.h"
 
 
 /**
diff --git a/Engine/Source/Runtime/Engine/Private/StaticMesh.cpp b/Engine/Source/Runtime/Engine/Private/StaticMesh.cpp
--- a/Engine/Source/Runtime/Engine/Private/StaticMesh.cpp
+++ b/Engine/Source/Runtime/Engine/Private/StaticMesh.cpp
@@ -21,6 +21,9 @@
 #include "DerivedDataCacheInterface.h"
 #include "UObjectAnnotation.h"
 #endif // #if WITH_EDITOR
+#include "Engine/StaticMeshSocket.h"
+#include "EditorFramework/AssetImportData.h"
+#include "AI/Navigation/NavCollision.h"
 
 DEFINE_LOG_CATEGORY(LogStaticMesh);	
 
diff --git a/Engine/Source/Runtime/Engine/Private/StaticMeshComponent.cpp b/Engine/Source/Runtime/Engine/Private/StaticMeshComponent.cpp
--- a/Engine/Source/Runtime/Engine/Private/StaticMeshComponent.cpp
+++ b/Engine/Source/Runtime/Engine/Private/StaticMeshComponent.cpp
@@ -16,6 +16,10 @@
 #include "LightMap.h"
 #include "ShadowMap.h"
 #include "ComponentReregisterContext.h"
+#include "Engine/ShadowMapTexture2D.h"
+#include "AI/Navigation/NavCollision.h"
+#include "Engine/StaticMeshSocket.h"
+#include "NavigationSystemHelpers.h"
 
 #define LOCTEXT_NAMESPACE "StaticMeshComponent"
 
diff --git a/Engine/Source/Runtime/Engine/Private/StaticMeshRender.cpp b/Engine/Source/Runtime/Engine/Private/StaticMeshRender.cpp
--- a/Engine/Source/Runtime/Engine/Private/StaticMeshRender.cpp
+++ b/Engine/Source/Runtime/Engine/Private/StaticMeshRender.cpp
@@ -12,6 +12,8 @@
 #include "ShadowMap.h"
 #include "DistanceFieldAtlas.h"
 #include "ComponentReregisterContext.h"
+#include "Components/BrushComponent.h"
+#include "AI/Navigation/NavCollision.h"
 
 /** If true, optimized depth-only index buffers are used for shadow rendering. */
 static bool GUseShadowIndexBuffer = true;
diff --git a/Engine/Source/Runtime/Engine/Private/TextureRenderTarget.cpp b/Engine/Source/Runtime/Engine/Private/TextureRenderTarget.cpp
--- a/Engine/Source/Runtime/Engine/Private/TextureRenderTarget.cpp
+++ b/Engine/Source/Runtime/Engine/Private/TextureRenderTarget.cpp
@@ -5,6 +5,7 @@
 =============================================================================*/
 
 #include "EnginePrivate.h"
+#include "Engine/TextureRenderTarget.h"
 
 /*-----------------------------------------------------------------------------
 UTextureRenderTarget
diff --git a/Engine/Source/Runtime/Engine/Private/Timeline.cpp b/Engine/Source/Runtime/Engine/Private/Timeline.cpp
--- a/Engine/Source/Runtime/Engine/Private/Timeline.cpp
+++ b/Engine/Source/Runtime/Engine/Private/Timeline.cpp
@@ -5,7 +5,10 @@
 =============================================================================*/
 
 #include "EnginePrivate.h"
+#include "Curves/CurveLinearColor.h"
+#include "Curves/CurveVector.h"
 #include "Net/UnrealNetwork.h"
+#include "Components/TimelineComponent.h"
 
 DEFINE_LOG_CATEGORY_STATIC(LogTimeline, Log, All);
 
diff --git a/Engine/Source/Runtime/Engine/Private/TimelineTemplate.cpp b/Engine/Source/Runtime/Engine/Private/TimelineTemplate.cpp
--- a/Engine/Source/Runtime/Engine/Private/TimelineTemplate.cpp
+++ b/Engine/Source/Runtime/Engine/Private/TimelineTemplate.cpp
@@ -1,8 +1,11 @@
 // Copyright 1998-2014 Epic Games, Inc. All Rights Reserved.
 
 #include "EnginePrivate.h"
+#include "Curves/CurveLinearColor.h"
+#include "Curves/CurveVector.h"
 #include "BlueprintUtilities.h"
 #include "LatentActions.h"
+#include "Engine/TimelineTemplate.h"
 
 namespace
 {
diff --git a/Engine/Source/Runtime/Engine/Private/TriggerActors.cpp b/Engine/Source/Runtime/Engine/Private/TriggerActors.cpp
--- a/Engine/Source/Runtime/Engine/Private/TriggerActors.cpp
+++ b/Engine/Source/Runtime/Engine/Private/TriggerActors.cpp
@@ -9,6 +9,9 @@
 #include "Engine/TriggerCapsule.h"
 #include "Engine/TriggerSphere.h"
 #include "Net/UnrealNetwork.h"
+#include "Components/BoxComponent.h"
+#include "Components/SphereComponent.h"
+#include "Components/CapsuleComponent.h"
 namespace
 {
 	static const FColor TriggerBaseColor(100, 255, 100, 255);
diff --git a/Engine/Source/Runtime/Engine/Private/UnrealEngine.cpp b/Engine/Source/Runtime/Engine/Private/UnrealEngine.cpp
--- a/Engine/Source/Runtime/Engine/Private/UnrealEngine.cpp
+++ b/Engine/Source/Runtime/Engine/Private/UnrealEngine.cpp
@@ -88,6 +88,13 @@
 #include "STestSuite.h"
 #include "Engine/DemoNetDriver.h"
 #include "SThrobber.h"
+#include "Engine/TextureCube.h"
+#include "AI/Navigation/AvoidanceManager.h"
+#include "Engine/GameEngine.h"
+#include "PhysicsEngine/PhysicsCollisionHandler.h"
+#include "Components/BrushComponent.h"
+#include "GameFramework/GameUserSettings.h"
+#include "PhysicalMaterials/PhysicalMaterial.h"
 
 DEFINE_LOG_CATEGORY_STATIC(LogEngine, Log, All);
 
diff --git a/Engine/Source/Runtime/Engine/Private/UnrealExporter.cpp b/Engine/Source/Runtime/Engine/Private/UnrealExporter.cpp
--- a/Engine/Source/Runtime/Engine/Private/UnrealExporter.cpp
+++ b/Engine/Source/Runtime/Engine/Private/UnrealExporter.cpp
@@ -7,6 +7,7 @@
 // Engine includes.
 #include "EnginePrivate.h"
 #include "Model.h"
+#include "UnrealExporter.h"
 
 DEFINE_LOG_CATEGORY_STATIC(LogExporter, Log, All);
 
diff --git a/Engine/Source/Runtime/Engine/Private/UserInterface/Console.cpp b/Engine/Source/Runtime/Engine/Private/UserInterface/Console.cpp
--- a/Engine/Source/Runtime/Engine/Private/UserInterface/Console.cpp
+++ b/Engine/Source/Runtime/Engine/Private/UserInterface/Console.cpp
@@ -9,6 +9,8 @@
 #include "Engine/LevelScriptActor.h"
 #include "SlateBasics.h"
 #include "DefaultValueHelper.h"
+#include "ConsoleSettings.h"
+#include "GameFramework/InputSettings.h"
 
 static const uint32 MAX_AUTOCOMPLETION_LINES = 20;
 
diff --git a/Engine/Source/Runtime/Engine/Private/UserInterface/PlayerInput.cpp b/Engine/Source/Runtime/Engine/Private/UserInterface/PlayerInput.cpp
--- a/Engine/Source/Runtime/Engine/Private/UserInterface/PlayerInput.cpp
+++ b/Engine/Source/Runtime/Engine/Private/UserInterface/PlayerInput.cpp
@@ -5,6 +5,8 @@
 =============================================================================*/
 
 #include "EnginePrivate.h"
+#include "GameFramework/PlayerInput.h"
+#include "GameFramework/InputSettings.h"
 
 DEFINE_LOG_CATEGORY(LogInput);
 
diff --git a/Engine/Source/Runtime/Engine/Private/VectorFieldVolume.cpp b/Engine/Source/Runtime/Engine/Private/VectorFieldVolume.cpp
--- a/Engine/Source/Runtime/Engine/Private/VectorFieldVolume.cpp
+++ b/Engine/Source/Runtime/Engine/Private/VectorFieldVolume.cpp
@@ -2,6 +2,8 @@
 
 #include "EnginePrivate.h"
 #include "VectorField/VectorFieldVolume.h"
+#include "Components/VectorFieldComponent.h"
+#include "Components/BillboardComponent.h"
 
 AVectorFieldVolume::AVectorFieldVolume(const FObjectInitializer& ObjectInitializer)
 	: Super(ObjectInitializer)
diff --git a/Engine/Source/Runtime/Engine/Private/Vehicles/WheeledVehicleMovementComponent.cpp b/Engine/Source/Runtime/Engine/Private/Vehicles/WheeledVehicleMovementComponent.cpp
--- a/Engine/Source/Runtime/Engine/Private/Vehicles/WheeledVehicleMovementComponent.cpp
+++ b/Engine/Source/Runtime/Engine/Private/Vehicles/WheeledVehicleMovementComponent.cpp
@@ -14,6 +14,8 @@
 #include "../Collision/PhysXCollision.h"
 #include "PhysXVehicleManager.h"
 #endif //WITH_PHYSX
+#include "AI/Navigation/AvoidanceManager.h"
+#include "PhysicalMaterials/PhysicalMaterial.h"
 
 #define LOCTEXT_NAMESPACE "UWheeledVehicleMovementComponent"
 
diff --git a/Engine/Source/Runtime/Engine/Private/World.cpp b/Engine/Source/Runtime/Engine/Private/World.cpp
--- a/Engine/Source/Runtime/Engine/Private/World.cpp
+++ b/Engine/Source/Runtime/Engine/Private/World.cpp
@@ -37,7 +37,10 @@
 #include "Camera/CameraActor.h"
 #include "Engine/DemoNetDriver.h"
 #include "Layers/Layer.h"
+#include "GameFramework/GameState.h"
+#include "GameFramework/PlayerState.h"
 
+#include "Materials/MaterialParameterCollectionInstance.h"
 #if WITH_EDITOR
 	#include "Editor/UnrealEd/Public/Kismet2/KismetEditorUtilities.h"
 	#include "Editor/UnrealEd/Public/Kismet2/BlueprintEditorUtils.h"
@@ -55,6 +58,19 @@
 #include "DataChannel.h"
 #include "ShaderCompiler.h"
 #include "Engine/LevelStreamingPersistent.h"
+#include "AI/Navigation/AvoidanceManager.h"
+#include "GeneralProjectSettings.h"
+#include "PhysicsEngine/PhysicsConstraintActor.h"
+#include "PhysicsEngine/PhysicsCollisionHandler.h"
+#include "Engine/GameEngine.h"
+#include "Engine/ShadowMapTexture2D.h"
+#include "Components/BrushComponent.h"
+#include "Components/LineBatchComponent.h"
+#include "Components/ModelComponent.h"
+#include "Materials/MaterialParameterCollection.h"
+#include "PhysicsEngine/PhysicsSettings.h"
+#include "Engine/Polys.h"
+#include "Engine/LightMapTexture2D.h"
 
 DEFINE_LOG_CATEGORY_STATIC(LogWorld, Log, All);
 DEFINE_LOG_CATEGORY(LogSpawn);
diff --git a/Engine/Source/Runtime/Engine/Private/WorldComposition.cpp b/Engine/Source/Runtime/Engine/Private/WorldComposition.cpp
--- a/Engine/Source/Runtime/Engine/Private/WorldComposition.cpp
+++ b/Engine/Source/Runtime/Engine/Private/WorldComposition.cpp
@@ -6,6 +6,7 @@
 #include "EnginePrivate.h"
 #include "Engine/WorldComposition.h"
 #include "LevelUtils.h"
+#include "Engine/LevelStreamingKismet.h"
 
 DEFINE_LOG_CATEGORY_STATIC(LogWorldComposition, Log, All);
 
diff --git a/Engine/Source/Runtime/Engine/Public/PhysicsPublic.h b/Engine/Source/Runtime/Engine/Public/PhysicsPublic.h
--- a/Engine/Source/Runtime/Engine/Public/PhysicsPublic.h
+++ b/Engine/Source/Runtime/Engine/Public/PhysicsPublic.h
@@ -10,6 +10,7 @@
 #include "PhysxUserData.h"
 #include "DynamicMeshBuilder.h"
 #include "LocalVertexFactory.h"
+#include "PhysicsEngine/PhysicsAsset.h"
 /**
  * Physics stats
  */
diff --git a/Engine/Source/Runtime/Renderer/Private/ScenePrivate.h b/Engine/Source/Runtime/Renderer/Private/ScenePrivate.h
--- a/Engine/Source/Runtime/Renderer/Private/ScenePrivate.h
+++ b/Engine/Source/Runtime/Renderer/Private/ScenePrivate.h
@@ -9,6 +9,7 @@
 
 class SceneRenderingAllocator;
 class USceneCaptureComponent;
+class UTextureRenderTarget;
 
 class SceneRenderingBitArrayAllocator
 	: public TInlineAllocator<4,SceneRenderingAllocator>
diff --git a/Engine/Source/Runtime/SlateCore/Private/Widgets/SWidget.cpp b/Engine/Source/Runtime/SlateCore/Private/Widgets/SWidget.cpp
--- a/Engine/Source/Runtime/SlateCore/Private/Widgets/SWidget.cpp
+++ b/Engine/Source/Runtime/SlateCore/Private/Widgets/SWidget.cpp
@@ -1,6 +1,8 @@
 // Copyright 1998-2014 Epic Games, Inc. All Rights Reserved.
 
 #include "SlateCorePrivatePCH.h"
+#include "Widgets/SWidget.h"
+#include "Input/Events.h"
 
 DECLARE_CYCLE_STAT(TEXT("OnPaint"), STAT_SlateOnPaint, STATGROUP_Slate);
 DECLARE_CYCLE_STAT(TEXT("ArrangeChildren"), STAT_SlateArrangeChildren, STATGROUP_Slate);
diff --git a/Engine/Source/Runtime/SlateCore/Public/Input/Events.h b/Engine/Source/Runtime/SlateCore/Public/Input/Events.h
--- a/Engine/Source/Runtime/SlateCore/Public/Input/Events.h
+++ b/Engine/Source/Runtime/SlateCore/Public/Input/Events.h
@@ -2,6 +2,7 @@
 
 #pragma once
 
+#include "InputCoreTypes.h"
 #include "Events.generated.h"
 
 
