CIS fix.

diff --git a/Engine/Source/Runtime/Engine/Private/EnginePrivate.h b/Engine/Source/Runtime/Engine/Private/EnginePrivate.h
--- a/Engine/Source/Runtime/Engine/Private/EnginePrivate.h
+++ b/Engine/Source/Runtime/Engine/Private/EnginePrivate.h
@@ -1,16 +1,21 @@
 // Copyright 1998-2014 Epic Games, Inc. All Rights Reserved.
 
 /*=============================================================================
-	EnginePrivate.h: Unreal engine private header file.
+EnginePrivate.h: Unreal engine private header file.
 =============================================================================*/
 
 #pragma once
+#define USE_NEW_PCH
+#ifdef USE_NEW_PCH
+#include "Core.h"
 #include "CoreUObject.h"
+#include "InputCore.h"
 #include "EngineDefines.h"
-#include "GameMapsSettings.h"
+#include "EngineSettings.h"
 #include "EngineStats.h"
 #include "EngineLogs.h"
 #include "EngineGlobals.h"
+
 #include "Engine/EngineBaseTypes.h"
 #include "Camera/CameraTypes.h"
 #include "Engine/EngineTypes.h"
@@ -23,69 +28,446 @@
 #include "AI/Navigation/NavigationTypes.h"
 #include "Kismet/BlueprintFunctionLibrary.h"
 #include "AI/Navigation/NavigationSystem.h"
+#include "AI/Navigation/NavCollision.h"
 #include "Camera/PlayerCameraManager.h"
 #include "GameFramework/Info.h"
 #include "GameFramework/GameMode.h"
 #include "Components/InputComponent.h"
+#include "GameFramework/ForceFeedbackEffect.h"
+#include "Engine/SubsurfaceProfile.h"
 #include "GameFramework/PlayerController.h"
+#include "Engine/DebugCameraController.h"
 #include "Engine/Canvas.h"
+#include "GameFramework/GameSession.h"
+#include "GameFramework/GameState.h"
+#include "GameFramework/PlayerState.h"
 #include "GameFramework/WorldSettings.h"
 #include "PhysicsEngine/BodyInstance.h"
 #include "Components/SceneComponent.h"
 #include "Components/LightComponentBase.h"
+#include "Components/LightComponent.h"
 #include "AI/Navigation/NavigationData.h"
+#include "AI/Navigation/NavAreas/NavArea.h"
 #include "AI/Navigation/RecastNavMesh.h"
+#include "Engine/NavigationObjectBase.h"
+#include "GameFramework/PlayerStart.h"
+#include "AI/Navigation/NavPathObserverInterface.h"
+#include "AI/Navigation/NavigationTestingActor.h"
+#include "AI/Navigation/NavLinkHostInterface.h"
+#include "Engine/Note.h"
 #include "Components/PrimitiveComponent.h"
 #include "GameFramework/Pawn.h"
+#include "Interfaces/NetworkPredictionInterface.h"
+#include "PhysicsEngine/PhysicsConstraintActor.h"
 #include "Components/MeshComponent.h"
 #include "Components/SkinnedMeshComponent.h"
+#include "PhysicsEngine/ConstraintInstance.h"
+#include "PhysicsEngine/PhysicsConstraintComponent.h"
 #include "Components/SkeletalMeshComponent.h"
+#include "Components/PawnNoiseEmitterComponent.h"
 #include "Components/AudioComponent.h"
+#include "Camera/CameraComponent.h"
+#include "Components/ChildActorComponent.h"
+#include "Components/DecalComponent.h"
+#include "Components/DirectionalLightComponent.h"
+#include "Components/PointLightComponent.h"
+#include "Components/SpotLightComponent.h"
+#include "Components/SkyLightComponent.h"
 #include "Components/ArrowComponent.h"
 #include "Components/BillboardComponent.h"
+#include "Components/BrushComponent.h"
+#include "Components/DrawFrustumComponent.h"
+#include "Debug/DebugDrawService.h"
+#include "Components/LineBatchComponent.h"
+#include "Components/MaterialBillboardComponent.h"
+#include "Components/DestructibleComponent.h"
 #include "Components/StaticMeshComponent.h"
+#include "Components/InstancedStaticMeshComponent.h"
 #include "Interfaces/Interface_CollisionDataProvider.h"
+#include "Components/ModelComponent.h"
 #include "Components/ShapeComponent.h"
+#include "Components/BoxComponent.h"
+#include "Components/CapsuleComponent.h"
+#include "Components/SphereComponent.h"
+#include "Components/VectorFieldComponent.h"
+#include "Components/ReflectionCaptureComponent.h"
+#include "Components/TimelineComponent.h"
+#include "EditorFramework/AssetImportData.h"
+#include "AI/Navigation/AvoidanceManager.h"
+#include "Engine/BlendableInterface.h"
 #include "Engine/BlueprintCore.h"
 #include "Engine/Blueprint.h"
+#include "Sound/DialogueTypes.h"
 #include "Kismet/GameplayStatics.h"
 #include "Engine/BlueprintGeneratedClass.h"
 #include "PhysicsEngine/BodySetup.h"
+#include "Camera/CameraAnim.h"
+#include "Camera/CameraAnimInst.h"
+#include "Camera/CameraModifier.h"
+#include "Camera/CameraShake.h"
+#include "Camera/CameraModifier_CameraShake.h"
+#include "GameFramework/CheatManager.h"
 #include "Engine/CollisionProfile.h"
+#include "Curves/CurveLinearColor.h"
+#include "Curves/CurveVector.h"
 #include "Engine/CurveTable.h"
+#include "GameFramework/DamageType.h"
 #include "Engine/DataTable.h"
+#include "Sound/DialogueVoice.h"
+#include "Sound/DialogueWave.h"
 #include "Distributions/Distribution.h"
 #include "Distributions/DistributionFloat.h"
 #include "Distributions/DistributionFloatConstant.h"
+#include "Distributions/DistributionFloatParameterBase.h"
+#include "Distributions/DistributionFloatConstantCurve.h"
 #include "Distributions/DistributionFloatUniform.h"
+#include "Distributions/DistributionFloatUniformCurve.h"
 #include "Distributions/DistributionVector.h"
 #include "Distributions/DistributionVectorConstant.h"
+#include "Distributions/DistributionVectorParameterBase.h"
+#include "Distributions/DistributionVectorConstantCurve.h"
 #include "Distributions/DistributionVectorUniform.h"
+#include "Distributions/DistributionVectorUniformCurve.h"
 #include "Engine/Engine.h"
+#include "Engine/GameEngine.h"
 #include "Exporters/Exporter.h"
+#include "Engine/FontImportOptions.h"
+#include "Engine/Font.h"
+#include "GameFramework/GameUserSettings.h"
+#include "GameFramework/PlayerInput.h"
+#include "GameFramework/InputSettings.h"
+#include "Engine/InterpCurveEdSetup.h"
 #include "Engine/Level.h"
 #include "Engine/LevelStreaming.h"
+#include "Engine/LevelStreamingKismet.h"
+#include "GameFramework/LocalMessage.h"
 #include "Materials/MaterialInterface.h"
 #include "Materials/Material.h"
 #include "Materials/MaterialInstance.h"
 #include "Materials/MaterialInstanceDynamic.h"
+#include "Materials/MaterialParameterCollection.h"
+#include "Materials/MaterialParameterCollectionInstance.h"
 #include "Engine/NetDriver.h"
+#include "Engine/PackageMapClient.h"
 #include "Engine/PendingNetGame.h"
+#include "PhysicalMaterials/PhysicalMaterial.h"
 #include "PhysicalMaterials/PhysicalMaterialPropertyBase.h"
+#include "PhysicsEngine/PhysicsAsset.h"
+#include "PhysicsEngine/PhysicsCollisionHandler.h"
+#include "PhysicsEngine/PhysicsConstraintTemplate.h"
+#include "PhysicsEngine/PhysicsSettings.h"
 #include "Engine/PlatformInterfaceBase.h"
+#include "Engine/InGameAdManager.h"
 #include "Engine/Player.h"
 #include "Engine/LocalPlayer.h"
 #include "Engine/NetConnection.h"
+#include "Engine/ChildConnection.h"
+#include "Engine/Polys.h"
+#include "Engine/SCS_Node.h"
 #include "Engine/Selection.h"
+#include "Engine/SimpleConstructionScript.h"
 #include "Engine/SkeletalMesh.h"
+#include "Engine/SkeletalMeshSocket.h"
 #include "Engine/StaticMesh.h"
+#include "Engine/StaticMeshSocket.h"
 #include "Engine/TextureDefines.h"
 #include "Engine/Texture.h"
 #include "Engine/Texture2D.h"
+#include "Engine/LightMapTexture2D.h"
+#include "Engine/ShadowMapTexture2D.h"
+#include "Engine/TextureLightProfile.h"
+#include "Engine/TextureCube.h"
+#include "Engine/TextureRenderTarget.h"
+#include "Engine/TextureRenderTarget2D.h"
+#include "Engine/TimelineTemplate.h"
+#include "GameFramework/TouchInterface.h"
+#include "Engine/UserDefinedStruct.h"
 #include "Engine/World.h"
 #include "ActiveSound.h"
+
 #include "SystemSettings.h"					// Scalability options.
 #include "SceneManagement.h"				// Scene management.
+
+#include "Engine/Light.h"
+#include "Interfaces/Interface_PostProcessVolume.h"
+#include "NavDataGenerator.h"
+#include "NavigationModifier.h"
+#include "NavigationOctree.h"
+#include "ContentStreaming.h"				// Content streaming class definitions.
+#include "ComponentReregisterContext.h"	
+#include "DrawDebugHelpers.h"
+#include "UnrealEngine.h"					// Unreal engine helpers.
+#include "CanvasTypes.h"							// Canvas.
+#include "EngineUtils.h"
+#include "InstancedFoliage.h"				// Instanced foliage.
+#include "UnrealExporter.h"					// Exporter definition.
+#include "TimerManager.h"					// Game play timers
+#include "NavigationSystemHelpers.h"
+#include "HardwareInfo.h"
+#else
+// Copyright 1998-2014 Epic Games, Inc. All Rights Reserved.
+
+/*=============================================================================
+EnginePrivate.h: Unreal engine private header file.
+=============================================================================*/
+
+#pragma once
+
+
+#include "Core.h"
+#include "CoreUObject.h"
+#include "InputCore.h"
+#include "EngineDefines.h"
+#include "EngineSettings.h"
+#include "EngineStats.h"
+#include "EngineLogs.h"
+#include "EngineGlobals.h"
+
+#include "Engine/EngineBaseTypes.h"
+#include "Camera/CameraTypes.h"
+#include "Engine/EngineTypes.h"
+#include "Components/ActorComponent.h"
+#include "GameFramework/Actor.h"
+#include "Engine/Brush.h"
+#include "AI/Navigation/NavAgentInterface.h"
+#include "GameFramework/Controller.h"
+#include "AI/Navigation/NavLinkDefinition.h"
+#include "AI/Navigation/NavigationTypes.h"
+#include "Kismet/BlueprintFunctionLibrary.h"
+#include "AI/Navigation/NavigationSystem.h"
+#include "AI/Navigation/NavCollision.h"
+#include "Camera/PlayerCameraManager.h"
+#include "GameFramework/Info.h"
+#include "GameFramework/GameMode.h"
+#include "Components/InputComponent.h"
+#include "GameFramework/ForceFeedbackEffect.h"
+#include "Engine/SubsurfaceProfile.h"
+#include "GameFramework/PlayerController.h"
+#include "Engine/DebugCameraController.h"
+#include "Engine/Canvas.h"
+#include "GameFramework/GameSession.h"
+#include "GameFramework/GameState.h"
+#include "GameFramework/PlayerState.h"
+#include "GameFramework/WorldSettings.h"
+#include "PhysicsEngine/BodyInstance.h"
+#include "Components/SceneComponent.h"
+#include "Components/LightComponentBase.h"
+#include "Components/LightComponent.h"
+#include "Engine/GeneratedMeshAreaLight.h"
+#include "AI/Navigation/NavigationData.h"
+#include "AI/Navigation/NavigationGraph.h"
+#include "AI/Navigation/NavAreas/NavArea.h"
+#include "AI/Navigation/RecastNavMesh.h"
+#include "AI/Navigation/NavigationGraphNode.h"
+#include "AI/Navigation/NavigationDataChunk.h"
+#include "Engine/NavigationObjectBase.h"
+#include "GameFramework/PlayerStart.h"
+#include "Engine/PlayerStartPIE.h"
+#include "AI/Navigation/NavPathObserverInterface.h"
+#include "AI/Navigation/NavigationTestingActor.h"
+#include "AI/Navigation/NavLinkHostInterface.h"
+#include "AI/Navigation/NavLinkProxy.h"
+#include "Engine/Note.h"
+#include "Components/PrimitiveComponent.h"
+#include "GameFramework/Pawn.h"
+#include "Interfaces/NetworkPredictionInterface.h"
+#include "Engine/ReflectionCapture.h"
+#include "Engine/BoxReflectionCapture.h"
+#include "Engine/PlaneReflectionCapture.h"
+#include "Engine/SphereReflectionCapture.h"
+#include "PhysicsEngine/RigidBodyBase.h"
+#include "PhysicsEngine/PhysicsConstraintActor.h"
+#include "PhysicsEngine/PhysicsThruster.h"
+#include "PhysicsEngine/RadialForceActor.h"
+#include "Engine/SceneCapture.h"
+#include "Engine/SceneCapture2D.h"
+#include "Engine/SceneCaptureCube.h"
+#include "Components/MeshComponent.h"
+#include "Components/SkinnedMeshComponent.h"
+#include "PhysicsEngine/ConstraintInstance.h"
+#include "PhysicsEngine/PhysicsConstraintComponent.h"
+#include "Components/SkeletalMeshComponent.h"
+#include "Engine/TargetPoint.h"
+#include "Engine/DataAsset.h"
+#include "Components/PawnNoiseEmitterComponent.h"
+#include "PhysicsEngine/PhysicsHandleComponent.h"
+#include "Atmosphere/AtmosphericFogComponent.h"
+#include "Components/AudioComponent.h"
+#include "Camera/CameraComponent.h"
+#include "Components/ChildActorComponent.h"
+#include "Components/DecalComponent.h"
+#include "Components/ExponentialHeightFogComponent.h"
+#include "Components/DirectionalLightComponent.h"
+#include "Components/PointLightComponent.h"
+#include "Components/SpotLightComponent.h"
+#include "Components/SkyLightComponent.h"
+#include "AI/Navigation/NavigationGraphNodeComponent.h"
+#include "PhysicsEngine/PhysicsThrusterComponent.h"
+#include "Components/PostProcessComponent.h"
+#include "Components/ArrowComponent.h"
+#include "Components/BillboardComponent.h"
+#include "Components/BrushComponent.h"
+#include "Components/DrawFrustumComponent.h"
+#include "Debug/DebugDrawService.h"
+#include "Components/LineBatchComponent.h"
+#include "Components/MaterialBillboardComponent.h"
+#include "Components/DestructibleComponent.h"
+#include "Components/PoseableMeshComponent.h"
+#include "Lightmass/LightmassPrimitiveSettingsObject.h"
+#include "Components/StaticMeshComponent.h"
+#include "Components/InstancedStaticMeshComponent.h"
+#include "Interfaces/Interface_CollisionDataProvider.h"
+#include "Components/ModelComponent.h"
+#include "AI/Navigation/NavLinkRenderingComponent.h"
+#include "AI/Navigation/NavMeshRenderingComponent.h"
+#include "AI/Navigation/NavTestRenderingComponent.h"
+#include "Components/ShapeComponent.h"
+#include "Components/BoxComponent.h"
+#include "Components/CapsuleComponent.h"
+#include "Components/SphereComponent.h"
+#include "Components/DrawSphereComponent.h"
+#include "Components/TextRenderComponent.h"
+#include "Components/VectorFieldComponent.h"
+#include "PhysicsEngine/RadialForceComponent.h"
+#include "Components/ReflectionCaptureComponent.h"
+#include "Components/BoxReflectionCaptureComponent.h"
+#include "Components/PlaneReflectionCaptureComponent.h"
+#include "Components/SphereReflectionCaptureComponent.h"
+#include "Components/SceneCaptureComponent.h"
+#include "Components/SceneCaptureComponent2D.h"
+#include "Components/SceneCaptureComponentCube.h"
+#include "GameFramework/SpringArmComponent.h"
+#include "Components/WindDirectionalSourceComponent.h"
+#include "Components/TimelineComponent.h"
+#include "EditorFramework/AssetImportData.h"
+#include "Engine/AssetUserData.h"
+#include "AI/Navigation/AvoidanceManager.h"
+#include "Engine/BlendableInterface.h"
+#include "Engine/BlueprintCore.h"
+#include "Engine/Blueprint.h"
+#include "Sound/DialogueTypes.h"
+#include "Kismet/GameplayStatics.h"
+#include "Kismet/HeadMountedDisplayFunctionLibrary.h"
+#include "Engine/BlueprintGeneratedClass.h"
+#include "PhysicsEngine/BodySetup.h"
+#include "Camera/CameraAnim.h"
+#include "Camera/CameraAnimInst.h"
+#include "Camera/CameraModifier.h"
+#include "Camera/CameraShake.h"
+#include "Camera/CameraModifier_CameraShake.h"
+#include "GameFramework/CheatManager.h"
+#include "Engine/CollisionProfile.h"
+#include "Curves/CurveLinearColor.h"
+#include "Curves/CurveVector.h"
+#include "Engine/CurveTable.h"
+#include "GameFramework/DamageType.h"
+#include "Engine/DataTable.h"
+#include "Sound/DialogueVoice.h"
+#include "Sound/DialogueWave.h"
+#include "Distributions/Distribution.h"
+#include "Distributions/DistributionFloat.h"
+#include "Distributions/DistributionFloatConstant.h"
+#include "Distributions/DistributionFloatParameterBase.h"
+#include "Distributions/DistributionFloatParticleParameter.h"
+#include "Distributions/DistributionFloatSoundParameter.h"
+#include "Distributions/DistributionFloatConstantCurve.h"
+#include "Distributions/DistributionFloatUniform.h"
+#include "Distributions/DistributionFloatUniformCurve.h"
+#include "Distributions/DistributionVector.h"
+#include "Distributions/DistributionVectorConstant.h"
+#include "Distributions/DistributionVectorParameterBase.h"
+#include "Distributions/DistributionVectorParticleParameter.h"
+#include "Distributions/DistributionVectorConstantCurve.h"
+#include "Distributions/DistributionVectorUniform.h"
+#include "Distributions/DistributionVectorUniformCurve.h"
+#include "Engine/Engine.h"
+#include "Engine/GameEngine.h"
+#include "Exporters/Exporter.h"
+#include "Engine/FontImportOptions.h"
+#include "Engine/Font.h"
+#include "GameFramework/GameUserSettings.h"
+#include "GameFramework/PlayerInput.h"
+#include "GameFramework/InputSettings.h"
+#include "Engine/InterpCurveEdSetup.h"
+#include "Engine/IntSerialization.h"
+#include "Layers/Layer.h"
+#include "Engine/Level.h"
+#include "Engine/LevelStreaming.h"
+#include "Engine/LevelStreamingAlwaysLoaded.h"
+#include "Engine/LevelStreamingBounds.h"
+#include "Engine/LevelStreamingKismet.h"
+#include "Engine/LevelStreamingPersistent.h"
+#include "Lightmass/LightmappedSurfaceCollection.h"
+#include "GameFramework/LocalMessage.h"
+#include "GameFramework/EngineMessage.h"
+#include "Materials/MaterialInterface.h"
+#include "Materials/Material.h"
+#include "Materials/MaterialInstance.h"
+#include "Materials/MaterialInstanceConstant.h"
+#include "Materials/MaterialInstanceDynamic.h"
+#include "Materials/MaterialParameterCollection.h"
+#include "Materials/MaterialParameterCollectionInstance.h"
+#include "Engine/NetDriver.h"
+#include "Engine/NetworkSettings.h"
+#include "Engine/ObjectLibrary.h"
+#include "Engine/ObjectReferencer.h"
+#include "GameFramework/OnlineSession.h"
+#include "Engine/PackageMapClient.h"
+#include "Engine/PendingNetGame.h"
+#include "PhysicalMaterials/PhysicalMaterial.h"
+#include "PhysicalMaterials/PhysicalMaterialPropertyBase.h"
+#include "PhysicsEngine/PhysicsAsset.h"
+#include "PhysicsEngine/PhysicsCollisionHandler.h"
+#include "PhysicsEngine/PhysicsConstraintTemplate.h"
+#include "PhysicsEngine/PhysicsSettings.h"
+#include "Engine/PlatformInterfaceBase.h"
+#include "Engine/CloudStorageBase.h"
+#include "Engine/InGameAdManager.h"
+#include "Engine/MicroTransactionBase.h"
+#include "Engine/TwitterIntegrationBase.h"
+#include "Engine/PlatformInterfaceWebResponse.h"
+#include "Engine/Player.h"
+#include "Engine/LocalPlayer.h"
+#include "Engine/NetConnection.h"
+#include "Engine/ChildConnection.h"
+#include "Engine/Polys.h"
+#include "Engine/RendererSettings.h"
+#include "MovieScene/RuntimeMovieScenePlayerInterface.h"
+#include "GameFramework/SaveGame.h"
+#include "Engine/SCS_Node.h"
+#include "Engine/Selection.h"
+#include "Engine/SimpleConstructionScript.h"
+#include "Engine/SkeletalMesh.h"
+#include "Engine/SkeletalMeshSocket.h"
+#include "Engine/StaticMesh.h"
+#include "Engine/StaticMeshSocket.h"
+#include "Camera/CameraStackTypes.h"
+#include "Engine/StreamableManager.h"
+#include "Engine/TextureDefines.h"
+#include "Tests/TextPropertyTestObject.h"
+#include "Engine/Texture.h"
+#include "Engine/Texture2D.h"
+#include "Engine/LightMapTexture2D.h"
+#include "Engine/ShadowMapTexture2D.h"
+#include "Engine/TextureLightProfile.h"
+#include "Engine/Texture2DDynamic.h"
+#include "Engine/TextureCube.h"
+#include "Engine/TextureRenderTarget.h"
+#include "Engine/TextureRenderTarget2D.h"
+#include "Engine/CanvasRenderTarget2D.h"
+#include "Engine/TextureRenderTargetCube.h"
+#include "EditorFramework/ThumbnailInfo.h"
+#include "Engine/TimelineTemplate.h"
+#include "GameFramework/TouchInterface.h"
+#include "Engine/UserDefinedEnum.h"
+#include "Engine/UserDefinedStruct.h"
+#include "Engine/World.h"
+
+#include "SystemSettings.h"					// Scalability options.
+#include "SceneManagement.h"				// Scene management.
+
+#include "NavDataGenerator.h"
+#include "NavLinkRenderingProxy.h"
 #include "NavigationModifier.h"
 #include "NavigationOctree.h"
 #include "ContentStreaming.h"				// Content streaming class definitions.
@@ -95,4 +477,9 @@
 #include "CanvasTypes.h"							// Canvas.
 #include "EngineUtils.h"
 #include "InstancedFoliage.h"				// Instanced foliage.
+#include "UnrealExporter.h"					// Exporter definition.
 #include "TimerManager.h"					// Game play timers
+#include "EngineService.h"
+#include "NavigationSystemHelpers.h"
+#include "HardwareInfo.h"
+#endif // USE_NEW_PCH
