BlueprintEd - Adding a null Node check when finding a place to insert a new node, ran into this when adding new nodes for events in widget blueprints.

diff --git a/Engine/Source/Runtime/Engine/Private/EdGraph/EdGraph.cpp b/Engine/Source/Runtime/Engine/Private/EdGraph/EdGraph.cpp
--- a/Engine/Source/Runtime/Engine/Private/EdGraph/EdGraph.cpp
+++ b/Engine/Source/Runtime/Engine/Private/EdGraph/EdGraph.cpp
@@ -232,8 +232,11 @@ FVector2D UEdGraph::GetGoodPlaceForNewNode()
 		for(int32 i=1; i<Nodes.Num(); i++)
 		{
 			Node = Nodes[i];
-			BottomLeft.X = FMath::Min<float>(BottomLeft.X, Node->NodePosX);
-			BottomLeft.Y = FMath::Max<float>(BottomLeft.Y, Node->NodePosY);
+			if ( Node )
+			{
+				BottomLeft.X = FMath::Min<float>(BottomLeft.X, Node->NodePosX);
+				BottomLeft.Y = FMath::Max<float>(BottomLeft.Y, Node->NodePosY);
+			}
 		}
 	}
 
