fixed compiler error

diff --git a/Engine/Source/Runtime/Engine/Private/SceneView.cpp b/Engine/Source/Runtime/Engine/Private/SceneView.cpp
--- a/Engine/Source/Runtime/Engine/Private/SceneView.cpp
+++ b/Engine/Source/Runtime/Engine/Private/SceneView.cpp
@@ -955,7 +955,7 @@ void FSceneView::StartFinalPostprocessSettings(FVector InViewLocation)
 	}
 }
 
-void FSceneView::EndFinalPostprocessSettings(FSceneViewInitOptions ViewInitOptions)
+void FSceneView::EndFinalPostprocessSettings(const FSceneViewInitOptions& ViewInitOptions)
 {
 	{
 		static const auto CVarMobileMSAA = IConsoleManager::Get().FindTConsoleVariableDataInt(TEXT("r.MobileMSAA"));
diff --git a/Engine/Source/Runtime/Engine/Public/SceneView.h b/Engine/Source/Runtime/Engine/Public/SceneView.h
--- a/Engine/Source/Runtime/Engine/Public/SceneView.h
+++ b/Engine/Source/Runtime/Engine/Public/SceneView.h
@@ -442,7 +442,7 @@ public:
 	void OverridePostProcessSettings(const FPostProcessSettings& Src, float Weight);
 
 	/** applied global restrictions from show flags */
-	void EndFinalPostprocessSettings(FSceneViewInitOptions ViewInitOptions);
+	void EndFinalPostprocessSettings(const FSceneViewInitOptions& ViewInitOptions);
 
 	/** Configure post process settings for the buffer visualization system */
 	void ConfigureBufferVisualizationSettings();
