Fix editor crash when changing a montage track's slot name, while an AnimNotify was selected.
https://jira.ol.epicgames.net/browse/OR-82

diff --git a/Engine/Source/Editor/Persona/Private/SAnimMontagePanel.cpp b/Engine/Source/Editor/Persona/Private/SAnimMontagePanel.cpp
--- a/Engine/Source/Editor/Persona/Private/SAnimMontagePanel.cpp
+++ b/Engine/Source/Editor/Persona/Private/SAnimMontagePanel.cpp
@@ -840,6 +840,9 @@ void SAnimMontagePanel::OnSlotNameChanged(TSharedPtr<FString> NewSelection, ESel
 			}
 			
 			UpdateSlotGroupWarningVisibility();
+
+			// Clear selection, so Details panel for AnimNotifies doesn't show outdated information.
+			ClearSelected();
 		}
 	}
 }
diff --git a/Engine/Source/Runtime/Engine/Classes/Animation/AnimLinkableElement.cpp b/Engine/Source/Runtime/Engine/Classes/Animation/AnimLinkableElement.cpp
--- a/Engine/Source/Runtime/Engine/Classes/Animation/AnimLinkableElement.cpp
+++ b/Engine/Source/Runtime/Engine/Classes/Animation/AnimLinkableElement.cpp
@@ -109,7 +109,7 @@ void FAnimLinkableElement::OnChanged(float NewMontageTime)
 		return;
 	}
 
-	SlotIndex = FMath::Clamp(SlotIndex, 0, LinkedMontage->SlotAnimTracks.Num());
+	SlotIndex = FMath::Clamp(SlotIndex, 0, LinkedMontage->SlotAnimTracks.Num()-1);
 
 	// If the link method changed, transform the link value
 	if(CachedLinkMethod != LinkMethod)
