Send ListSoundDurations output to requested archive instead of explicitly to the LogAudio channel.

diff --git a/Engine/Source/Runtime/Engine/Private/AudioDevice.cpp b/Engine/Source/Runtime/Engine/Private/AudioDevice.cpp
--- a/Engine/Source/Runtime/Engine/Private/AudioDevice.cpp
+++ b/Engine/Source/Runtime/Engine/Private/AudioDevice.cpp
@@ -457,8 +457,7 @@ bool FAudioDevice::HandleListSoundClassVolumesCommand( const TCHAR* Cmd, FOutput
 		{
 			const FSoundClassProperties& CurClass = It.Value();
 
-			FString Line = FString::Printf( TEXT("Cur (%3.2f, %3.2f) for SoundClass %s"), CurClass.Volume, CurClass.Pitch, *SoundClass->GetName() );
-			Ar.Logf( *Line );
+			Ar.Logf( TEXT("Cur (%3.2f, %3.2f) for SoundClass %s"), CurClass.Volume, CurClass.Pitch, *SoundClass->GetName() );
 		}
 	}
 
@@ -519,11 +518,11 @@ bool FAudioDevice::HandleListAudioComponentsCommand( const TCHAR* Cmd, FOutputDe
 
 bool FAudioDevice::HandleListSoundDurationsCommand( const TCHAR* Cmd, FOutputDevice& Ar )
 {
-	UE_LOG(LogAudio, Log,  TEXT( ",Sound,Duration,Channels" ) );
+	Ar.Logf(TEXT( "Sound,Duration,Channels" ) );
 	for( TObjectIterator<USoundWave> It; It; ++It )
 	{
 		USoundWave* SoundWave = *It;
-		UE_LOG(LogAudio, Log,  TEXT( ",%s,%f,%i" ), *SoundWave->GetPathName(), SoundWave->Duration, SoundWave->NumChannels );
+		Ar.Logf(TEXT( "%s,%f,%i" ), *SoundWave->GetPathName(), SoundWave->Duration, SoundWave->NumChannels );
 	}
 	return true;
 }
