Fixed navigation crashing when level has been saved with a navmesh instane and navigataion system has been disabled via WorldSettings.bEnableNavigationSystem #UE4

diff --git a/Engine/Source/Runtime/Engine/Private/AI/Navigation/NavigationData.cpp b/Engine/Source/Runtime/Engine/Private/AI/Navigation/NavigationData.cpp
--- a/Engine/Source/Runtime/Engine/Private/AI/Navigation/NavigationData.cpp
+++ b/Engine/Source/Runtime/Engine/Private/AI/Navigation/NavigationData.cpp
@@ -433,14 +433,17 @@ TArray<FBox> ANavigationData::GetNavigableBoundsInLevel(const FName& InLevelPack
 	TArray<FBox> Result;
 	const UNavigationSystem* NavSys = GetWorld()->GetNavigationSystem();
 	
-	const auto& NavigationBounds = NavSys->GetNavigationBounds();
-	Result.Reserve(NavigationBounds.Num());
-
-	for (const auto& Bounds : NavigationBounds)
+	if (NavSys != nullptr)
 	{
-		if (Bounds.PackageName == InLevelPackageName)
+		const auto& NavigationBounds = NavSys->GetNavigationBounds();
+		Result.Reserve(NavigationBounds.Num());
+
+		for (const auto& Bounds : NavigationBounds)
 		{
-			Result.Add(Bounds.AreaBox);
+			if (Bounds.PackageName == InLevelPackageName)
+			{
+				Result.Add(Bounds.AreaBox);
+			}
 		}
 	}
 	
