Fix incorrect calculations for interior volumes
Fix incorrect tooltips for what InteriorVolume/LPF represent

diff --git a/Engine/Source/Runtime/Engine/Classes/Sound/AudioVolume.h b/Engine/Source/Runtime/Engine/Classes/Sound/AudioVolume.h
--- a/Engine/Source/Runtime/Engine/Classes/Sound/AudioVolume.h
+++ b/Engine/Source/Runtime/Engine/Classes/Sound/AudioVolume.h
@@ -120,7 +120,7 @@ struct FInteriorSettings
 	UPROPERTY(EditAnywhere, BlueprintReadWrite, Category=InteriorSettings)
 	float ExteriorLPFTime;
 
-	// The desired volume of sounds inside the volume when the player is inside the volume
+	// The desired volume of sounds inside the volume when the player is outside the volume
 	UPROPERTY(EditAnywhere, BlueprintReadWrite, Category=InteriorSettings)
 	float InteriorVolume;
 
@@ -128,7 +128,7 @@ struct FInteriorSettings
 	UPROPERTY(EditAnywhere, BlueprintReadWrite, Category=InteriorSettings)
 	float InteriorTime;
 
-	// The desired LPF of sounds inside  the volume when the player is inside the volume
+	// The desired LPF of sounds inside  the volume when the player is outside the volume
 	UPROPERTY(EditAnywhere, BlueprintReadWrite, Category=InteriorSettings)
 	float InteriorLPF;
 
diff --git a/Engine/Source/Runtime/Engine/Private/ActiveSound.cpp b/Engine/Source/Runtime/Engine/Private/ActiveSound.cpp
--- a/Engine/Source/Runtime/Engine/Private/ActiveSound.cpp
+++ b/Engine/Source/Runtime/Engine/Private/ActiveSound.cpp
@@ -330,10 +330,10 @@ void FActiveSound::HandleInteriorVolumes( const FListener& Listener, FSoundParse
 	if( Listener.Volume == AudioVolume || !bAllowSpatialization )
 	{
 		// Ambient and listener in same ambient zone
-		CurrentInteriorVolume = ( SourceInteriorVolume * ( 1.0f - Listener.InteriorVolumeInterp ) ) + (Listener.InteriorSettings.InteriorVolume * Listener.InteriorVolumeInterp);
+		CurrentInteriorVolume = ( SourceInteriorVolume * ( 1.0f - Listener.InteriorVolumeInterp ) ) + Listener.InteriorVolumeInterp;
 		ParseParams.VolumeMultiplier *= CurrentInteriorVolume;
 
-		CurrentInteriorLPF = ( SourceInteriorLPF * ( 1.0f - Listener.InteriorLPFInterp ) ) + (Listener.InteriorSettings.InteriorLPF * Listener.InteriorLPFInterp);
+		CurrentInteriorLPF = ( SourceInteriorLPF * ( 1.0f - Listener.InteriorLPFInterp ) ) + Listener.InteriorLPFInterp;
 		ParseParams.HighFrequencyGain *= CurrentInteriorLPF;
 
 		UE_LOG(LogAudio, Verbose, TEXT( "Ambient in same volume. Volume *= %g LPF *= %g (%s)" ),
@@ -357,12 +357,12 @@ void FActiveSound::HandleInteriorVolumes( const FListener& Listener, FSoundParse
 		else
 		{
 			// The ambient sound is 'inside' - use the ambient sound's interior volume multiplied with the listeners exterior volume
-			CurrentInteriorVolume = ( SourceInteriorVolume * ( 1.0f - Listener.InteriorVolumeInterp ) ) + ( Ambient.InteriorVolume * Listener.InteriorVolumeInterp )
-										* ( SourceInteriorVolume * ( 1.0f - Listener.ExteriorVolumeInterp ) ) + ( Listener.InteriorSettings.ExteriorVolume * Listener.ExteriorVolumeInterp );
+			CurrentInteriorVolume = (( SourceInteriorVolume * ( 1.0f - Listener.InteriorVolumeInterp ) ) + ( Ambient.InteriorVolume * Listener.InteriorVolumeInterp ))
+										* (( SourceInteriorVolume * ( 1.0f - Listener.ExteriorVolumeInterp ) ) + ( Listener.InteriorSettings.ExteriorVolume * Listener.ExteriorVolumeInterp ));
 			ParseParams.VolumeMultiplier *= CurrentInteriorVolume;
 
-			CurrentInteriorLPF = ( SourceInteriorLPF * ( 1.0f - Listener.InteriorLPFInterp ) ) + ( Ambient.InteriorLPF * Listener.InteriorLPFInterp )
-										* ( SourceInteriorLPF * ( 1.0f - Listener.ExteriorLPFInterp ) ) + ( Listener.InteriorSettings.ExteriorLPF * Listener.ExteriorLPFInterp );
+			CurrentInteriorLPF = (( SourceInteriorLPF * ( 1.0f - Listener.InteriorLPFInterp ) ) + ( Ambient.InteriorLPF * Listener.InteriorLPFInterp ))
+										* (( SourceInteriorLPF * ( 1.0f - Listener.ExteriorLPFInterp ) ) + ( Listener.InteriorSettings.ExteriorLPF * Listener.ExteriorLPFInterp ));
 			ParseParams.HighFrequencyGain *= CurrentInteriorLPF;
 
 			UE_LOG(LogAudio, Verbose, TEXT( "Ambient in diff volume, ambient inside. Volume *= %g LPF *= %g (%s)" ),
