Orion Code Main:
AnimNotifyState
- Adding the OverridePSTemplate function for the bp's to be able to implement.

diff --git a/Engine/Source/Runtime/Engine/Classes/Animation/AnimNotifies/AnimNotifyState_Trail.h b/Engine/Source/Runtime/Engine/Classes/Animation/AnimNotifies/AnimNotifyState_Trail.h
--- a/Engine/Source/Runtime/Engine/Classes/Animation/AnimNotifies/AnimNotifyState_Trail.h
+++ b/Engine/Source/Runtime/Engine/Classes/Animation/AnimNotifies/AnimNotifyState_Trail.h
@@ -15,6 +15,9 @@ class UAnimNotifyState_Trail : public UAnimNotifyState
 	UPROPERTY(EditAnywhere, BlueprintReadWrite, Category = Trail)
 	UParticleSystem* PSTemplate;
 
+	UFUNCTION(BlueprintImplementableEvent)
+	virtual UParticleSystem* OverridePSTemplate(class USkeletalMeshComponent * MeshComp, class UAnimSequenceBase * Animation) const;
+
 	/** Name of the first socket defining this trail. */
 	UPROPERTY(EditAnywhere, BlueprintReadWrite, Category = Trail)
 	FName FirstSocketName;
diff --git a/Engine/Source/Runtime/Engine/Private/Animation/AnimNotifyState_Trail.cpp b/Engine/Source/Runtime/Engine/Private/Animation/AnimNotifyState_Trail.cpp
--- a/Engine/Source/Runtime/Engine/Private/Animation/AnimNotifyState_Trail.cpp
+++ b/Engine/Source/Runtime/Engine/Private/Animation/AnimNotifyState_Trail.cpp
@@ -52,6 +52,12 @@ void UAnimNotifyState_Trail::NotifyBegin(class USkeletalMeshComponent * MeshComp
 		Width = AnimInst->GetCurveValue(WidthScaleCurve);
 	}
 
+	UParticleSystem* ParticleSystemTemplate = OverridePSTemplate(MeshComp, Animation);
+	if (ParticleSystemTemplate == nullptr)
+	{
+		ParticleSystemTemplate = PSTemplate;
+	}
+
 	bool bFoundExistingTrail = false;
 	for (USceneComponent* Comp : Children)
 	{
