Non unity cis fix.

diff --git a/Engine/Source/Runtime/Engine/Classes/Distributions/Distribution.h b/Engine/Source/Runtime/Engine/Classes/Distributions/Distribution.h
--- a/Engine/Source/Runtime/Engine/Classes/Distributions/Distribution.h
+++ b/Engine/Source/Runtime/Engine/Classes/Distributions/Distribution.h
@@ -5,6 +5,15 @@
 #include "Math/CurveEdInterface.h"
 #include "Distribution.generated.h"
 
+UENUM()
+enum DistributionParamMode
+{
+	DPM_Normal,
+	DPM_Abs,
+	DPM_Direct,
+	DPM_MAX,
+};
+
 /** Lookup table for distributions. */
 #if !CPP      //noexport struct
 USTRUCT(noexport)
diff --git a/Engine/Source/Runtime/Engine/Classes/Distributions/DistributionFloatParameterBase.h b/Engine/Source/Runtime/Engine/Classes/Distributions/DistributionFloatParameterBase.h
--- a/Engine/Source/Runtime/Engine/Classes/Distributions/DistributionFloatParameterBase.h
+++ b/Engine/Source/Runtime/Engine/Classes/Distributions/DistributionFloatParameterBase.h
@@ -5,15 +5,6 @@
 #include "Distributions/DistributionFloatConstant.h"
 #include "DistributionFloatParameterBase.generated.h"
 
-UENUM()
-enum DistributionParamMode
-{
-	DPM_Normal,
-	DPM_Abs,
-	DPM_Direct,
-	DPM_MAX,
-};
-
 UCLASS(abstract, collapsecategories, hidecategories=Object, editinlinenew)
 class UDistributionFloatParameterBase : public UDistributionFloatConstant
 {
diff --git a/Engine/Source/Runtime/Engine/Classes/Distributions/DistributionVectorParameterBase.h b/Engine/Source/Runtime/Engine/Classes/Distributions/DistributionVectorParameterBase.h
--- a/Engine/Source/Runtime/Engine/Classes/Distributions/DistributionVectorParameterBase.h
+++ b/Engine/Source/Runtime/Engine/Classes/Distributions/DistributionVectorParameterBase.h
@@ -26,7 +26,7 @@ class UDistributionVectorParameterBase : public UDistributionVectorConstant
 	FVector MaxOutput;
 
 	UPROPERTY(EditAnywhere, Category=DistributionVectorParameterBase)
-	TEnumAsByte<enum DistributionParamMode> ParamModes[3];
+	TEnumAsByte<DistributionParamMode> ParamModes[3];
 
 
 	//Begin UDistributionVector Interface
