Fix root motion initial rotation handling

From UDN: https://udn.unrealengine.com/questions/224331/order-of-rotationtranslation-incorrect-when-extrac.html

diff --git a/Engine/Source/Runtime/Engine/Classes/Animation/AnimSequence.h b/Engine/Source/Runtime/Engine/Classes/Animation/AnimSequence.h
--- a/Engine/Source/Runtime/Engine/Classes/Animation/AnimSequence.h
+++ b/Engine/Source/Runtime/Engine/Classes/Animation/AnimSequence.h
@@ -567,7 +567,7 @@ class UAnimSequence : public UAnimSequenceBase
 	// End UAnimSequenceBase interface
 
 	// Extract Root Motion transform from the animation
-	FTransform ExtractRootMotion(const float & StartTime, const float & DeltaTime, const float bAllowLooping) const;
+	FTransform ExtractRootMotion(float StartTime, float DeltaTime, bool bAllowLooping) const;
 
 	// Extract Root Motion transform from a contiguous position range (no looping)
 	FTransform ExtractRootMotionFromRange(float StartTrackPosition, float EndTrackPosition) const;
diff --git a/Engine/Source/Runtime/Engine/Private/Animation/AnimSequence.cpp b/Engine/Source/Runtime/Engine/Private/Animation/AnimSequence.cpp
--- a/Engine/Source/Runtime/Engine/Private/Animation/AnimSequence.cpp
+++ b/Engine/Source/Runtime/Engine/Private/Animation/AnimSequence.cpp
@@ -731,7 +731,7 @@ FTransform UAnimSequence::ExtractRootTrackTransform(float Pos, const FBoneContai
 	return FTransform::Identity;
 }
 
-FTransform UAnimSequence::ExtractRootMotion(const float & StartTime, const float & DeltaTime, const float bAllowLooping) const
+FTransform UAnimSequence::ExtractRootMotion(float StartTime, float DeltaTime, bool bAllowLooping) const
 {
 	FRootMotionMovementParams RootMotionParams;
 
@@ -780,8 +780,8 @@ FTransform UAnimSequence::ExtractRootMotionFromRange(float StartTrackPosition, f
 
 	// Transform to Component Space Rotation (inverse root transform from first frame)
 	const FTransform RootToComponentRot = FTransform(InitialTransform.GetRotation().Inverse());
-	StartTransform = RootToComponentRot * StartTransform;
-	EndTransform = RootToComponentRot * EndTransform;
+	StartTransform = StartTransform * RootToComponentRot;
+	EndTransform = EndTransform * RootToComponentRot;
 
 	return EndTransform.GetRelativeTransform(StartTransform);
 }
