Revert playercontroller pausing back to original logic

diff --git a/Engine/Source/Runtime/Engine/Private/PlayerController.cpp b/Engine/Source/Runtime/Engine/Private/PlayerController.cpp
--- a/Engine/Source/Runtime/Engine/Private/PlayerController.cpp
+++ b/Engine/Source/Runtime/Engine/Private/PlayerController.cpp
@@ -1600,7 +1600,7 @@ void APlayerController::ClientReturnToMainMenu_Implementation(const FString& Ret
 bool APlayerController::SetPause( bool bPause, FCanUnpause CanUnpauseDelegate)
 {
 	bool bResult = false;
-	if (GetNetMode() == NM_Standalone)
+	if (GetNetMode() != NM_Client)
 	{
 		AGameMode* const GameMode = GetWorld()->GetAuthGameMode();
 		if (bPause)
