Add const qualifiers to some networking functions
AActor::IsNetRelevantFor is now const and takes const parameters
UNetDriver::IsLevelInitializedForActor is now const and takes const parameters
UNetConnection::ClientHasInitializedLevelFor is now const
FNetGUIDCache::SupportsObject is now const

diff --git a/Engine/Source/Developer/GameplayDebugger/Classes/GameplayDebuggingReplicator.h b/Engine/Source/Developer/GameplayDebugger/Classes/GameplayDebuggingReplicator.h
--- a/Engine/Source/Developer/GameplayDebugger/Classes/GameplayDebuggingReplicator.h
+++ b/Engine/Source/Developer/GameplayDebugger/Classes/GameplayDebuggingReplicator.h
@@ -95,7 +95,7 @@ class GAMEPLAYDEBUGGER_API AGameplayDebuggingReplicator : public AActor
 
 	virtual class UNetConnection* GetNetConnection() override;
 
-	virtual bool IsNetRelevantFor(class APlayerController* RealViewer, AActor* Viewer, const FVector& SrcLocation) override;
+	virtual bool IsNetRelevantFor(const APlayerController* RealViewer, const AActor* Viewer, const FVector& SrcLocation) const override;
 
 	virtual void PostInitializeComponents() override;
 
diff --git a/Engine/Source/Developer/GameplayDebugger/Private/GameplayDebuggingReplicator.cpp b/Engine/Source/Developer/GameplayDebugger/Private/GameplayDebuggingReplicator.cpp
--- a/Engine/Source/Developer/GameplayDebugger/Private/GameplayDebuggingReplicator.cpp
+++ b/Engine/Source/Developer/GameplayDebugger/Private/GameplayDebuggingReplicator.cpp
@@ -98,7 +98,7 @@ void AGameplayDebuggingReplicator::GetLifetimeReplicatedProps(TArray< FLifetimeP
 #endif
 }
 
-bool AGameplayDebuggingReplicator::IsNetRelevantFor(class APlayerController* RealViewer, AActor* Viewer, const FVector& SrcLocation)
+bool AGameplayDebuggingReplicator::IsNetRelevantFor(const APlayerController* RealViewer, const AActor* Viewer, const FVector& SrcLocation) const
 {
 	return LocalPlayerOwner == RealViewer;
 }
diff --git a/Engine/Source/Runtime/Engine/Classes/Engine/DemoNetConnection.h b/Engine/Source/Runtime/Engine/Classes/Engine/DemoNetConnection.h
--- a/Engine/Source/Runtime/Engine/Classes/Engine/DemoNetConnection.h
+++ b/Engine/Source/Runtime/Engine/Classes/Engine/DemoNetConnection.h
@@ -25,12 +25,15 @@ class UDemoNetConnection : public UNetConnection
 	virtual int32		IsNetReady( bool Saturate ) override;
 	virtual void		FlushNet( bool bIgnoreSimulation = false ) override;
 	virtual void		HandleClientPlayer( APlayerController* PC, class UNetConnection* NetConnection ) override;
-	virtual bool		ClientHasInitializedLevelFor( const UObject* TestObject ) override;
+	virtual bool		ClientHasInitializedLevelFor( const UObject* TestObject ) const override;
 	// End UNetConnection interface.
 
 	/** @return The DemoRecording driver object */
 	FORCEINLINE class UDemoNetDriver* GetDriver() { return (UDemoNetDriver*)Driver; }
 
+	/** @return The DemoRecording driver object */
+	FORCEINLINE const class UDemoNetDriver* GetDriver() const { return (UDemoNetDriver*)Driver; }
+
 	TArray< FQueuedDemoPacket > QueuedDemoPackets;
 };
 
diff --git a/Engine/Source/Runtime/Engine/Classes/Engine/NetConnection.h b/Engine/Source/Runtime/Engine/Classes/Engine/NetConnection.h
--- a/Engine/Source/Runtime/Engine/Classes/Engine/NetConnection.h
+++ b/Engine/Source/Runtime/Engine/Classes/Engine/NetConnection.h
@@ -484,7 +484,7 @@ public:
 	 * returns whether the client has initialized the level required for the given object
 	 * @return true if the client has initialized the level the object is in or the object is not in a level, false otherwise
 	 */
-	ENGINE_API virtual bool ClientHasInitializedLevelFor(const UObject* TestObject);
+	ENGINE_API virtual bool ClientHasInitializedLevelFor(const UObject* TestObject) const;
 
 	/**
 	 * Allows the connection to process the raw data that was received
diff --git a/Engine/Source/Runtime/Engine/Classes/Engine/NetDriver.h b/Engine/Source/Runtime/Engine/Classes/Engine/NetDriver.h
--- a/Engine/Source/Runtime/Engine/Classes/Engine/NetDriver.h
+++ b/Engine/Source/Runtime/Engine/Classes/Engine/NetDriver.h
@@ -540,5 +540,5 @@ protected:
 	/** Unregister all TickDispatch, TickFlush, PostTickFlush to tick in World */
 	ENGINE_API void UnregisterTickEvents(class UWorld* InWorld) const;
 	/** Returns true if this actor is considered to be in a loaded level */
-	bool IsLevelInitializedForActor(class AActor* InActor, class UNetConnection* InConnection);
+	bool IsLevelInitializedForActor(const AActor* InActor, const UNetConnection* InConnection) const;
 };
diff --git a/Engine/Source/Runtime/Engine/Classes/Engine/PackageMapClient.h b/Engine/Source/Runtime/Engine/Classes/Engine/PackageMapClient.h
--- a/Engine/Source/Runtime/Engine/Classes/Engine/PackageMapClient.h
+++ b/Engine/Source/Runtime/Engine/Classes/Engine/PackageMapClient.h
@@ -53,7 +53,7 @@ public:
 	FNetGUIDCache( UNetDriver * InDriver );
 
 	void			CleanReferences();
-	bool			SupportsObject( const UObject* Object );
+	bool			SupportsObject( const UObject* Object ) const;
 	bool			IsDynamicObject( const UObject* Object );
 	bool			IsNetGUIDAuthority() const;
 	FNetworkGUID	GetOrAssignNetGUID( const UObject* Object );
diff --git a/Engine/Source/Runtime/Engine/Classes/GameFramework/Actor.h b/Engine/Source/Runtime/Engine/Classes/GameFramework/Actor.h
--- a/Engine/Source/Runtime/Engine/Classes/GameFramework/Actor.h
+++ b/Engine/Source/Runtime/Engine/Classes/GameFramework/Actor.h
@@ -1597,7 +1597,7 @@ public:
 	  *
 	  * @return bool - true if this actor is network relevant to the client associated with RealViewer 
 	  */
-	virtual bool IsNetRelevantFor(class APlayerController* RealViewer, AActor* Viewer, const FVector& SrcLocation);
+	virtual bool IsNetRelevantFor(const APlayerController* RealViewer, const AActor* Viewer, const FVector& SrcLocation) const;
 
 	/**
 	 * Check if this actor is the owner when doing relevancy checks for actors marked bOnlyRelevantToOwner
diff --git a/Engine/Source/Runtime/Engine/Classes/GameFramework/Pawn.h b/Engine/Source/Runtime/Engine/Classes/GameFramework/Pawn.h
--- a/Engine/Source/Runtime/Engine/Classes/GameFramework/Pawn.h
+++ b/Engine/Source/Runtime/Engine/Classes/GameFramework/Pawn.h
@@ -198,7 +198,7 @@ public:
 	virtual FString GetHumanReadableName() const override;
 	virtual float GetNetPriority(const FVector& ViewPos, const FVector& ViewDir, APlayerController* Viewer, UActorChannel* InChannel, float Time, bool bLowBandwidth) override;
 	virtual bool ShouldTickIfViewportsOnly() const override;
-	virtual bool IsNetRelevantFor(APlayerController* RealViewer, AActor* Viewer, const FVector& SrcLocation) override;
+	virtual bool IsNetRelevantFor(const APlayerController* RealViewer, const AActor* Viewer, const FVector& SrcLocation) const override;
 	virtual void PostNetReceiveLocationAndRotation() override;
 	virtual void PostNetReceiveVelocity(const FVector& NewVelocity) override;
 	virtual void DisplayDebug(UCanvas* Canvas, const FDebugDisplayInfo& DebugDisplay, float& YL, float& YPos) override;
diff --git a/Engine/Source/Runtime/Engine/Classes/GameFramework/PlayerController.h b/Engine/Source/Runtime/Engine/Classes/GameFramework/PlayerController.h
--- a/Engine/Source/Runtime/Engine/Classes/GameFramework/PlayerController.h
+++ b/Engine/Source/Runtime/Engine/Classes/GameFramework/PlayerController.h
@@ -1151,7 +1151,7 @@ public:
 	virtual void GetActorEyesViewPoint(FVector& Location, FRotator& Rotation) const override;
 	virtual void CalcCamera(float DeltaTime, struct FMinimalViewInfo& OutResult) override;
 	virtual void TickActor(float DeltaTime, enum ELevelTick TickType, FActorTickFunction& ThisTickFunction) override;
-	virtual bool IsNetRelevantFor(APlayerController* RealViewer, AActor* Viewer, const FVector& SrcLocation) override;
+	virtual bool IsNetRelevantFor(const APlayerController* RealViewer, const AActor* Viewer, const FVector& SrcLocation) const override;
 	virtual void FellOutOfWorld(const class UDamageType& dmgType) override;
 	virtual void Reset() override;
 	virtual void Possess(APawn* aPawn) override;
diff --git a/Engine/Source/Runtime/Engine/Private/ActorReplication.cpp b/Engine/Source/Runtime/Engine/Private/ActorReplication.cpp
--- a/Engine/Source/Runtime/Engine/Private/ActorReplication.cpp
+++ b/Engine/Source/Runtime/Engine/Private/ActorReplication.cpp
@@ -122,7 +122,7 @@ void AActor::PostNetReceivePhysicState()
 	}
 }
 
-bool AActor::IsNetRelevantFor(APlayerController* RealViewer, AActor* Viewer, const FVector& SrcLocation)
+bool AActor::IsNetRelevantFor(const APlayerController* RealViewer, const AActor* Viewer, const FVector& SrcLocation) const
 {
 	if( bAlwaysRelevant || IsOwnedBy(Viewer) || IsOwnedBy(RealViewer) || this==Viewer || Viewer==Instigator )
 	{
diff --git a/Engine/Source/Runtime/Engine/Private/DemoNetDriver.cpp b/Engine/Source/Runtime/Engine/Private/DemoNetDriver.cpp
--- a/Engine/Source/Runtime/Engine/Private/DemoNetDriver.cpp
+++ b/Engine/Source/Runtime/Engine/Private/DemoNetDriver.cpp
@@ -825,7 +825,7 @@ void UDemoNetConnection::HandleClientPlayer( APlayerController* PC, UNetConnecti
 	}
 }
 
-bool UDemoNetConnection::ClientHasInitializedLevelFor(const UObject* TestObject)
+bool UDemoNetConnection::ClientHasInitializedLevelFor(const UObject* TestObject) const
 {
 	// We save all currently streamed levels into the demo stream so we can force the demo playback client
 	// to stay in sync with the recording server
diff --git a/Engine/Source/Runtime/Engine/Private/NetConnection.cpp b/Engine/Source/Runtime/Engine/Private/NetConnection.cpp
--- a/Engine/Source/Runtime/Engine/Private/NetConnection.cpp
+++ b/Engine/Source/Runtime/Engine/Private/NetConnection.cpp
@@ -375,7 +375,7 @@ void UNetConnection::SendPackageMap()
 {
 }
 
-bool UNetConnection::ClientHasInitializedLevelFor(const UObject* TestObject)
+bool UNetConnection::ClientHasInitializedLevelFor(const UObject* TestObject) const
 {
 	check(Driver);
 	checkSlow(Driver->IsServer());
diff --git a/Engine/Source/Runtime/Engine/Private/NetworkDriver.cpp b/Engine/Source/Runtime/Engine/Private/NetworkDriver.cpp
--- a/Engine/Source/Runtime/Engine/Private/NetworkDriver.cpp
+++ b/Engine/Source/Runtime/Engine/Private/NetworkDriver.cpp
@@ -596,7 +596,7 @@ void UNetDriver::TickDispatch( float DeltaTime )
 	}
 }
 
-bool UNetDriver::IsLevelInitializedForActor(class AActor* InActor, class UNetConnection* InConnection)
+bool UNetDriver::IsLevelInitializedForActor(const AActor* InActor, const UNetConnection* InConnection) const
 {
 	check(InActor);
     check(InConnection);
diff --git a/Engine/Source/Runtime/Engine/Private/PackageMapClient.cpp b/Engine/Source/Runtime/Engine/Private/PackageMapClient.cpp
--- a/Engine/Source/Runtime/Engine/Private/PackageMapClient.cpp
+++ b/Engine/Source/Runtime/Engine/Private/PackageMapClient.cpp
@@ -1357,7 +1357,7 @@ void FNetGUIDCache::CleanReferences()
 	UE_LOG( LogNetPackageMap, Warning, TEXT( "FNetGUIDCache::CleanReferences: ObjectLookup: %i, NetGUIDLookup: %i, Mem: %i kB" ), ObjectLookup.Num(), NetGUIDLookup.Num(), ( CountBytesAr.Mem / 1024 ) );
 }
 
-bool FNetGUIDCache::SupportsObject( const UObject* Object )
+bool FNetGUIDCache::SupportsObject( const UObject* Object ) const
 {
 	// NULL is always supported
 	if ( !Object )
diff --git a/Engine/Source/Runtime/Engine/Private/Pawn.cpp b/Engine/Source/Runtime/Engine/Private/Pawn.cpp
--- a/Engine/Source/Runtime/Engine/Private/Pawn.cpp
+++ b/Engine/Source/Runtime/Engine/Private/Pawn.cpp
@@ -1019,7 +1019,7 @@ bool APawn::IsBasedOnActor(const AActor* Other) const
 }
 
 
-bool APawn::IsNetRelevantFor(APlayerController* RealViewer, AActor* Viewer, const FVector& SrcLocation)
+bool APawn::IsNetRelevantFor(const APlayerController* RealViewer, const AActor* Viewer, const FVector& SrcLocation) const
 {
 	if( bAlwaysRelevant || RealViewer == Controller || IsOwnedBy(Viewer) || IsOwnedBy(RealViewer) || this==Viewer || Viewer==Instigator
 		|| IsBasedOnActor(Viewer) || (Viewer && Viewer->IsBasedOnActor(this)))
diff --git a/Engine/Source/Runtime/Engine/Private/PlayerController.cpp b/Engine/Source/Runtime/Engine/Private/PlayerController.cpp
--- a/Engine/Source/Runtime/Engine/Private/PlayerController.cpp
+++ b/Engine/Source/Runtime/Engine/Private/PlayerController.cpp
@@ -3795,7 +3795,7 @@ void APlayerController::TickActor( float DeltaSeconds, ELevelTick TickType, FAct
 #endif // !(UE_BUILD_SHIPPING || UE_BUILD_TEST)
 }
 
-bool APlayerController::IsNetRelevantFor(APlayerController* RealViewer, AActor* Viewer, const FVector& SrcLocation)
+bool APlayerController::IsNetRelevantFor(const APlayerController* RealViewer, const AActor* Viewer, const FVector& SrcLocation) const
 {
 	return ( this==RealViewer );
 }
diff --git a/Engine/Source/Runtime/GameplayAbilities/Private/Abilities/GameplayAbilityTargetActor.cpp b/Engine/Source/Runtime/GameplayAbilities/Private/Abilities/GameplayAbilityTargetActor.cpp
--- a/Engine/Source/Runtime/GameplayAbilities/Private/Abilities/GameplayAbilityTargetActor.cpp
+++ b/Engine/Source/Runtime/GameplayAbilities/Private/Abilities/GameplayAbilityTargetActor.cpp
@@ -90,7 +90,7 @@ void AGameplayAbilityTargetActor::CancelTargeting()
 	Destroy();
 }
 
-bool AGameplayAbilityTargetActor::IsNetRelevantFor(class APlayerController* RealViewer, AActor* Viewer, const FVector& SrcLocation)
+bool AGameplayAbilityTargetActor::IsNetRelevantFor(const APlayerController* RealViewer, const AActor* Viewer, const FVector& SrcLocation) const
 {
 	//The player who created the ability doesn't need to be updated about it - there should be local prediction in place.
 	if (RealViewer == MasterPC)
diff --git a/Engine/Source/Runtime/GameplayAbilities/Private/Abilities/GameplayAbilityWorldReticle.cpp b/Engine/Source/Runtime/GameplayAbilities/Private/Abilities/GameplayAbilityWorldReticle.cpp
--- a/Engine/Source/Runtime/GameplayAbilities/Private/Abilities/GameplayAbilityWorldReticle.cpp
+++ b/Engine/Source/Runtime/GameplayAbilities/Private/Abilities/GameplayAbilityWorldReticle.cpp
@@ -27,7 +27,7 @@ void AGameplayAbilityWorldReticle::InitializeReticle(AGameplayAbilityTargetActor
 	OnParametersInitialized();
 }
 
-bool AGameplayAbilityWorldReticle::IsNetRelevantFor(class APlayerController* RealViewer, AActor* Viewer, const FVector& SrcLocation)
+bool AGameplayAbilityWorldReticle::IsNetRelevantFor(const APlayerController* RealViewer, const AActor* Viewer, const FVector& SrcLocation) const
 {
 	//The player who created the ability doesn't need to be updated about it - there should be local prediction in place.
 	if (RealViewer == MasterPC)
diff --git a/Engine/Source/Runtime/GameplayAbilities/Public/Abilities/GameplayAbilityTargetActor.h b/Engine/Source/Runtime/GameplayAbilities/Public/Abilities/GameplayAbilityTargetActor.h
--- a/Engine/Source/Runtime/GameplayAbilities/Public/Abilities/GameplayAbilityTargetActor.h
+++ b/Engine/Source/Runtime/GameplayAbilities/Public/Abilities/GameplayAbilityTargetActor.h
@@ -72,7 +72,7 @@ public:
 	FAbilityTargetData	TargetDataReadyDelegate;
 	FAbilityTargetData	CanceledDelegate;
 
-	virtual bool IsNetRelevantFor(class APlayerController* RealViewer, AActor* Viewer, const FVector& SrcLocation) override;
+	virtual bool IsNetRelevantFor(const APlayerController* RealViewer, const AActor* Viewer, const FVector& SrcLocation) const override;
 
 	UPROPERTY(BlueprintReadOnly, Category = "Targeting")
 	APlayerController* MasterPC;
diff --git a/Engine/Source/Runtime/GameplayAbilities/Public/Abilities/GameplayAbilityWorldReticle.h b/Engine/Source/Runtime/GameplayAbilities/Public/Abilities/GameplayAbilityWorldReticle.h
--- a/Engine/Source/Runtime/GameplayAbilities/Public/Abilities/GameplayAbilityWorldReticle.h
+++ b/Engine/Source/Runtime/GameplayAbilities/Public/Abilities/GameplayAbilityWorldReticle.h
@@ -33,7 +33,7 @@ public:
 
 	// ------------------------------
 
-	virtual bool IsNetRelevantFor(class APlayerController* RealViewer, AActor* Viewer, const FVector& SrcLocation) override;
+	virtual bool IsNetRelevantFor(const APlayerController* RealViewer, const AActor* Viewer, const FVector& SrcLocation) const override;
 
 	void InitializeReticle(AGameplayAbilityTargetActor* InTargetingActor, FWorldReticleParameters InParameters);
 
