Made a render thread font cache for Canvas to use

This addresses a crash when trying to draw text from the render thread.

The extra font cache is managed by FEngineFontServices, and its creation is delayed until it is first needed. In Canvas you now use FEngineFontServices (rather than FSlateApplication directly) which will automatically return the correct cache for the current thread.

diff --git a/Engine/Source/Runtime/Engine/Engine.Build.cs b/Engine/Source/Runtime/Engine/Engine.Build.cs
--- a/Engine/Source/Runtime/Engine/Engine.Build.cs
+++ b/Engine/Source/Runtime/Engine/Engine.Build.cs
@@ -84,6 +84,21 @@ public class Engine : ModuleRules
 			}
 		);
 
+		if (Target.Type != TargetRules.TargetType.Server)
+		{
+			PrivateIncludePathModuleNames.AddRange(
+				new string[] { 
+					"SlateRHIRenderer",
+				}
+			);
+
+			DynamicallyLoadedModuleNames.AddRange(
+				new string[] {
+					"SlateRHIRenderer",
+				}
+			);
+		};
+
 		if (UEBuildConfiguration.bBuildDeveloperTools)
 		{
 			// Add "BlankModule" so that it gets compiled as an example and will be maintained and tested.  This can be removed
diff --git a/Engine/Source/Runtime/Engine/Private/UserInterface/Canvas.cpp b/Engine/Source/Runtime/Engine/Private/UserInterface/Canvas.cpp
--- a/Engine/Source/Runtime/Engine/Private/UserInterface/Canvas.cpp
+++ b/Engine/Source/Runtime/Engine/Private/UserInterface/Canvas.cpp
@@ -7,6 +7,7 @@
 #include "EnginePrivate.h"
 #include "SlateBasics.h"
 #include "Engine/Font.h"
+#include "EngineFontServices.h"
 #include "TileRendering.h"
 #include "RHIStaticStates.h"
 #include "WordWrapper.h"
@@ -555,6 +556,11 @@ void FCanvas::Flush_RenderThread(FRHICommandListImmediate& RHICmdList, bool bFor
 		return;
 	}
 
+	// Update the font cache with new text before elements are drawn
+	{
+		FEngineFontServices::Get().UpdateCache();
+	}
+
 	// FCanvasSortElement compare class
 	struct FCompareFCanvasSortElement
 	{
@@ -635,8 +641,7 @@ void FCanvas::Flush_GameThread(bool bForce)
 
 	// Update the font cache with new text before elements are drawn
 	{
-		TSharedRef<FSlateFontCache> FontCache = FSlateApplication::Get().GetRenderer()->GetFontCache();
-		FontCache->UpdateCache();
+		FEngineFontServices::Get().UpdateCache();
 	}
 
 	// FCanvasSortElement compare class
@@ -862,7 +867,7 @@ void FCanvas::DrawTile( float X, float Y, float SizeX,	float SizeY, float U, flo
 	BatchedElements->AddTriangle(V00,V11,V01,FinalTexture,BlendMode);
 }
 
-int32 FCanvas::DrawShadowedString( float StartX,float StartY,const TCHAR* Text,class UFont* Font,const FLinearColor& Color, const FLinearColor& ShadowColor )
+int32 FCanvas::DrawShadowedString( float StartX,float StartY,const TCHAR* Text,const UFont* Font,const FLinearColor& Color, const FLinearColor& ShadowColor )
 {
 	const float Z = 1.0f;
 	FCanvasTextItem TextItem( FVector2D( StartX, StartY ), FText::FromString( Text ), Font, Color );
@@ -886,7 +891,7 @@ void FCanvas::DrawNGon(const FVector2D& Center, const FColor& Color, int32 NumSi
 	DrawItem(NGonItem);
 }
 
-int32 FCanvas::DrawShadowedText( float StartX,float StartY,const FText& Text,class UFont* Font,const FLinearColor& Color, const FLinearColor& ShadowColor )
+int32 FCanvas::DrawShadowedText( float StartX,float StartY,const FText& Text,const UFont* Font,const FLinearColor& Color, const FLinearColor& ShadowColor )
 {
 	const float Z = 1.0f;
 	FCanvasTextItem TextItem( FVector2D( StartX, StartY ), Text, Font, Color );
diff --git a/Engine/Source/Runtime/Engine/Private/UserInterface/CanvasItem.cpp b/Engine/Source/Runtime/Engine/Private/UserInterface/CanvasItem.cpp
--- a/Engine/Source/Runtime/Engine/Private/UserInterface/CanvasItem.cpp
+++ b/Engine/Source/Runtime/Engine/Private/UserInterface/CanvasItem.cpp
@@ -10,6 +10,7 @@
 #include "RHIStaticStates.h"
 #include "Internationalization.h"
 #include "SlateBasics.h"
+#include "EngineFontServices.h"
 
 
 DECLARE_CYCLE_STAT(TEXT("CanvasTileTextureItem Time"),STAT_Canvas_TileTextureItemTime,STATGROUP_Canvas);
@@ -1027,8 +1028,12 @@ void FCanvasTextItem::DrawStringInternal_RuntimeCache( FCanvas* InCanvas, const
 		return;
 	}
 
-	TSharedRef<FSlateFontCache> FontCache = FSlateApplication::Get().GetRenderer()->GetFontCache();
-	
+	TSharedPtr<FSlateFontCache> FontCache = FEngineFontServices::Get().GetFontCache();
+	if( !FontCache.IsValid() )
+	{
+		return;
+	}
+
 	const float FontScale = 1.0f;
 	const FSlateFontInfo LegacyFontInfo = (SlateFontInfo.IsSet()) ? SlateFontInfo.GetValue() : Font->GetLegacySlateFontInfo();
 	FCharacterList& CharacterList = FontCache->GetCharacterList( LegacyFontInfo, FontScale );
diff --git a/Engine/Source/Runtime/Engine/Private/UserInterface/EngineFontServices.cpp b/Engine/Source/Runtime/Engine/Private/UserInterface/EngineFontServices.cpp
new file mode 100644
--- /dev/null
+++ b/Engine/Source/Runtime/Engine/Private/UserInterface/EngineFontServices.cpp
@@ -0,0 +1,133 @@
+// Copyright 1998-2014 Epic Games, Inc. All Rights Reserved.
+
+#include "EnginePrivate.h"
+#include "EngineFontServices.h"
+#include "SlateBasics.h"
+#if !UE_SERVER
+	#include "ISlateRHIRendererModule.h"
+#endif
+
+FEngineFontServices* FEngineFontServices::Instance = nullptr;
+
+FEngineFontServices::FEngineFontServices()
+{
+	check(IsInGameThread());
+
+#if !UE_SERVER
+	RenderThreadFontAtlasFactory = FModuleManager::Get().GetModuleChecked<ISlateRHIRendererModule>("SlateRHIRenderer").CreateSlateFontAtlasFactory();
+#endif
+}
+
+FEngineFontServices::~FEngineFontServices()
+{
+	check(IsInGameThread());
+
+	if(RenderThreadFontCache.IsValid())
+	{
+		RenderThreadFontCache->ReleaseResources();
+		FlushRenderingCommands();
+
+		RenderThreadFontCache.Reset();
+	}
+}
+
+void FEngineFontServices::Create()
+{
+	check(!Instance);
+	Instance = new FEngineFontServices();
+}
+
+void FEngineFontServices::Destroy()
+{
+	check(Instance);
+	delete Instance;
+	Instance = nullptr;
+}
+
+FEngineFontServices& FEngineFontServices::Get()
+{
+	check(Instance);
+	return *Instance;
+}
+
+TSharedPtr<FSlateFontCache> FEngineFontServices::GetFontCache()
+{
+	check(IsInGameThread() || IsInRenderingThread());
+
+#if !UE_SERVER
+	if(IsInGameThread())
+	{
+		if(FSlateApplication::IsInitialized())
+		{
+			return FSlateApplication::Get().GetRenderer()->GetFontCache();
+		}
+	}
+	else
+	{
+		ConditionalCreateRenderThreadFontCache();
+		return RenderThreadFontCache;
+	}
+#endif
+
+	return nullptr;
+}
+
+TSharedPtr<FSlateFontMeasure> FEngineFontServices::GetFontMeasure()
+{
+	check(IsInGameThread() || IsInRenderingThread());
+
+#if !UE_SERVER
+	if(IsInGameThread())
+	{
+		if(FSlateApplication::IsInitialized())
+		{
+			FSlateApplication::Get().GetRenderer()->GetFontMeasureService();
+		}
+	}
+	else
+	{
+		ConditionalCreatRenderThreadFontMeasure();
+		return RenderThreadFontMeasure;
+	}
+#endif
+
+	return nullptr;
+}
+
+void FEngineFontServices::UpdateCache()
+{
+	check(IsInGameThread() || IsInRenderingThread());
+
+	if(IsInGameThread())
+	{
+		if(FSlateApplication::IsInitialized())
+		{
+			FSlateApplication::Get().GetRenderer()->GetFontCache()->UpdateCache();
+		}
+	}
+	else
+	{
+		if(RenderThreadFontCache.IsValid())
+		{
+			RenderThreadFontCache->UpdateCache();
+		}
+	}
+}
+
+void FEngineFontServices::ConditionalCreateRenderThreadFontCache()
+{
+	if(!RenderThreadFontCache.IsValid())
+	{
+		RenderThreadFontCache = MakeShareable(new FSlateFontCache(RenderThreadFontAtlasFactory.ToSharedRef()));
+	}
+}
+
+void FEngineFontServices::ConditionalCreatRenderThreadFontMeasure()
+{
+	ConditionalCreateRenderThreadFontCache();
+
+	if(!RenderThreadFontMeasure.IsValid())
+	{
+		RenderThreadFontMeasure = FSlateFontMeasure::Create(RenderThreadFontCache.ToSharedRef());
+	}
+}
diff --git a/Engine/Source/Runtime/Engine/Public/CanvasTypes.h b/Engine/Source/Runtime/Engine/Public/CanvasTypes.h
--- a/Engine/Source/Runtime/Engine/Public/CanvasTypes.h
+++ b/Engine/Source/Runtime/Engine/Public/CanvasTypes.h
@@ -484,9 +484,9 @@ public:
 	* @param ShadowColor - Shadow color to draw underneath the text (ignored for distance field fonts)
 	* @return total size in pixels of text drawn
 	*/
-	ENGINE_API int32 DrawShadowedString( float StartX, float StartY, const TCHAR* Text, class UFont* Font, const FLinearColor& Color, const FLinearColor& ShadowColor = FLinearColor::Black );
+	ENGINE_API int32 DrawShadowedString( float StartX, float StartY, const TCHAR* Text, const UFont* Font, const FLinearColor& Color, const FLinearColor& ShadowColor = FLinearColor::Black );
 	
-	ENGINE_API int32 DrawShadowedText( float StartX, float StartY, const FText& Text, class UFont* Font, const FLinearColor& Color, const FLinearColor& ShadowColor = FLinearColor::Black );
+	ENGINE_API int32 DrawShadowedText( float StartX, float StartY, const FText& Text, const UFont* Font, const FLinearColor& Color, const FLinearColor& ShadowColor = FLinearColor::Black );
 
 	ENGINE_API void DrawNGon(const FVector2D& Center, const FColor& Color, int32 NumSides, float Radius);
 
diff --git a/Engine/Source/Runtime/Engine/Public/EngineFontServices.h b/Engine/Source/Runtime/Engine/Public/EngineFontServices.h
new file mode 100644
--- /dev/null
+++ b/Engine/Source/Runtime/Engine/Public/EngineFontServices.h
@@ -0,0 +1,54 @@
+// Copyright 1998-2014 Epic Games, Inc. All Rights Reserved.
+
+#pragma once
+
+/** 
+ * Provides the Engine with access to the Slate font cache and font measuring services (for Canvas)
+ * On the game thread this just leverages the Slate services, but the render thread needs its own instances
+ */
+class ENGINE_API FEngineFontServices
+{
+public:
+	/** Create the singular instance of this class - must be called from the game thread */
+	static void Create();
+
+	/** Destroy the singular instance of this class - must be called from the game thread */
+	static void Destroy();
+
+	/** Get the singular instance of this class */
+	static FEngineFontServices& Get();
+
+	/** Get the font cache to use for the current thread */
+	TSharedPtr<FSlateFontCache> GetFontCache();
+
+	/** Get the font measure to use for the current thread */
+	TSharedPtr<FSlateFontMeasure> GetFontMeasure();
+
+	/** Update the cache for the current thread */
+	void UpdateCache();
+
+private:
+	/** Constructor - must be called from the game thread */
+	FEngineFontServices();
+
+	/** Destructor - must be called from the game thread */
+	~FEngineFontServices();
+
+	/** Create the font cache for the render thread if it doesn't yet exist */
+	void ConditionalCreateRenderThreadFontCache();
+
+	/** Create the font measure for the render thread if it doesn't yet exist */
+	void ConditionalCreatRenderThreadFontMeasure();
+
+	/** Font atlas factory to use for the render thread - created on the game thread as it depends on another module */
+	TSharedPtr<ISlateFontAtlasFactory> RenderThreadFontAtlasFactory;
+
+	/** Font cache used by the render thread - creation is delayed until the first request is made */
+	TSharedPtr<FSlateFontCache> RenderThreadFontCache;
+
+	/** Font measure used by the render thread - creation is delayed until the first request is made */
+	TSharedPtr<FSlateFontMeasure> RenderThreadFontMeasure;
+
+	/** Singular instance of this class */
+	static FEngineFontServices* Instance;
+};
diff --git a/Engine/Source/Runtime/Launch/Private/LaunchEngineLoop.cpp b/Engine/Source/Runtime/Launch/Private/LaunchEngineLoop.cpp
--- a/Engine/Source/Runtime/Launch/Private/LaunchEngineLoop.cpp
+++ b/Engine/Source/Runtime/Launch/Private/LaunchEngineLoop.cpp
@@ -49,6 +49,7 @@
 #if !UE_SERVER
 	#include "HeadMountedDisplay.h"
 	#include "ISlateRHIRendererModule.h"
+	#include "EngineFontServices.h"
 #endif
 
 	#include "MoviePlayer.h"
@@ -1246,6 +1247,9 @@ int32 FEngineLoop::PreInit( const TCHAR* CmdLine )
 		FCoreStyle::ResetToDefault();
 	}
 
+#if !UE_SERVER
+	FEngineFontServices::Create();
+#endif
 	
 	FScopedSlowTask SlowTask(100, NSLOCTEXT("EngineLoop", "EngineLoop_Initializing", "Initializing..."));
 
@@ -2003,6 +2007,10 @@ void FEngineLoop::Exit()
 		GEngine->PreExit();
 	}
 
+#if !UE_SERVER
+	FEngineFontServices::Destroy();
+#endif
+
 	// close all windows
 	FSlateApplication::Shutdown();
 
diff --git a/Engine/Source/Runtime/SlateCore/Private/Fonts/FontBulkData.cpp b/Engine/Source/Runtime/SlateCore/Private/Fonts/FontBulkData.cpp
--- a/Engine/Source/Runtime/SlateCore/Private/Fonts/FontBulkData.cpp
+++ b/Engine/Source/Runtime/SlateCore/Private/Fonts/FontBulkData.cpp
@@ -40,6 +40,7 @@ UFontBulkData::UFontBulkData(const void* const InFontData, const int32 InFontDat
 
 const void* UFontBulkData::Lock(int32& OutFontDataSizeBytes) const
 {
+	CriticalSection.Lock();
 	OutFontDataSizeBytes = BulkData.GetBulkDataSize();
 	return BulkData.LockReadOnly();
 }
@@ -47,6 +48,7 @@ const void* UFontBulkData::Lock(int32& OutFontDataSizeBytes) const
 void UFontBulkData::Unlock() const
 {
 	BulkData.Unlock();
+	CriticalSection.Unlock();
 }
 
 void UFontBulkData::Serialize(FArchive& Ar)
diff --git a/Engine/Source/Runtime/SlateCore/Private/Fonts/LegacySlateFontInfoCache.cpp b/Engine/Source/Runtime/SlateCore/Private/Fonts/LegacySlateFontInfoCache.cpp
--- a/Engine/Source/Runtime/SlateCore/Private/Fonts/LegacySlateFontInfoCache.cpp
+++ b/Engine/Source/Runtime/SlateCore/Private/Fonts/LegacySlateFontInfoCache.cpp
@@ -68,6 +68,9 @@ const FFontData& FLegacySlateFontInfoCache::GetFallbackFont()
 {
 	const FName FallbackFontName = *(FPaths::EngineContentDir() / TEXT("Slate/Fonts/") / (NSLOCTEXT("Slate", "FallbackFont", "DroidSansFallback").ToString() + TEXT(".ttf")));
 
+	// GetFallbackFont is called directly from the font cache, so may be called from multiple threads at once
+	FScopeLock Lock(&FallbackFontCS);
+
 	{
 		TSharedPtr<const FFontData>* const ExistingFallbackFont = FallbackFonts.Find(FallbackFontName);
 		if(ExistingFallbackFont)
@@ -84,11 +87,15 @@ const FFontData& FLegacySlateFontInfoCache::GetFallbackFont()
 
 const FFontData& FLegacySlateFontInfoCache::GetLastResortFont()
 {
+	// GetLastResortFont is called directly from the font cache, so may be called from multiple threads at once
+	FScopeLock Lock(&LastResortFontCS);
+
 	if (!LastResortFont.IsValid())
 	{
 		const FString LastResortFontPath = FPaths::EngineContentDir() / TEXT("Slate/Fonts/LastResort.ttf");
 		LastResortFont = MakeShareable(new FFontData(LastResortFontPath, new UFontBulkData(LastResortFontPath), EFontHinting::Default));
 	}
+
 	return *LastResortFont;
 }
 
diff --git a/Engine/Source/Runtime/SlateCore/Private/Fonts/LegacySlateFontInfoCache.h b/Engine/Source/Runtime/SlateCore/Private/Fonts/LegacySlateFontInfoCache.h
--- a/Engine/Source/Runtime/SlateCore/Private/Fonts/LegacySlateFontInfoCache.h
+++ b/Engine/Source/Runtime/SlateCore/Private/Fonts/LegacySlateFontInfoCache.h
@@ -80,5 +80,8 @@ private:
 	TMap<FName, TSharedPtr<const FFontData>> FallbackFonts;
 	TSharedPtr<const FFontData> LastResortFont;
 
+	FCriticalSection FallbackFontCS;
+	FCriticalSection LastResortFontCS;
+
 	static TSharedPtr<FLegacySlateFontInfoCache> Singleton;
 };
diff --git a/Engine/Source/Runtime/SlateCore/Private/Textures/TextureAtlas.cpp b/Engine/Source/Runtime/SlateCore/Private/Textures/TextureAtlas.cpp
--- a/Engine/Source/Runtime/SlateCore/Private/Textures/TextureAtlas.cpp
+++ b/Engine/Source/Runtime/SlateCore/Private/Textures/TextureAtlas.cpp
@@ -50,7 +50,8 @@ const FAtlasedTextureSlot* FSlateTextureAtlas::AddTexture( uint32 TextureWidth,
 
 void FSlateTextureAtlas::MarkTextureDirty()
 {
-	check( IsThreadSafeForSlateRendering() );
+	//check( IsThreadSafeForSlateRendering() );
+	check( (GSlateLoadingThreadId != 0) || FPlatformTLS::GetCurrentThreadId() == AtlasOwnerThreadId );
 	bNeedsUpdate = true;
 }
 	
@@ -84,6 +85,8 @@ void FSlateTextureAtlas::InitAtlasData()
 	AtlasData.Reserve(AtlasWidth * AtlasHeight * Stride);
 	AtlasData.AddZeroed(AtlasWidth * AtlasHeight * Stride);
 
+	AtlasOwnerThreadId = FPlatformTLS::GetCurrentThreadId();
+
 	INC_MEMORY_STAT_BY(STAT_SlateTextureAtlasMemory, AtlasData.GetAllocatedSize());
 }
 
diff --git a/Engine/Source/Runtime/SlateCore/Public/Fonts/FontBulkData.h b/Engine/Source/Runtime/SlateCore/Public/Fonts/FontBulkData.h
--- a/Engine/Source/Runtime/SlateCore/Public/Fonts/FontBulkData.h
+++ b/Engine/Source/Runtime/SlateCore/Public/Fonts/FontBulkData.h
@@ -31,4 +31,7 @@ public:
 private:
 	/** Internal bulk data */
 	FByteBulkData BulkData;
+
+	/** Critical section to prevent concurrent access when locking the internal bulk data */
+	mutable FCriticalSection CriticalSection;
 };
diff --git a/Engine/Source/Runtime/SlateCore/Public/Textures/TextureAtlas.h b/Engine/Source/Runtime/SlateCore/Public/Textures/TextureAtlas.h
--- a/Engine/Source/Runtime/SlateCore/Public/Textures/TextureAtlas.h
+++ b/Engine/Source/Runtime/SlateCore/Public/Textures/TextureAtlas.h
@@ -188,4 +188,7 @@ protected:
 
 	/** True if this texture needs to have its rendering resources updated */
 	bool bNeedsUpdate;
+
+	/** ID of the thread that owns this atlas - this is the only thread that can safely update it */
+	uint32 AtlasOwnerThreadId;
 };
diff --git a/Engine/Source/Runtime/SlateRHIRenderer/Private/SlateRHIRendererModule.cpp b/Engine/Source/Runtime/SlateRHIRenderer/Private/SlateRHIRendererModule.cpp
--- a/Engine/Source/Runtime/SlateRHIRenderer/Private/SlateRHIRendererModule.cpp
+++ b/Engine/Source/Runtime/SlateRHIRenderer/Private/SlateRHIRendererModule.cpp
@@ -54,6 +54,11 @@ public:
 		return MakeShareable( new FSlate3DRenderer( ResourceManager, FontCache ) );
 	}
 
+	virtual TSharedRef<ISlateFontAtlasFactory> CreateSlateFontAtlasFactory() override
+	{
+		return MakeShareable(new FSlateRHIFontAtlasFactory);
+	}
+
 	virtual void StartupModule( ) override { }
 	virtual void ShutdownModule( ) override { }
 
diff --git a/Engine/Source/Runtime/SlateRHIRenderer/Public/Interfaces/ISlateRHIRendererModule.h b/Engine/Source/Runtime/SlateRHIRenderer/Public/Interfaces/ISlateRHIRendererModule.h
--- a/Engine/Source/Runtime/SlateRHIRenderer/Public/Interfaces/ISlateRHIRendererModule.h
+++ b/Engine/Source/Runtime/SlateRHIRenderer/Public/Interfaces/ISlateRHIRendererModule.h
@@ -6,6 +6,7 @@
 #include "SlateCore.h"
 
 class ISlate3DRenderer;
+class ISlateFontAtlasFactory;
 
 /**
  * Interface for the Slate RHI Renderer module.
@@ -23,4 +24,6 @@ public:
 	virtual TSharedRef<FSlateRenderer> CreateSlateRHIRenderer( ) = 0;
 
 	virtual TSharedRef<ISlate3DRenderer> CreateSlate3DRenderer() = 0;
+
+	virtual TSharedRef<ISlateFontAtlasFactory> CreateSlateFontAtlasFactory() = 0;
 };
