Fix some missing read locks for overlap queries and destructible sync update

diff --git a/Engine/Source/Runtime/Engine/Private/Collision/PhysXCollision.cpp b/Engine/Source/Runtime/Engine/Private/Collision/PhysXCollision.cpp
--- a/Engine/Source/Runtime/Engine/Private/Collision/PhysXCollision.cpp
+++ b/Engine/Source/Runtime/Engine/Private/Collision/PhysXCollision.cpp
@@ -1612,6 +1612,7 @@ bool GeomOverlapMulti_PhysX(const UWorld* World, const PxGeometry& PGeom, const
 {
 	SCOPE_CYCLE_COUNTER(STAT_Collision_GeomOverlapMultiple);
 	bool bHaveBlockingHit = false;
+	PxScene* LockedScenes[] = { nullptr, nullptr, nullptr };
 
 	// overlapMultiple only supports sphere/capsule/box 
 	if (PGeom.getType()==PxGeometryType::eSPHERE || PGeom.getType()==PxGeometryType::eCAPSULE || PGeom.getType()==PxGeometryType::eBOX || PGeom.getType()==PxGeometryType::eCONVEXMESH )
@@ -1625,7 +1626,8 @@ bool GeomOverlapMulti_PhysX(const UWorld* World, const PxGeometry& PGeom, const
 		FPhysScene* PhysScene = World->GetPhysicsScene();
 		PxScene* SyncScene = PhysScene->GetPhysXScene(PST_Sync);
 
-		SCOPED_SCENE_READ_LOCK(SyncScene);
+		SCENE_LOCK_READ(SyncScene);		//we can't use scoped because we later do a conversion which depends on these results and it should all be atomic
+		LockedScenes[PST_Sync] = SyncScene;
 
 		// Create buffer for hits. Note: memory is not initialized (for perf reasons), since API does not require it.
 		TTypeCompatibleBytes<PxOverlapHit> RawPOverlapArray[OVERLAP_BUFFER_SIZE];
@@ -1647,7 +1649,8 @@ bool GeomOverlapMulti_PhysX(const UWorld* World, const PxGeometry& PGeom, const
 		if (Params.bTraceAsyncScene && PhysScene->HasAsyncScene())
 		{		
 			PxScene* AsyncScene = PhysScene->GetPhysXScene(PST_Async);
-			SCOPED_SCENE_READ_LOCK(AsyncScene);
+			SCENE_LOCK_READ(AsyncScene);			//we can't use scoped because we later do a conversion which depends on these results and it should all be atomic
+			LockedScenes[PST_Async] = AsyncScene;
 
 			// Write into the same PHits buffer
 			PxOverlapHit* PAsyncOverlapArray = POverlapArray + NumHits;
@@ -1673,6 +1676,11 @@ bool GeomOverlapMulti_PhysX(const UWorld* World, const PxGeometry& PGeom, const
 			bHaveBlockingHit = ConvertOverlapResults(NumHits, POverlapArray, PFilter, OutOverlaps);
 		}
 
+		for (PxScene* LockedScene : LockedScenes)
+		{
+			SCENE_UNLOCK_READ(LockedScene);
+		}
+		
 #if !(UE_BUILD_SHIPPING || UE_BUILD_TEST)
 		if ((World->DebugDrawTraceTag != NAME_None) && (World->DebugDrawTraceTag == Params.TraceTag))
 		{
diff --git a/Engine/Source/Runtime/Engine/Private/PhysicsEngine/PhysScene.cpp b/Engine/Source/Runtime/Engine/Private/PhysicsEngine/PhysScene.cpp
--- a/Engine/Source/Runtime/Engine/Private/PhysicsEngine/PhysScene.cpp
+++ b/Engine/Source/Runtime/Engine/Private/PhysicsEngine/PhysScene.cpp
@@ -768,6 +768,7 @@ void FPhysScene::SyncComponentsToBodies(uint32 SceneType)
 #if WITH_APEX
 	if (ActiveDestructibleActors[SceneType].Num())
 	{
+		SCOPED_SCENE_READ_LOCK(GetPhysXScene(SceneType));
 		UDestructibleComponent::UpdateDestructibleChunkTM(ActiveDestructibleActors[SceneType]);
 	}
 #endif
