#docs - Documentation for more framework classes.

diff --git a/Engine/Source/Editor/UnrealEd/Classes/Editor/EditorEngine.h b/Engine/Source/Editor/UnrealEd/Classes/Editor/EditorEngine.h
--- a/Engine/Source/Editor/UnrealEd/Classes/Editor/EditorEngine.h
+++ b/Engine/Source/Editor/UnrealEd/Classes/Editor/EditorEngine.h
@@ -198,6 +198,10 @@ private:
 	TSet<FString> ActorLabels;
 };
 
+/**
+ * Engine that drives the Editor.
+ * Separate from UGameEngine because it may have much different functionality than desired for an instance of a game itself.
+ */
 UCLASS(config=Engine, transient)
 class UNREALED_API UEditorEngine : public UEngine
 {
diff --git a/Engine/Source/Runtime/Engine/Classes/Components/MeshComponent.h b/Engine/Source/Runtime/Engine/Classes/Components/MeshComponent.h
--- a/Engine/Source/Runtime/Engine/Classes/Components/MeshComponent.h
+++ b/Engine/Source/Runtime/Engine/Classes/Components/MeshComponent.h
@@ -8,6 +8,12 @@
 
 struct FMaterialRelevance;
 
+/**
+ * MeshComponent is an abstract base for any component that is an instance of a renderable collection of triangles.
+ *
+ * @see UStaticMeshComponent
+ * @see USkeletalMeshComponent
+ */
 UCLASS(abstract)
 class ENGINE_API UMeshComponent : public UPrimitiveComponent
 {
diff --git a/Engine/Source/Runtime/Engine/Classes/Components/ShapeComponent.h b/Engine/Source/Runtime/Engine/Classes/Components/ShapeComponent.h
--- a/Engine/Source/Runtime/Engine/Classes/Components/ShapeComponent.h
+++ b/Engine/Source/Runtime/Engine/Classes/Components/ShapeComponent.h
@@ -4,6 +4,9 @@
 #include "Components/PrimitiveComponent.h"
 #include "ShapeComponent.generated.h"
 
+/**
+ * ShapeComponent is a PrimitiveComponent that is represented by a simple geometrical shape (sphere, capsule, box, etc).
+ */
 UCLASS(abstract, hidecategories=(Object,LOD,Lighting,TextureStreaming,Activation,"Components|Activation"), editinlinenew, meta=(BlueprintSpawnableComponent), showcategories=(Mobility))
 class ENGINE_API UShapeComponent : public UPrimitiveComponent
 {
diff --git a/Engine/Source/Runtime/Engine/Classes/Engine/Engine.h b/Engine/Source/Runtime/Engine/Classes/Engine/Engine.h
--- a/Engine/Source/Runtime/Engine/Classes/Engine/Engine.h
+++ b/Engine/Source/Runtime/Engine/Classes/Engine/Engine.h
@@ -526,6 +526,10 @@ class IAnalyticsProvider;
 DECLARE_DELEGATE_OneParam(FBeginStreamingPauseDelegate, FViewport*);
 DECLARE_DELEGATE(FEndStreamingPauseDelegate);
 
+/**
+ * Abstract base class of all Engine classes, responsible for management of systems critical to editor or game systems.
+ * Also defines default classes for certain engine systems.
+ */
 UCLASS(abstract, config=Engine, defaultconfig, transient)
 class ENGINE_API UEngine
 	: public UObject
diff --git a/Engine/Source/Runtime/Engine/Classes/Engine/GameEngine.h b/Engine/Source/Runtime/Engine/Classes/Engine/GameEngine.h
--- a/Engine/Source/Runtime/Engine/Classes/Engine/GameEngine.h
+++ b/Engine/Source/Runtime/Engine/Classes/Engine/GameEngine.h
@@ -10,6 +10,9 @@
 
 #include "GameEngine.generated.h"
 
+/**
+ * Engine that manages core systems that enable a game.
+ */
 UCLASS(config=Engine, transient)
 class ENGINE_API UGameEngine : public UEngine
 {
diff --git a/Engine/Source/Runtime/Engine/Classes/GameFramework/GameUserSettings.h b/Engine/Source/Runtime/Engine/Classes/GameFramework/GameUserSettings.h
--- a/Engine/Source/Runtime/Engine/Classes/GameFramework/GameUserSettings.h
+++ b/Engine/Source/Runtime/Engine/Classes/GameFramework/GameUserSettings.h
@@ -4,6 +4,9 @@
 #include "Scalability.h"
 #include "GameUserSettings.generated.h"
 
+/**
+ * Stores user settings for a game (for example graphics and sound settings), with the ability to save and load to and from a file.
+ */
 UCLASS(config=GameUserSettings, configdonotcheckdefaults)
 class ENGINE_API UGameUserSettings : public UObject
 {
diff --git a/Engine/Source/Runtime/Engine/Classes/GameFramework/Info.h b/Engine/Source/Runtime/Engine/Classes/GameFramework/Info.h
--- a/Engine/Source/Runtime/Engine/Classes/GameFramework/Info.h
+++ b/Engine/Source/Runtime/Engine/Classes/GameFramework/Info.h
@@ -9,6 +9,10 @@
 #include "GameFramework/Actor.h"
 #include "Info.generated.h"
 
+/**
+ * Info is the base class of an Actor that isn't meant to have a physical representation in the world, used primarily
+ * for "manager" type classes that hold settings data about the world, but might need to be an Actor for replication purposes.
+ */
 UCLASS(abstract, hidecategories=(Input, Movement, Collision, Rendering, "Utilities|Transformation"), showcategories=("Input|MouseInput", "Input|TouchInput"), MinimalAPI, NotBlueprintable)
 class AInfo : public AActor
 {
diff --git a/Engine/Source/Runtime/Engine/Classes/GameFramework/PainCausingVolume.h b/Engine/Source/Runtime/Engine/Classes/GameFramework/PainCausingVolume.h
--- a/Engine/Source/Runtime/Engine/Classes/GameFramework/PainCausingVolume.h
+++ b/Engine/Source/Runtime/Engine/Classes/GameFramework/PainCausingVolume.h
@@ -9,7 +9,9 @@
 #include "GameFramework/PhysicsVolume.h"
 #include "PainCausingVolume.generated.h"
 
-
+/**
+ * Volume that causes damage over time to any Actor that overlaps its collision.
+ */
 UCLASS()
 class ENGINE_API APainCausingVolume : public APhysicsVolume
 {
