#docs - ActorComponent docs.

diff --git a/Engine/Source/Runtime/Engine/Classes/Components/ActorComponent.h b/Engine/Source/Runtime/Engine/Classes/Components/ActorComponent.h
--- a/Engine/Source/Runtime/Engine/Classes/Components/ActorComponent.h
+++ b/Engine/Source/Runtime/Engine/Classes/Components/ActorComponent.h
@@ -8,6 +8,12 @@
 
 struct FReplicationFlags;
 
+/**
+ * ActorComponent is the base class of any Component that manages or represents functionality of an Actor.
+ * Everything the player sees or interacts with in the world when playing the game is actually the work of one type of ActorComponent or another.
+ *
+ * @see [ActorComponent](https://docs.unrealengine.com/latest/INT/Programming/UnrealArchitecture/Actors/Components/index.html#actorcomponents)
+ */
 UCLASS(DefaultToInstanced, abstract, hidecategories=(ComponentReplication))
 class ENGINE_API UActorComponent : public UObject, public IInterface_AssetUserData
 {
