Fix shootergame control issues when using any controllerid but 0 to play.

diff --git a/Engine/Source/Runtime/Engine/Private/Slate/SceneViewport.cpp b/Engine/Source/Runtime/Engine/Private/Slate/SceneViewport.cpp
--- a/Engine/Source/Runtime/Engine/Private/Slate/SceneViewport.cpp
+++ b/Engine/Source/Runtime/Engine/Private/Slate/SceneViewport.cpp
@@ -362,7 +362,7 @@ FReply FSceneViewport::OnMouseButtonDown( const FGeometry& InGeometry, const FPo
 		}
 
 		// Process the mouse event
-		if( !ViewportClient->InputKey( this, 0, InMouseEvent.GetEffectingButton(), IE_Pressed ) )
+		if (!ViewportClient->InputKey(this, InMouseEvent.GetUserIndex(), InMouseEvent.GetEffectingButton(), IE_Pressed))
 		{
 			CurrentReplyState = FReply::Unhandled(); 
 		}
@@ -419,7 +419,7 @@ FReply FSceneViewport::OnMouseButtonUp( const FGeometry& InGeometry, const FPoin
 	bool bReleaseMouse = true;
 	if( ViewportClient && GetSizeXY() != FIntPoint::ZeroValue  )
 	{
-		if( !ViewportClient->InputKey(this,0,InMouseEvent.GetEffectingButton(), IE_Released ) )
+		if (!ViewportClient->InputKey(this, InMouseEvent.GetUserIndex(), InMouseEvent.GetEffectingButton(), IE_Released))
 		{
 			CurrentReplyState = FReply::Unhandled(); 
 		}
@@ -518,8 +518,8 @@ FReply FSceneViewport::OnMouseWheel( const FGeometry& InGeometry, const FPointer
 		FKey const ViewportClientKey = InMouseEvent.GetWheelDelta() < 0 ? EKeys::MouseScrollDown : EKeys::MouseScrollUp;
 
 		// Pressed and released should be sent
-		ViewportClient->InputKey( this, 0, ViewportClientKey, IE_Pressed );
-		ViewportClient->InputKey( this, 0, ViewportClientKey, IE_Released );
+		ViewportClient->InputKey(this, InMouseEvent.GetUserIndex(), ViewportClientKey, IE_Pressed);
+		ViewportClient->InputKey(this, InMouseEvent.GetUserIndex(), ViewportClientKey, IE_Released);
 	}
 	return CurrentReplyState;
 }
@@ -543,7 +543,7 @@ FReply FSceneViewport::OnMouseButtonDoubleClick( const FGeometry& InGeometry, co
 		// Switch to the viewport clients world before processing input
 		FScopedConditionalWorldSwitcher WorldSwitcher( ViewportClient );
 
-		if( !ViewportClient->InputKey( this, 0, InMouseEvent.GetEffectingButton(), IE_DoubleClick ) )
+		if (!ViewportClient->InputKey(this, InMouseEvent.GetUserIndex(), InMouseEvent.GetEffectingButton(), IE_DoubleClick))
 		{
 			CurrentReplyState = FReply::Unhandled(); 
 		}
@@ -685,7 +685,7 @@ FReply FSceneViewport::OnKeyDown( const FGeometry& InGeometry, const FKeyEvent&
 		// Switch to the viewport clients world before processing input
 		FScopedConditionalWorldSwitcher WorldSwitcher( ViewportClient );
 
-		if( !ViewportClient->InputKey( this, 0, Key, InKeyEvent.IsRepeat() ? IE_Repeat : IE_Pressed ) )
+		if (!ViewportClient->InputKey(this, InKeyEvent.GetUserIndex(), Key, InKeyEvent.IsRepeat() ? IE_Repeat : IE_Pressed))
 		{
 			CurrentReplyState = FReply::Unhandled();
 		}
@@ -706,7 +706,7 @@ FReply FSceneViewport::OnKeyUp( const FGeometry& InGeometry, const FKeyEvent& In
 		// Switch to the viewport clients world before processing input
 		FScopedConditionalWorldSwitcher WorldSwitcher( ViewportClient );
 
-		if( !ViewportClient->InputKey(this, 0, Key, IE_Released) )
+		if (!ViewportClient->InputKey(this, InKeyEvent.GetUserIndex(), Key, IE_Released))
 		{
 			CurrentReplyState = FReply::Unhandled();
 		}
