MeshParticleVertexFactory always assumed the existence of a valid dynamic parameter buffer, but we don't support them in the non-instanced case (anything <SM5), and they should probably be able to be null in other cases as well.

diff --git a/Engine/Source/Runtime/Engine/Private/MeshParticleVertexFactory.cpp b/Engine/Source/Runtime/Engine/Private/MeshParticleVertexFactory.cpp
--- a/Engine/Source/Runtime/Engine/Private/MeshParticleVertexFactory.cpp
+++ b/Engine/Source/Runtime/Engine/Private/MeshParticleVertexFactory.cpp
@@ -60,7 +60,10 @@ public:
 			SetShaderValue(RHICmdList, VertexShaderRHI, SubUVLerp, Vertex->SubUVLerp);
 			SetShaderValue(RHICmdList, VertexShaderRHI, ParticleDirection, Vertex->Velocity);
 			SetShaderValue(RHICmdList, VertexShaderRHI, RelativeTime, Vertex->RelativeTime);
-			SetShaderValue(RHICmdList, VertexShaderRHI, DynamicParameter, FVector4(DynamicVertex->DynamicValue[0], DynamicVertex->DynamicValue[1], DynamicVertex->DynamicValue[2], DynamicVertex->DynamicValue[3]));
+			if (BatchParameters->DynamicParameterBuffer)
+			{
+				SetShaderValue(RHICmdList, VertexShaderRHI, DynamicParameter, FVector4(DynamicVertex->DynamicValue[0], DynamicVertex->DynamicValue[1], DynamicVertex->DynamicValue[2], DynamicVertex->DynamicValue[3]));
+			}
 			SetShaderValue(RHICmdList, VertexShaderRHI, ParticleColor, FVector4(Vertex->Color.Component(0), Vertex->Color.Component(1), Vertex->Color.Component(2), Vertex->Color.Component(3)));
 		}
 	}
