Fixed crash when layer node data is corrupted, fix up manually.

https://jira.ol.epicgames.net/browse/OR-83

diff --git a/Engine/Source/Editor/AnimGraph/Classes/AnimGraphNode_LayeredBoneBlend.h b/Engine/Source/Editor/AnimGraph/Classes/AnimGraphNode_LayeredBoneBlend.h
--- a/Engine/Source/Editor/AnimGraph/Classes/AnimGraphNode_LayeredBoneBlend.h
+++ b/Engine/Source/Editor/AnimGraph/Classes/AnimGraphNode_LayeredBoneBlend.h
@@ -13,6 +13,10 @@ class UAnimGraphNode_LayeredBoneBlend : public UAnimGraphNode_BlendListBase
 	UPROPERTY(EditAnywhere, Category=Settings)
 	FAnimNode_LayeredBoneBlend Node;
 
+	// UObject interface
+	virtual void Serialize(FArchive& Ar) override;
+	// End of UObject interface
+
 	// Adds a new pose pin
 	//@TODO: Generalize this behavior (returning a list of actions/delegates maybe?)
 	ANIMGRAPH_API virtual void AddPinToBlendByFilter();
diff --git a/Engine/Source/Editor/AnimGraph/Private/AnimGraphNode_LayeredBoneBlend.cpp b/Engine/Source/Editor/AnimGraph/Private/AnimGraphNode_LayeredBoneBlend.cpp
--- a/Engine/Source/Editor/AnimGraph/Private/AnimGraphNode_LayeredBoneBlend.cpp
+++ b/Engine/Source/Editor/AnimGraph/Private/AnimGraphNode_LayeredBoneBlend.cpp
@@ -14,6 +14,7 @@
 UAnimGraphNode_LayeredBoneBlend::UAnimGraphNode_LayeredBoneBlend(const FObjectInitializer& ObjectInitializer)
 	: Super(ObjectInitializer)
 {
+	Node.AddPose();
 }
 
 FLinearColor UAnimGraphNode_LayeredBoneBlend::GetNodeTitleColor() const
@@ -90,4 +91,9 @@ void UAnimGraphNode_LayeredBoneBlend::GetContextMenuActions(const FGraphNodeCont
 	}
 }
 
+void UAnimGraphNode_LayeredBoneBlend::Serialize(FArchive& Ar)
+{
+	Super::Serialize(Ar);
+	Node.ValidateData();
+}
 #undef LOCTEXT_NAMESPACE
\ No newline at end of file
diff --git a/Engine/Source/Runtime/Engine/Classes/Animation/AnimNode_LayeredBoneBlend.h b/Engine/Source/Runtime/Engine/Classes/Animation/AnimNode_LayeredBoneBlend.h
--- a/Engine/Source/Runtime/Engine/Classes/Animation/AnimNode_LayeredBoneBlend.h
+++ b/Engine/Source/Runtime/Engine/Classes/Animation/AnimNode_LayeredBoneBlend.h
@@ -38,9 +38,6 @@ protected:
 public:	
 	FAnimNode_LayeredBoneBlend()
 	{
-#if WITH_EDITOR
-		AddPose();
-#endif
 	}
 
 	// FAnimNode_Base interface
@@ -65,6 +62,11 @@ public:
 		BlendPoses.RemoveAt(PoseIndex);
 		LayerSetup.RemoveAt(PoseIndex);
 	}
+
+	// ideally you don't like to get to situation where it becomes inconsistent, but this happened, 
+	// and we don't know what caused this. Possibly copy/paste, but I tried copy/paste and that didn't work
+	// so here we add code to fix this up manually in editor, so that they can continue working on it. 
+	void ValidateData();
 #endif
 
 	void ReinitializeBoneBlendWeights(const FBoneContainer& RequiredBones, const USkeleton* Skeleton);
diff --git a/Engine/Source/Runtime/Engine/Private/Animation/AnimNode_LayeredBoneBlend.cpp b/Engine/Source/Runtime/Engine/Private/Animation/AnimNode_LayeredBoneBlend.cpp
--- a/Engine/Source/Runtime/Engine/Private/Animation/AnimNode_LayeredBoneBlend.cpp
+++ b/Engine/Source/Runtime/Engine/Private/Animation/AnimNode_LayeredBoneBlend.cpp
@@ -126,4 +126,38 @@ void FAnimNode_LayeredBoneBlend::GatherDebugData(FNodeDebugData& DebugData)
 	{
 		BlendPoses[ChildIndex].GatherDebugData(DebugData.BranchFlow(BlendWeights[ChildIndex]));
 	}
+}
+
+void FAnimNode_LayeredBoneBlend::ValidateData()
+{
+	// ideally you don't like to get to situation where it becomes inconsistent, but this happened, 
+	// and we don't know what caused this. Possibly copy/paste, but I tried copy/paste and that didn't work
+	// so here we add code to fix this up manually in editor, so that they can continue working on it. 
+	int32 PoseNum = BlendPoses.Num();
+	int32 WeightNum = BlendWeights.Num();
+	int32 LayerNum = LayerSetup.Num();
+
+	int32 Max = FMath::Max3(PoseNum, WeightNum, LayerNum);
+	int32 Min = FMath::Min3(PoseNum, WeightNum, LayerNum);
+	// if they are not all same
+	if (Min != Max)
+	{
+		// we'd like to increase to all Max
+		// sadly we don't have add X for how many
+		for (int32 Index=PoseNum; Index<Max; ++Index)
+		{
+			BlendPoses.Add(FPoseLink());
+		}
+
+		for(int32 Index=WeightNum; Index<Max; ++Index)
+		{
+			BlendWeights.Add(1.f);
+		}
+
+		for(int32 Index=LayerNum; Index<Max; ++Index)
+		{
+			LayerSetup.Add(FInputBlendPose());
+		}
+	}
+
 }
\ No newline at end of file
