#ue4 - Simulated proxies should check the replicated server state when testing if component had failed to resolve on client. Previously it was checking the local version which could be incorrect on the client.

diff --git a/Engine/Source/Runtime/Engine/Private/CharacterMovementComponent.cpp b/Engine/Source/Runtime/Engine/Private/CharacterMovementComponent.cpp
--- a/Engine/Source/Runtime/Engine/Private/CharacterMovementComponent.cpp
+++ b/Engine/Source/Runtime/Engine/Private/CharacterMovementComponent.cpp
@@ -1060,7 +1060,7 @@ void UCharacterMovementComponent::SimulateMovement(float DeltaSeconds)
 	}
 
 	// If base is not resolved on the client, we should not try to simulate at all
-	if (CharacterOwner->GetBasedMovement().IsBaseUnresolved())
+	if (CharacterOwner->GetReplicatedBasedMovement().IsBaseUnresolved())
 	{
 		UE_LOG(LogCharacterMovement, Verbose, TEXT("Base for simulated character '%s' is not resolved on client, skipping SimulateMovement"), *CharacterOwner->GetName());
 		return;
