Fixed crash when a matinee sequence calls an event which stops the sequence prematurely.

diff --git a/Engine/Source/Runtime/Engine/Private/Interpolation.cpp b/Engine/Source/Runtime/Engine/Private/Interpolation.cpp
--- a/Engine/Source/Runtime/Engine/Private/Interpolation.cpp
+++ b/Engine/Source/Runtime/Engine/Private/Interpolation.cpp
@@ -716,6 +716,18 @@ void AMatineeActor::UpdateInterp( float NewPosition, bool bPreview, bool bJump )
 			for( int32 GroupIndex = 0; GroupIndex < Groups.Num(); ++GroupIndex )
 			{
 				Groups[GroupIndex]->Group->UpdateGroup( NewPosition, Groups[GroupIndex], bPreview, bJump );
+
+				const bool bhasBeenTerminated = (GroupInst.Num() == 0);
+#if WITH_EDITORONLY_DATA
+				if (bhasBeenTerminated && !bIsBeingEdited)
+#else
+				if (bhasBeenTerminated)
+#endif
+				{
+					UE_LOG(LogMatinee, Log, TEXT("WARNING: A matinee was stopped while updating group '%s'; the next groups will not be updated."), *Groups[GroupIndex]->Group->GetFullGroupName(true));
+					InterpPosition = NewPosition;
+					return;
+				}
 			}
 		}
 
