Monolithic build CIS fix

diff --git a/Engine/Source/Runtime/Engine/Private/Animation/AnimNode_LayeredBoneBlend.cpp b/Engine/Source/Runtime/Engine/Private/Animation/AnimNode_LayeredBoneBlend.cpp
--- a/Engine/Source/Runtime/Engine/Private/Animation/AnimNode_LayeredBoneBlend.cpp
+++ b/Engine/Source/Runtime/Engine/Private/Animation/AnimNode_LayeredBoneBlend.cpp
@@ -128,6 +128,7 @@ void FAnimNode_LayeredBoneBlend::GatherDebugData(FNodeDebugData& DebugData)
 	}
 }
 
+#if WITH_EDITOR
 void FAnimNode_LayeredBoneBlend::ValidateData()
 {
 	// ideally you don't like to get to situation where it becomes inconsistent, but this happened, 
@@ -159,5 +160,5 @@ void FAnimNode_LayeredBoneBlend::ValidateData()
 			LayerSetup.Add(FInputBlendPose());
 		}
 	}
-
-}
\ No newline at end of file
+}
+#endif
\ No newline at end of file
