Data table code doesn't check HACK_HEADER_GENERATOR anymore.

diff --git a/Engine/Source/Runtime/Engine/Classes/Engine/DataTable.h b/Engine/Source/Runtime/Engine/Classes/Engine/DataTable.h
--- a/Engine/Source/Runtime/Engine/Classes/Engine/DataTable.h
+++ b/Engine/Source/Runtime/Engine/Classes/Engine/DataTable.h
@@ -119,7 +119,7 @@ class UDataTable : public UObject
 	/** Empty the table info (will not clear RowStruct) */
 	ENGINE_API void EmptyTable();
 
-#if WITH_EDITOR || HACK_HEADER_GENERATOR
+#if WITH_EDITOR
 
 private:
 	//when RowStruct is being modified, row data is stored serialized with tags
@@ -159,7 +159,7 @@ public:
 
 	/** Get array for each row in the table. The first row is the titles*/
 	ENGINE_API TArray< TArray<FString> > GetTableData() const;
-#endif //WITH_EDITOR || HACK_HEADER_GENERATOR
+#endif //WITH_EDITOR
 
 	ENGINE_API static FString AssignStringToProperty(const FString& InString, const UProperty* InProp, uint8* InData);
 
diff --git a/Engine/Source/Runtime/Engine/Private/DataTable.cpp b/Engine/Source/Runtime/Engine/Private/DataTable.cpp
--- a/Engine/Source/Runtime/Engine/Private/DataTable.cpp
+++ b/Engine/Source/Runtime/Engine/Private/DataTable.cpp
@@ -182,9 +182,9 @@ void UDataTable::AddReferencedObjects(UObject* InThis, FReferenceCollector& Coll
 		}
 	}
 
-#if WITH_EDITOR || HACK_HEADER_GENERATOR
+#if WITH_EDITOR
 	Collector.AddReferencedObjects(This->TemporarilyReferencedObjects);
-#endif //WITH_EDITOR || HACK_HEADER_GENERATOR
+#endif //WITH_EDITOR
 
 	Super::AddReferencedObjects( This, Collector );
 }
@@ -240,7 +240,7 @@ UProperty* UDataTable::FindTableProperty(const FName& PropertyName) const
 	return Property;
 }
 
-#if WITH_EDITOR || HACK_HEADER_GENERATOR
+#if WITH_EDITOR
 
 struct FPropertyDisplayNameHelper
 {
@@ -670,4 +670,4 @@ TArray< TArray<FString> > UDataTable::GetTableData() const
 
 }
 
-#endif //WITH_EDITOR || HACK_HEADER_GENERATOR
\ No newline at end of file
+#endif //WITH_EDITOR
\ No newline at end of file
