Move unwelding code out of FBodyInstance::TermBody and put into UPrimitiveComponent where it makes more sense.
Fixes cases where BodyInstance is considered invalid and so TermBody is never called which orphans shapes inside PhysX

diff --git a/Engine/Source/Runtime/Engine/Classes/Components/PrimitiveComponent.h b/Engine/Source/Runtime/Engine/Classes/Components/PrimitiveComponent.h
--- a/Engine/Source/Runtime/Engine/Classes/Components/PrimitiveComponent.h
+++ b/Engine/Source/Runtime/Engine/Classes/Components/PrimitiveComponent.h
@@ -1127,6 +1127,11 @@ public:
 	*/
 	virtual void UnWeldFromParent();
 
+	/**
+	*   Unwelds the children of this component. Attachment is maintained
+	*/
+	virtual void UnWeldChildren();
+
 	/**
 	*	Adds the bodies that are currently welded to the OutWeldedBodies array 
 	*/
diff --git a/Engine/Source/Runtime/Engine/Private/PhysicsEngine/BodyInstance.cpp b/Engine/Source/Runtime/Engine/Private/PhysicsEngine/BodyInstance.cpp
--- a/Engine/Source/Runtime/Engine/Private/PhysicsEngine/BodyInstance.cpp
+++ b/Engine/Source/Runtime/Engine/Private/PhysicsEngine/BodyInstance.cpp
@@ -1503,28 +1503,6 @@ void TermBodyHelper(int32& SceneIndex, PxRigidActor*& PRigidActor, FBodyInstance
  */
 void FBodyInstance::TermBody()
 {
-	if (UPrimitiveComponent* OwnerComponentInst = OwnerComponent.Get())
-	{
-		OwnerComponentInst->UnWeldFromParent();
-
-		for (USceneComponent * Child : OwnerComponentInst->AttachChildren)
-		{
-			if (UPrimitiveComponent * PrimChild = Cast<UPrimitiveComponent>(Child))
-			{
-				if (FBodyInstance* ChildBI = PrimChild->GetBodyInstance())
-				{
-					if (ChildBI->bWelded)
-					{
-						if (ChildBI->OwnerComponent.IsValid())
-						{
-							ChildBI->OwnerComponent->UnWeldFromParent();
-						}
-					}
-				}
-			}
-		}
-	}
-
 #if WITH_BOX2D
 	if (BodyInstancePtr != NULL)
 	{
@@ -1567,7 +1545,7 @@ void FBodyInstance::TermBody()
 	{
 		DOFConstraint->TermConstraint();
 		FConstraintInstance::Free(DOFConstraint);
-		DOFConstraint = NULL;
+			DOFConstraint = NULL;
 	}
 	
 
diff --git a/Engine/Source/Runtime/Engine/Private/PrimitiveComponent.cpp b/Engine/Source/Runtime/Engine/Private/PrimitiveComponent.cpp
--- a/Engine/Source/Runtime/Engine/Private/PrimitiveComponent.cpp
+++ b/Engine/Source/Runtime/Engine/Private/PrimitiveComponent.cpp
@@ -428,6 +428,10 @@ void UPrimitiveComponent::SendPhysicsTransform(bool bTeleport)
 
 void UPrimitiveComponent::DestroyPhysicsState()
 {
+	// we remove welding related to this component
+	UnWeldFromParent();
+	UnWeldChildren();
+
 	// clean up physics engine representation
 	if(BodyInstance.IsValidBodyInstance())
 	{
diff --git a/Engine/Source/Runtime/Engine/Private/PrimitiveComponentPhysics.cpp b/Engine/Source/Runtime/Engine/Private/PrimitiveComponentPhysics.cpp
--- a/Engine/Source/Runtime/Engine/Private/PrimitiveComponentPhysics.cpp
+++ b/Engine/Source/Runtime/Engine/Private/PrimitiveComponentPhysics.cpp
@@ -703,6 +703,17 @@ void UPrimitiveComponent::UnWeldFromParent()
 	}
 }
 
+void UPrimitiveComponent::UnWeldChildren()
+{
+	for (USceneComponent* ChildComponent : AttachChildren)
+	{
+		if (UPrimitiveComponent* PrimComp = Cast<UPrimitiveComponent>(ChildComponent))
+		{
+			PrimComp->UnWeldFromParent();
+		}
+	}
+}
+
 FBodyInstance* UPrimitiveComponent::GetBodyInstance(FName BoneName, bool bGetWelded) const
 {
 	if (bGetWelded && BodyInstance.bWelded)
diff --git a/Engine/Source/Runtime/Engine/Private/SkeletalMeshComponentPhysics.cpp b/Engine/Source/Runtime/Engine/Private/SkeletalMeshComponentPhysics.cpp
--- a/Engine/Source/Runtime/Engine/Private/SkeletalMeshComponentPhysics.cpp
+++ b/Engine/Source/Runtime/Engine/Private/SkeletalMeshComponentPhysics.cpp
@@ -863,14 +863,9 @@ void USkeletalMeshComponent::TermArticulated()
 	{
 		check( Bodies[i] );
 		Bodies[i]->TermBody();
-	}
-	
-	for(int32 i=0; i<Bodies.Num(); i++)
-	{
-		check( Bodies[i] );
 		delete Bodies[i];
 	}
-
+	
 	Bodies.Empty();
 
 #if WITH_PHYSX
@@ -1315,6 +1310,8 @@ void USkeletalMeshComponent::DestroyPhysicsState()
 {
 	if (bEnablePerPolyCollision == false)
 	{
+		UnWeldFromParent();
+		UnWeldChildren();
 		TermArticulated();
 	}
 
