hide metallic pin for SubsurfaceProfile (not supported at the moment)

diff --git a/Engine/Source/Runtime/Engine/Private/Materials/Material.cpp b/Engine/Source/Runtime/Engine/Private/Materials/Material.cpp
--- a/Engine/Source/Runtime/Engine/Private/Materials/Material.cpp
+++ b/Engine/Source/Runtime/Engine/Private/Materials/Material.cpp
@@ -3653,7 +3653,7 @@ bool UMaterial::IsPropertyActive(EMaterialProperty InProperty) const
 		Active = ShadingModel != MSM_Unlit;
 		break;
 	case MP_Metallic:
-		Active = ShadingModel != MSM_Unlit && ShadingModel != MSM_Subsurface && ShadingModel != MSM_PreintegratedSkin;
+		Active = ShadingModel != MSM_Unlit && ShadingModel != MSM_Subsurface && ShadingModel != MSM_PreintegratedSkin && ShadingModel != MSM_SubsurfaceProfile;
 		break;
 	case MP_Normal:
 		Active = ShadingModel != MSM_Unlit || Refraction.IsConnected();
