Fixes for MT Anim Evaluation in Unreal Tournament

diff --git a/Engine/Source/Runtime/Engine/Classes/Components/SkeletalMeshComponent.h b/Engine/Source/Runtime/Engine/Classes/Components/SkeletalMeshComponent.h
--- a/Engine/Source/Runtime/Engine/Classes/Components/SkeletalMeshComponent.h
+++ b/Engine/Source/Runtime/Engine/Classes/Components/SkeletalMeshComponent.h
@@ -1159,7 +1159,7 @@ public:
 	void ParallelAnimationEvaluation() { PerformAnimationEvaluation(AnimEvaluationContext.SkeletalMesh, AnimEvaluationContext.AnimInstance, AnimEvaluationContext.SpaceBases, AnimEvaluationContext.LocalAtoms, AnimEvaluationContext.VertexAnims, AnimEvaluationContext.RootBoneTranslation); }
 	void CompleteParallelAnimationEvaluation()
 	{
-		if (AnimEvaluationContext.AnimInstance == AnimScriptInstance && AnimEvaluationContext.SkeletalMesh == SkeletalMesh)
+		if ((AnimEvaluationContext.AnimInstance == AnimScriptInstance) && (AnimEvaluationContext.SkeletalMesh == SkeletalMesh) && (AnimEvaluationContext.SpaceBases.Num() == SpaceBases.Num()))
 		{
 			Exchange(AnimEvaluationContext.SpaceBases, AnimEvaluationContext.bDoInterpolation ? CachedSpaceBases : SpaceBases);
 			Exchange(AnimEvaluationContext.LocalAtoms, AnimEvaluationContext.bDoInterpolation ? CachedLocalAtoms : LocalAtoms);
diff --git a/Engine/Source/Runtime/Engine/Private/SkeletalMeshComponent.cpp b/Engine/Source/Runtime/Engine/Private/SkeletalMeshComponent.cpp
--- a/Engine/Source/Runtime/Engine/Private/SkeletalMeshComponent.cpp
+++ b/Engine/Source/Runtime/Engine/Private/SkeletalMeshComponent.cpp
@@ -649,8 +649,8 @@ void USkeletalMeshComponent::FillSpaceBases(const USkeletalMesh* InSkeletalMesh,
 #endif
 		FTransform::Multiply(SpaceBasesData + BoneIndex, LocalTransformsData + BoneIndex, SpaceBasesData + ParentIndex);
 
-		checkSlow( SpaceBases[BoneIndex].IsRotationNormalized() );
-		checkSlow( !SpaceBases[BoneIndex].ContainsNaN() );
+		checkSlow(DestSpaceBases[BoneIndex].IsRotationNormalized());
+		checkSlow(!DestSpaceBases[BoneIndex].ContainsNaN());
 	}
 
 	/**
