Don't assert when reinitializing a previously hidden streaming level that contains an emitter.

diff --git a/Engine/Source/Runtime/Engine/Private/Particles/ParticleComponents.cpp b/Engine/Source/Runtime/Engine/Private/Particles/ParticleComponents.cpp
--- a/Engine/Source/Runtime/Engine/Private/Particles/ParticleComponents.cpp
+++ b/Engine/Source/Runtime/Engine/Private/Particles/ParticleComponents.cpp
@@ -359,7 +359,7 @@ void AEmitter::PostInitializeComponents()
 	// Set Notification Delegate
 	if (ParticleSystemComponent)
 	{
-		ParticleSystemComponent->OnSystemFinished.AddDynamic( this, &AEmitter::OnParticleSystemFinished );
+		ParticleSystemComponent->OnSystemFinished.AddUniqueDynamic( this, &AEmitter::OnParticleSystemFinished );
 		bCurrentlyActive = ParticleSystemComponent->bAutoActivate;
 	}
 
