Clear Lifespan timers of Actors when removing the streaming level they belong to.
Log a warning rather than crashing when calling Destroy on an Actor that belongs to a Level that has no OwningLevel (hidden streaming levels essentially).

diff --git a/Engine/Source/Runtime/Engine/Private/Actor.cpp b/Engine/Source/Runtime/Engine/Private/Actor.cpp
--- a/Engine/Source/Runtime/Engine/Private/Actor.cpp
+++ b/Engine/Source/Runtime/Engine/Private/Actor.cpp
@@ -1519,6 +1519,12 @@ void AActor::EndPlay(const EEndPlayReason::Type EndPlayReason)
 	{
 		bActorInitialized = false;
 		GetWorld()->RemoveNetworkActor(this);
+
+		// Clear any ticking lifespan timers
+		if (InitialLifeSpan > 0.f)
+		{
+			SetLifeSpan(0.f);
+		}
 	}
 
 	UNavigationSystem::OnActorUnregistered(this);
@@ -2649,7 +2655,15 @@ bool AActor::Destroy( bool bNetForce, bool bShouldModifyLevel )
 	// It's already pending kill, no need to beat the corpse
 	if (!IsPendingKill())
 	{
-		GetWorld()->DestroyActor( this, bNetForce, bShouldModifyLevel );
+		UWorld* World = GetWorld();
+		if (World)
+		{
+			World->DestroyActor( this, bNetForce, bShouldModifyLevel );
+		}
+		else
+		{
+			UE_LOG(LogSpawn, Warning, TEXT("Destroying %s, which doesn't have a valid world pointer"), *GetPathName());
+		}
 	}
 	return IsPendingKill();
 }
@@ -3380,7 +3394,7 @@ void AActor::SetLifeSpan( float InLifespan )
 float AActor::GetLifeSpan() const
 {
 	// Timer remaining returns -1.0f if there is no such timer - return this as ZERO
-	float CurrentLifespan = GetWorldTimerManager().GetTimerRemaining( this, &AActor::LifeSpanExpired );
+	const float CurrentLifespan = GetWorldTimerManager().GetTimerRemaining( this, &AActor::LifeSpanExpired );
 	return ( CurrentLifespan != -1.0f ) ? CurrentLifespan : 0.0f;
 }
 
