ok Some reason it didnt take trying again, fixing a get on SlateRHIRenderer that hasnt been loading

diff --git a/Engine/Source/Runtime/Engine/Private/UserInterface/EngineFontServices.cpp b/Engine/Source/Runtime/Engine/Private/UserInterface/EngineFontServices.cpp
--- a/Engine/Source/Runtime/Engine/Private/UserInterface/EngineFontServices.cpp
+++ b/Engine/Source/Runtime/Engine/Private/UserInterface/EngineFontServices.cpp
@@ -14,7 +14,10 @@ FEngineFontServices::FEngineFontServices()
 	check(IsInGameThread());
 
 #if !UE_SERVER
-	RenderThreadFontAtlasFactory = FModuleManager::Get().GetModuleChecked<ISlateRHIRendererModule>("SlateRHIRenderer").CreateSlateFontAtlasFactory();
+	if (!IsRunningDedicatedServer())
+	{
+		RenderThreadFontAtlasFactory = FModuleManager::Get().GetModuleChecked<ISlateRHIRendererModule>("SlateRHIRenderer").CreateSlateFontAtlasFactory();
+	}
 #endif
 }
 
