Fixed crash from realtime preview of parameter default value in some cases

diff --git a/Engine/Source/Runtime/Engine/Private/Materials/Material.cpp b/Engine/Source/Runtime/Engine/Private/Materials/Material.cpp
--- a/Engine/Source/Runtime/Engine/Private/Materials/Material.cpp
+++ b/Engine/Source/Runtime/Engine/Private/Materials/Material.cpp
@@ -681,13 +681,15 @@ void UMaterial::OverrideVectorParameterDefault(FName ParameterName, const FLinea
 	// Iterate over each of the material's texture expressions.
 	for (FMaterialUniformExpression* UniformExpression : UniformExpressions)
 	{
-		check(UniformExpression->GetType() == &FMaterialUniformExpressionVectorParameter::StaticType);
-		FMaterialUniformExpressionVectorParameter* VectorExpression = static_cast<FMaterialUniformExpressionVectorParameter*>(UniformExpression);
-
-		if (VectorExpression->GetParameterName() == ParameterName)
+		if (UniformExpression->GetType() == &FMaterialUniformExpressionVectorParameter::StaticType)
 		{
-			VectorExpression->SetTransientOverrideDefaultValue(Value, bOverride);
-			bShouldRecacheMaterialExpressions = true;
+			FMaterialUniformExpressionVectorParameter* VectorExpression = static_cast<FMaterialUniformExpressionVectorParameter*>(UniformExpression);
+
+			if (VectorExpression->GetParameterName() == ParameterName)
+			{
+				VectorExpression->SetTransientOverrideDefaultValue(Value, bOverride);
+				bShouldRecacheMaterialExpressions = true;
+			}
 		}
 	}
 
@@ -711,13 +713,15 @@ void UMaterial::OverrideScalarParameterDefault(FName ParameterName, float Value,
 	// Iterate over each of the material's texture expressions.
 	for (FMaterialUniformExpression* UniformExpression : UniformExpressions)
 	{
-		check(UniformExpression->GetType() == &FMaterialUniformExpressionScalarParameter::StaticType);
-		FMaterialUniformExpressionScalarParameter* ScalarExpression = static_cast<FMaterialUniformExpressionScalarParameter*>(UniformExpression);
-
-		if (ScalarExpression->GetParameterName() == ParameterName)
+		if (UniformExpression->GetType() == &FMaterialUniformExpressionScalarParameter::StaticType)
 		{
-			ScalarExpression->SetTransientOverrideDefaultValue(Value, bOverride);
-			bShouldRecacheMaterialExpressions = true;
+			FMaterialUniformExpressionScalarParameter* ScalarExpression = static_cast<FMaterialUniformExpressionScalarParameter*>(UniformExpression);
+
+			if (ScalarExpression->GetParameterName() == ParameterName)
+			{
+				ScalarExpression->SetTransientOverrideDefaultValue(Value, bOverride);
+				bShouldRecacheMaterialExpressions = true;
+			}
 		}
 	}
 
diff --git a/Engine/Source/Runtime/Engine/Private/Materials/MaterialInstance.cpp b/Engine/Source/Runtime/Engine/Private/Materials/MaterialInstance.cpp
--- a/Engine/Source/Runtime/Engine/Private/Materials/MaterialInstance.cpp
+++ b/Engine/Source/Runtime/Engine/Private/Materials/MaterialInstance.cpp
@@ -789,13 +789,15 @@ void UMaterialInstance::OverrideVectorParameterDefault(FName ParameterName, cons
 		for (int32 ExpressionIndex = 0; ExpressionIndex < UniformExpressions.Num(); ExpressionIndex++)
 		{
 			FMaterialUniformExpression* UniformExpression = UniformExpressions[ExpressionIndex];
-			check(UniformExpression->GetType() == &FMaterialUniformExpressionVectorParameter::StaticType);
-			FMaterialUniformExpressionVectorParameter* VectorExpression = static_cast<FMaterialUniformExpressionVectorParameter*>(UniformExpression);
-
-			if (VectorExpression->GetParameterName() == ParameterName)
+			if (UniformExpression->GetType() == &FMaterialUniformExpressionVectorParameter::StaticType)
 			{
-				VectorExpression->SetTransientOverrideDefaultValue(Value, bOverride);
-				bShouldRecacheMaterialExpressions = true;
+				FMaterialUniformExpressionVectorParameter* VectorExpression = static_cast<FMaterialUniformExpressionVectorParameter*>(UniformExpression);
+
+				if (VectorExpression->GetParameterName() == ParameterName)
+				{
+					VectorExpression->SetTransientOverrideDefaultValue(Value, bOverride);
+					bShouldRecacheMaterialExpressions = true;
+				}
 			}
 		}
 	}
@@ -822,13 +824,15 @@ void UMaterialInstance::OverrideScalarParameterDefault(FName ParameterName, floa
 		for (int32 ExpressionIndex = 0; ExpressionIndex < UniformExpressions.Num(); ExpressionIndex++)
 		{
 			FMaterialUniformExpression* UniformExpression = UniformExpressions[ExpressionIndex];
-			check(UniformExpression->GetType() == &FMaterialUniformExpressionScalarParameter::StaticType);
-			FMaterialUniformExpressionScalarParameter* ScalarExpression = static_cast<FMaterialUniformExpressionScalarParameter*>(UniformExpression);
-
-			if (ScalarExpression->GetParameterName() == ParameterName)
+			if (UniformExpression->GetType() == &FMaterialUniformExpressionScalarParameter::StaticType)
 			{
-				ScalarExpression->SetTransientOverrideDefaultValue(Value, bOverride);
-				bShouldRecacheMaterialExpressions = true;
+				FMaterialUniformExpressionScalarParameter* ScalarExpression = static_cast<FMaterialUniformExpressionScalarParameter*>(UniformExpression);
+
+				if (ScalarExpression->GetParameterName() == ParameterName)
+				{
+					ScalarExpression->SetTransientOverrideDefaultValue(Value, bOverride);
+					bShouldRecacheMaterialExpressions = true;
+				}
 			}
 		}
 	}
