Fix for editor build - explained by pch differences between fornite-simple and ue4?

diff --git a/Engine/Source/Runtime/Engine/Private/SkeletalMeshMerge.cpp b/Engine/Source/Runtime/Engine/Private/SkeletalMeshMerge.cpp
--- a/Engine/Source/Runtime/Engine/Private/SkeletalMeshMerge.cpp
+++ b/Engine/Source/Runtime/Engine/Private/SkeletalMeshMerge.cpp
@@ -5,6 +5,7 @@
 =============================================================================*/
 
 #include "EnginePrivate.h"
+#include "Animation/Skeleton.h"
 #include "RawIndexBuffer.h"
 #include "SkeletalMeshMerge.h"
 #include "Engine/SkeletalMeshSocket.h"
