Fix cases where scene view state may not have referenced objects added for GC

diff --git a/Engine/Source/Editor/UnrealEd/Private/EditorViewportClient.cpp b/Engine/Source/Editor/UnrealEd/Private/EditorViewportClient.cpp
--- a/Engine/Source/Editor/UnrealEd/Private/EditorViewportClient.cpp
+++ b/Engine/Source/Editor/UnrealEd/Private/EditorViewportClient.cpp
@@ -2408,6 +2408,11 @@ void FEditorViewportClient::AddReferencedObjects( FReferenceCollector& Collector
 	{
 		PreviewScene->AddReferencedObjects( Collector );
 	}
+
+	if (ViewState.GetReference())
+	{
+		ViewState.GetReference()->AddReferencedObjects(Collector);
+	}
 }
 
 FSceneInterface* FEditorViewportClient::GetScene() const
diff --git a/Engine/Source/Runtime/Engine/Private/LocalPlayer.cpp b/Engine/Source/Runtime/Engine/Private/LocalPlayer.cpp
--- a/Engine/Source/Runtime/Engine/Private/LocalPlayer.cpp
+++ b/Engine/Source/Runtime/Engine/Private/LocalPlayer.cpp
@@ -1620,12 +1620,17 @@ void ULocalPlayer::AddReferencedObjects(UObject* InThis, FReferenceCollector& Co
 	ULocalPlayer* This = CastChecked<ULocalPlayer>(InThis);
 
 	FSceneViewStateInterface* Ref = This->ViewState.GetReference();
-
 	if(Ref)
 	{
 		Ref->AddReferencedObjects(Collector);
 	}
 
+	FSceneViewStateInterface* StereoRef = This->StereoViewState.GetReference();
+	if (StereoRef)
+	{
+		StereoRef->AddReferencedObjects(Collector);
+	}
+
 	UPlayer::AddReferencedObjects(This, Collector);
 }
 
