#ue4 - Mark CharacterMovementComponent::PhysicsVolumeChanged as an override.

diff --git a/Engine/Source/Runtime/Engine/Classes/GameFramework/CharacterMovementComponent.h b/Engine/Source/Runtime/Engine/Classes/GameFramework/CharacterMovementComponent.h
--- a/Engine/Source/Runtime/Engine/Classes/GameFramework/CharacterMovementComponent.h
+++ b/Engine/Source/Runtime/Engine/Classes/GameFramework/CharacterMovementComponent.h
@@ -1117,7 +1117,7 @@ public:
 	virtual void PhysicsRotation(float DeltaTime);
 
 	/** Delegate when PhysicsVolume of UpdatedComponent has been changed **/
-	virtual void PhysicsVolumeChanged(class APhysicsVolume* NewVolume);
+	virtual void PhysicsVolumeChanged(class APhysicsVolume* NewVolume) override;
 
 	/** Set movement mode to the default based on the current physics volume. */
 	virtual void SetDefaultMovementMode();
