#ue4 - Remove dead code in PhysFalling().

diff --git a/Engine/Source/Runtime/Engine/Private/CharacterMovementComponent.cpp b/Engine/Source/Runtime/Engine/Private/CharacterMovementComponent.cpp
--- a/Engine/Source/Runtime/Engine/Private/CharacterMovementComponent.cpp
+++ b/Engine/Source/Runtime/Engine/Private/CharacterMovementComponent.cpp
@@ -2933,20 +2933,6 @@ void UCharacterMovementComponent::PhysFalling(float deltaTime, int32 Iterations)
 			NotifyJumpApex();
 		}
 
-		// Not really necessary to limit air control since the first deflection uses non air control velocity for deflection!
-		/*
-		if (bHasAirControl)
-		{
-			const float LookAheadTime = 0.02f;
-			FHitResult TestAirResult(1.f);
-			if (FindAirControlImpact(timeTick, LookAheadTime, Velocity, FallAcceleration, Gravity, TestAirResult))
-			{
-				const FVector LimitedAirAccel = LimitAirControl(timeTick, AirControlAccel, TestAirResult, true);
-				Velocity = VelocityNoAirControl + (LimitedAirAccel * timeTick);
-			}
-		}
-		*/
-
 
 		// Move
 		FHitResult Hit(1.f);
