#docs - Updated ActorComponent with a better description.

diff --git a/Engine/Source/Runtime/Engine/Classes/Components/ActorComponent.h b/Engine/Source/Runtime/Engine/Classes/Components/ActorComponent.h
--- a/Engine/Source/Runtime/Engine/Classes/Components/ActorComponent.h
+++ b/Engine/Source/Runtime/Engine/Classes/Components/ActorComponent.h
@@ -9,10 +9,12 @@
 struct FReplicationFlags;
 
 /**
- * ActorComponent is the base class of any Component that manages or represents functionality of an Actor.
- * Everything the player sees or interacts with in the world when playing the game is actually the work of one type of ActorComponent or another.
+ * ActorComponent is the base class for components that define reusable behavior that can be added to different types of Actors.
+ * ActorComponents that have a transform are known as SceneComponents and those that can be rendered are PrimitiveComponents.
  *
  * @see [ActorComponent](https://docs.unrealengine.com/latest/INT/Programming/UnrealArchitecture/Actors/Components/index.html#actorcomponents)
+ * @see USceneComponent
+ * @see UPrimitiveComponent
  */
 UCLASS(DefaultToInstanced, abstract, hidecategories=(ComponentReplication))
 class ENGINE_API UActorComponent : public UObject, public IInterface_AssetUserData
