Add ENGINE_API to UTireType to allow subclassing outside engine
[https://answers.unrealengine.com/questions/123603/utiretype-is-not-exported-from-ue4-engine.html]

diff --git a/Engine/Source/Runtime/Engine/Classes/Vehicles/TireType.h b/Engine/Source/Runtime/Engine/Classes/Vehicles/TireType.h
--- a/Engine/Source/Runtime/Engine/Classes/Vehicles/TireType.h
+++ b/Engine/Source/Runtime/Engine/Classes/Vehicles/TireType.h
@@ -9,7 +9,7 @@
 #include "TireType.generated.h"
 
 UCLASS()
-class UTireType : public UDataAsset
+class ENGINE_API UTireType : public UDataAsset
 {
 	GENERATED_UCLASS_BODY()
 
