Added call in HandleSeamlesTravelPlayer to ensure a client is in spectating state when required.

diff --git a/Engine/Source/Runtime/Engine/Private/GameMode.cpp b/Engine/Source/Runtime/Engine/Private/GameMode.cpp
--- a/Engine/Source/Runtime/Engine/Private/GameMode.cpp
+++ b/Engine/Source/Runtime/Engine/Private/GameMode.cpp
@@ -871,6 +871,7 @@ void AGameMode::HandleSeamlessTravelPlayer(AController*& C)
 			NumTravellingPlayers--;
 			PC->bPlayerIsWaiting = true;
 			PC->ChangeState(NAME_Spectating);
+			PC->ClientGotoState(NAME_Spectating);
 		}
 	}
 	else
