Fixed UE-5133 "Visual Logger automation tests logs out 'Hello World'"

diff --git a/Engine/Source/Runtime/Engine/Private/VisualLogger/VisualLoggerAutomationTests.cpp b/Engine/Source/Runtime/Engine/Private/VisualLogger/VisualLoggerAutomationTests.cpp
--- a/Engine/Source/Runtime/Engine/Private/VisualLogger/VisualLoggerAutomationTests.cpp
+++ b/Engine/Source/Runtime/Engine/Private/VisualLogger/VisualLoggerAutomationTests.cpp
@@ -92,7 +92,7 @@ bool FVisualLogTest::RunTest(const FString& Parameters)
 	FVisualLogger::Get().SetIsRecording(false);
 	CHECK_FAIL(FVisualLogger::Get().IsRecording());
 
-	UE_VLOG(GWorld, LogVisual, Log, TEXT("Hello World"));
+	UE_VLOG(GWorld, LogVisual, Log, TEXT("Simple text line to test vlog"));
 	CHECK_SUCCESS(Context.Device.LastObject == NULL);
 	CHECK_SUCCESS(Context.Device.LastEntry.TimeStamp == -1);
 
@@ -100,8 +100,8 @@ bool FVisualLogTest::RunTest(const FString& Parameters)
 	CHECK_SUCCESS(FVisualLogger::Get().IsRecording());
 
 	{
-		const FString TextToLog = TEXT("Hello World");
-		UE_VLOG_UELOG(GWorld, LogVisual, Warning, TEXT("%s"), *TextToLog);
+		const FString TextToLog = TEXT("Simple text line to test if UE_VLOG_UELOG works fine");
+		UE_VLOG_UELOG(GWorld, LogVisual, Log, TEXT("%s"), *TextToLog);
 		CHECK_SUCCESS(Context.Device.LastObject != GWorld);
 		CHECK_SUCCESS(Context.Device.LastEntry.TimeStamp == -1);
 		FVisualLogEntry* CurrentEntry = FVisualLogger::Get().GetEntryToWrite(GWorld, GWorld->TimeSeconds, ECreateIfNeeded::DontCreate);
@@ -134,7 +134,7 @@ bool FVisualLogSegmentsTest::RunTest(const FString& Parameters)
 	FVisualLogger::Get().SetIsRecording(false);
 	CHECK_FAIL(FVisualLogger::Get().IsRecording());
 
-	UE_VLOG_SEGMENT(GWorld, LogVisual, Log, FVector(0, 0, 0), FVector(1, 0, 0), FColor::Red, TEXT("Simple segment log"));
+	UE_VLOG_SEGMENT(GWorld, LogVisual, Log, FVector(0, 0, 0), FVector(1, 0, 0), FColor::Red, TEXT("Simple segment log to test vlog"));
 	CHECK_SUCCESS(Context.Device.LastObject == NULL);
 	CHECK_SUCCESS(Context.Device.LastEntry.TimeStamp == -1);
 
@@ -143,7 +143,7 @@ bool FVisualLogSegmentsTest::RunTest(const FString& Parameters)
 	{
 		const FVector StartPoint(0, 0, 0);
 		const FVector EndPoint(1, 0, 0);
-		UE_VLOG_SEGMENT(GWorld, LogVisual, Log, StartPoint, EndPoint, FColor::Red, TEXT("Simple segment log"));
+		UE_VLOG_SEGMENT(GWorld, LogVisual, Log, StartPoint, EndPoint, FColor::Red, TEXT("Simple segment log to test vlog"));
 		CHECK_SUCCESS(Context.Device.LastObject == NULL);
 		CHECK_SUCCESS(Context.Device.LastEntry.TimeStamp == -1);
 		FVisualLogEntry* CurrentEntry = FVisualLogger::Get().GetEntryToWrite(GWorld, GWorld->TimeSeconds, ECreateIfNeeded::DontCreate);
@@ -166,11 +166,11 @@ bool FVisualLogSegmentsTest::RunTest(const FString& Parameters)
 	return true;
 }
 
-DEFINE_VLOG_EVENT(EventTest, All, "Simple event for tests")
+DEFINE_VLOG_EVENT(EventTest, All, "Simple event for vlog tests")
 
-DEFINE_VLOG_EVENT(EventTest2, All, "Second simple event for tests")
+DEFINE_VLOG_EVENT(EventTest2, All, "Second simple event for vlog tests")
 
-DEFINE_VLOG_EVENT(EventTest3, All, "Third simple event for tests")
+DEFINE_VLOG_EVENT(EventTest3, All, "Third simple event for vlog tests")
 
 IMPLEMENT_SIMPLE_AUTOMATION_TEST(FVisualLogEventsTest, "Engine.VisualLogger.Logging events", EAutomationTestFlags::ATF_Game | EAutomationTestFlags::ATF_Editor)
 bool FVisualLogEventsTest::RunTest(const FString& Parameters)
@@ -179,13 +179,13 @@ bool FVisualLogEventsTest::RunTest(const FString& Parameters)
 	FVisualLogger::Get().SetIsRecording(true);
 
 	CHECK_SUCCESS(EventTest.Name == TEXT("EventTest"));
-	CHECK_SUCCESS(EventTest.FriendlyDesc == TEXT("Simple event for tests"));
+	CHECK_SUCCESS(EventTest.FriendlyDesc == TEXT("Simple event for vlog tests"));
 
 	CHECK_SUCCESS(EventTest2.Name == TEXT("EventTest2"));
-	CHECK_SUCCESS(EventTest2.FriendlyDesc == TEXT("Second simple event for tests"));
+	CHECK_SUCCESS(EventTest2.FriendlyDesc == TEXT("Second simple event for vlog tests"));
 
 	CHECK_SUCCESS(EventTest3.Name == TEXT("EventTest3"));
-	CHECK_SUCCESS(EventTest3.FriendlyDesc == TEXT("Third simple event for tests"));
+	CHECK_SUCCESS(EventTest3.FriendlyDesc == TEXT("Third simple event for vlog tests"));
 
 	FVisualLogEntry* CurrentEntry = FVisualLogger::Get().GetEntryToWrite(GWorld, GWorld->TimeSeconds, ECreateIfNeeded::DontCreate);
 	float CurrentTimestamp = GWorld->TimeSeconds;
@@ -219,9 +219,9 @@ bool FVisualLogEventsTest::RunTest(const FString& Parameters)
 		CHECK_SUCCESS(CurrentEntry->Events[2].Counter == 1);
 		CHECK_SUCCESS(CurrentEntry->Events[2].Name == TEXT("EventTest3"));
 
-		CHECK_SUCCESS(CurrentEntry->Events[0].UserFriendlyDesc == TEXT("Simple event for tests"));
-		CHECK_SUCCESS(CurrentEntry->Events[1].UserFriendlyDesc == TEXT("Second simple event for tests"));
-		CHECK_SUCCESS(CurrentEntry->Events[2].UserFriendlyDesc == TEXT("Third simple event for tests"));
+		CHECK_SUCCESS(CurrentEntry->Events[0].UserFriendlyDesc == TEXT("Simple event for vlog tests"));
+		CHECK_SUCCESS(CurrentEntry->Events[1].UserFriendlyDesc == TEXT("Second simple event for vlog tests"));
+		CHECK_SUCCESS(CurrentEntry->Events[2].UserFriendlyDesc == TEXT("Third simple event for vlog tests"));
 
 		FVisualLogger::Get().GetEntryToWrite(GWorld, CurrentTimestamp+0.1); //generate new entry and serialize old one
 		CurrentEntry = &Context.Device.LastEntry;
