Add bIgnoreOwningActor to RadialForceComponent (UE-2802)

diff --git a/Engine/Source/Runtime/Engine/Classes/PhysicsEngine/RadialForceComponent.h b/Engine/Source/Runtime/Engine/Classes/PhysicsEngine/RadialForceComponent.h
--- a/Engine/Source/Runtime/Engine/Classes/PhysicsEngine/RadialForceComponent.h
+++ b/Engine/Source/Runtime/Engine/Classes/PhysicsEngine/RadialForceComponent.h
@@ -28,6 +28,10 @@ class URadialForceComponent : public USceneComponent
 	UPROPERTY(EditAnywhere, BlueprintReadWrite, Category=Impulse)
 	uint32 bImpulseVelChange:1;
 
+	/** If true, do not apply force/impulse to any physics objects that are part of the Actor that owns this component. */
+	UPROPERTY(EditAnywhere, BlueprintReadWrite, Category = Impulse)
+	uint32 bIgnoreOwningActor:1;
+
 	/** How strong the force should be */
 	UPROPERTY(EditAnywhere, BlueprintReadWrite, Category=Force)
 	float ForceStrength;
diff --git a/Engine/Source/Runtime/Engine/Private/PhysicsEngine/RadialForceComponent.cpp b/Engine/Source/Runtime/Engine/Private/PhysicsEngine/RadialForceComponent.cpp
--- a/Engine/Source/Runtime/Engine/Private/PhysicsEngine/RadialForceComponent.cpp
+++ b/Engine/Source/Runtime/Engine/Private/PhysicsEngine/RadialForceComponent.cpp
@@ -95,6 +95,13 @@ void URadialForceComponent::FireImpulse()
 
 	FCollisionQueryParams Params(FireImpulseOverlapName, false);
 	Params.bTraceAsyncScene = true; // want to hurt stuff in async scene
+
+	// Ignore owner actor if desired
+	if (bIgnoreOwningActor)
+	{
+		Params.AddIgnoredActor(GetOwner());
+	}
+
 	GetWorld()->OverlapMulti(Overlaps, Origin, FQuat::Identity, FCollisionShape::MakeSphere(Radius), Params, CollisionObjectQueryParams);
 
 	// Iterate over each and apply an impulse
