Add 'AssetUserData' support to SkeletalMesh (UE-2887)

diff --git a/Engine/Source/Runtime/Engine/Classes/Engine/SkeletalMesh.h b/Engine/Source/Runtime/Engine/Classes/Engine/SkeletalMesh.h
--- a/Engine/Source/Runtime/Engine/Classes/Engine/SkeletalMesh.h
+++ b/Engine/Source/Runtime/Engine/Classes/Engine/SkeletalMesh.h
@@ -10,6 +10,7 @@
 #include "Animation/PreviewAssetAttachComponent.h"
 #include "BoneContainer.h"
 #include "Interfaces/Interface_CollisionDataProvider.h"
+#include "Interfaces/Interface_AssetUserData.h"
 #include "SkeletalMesh.generated.h"
 
 /** The maximum number of skeletal mesh LODs allowed. */
@@ -564,7 +565,7 @@ class FSkeletalMeshResource;
  * @see https://docs.unrealengine.com/latest/INT/Engine/Content/Types/SkeletalMeshes/
  */
 UCLASS(hidecategories=Object, MinimalAPI, BlueprintType)
-class USkeletalMesh : public UObject, public IInterface_CollisionDataProvider
+class USkeletalMesh : public UObject, public IInterface_CollisionDataProvider, public IInterface_AssetUserData
 {
 	GENERATED_UCLASS_BODY()
 
@@ -714,6 +715,12 @@ public:
 	UPROPERTY(EditAnywhere, editfixedsize, BlueprintReadOnly, Category=Clothing)
 	TArray<FClothingAssetData>		ClothingAssets;
 
+protected:
+
+	/** Array of user data stored with the asset */
+	UPROPERTY(EditAnywhere, AdvancedDisplay, Instanced, Category=SkeletalMesh)
+	TArray<UAssetUserData*> AssetUserData;
+
 private:
 	/** Skeletal mesh source data */
 	class FSkeletalMeshSourceData* SourceData;
@@ -899,6 +906,13 @@ public:
 	}
 	// End Interface_CollisionDataProvider Interface
 
+	// Begin IInterface_AssetUserData Interface
+	virtual void AddAssetUserData(UAssetUserData* InUserData) override;
+	virtual void RemoveUserDataOfClass(TSubclassOf<UAssetUserData> InUserDataClass) override;
+	virtual UAssetUserData* GetAssetUserDataOfClass(TSubclassOf<UAssetUserData> InUserDataClass) override;
+	virtual const TArray<UAssetUserData*>* GetAssetUserDataArray() const override;
+	// End IInterface_AssetUserData Interface
+
 private:
 
 #if WITH_EDITOR
diff --git a/Engine/Source/Runtime/Engine/Private/SkeletalMesh.cpp b/Engine/Source/Runtime/Engine/Private/SkeletalMesh.cpp
--- a/Engine/Source/Runtime/Engine/Private/SkeletalMesh.cpp
+++ b/Engine/Source/Runtime/Engine/Private/SkeletalMesh.cpp
@@ -26,6 +26,7 @@
 #include "Animation/AnimBlueprintGeneratedClass.h"
 #include "ComponentReregisterContext.h"
 #include "PhysicsEngine/PhysicsAsset.h"
+#include "Engine/AssetUserData.h"
 
 #if WITH_EDITOR
 #include "MeshUtilities.h"
@@ -3474,6 +3475,49 @@ bool USkeletalMesh::GetPhysicsTriMeshData(FTriMeshCollisionData* CollisionData,
 #endif // #if WITH_EDITORONLY_DATA
 }
 
+void USkeletalMesh::AddAssetUserData(UAssetUserData* InUserData)
+{
+	if (InUserData != NULL)
+	{
+		UAssetUserData* ExistingData = GetAssetUserDataOfClass(InUserData->GetClass());
+		if (ExistingData != NULL)
+		{
+			AssetUserData.Remove(ExistingData);
+		}
+		AssetUserData.Add(InUserData);
+	}
+}
+
+UAssetUserData* USkeletalMesh::GetAssetUserDataOfClass(TSubclassOf<UAssetUserData> InUserDataClass)
+{
+	for (int32 DataIdx = 0; DataIdx < AssetUserData.Num(); DataIdx++)
+	{
+		UAssetUserData* Datum = AssetUserData[DataIdx];
+		if (Datum != NULL && Datum->IsA(InUserDataClass))
+		{
+			return Datum;
+		}
+	}
+	return NULL;
+}
+
+void USkeletalMesh::RemoveUserDataOfClass(TSubclassOf<UAssetUserData> InUserDataClass)
+{
+	for (int32 DataIdx = 0; DataIdx < AssetUserData.Num(); DataIdx++)
+	{
+		UAssetUserData* Datum = AssetUserData[DataIdx];
+		if (Datum != NULL && Datum->IsA(InUserDataClass))
+		{
+			AssetUserData.RemoveAt(DataIdx);
+			return;
+		}
+	}
+}
+
+const TArray<UAssetUserData*>* USkeletalMesh::GetAssetUserDataArray() const
+{
+	return &AssetUserData;
+}
 
 ////// SKELETAL MESH THUMBNAIL SUPPORT ////////
 
