Re-enable aggregates for ragdolls - Added threshold setting for spawning aggregates to project settings
[UE-2478] - re-enabled PxAggregate support for ragdolls .

diff --git a/Engine/Source/Runtime/Engine/Classes/Components/SkeletalMeshComponent.h b/Engine/Source/Runtime/Engine/Classes/Components/SkeletalMeshComponent.h
--- a/Engine/Source/Runtime/Engine/Classes/Components/SkeletalMeshComponent.h
+++ b/Engine/Source/Runtime/Engine/Classes/Components/SkeletalMeshComponent.h
@@ -504,6 +504,10 @@ public:
 	UPROPERTY()
 	FVector LineCheckBoundsScale;
 
+	/** Threshold for physics asset bodies above which we use an aggregate for broadphase collisions */
+	UPROPERTY()
+	int32 RagdollAggregateThreshold;
+
 	/** Notification when constraint is broken. */
 	UPROPERTY(BlueprintAssignable)
 	FConstraintBrokenSignature OnConstraintBroken;
diff --git a/Engine/Source/Runtime/Engine/Classes/PhysicsEngine/PhysicsSettings.h b/Engine/Source/Runtime/Engine/Classes/PhysicsEngine/PhysicsSettings.h
--- a/Engine/Source/Runtime/Engine/Classes/PhysicsEngine/PhysicsSettings.h
+++ b/Engine/Source/Runtime/Engine/Classes/PhysicsEngine/PhysicsSettings.h
@@ -86,6 +86,10 @@ class ENGINE_API UPhysicsSettings
 	UPROPERTY(config, EditAnywhere, Category = Constants)
 	float DefaultFluidFriction;
 
+	/** Threshold for ragdoll bodies above which they will be added to an aggregate before being added to the scene */
+	UPROPERTY(config, EditAnywhere, meta = (ClampMin = "1", UIMin = "1", ClampMax = "127", UIMax = "127"), Category = Constants)
+	int32 RagdollAggregateThreshold;
+
 	/** Enables the use of an async scene */
 	UPROPERTY(config, EditAnywhere, AdvancedDisplay, Category=Simulation)
 	bool bEnableAsyncScene;
diff --git a/Engine/Source/Runtime/Engine/Private/PhysicsEngine/BodyInstance.cpp b/Engine/Source/Runtime/Engine/Private/PhysicsEngine/BodyInstance.cpp
--- a/Engine/Source/Runtime/Engine/Private/PhysicsEngine/BodyInstance.cpp
+++ b/Engine/Source/Runtime/Engine/Private/PhysicsEngine/BodyInstance.cpp
@@ -1320,8 +1320,8 @@ void FBodyInstance::InitBody(UBodySetup* Setup, const FTransform& Transform, UPr
 	// Need to add actor into scene before calling putToSleep
 	// check if InAggregate is passed in
 	// BRG N.B. : For now only using InAggregate and BodyAggregate for dynamic bodies that are in the same scene as the aggregate, since otherwise we might have actors from two scenes.
-	// Right now aggregates are effectively disabled anyhow.
-	if(InAggregate && PNewDynamic != NULL && InAggregate->getScene() == PSceneForNewDynamic)
+	// If the aggregate scene is NULL we trust the caller that incoming bodies will be in the same scene.
+	if(InAggregate && PNewDynamic != NULL && (InAggregate->getScene() == PSceneForNewDynamic || InAggregate->getScene() == NULL))
 	{
 		InAggregate->addActor(*PNewDynamic);
 	}
@@ -1384,7 +1384,7 @@ void FBodyInstance::InitBody(UBodySetup* Setup, const FTransform& Transform, UPr
 		CreateDOFLock();
 
 		// wakeUp and putToSleep will issue warnings on kinematic actors
-		if (IsRigidBodyNonKinematic(PNewDynamic))
+		if (IsRigidBodyNonKinematic(PNewDynamic) && PNewDynamic->getScene())
 		{
 		    // Sleep/wake up as appropriate
 		    if (bShouldStartAwake)
diff --git a/Engine/Source/Runtime/Engine/Private/PhysicsEngine/PhysXSupport.h b/Engine/Source/Runtime/Engine/Private/PhysicsEngine/PhysXSupport.h
--- a/Engine/Source/Runtime/Engine/Private/PhysicsEngine/PhysXSupport.h
+++ b/Engine/Source/Runtime/Engine/Private/PhysicsEngine/PhysXSupport.h
@@ -139,7 +139,6 @@ private:
 };
 
 /** Thresholds for aggregates  */
-const uint32 AggregatePhysicsAssetThreshold  = 999999999;
 const uint32 AggregateMaxSize	   = 128;
 const uint32 AggregateBodyShapesThreshold	   = 999999999;
 
diff --git a/Engine/Source/Runtime/Engine/Private/PhysicsEngine/PhysicsSettings.cpp b/Engine/Source/Runtime/Engine/Private/PhysicsEngine/PhysicsSettings.cpp
--- a/Engine/Source/Runtime/Engine/Private/PhysicsEngine/PhysicsSettings.cpp
+++ b/Engine/Source/Runtime/Engine/Private/PhysicsEngine/PhysicsSettings.cpp
@@ -10,6 +10,7 @@ UPhysicsSettings::UPhysicsSettings(const FObjectInitializer& ObjectInitializer)
 	, DefaultGravityZ(-980.f)
 	, DefaultTerminalVelocity(4000.f)
 	, DefaultFluidFriction(0.3f)
+	, RagdollAggregateThreshold(4)
 	, bEnableAsyncScene(false)
 	, bEnable2DPhysics(false)
 	, bSimulateSkeletalMeshOnDedicatedServer(true)
diff --git a/Engine/Source/Runtime/Engine/Private/SkeletalMeshComponent.cpp b/Engine/Source/Runtime/Engine/Private/SkeletalMeshComponent.cpp
--- a/Engine/Source/Runtime/Engine/Private/SkeletalMeshComponent.cpp
+++ b/Engine/Source/Runtime/Engine/Private/SkeletalMeshComponent.cpp
@@ -136,6 +136,7 @@ USkeletalMeshComponent::USkeletalMeshComponent(const FObjectInitializer& ObjectI
 	bDefaultPlaying_DEPRECATED = true;
 	bEnablePhysicsOnDedicatedServer = UPhysicsSettings::Get()->bSimulateSkeletalMeshOnDedicatedServer;
 	bEnableUpdateRateOptimizations = false;
+	RagdollAggregateThreshold = UPhysicsSettings::Get()->RagdollAggregateThreshold;
 
 	bTickInEditor = true;
 }
diff --git a/Engine/Source/Runtime/Engine/Private/SkeletalMeshComponentPhysics.cpp b/Engine/Source/Runtime/Engine/Private/SkeletalMeshComponentPhysics.cpp
--- a/Engine/Source/Runtime/Engine/Private/SkeletalMeshComponentPhysics.cpp
+++ b/Engine/Source/Runtime/Engine/Private/SkeletalMeshComponentPhysics.cpp
@@ -695,8 +695,12 @@ void USkeletalMeshComponent::SetPhysMaterialOverride(UPhysicalMaterial* NewPhysM
 	}
 }
 
+DEFINE_STAT(STAT_InitArticulated);
+
 void USkeletalMeshComponent::InitArticulated(FPhysScene* PhysScene)
 {
+	SCOPE_CYCLE_COUNTER(STAT_InitArticulated);
+
 	UPhysicsAsset* const PhysicsAsset = GetPhysicsAsset();
 
 	if(PhysScene == NULL || PhysicsAsset == NULL || SkeletalMesh == NULL)
@@ -736,15 +740,19 @@ void USkeletalMeshComponent::InitArticulated(FPhysScene* PhysScene)
 	uint32 SkelMeshCompID = GetUniqueID();
 	PhysScene->DeferredAddCollisionDisableTable(SkelMeshCompID, &PhysicsAsset->CollisionDisableTable);
 
-	if(Aggregate == NULL && Bodies.Num() > AggregatePhysicsAssetThreshold)
+	int32 NumBodies = PhysicsAsset->BodySetup.Num();
+	if(Aggregate == NULL && NumBodies > RagdollAggregateThreshold && NumBodies <= AggregateMaxSize)
 	{
-		Aggregate = GPhysXSDK->createAggregate(AggregateMaxSize, true);
+		Aggregate = GPhysXSDK->createAggregate(PhysicsAsset->BodySetup.Num(), true);
+	}
+	else if(Aggregate && NumBodies > AggregateMaxSize)
+	{
+		UE_LOG(LogSkeletalMesh, Log, TEXT("USkeletalMeshComponent::InitArticulated : Too many bodies to create aggregate, Max: %u, This: %d"), AggregateMaxSize, NumBodies);
 	}
 #endif //WITH_PHYSX
 
 	// Create all the bodies.
 	check(Bodies.Num() == 0);
-	int32 NumBodies = PhysicsAsset->BodySetup.Num();
 	Bodies.AddZeroed(NumBodies);
 	for(int32 i=0; i<NumBodies; i++)
 	{
@@ -798,6 +806,19 @@ void USkeletalMeshComponent::InitArticulated(FPhysScene* PhysScene)
 		// Get the scene type from the SkeletalMeshComponent's BodyInstance
 		const uint32 SceneType = BodyInstance.UseAsyncScene() ? PST_Async : PST_Sync;
 		PhysScene->GetPhysXScene(SceneType)->addAggregate(*Aggregate);
+
+		// If we've used an aggregate, InitBody would not be able to set awake status as we *must* have a scene
+		// to do that, so we reconcile this here.
+		AActor* Owner = GetOwner();
+		bool bShouldSleep = !BodyInstance.bStartAwake && (Owner && Owner->GetVelocity().SizeSquared() <= KINDA_SMALL_NUMBER);
+
+		if(bShouldSleep)
+		{
+			for(FBodyInstance* Body : Bodies)
+			{
+				Body->GetPxRigidDynamic()->putToSleep();
+			}
+		}
 	}
 #endif //WITH_PHYSX
 
diff --git a/Engine/Source/Runtime/Engine/Public/PhysicsPublic.h b/Engine/Source/Runtime/Engine/Public/PhysicsPublic.h
--- a/Engine/Source/Runtime/Engine/Public/PhysicsPublic.h
+++ b/Engine/Source/Runtime/Engine/Public/PhysicsPublic.h
@@ -22,6 +22,7 @@ DECLARE_CYCLE_STAT_EXTERN(TEXT("Phys SetBodyTransform"),STAT_SetBodyTransform,ST
 DECLARE_CYCLE_STAT_EXTERN(TEXT("Phys SubstepStart"), STAT_SubstepSimulationStart,STATGROUP_Physics, );
 DECLARE_CYCLE_STAT_EXTERN(TEXT("Phys SubstepEnd"), STAT_SubstepSimulationEnd, STATGROUP_Physics, );
 DECLARE_CYCLE_STAT_EXTERN(TEXT("SyncComponentsToBodies"), STAT_SyncComponentsToBodies, STATGROUP_Physics, );
+DECLARE_CYCLE_STAT_EXTERN(TEXT("Init Articulated"), STAT_InitArticulated, STATGROUP_Physics, );
 
 DECLARE_DWORD_COUNTER_STAT_EXTERN(TEXT("Broadphase Adds"), STAT_NumBroadphaseAdds, STATGROUP_Physics, );
 DECLARE_DWORD_COUNTER_STAT_EXTERN(TEXT("Broadphase Removes"), STAT_NumBroadphaseRemoves, STATGROUP_Physics, );
