Add HitComponent to OnComponentHit multicaster

diff --git a/Engine/Source/Runtime/Engine/Classes/Components/PrimitiveComponent.h b/Engine/Source/Runtime/Engine/Classes/Components/PrimitiveComponent.h
--- a/Engine/Source/Runtime/Engine/Classes/Components/PrimitiveComponent.h
+++ b/Engine/Source/Runtime/Engine/Classes/Components/PrimitiveComponent.h
@@ -80,7 +80,7 @@ struct FSpriteCategoryInfo
  * Delegate for notification of blocking collision against a specific component.  
  * NormalImpulse will be filled in for physics-simulating bodies, but will be zero for swept-component blocking collisions. 
  */
-DECLARE_DYNAMIC_MULTICAST_DELEGATE_FourParams( FComponentHitSignature, class AActor*, OtherActor, class UPrimitiveComponent*, OtherComp, FVector, NormalImpulse, const FHitResult&, Hit );
+DECLARE_DYNAMIC_MULTICAST_DELEGATE_FiveParams( FComponentHitSignature, class AActor*, OtherActor, class UPrimitiveComponent*, HitComponent, class UPrimitiveComponent*, OtherComp, FVector, NormalImpulse, const FHitResult&, Hit );
 /** Delegate for notification of start of overlap with a specific component */
 DECLARE_DYNAMIC_MULTICAST_DELEGATE_FiveParams( FComponentBeginOverlapSignature,class AActor*, OtherActor, class UPrimitiveComponent*, OtherComp, int32, OtherBodyIndex, bool, bFromSweep, const FHitResult &, SweepResult);
 /** Delegate for notification of end of overlap with a specific component */
diff --git a/Engine/Source/Runtime/Engine/Private/Actor.cpp b/Engine/Source/Runtime/Engine/Private/Actor.cpp
--- a/Engine/Source/Runtime/Engine/Private/Actor.cpp
+++ b/Engine/Source/Runtime/Engine/Private/Actor.cpp
@@ -1759,7 +1759,7 @@ void AActor::DispatchBlockingHit(UPrimitiveComponent* MyComp, UPrimitiveComponen
 	// If component is still alive, call delegate on component
 	if(!MyComp->IsPendingKill())
 	{
-		MyComp->OnComponentHit.Broadcast(OtherActor, OtherComp, FVector(0,0,0), Hit);
+		MyComp->OnComponentHit.Broadcast(OtherActor, MyComp, OtherComp, FVector(0,0,0), Hit);
 	}
 }
 
@@ -3030,7 +3030,7 @@ void AActor::DispatchPhysicsCollisionHit(const FRigidBodyCollisionInfo& MyInfo,
 
 	if(MyInfo.Component.IsValid() && MyInfo.Component.Get()->OnComponentHit.IsBound())
 	{
-		MyInfo.Component.Get()->OnComponentHit.Broadcast(OtherInfo.Actor.Get(), OtherInfo.Component.Get(), RigidCollisionData.TotalNormalImpulse, Result);
+		MyInfo.Component.Get()->OnComponentHit.Broadcast(OtherInfo.Actor.Get(), MyInfo.Component.Get(), OtherInfo.Component.Get(), RigidCollisionData.TotalNormalImpulse, Result);
 	}
 }
 
diff --git a/Engine/Source/Runtime/GameplayAbilities/Private/Abilities/Tasks/AbilityTask_WaitOverlap.cpp b/Engine/Source/Runtime/GameplayAbilities/Private/Abilities/Tasks/AbilityTask_WaitOverlap.cpp
--- a/Engine/Source/Runtime/GameplayAbilities/Private/Abilities/Tasks/AbilityTask_WaitOverlap.cpp
+++ b/Engine/Source/Runtime/GameplayAbilities/Private/Abilities/Tasks/AbilityTask_WaitOverlap.cpp
@@ -17,7 +17,7 @@ void UAbilityTask_WaitOverlap::OnOverlapCallback(AActor* OtherActor, class UPrim
 	}
 }
 
-void UAbilityTask_WaitOverlap::OnHitCallback(AActor* OtherActor, UPrimitiveComponent* OtherComp, FVector NormalImpulse, const FHitResult& Hit)
+void UAbilityTask_WaitOverlap::OnHitCallback(AActor* OtherActor, UPrimitiveComponent* HitComponent, UPrimitiveComponent* OtherComp, FVector NormalImpulse, const FHitResult& Hit)
 {
 	if(OtherActor)
 	{
diff --git a/Engine/Source/Runtime/GameplayAbilities/Public/Abilities/Tasks/AbilityTask_WaitOverlap.h b/Engine/Source/Runtime/GameplayAbilities/Public/Abilities/Tasks/AbilityTask_WaitOverlap.h
--- a/Engine/Source/Runtime/GameplayAbilities/Public/Abilities/Tasks/AbilityTask_WaitOverlap.h
+++ b/Engine/Source/Runtime/GameplayAbilities/Public/Abilities/Tasks/AbilityTask_WaitOverlap.h
@@ -26,7 +26,7 @@ class UAbilityTask_WaitOverlap : public UAbilityTask
 	void OnOverlapCallback(AActor* OtherActor, UPrimitiveComponent* OtherComp, int32 OtherBodyIndex, bool bFromSweep, const FHitResult& SweepResult);
 
 	UFUNCTION()
-	void OnHitCallback(AActor* OtherActor, UPrimitiveComponent* OtherComp, FVector NormalImpulse, const FHitResult& Hit);
+	void OnHitCallback(AActor* OtherActor, UPrimitiveComponent* HitComponent, UPrimitiveComponent* OtherComp, FVector NormalImpulse, const FHitResult& Hit);
 
 	virtual void Activate() override;
 
