Make SplineComponent a PrimitiveComponent so that subclases can add rendering functionality
Move LineBatchComponent from Level.cpp to its own LineBatchComponent.cpp file
Add comments for several framework files (UE-3617)

diff --git a/Engine/Source/Runtime/Engine/Classes/Components/BrushComponent.h b/Engine/Source/Runtime/Engine/Classes/Components/BrushComponent.h
--- a/Engine/Source/Runtime/Engine/Classes/Components/BrushComponent.h
+++ b/Engine/Source/Runtime/Engine/Classes/Components/BrushComponent.h
@@ -1,12 +1,16 @@
 // Copyright 1998-2014 Epic Games, Inc. All Rights Reserved.
 
-
 #pragma once
 #include "BrushComponent.generated.h"
 
 struct FEngineShowFlags;
 struct FConvexVolume;
 
+/** 
+ *	A brush component defines a shape that can be modified within the editor. They are used both as part of BSP building, and for volumes. 
+ *	@see https://docs.unrealengine.com/latest/INT/Engine/Actors/Volumes
+ *	@see https://docs.unrealengine.com/latest/INT/Engine/Actors/Brushes
+ */
 UCLASS(editinlinenew, MinimalAPI, hidecategories=(Physics, Lighting, LOD, Rendering, TextureStreaming, Transform, Activation, "Components|Activation"), showcategories=(Mobility, "Rendering|Material"))
 class UBrushComponent : public UPrimitiveComponent
 {
diff --git a/Engine/Source/Runtime/Engine/Classes/Components/LineBatchComponent.h b/Engine/Source/Runtime/Engine/Classes/Components/LineBatchComponent.h
--- a/Engine/Source/Runtime/Engine/Classes/Components/LineBatchComponent.h
+++ b/Engine/Source/Runtime/Engine/Classes/Components/LineBatchComponent.h
@@ -104,17 +104,22 @@ public:
 	float RemainingLifeTime;
 };
 
+/** 
+ * The line batch component buffers and draws lines (and some other line-based shapes) in a scene. 
+ *	This can be useful for debug drawing, but is not very performant for runtime use.
+ */
 UCLASS(MinimalAPI)
 class ULineBatchComponent : public UPrimitiveComponent
 {
 	GENERATED_UCLASS_BODY()
 
+	/** Buffer of lines to draw */
 	TArray<struct FBatchedLine> BatchedLines;
-
+	/** Buffer or points to draw */
 	TArray<struct FBatchedPoint> BatchedPoints;
-
+	/** Default time that lines/points will draw for */
 	float DefaultLifeTime;
-
+	/** Buffer of simple meshes to draw */
 	TArray<struct FBatchedMesh> BatchedMeshes;
 
 	/** Provide many lines to draw - faster than calling DrawLine many times. */
@@ -147,6 +152,7 @@ class ULineBatchComponent : public UPrimitiveComponent
 
 	/** Draw a box */
 	void DrawSolidBox(FBox const& Box, FTransform const& Xform, const FColor& Color, uint8 DepthPriority, float LifeTime);
+	/** Draw a mesh */
 	void DrawMesh(TArray<FVector> const& Verts, TArray<int32> const& Indices, FColor const& Color, uint8 DepthPriority, float LifeTime);
 
 	// Begin UPrimitiveComponent interface.
@@ -160,6 +166,7 @@ class ULineBatchComponent : public UPrimitiveComponent
 	virtual void ApplyWorldOffset(const FVector& InOffset, bool bWorldShift) override;
 	// End UActorComponent interface.
 
+	/** Clear all batched lines, points and meshes */
 	ENGINE_API void Flush();
 };
 
diff --git a/Engine/Source/Runtime/Engine/Classes/Components/SplineComponent.h b/Engine/Source/Runtime/Engine/Classes/Components/SplineComponent.h
--- a/Engine/Source/Runtime/Engine/Classes/Components/SplineComponent.h
+++ b/Engine/Source/Runtime/Engine/Classes/Components/SplineComponent.h
@@ -4,10 +4,12 @@
 #pragma once
 #include "SplineComponent.generated.h"
 
-
-
+/** 
+ *	A spline component is a spline shape which can be used for other purposes (e.g. animating objects). It does not contain rendering capabilities itself (outside the editor) 
+ *	@see https://docs.unrealengine.com/latest/INT/Resources/ContentExamples/Blueprint_Splines
+ */
 UCLASS(ClassGroup=Shapes, meta=(BlueprintSpawnableComponent))
-class ENGINE_API USplineComponent : public USceneComponent
+class ENGINE_API USplineComponent : public UPrimitiveComponent
 {
 	GENERATED_UCLASS_BODY()
 
diff --git a/Engine/Source/Runtime/Engine/Classes/Components/SplineMeshComponent.h b/Engine/Source/Runtime/Engine/Classes/Components/SplineMeshComponent.h
--- a/Engine/Source/Runtime/Engine/Classes/Components/SplineMeshComponent.h
+++ b/Engine/Source/Runtime/Engine/Classes/Components/SplineMeshComponent.h
@@ -84,6 +84,10 @@ struct FSplineMeshParams
 
 };
 
+/** 
+ *	A Spline Mesh Component is a derivation of a Static Mesh Component which can be deformed using a spline. Only a start and end position (and tangent) can be specified.  
+ *	@see https://docs.unrealengine.com/latest/INT/Resources/ContentExamples/Blueprint_Splines
+ */
 UCLASS(ClassGroup=Rendering, hidecategories=(Physics), meta=(BlueprintSpawnableComponent))
 class ENGINE_API USplineMeshComponent : public UStaticMeshComponent, public IInterface_CollisionDataProvider
 {
diff --git a/Engine/Source/Runtime/Engine/Classes/GameFramework/PlayerStart.h b/Engine/Source/Runtime/Engine/Classes/GameFramework/PlayerStart.h
--- a/Engine/Source/Runtime/Engine/Classes/GameFramework/PlayerStart.h
+++ b/Engine/Source/Runtime/Engine/Classes/GameFramework/PlayerStart.h
@@ -1,22 +1,26 @@
 // Copyright 1998-2014 Epic Games, Inc. All Rights Reserved.
 
-//=============================================================================
-// Player start location.
-//=============================================================================
-
 #pragma once
 #include "Engine/NavigationObjectBase.h"
 #include "PlayerStart.generated.h"
 
+/** 
+ *	This class indicates a location where a player can spawn when the game begins
+ *	
+ *	@see https://docs.unrealengine.com/latest/INT/Engine/Actors/PlayerStart/
+ */
 UCLASS(ClassGroup=Common, hidecategories=Collision)
 class ENGINE_API APlayerStart : public ANavigationObjectBase
 {
+	/** To take more control over PlayerStart selection, you can override the virtual AGameMode::FindPlayerStart and AGameMode::ChoosePlayerStart function */
+
 	GENERATED_UCLASS_BODY()
 
 	/** Used when searching for which playerstart to use. */
 	UPROPERTY(EditAnywhere, BlueprintReadWrite, Category=Object)
 	FName PlayerStartTag;
 
+	/** Arrow component to indicate forward direction of start */
 #if WITH_EDITORONLY_DATA
 private_subobject:
 	DEPRECATED_FORGAME(4.6, "ArrowComponent should not be accessed directly, please use GetArrowComponent() function instead. ArrowComponent will soon be private and your code will not compile.")
@@ -25,9 +29,10 @@ private_subobject:
 public:
 #endif
 
-	virtual void PostInitializeComponents() override;
-	
+	// Begin AActor interface
+	virtual void PostInitializeComponents() override;	
 	virtual void PostUnregisterAllComponents() override;
+	// End AActor interface
 
 #if WITH_EDITORONLY_DATA
 	/** Returns ArrowComponent subobject **/
diff --git a/Engine/Source/Runtime/Engine/Classes/GameFramework/SaveGame.h b/Engine/Source/Runtime/Engine/Classes/GameFramework/SaveGame.h
--- a/Engine/Source/Runtime/Engine/Classes/GameFramework/SaveGame.h
+++ b/Engine/Source/Runtime/Engine/Classes/GameFramework/SaveGame.h
@@ -3,9 +3,24 @@
 #pragma once
 #include "SaveGame.generated.h"
 
+/** 
+ *	This class acts as a base class for a save game object that can be used to save state about the game. 
+ *	When you create your own save game subclass, you would add member variables for the information that you want to save.
+ *	Then when you want to save a game, create an instance of this object using CreateSaveGameObject, fill in the data, and use SaveGameToSlot, providing a slot name.
+ *	To load the game you then just use LoadGameFromSlot, and read the data from the resulting object.
+ *
+ *	@see https://docs.unrealengine.com/latest/INT/Gameplay/SaveGame
+ */
 UCLASS(abstract, Blueprintable, BlueprintType)
 class ENGINE_API USaveGame : public UObject
 {
-	GENERATED_UCLASS_BODY()
+	/**
+	 *	@see UGameplayStatics::CreateSaveGameObject
+	 *	@see UGameplayStatics::SaveGameToSlot
+	 *	@see UGameplayStatics::DoesSaveGameExist
+	 *	@see UGameplayStatics::LoadGameFromSlot
+	 *	@see UGameplayStatics::DeleteGameInSlot
+	 */
 
+	GENERATED_UCLASS_BODY()
 };
diff --git a/Engine/Source/Runtime/Engine/Private/Level.cpp b/Engine/Source/Runtime/Engine/Private/Level.cpp
--- a/Engine/Source/Runtime/Engine/Private/Level.cpp
+++ b/Engine/Source/Runtime/Engine/Private/Level.cpp
@@ -32,7 +32,6 @@ Level.cpp: Level-related functions
 #include "Engine/ShadowMapTexture2D.h"
 #include "Components/ModelComponent.h"
 #include "Engine/LightMapTexture2D.h"
-#include "Components/LineBatchComponent.h"
 DEFINE_LOG_CATEGORY(LogLevel);
 
 /*-----------------------------------------------------------------------------
@@ -1969,409 +1968,3 @@ void ULevel::RemoveUserDataOfClass(TSubclassOf<UAssetUserData> InUserDataClass)
 	}
 }
 
-/*-----------------------------------------------------------------------------
-ULineBatchComponent implementation.
------------------------------------------------------------------------------*/
-
-/** Represents a LineBatchComponent to the scene manager. */
-class FLineBatcherSceneProxy : public FPrimitiveSceneProxy
-{
-public:
-	FLineBatcherSceneProxy(const ULineBatchComponent* InComponent) :
-		FPrimitiveSceneProxy(InComponent), Lines(InComponent->BatchedLines), 
-		Points(InComponent->BatchedPoints), Meshes(InComponent->BatchedMeshes)
-	{
-		bWillEverBeLit = false;
-		ViewRelevance.bDrawRelevance = true;
-		ViewRelevance.bDynamicRelevance = true;
-		ViewRelevance.bNormalTranslucencyRelevance = true;
-	}
-
-	virtual void GetDynamicMeshElements(const TArray<const FSceneView*>& Views, const FSceneViewFamily& ViewFamily, uint32 VisibilityMap, FMeshElementCollector& Collector) const override
-	{
-		QUICK_SCOPE_CYCLE_COUNTER( STAT_LineBatcherSceneProxy_GetDynamicMeshElements );
-
-		for (int32 ViewIndex = 0; ViewIndex < Views.Num(); ViewIndex++)
-		{
-			if (VisibilityMap & (1 << ViewIndex))
-			{
-				const FSceneView* View = Views[ViewIndex];
-				FPrimitiveDrawInterface* PDI = Collector.GetPDI(ViewIndex);
-
-				for (int32 i = 0; i < Lines.Num(); i++)
-				{
-					PDI->DrawLine(Lines[i].Start, Lines[i].End, Lines[i].Color, Lines[i].DepthPriority, Lines[i].Thickness);
-				}
-
-				for (int32 i = 0; i < Points.Num(); i++)
-				{
-					PDI->DrawPoint(Points[i].Position, Points[i].Color, Points[i].PointSize, Points[i].DepthPriority);
-				}
-
-				for (int32 i = 0; i < Meshes.Num(); i++)
-				{
-					static FVector const PosX(1.f,0,0);
-					static FVector const PosY(0,1.f,0);
-					static FVector const PosZ(0,0,1.f);
-
-					FBatchedMesh const& M = Meshes[i];
-
-					// this seems far from optimal in terms of perf, but it's for debugging
-					FDynamicMeshBuilder MeshBuilder;
-
-					// set up geometry
-					for (int32 VertIdx=0; VertIdx < M.MeshVerts.Num(); ++VertIdx)
-					{
-						MeshBuilder.AddVertex( M.MeshVerts[VertIdx], FVector2D::ZeroVector, PosX, PosY, PosZ, FColor::White );
-					}
-					//MeshBuilder.AddTriangles(M.MeshIndices);
-					for (int32 Idx=0; Idx < M.MeshIndices.Num(); Idx+=3)
-					{
-						MeshBuilder.AddTriangle( M.MeshIndices[Idx], M.MeshIndices[Idx+1], M.MeshIndices[Idx+2] );
-					}
-
-					FMaterialRenderProxy* const MaterialRenderProxy = new(FMemStack::Get()) FColoredMaterialRenderProxy(GEngine->DebugMeshMaterial->GetRenderProxy(false), M.Color);
-					MeshBuilder.GetMesh(FMatrix::Identity, MaterialRenderProxy, M.DepthPriority, false, false, ViewIndex, Collector);
-				}
-			}
-		}
-	}
-
-	/** 
-	* Draw the scene proxy as a dynamic element
-	*
-	* @param	PDI - draw interface to render to
-	* @param	View - current view
-	*/
-	virtual void DrawDynamicElements(FPrimitiveDrawInterface* PDI,const FSceneView* View) override
-	{
-		QUICK_SCOPE_CYCLE_COUNTER( STAT_LineBatcherSceneProxy_DrawDynamicElements );
-
-		for (int32 i = 0; i < Lines.Num(); i++)
-		{
-			PDI->DrawLine(Lines[i].Start, Lines[i].End, Lines[i].Color, Lines[i].DepthPriority, Lines[i].Thickness);
-		}
-
-		for (int32 i = 0; i < Points.Num(); i++)
-		{
-			PDI->DrawPoint(Points[i].Position, Points[i].Color, Points[i].PointSize, Points[i].DepthPriority);
-		}
-
-		for (int32 i = 0; i < Meshes.Num(); i++)
-		{
-			static FVector const PosX(1.f,0,0);
-			static FVector const PosY(0,1.f,0);
-			static FVector const PosZ(0,0,1.f);
-
-			FBatchedMesh const& M = Meshes[i];
-
-			// this seems far from optimal in terms of perf, but it's for debugging
-			FDynamicMeshBuilder MeshBuilder;
-
-			// set up geometry
-			for (int32 VertIdx=0; VertIdx < M.MeshVerts.Num(); ++VertIdx)
-			{
-				MeshBuilder.AddVertex( M.MeshVerts[VertIdx], FVector2D::ZeroVector, PosX, PosY, PosZ, FColor::White );
-			}
-			//MeshBuilder.AddTriangles(M.MeshIndices);
-			for (int32 Idx=0; Idx < M.MeshIndices.Num(); Idx+=3)
-			{
-				MeshBuilder.AddTriangle( M.MeshIndices[Idx], M.MeshIndices[Idx+1], M.MeshIndices[Idx+2] );
-			}
-
-			FMaterialRenderProxy* const MaterialRenderProxy = new(FMemStack::Get()) FColoredMaterialRenderProxy(GEngine->DebugMeshMaterial->GetRenderProxy(false), M.Color);
-			MeshBuilder.Draw(PDI, FMatrix::Identity, MaterialRenderProxy, M.DepthPriority);
-		}
-	}
-
-	/**
-	*  Returns a struct that describes to the renderer when to draw this proxy.
-	*	@param		Scene view to use to determine our relevence.
-	*  @return		View relevance struct
-	*/
-	virtual FPrimitiveViewRelevance GetViewRelevance(const FSceneView* View)
-	{
-		return ViewRelevance;
-	}
-	virtual uint32 GetMemoryFootprint( void ) const { return( sizeof( *this ) + GetAllocatedSize() ); }
-	uint32 GetAllocatedSize( void ) const 
-	{ 
-		return( FPrimitiveSceneProxy::GetAllocatedSize() + Lines.GetAllocatedSize() + Points.GetAllocatedSize() + Meshes.GetAllocatedSize() ); 
-	}
-
-private:
-	TArray<FBatchedLine> Lines;
-	TArray<FBatchedPoint> Points;
-	TArray<FBatchedMesh> Meshes;
-	FPrimitiveViewRelevance ViewRelevance;
-};
-
-ULineBatchComponent::ULineBatchComponent(const FObjectInitializer& ObjectInitializer)
-	: Super(ObjectInitializer)
-{
-	bAutoActivate = true;
-	bTickInEditor = true;
-	PrimaryComponentTick.bCanEverTick = true;
-	BodyInstance.bEnableCollision_DEPRECATED = false;
-
-	SetCollisionProfileName(UCollisionProfile::NoCollision_ProfileName);
-
-	bUseEditorCompositing = true;
-	bGenerateOverlapEvents = false;
-}
-
-void ULineBatchComponent::DrawLine(const FVector& Start,const FVector& End,const FLinearColor& Color,uint8 DepthPriority,const float Thickness, const float LifeTime)
-{
-	new(BatchedLines) FBatchedLine(Start,End,Color,LifeTime,Thickness,DepthPriority);
-	// LineBatcher and PersistentLineBatcher components will be updated at the end of UWorld::Tick
-	MarkRenderStateDirty();
-}
-
-void ULineBatchComponent::DrawLines(const TArray<FBatchedLine>& InLines)
-{
-	BatchedLines.Append(InLines);
-	// LineBatcher and PersistentLineBatcher components will be updated at the end of UWorld::Tick
-	MarkRenderStateDirty();
-}
-
-void ULineBatchComponent::DrawPoint(
-	const FVector& Position,
-	const FLinearColor& Color,
-	float PointSize,
-	uint8 DepthPriority,
-	float LifeTime
-	)
-{
-	new(BatchedPoints) FBatchedPoint(Position, Color, PointSize, LifeTime, DepthPriority);
-	// LineBatcher and PersistentLineBatcher components will be updated at the end of UWorld::Tick
-	MarkRenderStateDirty();
-}
-
-void ULineBatchComponent::DrawBox(const FBox& Box, const FMatrix& TM, const FColor& Color, uint8 DepthPriorityGroup)
-{
-	FVector	B[2], P, Q;
-	int32 ai, aj;
-	B[0] = Box.Min;
-	B[1] = Box.Max;
-
-	for( ai=0; ai<2; ai++ ) for( aj=0; aj<2; aj++ )
-	{
-		P.X=B[ai].X; Q.X=B[ai].X;
-		P.Y=B[aj].Y; Q.Y=B[aj].Y;
-		P.Z=B[0].Z; Q.Z=B[1].Z;
-		new(BatchedLines) FBatchedLine(TM.TransformPosition(P), TM.TransformPosition(Q), Color, DefaultLifeTime, 0.0f, DepthPriorityGroup);
-
-		P.Y=B[ai].Y; Q.Y=B[ai].Y;
-		P.Z=B[aj].Z; Q.Z=B[aj].Z;
-		P.X=B[0].X; Q.X=B[1].X;
-		new(BatchedLines) FBatchedLine(TM.TransformPosition(P), TM.TransformPosition(Q), Color, DefaultLifeTime, 0.0f, DepthPriorityGroup);
-
-		P.Z=B[ai].Z; Q.Z=B[ai].Z;
-		P.X=B[aj].X; Q.X=B[aj].X;
-		P.Y=B[0].Y; Q.Y=B[1].Y;
-		new(BatchedLines) FBatchedLine(TM.TransformPosition(P), TM.TransformPosition(Q), Color, DefaultLifeTime, 0.0f, DepthPriorityGroup);
-	}
-	// LineBatcher and PersistentLineBatcher components will be updated at the end of UWorld::Tick
-	MarkRenderStateDirty();
-}
-
-void ULineBatchComponent::DrawSolidBox(const FBox& Box, const FTransform& Xform, const FColor& Color, uint8 DepthPriority, float LifeTime)
-{
-	int32 const NewMeshIdx = BatchedMeshes.Add(FBatchedMesh());
-	FBatchedMesh& BM = BatchedMeshes[NewMeshIdx];
-
-	BM.Color = Color;
-	BM.DepthPriority = DepthPriority;
-	BM.RemainingLifeTime = LifeTime;
-
-	BM.MeshVerts.AddUninitialized(8);
-	BM.MeshVerts[0] = Xform.TransformPosition( FVector(Box.Min.X, Box.Min.Y, Box.Max.Z) );
-	BM.MeshVerts[1] = Xform.TransformPosition( FVector(Box.Max.X, Box.Min.Y, Box.Max.Z) );
-	BM.MeshVerts[2] = Xform.TransformPosition( FVector(Box.Min.X, Box.Min.Y, Box.Min.Z) );
-	BM.MeshVerts[3] = Xform.TransformPosition( FVector(Box.Max.X, Box.Min.Y, Box.Min.Z) );
-	BM.MeshVerts[4] = Xform.TransformPosition( FVector(Box.Min.X, Box.Max.Y, Box.Max.Z) );
-	BM.MeshVerts[5] = Xform.TransformPosition( FVector(Box.Max.X, Box.Max.Y, Box.Max.Z) );
-	BM.MeshVerts[6] = Xform.TransformPosition( FVector(Box.Min.X, Box.Max.Y, Box.Min.Z) );
-	BM.MeshVerts[7] = Xform.TransformPosition( FVector(Box.Max.X, Box.Max.Y, Box.Min.Z) );
-
-	// clockwise
-	BM.MeshIndices.AddUninitialized(36);
-	int32 const Indices[36] = {	3,2,0,
-		3,0,1,
-		7,3,1,
-		7,1,5,
-		6,7,5,
-		6,5,4,
-		2,6,4,
-		2,4,0,
-		1,0,4,
-		1,4,5,
-		7,6,2,
-		7,2,3	};
-
-	for (int32 Idx=0; Idx<36; ++Idx)
-	{
-		BM.MeshIndices[Idx] = Indices[Idx];
-	}
-
-	MarkRenderStateDirty();
-}
-
-void ULineBatchComponent::DrawMesh(TArray<FVector> const& Verts, TArray<int32> const& Indices, FColor const& Color, uint8 DepthPriority, float LifeTime)
-{
-	// modifying array element directly to avoid copying arrays
-	int32 const NewMeshIdx = BatchedMeshes.Add(FBatchedMesh());
-	FBatchedMesh& BM = BatchedMeshes[NewMeshIdx];
-
-	BM.MeshIndices = Indices;
-	BM.MeshVerts = Verts;
-	BM.Color = Color;
-	BM.DepthPriority = DepthPriority;
-	BM.RemainingLifeTime = LifeTime;
-
-	MarkRenderStateDirty();
-}
-
-void ULineBatchComponent::DrawDirectionalArrow(const FMatrix& ArrowToWorld,FColor InColor,float Length,float ArrowSize,uint8 DepthPriority)
-{
-	const FVector Tip = ArrowToWorld.TransformPosition(FVector(Length,0,0));
-	new(BatchedLines) FBatchedLine(Tip,ArrowToWorld.TransformPosition(FVector::ZeroVector),InColor,DefaultLifeTime,0.0f,DepthPriority);
-	new(BatchedLines) FBatchedLine(Tip,ArrowToWorld.TransformPosition(FVector(Length-ArrowSize,+ArrowSize,+ArrowSize)),InColor,DefaultLifeTime,0.0f,DepthPriority);
-	new(BatchedLines) FBatchedLine(Tip,ArrowToWorld.TransformPosition(FVector(Length-ArrowSize,+ArrowSize,-ArrowSize)),InColor,DefaultLifeTime,0.0f,DepthPriority);
-	new(BatchedLines) FBatchedLine(Tip,ArrowToWorld.TransformPosition(FVector(Length-ArrowSize,-ArrowSize,+ArrowSize)),InColor,DefaultLifeTime,0.0f,DepthPriority);
-	new(BatchedLines) FBatchedLine(Tip,ArrowToWorld.TransformPosition(FVector(Length-ArrowSize,-ArrowSize,-ArrowSize)),InColor,DefaultLifeTime,0.0f,DepthPriority);
-	MarkRenderStateDirty();
-}
-
-/** Draw a circle */
-void ULineBatchComponent::DrawCircle(const FVector& Base,const FVector& X,const FVector& Y,FColor Color,float Radius,int32 NumSides,uint8 DepthPriority)
-{
-	const float	AngleDelta = 2.0f * PI / NumSides;
-	FVector	LastVertex = Base + X * Radius;
-
-	for(int32 SideIndex = 0;SideIndex < NumSides;SideIndex++)
-	{
-		const FVector Vertex = Base + (X * FMath::Cos(AngleDelta * (SideIndex + 1)) + Y * FMath::Sin(AngleDelta * (SideIndex + 1))) * Radius;
-		new(BatchedLines) FBatchedLine(LastVertex,Vertex,Color,DefaultLifeTime,0.0f,DepthPriority);
-		LastVertex = Vertex;
-	}
-
-	MarkRenderStateDirty();
-}
-
-void ULineBatchComponent::TickComponent(float DeltaTime, enum ELevelTick TickType, FActorComponentTickFunction *ThisTickFunction)
-{
-	bool bDirty = false;
-	// Update the life time of batched lines, removing the lines which have expired.
-	for(int32 LineIndex=0; LineIndex < BatchedLines.Num(); LineIndex++)
-	{
-		FBatchedLine& Line = BatchedLines[LineIndex];
-		if (Line.RemainingLifeTime > 0.0f)
-		{
-			Line.RemainingLifeTime -= DeltaTime;
-			if(Line.RemainingLifeTime <= 0.0f)
-			{
-				// The line has expired, remove it.
-				BatchedLines.RemoveAt(LineIndex--);
-				bDirty = true;
-			}
-		}
-	}
-
-	// Update the life time of batched points, removing the points which have expired.
-	for(int32 PtIndex=0; PtIndex < BatchedPoints.Num(); PtIndex++)
-	{
-		FBatchedPoint& Pt = BatchedPoints[PtIndex];
-		if (Pt.RemainingLifeTime > 0.0f)
-		{
-			Pt.RemainingLifeTime -= DeltaTime;
-			if(Pt.RemainingLifeTime <= 0.0f)
-			{
-				// The point has expired, remove it.
-				BatchedPoints.RemoveAt(PtIndex--);
-				bDirty = true;
-			}
-		}
-	}
-
-	// Update the life time of batched meshes, removing the meshes which have expired.
-	for(int32 MeshIndex=0; MeshIndex < BatchedMeshes.Num(); MeshIndex++)
-	{
-		FBatchedMesh& Mesh = BatchedMeshes[MeshIndex];
-		if (Mesh.RemainingLifeTime > 0.0f)
-		{
-			Mesh.RemainingLifeTime -= DeltaTime;
-			if(Mesh.RemainingLifeTime <= 0.0f)
-			{
-				// The mesh has expired, remove it.
-				BatchedMeshes.RemoveAt(MeshIndex--);
-				bDirty = true;
-			}
-		}
-	}
-
-	if(bDirty)
-	{
-		MarkRenderStateDirty();
-	}
-}
-
-void ULineBatchComponent::ApplyWorldOffset(const FVector& InOffset, bool bWorldShift)
-{
-	Super::ApplyWorldOffset(InOffset, bWorldShift);
-
-	bool bDirty = false;
-	for (FBatchedLine& Line : BatchedLines)
-	{
-		Line.Start += InOffset;
-		Line.End += InOffset;
-		bDirty = true;
-	}
-
-	for (FBatchedPoint& Point : BatchedPoints)
-	{
-		Point.Position += InOffset;
-		bDirty = true;
-	}
-
-	for (FBatchedMesh& Mesh : BatchedMeshes)
-	{
-		for (FVector& Vert : Mesh.MeshVerts)
-		{
-			Vert += InOffset;
-			bDirty = true;
-		}
-	}
-
-	if (bDirty)
-	{
-		MarkRenderStateDirty();
-	}
-}
-
-/**
-* Creates a new scene proxy for the line batcher component.
-* @return	Pointer to the FLineBatcherSceneProxy
-*/
-FPrimitiveSceneProxy* ULineBatchComponent::CreateSceneProxy()
-{
-	return new FLineBatcherSceneProxy(this);
-}
-
-FBoxSphereBounds ULineBatchComponent::CalcBounds( const FTransform& LocalToWorld ) const 
-{
-	const FVector BoxExtent(HALF_WORLD_MAX);
-	return FBoxSphereBounds(FVector::ZeroVector, BoxExtent, BoxExtent.Size());
-}
-
-void ULineBatchComponent::Flush()
-{
-	if (BatchedLines.Num() > 0 || BatchedPoints.Num() > 0 || BatchedMeshes.Num() > 0)
-	{
-	BatchedLines.Empty();
-	BatchedPoints.Empty();
-	BatchedMeshes.Empty();
-	MarkRenderStateDirty();
-}
-}
-
diff --git a/Engine/Source/Runtime/Engine/Private/LineBatchComponent.cpp b/Engine/Source/Runtime/Engine/Private/LineBatchComponent.cpp
new file mode 100644
--- /dev/null
+++ b/Engine/Source/Runtime/Engine/Private/LineBatchComponent.cpp
@@ -0,0 +1,413 @@
+// Copyright 1998-2014 Epic Games, Inc. All Rights Reserved.
+
+/*-----------------------------------------------------------------------------
+	ULineBatchComponent implementation.
+-----------------------------------------------------------------------------*/
+
+#include "EnginePrivate.h"
+#include "Components/LineBatchComponent.h"
+#include "PrimitiveSceneProxy.h"
+#include "DynamicMeshBuilder.h"
+
+
+/** Represents a LineBatchComponent to the scene manager. */
+class FLineBatcherSceneProxy : public FPrimitiveSceneProxy
+{
+public:
+	FLineBatcherSceneProxy(const ULineBatchComponent* InComponent) :
+		FPrimitiveSceneProxy(InComponent), Lines(InComponent->BatchedLines), 
+		Points(InComponent->BatchedPoints), Meshes(InComponent->BatchedMeshes)
+	{
+		bWillEverBeLit = false;
+		ViewRelevance.bDrawRelevance = true;
+		ViewRelevance.bDynamicRelevance = true;
+		ViewRelevance.bNormalTranslucencyRelevance = true;
+	}
+
+	virtual void GetDynamicMeshElements(const TArray<const FSceneView*>& Views, const FSceneViewFamily& ViewFamily, uint32 VisibilityMap, FMeshElementCollector& Collector) const override
+	{
+		QUICK_SCOPE_CYCLE_COUNTER( STAT_LineBatcherSceneProxy_GetDynamicMeshElements );
+
+		for (int32 ViewIndex = 0; ViewIndex < Views.Num(); ViewIndex++)
+		{
+			if (VisibilityMap & (1 << ViewIndex))
+			{
+				const FSceneView* View = Views[ViewIndex];
+				FPrimitiveDrawInterface* PDI = Collector.GetPDI(ViewIndex);
+
+				for (int32 i = 0; i < Lines.Num(); i++)
+				{
+					PDI->DrawLine(Lines[i].Start, Lines[i].End, Lines[i].Color, Lines[i].DepthPriority, Lines[i].Thickness);
+				}
+
+				for (int32 i = 0; i < Points.Num(); i++)
+				{
+					PDI->DrawPoint(Points[i].Position, Points[i].Color, Points[i].PointSize, Points[i].DepthPriority);
+				}
+
+				for (int32 i = 0; i < Meshes.Num(); i++)
+				{
+					static FVector const PosX(1.f,0,0);
+					static FVector const PosY(0,1.f,0);
+					static FVector const PosZ(0,0,1.f);
+
+					FBatchedMesh const& M = Meshes[i];
+
+					// this seems far from optimal in terms of perf, but it's for debugging
+					FDynamicMeshBuilder MeshBuilder;
+
+					// set up geometry
+					for (int32 VertIdx=0; VertIdx < M.MeshVerts.Num(); ++VertIdx)
+					{
+						MeshBuilder.AddVertex( M.MeshVerts[VertIdx], FVector2D::ZeroVector, PosX, PosY, PosZ, FColor::White );
+					}
+					//MeshBuilder.AddTriangles(M.MeshIndices);
+					for (int32 Idx=0; Idx < M.MeshIndices.Num(); Idx+=3)
+					{
+						MeshBuilder.AddTriangle( M.MeshIndices[Idx], M.MeshIndices[Idx+1], M.MeshIndices[Idx+2] );
+					}
+
+					FMaterialRenderProxy* const MaterialRenderProxy = new(FMemStack::Get()) FColoredMaterialRenderProxy(GEngine->DebugMeshMaterial->GetRenderProxy(false), M.Color);
+					MeshBuilder.GetMesh(FMatrix::Identity, MaterialRenderProxy, M.DepthPriority, false, false, ViewIndex, Collector);
+				}
+			}
+		}
+	}
+
+	/** 
+	* Draw the scene proxy as a dynamic element
+	*
+	* @param	PDI - draw interface to render to
+	* @param	View - current view
+	*/
+	virtual void DrawDynamicElements(FPrimitiveDrawInterface* PDI,const FSceneView* View) override
+	{
+		QUICK_SCOPE_CYCLE_COUNTER( STAT_LineBatcherSceneProxy_DrawDynamicElements );
+
+		for (int32 i = 0; i < Lines.Num(); i++)
+		{
+			PDI->DrawLine(Lines[i].Start, Lines[i].End, Lines[i].Color, Lines[i].DepthPriority, Lines[i].Thickness);
+		}
+
+		for (int32 i = 0; i < Points.Num(); i++)
+		{
+			PDI->DrawPoint(Points[i].Position, Points[i].Color, Points[i].PointSize, Points[i].DepthPriority);
+		}
+
+		for (int32 i = 0; i < Meshes.Num(); i++)
+		{
+			static FVector const PosX(1.f,0,0);
+			static FVector const PosY(0,1.f,0);
+			static FVector const PosZ(0,0,1.f);
+
+			FBatchedMesh const& M = Meshes[i];
+
+			// this seems far from optimal in terms of perf, but it's for debugging
+			FDynamicMeshBuilder MeshBuilder;
+
+			// set up geometry
+			for (int32 VertIdx=0; VertIdx < M.MeshVerts.Num(); ++VertIdx)
+			{
+				MeshBuilder.AddVertex( M.MeshVerts[VertIdx], FVector2D::ZeroVector, PosX, PosY, PosZ, FColor::White );
+			}
+			//MeshBuilder.AddTriangles(M.MeshIndices);
+			for (int32 Idx=0; Idx < M.MeshIndices.Num(); Idx+=3)
+			{
+				MeshBuilder.AddTriangle( M.MeshIndices[Idx], M.MeshIndices[Idx+1], M.MeshIndices[Idx+2] );
+			}
+
+			FMaterialRenderProxy* const MaterialRenderProxy = new(FMemStack::Get()) FColoredMaterialRenderProxy(GEngine->DebugMeshMaterial->GetRenderProxy(false), M.Color);
+			MeshBuilder.Draw(PDI, FMatrix::Identity, MaterialRenderProxy, M.DepthPriority);
+		}
+	}
+
+	/**
+	*  Returns a struct that describes to the renderer when to draw this proxy.
+	*	@param		Scene view to use to determine our relevence.
+	*  @return		View relevance struct
+	*/
+	virtual FPrimitiveViewRelevance GetViewRelevance(const FSceneView* View)
+	{
+		return ViewRelevance;
+	}
+	virtual uint32 GetMemoryFootprint( void ) const { return( sizeof( *this ) + GetAllocatedSize() ); }
+	uint32 GetAllocatedSize( void ) const 
+	{ 
+		return( FPrimitiveSceneProxy::GetAllocatedSize() + Lines.GetAllocatedSize() + Points.GetAllocatedSize() + Meshes.GetAllocatedSize() ); 
+	}
+
+private:
+	TArray<FBatchedLine> Lines;
+	TArray<FBatchedPoint> Points;
+	TArray<FBatchedMesh> Meshes;
+	FPrimitiveViewRelevance ViewRelevance;
+};
+
+ULineBatchComponent::ULineBatchComponent(const FObjectInitializer& ObjectInitializer)
+	: Super(ObjectInitializer)
+{
+	bAutoActivate = true;
+	bTickInEditor = true;
+	PrimaryComponentTick.bCanEverTick = true;
+	BodyInstance.bEnableCollision_DEPRECATED = false;
+
+	SetCollisionProfileName(UCollisionProfile::NoCollision_ProfileName);
+
+	bUseEditorCompositing = true;
+	bGenerateOverlapEvents = false;
+}
+
+void ULineBatchComponent::DrawLine(const FVector& Start,const FVector& End,const FLinearColor& Color,uint8 DepthPriority,const float Thickness, const float LifeTime)
+{
+	new(BatchedLines) FBatchedLine(Start,End,Color,LifeTime,Thickness,DepthPriority);
+	// LineBatcher and PersistentLineBatcher components will be updated at the end of UWorld::Tick
+	MarkRenderStateDirty();
+}
+
+void ULineBatchComponent::DrawLines(const TArray<FBatchedLine>& InLines)
+{
+	BatchedLines.Append(InLines);
+	// LineBatcher and PersistentLineBatcher components will be updated at the end of UWorld::Tick
+	MarkRenderStateDirty();
+}
+
+void ULineBatchComponent::DrawPoint(
+	const FVector& Position,
+	const FLinearColor& Color,
+	float PointSize,
+	uint8 DepthPriority,
+	float LifeTime
+	)
+{
+	new(BatchedPoints) FBatchedPoint(Position, Color, PointSize, LifeTime, DepthPriority);
+	// LineBatcher and PersistentLineBatcher components will be updated at the end of UWorld::Tick
+	MarkRenderStateDirty();
+}
+
+void ULineBatchComponent::DrawBox(const FBox& Box, const FMatrix& TM, const FColor& Color, uint8 DepthPriorityGroup)
+{
+	FVector	B[2], P, Q;
+	int32 ai, aj;
+	B[0] = Box.Min;
+	B[1] = Box.Max;
+
+	for( ai=0; ai<2; ai++ ) for( aj=0; aj<2; aj++ )
+	{
+		P.X=B[ai].X; Q.X=B[ai].X;
+		P.Y=B[aj].Y; Q.Y=B[aj].Y;
+		P.Z=B[0].Z; Q.Z=B[1].Z;
+		new(BatchedLines) FBatchedLine(TM.TransformPosition(P), TM.TransformPosition(Q), Color, DefaultLifeTime, 0.0f, DepthPriorityGroup);
+
+		P.Y=B[ai].Y; Q.Y=B[ai].Y;
+		P.Z=B[aj].Z; Q.Z=B[aj].Z;
+		P.X=B[0].X; Q.X=B[1].X;
+		new(BatchedLines) FBatchedLine(TM.TransformPosition(P), TM.TransformPosition(Q), Color, DefaultLifeTime, 0.0f, DepthPriorityGroup);
+
+		P.Z=B[ai].Z; Q.Z=B[ai].Z;
+		P.X=B[aj].X; Q.X=B[aj].X;
+		P.Y=B[0].Y; Q.Y=B[1].Y;
+		new(BatchedLines) FBatchedLine(TM.TransformPosition(P), TM.TransformPosition(Q), Color, DefaultLifeTime, 0.0f, DepthPriorityGroup);
+	}
+	// LineBatcher and PersistentLineBatcher components will be updated at the end of UWorld::Tick
+	MarkRenderStateDirty();
+}
+
+void ULineBatchComponent::DrawSolidBox(const FBox& Box, const FTransform& Xform, const FColor& Color, uint8 DepthPriority, float LifeTime)
+{
+	int32 const NewMeshIdx = BatchedMeshes.Add(FBatchedMesh());
+	FBatchedMesh& BM = BatchedMeshes[NewMeshIdx];
+
+	BM.Color = Color;
+	BM.DepthPriority = DepthPriority;
+	BM.RemainingLifeTime = LifeTime;
+
+	BM.MeshVerts.AddUninitialized(8);
+	BM.MeshVerts[0] = Xform.TransformPosition( FVector(Box.Min.X, Box.Min.Y, Box.Max.Z) );
+	BM.MeshVerts[1] = Xform.TransformPosition( FVector(Box.Max.X, Box.Min.Y, Box.Max.Z) );
+	BM.MeshVerts[2] = Xform.TransformPosition( FVector(Box.Min.X, Box.Min.Y, Box.Min.Z) );
+	BM.MeshVerts[3] = Xform.TransformPosition( FVector(Box.Max.X, Box.Min.Y, Box.Min.Z) );
+	BM.MeshVerts[4] = Xform.TransformPosition( FVector(Box.Min.X, Box.Max.Y, Box.Max.Z) );
+	BM.MeshVerts[5] = Xform.TransformPosition( FVector(Box.Max.X, Box.Max.Y, Box.Max.Z) );
+	BM.MeshVerts[6] = Xform.TransformPosition( FVector(Box.Min.X, Box.Max.Y, Box.Min.Z) );
+	BM.MeshVerts[7] = Xform.TransformPosition( FVector(Box.Max.X, Box.Max.Y, Box.Min.Z) );
+
+	// clockwise
+	BM.MeshIndices.AddUninitialized(36);
+	int32 const Indices[36] = {	3,2,0,
+		3,0,1,
+		7,3,1,
+		7,1,5,
+		6,7,5,
+		6,5,4,
+		2,6,4,
+		2,4,0,
+		1,0,4,
+		1,4,5,
+		7,6,2,
+		7,2,3	};
+
+	for (int32 Idx=0; Idx<36; ++Idx)
+	{
+		BM.MeshIndices[Idx] = Indices[Idx];
+	}
+
+	MarkRenderStateDirty();
+}
+
+void ULineBatchComponent::DrawMesh(TArray<FVector> const& Verts, TArray<int32> const& Indices, FColor const& Color, uint8 DepthPriority, float LifeTime)
+{
+	// modifying array element directly to avoid copying arrays
+	int32 const NewMeshIdx = BatchedMeshes.Add(FBatchedMesh());
+	FBatchedMesh& BM = BatchedMeshes[NewMeshIdx];
+
+	BM.MeshIndices = Indices;
+	BM.MeshVerts = Verts;
+	BM.Color = Color;
+	BM.DepthPriority = DepthPriority;
+	BM.RemainingLifeTime = LifeTime;
+
+	MarkRenderStateDirty();
+}
+
+void ULineBatchComponent::DrawDirectionalArrow(const FMatrix& ArrowToWorld,FColor InColor,float Length,float ArrowSize,uint8 DepthPriority)
+{
+	const FVector Tip = ArrowToWorld.TransformPosition(FVector(Length,0,0));
+	new(BatchedLines) FBatchedLine(Tip,ArrowToWorld.TransformPosition(FVector::ZeroVector),InColor,DefaultLifeTime,0.0f,DepthPriority);
+	new(BatchedLines) FBatchedLine(Tip,ArrowToWorld.TransformPosition(FVector(Length-ArrowSize,+ArrowSize,+ArrowSize)),InColor,DefaultLifeTime,0.0f,DepthPriority);
+	new(BatchedLines) FBatchedLine(Tip,ArrowToWorld.TransformPosition(FVector(Length-ArrowSize,+ArrowSize,-ArrowSize)),InColor,DefaultLifeTime,0.0f,DepthPriority);
+	new(BatchedLines) FBatchedLine(Tip,ArrowToWorld.TransformPosition(FVector(Length-ArrowSize,-ArrowSize,+ArrowSize)),InColor,DefaultLifeTime,0.0f,DepthPriority);
+	new(BatchedLines) FBatchedLine(Tip,ArrowToWorld.TransformPosition(FVector(Length-ArrowSize,-ArrowSize,-ArrowSize)),InColor,DefaultLifeTime,0.0f,DepthPriority);
+	MarkRenderStateDirty();
+}
+
+/** Draw a circle */
+void ULineBatchComponent::DrawCircle(const FVector& Base,const FVector& X,const FVector& Y,FColor Color,float Radius,int32 NumSides,uint8 DepthPriority)
+{
+	const float	AngleDelta = 2.0f * PI / NumSides;
+	FVector	LastVertex = Base + X * Radius;
+
+	for(int32 SideIndex = 0;SideIndex < NumSides;SideIndex++)
+	{
+		const FVector Vertex = Base + (X * FMath::Cos(AngleDelta * (SideIndex + 1)) + Y * FMath::Sin(AngleDelta * (SideIndex + 1))) * Radius;
+		new(BatchedLines) FBatchedLine(LastVertex,Vertex,Color,DefaultLifeTime,0.0f,DepthPriority);
+		LastVertex = Vertex;
+	}
+
+	MarkRenderStateDirty();
+}
+
+void ULineBatchComponent::TickComponent(float DeltaTime, enum ELevelTick TickType, FActorComponentTickFunction *ThisTickFunction)
+{
+	bool bDirty = false;
+	// Update the life time of batched lines, removing the lines which have expired.
+	for(int32 LineIndex=0; LineIndex < BatchedLines.Num(); LineIndex++)
+	{
+		FBatchedLine& Line = BatchedLines[LineIndex];
+		if (Line.RemainingLifeTime > 0.0f)
+		{
+			Line.RemainingLifeTime -= DeltaTime;
+			if(Line.RemainingLifeTime <= 0.0f)
+			{
+				// The line has expired, remove it.
+				BatchedLines.RemoveAt(LineIndex--);
+				bDirty = true;
+			}
+		}
+	}
+
+	// Update the life time of batched points, removing the points which have expired.
+	for(int32 PtIndex=0; PtIndex < BatchedPoints.Num(); PtIndex++)
+	{
+		FBatchedPoint& Pt = BatchedPoints[PtIndex];
+		if (Pt.RemainingLifeTime > 0.0f)
+		{
+			Pt.RemainingLifeTime -= DeltaTime;
+			if(Pt.RemainingLifeTime <= 0.0f)
+			{
+				// The point has expired, remove it.
+				BatchedPoints.RemoveAt(PtIndex--);
+				bDirty = true;
+			}
+		}
+	}
+
+	// Update the life time of batched meshes, removing the meshes which have expired.
+	for(int32 MeshIndex=0; MeshIndex < BatchedMeshes.Num(); MeshIndex++)
+	{
+		FBatchedMesh& Mesh = BatchedMeshes[MeshIndex];
+		if (Mesh.RemainingLifeTime > 0.0f)
+		{
+			Mesh.RemainingLifeTime -= DeltaTime;
+			if(Mesh.RemainingLifeTime <= 0.0f)
+			{
+				// The mesh has expired, remove it.
+				BatchedMeshes.RemoveAt(MeshIndex--);
+				bDirty = true;
+			}
+		}
+	}
+
+	if(bDirty)
+	{
+		MarkRenderStateDirty();
+	}
+}
+
+void ULineBatchComponent::ApplyWorldOffset(const FVector& InOffset, bool bWorldShift)
+{
+	Super::ApplyWorldOffset(InOffset, bWorldShift);
+
+	bool bDirty = false;
+	for (FBatchedLine& Line : BatchedLines)
+	{
+		Line.Start += InOffset;
+		Line.End += InOffset;
+		bDirty = true;
+	}
+
+	for (FBatchedPoint& Point : BatchedPoints)
+	{
+		Point.Position += InOffset;
+		bDirty = true;
+	}
+
+	for (FBatchedMesh& Mesh : BatchedMeshes)
+	{
+		for (FVector& Vert : Mesh.MeshVerts)
+		{
+			Vert += InOffset;
+			bDirty = true;
+		}
+	}
+
+	if (bDirty)
+	{
+		MarkRenderStateDirty();
+	}
+}
+
+/**
+* Creates a new scene proxy for the line batcher component.
+* @return	Pointer to the FLineBatcherSceneProxy
+*/
+FPrimitiveSceneProxy* ULineBatchComponent::CreateSceneProxy()
+{
+	return new FLineBatcherSceneProxy(this);
+}
+
+FBoxSphereBounds ULineBatchComponent::CalcBounds( const FTransform& LocalToWorld ) const 
+{
+	const FVector BoxExtent(HALF_WORLD_MAX);
+	return FBoxSphereBounds(FVector::ZeroVector, BoxExtent, BoxExtent.Size());
+}
+
+void ULineBatchComponent::Flush()
+{
+	if (BatchedLines.Num() > 0 || BatchedPoints.Num() > 0 || BatchedMeshes.Num() > 0)
+	{
+		BatchedLines.Empty();
+		BatchedPoints.Empty();
+		BatchedMeshes.Empty();
+		MarkRenderStateDirty();
+	}
+}
