Fixing up a licensee crash in static mesh vert-painting (flushing render commands when a resource needs to be destroyed).

diff --git a/Engine/Source/Runtime/Engine/Private/StaticMeshComponent.cpp b/Engine/Source/Runtime/Engine/Private/StaticMeshComponent.cpp
--- a/Engine/Source/Runtime/Engine/Private/StaticMeshComponent.cpp
+++ b/Engine/Source/Runtime/Engine/Private/StaticMeshComponent.cpp
@@ -764,16 +764,14 @@ void UStaticMeshComponent::CopyInstanceVertexColorsIfCompatible( UStaticMeshComp
 	{
 		Modify();
 
-		// Remove any and all vertex colors from the target static mesh, if they exist.
-		if ( IsRenderStateCreated() )
+		bool bRegistered = IsRegistered();
+		FComponentReregisterContext ReregisterContext(this);
+		if (bRegistered)
 		{
-			FComponentReregisterContext ComponentReregisterContext( this );
-			RemoveInstanceVertexColors();
-		}
-		else
-		{
-			RemoveInstanceVertexColors();
+			FlushRenderingCommands(); // don't sync threads unless we have to
 		}
+		// Remove any and all vertex colors from the target static mesh, if they exist.
+		RemoveInstanceVertexColors();
 
 		int32 NumSourceLODs = SourceComponent->StaticMesh->GetNumLODs();
 
