Cleaned up UProperties and UFunctions Navigation vs AI|Navigation categories use #UE4

- "Navigation" is to be used for properties
- "AI|Navigation" is for ufunctions

diff --git a/Engine/Source/Runtime/Engine/Classes/AI/Navigation/NavModifierVolume.h b/Engine/Source/Runtime/Engine/Classes/AI/Navigation/NavModifierVolume.h
--- a/Engine/Source/Runtime/Engine/Classes/AI/Navigation/NavModifierVolume.h
+++ b/Engine/Source/Runtime/Engine/Classes/AI/Navigation/NavModifierVolume.h
@@ -8,7 +8,7 @@
 /** 
  *	Allows applying selected AreaClass to navmesh, using Volume's shape
  */
-UCLASS(hidecategories=(Navigation, "AI|Navigation"))
+UCLASS(hidecategories=(Navigation))
 class ANavModifierVolume : public AVolume, public INavRelevantInterface
 {
 	GENERATED_UCLASS_BODY()
diff --git a/Engine/Source/Runtime/Engine/Classes/AI/Navigation/NavRelevantComponent.h b/Engine/Source/Runtime/Engine/Classes/AI/Navigation/NavRelevantComponent.h
--- a/Engine/Source/Runtime/Engine/Classes/AI/Navigation/NavRelevantComponent.h
+++ b/Engine/Source/Runtime/Engine/Classes/AI/Navigation/NavRelevantComponent.h
@@ -23,7 +23,7 @@ class ENGINE_API UNavRelevantComponent : public UActorComponent, public INavRele
 
 	virtual void CalcBounds();
 
-	UFUNCTION(BlueprintCallable, Category="Navigation")
+	UFUNCTION(BlueprintCallable, Category="AI|Navigation")
 	void SetNavigationRelevancy(bool bRelevant);
 
 	/** force refresh in navigation octree */
diff --git a/Engine/Source/Runtime/Engine/Classes/AI/Navigation/NavigationSystem.h b/Engine/Source/Runtime/Engine/Classes/AI/Navigation/NavigationSystem.h
--- a/Engine/Source/Runtime/Engine/Classes/AI/Navigation/NavigationSystem.h
+++ b/Engine/Source/Runtime/Engine/Classes/AI/Navigation/NavigationSystem.h
@@ -539,7 +539,7 @@ public:
 	FORCEINLINE bool IsNavigationBuildingLocked() const { return bNavigationBuildingLocked || bInitialBuildingLockActive; }
 
 	// @todo document
-	UFUNCTION(BlueprintCallable, Category = Navigation)
+	UFUNCTION(BlueprintCallable, Category = "AI|Navigation")
 	void OnNavigationBoundsUpdated(ANavMeshBoundsVolume* NavVolume);
 	void OnNavigationBoundsAdded(ANavMeshBoundsVolume* NavVolume);
 	void OnNavigationBoundsRemoved(ANavMeshBoundsVolume* NavVolume);
diff --git a/Engine/Source/Runtime/Engine/Classes/GameFramework/Pawn.h b/Engine/Source/Runtime/Engine/Classes/GameFramework/Pawn.h
--- a/Engine/Source/Runtime/Engine/Classes/GameFramework/Pawn.h
+++ b/Engine/Source/Runtime/Engine/Classes/GameFramework/Pawn.h
@@ -26,7 +26,7 @@ class UPrimitiveComponent;
  *
  * @see https://docs.unrealengine.com/latest/INT/Gameplay/Framework/Pawn/
  */
-UCLASS(abstract, config=Game, BlueprintType, Blueprintable, hidecategories=(Navigation, "AI|Navigation"))
+UCLASS(abstract, config=Game, BlueprintType, Blueprintable, hidecategories=(Navigation))
 class ENGINE_API APawn : public AActor, public INavAgentInterface
 {
 	GENERATED_BODY()
@@ -222,7 +222,7 @@ public:
 
 	/** Use SetCanAffectNavigationGeneration to change this value at runtime.
 	 *	Note that calling this function at runtime will result in any navigation change only if runtime navigation generation is enabled. */
-	UFUNCTION(BlueprintCallable, Category="Navigation")
+	UFUNCTION(BlueprintCallable, Category="AI|Navigation")
 	void SetCanAffectNavigationGeneration(bool bNewValue);
 
 	/** update all components relevant for navigation generators to match bCanAffectNavigationGeneration flag */
