Added to SplineComponent: get/set spline point type; set editor colors for selected/unselected spline segments.

diff --git a/Engine/Source/Editor/ComponentVisualizers/Private/SplineComponentVisualizer.cpp b/Engine/Source/Editor/ComponentVisualizers/Private/SplineComponentVisualizer.cpp
--- a/Engine/Source/Editor/ComponentVisualizers/Private/SplineComponentVisualizer.cpp
+++ b/Engine/Source/Editor/ComponentVisualizers/Private/SplineComponentVisualizer.cpp
@@ -140,8 +140,8 @@ void FSplineComponentVisualizer::DrawVisualization(const UActorComponent* Compon
 
 		USplineComponent* EditedSplineComp = GetEditedSplineComponent();
 
-		const FColor NormalColor(255, 255, 255);
-		const FColor SelectedColor(255, 0, 0);
+		const FColor NormalColor = SplineComp->EditorUnselectedSplineSegmentColor;
+		const FColor SelectedColor = SplineComp->EditorSelectedSplineSegmentColor;
 		const float GrabHandleSize = 12.0f;
 
 		const FInterpCurveVector& SplineInfo = SplineComp->SplineInfo;
diff --git a/Engine/Source/Runtime/Engine/Classes/Components/SplineComponent.h b/Engine/Source/Runtime/Engine/Classes/Components/SplineComponent.h
--- a/Engine/Source/Runtime/Engine/Classes/Components/SplineComponent.h
+++ b/Engine/Source/Runtime/Engine/Classes/Components/SplineComponent.h
@@ -4,6 +4,22 @@
 #pragma once
 #include "SplineComponent.generated.h"
 
+
+
+/** Permitted spline point types for SplineComponent */
+UENUM()
+namespace ESplinePointType
+{
+	enum Type
+	{
+		Linear,
+		Curve,
+		Constant,
+		CurveClamped
+	};
+}
+
+
 /** 
  *	A spline component is a spline shape which can be used for other purposes (e.g. animating objects). It does not contain rendering capabilities itself (outside the editor) 
  *	@see https://docs.unrealengine.com/latest/INT/Resources/ContentExamples/Blueprint_Splines
@@ -46,6 +62,16 @@ private:
 	bool bClosedLoop;
 
 public:
+#if WITH_EDITORONLY_DATA
+	/** Color of an unselected spline component segment in the editor */
+	UPROPERTY(EditAnywhere, BlueprintReadWrite, Category = Editor)
+	FLinearColor EditorUnselectedSplineSegmentColor;
+
+	/** Color of a selected spline component segment in the editor */
+	UPROPERTY(EditAnywhere, BlueprintReadWrite, Category = Editor)
+	FLinearColor EditorSelectedSplineSegmentColor;
+#endif
+
 	// Begin UActorComponent interface.
 	virtual FComponentInstanceDataBase* GetComponentInstanceData() const override;
 	virtual FName GetComponentInstanceDataType() const override;
@@ -87,6 +113,14 @@ public:
 	UFUNCTION(BlueprintCallable, Category = Spline)
 	void SetWorldLocationAtSplinePoint(int32 PointIndex, const FVector& InLocation);
 
+	/** Get the type of a spline point */
+	UFUNCTION(BlueprintCallable, Category = Spline)
+	ESplinePointType::Type GetSplinePointType(int32 PointIndex) const;
+
+	/** Specify the type of a spline point */
+	UFUNCTION(BlueprintCallable, Category = Spline)
+	void SetSplinePointType(int32 PointIndex, ESplinePointType::Type Type);
+
 	/** Get the number of points that make up this spline */
 	UFUNCTION(BlueprintCallable, Category=Spline)
 	int32 GetNumSplinePoints() const;
diff --git a/Engine/Source/Runtime/Engine/Private/SplineComponent.cpp b/Engine/Source/Runtime/Engine/Private/SplineComponent.cpp
--- a/Engine/Source/Runtime/Engine/Private/SplineComponent.cpp
+++ b/Engine/Source/Runtime/Engine/Private/SplineComponent.cpp
@@ -16,6 +16,10 @@ USplineComponent::USplineComponent(const FObjectInitializer& ObjectInitializer)
 	, Duration(1.0f)
 	, bStationaryEndpoints(false)
 	, bClosedLoop(false)
+#if WITH_EDITORONLY_DATA
+	, EditorSelectedSplineSegmentColor(FLinearColor(1.0f, 0.0f, 0.0f))
+	, EditorUnselectedSplineSegmentColor(FLinearColor(1.0f, 1.0f, 1.0f))
+#endif
 {
 	SplineInfo.Points.Reset(10);
 
@@ -360,6 +364,57 @@ void USplineComponent::SetWorldLocationAtSplinePoint(int32 PointIndex, const FVe
 	}
 }
 
+
+ESplinePointType::Type USplineComponent::GetSplinePointType(int32 PointIndex) const
+{
+	if ((PointIndex >= 0) && (PointIndex < SplineInfo.Points.Num()))
+	{
+		switch (SplineInfo.Points[PointIndex].InterpMode)
+		{
+			case CIM_CurveAuto:			return ESplinePointType::Curve;
+			case CIM_CurveAutoClamped:	return ESplinePointType::CurveClamped;
+			case CIM_Linear:			return ESplinePointType::Linear;
+		}
+	}
+
+	return ESplinePointType::Constant;
+}
+
+
+void USplineComponent::SetSplinePointType(int32 PointIndex, ESplinePointType::Type Type)
+{
+	EInterpCurveMode InterpMode = CIM_Constant;
+	switch (Type)
+	{
+		case ESplinePointType::Curve:			InterpMode = CIM_CurveAuto; break;
+		case ESplinePointType::CurveClamped:	InterpMode = CIM_CurveAutoClamped; break;
+		case ESplinePointType::Linear:			InterpMode = CIM_Linear; break;
+	}
+
+	const int32 NumPoints = SplineInfo.Points.Num();
+
+	if ((PointIndex >= 0) && (PointIndex < NumPoints))
+	{
+		SplineInfo.Points[PointIndex].InterpMode = InterpMode;
+
+		if (IsClosedLoop())
+		{
+			// In a closed loop, the first and last points are tied, so update one with the other
+			if (PointIndex == 0)
+			{
+				SplineInfo.Points[NumPoints - 1].InterpMode = InterpMode;
+			}
+			else if (PointIndex == NumPoints - 1)
+			{
+				SplineInfo.Points[0].InterpMode = InterpMode;
+			}
+		}
+
+		UpdateSpline();
+	}
+}
+
+
 int32 USplineComponent::GetNumSplinePoints() const
 {
 	return SplineInfo.Points.Num();
