Fixed 'show collision' console command so that new actors are spawned with collision showing.

diff --git a/Engine/Source/Runtime/Engine/Classes/Engine/GameViewportClient.h b/Engine/Source/Runtime/Engine/Classes/Engine/GameViewportClient.h
--- a/Engine/Source/Runtime/Engine/Classes/Engine/GameViewportClient.h
+++ b/Engine/Source/Runtime/Engine/Classes/Engine/GameViewportClient.h
@@ -612,6 +612,9 @@ private:
 		return EnabledStats.Num();
 	}
 
+	/** Process the 'show collision' console command */
+	void ToggleShowCollision();
+
 	/** Delegate handler to see if a stat is enabled on this viewport */
 	void HandleViewportStatCheckEnabled(const TCHAR* InName, bool& bOutCurrentEnabled, bool& bOutOthersEnabled);
 
@@ -624,6 +627,9 @@ private:
 	/** Delegate handler for when all stats are disabled in a viewport */
 	void HandleViewportStatDisableAll(const bool bInAnyViewport);
 
+	/** Delegate handler for when an actor is spawned */
+	void ShowCollisionOnSpawnedActors(AActor* Actor);
+
 private:
 	/** Slate window associated with this viewport client.  The same window may host more than one viewport client. */
 	TWeakPtr<SWindow> Window;
@@ -668,6 +674,9 @@ private:
 	/** Those sound stat flags which are enabled on this viewport */
 	static ESoundShowFlags::Type SoundShowFlags;
 
+	/** Number of viewports which have enabled 'show collision' */
+	static int32 NumViewportsShowingCollision;
+
 	/** Disables splitscreen, useful when game code is in menus, and doesn't want splitscreen on */
 	bool bDisableSplitScreenOverride;
 
@@ -679,7 +688,6 @@ private:
 
 	/** Whether or not the cursor is hidden when the viewport captures the mouse */
 	bool bHideCursorDuringCapture;
-
 };
 
 
diff --git a/Engine/Source/Runtime/Engine/Private/GameViewportClient.cpp b/Engine/Source/Runtime/Engine/Private/GameViewportClient.cpp
--- a/Engine/Source/Runtime/Engine/Private/GameViewportClient.cpp
+++ b/Engine/Source/Runtime/Engine/Private/GameViewportClient.cpp
@@ -48,6 +48,9 @@ FOnScreenshotCaptured UGameViewportClient::ScreenshotCapturedDelegate;
 /** A list of all the stat names which are enabled for this viewport (static so they persist between runs) */
 TArray<FString> UGameViewportClient::EnabledStats;
 
+/** The number of GameViewportClients which have enabled 'show collision' */
+int32 UGameViewportClient::NumViewportsShowingCollision = 0;
+
 /** Those sound stat flags which are enabled on this viewport */
 FViewportClient::ESoundShowFlags::Type UGameViewportClient::SoundShowFlags = FViewportClient::ESoundShowFlags::Disabled;
 
@@ -143,6 +146,12 @@ UGameViewportClient::UGameViewportClient(const FObjectInitializer& ObjectInitial
 
 UGameViewportClient::~UGameViewportClient()
 {
+	if (EngineShowFlags.Collision)
+	{
+		EngineShowFlags.Collision = false;
+		ToggleShowCollision();
+	}
+
 	FCoreDelegates::StatCheckEnabled.RemoveAll(this);
 	FCoreDelegates::StatEnabled.RemoveAll(this);
 	FCoreDelegates::StatDisabled.RemoveAll(this);
@@ -1968,6 +1977,26 @@ bool UGameViewportClient::HandleForceFullscreenCommand( const TCHAR* Cmd, FOutpu
 	return true;
 }
 
+/** Contains the previous state of a primitive before turning on collision visibility */
+struct CollVisibilityState
+{
+	bool bHiddenInGame;
+	bool bVisible;
+
+	CollVisibilityState(bool InHidden, bool InVisible) :
+		bHiddenInGame(InHidden),
+		bVisible(InVisible)
+	{
+	}
+};
+
+typedef TMap<TWeakObjectPtr<UPrimitiveComponent>, CollVisibilityState> CollisionComponentVisibilityMap;
+CollisionComponentVisibilityMap& GetCollisionComponentVisibilityMap()
+{
+	static CollisionComponentVisibilityMap Mapping;
+	return Mapping;
+}
+
 bool UGameViewportClient::HandleShowCommand( const TCHAR* Cmd, FOutputDevice& Ar, UWorld* InWorld )
 {
 #if UE_BUILD_SHIPPING
@@ -2034,135 +2063,58 @@ bool UGameViewportClient::HandleShowCommand( const TCHAR* Cmd, FOutputDevice& Ar
 			bool bIsACollisionFlag = FEngineShowFlags::IsNameThere(Cmd, TEXT("Collision"));
 
 			if(bCanBeToggled)
-				{
-					bool bOldState = EngineShowFlags.GetSingleFlag(FlagIndex);
-
-					if(bIsACollisionFlag && !bOldState)
-					{
-						// we only want one active at a time
-						EngineShowFlags.Collision = 0;
-					}
+			{
+				bool bOldState = EngineShowFlags.GetSingleFlag(FlagIndex);
 
-					EngineShowFlags.SetSingleFlag(FlagIndex, !bOldState);
+				EngineShowFlags.SetSingleFlag(FlagIndex, !bOldState);
 
-					if(FEngineShowFlags::IsNameThere(Cmd, TEXT("Navigation,Cover")))
-					{
-						VerifyPathRenderingComponents();
-					}
+				if(FEngineShowFlags::IsNameThere(Cmd, TEXT("Navigation,Cover")))
+				{
+					VerifyPathRenderingComponents();
+				}
 					
-					if(FEngineShowFlags::IsNameThere(Cmd, TEXT("Volumes")))
+				if(FEngineShowFlags::IsNameThere(Cmd, TEXT("Volumes")))
+				{
+					// TODO: Investigate why this is doesn't appear to work
+					if (AllowDebugViewmodes())
 					{
-						if (AllowDebugViewmodes())
+						// Iterate over all brushes
+						for( TObjectIterator<UBrushComponent> It; It; ++It )
 						{
-							// Iterate over all brushes
-							for( TObjectIterator<UBrushComponent> It; It; ++It )
+							UBrushComponent* BrushComponent = *It;
+							AVolume* Owner = Cast<AVolume>( BrushComponent->GetOwner() );
+
+							// Only bother with volume brushes that belong to the world's scene
+							if( Owner && BrushComponent->GetScene() == GetWorld()->Scene )
 							{
-								UBrushComponent* BrushComponent = *It;
-								AVolume* Owner = Cast<AVolume>( BrushComponent->GetOwner() );
+								// We're expecting this to be in the game at this point
+								check( Owner->GetWorld()->IsGameWorld() );
 
-								// Only bother with volume brushes that belong to the world's scene
-								if( Owner && BrushComponent->GetScene() == GetWorld()->Scene )
+								// Toggle visibility of this volume
+								if( BrushComponent->IsVisible() )
 								{
-									// We're expecting this to be in the game at this point
-									check( Owner->GetWorld()->IsGameWorld() );
-
-									// Toggle visibility of this volume
-									if( BrushComponent->IsVisible() )
-									{
-										Owner->bHidden = true;
-										BrushComponent->SetVisibility( false );
-									}
-									else
-									{
-										Owner->bHidden = false;
-										BrushComponent->SetVisibility( true );
-									}
+									Owner->bHidden = true;
+									BrushComponent->SetVisibility( false );
+								}
+								else
+								{
+									Owner->bHidden = false;
+									BrushComponent->SetVisibility( true );
 								}
 							}
 						}
-						else
-						{
-							Ar.Logf(TEXT("Debug viewmodes not allowed on consoles by default.  See AllowDebugViewmodes()."));
-						}
 					}
-				}
-
-			if(bIsACollisionFlag)
-			{
-				// special case: for the Engine.Collision flag, we need to un-hide any primitive components that collide so their collision geometry gets rendered
-
-				/** Contains the previous state of a primitive before turning on collision visibility */
-				struct CollVisibilityState
-				{
-					bool bHiddenInGame;
-					bool bVisible;
-
-					CollVisibilityState(bool InHidden, bool InVisible) :
-					bHiddenInGame(InHidden),
-						bVisible(InVisible)
+					else
 					{
+						Ar.Logf(TEXT("Debug viewmodes not allowed on consoles by default.  See AllowDebugViewmodes()."));
 					}
-				};
-
-				static TMap< TWeakObjectPtr<UPrimitiveComponent>, CollVisibilityState> Mapping;
-
-				{
-					// Restore state to any object touched above
-					for (TMap< TWeakObjectPtr<UPrimitiveComponent>, CollVisibilityState>::TIterator It(Mapping); It; ++It)
-					{
-						TWeakObjectPtr<UPrimitiveComponent>& PrimitiveComponent = It.Key();
-						if (PrimitiveComponent.IsValid())
-						{
-							const CollVisibilityState& VisState = It.Value();
-							PrimitiveComponent->SetHiddenInGame(VisState.bHiddenInGame);
-							PrimitiveComponent->SetVisibility(VisState.bVisible);
-						}
-					}
-					Mapping.Empty();
 				}
-
-				if (EngineShowFlags.Collision)
-					{
-						for (TObjectIterator<UPrimitiveComponent> It; It; ++It)
-						{
-							UPrimitiveComponent* PrimitiveComponent = *It;
-							if( !PrimitiveComponent->IsVisible() && PrimitiveComponent->IsCollisionEnabled() && PrimitiveComponent->GetScene() == GetWorld()->Scene )
-							{
-								check( PrimitiveComponent->GetOwner() && PrimitiveComponent->GetOwner()->GetWorld() && PrimitiveComponent->GetOwner()->GetWorld()->IsGameWorld() );
-								
-								// Save state before modifying the collision visibility
-								Mapping.Add(PrimitiveComponent, CollVisibilityState(PrimitiveComponent->bHiddenInGame, PrimitiveComponent->bVisible));
-								PrimitiveComponent->SetHiddenInGame(false);
-								PrimitiveComponent->SetVisibility(true);
-							}
-						}
-					}
-
-#if !(UE_BUILD_SHIPPING || UE_BUILD_TEST)
-				if( EngineShowFlags.Collision )
-				{
-					for (FLocalPlayerIterator It((UEngine*)GetOuter(), InWorld); It; ++It)
-					{
-						APlayerController* PC = It->PlayerController;
-						if( PC != NULL && PC->GetPawn() != NULL )
-						{
-							PC->ClientMessage( FString::Printf(TEXT("!!!! Player Pawn %s Collision Info !!!!"), *PC->GetPawn()->GetName()) );
-							if (PC->GetPawn()->GetMovementBase())
-							{
-								PC->ClientMessage( FString::Printf(TEXT("Base %s"), *PC->GetPawn()->GetMovementBase()->GetName() ));
-							}
-							TArray<AActor*> Touching;
-							PC->GetPawn()->GetOverlappingActors(Touching);
-							for( int32 i = 0; i < Touching.Num(); i++ )
-							{
-								PC->ClientMessage( FString::Printf(TEXT("Touching %d: %s"), i, *Touching[i]->GetName() ));
-							}
-						}
-					}
-				}
-#endif
 			}
 
+			if(bIsACollisionFlag)
+			{
+				ToggleShowCollision();
+			}
 
 			return true;
 		}
@@ -2205,6 +2157,102 @@ bool UGameViewportClient::HandleShowCommand( const TCHAR* Cmd, FOutputDevice& Ar
 	return true;
 }
 
+void UGameViewportClient::ToggleShowCollision()
+{
+	// special case: for the Engine.Collision flag, we need to un-hide any primitive components that collide so their collision geometry gets rendered
+	const bool bIsShowingCollision = EngineShowFlags.Collision;
+
+	if (bIsShowingCollision)
+	{
+		NumViewportsShowingCollision++;
+		GetWorld()->AddOnActorSpawnedHandler(FOnActorSpawned::FDelegate::CreateUObject(this, &UGameViewportClient::ShowCollisionOnSpawnedActors));
+	}
+	else
+	{
+		NumViewportsShowingCollision--;
+		check(NumViewportsShowingCollision >= 0);
+		GetWorld()->RemoveOnActorSpawnedHandler(FOnActorSpawned::FDelegate::CreateUObject(this, &UGameViewportClient::ShowCollisionOnSpawnedActors));
+	}
+
+	CollisionComponentVisibilityMap& Mapping = GetCollisionComponentVisibilityMap();
+
+	// Restore state to any object in the map above
+	for (CollisionComponentVisibilityMap::TIterator It(Mapping); It; ++It)
+	{
+		TWeakObjectPtr<UPrimitiveComponent>& PrimitiveComponent = It.Key();
+		if (PrimitiveComponent.IsValid())
+		{
+			const CollVisibilityState& VisState = It.Value();
+			PrimitiveComponent->SetHiddenInGame(VisState.bHiddenInGame);
+			PrimitiveComponent->SetVisibility(VisState.bVisible);
+		}
+	}
+	Mapping.Empty();
+
+	if (NumViewportsShowingCollision > 0)
+	{
+		for (TObjectIterator<UPrimitiveComponent> It; It; ++It)
+		{
+			UPrimitiveComponent* PrimitiveComponent = *It;
+			if (!PrimitiveComponent->IsVisible() && PrimitiveComponent->IsCollisionEnabled() && PrimitiveComponent->GetScene() == GetWorld()->Scene)
+			{
+				check(PrimitiveComponent->GetOwner() && PrimitiveComponent->GetOwner()->GetWorld() && PrimitiveComponent->GetOwner()->GetWorld()->IsGameWorld());
+
+				// Save state before modifying the collision visibility
+				Mapping.Add(PrimitiveComponent, CollVisibilityState(PrimitiveComponent->bHiddenInGame, PrimitiveComponent->bVisible));
+				PrimitiveComponent->SetHiddenInGame(false);
+				PrimitiveComponent->SetVisibility(true);
+			}
+		}
+	}
+
+#if !(UE_BUILD_SHIPPING || UE_BUILD_TEST)
+	if (EngineShowFlags.Collision)
+	{
+		for (FLocalPlayerIterator It((UEngine*)GetOuter(), GetWorld()); It; ++It)
+		{
+			APlayerController* PC = It->PlayerController;
+			if (PC != NULL && PC->GetPawn() != NULL)
+			{
+				PC->ClientMessage(FString::Printf(TEXT("!!!! Player Pawn %s Collision Info !!!!"), *PC->GetPawn()->GetName()));
+				if (PC->GetPawn()->GetMovementBase())
+				{
+					PC->ClientMessage(FString::Printf(TEXT("Base %s"), *PC->GetPawn()->GetMovementBase()->GetName()));
+				}
+				TArray<AActor*> Touching;
+				PC->GetPawn()->GetOverlappingActors(Touching);
+				for (int32 i = 0; i < Touching.Num(); i++)
+				{
+					PC->ClientMessage(FString::Printf(TEXT("Touching %d: %s"), i, *Touching[i]->GetName()));
+				}
+			}
+		}
+	}
+#endif
+}
+
+void UGameViewportClient::ShowCollisionOnSpawnedActors(AActor* Actor)
+{
+	CollisionComponentVisibilityMap& Mapping = GetCollisionComponentVisibilityMap();
+
+	TArray<UPrimitiveComponent*> Components;
+	check(Actor != nullptr);
+	Actor->GetComponents(Components);
+
+	for (auto Component : Components)
+	{
+		if (!Mapping.Contains(Component) && !Component->IsVisible() && Component->IsCollisionEnabled() && Component->GetScene() == GetWorld()->Scene)
+		{
+			check(Component->GetOwner() && Component->GetOwner()->GetWorld() && Component->GetOwner()->GetWorld()->IsGameWorld());
+
+			// Save state before modifying the collision visibility
+			Mapping.Add(Component, CollVisibilityState(Component->bHiddenInGame, Component->bVisible));
+			Component->SetHiddenInGame(false);
+			Component->SetVisibility(true);
+		}
+	}
+}
+
 bool UGameViewportClient::HandleShowLayerCommand( const TCHAR* Cmd, FOutputDevice& Ar, UWorld* InWorld )
 {
 	FString LayerName = FParse::Token(Cmd, 0);
