Added a check to make sure that engine font services had been initialised before we call destroy on it

UE-5104 - Crash occurs when selecting no in the Missing Project Modules window

diff --git a/Engine/Source/Runtime/Engine/Private/UserInterface/EngineFontServices.cpp b/Engine/Source/Runtime/Engine/Private/UserInterface/EngineFontServices.cpp
--- a/Engine/Source/Runtime/Engine/Private/UserInterface/EngineFontServices.cpp
+++ b/Engine/Source/Runtime/Engine/Private/UserInterface/EngineFontServices.cpp
@@ -47,6 +47,11 @@ void FEngineFontServices::Destroy()
 	Instance = nullptr;
 }
 
+bool FEngineFontServices::IsInitialized()
+{
+	return Instance != nullptr;
+}
+
 FEngineFontServices& FEngineFontServices::Get()
 {
 	check(Instance);
diff --git a/Engine/Source/Runtime/Engine/Public/EngineFontServices.h b/Engine/Source/Runtime/Engine/Public/EngineFontServices.h
--- a/Engine/Source/Runtime/Engine/Public/EngineFontServices.h
+++ b/Engine/Source/Runtime/Engine/Public/EngineFontServices.h
@@ -15,6 +15,9 @@ public:
 	/** Destroy the singular instance of this class - must be called from the game thread */
 	static void Destroy();
 
+	/** Check to see if the singular instance of this class is currently initialized and ready */
+	static bool IsInitialized();
+
 	/** Get the singular instance of this class */
 	static FEngineFontServices& Get();
 
diff --git a/Engine/Source/Runtime/Launch/Private/LaunchEngineLoop.cpp b/Engine/Source/Runtime/Launch/Private/LaunchEngineLoop.cpp
--- a/Engine/Source/Runtime/Launch/Private/LaunchEngineLoop.cpp
+++ b/Engine/Source/Runtime/Launch/Private/LaunchEngineLoop.cpp
@@ -2008,7 +2008,10 @@ void FEngineLoop::Exit()
 	}
 
 #if !UE_SERVER
-	FEngineFontServices::Destroy();
+	if ( FEngineFontServices::IsInitialized() )
+	{
+		FEngineFontServices::Destroy();
+	}
 #endif
 
 	// close all windows
