UE4 - Safe gpu particle fixes in prep for Metal

diff --git a/Engine/Source/Runtime/Engine/Private/Particles/ParticleResources.cpp b/Engine/Source/Runtime/Engine/Private/Particles/ParticleResources.cpp
--- a/Engine/Source/Runtime/Engine/Private/Particles/ParticleResources.cpp
+++ b/Engine/Source/Runtime/Engine/Private/Particles/ParticleResources.cpp
@@ -68,13 +68,13 @@ void FParticleScratchVertexBuffer::InitRHI()
 {
 	// Create a scratch vertex buffer for injecting particles and rendering tiles.
 	uint32 Flags = BUF_Volatile;
-	if (GetFeatureLevel() >= ERHIFeatureLevel::SM4)
+	if (GetFeatureLevel() >= ERHIFeatureLevel::ES3_1)
 	{
 		Flags |= BUF_ShaderResource;
 	}
 	FRHIResourceCreateInfo CreateInfo;
 	VertexBufferRHI = RHICreateVertexBuffer(GParticleScratchVertexBufferSize, Flags, CreateInfo);
-	if (GetFeatureLevel() >= ERHIFeatureLevel::SM4)
+	if (GetFeatureLevel() >= ERHIFeatureLevel::ES3_1)
 	{
 		VertexBufferSRV_G32R32F = RHICreateShaderResourceView( VertexBufferRHI, /*Stride=*/ sizeof(FVector2D), PF_G32R32F );
 	}
diff --git a/Engine/Source/Runtime/Engine/Private/VectorField.cpp b/Engine/Source/Runtime/Engine/Private/VectorField.cpp
--- a/Engine/Source/Runtime/Engine/Private/VectorField.cpp
+++ b/Engine/Source/Runtime/Engine/Private/VectorField.cpp
@@ -188,7 +188,7 @@ public:
 	 */
 	virtual void InitRHI() override
 	{
-		if (VolumeData && GetFeatureLevel() >= ERHIFeatureLevel::SM4)
+		if (VolumeData && GetFeatureLevel() >= ERHIFeatureLevel::ES3_1)
 		{
 			const uint32 DataSize = SizeX * SizeY * SizeZ * sizeof(FFloat16Color);
 			FVectorFieldStaticResourceBulkDataInterface BulkDataInterface(VolumeData, DataSize);
@@ -755,7 +755,7 @@ public:
 	 */
 	virtual void InitRHI() override
 	{
-		if (GetFeatureLevel() >= ERHIFeatureLevel::SM4)
+		if (GetFeatureLevel() >= ERHIFeatureLevel::ES3_1)
 		{
 			check(SizeX > 0);
 			check(SizeY > 0);
diff --git a/Engine/Source/Runtime/RenderCore/Private/RenderUtils.cpp b/Engine/Source/Runtime/RenderCore/Private/RenderUtils.cpp
--- a/Engine/Source/Runtime/RenderCore/Private/RenderUtils.cpp
+++ b/Engine/Source/Runtime/RenderCore/Private/RenderUtils.cpp
@@ -244,7 +244,7 @@ public:
 	 */
 	virtual void InitRHI() override
 	{
-		if (GetFeatureLevel() >= ERHIFeatureLevel::SM4)
+		if (GetFeatureLevel() >= ERHIFeatureLevel::ES3_1)
 		{
 			// Create the texture.
 			FBlackVolumeTextureResourceBulkDataInterface BlackTextureBulkData;
@@ -513,7 +513,7 @@ public:
 FTexture* GBlackCubeArrayTexture = new TGlobalResource<FBlackCubeArrayTexture>;
 
 /*
- 	3 XYZ packed in 4 bytes. (11:11:10 for X:Y:Z)
+	3 XYZ packed in 4 bytes. (11:11:10 for X:Y:Z)
 */
 
 /**
