Full screen mode changes
 - Removed option 2 from r.FullScreenMode
 - Leaves 0 = Native, 1 = Windowed
 - FSystemResolution::RequestResolutionChange() sets the r.setres cvar as a 'console' so that game systems can still change the resolution after the user plays with the cvar

Stop game user settings from incorrectly interpreting r.FullScreenMode as an EWindowMode. It's not.

Simplify ShooterGame and PlatformerGame's fullscreen mode. Make them obey r.FullScreenMode when processing menu fullscreen flags.

diff --git a/Engine/Source/Runtime/Core/Private/HAL/ConsoleManager.cpp b/Engine/Source/Runtime/Core/Private/HAL/ConsoleManager.cpp
--- a/Engine/Source/Runtime/Core/Private/HAL/ConsoleManager.cpp
+++ b/Engine/Source/Runtime/Core/Private/HAL/ConsoleManager.cpp
@@ -1687,11 +1687,10 @@ static TAutoConsoleVariable<int32> CVarPostProcessAAQuality(
 
 static TAutoConsoleVariable<int32> CVarFullscreenMode(
 	TEXT("r.FullScreenMode"),
-	2,
+	1,
 	TEXT("Defines how we do full screen when requested (e.g. command line option -fullscreen or in ini [SystemSettings] fullscreen=true)\n")
 	TEXT(" 0: normal full screen (renders faster, more control over vsync, less GPU memory, 10bit color if possible)\n")
-	TEXT(" 1: windowed full screen, desktop resolution (quick switch between applications and window mode, full quality)\n")
-	TEXT(" 2: windowed full screen, specified resolution (like 1 but no unintuitive performance cliff, can be blurry, default)\n")
+	TEXT(" 1: windowed full screen (quick switch between applications and window mode, slight performance loss)\n")
 	TEXT(" any other number behaves like 0"),
 	ECVF_Scalability);
 
diff --git a/Engine/Source/Runtime/Engine/Private/GameEngine.cpp b/Engine/Source/Runtime/Engine/Private/GameEngine.cpp
--- a/Engine/Source/Runtime/Engine/Private/GameEngine.cpp
+++ b/Engine/Source/Runtime/Engine/Private/GameEngine.cpp
@@ -62,31 +62,6 @@ static FAutoConsoleCommand GDumpDrawListStatsCmd(
 	FConsoleCommandWithArgsDelegate::CreateStatic(&RunSynthBenchmark)
 	);
 
-
-int32 GetBoundFullScreenModeCVar()
-{
-	if (GEngine && GEngine->HMDDevice.IsValid() && GEngine->HMDDevice->IsHMDConnected())
-	{
-		// For HMD, Fullscreen mode should be always 0 (normal fullscreen).
-		return 0;
-	}
-	static const auto CVar = IConsoleManager::Get().FindTConsoleVariableDataInt(TEXT("r.FullScreenMode")); 
-
-	if (FPlatformProperties::SupportsWindowedMode())
-	{
-		int32 Value = CVar->GetValueOnGameThread();
-
-		if (Value >= 0 && Value < EWindowMode::NumWindowModes)
-		{
-			return Value;
-		}
-	}
-
-	// every other value behaves like 0
-	return 0;
-}
-
-// depending on WindowMode and the console variable r.FullScreenMode
 EWindowMode::Type GetWindowModeType(EWindowMode::Type WindowMode)
 {
 	if (FPlatformProperties::SupportsWindowedMode())
@@ -192,7 +167,9 @@ void UGameEngine::ConditionallyOverrideSettings(int32& ResolutionX, int32& Resol
 	else if (FParse::Param(FCommandLine::Get(),TEXT("FullScreen")))
 	{
 		// -FullScreen
-		WindowMode = EWindowMode::Fullscreen;
+		auto CVar = IConsoleManager::Get().FindTConsoleVariableDataInt(TEXT("r.FullScreenMode"));
+		check(CVar);
+		WindowMode = CVar->GetValueOnGameThread() == 0 ? EWindowMode::Fullscreen : EWindowMode::WindowedFullscreen;
 	}
 
 	//fullscreen is always supported, but don't allow windowed mode on platforms that dont' support it.
diff --git a/Engine/Source/Runtime/Engine/Private/GameUserSettings.cpp b/Engine/Source/Runtime/Engine/Private/GameUserSettings.cpp
--- a/Engine/Source/Runtime/Engine/Private/GameUserSettings.cpp
+++ b/Engine/Source/Runtime/Engine/Private/GameUserSettings.cpp
@@ -194,11 +194,6 @@ void UGameUserSettings::ApplyNonResolutionSettings()
 	bool bIsDirty = IsDirty();
 	EWindowMode::Type NewWindowMode = GetFullscreenMode();
 
-	{
-		static IConsoleVariable* CVar = IConsoleManager::Get().FindConsoleVariable(TEXT("r.FullScreenMode")); 
-		CVar->Set(NewWindowMode, ECVF_SetByGameSetting);
-	}
-
 	// Update vsync cvar
 	{
 		static auto CVar = IConsoleManager::Get().FindConsoleVariable(TEXT("r.VSync")); 
@@ -220,11 +215,6 @@ void UGameUserSettings::ApplyResolutionSettings(bool bCheckForCommandLineOverrid
 
 	EWindowMode::Type NewFullscreenMode = GetFullscreenMode();
 
-	{
-		static IConsoleVariable* CVar = IConsoleManager::Get().FindConsoleVariable(TEXT("r.FullScreenMode"));
-		CVar->Set(NewFullscreenMode, ECVF_SetByGameSetting);
-	}
-
 	// Request a resolution change
 	RequestResolutionChange(ResolutionSizeX, ResolutionSizeY, NewFullscreenMode, bCheckForCommandLineOverrides);
 	IConsoleManager::Get().CallAllConsoleVariableSinks();
@@ -273,11 +263,6 @@ void UGameUserSettings::RequestResolutionChange(int32 InResolutionX, int32 InRes
 		UGameEngine::ConditionallyOverrideSettings(InResolutionX, InResolutionY, InWindowMode);
 	}
 
-
-	{
-		static IConsoleVariable* CVar = IConsoleManager::Get().FindConsoleVariable(TEXT("r.FullScreenMode"));
-		CVar->Set(InWindowMode, ECVF_SetByGameSetting);
-	}
 	FSystemResolution::RequestResolutionChange(InResolutionX, InResolutionY, InWindowMode);
 }
 
@@ -334,11 +319,6 @@ void UGameUserSettings::PreloadResolutionSettings()
 #endif
 	}
 
-	{
-		static auto CVar = IConsoleManager::Get().FindConsoleVariable(TEXT("r.FullScreenMode")); 
-		CVar->Set((int32)WindowMode, ECVF_SetByGameSetting);
-	}
-
 	RequestResolutionChange(ResolutionX, ResolutionY, WindowMode);
 
 	IConsoleManager::Get().CallAllConsoleVariableSinks();
diff --git a/Engine/Source/Runtime/Engine/Private/GameViewportClient.cpp b/Engine/Source/Runtime/Engine/Private/GameViewportClient.cpp
--- a/Engine/Source/Runtime/Engine/Private/GameViewportClient.cpp
+++ b/Engine/Source/Runtime/Engine/Private/GameViewportClient.cpp
@@ -40,8 +40,6 @@ extern ENGINE_API UPrimitiveComponent* GDebugSelectedComponent;
 /** The lightmap used by the currently selected component, if it's a static mesh component. */
 extern ENGINE_API class FLightMap2D* GDebugSelectedLightmap;
 
-extern int32 GetBoundFullScreenModeCVar();
-
 /** Delegate called at the end of the frame when a screenshot is captured */
 FOnScreenshotCaptured UGameViewportClient::ScreenshotCapturedDelegate;
 
@@ -2511,7 +2509,9 @@ bool UGameViewportClient::SetDisplayConfiguration(const FIntPoint* Dimensions, E
 
 bool UGameViewportClient::HandleToggleFullscreenCommand(const TCHAR* Cmd, FOutputDevice& Ar)
 {
-	auto FullScreenMode = GetBoundFullScreenModeCVar() == 0 ? EWindowMode::Fullscreen : EWindowMode::WindowedFullscreen;
+	auto CVar = IConsoleManager::Get().FindTConsoleVariableDataInt(TEXT("r.FullScreenMode"));
+	check(CVar);
+	auto FullScreenMode = CVar->GetValueOnGameThread() == 0 ? EWindowMode::Fullscreen : EWindowMode::WindowedFullscreen;
 	FullScreenMode = Viewport->IsFullscreen() ? EWindowMode::Windowed : FullScreenMode;
 	if (GEngine->HMDDevice.IsValid() && GEngine->HMDDevice->IsHMDEnabled())
 	{
diff --git a/Engine/Source/Runtime/Engine/Private/LocalPlayer.cpp b/Engine/Source/Runtime/Engine/Private/LocalPlayer.cpp
--- a/Engine/Source/Runtime/Engine/Private/LocalPlayer.cpp
+++ b/Engine/Source/Runtime/Engine/Private/LocalPlayer.cpp
@@ -44,8 +44,6 @@ static TAutoConsoleVariable<int32> CVarViewportTest(
 bool GShouldLogOutAFrameOfMoveComponent = false;
 bool GShouldLogOutAFrameOfSetBodyTransform = false;
 
-extern int32 GetBoundFullScreenModeCVar();
-
 //////////////////////////////////////////////////////////////////////////
 // ULocalPlayer
 
diff --git a/Engine/Source/Runtime/Engine/Private/SceneView.cpp b/Engine/Source/Runtime/Engine/Private/SceneView.cpp
--- a/Engine/Source/Runtime/Engine/Private/SceneView.cpp
+++ b/Engine/Source/Runtime/Engine/Private/SceneView.cpp
@@ -1253,11 +1253,8 @@ void FSceneView::EndFinalPostprocessSettings(const FSceneViewInitOptions& ViewIn
 
 		if (bFullscreen)
 		{
-			extern int32 GetBoundFullScreenModeCVar();
-			int32 WindowModeType = GetBoundFullScreenModeCVar();
-
 			// CVar mode 2 is fullscreen with upscale
-			if(WindowModeType == 2)
+			if(GSystemResolution.WindowMode == EWindowMode::WindowedFullscreen)
 			{
 //				FIntPoint WindowSize = Viewport->GetSizeXY();
 				FIntPoint WindowSize = Family->RenderTarget->GetSizeXY();
diff --git a/Engine/Source/Runtime/Engine/Private/Slate/SceneViewport.cpp b/Engine/Source/Runtime/Engine/Private/Slate/SceneViewport.cpp
--- a/Engine/Source/Runtime/Engine/Private/Slate/SceneViewport.cpp
+++ b/Engine/Source/Runtime/Engine/Private/Slate/SceneViewport.cpp
@@ -9,7 +9,6 @@
 
 #include "IHeadMountedDisplay.h"
 
-extern int32 GetBoundFullScreenModeCVar();
 extern EWindowMode::Type GetWindowModeType(EWindowMode::Type WindowMode);
 
 FSceneViewport::FSceneViewport( FViewportClient* InViewportClient, TSharedPtr<SViewport> InViewportWidget )
@@ -844,13 +843,12 @@ void FSceneViewport::ResizeFrame(uint32 NewSizeX, uint32 NewSizeY, EWindowMode::
 			// scaling so we actual render to the resolution we've been asked for.
 			if (DesiredWindowMode == EWindowMode::WindowedFullscreen)
 			{
-				FSlateRect Rect = WindowToResize->GetFullScreenInfo();
-				if (Rect.IsValid())
-				{
-					NewSizeX = Rect.GetSize().X;
-					NewSizeY = Rect.GetSize().Y;
-				}
+				FDisplayMetrics DisplayMetrics;
+				FSlateApplication::Get().GetInitialDisplayMetrics(DisplayMetrics);
+				NewSizeX = DisplayMetrics.PrimaryDisplayWidth;;
+				NewSizeY = DisplayMetrics.PrimaryDisplayHeight;;
 			}
+
 			uint32 ViewportSizeX = NewSizeX;
 			uint32 ViewportSizeY = NewSizeY;
 
diff --git a/Engine/Source/Runtime/Engine/Private/UnrealEngine.cpp b/Engine/Source/Runtime/Engine/Private/UnrealEngine.cpp
--- a/Engine/Source/Runtime/Engine/Private/UnrealEngine.cpp
+++ b/Engine/Source/Runtime/Engine/Private/UnrealEngine.cpp
@@ -10328,7 +10328,7 @@ void FSystemResolution::RequestResolutionChange(int32 InResX, int32 InResY, EWin
 	}
 
 	FString NewValue = FString::Printf(TEXT("%dx%d%s"), InResX, InResY, *WindowModeSuffix);
-	CVarSystemResolution->Set(*NewValue, ECVF_SetByCode);
+	CVarSystemResolution->Set(*NewValue, ECVF_SetByConsole);
 }
 
 
