Fix for timelines using external curves preventing users loading new maps
#UE-4678 - LIVE: CRASH: Editor crashes when attempting to open a new level from a level that has an external curve in a timeline within the level blueprint
#Branch UE4
#Proj Engine

#Note - the issue appears to be that external curves were being duplicated to the generated class on UTimelineTemplate::PostDuplicate, this was incorrect and left open references to the external curves that later prevented the current map from being unloaded.

diff --git a/Engine/Source/Runtime/Engine/Private/TimelineTemplate.cpp b/Engine/Source/Runtime/Engine/Private/TimelineTemplate.cpp
--- a/Engine/Source/Runtime/Engine/Private/TimelineTemplate.cpp
+++ b/Engine/Source/Runtime/Engine/Private/TimelineTemplate.cpp
@@ -230,13 +230,15 @@ void UTimelineTemplate::PostDuplicate(bool bDuplicateForPIE)
 	}
 #endif // WITH_EDITOR
 
+	const bool bTransientPackage = GetOutermost() == GetTransientPackage();
+
 	for(TArray<struct FTTFloatTrack>::TIterator It = FloatTracks.CreateIterator();It;++It)
 	{
 		FTTFloatTrack& Track = *It;
 		if( Track.CurveFloat != NULL )
 		{
 			// Do not duplicate external curves unless duplicating to a transient package
-			if(!Track.CurveFloat->GetOuter()->IsA(UPackage::StaticClass()) || GetOutermost() == GetTransientPackage())
+			if(!Track.CurveFloat->GetOuter()->IsA(UPackage::StaticClass()) && bTransientPackage )
 			{
 				Track.CurveFloat = DuplicateObject<UCurveFloat>(Track.CurveFloat, NewCurveOuter, *MakeUniqueCurveName(Track.CurveFloat, NewCurveOuter));
 			}
@@ -253,7 +255,7 @@ void UTimelineTemplate::PostDuplicate(bool bDuplicateForPIE)
 		if( Track.CurveKeys != NULL )
 		{
 			// Do not duplicate external curves unless duplicating to a transient package
-			if(!Track.CurveKeys->GetOuter()->IsA(UPackage::StaticClass()) || GetOutermost() == GetTransientPackage())
+			if(!Track.CurveKeys->GetOuter()->IsA(UPackage::StaticClass()) && bTransientPackage )
 			{
 				Track.CurveKeys = DuplicateObject<UCurveFloat>(Track.CurveKeys, NewCurveOuter, *MakeUniqueCurveName(Track.CurveKeys, NewCurveOuter));
 			}
@@ -270,7 +272,7 @@ void UTimelineTemplate::PostDuplicate(bool bDuplicateForPIE)
 		if( Track.CurveVector != NULL )
 		{
 			// Do not duplicate external curves unless duplicating to a transient package
-			if(!Track.CurveVector->GetOuter()->IsA(UPackage::StaticClass()) || GetOutermost() == GetTransientPackage())
+			if(!Track.CurveVector->GetOuter()->IsA(UPackage::StaticClass()) && bTransientPackage )
 			{
 				Track.CurveVector = DuplicateObject<UCurveVector>(Track.CurveVector, NewCurveOuter, *MakeUniqueCurveName(Track.CurveVector, NewCurveOuter));
 			}
@@ -287,7 +289,7 @@ void UTimelineTemplate::PostDuplicate(bool bDuplicateForPIE)
 		if( Track.CurveLinearColor != NULL )
 		{
 			// Do not duplicate external curves unless duplicating to a transient package
-			if(!Track.CurveLinearColor->GetOuter()->IsA(UPackage::StaticClass()) || GetOutermost() == GetTransientPackage())
+			if(!Track.CurveLinearColor->GetOuter()->IsA(UPackage::StaticClass()) && bTransientPackage )
 			{
 				Track.CurveLinearColor = DuplicateObject<UCurveLinearColor>(Track.CurveLinearColor, NewCurveOuter, *MakeUniqueCurveName(Track.CurveLinearColor, NewCurveOuter));
 			}
