Fixed crash when using DOF lock constraints with ragdolls that have enough bodies to create an aggregate.

diff --git a/Engine/Source/Runtime/Engine/Private/PhysicsEngine/BodyInstance.cpp b/Engine/Source/Runtime/Engine/Private/PhysicsEngine/BodyInstance.cpp
--- a/Engine/Source/Runtime/Engine/Private/PhysicsEngine/BodyInstance.cpp
+++ b/Engine/Source/Runtime/Engine/Private/PhysicsEngine/BodyInstance.cpp
@@ -1381,22 +1381,26 @@ void FBodyInstance::InitBody(UBodySetup* Setup, const FTransform& Transform, UPr
 		int32 VelocityIterCount = FMath::Clamp(VelocitySolverIterationCount, 1, 255);
 		PNewDynamic->setSolverIterationCounts(PositionIterCount, VelocityIterCount);
 
-		CreateDOFLock();
-
-		// wakeUp and putToSleep will issue warnings on kinematic actors
-		if (IsRigidBodyNonKinematic(PNewDynamic) && PNewDynamic->getScene())
-		{
-		    // Sleep/wake up as appropriate
-		    if (bShouldStartAwake)
-		    {
-			    // Wake up bodies that are part of a moving actor.
-			    PNewDynamic->wakeUp();
-		    }
-		    else
-		    {
-			    // Bodies should start out sleeping.
-			    PNewDynamic->putToSleep();
-		    }
+		if(PNewDynamic->getScene())
+		{
+			// A scene is required to build the DOF constraint
+			CreateDOFLock();
+
+			// wakeUp and putToSleep will issue warnings on kinematic actors
+			if(IsRigidBodyNonKinematic(PNewDynamic) && PNewDynamic->getScene())
+			{
+				// Sleep/wake up as appropriate
+				if(bShouldStartAwake)
+				{
+					// Wake up bodies that are part of a moving actor.
+					PNewDynamic->wakeUp();
+				}
+				else
+				{
+					// Bodies should start out sleeping.
+					PNewDynamic->putToSleep();
+				}
+			}
 		}
 	}
 #endif // WITH_PHYSX
diff --git a/Engine/Source/Runtime/Engine/Private/SkeletalMeshComponentPhysics.cpp b/Engine/Source/Runtime/Engine/Private/SkeletalMeshComponentPhysics.cpp
--- a/Engine/Source/Runtime/Engine/Private/SkeletalMeshComponentPhysics.cpp
+++ b/Engine/Source/Runtime/Engine/Private/SkeletalMeshComponentPhysics.cpp
@@ -812,9 +812,13 @@ void USkeletalMeshComponent::InitArticulated(FPhysScene* PhysScene)
 		AActor* Owner = GetOwner();
 		bool bShouldSleep = !BodyInstance.bStartAwake && (Owner && Owner->GetVelocity().SizeSquared() <= KINDA_SMALL_NUMBER);
 
-		if(bShouldSleep)
+		for(FBodyInstance* Body : Bodies)
 		{
-			for(FBodyInstance* Body : Bodies)
+			// Creates a DOF constraint if necessary for the body - also requires the scene to exist within the actor
+			Body->CreateDOFLock();
+
+			// Set to sleep if necessary
+			if(bShouldSleep)
 			{
 				Body->GetPxRigidDynamic()->putToSleep();
 			}
