Don't return true when we really didn't spawn an actor

diff --git a/Engine/Source/Runtime/Engine/Private/PackageMapClient.cpp b/Engine/Source/Runtime/Engine/Private/PackageMapClient.cpp
--- a/Engine/Source/Runtime/Engine/Private/PackageMapClient.cpp
+++ b/Engine/Source/Runtime/Engine/Private/PackageMapClient.cpp
@@ -202,7 +202,7 @@ bool UPackageMapClient::SerializeNewActor(FArchive& Ar, class UActorChannel *Cha
 		}
 
 		UE_LOG( LogNetPackageMap, Log, TEXT( "UPackageMapClient::SerializeNewActor:  Skipping full read because we are deleting dynamic actor: %s" ), Actor ? *Actor->GetName() : TEXT( "NULL" ) );
-		return true;
+		return false;		// This doesn't mean an error. This just simply means we didn't spawn an actor.
 	}
 
 	if (NetGUID.IsDynamic())
