Better message on global shader compile error

diff --git a/Engine/Source/Runtime/Engine/Private/GlobalShader.cpp b/Engine/Source/Runtime/Engine/Private/GlobalShader.cpp
--- a/Engine/Source/Runtime/Engine/Private/GlobalShader.cpp
+++ b/Engine/Source/Runtime/Engine/Private/GlobalShader.cpp
@@ -628,7 +628,7 @@ void ProcessCompiledGlobalShaders(const TArray<FShaderCompileJob*>& CompilationR
 		}
 		else
 		{
-			UE_LOG(LogShaders, Fatal,TEXT("Failed to compile global shader %s"), GlobalShaderType->GetName());
+			UE_LOG(LogShaders, Fatal,TEXT("Failed to compile global shader %s.  Enable 'r.ShaderDevelopmentMode' in ConsoleVariables.ini for retries."), GlobalShaderType->GetName());
 		}
 	}
 
