Don't allow the non-transactional blueprint created components from being referenced by the transaction buffer

diff --git a/Engine/Source/Runtime/Engine/Private/Actor.cpp b/Engine/Source/Runtime/Engine/Private/Actor.cpp
--- a/Engine/Source/Runtime/Engine/Private/Actor.cpp
+++ b/Engine/Source/Runtime/Engine/Private/Actor.cpp
@@ -838,11 +838,39 @@ bool AActor::IsBasedOnActor(const AActor* Other) const
 
 bool AActor::Modify( bool bAlwaysMarkDirty/*=true*/ )
 {
+	// Any properties that reference a blueprint constructed component needs to avoid creating a reference to the component from the transaction
+	// buffer, so we temporarily switch the property to non-transactional while the modify occurs
+	TArray<UObjectProperty*> TemporarilyNonTransactionalProperties;
+	if (GUndo)
+	{
+		for (TFieldIterator<UObjectProperty> PropertyIt(GetClass(), EFieldIteratorFlags::IncludeSuper); PropertyIt; ++PropertyIt)
+		{
+			UObjectProperty* ObjProp = *PropertyIt;
+			if (!ObjProp->HasAllPropertyFlags(CPF_NonTransactional))
+			{
+				UActorComponent* ActorComponent = Cast<UActorComponent>(ObjProp->GetObjectPropertyValue(ObjProp->ContainerPtrToValuePtr<void>(this)));
+				if (ActorComponent && ActorComponent->bCreatedByConstructionScript)
+				{
+					ObjProp->SetPropertyFlags(CPF_NonTransactional);
+					TemporarilyNonTransactionalProperties.Add(ObjProp);
+				}
+			}
+		}
+	}
+
 	bool bSavedToTransactionBuffer = UObject::Modify( bAlwaysMarkDirty );
-	if( RootComponent )
+
+	for (UObjectProperty* ObjProp : TemporarilyNonTransactionalProperties)
 	{
-		bSavedToTransactionBuffer = GetRootComponent()->Modify( bAlwaysMarkDirty ) || bSavedToTransactionBuffer;
+		ObjProp->ClearPropertyFlags(CPF_NonTransactional);
 	}
+
+	// If the root component is blueprint constructed we don't save it to the transaction buffer
+	if( RootComponent && !RootComponent->bCreatedByConstructionScript )
+	{
+		bSavedToTransactionBuffer = RootComponent->Modify( bAlwaysMarkDirty ) || bSavedToTransactionBuffer;
+	}
+
 	return bSavedToTransactionBuffer;
 }
 
