Fixing it so timelines with null curves don't get flagged as differing.

diff --git a/Engine/Source/Runtime/Engine/Private/TimelineTemplate.cpp b/Engine/Source/Runtime/Engine/Private/TimelineTemplate.cpp
--- a/Engine/Source/Runtime/Engine/Private/TimelineTemplate.cpp
+++ b/Engine/Source/Runtime/Engine/Private/TimelineTemplate.cpp
@@ -329,28 +329,43 @@ bool FTTTrackBase::operator==( const FTTTrackBase& T2 ) const
 		   (bIsExternalCurve == T2.bIsExternalCurve);
 }
 
-
 bool FTTEventTrack::operator==( const FTTEventTrack& T2 ) const
 {
-	return FTTTrackBase::operator==(T2) && (CurveKeys && T2.CurveKeys) && 
-		   (*CurveKeys == *T2.CurveKeys);
+	bool bKeyCurvesEqual = (CurveKeys == T2.CurveKeys);
+	if (!bKeyCurvesEqual && CurveKeys && T2.CurveKeys)
+	{
+		bKeyCurvesEqual = (*CurveKeys == *T2.CurveKeys);
+	}
+	return FTTTrackBase::operator==(T2) && bKeyCurvesEqual;
 }
 
 bool FTTFloatTrack::operator==( const FTTFloatTrack& T2 ) const
 {
-	return FTTTrackBase::operator==(T2) &&  (CurveFloat && T2.CurveFloat) && 
-		(*CurveFloat == *T2.CurveFloat);
+	bool bFloatCurvesEqual = (CurveFloat == T2.CurveFloat);
+	if (!bFloatCurvesEqual && CurveFloat && T2.CurveFloat)
+	{
+		bFloatCurvesEqual = (*CurveFloat == *T2.CurveFloat);
+	}
+	return FTTTrackBase::operator==(T2) && bFloatCurvesEqual;
 }
 
 bool FTTVectorTrack::operator==( const FTTVectorTrack& T2 ) const
 {
-	return FTTTrackBase::operator==(T2)  &&  (CurveVector && T2.CurveVector) && 
-		 (*CurveVector == *T2.CurveVector);
+	bool bVectorCurvesEqual = (CurveVector == T2.CurveVector);
+	if (!bVectorCurvesEqual && CurveVector && T2.CurveVector)
+	{
+		bVectorCurvesEqual = (*CurveVector == *T2.CurveVector);
+	}
+	return FTTTrackBase::operator==(T2) && bVectorCurvesEqual;
 }
 
 bool FTTLinearColorTrack::operator==( const FTTLinearColorTrack& T2 ) const
 {
-	return FTTTrackBase::operator==(T2)  &&  (CurveLinearColor && T2.CurveLinearColor) && 
-		 (*CurveLinearColor == *T2.CurveLinearColor);
+	bool bColorCurvesEqual = (CurveLinearColor == T2.CurveLinearColor);
+	if (!bColorCurvesEqual && CurveLinearColor && T2.CurveLinearColor)
+	{
+		bColorCurvesEqual = (*CurveLinearColor == *T2.CurveLinearColor);
+	}
+	return FTTTrackBase::operator==(T2) && bColorCurvesEqual;
 }
 
