Engine: Fix mismatches between WITH_EDITORONLY_DATA and WITH_EDITOR usage in UCameraComponent

diff --git a/Engine/Source/Runtime/Engine/Classes/Camera/CameraComponent.h b/Engine/Source/Runtime/Engine/Classes/Camera/CameraComponent.h
--- a/Engine/Source/Runtime/Engine/Classes/Camera/CameraComponent.h
+++ b/Engine/Source/Runtime/Engine/Classes/Camera/CameraComponent.h
@@ -83,7 +83,7 @@ protected:
 	class UStaticMeshComponent* ProxyMeshComponent;
 #endif
 public:
-#if WITH_EDITOR
+#if WITH_EDITORONLY_DATA
 	// Refreshes the visual components to match the component state
 	ENGINE_API virtual void RefreshVisualRepresentation();
 
diff --git a/Engine/Source/Runtime/Engine/Private/Camera/CameraComponent.cpp b/Engine/Source/Runtime/Engine/Private/Camera/CameraComponent.cpp
--- a/Engine/Source/Runtime/Engine/Private/Camera/CameraComponent.cpp
+++ b/Engine/Source/Runtime/Engine/Private/Camera/CameraComponent.cpp
@@ -41,7 +41,7 @@ UCameraComponent::UCameraComponent(const FObjectInitializer& ObjectInitializer)
 
 void UCameraComponent::OnRegister()
 {
-#if WITH_EDITOR
+#if WITH_EDITORONLY_DATA
 	if (ProxyMeshComponent == NULL)
 	{
 		ProxyMeshComponent = ConstructObject<UStaticMeshComponent>(UStaticMeshComponent::StaticClass(), GetOuter(), NAME_None, RF_Transactional);
@@ -115,7 +115,7 @@ void UCameraComponent::PostLoad()
 }
 
 
-#if WITH_EDITOR
+#if WITH_EDITORONLY_DATA
 void UCameraComponent::RefreshVisualRepresentation()
 {
 	if (DrawFrustum != NULL)
@@ -147,12 +147,16 @@ void UCameraComponent::RestoreFrustumColor()
 		//Cam->DrawFrustum->FrustumColor = DefCam->DrawFrustum->FrustumColor;
 	}
 }
+#endif
 
+#if WITH_EDITOR
 void UCameraComponent::PostEditChangeProperty(FPropertyChangedEvent& PropertyChangedEvent)
 {
 	Super::PostEditChangeProperty(PropertyChangedEvent);
 
+#if WITH_EDITORONLY_DATA
 	RefreshVisualRepresentation();
+#endif
 }
 #endif
 
