Applied licensee fix for a crash during foliage editing
UE-5331

diff --git a/Engine/Source/Runtime/Engine/Private/InstancedFoliage.cpp b/Engine/Source/Runtime/Engine/Private/InstancedFoliage.cpp
--- a/Engine/Source/Runtime/Engine/Private/InstancedFoliage.cpp
+++ b/Engine/Source/Runtime/Engine/Private/InstancedFoliage.cpp
@@ -1660,9 +1660,12 @@ void AInstancedFoliageActor::ApplyWorldOffset(const FVector& InOffset, bool bWor
 			for (int32 InstanceIdx = 0; InstanceIdx < MeshInfo.Instances.Num(); InstanceIdx++)
 			{
 				FFoliageInstance& Instance = MeshInfo.Instances[InstanceIdx];
-				Instance.Location += InOffset;
-				// Rehash instance location
-				MeshInfo.InstanceHash->InsertInstance(Instance.Location, InstanceIdx);
+				if (Instance.ClusterIndex != INDEX_NONE)
+				{
+					Instance.Location += InOffset;
+					// Rehash instance location
+					MeshInfo.InstanceHash->InsertInstance(Instance.Location, InstanceIdx);
+				}
 			}
 
 			for (auto It = MeshInfo.ComponentHash.CreateIterator(); It; ++It)
