Fixed crash in Ulevel::PostLoad which happenign when there is a queued sub-level load request during traveling betwen maps
UE-5332

diff --git a/Engine/Source/Runtime/Engine/Classes/Engine/EngineTypes.h b/Engine/Source/Runtime/Engine/Classes/Engine/EngineTypes.h
--- a/Engine/Source/Runtime/Engine/Classes/Engine/EngineTypes.h
+++ b/Engine/Source/Runtime/Engine/Classes/Engine/EngineTypes.h
@@ -474,6 +474,13 @@ namespace EWorldType
 	};
 }
 
+enum class EFlushLevelStreamingType
+{
+	None,			
+	Full,			// Allow multiple load requests
+	Visibility,		// Flush visibility only, do not allow load requests, flushes async loading as well
+};
+
 USTRUCT()
 struct FResponseChannel
 {
diff --git a/Engine/Source/Runtime/Engine/Classes/Engine/World.h b/Engine/Source/Runtime/Engine/Classes/Engine/World.h
--- a/Engine/Source/Runtime/Engine/Classes/Engine/World.h
+++ b/Engine/Source/Runtime/Engine/Classes/Engine/World.h
@@ -843,9 +843,6 @@ public:
 	/** Whether world object has been initialized via Init()																	*/
 	bool										bIsWorldInitialized;
 	
-	/** Override, forcing level load requests to be allowed. < 0 == not allowed, 0 == have code choose, > 1 == force allow.		 */
-	int32										AllowLevelLoadOverride;
-
 	/** Number of frames to delay Streaming Volume updating, useful if you preload a bunch of levels but the camera hasn't caught up yet (INDEX_NONE for infinite) */
 	int32										StreamingVolumeUpdateDelay;
 
@@ -918,8 +915,8 @@ public:
 	UPROPERTY()
 	class UWorldComposition* WorldComposition;
 	
-	/** Whether we currently flushing level streaming state */ 
-	bool bFlushingLevelStreaming;
+	/** Whether we flushing level streaming state */ 
+	EFlushLevelStreamingType FlushLevelStreamingType;
 
 public:
 	/** The type of travel to perform next when doing a server travel */
@@ -1774,10 +1771,10 @@ public:
 	 * on all async operation like updating components.
 	 *
 	 * @param ViewFamily				Optional collection of views to take into account
-	 * @param bOnlyFlushVisibility		Whether to only flush level visibility operations (optional)
+	 * @param FlushType					Whether to only flush level visibility operations (optional)
 	 * @param ExcludeType				Exclude packages of this type from flushing
 	 */
-	void FlushLevelStreaming( FSceneViewFamily* ViewFamily = NULL, bool bOnlyFlushVisibility = false, FName ExcludeType = NAME_None);
+	void FlushLevelStreaming( FSceneViewFamily* ViewFamily = NULL, EFlushLevelStreamingType FlushType = EFlushLevelStreamingType::Full, FName ExcludeType = NAME_None );
 
 	/**
 	 * Triggers a call to ULevel::BuildStreamingData(this,NULL,NULL) within a few seconds.
diff --git a/Engine/Source/Runtime/Engine/Private/GameEngine.cpp b/Engine/Source/Runtime/Engine/Private/GameEngine.cpp
--- a/Engine/Source/Runtime/Engine/Private/GameEngine.cpp
+++ b/Engine/Source/Runtime/Engine/Private/GameEngine.cpp
@@ -501,7 +501,7 @@ void UGameEngine::PreExit()
 
 			World->GetGameInstance()->Shutdown();
 
-			World->FlushLevelStreaming( NULL, true );
+			World->FlushLevelStreaming( NULL, EFlushLevelStreamingType::Visibility );
 			World->CleanupWorld();
 		}
 	}
diff --git a/Engine/Source/Runtime/Engine/Private/UnrealEngine.cpp b/Engine/Source/Runtime/Engine/Private/UnrealEngine.cpp
--- a/Engine/Source/Runtime/Engine/Private/UnrealEngine.cpp
+++ b/Engine/Source/Runtime/Engine/Private/UnrealEngine.cpp
@@ -8603,8 +8603,8 @@ bool UEngine::LoadMap( FWorldContext& WorldContext, FURL URL, class UPendingNetG
 		// Clean up networking
 		ShutdownWorldNetDriver(WorldContext.World());
 
-		// Clean up game state.
-		WorldContext.World()->FlushLevelStreaming( NULL, true );
+		// Make sure there are no pending visibility requests.
+		WorldContext.World()->FlushLevelStreaming( NULL, EFlushLevelStreamingType::Visibility );
 		
 		// send a message that all levels are going away (NULL means every sublevel is being removed
 		// without a call to RemoveFromWorld for each)
@@ -9772,7 +9772,7 @@ bool UEngine::CommitMapChange( FWorldContext &Context )
 		// Update level streaming, forcing existing levels to be unloaded and their streaming objects 
 		// removed from the world info.	We can't kick off async loading in this update as we want to 
 		// collect garbage right below.
-		Context.World()->FlushLevelStreaming( NULL, true );
+		Context.World()->FlushLevelStreaming( NULL, EFlushLevelStreamingType::Visibility );
 		
 		// make sure any looping sounds, etc are stopped
 		if (GetAudioDevice() != NULL)
@@ -9819,12 +9819,12 @@ bool UEngine::CommitMapChange( FWorldContext &Context )
 
 			Context.PendingLevelStreamingStatusUpdates.Empty();
 
-			Context.World()->FlushLevelStreaming(NULL, false);
+			Context.World()->FlushLevelStreaming( NULL, EFlushLevelStreamingType::Full );
 		}
 		else
 		{
-			// This will cause the newly added persistent level to be made visible and kick off async loading for others.
-			Context.World()->FlushLevelStreaming( NULL, true );
+			// Make sure there are no pending visibility requests.
+			Context.World()->FlushLevelStreaming( NULL, EFlushLevelStreamingType::Visibility );
 		}
 
 		// delay the use of streaming volumes for a few frames
diff --git a/Engine/Source/Runtime/Engine/Private/World.cpp b/Engine/Source/Runtime/Engine/Private/World.cpp
--- a/Engine/Source/Runtime/Engine/Private/World.cpp
+++ b/Engine/Source/Runtime/Engine/Private/World.cpp
@@ -101,6 +101,7 @@ UWorld::UWorld( const FObjectInitializer& ObjectInitializer )
 ,	FXSystem(NULL)
 ,	TickTaskLevel(FTickTaskManagerInterface::Get().AllocateTickTaskLevel())
 ,   bIsBuilt(false)
+,	FlushLevelStreamingType(EFlushLevelStreamingType::None)
 ,	NextTravelType(TRAVEL_Relative)
 {
 	TimerManager = new FTimerManager();
@@ -119,6 +120,7 @@ UWorld::UWorld( const FObjectInitializer& ObjectInitializer,const FURL& InURL )
 ,	FXSystem(NULL)
 ,	TickTaskLevel(FTickTaskManagerInterface::Get().AllocateTickTaskLevel())
 ,   bIsBuilt(false)
+,	FlushLevelStreamingType(EFlushLevelStreamingType::None)
 ,	NextTravelType(TRAVEL_Relative)
 {
 	SetFlags( RF_Transactional );
@@ -994,7 +996,7 @@ void UWorld::DestroyWorld( bool bInformEngineOfWorld, UWorld* NewWorld )
 	// Clean up existing world and remove it from root set so it can be garbage collected.
 	bIsLevelStreamingFrozen = false;
 	bShouldForceUnloadStreamingLevels = true;
-	FlushLevelStreaming( NULL, true );
+	FlushLevelStreaming();
 	CleanupWorld(true, true, NewWorld);
 
 	check( NetworkActors.Num() == 0 );
@@ -2214,21 +2216,19 @@ void UWorld::UpdateLevelStreamingInner( UWorld* PersistentWorld, FSceneViewFamil
 				{
 					const FSceneView* View		= ViewFamily->Views[ViewIndex];
 					const FVector& ViewLocation	= View->ViewMatrices.ViewOrigin;
-					bShouldBeLoaded	= bShouldBeLoaded  || ( StreamingLevel && (!IsGameWorld() || StreamingLevel->ShouldBeLoaded(ViewLocation)) );
-					bShouldBeVisible= bShouldBeVisible || ( bShouldBeLoaded && StreamingLevel && StreamingLevel->ShouldBeVisible(ViewLocation) );
+					bShouldBeLoaded	= bShouldBeLoaded  || !IsGameWorld() || StreamingLevel->ShouldBeLoaded(ViewLocation);
+					bShouldBeVisible= bShouldBeVisible || ( bShouldBeLoaded && StreamingLevel->ShouldBeVisible(ViewLocation) );
 				}
 			}
 			// Or default to view location of 0,0,0
 			else
 			{
 				FVector ViewLocation(0,0,0);
-				bShouldBeLoaded		= bShouldBeLoaded  || ( StreamingLevel && (!IsGameWorld() || StreamingLevel->ShouldBeLoaded(ViewLocation)) );
-				bShouldBeVisible	= bShouldBeVisible || ( bShouldBeLoaded && StreamingLevel && StreamingLevel->ShouldBeVisible(ViewLocation) );
+				bShouldBeLoaded		= bShouldBeLoaded  || !IsGameWorld() || StreamingLevel->ShouldBeLoaded(ViewLocation);
+				bShouldBeVisible	= bShouldBeVisible || ( bShouldBeLoaded && StreamingLevel->ShouldBeVisible(ViewLocation) );
 			}
 		}
 
-		// Figure out whether there are any levels we haven't collected garbage yet.
-		bool bAreLevelsPendingPurge	=	FLevelStreamingGCHelper::GetNumLevelsPendingPurge() > 0;
 		// We want to give the garbage collector a chance to remove levels before we stream in more. We can't do this in the
 		// case of a blocking load as it means those requests should be fulfilled right away. By waiting on GC before kicking
 		// off new levels we potentially delay streaming in maps, but AllowLevelLoadRequests already looks and checks whether
@@ -2236,10 +2236,10 @@ void UWorld::UpdateLevelStreamingInner( UWorld* PersistentWorld, FSceneViewFamil
 		// on purpose as well so the GC code has a chance to execute between consecutive loads of maps.
 		//
 		// NOTE: AllowLevelLoadRequests not an invariant as streaming might affect the result, do NOT pulled out of the loop.
-		bool bAllowLevelLoadRequests =	(PersistentWorld->AllowLevelLoadRequests() && !bAreLevelsPendingPurge) 
-										|| bShouldBlockOnLoad 
-										|| bFlushingLevelStreaming;
+		bool bAllowLevelLoadRequests =	PersistentWorld->AllowLevelLoadRequests() || bShouldBlockOnLoad;
 
+		// Figure out whether there are any levels we haven't collected garbage yet.
+		bool bAreLevelsPendingPurge	=	FLevelStreamingGCHelper::GetNumLevelsPendingPurge() > 0;
 		// Request a 'soft' GC if there are levels pending purge and there are levels to be loaded. In the case of a blocking
 		// load this is going to guarantee GC firing first thing afterwards and otherwise it is going to sneak in right before
 		// kicking off the async load.
@@ -2383,29 +2383,20 @@ void UWorld::EvaluateWorldOriginLocation(const FSceneViewFamily& ViewFamily)
 	}
 }
 
-void UWorld::FlushLevelStreaming( FSceneViewFamily* ViewFamily, bool bOnlyFlushVisibility, FName ExcludeType)
+void UWorld::FlushLevelStreaming(FSceneViewFamily* ViewFamily, EFlushLevelStreamingType FlushType, FName ExcludeType)
 {
 	AWorldSettings* WorldSettings = GetWorldSettings();
 
-	TGuardValue<bool> FlushingLevelStreaming(bFlushingLevelStreaming, true);
-
-	// Allow queuing multiple load requests if we're performing a full flush and disallow if we're just
-	// flushing level visibility.
-	int32 OldAllowLevelLoadOverride = AllowLevelLoadOverride;
-	AllowLevelLoadOverride = bOnlyFlushVisibility ? 0 : 1;
+	TGuardValue<EFlushLevelStreamingType> FlushingLevelStreamingGuard(FlushLevelStreamingType, FlushType);
 
 	// Update internals with current loaded/ visibility flags.
-	UpdateLevelStreaming();
+	UpdateLevelStreaming(ViewFamily);
 
-	// Only flush async loading if we're performing a full flush.
-	if( !bOnlyFlushVisibility )
-	{
-		// Make sure all outstanding loads are taken care of, other than ones associated with the excluded type
-		FlushAsyncLoading(ExcludeType);
-	}
+	// Make sure all outstanding loads are taken care of, other than ones associated with the excluded type
+	FlushAsyncLoading(ExcludeType);
 
 	// Kick off making levels visible if loading finished by flushing.
-	UpdateLevelStreaming();
+	UpdateLevelStreaming(ViewFamily);
 
 	// Making levels visible is spread across several frames so we simply loop till it is done.
 	bool bLevelsPendingVisibility = true;
@@ -2417,7 +2408,7 @@ void UWorld::FlushLevelStreaming( FSceneViewFamily* ViewFamily, bool bOnlyFlushV
 		if( bLevelsPendingVisibility )
 		{
 			// Only flush async loading if we're performing a full flush.
-			if( !bOnlyFlushVisibility )
+			if (FlushLevelStreamingType == EFlushLevelStreamingType::Full)
 			{
 				// Make sure all outstanding loads are taken care of...
 				FlushAsyncLoading(NAME_None);
@@ -2435,12 +2426,10 @@ void UWorld::FlushLevelStreaming( FSceneViewFamily* ViewFamily, bool bOnlyFlushV
 	EnsureCollisionTreeIsBuilt();
 
 	// We already blocked on async loading.
-	if( !bOnlyFlushVisibility )
+	if (FlushLevelStreamingType == EFlushLevelStreamingType::Full)
 	{
 		bRequestedBlockOnAsyncLoading = false;
 	}
-
-	AllowLevelLoadOverride = OldAllowLevelLoadOverride;
 }
 
 /**
@@ -2509,39 +2498,24 @@ bool UWorld::AreAlwaysLoadedLevelsLoaded() const
 	return true;
 }
 
-
 bool UWorld::AllowLevelLoadRequests()
 {
-	bool bAllowLevelLoadRequests = false;
-	// Hold off requesting new loads if there is an active async load request.
-	if( !GIsEditor )
+	// Always allow level load request in the editor or when we do full streaming flush
+	if (IsGameWorld() && FlushLevelStreamingType != EFlushLevelStreamingType::Full)
 	{
-		// Let code choose.
-		if( AllowLevelLoadOverride == 0 )
-		{
-			// There are pending queued requests and gameplay has already started so hold off queueing.
-			if( IsAsyncLoading() && GetTimeSeconds() > 1.f )
-			{
-				bAllowLevelLoadRequests = false;
-			}
-			// No load requests or initial load so it's save to queue.
-			else
-			{
-				bAllowLevelLoadRequests = true;
-			}
-		}
-		// Use override, < 0 == don't allow, > 0 == allow
-		else
+		const bool bAreLevelsPendingPurge = FLevelStreamingGCHelper::GetNumLevelsPendingPurge() > 0;
+		
+		// Let code choose. Hold off queueing in case: 
+		// We are only flushing levels visibility
+		// There pending unload requests
+		// There pending load requests and gameplay has already started.
+		if (bAreLevelsPendingPurge || FlushLevelStreamingType == EFlushLevelStreamingType::Visibility || (IsAsyncLoading() && GetTimeSeconds() > 1.f))
 		{
-			bAllowLevelLoadRequests = AllowLevelLoadOverride > 0 ? true : false;
+			return false;
 		}
 	}
-	// Always allow load request in the Editor.
-	else
-	{
-		bAllowLevelLoadRequests = true;
-	}
-	return bAllowLevelLoadRequests;
+
+	return true;
 }
 
 bool UWorld::Exec( UWorld* InWorld, const TCHAR* Cmd, FOutputDevice& Ar )
@@ -4521,7 +4495,7 @@ UWorld* FSeamlessTravelHandler::Tick()
 		else
 		{
 			// Make sure there are no pending visibility requests.
-			CurrentWorld->FlushLevelStreaming( NULL, true );
+			CurrentWorld->FlushLevelStreaming( NULL, EFlushLevelStreamingType::Visibility );
 
 			if (CurrentWorld->GameState)
 			{
diff --git a/Engine/Source/Runtime/Engine/Private/WorldComposition.cpp b/Engine/Source/Runtime/Engine/Private/WorldComposition.cpp
--- a/Engine/Source/Runtime/Engine/Private/WorldComposition.cpp
+++ b/Engine/Source/Runtime/Engine/Private/WorldComposition.cpp
@@ -605,12 +605,12 @@ void UWorldComposition::CommitTileStreamingState(UWorld* PersistenWorld, int32 T
 	}
 
 	// Quit early in case we have cooldown on streaming state changes
-	if (TilesStreamingTimeThreshold > 0.0)
+	const bool bUseStreamingStateCooldown = (PersistenWorld->IsGameWorld() && PersistenWorld->FlushLevelStreamingType == EFlushLevelStreamingType::None);
+	if (bUseStreamingStateCooldown && TilesStreamingTimeThreshold > 0.0)
 	{
 		const double CurrentTime = FPlatformTime::Seconds();
 		const double TimePassed = CurrentTime - Tile.StreamingLevelStateChangeTime;
-
-		if (!PersistenWorld->bFlushingLevelStreaming &&	TimePassed < TilesStreamingTimeThreshold)
+		if (TimePassed < TilesStreamingTimeThreshold)
 		{
 			return;
 		}
