Removed UWorld::UpdateWorldStreamingState declaration
UE-5314

diff --git a/Engine/Source/Runtime/Engine/Classes/Engine/World.h b/Engine/Source/Runtime/Engine/Classes/Engine/World.h
--- a/Engine/Source/Runtime/Engine/Classes/Engine/World.h
+++ b/Engine/Source/Runtime/Engine/Classes/Engine/World.h
@@ -1757,14 +1757,6 @@ private:
 	void EvaluateWorldOriginLocation( const FSceneViewFamily& ViewFamily );
 
 public:
-
-	/**
-	 * Keeps streaming levels list updated according to camera view
-	 *
-	 * @param ViewFamily Collection of views to take into account
-	 */
-	void UpdateWorldStreamingState(const FSceneViewFamily& ViewFamily);
-	
 	/**
 	 * Flushes level streaming in blocking fashion and returns when all levels are loaded/ visible/ hidden
 	 * so further calls to UpdateLevelStreaming won't do any work unless state changes. Basically blocks
