Fix for UE-3820

Keeping emitters with any enabled LODs as not complete so that they still appear when the player moves back into LOD range.

diff --git a/Engine/Source/Runtime/Engine/Classes/Particles/ParticleEmitter.h b/Engine/Source/Runtime/Engine/Classes/Particles/ParticleEmitter.h
--- a/Engine/Source/Runtime/Engine/Classes/Particles/ParticleEmitter.h
+++ b/Engine/Source/Runtime/Engine/Classes/Particles/ParticleEmitter.h
@@ -321,6 +321,9 @@ class UParticleEmitter : public UObject
 		}
 		return Q;
 	}
+
+	/** Returns true if the is emitter has any enabled LODs, false otherwise. */
+	bool HasAnyEnabledLODs()const;
 };
 
 
diff --git a/Engine/Source/Runtime/Engine/Private/Particles/ParticleComponents.cpp b/Engine/Source/Runtime/Engine/Private/Particles/ParticleComponents.cpp
--- a/Engine/Source/Runtime/Engine/Private/Particles/ParticleComponents.cpp
+++ b/Engine/Source/Runtime/Engine/Private/Particles/ParticleComponents.cpp
@@ -1783,6 +1783,19 @@ void UParticleEmitter::Build()
 	}
 }
 
+bool UParticleEmitter::HasAnyEnabledLODs()const
+{
+	for (UParticleLODLevel* LodLevel : LODLevels)
+	{
+		if (LodLevel && LodLevel->bEnabled)
+		{
+			return true;
+		}
+	}
+	
+	return false;
+}
+
 /*-----------------------------------------------------------------------------
 	UParticleSpriteEmitter implementation.
 -----------------------------------------------------------------------------*/
@@ -4166,7 +4179,7 @@ void UParticleSystemComponent::InitParticles()
 					// Must have a slot for each emitter instance - even if it's NULL.
 					// This is so the indexing works correctly.
 					UParticleEmitter* Emitter = Template->Emitters[Idx];
-					if (Emitter && Emitter->DetailMode <= GlobalDetailMode)
+					if (Emitter && Emitter->DetailMode <= GlobalDetailMode && Emitter->HasAnyEnabledLODs())
 					{
 						EmitterInstances.Add(Emitter->CreateInstance(this));
 					}
@@ -4895,6 +4908,12 @@ bool UParticleSystemComponent::HasCompleted()
 					}
 				}
 			}
+			else
+			{
+				//Don't assume emitter is complete if the current lod is disabled.
+				//Any emitters with all their lods disabled aren't created in the first place.
+				bHasCompleted = false;
+			}
 		}
 	}
 
