Double buffered bone space bases

diff --git a/Engine/Source/Editor/AnimGraph/Private/AnimGraphNode_BoneDrivenController.cpp b/Engine/Source/Editor/AnimGraph/Private/AnimGraphNode_BoneDrivenController.cpp
--- a/Engine/Source/Editor/AnimGraph/Private/AnimGraphNode_BoneDrivenController.cpp
+++ b/Engine/Source/Editor/AnimGraph/Private/AnimGraphNode_BoneDrivenController.cpp
@@ -61,8 +61,8 @@ void UAnimGraphNode_BoneDrivenController::Draw(FPrimitiveDrawInterface* PDI, USk
 
 	if(SourceIdx != INDEX_NONE && TargetIdx != INDEX_NONE)
 	{
-		FTransform SourceTM = SkelMeshComp->SpaceBases[SourceIdx] * SkelMeshComp->ComponentToWorld;
-		FTransform TargetTM = SkelMeshComp->SpaceBases[TargetIdx] * SkelMeshComp->ComponentToWorld;
+		FTransform SourceTM = SkelMeshComp->GetSpaceBases()[SourceIdx] * SkelMeshComp->ComponentToWorld;
+		FTransform TargetTM = SkelMeshComp->GetSpaceBases()[TargetIdx] * SkelMeshComp->ComponentToWorld;
 
 		PDI->DrawLine(TargetTM.GetLocation(), SourceTM.GetLocation(), FLinearColor(0.0f, 0.0f, 1.0f), SDPG_Foreground, 0.5f);
 
diff --git a/Engine/Source/Editor/Persona/Private/AnimationEditorViewportClient.cpp b/Engine/Source/Editor/Persona/Private/AnimationEditorViewportClient.cpp
--- a/Engine/Source/Editor/Persona/Private/AnimationEditorViewportClient.cpp
+++ b/Engine/Source/Editor/Persona/Private/AnimationEditorViewportClient.cpp
@@ -670,7 +670,7 @@ void FAnimationViewportClient::ShowBoneNames( FCanvas* Canvas, FSceneView* View
 		const FColor BoneColor = FColor::White;
 		if (BoneColor.A != 0)
 		{
-			const FVector BonePos = PreviewSkelMeshComp->ComponentToWorld.TransformPosition(PreviewSkelMeshComp->SpaceBases[BoneIndex].GetLocation());
+			const FVector BonePos = PreviewSkelMeshComp->ComponentToWorld.TransformPosition(PreviewSkelMeshComp->GetSpaceBases()[BoneIndex].GetLocation());
 
 			const FPlane proj = View->Project(BonePos);
 			if (proj.W > 0.f)
@@ -881,7 +881,7 @@ void FAnimationViewportClient::DisplayInfo(FCanvas* Canvas, FSceneView* View, bo
 				const FName BoneName = PreviewSkelMeshComp->SkeletalMesh->RefSkeleton.GetBoneName(BoneIndex);
 				FTransform ReferenceTransform = PreviewSkelMeshComp->SkeletalMesh->RefSkeleton.GetRefBonePose()[BoneIndex];
 				FTransform LocalTransform = PreviewSkelMeshComp->LocalAtoms[BoneIndex];
-				FTransform ComponentTransform = PreviewSkelMeshComp->SpaceBases[BoneIndex];
+				FTransform ComponentTransform = PreviewSkelMeshComp->GetSpaceBases()[BoneIndex];
 
 				CurYOffset += YL + 2;
 				InfoString = FString::Printf(TEXT("Local :%s"), *LocalTransform.ToString());
@@ -1483,10 +1483,10 @@ void FAnimationViewportClient::DrawMeshBones(USkeletalMeshComponent * MeshCompon
 	if ( MeshComponent && MeshComponent->SkeletalMesh )
 	{
 		TArray<FTransform> WorldTransforms;
-		WorldTransforms.AddUninitialized(MeshComponent->SpaceBases.Num());
+		WorldTransforms.AddUninitialized(MeshComponent->GetNumSpaceBases());
 
 		TArray<FLinearColor> BoneColours;
-		BoneColours.AddUninitialized(MeshComponent->SpaceBases.Num());
+		BoneColours.AddUninitialized(MeshComponent->GetNumSpaceBases());
 
 		TArray<int32> SelectedBones;
 		if(UDebugSkelMeshComponent* DebugMeshComponent = Cast<UDebugSkelMeshComponent>(MeshComponent))
@@ -1500,7 +1500,7 @@ void FAnimationViewportClient::DrawMeshBones(USkeletalMeshComponent * MeshCompon
 			const int32 BoneIndex = MeshComponent->RequiredBones[Index];
 			const int32 ParentIndex = MeshComponent->SkeletalMesh->RefSkeleton.GetParentIndex(BoneIndex);
 
-			WorldTransforms[BoneIndex] = MeshComponent->SpaceBases[BoneIndex]*MeshComponent->ComponentToWorld;
+			WorldTransforms[BoneIndex] = MeshComponent->GetSpaceBases()[BoneIndex]*MeshComponent->ComponentToWorld;
 			
 			if(SelectedBones.Contains(BoneIndex))
 			{
@@ -1598,10 +1598,10 @@ void FAnimationViewportClient::DrawMeshSubsetBones(const USkeletalMeshComponent*
 	if ( MeshComponent && MeshComponent->SkeletalMesh && BonesOfInterest.Num() > 0 )
 	{
 		TArray<FTransform> WorldTransforms;
-		WorldTransforms.AddUninitialized(MeshComponent->SpaceBases.Num());
+		WorldTransforms.AddUninitialized(MeshComponent->GetNumSpaceBases());
 
 		TArray<FLinearColor> BoneColours;
-		BoneColours.AddUninitialized(MeshComponent->SpaceBases.Num());
+		BoneColours.AddUninitialized(MeshComponent->GetNumSpaceBases());
 
 		TArray<FBoneIndexType> RequiredBones;
 
@@ -1627,7 +1627,7 @@ void FAnimationViewportClient::DrawMeshSubsetBones(const USkeletalMeshComponent*
 					//found a bone we are interested in
 					if(ParentIndex >= 0)
 					{
-						WorldTransforms[ParentIndex] = MeshComponent->SpaceBases[ParentIndex]*MeshComponent->ComponentToWorld;
+						WorldTransforms[ParentIndex] = MeshComponent->GetSpaceBases()[ParentIndex]*MeshComponent->ComponentToWorld;
 					}
 					BoneColours[BoneIndex] = LinearSelectionColor;
 					bDrawBone = true;
@@ -1645,7 +1645,7 @@ void FAnimationViewportClient::DrawMeshSubsetBones(const USkeletalMeshComponent*
 			{
 				//add to the list
 				RequiredBones.AddUnique(BoneIndex);
-				WorldTransforms[BoneIndex] = MeshComponent->SpaceBases[BoneIndex]*MeshComponent->ComponentToWorld;
+				WorldTransforms[BoneIndex] = MeshComponent->GetSpaceBases()[BoneIndex]*MeshComponent->ComponentToWorld;
 			}
 		}
 
@@ -1673,7 +1673,7 @@ void FAnimationViewportClient::DrawSockets( TArray<USkeletalMeshSocket*>& Socket
 			FVector Start, End;
 			if (ParentIndex >=0)
 			{
-				FTransform WorldTransformParent = DebugMeshComponent->SpaceBases[ParentIndex]*DebugMeshComponent->ComponentToWorld;
+				FTransform WorldTransformParent = DebugMeshComponent->GetSpaceBases()[ParentIndex]*DebugMeshComponent->ComponentToWorld;
 				Start = WorldTransformParent.GetLocation();
 				End = WorldTransformSocket.GetLocation();
 			}
diff --git a/Engine/Source/Editor/Persona/Private/SAnimationScrubPanel.cpp b/Engine/Source/Editor/Persona/Private/SAnimationScrubPanel.cpp
--- a/Engine/Source/Editor/Persona/Private/SAnimationScrubPanel.cpp
+++ b/Engine/Source/Editor/Persona/Private/SAnimationScrubPanel.cpp
@@ -503,7 +503,7 @@ void SAnimationScrubPanel::OnReZeroAnimSequence( )
 				AnimSequence->Modify();
 
 				// Find vector that would translate current root bone location onto origin.
-				FVector ApplyTranslation = -1.f * PreviewSkelComp->SpaceBases[0].GetLocation();
+				FVector ApplyTranslation = -1.f * PreviewSkelComp->GetSpaceBases()[0].GetLocation();
 
 				// Convert into world space and eliminate 'z' translation. Don't want to move character into ground.
 				FVector WorldApplyTranslation = PreviewSkelComp->ComponentToWorld.TransformVector(ApplyTranslation);
diff --git a/Engine/Source/Editor/Persona/Private/SRetargetManager.cpp b/Engine/Source/Editor/Persona/Private/SRetargetManager.cpp
--- a/Engine/Source/Editor/Persona/Private/SRetargetManager.cpp
+++ b/Engine/Source/Editor/Persona/Private/SRetargetManager.cpp
@@ -245,10 +245,10 @@ FReply SRetargetManager::OnSaveRetargetBasePose()
 			const FScopedTransaction Transaction( LOCTEXT("SaveRetargetBasePose_Action", "Save Retarget Base Pose") );
 			PreviewMesh->Modify();
 			// get space bases and calculate local
-			const TArray<FTransform>& SpaceBases = PreviewMeshComp->SpaceBases;
+			const TArray<FTransform> & SpaceBases = PreviewMeshComp->GetSpaceBases();
 			// @todo check to see if skeleton vs preview mesh makes it different for missing bones
 			const FReferenceSkeleton& RefSkeleton = PreviewMesh->RefSkeleton;
-			TArray<FTransform>& NewRetargetBasePose = PreviewMesh->RetargetBasePose;
+			TArray<FTransform> & NewRetargetBasePose = PreviewMesh->RetargetBasePose;
 			// if you're using master pose component in preview, this won't work
 			check (PreviewMesh->RefSkeleton.GetNum() == SpaceBases.Num());
 			int32 TotalNumBones = PreviewMesh->RefSkeleton.GetNum();
diff --git a/Engine/Source/Editor/Persona/Private/Shared/AnimationRecorder.cpp b/Engine/Source/Editor/Persona/Private/Shared/AnimationRecorder.cpp
--- a/Engine/Source/Editor/Persona/Private/Shared/AnimationRecorder.cpp
+++ b/Engine/Source/Editor/Persona/Private/Shared/AnimationRecorder.cpp
@@ -91,7 +91,7 @@ void FAnimationRecorder::StartRecord(USkeletalMeshComponent * Component, UAnimSe
 	AnimationObject->TrackToSkeletonMapTable.Empty();
 	AnimationObject->AnimationTrackNames.Empty();
 
-	PreviousSpacesBases = Component->SpaceBases;
+	PreviousSpacesBases = Component->GetSpaceBases();
 
 	LastFrame = 0;
 	AnimationObject->SequenceLength = 0.f;
@@ -194,7 +194,7 @@ void FAnimationRecorder::UpdateRecord(USkeletalMeshComponent * Component, float
 		return;
 	}
 
-	TArray<FTransform>& SpaceBases = Component->SpaceBases;
+	const TArray<FTransform> & SpaceBases = Component->GetSpaceBases();
 
 	check (SpaceBases.Num() == PreviousSpacesBases.Num());
 
@@ -225,7 +225,7 @@ void FAnimationRecorder::UpdateRecord(USkeletalMeshComponent * Component, float
 	}
 
 	//save to current transform
-	PreviousSpacesBases = Component->SpaceBases;
+	PreviousSpacesBases = Component->GetSpaceBases();
 
 	// if we passed MaxFrame, just stop it
 	if (FramesRecorded > MaxFrame)
diff --git a/Engine/Source/Editor/PhAT/Private/PhATPreviewViewportClient.cpp b/Engine/Source/Editor/PhAT/Private/PhATPreviewViewportClient.cpp
--- a/Engine/Source/Editor/PhAT/Private/PhATPreviewViewportClient.cpp
+++ b/Engine/Source/Editor/PhAT/Private/PhATPreviewViewportClient.cpp
@@ -164,9 +164,9 @@ void FPhATEdPreviewViewportClient::DrawCanvas( FViewport& InViewport, FSceneView
 	if ((SharedData->bShowHierarchy && SharedData->EditorSimOptions->bShowNamesInHierarchy))
 	{
 		// Iterate over each graphics bone.
-		for (int32 i = 0; i <SharedData->EditorSkelComp->SpaceBases.Num(); ++i)
+		for (int32 i = 0; i <SharedData->EditorSkelComp->GetNumSpaceBases(); ++i)
 		{
-			FVector BonePos = SharedData->EditorSkelComp->ComponentToWorld.TransformPosition(SharedData->EditorSkelComp->SpaceBases[i].GetLocation());
+			FVector BonePos = SharedData->EditorSkelComp->ComponentToWorld.TransformPosition(SharedData->EditorSkelComp->GetSpaceBases()[i].GetLocation());
 
 			FPlane proj = View.Project(BonePos);
 			if (proj.W > 0.f) // This avoids drawing bone names that are behind us.
@@ -666,7 +666,7 @@ void FPhATEdPreviewViewportClient::Tick(float DeltaSeconds)
 		Setting->bWorldGravitySet = true;
 
 		// We back up the transforms array now
-		SharedData->EditorSkelComp->AnimationSpaceBases = SharedData->EditorSkelComp->SpaceBases;
+		SharedData->EditorSkelComp->AnimationSpaceBases = SharedData->EditorSkelComp->GetSpaceBases();
 		SharedData->EditorSkelComp->SetPhysicsBlendWeight(SharedData->EditorSimOptions->PhysicsBlend);
 	}
 }
diff --git a/Engine/Source/Editor/PhAT/Private/PhATRender.cpp b/Engine/Source/Editor/PhAT/Private/PhATRender.cpp
--- a/Engine/Source/Editor/PhAT/Private/PhATRender.cpp
+++ b/Engine/Source/Editor/PhAT/Private/PhATRender.cpp
@@ -279,7 +279,7 @@ FPrimitiveSceneProxy* UPhATEdSkeletalMeshComponent::CreateSceneProxy()
 
 void UPhATEdSkeletalMeshComponent::DrawHierarchy(FPrimitiveDrawInterface* PDI, bool bAnimSkel)
 {
-	for (int32 i=1; i <SpaceBases.Num(); ++i)
+	for (int32 i=1; i<GetNumSpaceBases(); ++i)
 	{
 		int32 ParentIndex = SkeletalMesh->RefSkeleton.GetParentIndex(i);
 
@@ -291,8 +291,8 @@ void UPhATEdSkeletalMeshComponent::DrawHierarchy(FPrimitiveDrawInterface* PDI, b
 		}
 		else
 		{
-			ParentPos = ComponentToWorld.TransformPosition(SpaceBases[ParentIndex].GetLocation());
-			ChildPos = ComponentToWorld.TransformPosition(SpaceBases[i].GetLocation());
+			ParentPos = ComponentToWorld.TransformPosition(GetSpaceBases()[ParentIndex].GetLocation());
+			ChildPos = ComponentToWorld.TransformPosition(GetSpaceBases()[i].GetLocation());
 		}
 
 		FColor DrawColor = bAnimSkel ? AnimSkelDrawColor : HierarchyDrawColor;
diff --git a/Engine/Source/Editor/UnrealEd/Classes/Animation/DebugSkelMeshComponent.h b/Engine/Source/Editor/UnrealEd/Classes/Animation/DebugSkelMeshComponent.h
--- a/Engine/Source/Editor/UnrealEd/Classes/Animation/DebugSkelMeshComponent.h
+++ b/Engine/Source/Editor/UnrealEd/Classes/Animation/DebugSkelMeshComponent.h
@@ -248,6 +248,11 @@ class UDebugSkelMeshComponent : public USkeletalMeshComponent
 	virtual void CheckClothTeleport(float DeltaTime) override;
 
 #endif //#if WITH_APEX_CLOTHING
+
+private:
+
+	// Helper function to generate space bases for current frame
+	void GenSpaceBases(TArray<FTransform>& OutSpaceBases);
 };
 
 
diff --git a/Engine/Source/Editor/UnrealEd/Private/Animation/DebugSkelMeshComponent.cpp b/Engine/Source/Editor/UnrealEd/Private/Animation/DebugSkelMeshComponent.cpp
--- a/Engine/Source/Editor/UnrealEd/Private/Animation/DebugSkelMeshComponent.cpp
+++ b/Engine/Source/Editor/UnrealEd/Private/Animation/DebugSkelMeshComponent.cpp
@@ -130,7 +130,8 @@ FBoxSphereBounds UDebugSkelMeshComponent::CalcBounds(const FTransform& LocalToWo
 	{
 		// extend bounds by bones but without root bone
 		FBox BoundingBox(0);
-		for (int32 BoneIndex = 1; BoneIndex < SpaceBases.Num(); ++ BoneIndex)
+		const int32 NumBones = GetNumSpaceBases();
+		for (int32 BoneIndex = 1; BoneIndex < NumBones; ++BoneIndex)
 		{
 			BoundingBox += GetBoneMatrix(BoneIndex).GetOrigin();
 		}
@@ -425,42 +426,44 @@ void UDebugSkelMeshComponent::SetShowBoneWeight(bool bNewShowBoneWeight)
 	bDrawBoneInfluences = bNewShowBoneWeight;
 }
 
+void UDebugSkelMeshComponent::GenSpaceBases(TArray<FTransform>& OutSpaceBases)
+{
+	TArray<FTransform> TempLocalAtoms;
+	TArray<FActiveVertexAnim> TempVertexAnims;
+	FVector TempRootBoneTranslation;
+	PerformAnimationEvaluation(SkeletalMesh, AnimScriptInstance, OutSpaceBases, CachedLocalAtoms, TempVertexAnims, TempRootBoneTranslation);
+}
+
 void UDebugSkelMeshComponent::RefreshBoneTransforms(FActorComponentTickFunction* TickFunction)
 {
+	bool bGenerateRAWAnimation = bDisplayRawAnimation && AnimScriptInstance && AnimScriptInstance->RequiredBones.IsValid();
+
+	if (bGenerateRAWAnimation)
+	{
+		AnimScriptInstance->RequiredBones.SetUseRAWData(true);
+	}
+
 	// Run regular update first so we get RequiredBones up to date.
 	Super::RefreshBoneTransforms(NULL); // Pass NULL so we force non threaded work
 
+	if (bGenerateRAWAnimation)
+	{
+		AnimScriptInstance->RequiredBones.SetUseRAWData(false);
+	}
+
 	// Non retargeted pose.
 	NonRetargetedSpaceBases.Empty();
 	if( bDisplayNonRetargetedPose && AnimScriptInstance && AnimScriptInstance->RequiredBones.IsValid() )
 	{
-		TArray<FTransform> BackupSpaceBases = SpaceBases;
-
 		AnimScriptInstance->RequiredBones.SetDisableRetargeting(true);
-		Super::RefreshBoneTransforms(NULL);
+		GenSpaceBases(NonRetargetedSpaceBases);
 		AnimScriptInstance->RequiredBones.SetDisableRetargeting(false);
-
-		NonRetargetedSpaceBases = SpaceBases;
-		SpaceBases = BackupSpaceBases;
 	}
 
-	if( bDisplayRawAnimation )
+	if (bDisplayRawAnimation)
 	{
-		// save the transform in CompressedSpaceBases
-		CompressedSpaceBases = SpaceBases;
-
-		// use raw data now
-		if( AnimScriptInstance && AnimScriptInstance->RequiredBones.IsValid() )
-		{
-			AnimScriptInstance->RequiredBones.SetUseRAWData(true);
-			Super::RefreshBoneTransforms(NULL);
-			AnimScriptInstance->RequiredBones.SetUseRAWData(false);
-		}
-		// Otherwise we'll just get ref pose.
-		else
-		{
-			Super::RefreshBoneTransforms(NULL);
-		}
+		// Generate the normal compressed space bases
+		GenSpaceBases(CompressedSpaceBases);
 	}
 	else
 	{
diff --git a/Engine/Source/Runtime/Engine/Classes/Components/SkeletalMeshComponent.h b/Engine/Source/Runtime/Engine/Classes/Components/SkeletalMeshComponent.h
--- a/Engine/Source/Runtime/Engine/Classes/Components/SkeletalMeshComponent.h
+++ b/Engine/Source/Runtime/Engine/Classes/Components/SkeletalMeshComponent.h
@@ -1014,7 +1014,7 @@ public:
 	 *	Iterate over each physics body in the physics for this mesh, and for each 'kinematic' (ie fixed or default if owner isn't simulating) one, update its
 	 *	transform based on the animated transform.
 	 */
-	void UpdateKinematicBonesToPhysics(bool bTeleport, bool bNeedsSkinning, bool bForceUpdate = false);
+	void UpdateKinematicBonesToPhysics(const TArray<FTransform>& InSpaceBases, bool bTeleport, bool bNeedsSkinning, bool bForceUpdate = false);
 	
 	/**
 	 * Look up all bodies for broken constraints.
@@ -1163,9 +1163,9 @@ public:
 	void ParallelAnimationEvaluation() { PerformAnimationEvaluation(AnimEvaluationContext.SkeletalMesh, AnimEvaluationContext.AnimInstance, AnimEvaluationContext.SpaceBases, AnimEvaluationContext.LocalAtoms, AnimEvaluationContext.VertexAnims, AnimEvaluationContext.RootBoneTranslation); }
 	void CompleteParallelAnimationEvaluation()
 	{
-		if ((AnimEvaluationContext.AnimInstance == AnimScriptInstance) && (AnimEvaluationContext.SkeletalMesh == SkeletalMesh) && (AnimEvaluationContext.SpaceBases.Num() == SpaceBases.Num()))
+		if ((AnimEvaluationContext.AnimInstance == AnimScriptInstance) && (AnimEvaluationContext.SkeletalMesh == SkeletalMesh) && (AnimEvaluationContext.SpaceBases.Num() == GetNumSpaceBases()))
 		{
-			Exchange(AnimEvaluationContext.SpaceBases, AnimEvaluationContext.bDoInterpolation ? CachedSpaceBases : SpaceBases);
+			Exchange(AnimEvaluationContext.SpaceBases, AnimEvaluationContext.bDoInterpolation ? CachedSpaceBases : GetEditableSpaceBases() );
 			Exchange(AnimEvaluationContext.LocalAtoms, AnimEvaluationContext.bDoInterpolation ? CachedLocalAtoms : LocalAtoms);
 			Exchange(AnimEvaluationContext.VertexAnims, ActiveVertexAnims);
 			Exchange(AnimEvaluationContext.RootBoneTranslation, RootBoneTranslation);
diff --git a/Engine/Source/Runtime/Engine/Classes/Components/SkinnedMeshComponent.h b/Engine/Source/Runtime/Engine/Classes/Components/SkinnedMeshComponent.h
--- a/Engine/Source/Runtime/Engine/Classes/Components/SkinnedMeshComponent.h
+++ b/Engine/Source/Runtime/Engine/Classes/Components/SkinnedMeshComponent.h
@@ -149,17 +149,14 @@ class ENGINE_API USkinnedMeshComponent : public UMeshComponent
 {
 	GENERATED_UCLASS_BODY()
 
-	/** The skeletal mesh used by this component. */
-	UPROPERTY(EditAnywhere, BlueprintReadOnly, Category="Mesh")
+		/** The skeletal mesh used by this component. */
+		UPROPERTY(EditAnywhere, BlueprintReadOnly, Category = "Mesh")
 	class USkeletalMesh* SkeletalMesh;
 
-	/** Temporary array of of component-space bone matrices, update each frame and used for rendering the mesh. */
-	TArray<FTransform> SpaceBases;
-
 	//
 	// MasterPoseComponent.
 	//
-	
+
 	/**
 	 *	If set, this SkeletalMeshComponent will not use its SpaceBase for bone transform, but will
 	 *	use the SpaceBases array in the MasterPoseComponent. This is used when constructing a character using multiple skeletal meshes sharing the same
@@ -167,9 +164,22 @@ class ENGINE_API USkinnedMeshComponent : public UMeshComponent
 	 */
 	TWeakObjectPtr< class USkinnedMeshComponent > MasterPoseComponent;
 
+private:
+	/** Temporary array of of component-space bone matrices, update each frame and used for rendering the mesh. */
+	TArray<FTransform> SpaceBasesArray[2];
+
+	/** The index for the space bases buffer we can currently write to */
+	int32 CurrentEditableSpaceBases;
+
+	/** The index for the space bases buffer we can currently read from */
+	int32 CurrentReadSpaceBases;
+
+	/** Are we using double buffered blend spaces */
+	bool bDoubleBufferedBlendSpaces;
+
 protected:
-	/** 
-	 * If set, this component has slave pose components that are associated with this 
+	/**
+	 * If set, this component has slave pose components that are associated with this
 	 * Note this is weak object ptr, so it will go away unless you have other strong reference
 	 */
 	TArray< TWeakObjectPtr<USkinnedMeshComponent> > SlavePoseComponents;
@@ -180,121 +190,121 @@ public:
 	 */
 	TArray<int32> MasterBoneMap;
 
-	/** 
+	/**
 	 * When true, we will just using the bounds from our MasterPoseComponent.  This is useful for when we have a Mesh Parented
 	 * to the main SkelMesh (e.g. outline mesh or a full body overdraw effect that is toggled) that is always going to be the same
 	 * bounds as parent.  We want to do no calculations in that case.
 	 */
 	UPROPERTY()
-	uint32 bUseBoundsFromMasterPoseComponent:1;
+		uint32 bUseBoundsFromMasterPoseComponent : 1;
 
 	/** Array indicating all active vertex animations. This array is updated inside RefreshBoneTransforms based on the Anim Blueprint. */
 	UPROPERTY(transient)
-	TArray<struct FActiveVertexAnim> ActiveVertexAnims;
+		TArray<struct FActiveVertexAnim> ActiveVertexAnims;
 
 #if WITH_EDITORONLY_DATA
 	/** Index of the chunk to preview... If set to -1, all chunks will be rendered */
 	UPROPERTY(transient)
-	int32 ChunkIndexPreview;
+		int32 ChunkIndexPreview;
 
 	/** Index of the section to preview... If set to -1, all section will be rendered */
 	UPROPERTY(transient)
-	int32 SectionIndexPreview;
+		int32 SectionIndexPreview;
 
 #endif // WITH_EDITORONLY_DATA
 	//
 	// Physics.
 	//
-	
+
 	/**
 	 *	PhysicsAsset is set in SkeletalMesh by default, but you can override with this value
 	 */
-	UPROPERTY(EditAnywhere, AdvancedDisplay, BlueprintReadOnly, Category=Physics)
+	UPROPERTY(EditAnywhere, AdvancedDisplay, BlueprintReadOnly, Category = Physics)
 	class UPhysicsAsset* PhysicsAssetOverride;
 
 	//
 	// Level of detail.
 	//
-	
+
 	/** If 0, auto-select LOD level. if >0, force to (ForcedLodModel-1). */
-	UPROPERTY(EditAnywhere, AdvancedDisplay, BlueprintReadOnly, Category=LOD)
-	int32 ForcedLodModel;
+	UPROPERTY(EditAnywhere, AdvancedDisplay, BlueprintReadOnly, Category = LOD)
+		int32 ForcedLodModel;
 
 	/**
 	 * This is the min LOD that this component will use.  (e.g. if set to 2 then only 2+ LOD Models will be used.) This is useful to set on
 	 * meshes which are known to be a certain distance away and still want to have better LODs when zoomed in on them.
 	 **/
-	UPROPERTY(EditAnywhere, AdvancedDisplay, BlueprintReadOnly, Category=LOD)
-	int32 MinLodModel;
+	UPROPERTY(EditAnywhere, AdvancedDisplay, BlueprintReadOnly, Category = LOD)
+		int32 MinLodModel;
 
-	/** 
-	 *	Best LOD that was 'predicted' by UpdateSkelPose. 
-	 *	This is what bones were updated based on, so we do not allow rendering at a better LOD than this. 
+	/**
+	 *	Best LOD that was 'predicted' by UpdateSkelPose.
+	 *	This is what bones were updated based on, so we do not allow rendering at a better LOD than this.
 	 */
 	UPROPERTY()
-	int32 PredictedLODLevel;
+		int32 PredictedLODLevel;
 
 	/** LOD level from previous frame, so we can detect changes in LOD to recalc required bones. */
 	UPROPERTY()
-	int32 OldPredictedLODLevel;
+		int32 OldPredictedLODLevel;
 
 	/**	High (best) DistanceFactor that was desired for rendering this USkeletalMesh last frame. Represents how big this mesh was in screen space   */
 	UPROPERTY()
-	float MaxDistanceFactor;
+		float MaxDistanceFactor;
 
 	/** LOD array info. Each index will correspond to the LOD index **/
 	UPROPERTY(transient)
-	TArray<struct FSkelMeshComponentLODInfo> LODInfo;
+		TArray<struct FSkelMeshComponentLODInfo> LODInfo;
 
 	//
 	// Rendering options.
 	//
-	
+
 	/**
 	 * Allows adjusting the desired streaming distance of streaming textures that uses UV 0.
 	 * 1.0 is the default, whereas a higher value makes the textures stream in sooner from far away.
 	 * A lower value (0.0-1.0) makes the textures stream in later (you have to be closer).
 	 */
-	UPROPERTY(EditAnywhere, AdvancedDisplay, BlueprintReadWrite, Category=SkeletalMesh)
-	float StreamingDistanceMultiplier;
+	UPROPERTY(EditAnywhere, AdvancedDisplay, BlueprintReadWrite, Category = SkeletalMesh)
+		float StreamingDistanceMultiplier;
 
 	/**
 	 * Wireframe color
 	 */
-	UPROPERTY(EditAnywhere, AdvancedDisplay, BlueprintReadOnly, Category=SkeletalMesh)
-	FColor WireframeColor;
+	UPROPERTY(EditAnywhere, AdvancedDisplay, BlueprintReadOnly, Category = SkeletalMesh)
+		FColor WireframeColor;
 
 	/** Forces the mesh to draw in wireframe mode. */
 	UPROPERTY()
-	uint32 bForceWireframe:1;
+		uint32 bForceWireframe : 1;
 
 	/** Draw the skeleton hierarchy for this skel mesh. */
 	UPROPERTY()
-	uint32 bDisplayBones:1;
+		uint32 bDisplayBones : 1;
 
 	/** Don't bother rendering the skin. */
 	UPROPERTY()
-	uint32 bHideSkin:1;
+		uint32 bHideSkin : 1;
 	/** Array of bone visibilities (containing one of the values in EBoneVisibilityStatus for each bone).  A bone is only visible if it is *exactly* 1 (BVS_Visible) */
 	TArray<uint8> BoneVisibilityStates;
 
 	/**
 	 *	If true, use per-bone motion blur on this skeletal mesh (requires additional rendering, can be disabled to save performance).
 	 */
-	UPROPERTY(EditAnywhere, AdvancedDisplay, BlueprintReadOnly, Category=SkeletalMesh)
-	uint32 bPerBoneMotionBlur:1;
+	UPROPERTY(EditAnywhere, AdvancedDisplay, BlueprintReadOnly, Category = SkeletalMesh)
+		uint32 bPerBoneMotionBlur : 1;
 
 	//
 	// Misc.
 	//
-	
+
 	/** When true, skip using the physics asset etc. and always use the fixed bounds defined in the SkeletalMesh. */
-	UPROPERTY(EditAnywhere, AdvancedDisplay, BlueprintReadWrite, Category=SkeletalMesh)
-	uint32 bComponentUseFixedSkelBounds:1;
+	UPROPERTY(EditAnywhere, AdvancedDisplay, BlueprintReadWrite, Category = SkeletalMesh)
+		uint32 bComponentUseFixedSkelBounds : 1;
 
 	/** If true, when updating bounds from a PhysicsAsset, consider _all_ BodySetups, not just those flagged with bConsiderForBounds. */
-	UPROPERTY(EditAnywhere, AdvancedDisplay, BlueprintReadWrite, Category=SkeletalMesh)
-	uint32 bConsiderAllBodiesForBounds:1;
+	UPROPERTY(EditAnywhere, AdvancedDisplay, BlueprintReadWrite, Category = SkeletalMesh)
+		uint32 bConsiderAllBodiesForBounds : 1;
 
 	/** If true, update skeleton/attachments even when our Owner has not been rendered recently
 	 * @note if this is false, bone information may not be accurate, so be careful setting this to false if bone info is relevant to gameplay
@@ -302,106 +312,106 @@ public:
 	 * @note: In the output from SHOWSKELCOMPTICKTIME you want UpdatePoseTotal to be 0 when this is false for a specific component
 	 */
 	UPROPERTY()
-	uint32 bUpdateSkelWhenNotRendered_DEPRECATED:1;
+		uint32 bUpdateSkelWhenNotRendered_DEPRECATED : 1;
 
 	/** This is update frequency flag even when our Owner has not been rendered recently
-	 * 
+	 *
 	 * SMU_AlwaysTickPoseAndRefreshBones,			// Always Tick and Refresh BoneTransforms whether rendered or not
 	 * SMU_AlwaysTickPose,							// Always Tick, but Refresh BoneTransforms only when rendered
 	 * SMU_OnlyTickPoseWhenRendered,				// Tick only when rendered, and it will only RefreshBoneTransforms when rendered
-	 * 
+	 *
 	 */
-	UPROPERTY(EditAnywhere, AdvancedDisplay, BlueprintReadWrite, Category=SkeletalMesh)
-	TEnumAsByte<EMeshComponentUpdateFlag::Type> MeshComponentUpdateFlag;
+	UPROPERTY(EditAnywhere, AdvancedDisplay, BlueprintReadWrite, Category = SkeletalMesh)
+		TEnumAsByte<EMeshComponentUpdateFlag::Type> MeshComponentUpdateFlag;
 
 private:
 	/** If true, UpdateTransform will always result in a call to MeshObject->Update. */
 	UPROPERTY(transient)
-	uint32 bForceMeshObjectUpdate:1;
+		uint32 bForceMeshObjectUpdate : 1;
 
 public:
 	/** If true, DistanceFactor for this SkinnedMeshComponent will be added to global chart. */
 	UPROPERTY()
-	uint32 bChartDistanceFactor:1;
+		uint32 bChartDistanceFactor : 1;
 
 	/** Whether or not we can highlight selected sections - this should really only be done in the editor */
 	UPROPERTY(transient)
-	uint32 bCanHighlightSelectedSections:1;
+		uint32 bCanHighlightSelectedSections : 1;
 
 	/** true if mesh has been recently rendered, false otherwise */
 	UPROPERTY(transient)
-	uint32 bRecentlyRendered:1;
+		uint32 bRecentlyRendered : 1;
 
 	/** Editor only. Used for visualizing drawing order in Animset Viewer. If < 1.0,
 	  * only the specified fraction of triangles will be rendered
 	  */
 	UPROPERTY(transient)
-	float ProgressiveDrawingFraction;
+		float ProgressiveDrawingFraction;
 
 	/** Editor only. Used for manually selecting the alternate indices for
 	  * TRISORT_CustomLeftRight sections.
 	  */
 	UPROPERTY(transient)
-	uint8 CustomSortAlternateIndexMode;
+		uint8 CustomSortAlternateIndexMode;
 
-	/** 
+	/**
 	 * Override the Physics Asset of the mesh. It uses SkeletalMesh.PhysicsAsset, but if you'd like to override use this function
-	 * 
+	 *
 	 * @param	NewPhysicsAsset	New PhysicsAsset
 	 * @param	bForceReInit	Force reinitialize
 	 */
-	UFUNCTION(BlueprintCallable, Category="Components|SkinnedMesh")
-	virtual void SetPhysicsAsset(class UPhysicsAsset* NewPhysicsAsset, bool bForceReInit = false);
+	UFUNCTION(BlueprintCallable, Category = "Components|SkinnedMesh")
+		virtual void SetPhysicsAsset(class UPhysicsAsset* NewPhysicsAsset, bool bForceReInit = false);
 
 	/**
 	 * Find the index of bone by name. Looks in the current SkeletalMesh being used by this SkeletalMeshComponent.
-	 * 
+	 *
 	 * @param BoneName Name of bone to look up
-	 * 
+	 *
 	 * @return Index of the named bone in the current SkeletalMesh. Will return INDEX_NONE if bone not found.
 	 *
 	 * @see USkeletalMesh::GetBoneIndex.
 	 */
-	UFUNCTION(BlueprintCallable, Category="Components|SkinnedMesh")
-	int32 GetBoneIndex( FName BoneName ) const;
+	UFUNCTION(BlueprintCallable, Category = "Components|SkinnedMesh")
+		int32 GetBoneIndex(FName BoneName) const;
 
-	/** 
+	/**
 	 * Get Bone Name from index
 	 * @param BoneIndex Index of the bone
 	 *
-	 * @return the name of the bone at the specified index 
+	 * @return the name of the bone at the specified index
 	 */
-	UFUNCTION(BlueprintCallable, Category="Components|SkinnedMesh")
-	FName GetBoneName(int32 BoneIndex) const;
+	UFUNCTION(BlueprintCallable, Category = "Components|SkinnedMesh")
+		FName GetBoneName(int32 BoneIndex) const;
 
-		/**
-	 * Returns bone name linked to a given named socket on the skeletal mesh component.
-	 * If you're unsure to deal with sockets or bones names, you can use this function to filter through, and always return the bone name.
-	 *
-	 * @param	bone name or socket name
-	 *
-	 * @return	bone name
-	 */
-	UFUNCTION(BlueprintCallable, Category="Components|SkinnedMesh")
-	FName GetSocketBoneName(FName InSocketName);
+	/**
+ * Returns bone name linked to a given named socket on the skeletal mesh component.
+ * If you're unsure to deal with sockets or bones names, you can use this function to filter through, and always return the bone name.
+ *
+ * @param	bone name or socket name
+ *
+ * @return	bone name
+ */
+	UFUNCTION(BlueprintCallable, Category = "Components|SkinnedMesh")
+		FName GetSocketBoneName(FName InSocketName);
 
-	/** 
-	 * Change the SkeletalMesh that is rendered for this Component. Will re-initialize the animation tree etc. 
+	/**
+	 * Change the SkeletalMesh that is rendered for this Component. Will re-initialize the animation tree etc.
 	 *
 	 * @param NewMesh New mesh to set for this component
 	 */
-	UFUNCTION(BlueprintCallable, Category="Components|SkinnedMesh")
-	virtual void SetSkeletalMesh(class USkeletalMesh* NewMesh);
+	UFUNCTION(BlueprintCallable, Category = "Components|SkinnedMesh")
+		virtual void SetSkeletalMesh(class USkeletalMesh* NewMesh);
 
-	/** 
+	/**
 	 * Get Parent Bone of the input bone
-	 * 
+	 *
 	 * @param BoneName Name of the bone
 	 *
-	 * @return the name of the parent bone for the specified bone. Returns 'None' if the bone does not exist or it is the root bone 
+	 * @return the name of the parent bone for the specified bone. Returns 'None' if the bone does not exist or it is the root bone
 	 */
-	UFUNCTION(BlueprintCallable, Category="Components|SkinnedMesh")
-	FName GetParentBone(FName BoneName) const;
+	UFUNCTION(BlueprintCallable, Category = "Components|SkinnedMesh")
+		FName GetParentBone(FName BoneName) const;
 
 public:
 	/** Object responsible for sending bone transforms, vertex anim state etc. to render thread. */
@@ -444,7 +454,7 @@ public:
 	virtual FBoxSphereBounds CalcBounds(const FTransform& LocalToWorld) const override;
 	virtual bool HasAnySockets() const override;
 	virtual void QuerySupportedSockets(TArray<FComponentSocketDescription>& OutSockets) const override;
-	virtual void UpdateOverlaps(TArray<FOverlapInfo> const* PendingOverlaps=NULL, bool bDoNotifies=true, const TArray<FOverlapInfo>* OverlapsAtEndLocation=NULL) override;
+	virtual void UpdateOverlaps(TArray<FOverlapInfo> const* PendingOverlaps = NULL, bool bDoNotifies = true, const TArray<FOverlapInfo>* OverlapsAtEndLocation = NULL) override;
 	// End USceneComponent interface
 
 	// Begin UPrimitiveComponent interface
@@ -470,18 +480,18 @@ public:
 	 */
 	virtual bool ShouldCPUSkin();
 
-	/** 
-	 * Function to operate on mesh object after its created, 
+	/**
+	 * Function to operate on mesh object after its created,
 	 * but before it's attached.
-	 * 
+	 *
 	 * @param MeshObject - Mesh Object owned by this component
 	 */
 	virtual void PostInitMeshObject(class FSkeletalMeshObject*) {}
 
-	/** 
-	 * Simple, CPU evaluation of a vertex's skinned position (returned in component space) 
+	/**
+	 * Simple, CPU evaluation of a vertex's skinned position (returned in component space)
 	 *
-	 * @param VertexIndex Vertex Index. If compressed, this will be slow. 
+	 * @param VertexIndex Vertex Index. If compressed, this will be slow.
 	 */
 	virtual FVector GetSkinnedVertexPosition(int32 VertexIndex) const;
 
@@ -503,13 +513,13 @@ public:
 	 * Update functions
 	 */
 
-	/** 
+	/**
 	 * Refresh Bone Transform (SpaceBases)
 	 * Each class will need to implement this function
-	 * Ideally this function should be atomic (not relying on Tick or any other update.) 
-	 * 
+	 * Ideally this function should be atomic (not relying on Tick or any other update.)
+	 *
 	 * @param TickFunction Supplied as non null if we are running in a tick, allows us to create graph tasks for parallelism
-	 * 
+	 *
 	 */
 	virtual void RefreshBoneTransforms(FActorComponentTickFunction* TickFunction = NULL) PURE_VIRTUAL(USkinnedMeshComponent::RefreshBoneTransforms, );
 
@@ -520,17 +530,17 @@ public:
 	 *
 	 * @return	Return true if anything modified. Return false otherwise
 	 */
-	virtual void TickPose( float DeltaTime ) {}
+	virtual void TickPose(float DeltaTime) {}
 
-	/** 
+	/**
 	 * Update Slave Component. This gets called when MasterPoseComponent!=NULL
-	 * 
+	 *
 	 */
 	virtual void UpdateSlaveComponent();
 
-	/** 
-	 * Update the PredictedLODLevel and MaxDistanceFactor in the component from its MeshObject. 
-	 * 
+	/**
+	 * Update the PredictedLODLevel and MaxDistanceFactor in the component from its MeshObject.
+	 *
 	 * @return true if LOD has been changed. false otherwise.
 	 */
 	virtual bool UpdateLODStatus();
@@ -539,7 +549,7 @@ public:
 	void InitLODInfos();
 
 	/**
-	 * Rebuild BoneVisibilityStates array. Mostly refresh information of bones for BVS_HiddenByParent 
+	 * Rebuild BoneVisibilityStates array. Mostly refresh information of bones for BVS_HiddenByParent
 	 */
 	void RebuildVisibilityArray();
 
@@ -549,7 +559,34 @@ public:
 	 * Checks/updates material usage on proxy based on current morph target usage
 	 */
 	void UpdateMorphMaterialUsageOnProxy();
+	
+	/** Access Space Bases for reading */
+	const TArray<FTransform>& GetSpaceBases() const { return SpaceBasesArray[CurrentReadSpaceBases]; }
 
+	/** Get Access to the current editable space bases */
+	TArray<FTransform>& GetEditableSpaceBases() { return SpaceBasesArray[CurrentEditableSpaceBases]; }
+	const TArray<FTransform>& GetEditableSpaceBases() const { return SpaceBasesArray[CurrentEditableSpaceBases]; }
+
+	/** Get the number of space bases */
+	int32 GetNumSpaceBases() const { return GetSpaceBases().Num(); }
+
+	/** Flip the editable space base buffer */
+	void FlipEditableSpaceBases();
+
+	void SetSpaceBaseDoubleBuffering(bool bInDoubleBufferedBlendSpaces)
+	{
+		bDoubleBufferedBlendSpaces = bInDoubleBufferedBlendSpaces;
+
+		if (bDoubleBufferedBlendSpaces)
+		{
+			CurrentReadSpaceBases = CurrentReadSpaceBases;
+			CurrentEditableSpaceBases = 1 - CurrentReadSpaceBases;
+		}
+		else
+		{
+			CurrentReadSpaceBases = CurrentEditableSpaceBases = CurrentReadSpaceBases;
+		}
+	}
 
 protected:
 	/**
diff --git a/Engine/Source/Runtime/Engine/Private/DestructibleComponent.cpp b/Engine/Source/Runtime/Engine/Private/DestructibleComponent.cpp
--- a/Engine/Source/Runtime/Engine/Private/DestructibleComponent.cpp
+++ b/Engine/Source/Runtime/Engine/Private/DestructibleComponent.cpp
@@ -43,6 +43,8 @@ UDestructibleComponent::UDestructibleComponent(const FObjectInitializer& ObjectI
 	bMultiBodyOverlap = true;
 
 	LargeChunkThreshold = 25.f;
+
+	SetSpaceBaseDoubleBuffering(false);
 }
 
 void UDestructibleComponent::Serialize(FArchive& Ar)
@@ -924,7 +926,7 @@ void UDestructibleComponent::SetChunksWorldTM(const TArray<FUpdateChunksInfo>& U
 		const FQuat BoneRotation = InvRotation*WorldRotation;
 		const FVector BoneTranslation = InvRotation.RotateVector(WorldTranslation - ComponentToWorld.GetTranslation()) / ComponentToWorld.GetScale3D();
 
-		SpaceBases[BoneIndex] = FTransform(BoneRotation, BoneTranslation);
+		GetEditableSpaceBases()[BoneIndex] = FTransform(BoneRotation, BoneTranslation);
 	}
 
 	// Mark the transform as dirty, so the bounds are updated and sent to the render thread
@@ -932,6 +934,9 @@ void UDestructibleComponent::SetChunksWorldTM(const TArray<FUpdateChunksInfo>& U
 
 	// New bone positions need to be sent to render thread
 	MarkRenderDynamicDataDirty();
+
+	//Update bone visibilty and flip the editable space base buffer
+	FlipEditableSpaceBases();
 }
 
 void UDestructibleComponent::SetChunkWorldRT( int32 ChunkIndex, const FQuat& WorldRotation, const FVector& WorldTranslation )
@@ -954,7 +959,7 @@ void UDestructibleComponent::SetChunkWorldRT( int32 ChunkIndex, const FQuat& Wor
 	// More optimal form of the above
 	const FQuat BoneRotation = ComponentToWorld.GetRotation().Inverse()*WorldRotation;
 	const FVector BoneTranslation = ComponentToWorld.GetRotation().Inverse().RotateVector(WorldTranslation - ComponentToWorld.GetTranslation())/ComponentToWorld.GetScale3D();
-	SpaceBases[BoneIndex] = FTransform(BoneRotation, BoneTranslation);
+	GetEditableSpaceBases()[BoneIndex] = FTransform(BoneRotation, BoneTranslation);
 #endif
 }
 
diff --git a/Engine/Source/Runtime/Engine/Private/PhysicsEngine/PhysAnim.cpp b/Engine/Source/Runtime/Engine/Private/PhysicsEngine/PhysAnim.cpp
--- a/Engine/Source/Runtime/Engine/Private/PhysicsEngine/PhysAnim.cpp
+++ b/Engine/Source/Runtime/Engine/Private/PhysicsEngine/PhysAnim.cpp
@@ -69,11 +69,13 @@ void USkeletalMeshComponent::BlendPhysicsBones( TArray<FBoneIndexType>& InRequir
 	// Make sure scratch space is big enough.
 	TArray<FAssetWorldBoneTM> WorldBoneTMs;
 	WorldBoneTMs.Reset();
-	WorldBoneTMs.AddZeroed(SpaceBases.Num());
+	WorldBoneTMs.AddZeroed(GetNumSpaceBases());
 
 	FTransform LocalToWorldTM = ComponentToWorld;
 	LocalToWorldTM.RemoveScaling();
 
+	TArray<FTransform>& EditableSpaceBases = GetEditableSpaceBases();
+
 	// For each bone - see if we need to provide some data for it.
 	for(int32 i=0; i<InRequiredBones.Num(); i++)
 	{
@@ -155,12 +157,12 @@ void USkeletalMeshComponent::BlendPhysicsBones( TArray<FBoneIndexType>& InRequir
 		// Update SpaceBases entry for this bone now
 		if( BoneIndex == 0 )
 		{
-			SpaceBases[0] = LocalAtoms[0];
+			EditableSpaceBases[0] = LocalAtoms[0];
 		}
 		else
 		{
 			const int32 ParentIndex	= SkeletalMesh->RefSkeleton.GetParentIndex(BoneIndex);
-			SpaceBases[BoneIndex] = LocalAtoms[BoneIndex] * SpaceBases[ParentIndex];
+			EditableSpaceBases[BoneIndex] = LocalAtoms[BoneIndex] * EditableSpaceBases[ParentIndex];
 
 			/**
 			* Normalize rotations.
@@ -169,12 +171,12 @@ void USkeletalMeshComponent::BlendPhysicsBones( TArray<FBoneIndexType>& InRequir
 			* SpaceBases are used by external systems, we feed this to PhysX, send this to gameplay through bone and socket queries, etc.
 			* So this is a good place to make sure all transforms are normalized.
 			*/
-			SpaceBases[BoneIndex].NormalizeRotation();
+			EditableSpaceBases[BoneIndex].NormalizeRotation();
 		}
 
 		if (bUpdatePhysics && BodyInstance)
 		{
-			BodyInstance->SetBodyTransform(SpaceBases[BoneIndex] * ComponentToWorld, true);
+			BodyInstance->SetBodyTransform(EditableSpaceBases[BoneIndex] * ComponentToWorld, true);
 		}
 	}
 
@@ -226,7 +228,7 @@ void USkeletalMeshComponent::BlendInPhysics()
 
 
 
-void USkeletalMeshComponent::UpdateKinematicBonesToPhysics(bool bTeleport, bool bNeedsSkinning, bool bForceUpdate)
+void USkeletalMeshComponent::UpdateKinematicBonesToPhysics(const TArray<FTransform>& InSpaceBases, bool bTeleport, bool bNeedsSkinning, bool bForceUpdate)
 {
 	SCOPE_CYCLE_COUNTER(STAT_UpdateRBBones);
 
@@ -252,14 +254,14 @@ void USkeletalMeshComponent::UpdateKinematicBonesToPhysics(bool bTeleport, bool
 	}
 
 	// If desired, draw the skeleton at the point where we pass it to the physics.
-	if( bShowPrePhysBones && SkeletalMesh && SpaceBases.Num() == SkeletalMesh->RefSkeleton.GetNum() )
+	if (bShowPrePhysBones && SkeletalMesh && InSpaceBases.Num() == SkeletalMesh->RefSkeleton.GetNum())
 	{
-		for(int32 i=1; i<SpaceBases.Num(); i++)
+		for (int32 i = 1; i<InSpaceBases.Num(); i++)
 		{
-			FVector ThisPos = CurrentLocalToWorld.TransformPosition( SpaceBases[i].GetLocation() );
+			FVector ThisPos = CurrentLocalToWorld.TransformPosition(InSpaceBases[i].GetLocation());
 
 			int32 ParentIndex = SkeletalMesh->RefSkeleton.GetParentIndex(i);
-			FVector ParentPos = CurrentLocalToWorld.TransformPosition( SpaceBases[ParentIndex].GetLocation() );
+			FVector ParentPos = CurrentLocalToWorld.TransformPosition(InSpaceBases[ParentIndex].GetLocation());
 
 			GetWorld()->LineBatcher->DrawLine(ThisPos, ParentPos, AnimSkelDrawColor, SDPG_Foreground);
 		}
@@ -303,14 +305,14 @@ void USkeletalMeshComponent::UpdateKinematicBonesToPhysics(bool bTeleport, bool
 					int32 const BoneIndex = SkeletalMesh->RefSkeleton.FindBoneIndex(BodyName);
 
 					// If we could not find it - warn.
-					if (BoneIndex == INDEX_NONE || BoneIndex >= SpaceBases.Num())
+					if (BoneIndex == INDEX_NONE || BoneIndex >= GetNumSpaceBases())
 					{
 						UE_LOG(LogPhysics, Log, TEXT("UpdateRBBones: WARNING: Failed to find bone '%s' need by PhysicsAsset '%s' in SkeletalMesh '%s'."), *BodyName.ToString(), *PhysicsAsset->GetName(), *SkeletalMesh->GetName());
 					}
 					else
 					{
 						// update bone transform to world
-						FTransform BoneTransform = SpaceBases[BoneIndex] * CurrentLocalToWorld;
+						FTransform BoneTransform = InSpaceBases[BoneIndex] * CurrentLocalToWorld;
 
 						// move body
 						BodyInst->SetBodyTransform(BoneTransform, bTeleport);
diff --git a/Engine/Source/Runtime/Engine/Private/PoseableMeshComponent.cpp b/Engine/Source/Runtime/Engine/Private/PoseableMeshComponent.cpp
--- a/Engine/Source/Runtime/Engine/Private/PoseableMeshComponent.cpp
+++ b/Engine/Source/Runtime/Engine/Private/PoseableMeshComponent.cpp
@@ -55,7 +55,7 @@ void UPoseableMeshComponent::RefreshBoneTransforms(FActorComponentTickFunction*
 	}
 
 	// Do nothing more if no bones in skeleton.
-	if( SpaceBases.Num() == 0 )
+	if( GetNumSpaceBases() == 0 )
 	{
 		return;
 	}
@@ -77,7 +77,7 @@ void UPoseableMeshComponent::FillSpaceBases()
 
 	// right now all this does is to convert to SpaceBases
 	check( SkeletalMesh->RefSkeleton.GetNum() == LocalAtoms.Num() );
-	check( SkeletalMesh->RefSkeleton.GetNum() == SpaceBases.Num() );
+	check( SkeletalMesh->RefSkeleton.GetNum() == GetNumSpaceBases());
 	check( SkeletalMesh->RefSkeleton.GetNum() == BoneVisibilityStates.Num() );
 
 	const int32 NumBones = LocalAtoms.Num();
@@ -89,9 +89,9 @@ void UPoseableMeshComponent::FillSpaceBases()
 #endif
 	// Build in 3 passes.
 	FTransform* LocalTransformsData = LocalAtoms.GetData(); 
-	FTransform* SpaceBasesData = SpaceBases.GetData();
+	FTransform* SpaceBasesData = GetEditableSpaceBases().GetData();
 	
-	SpaceBases[0] = LocalAtoms[0];
+	GetEditableSpaceBases()[0] = LocalAtoms[0];
 #if (UE_BUILD_DEBUG || UE_BUILD_DEVELOPMENT)
 	BoneProcessed[0] = 1;
 #endif
@@ -114,8 +114,8 @@ void UPoseableMeshComponent::FillSpaceBases()
 #endif
 		FTransform::Multiply(SpaceBasesData + BoneIndex, LocalTransformsData + BoneIndex, SpaceBasesData + ParentIndex);
 
-		checkSlow( SpaceBases[BoneIndex].IsRotationNormalized() );
-		checkSlow( !SpaceBases[BoneIndex].ContainsNaN() );
+		checkSlow(GetEditableSpaceBases()[BoneIndex].IsRotationNormalized());
+		checkSlow(!GetEditableSpaceBases()[BoneIndex].ContainsNaN());
 	}
 }
 
diff --git a/Engine/Source/Runtime/Engine/Private/SkeletalMesh.cpp b/Engine/Source/Runtime/Engine/Private/SkeletalMesh.cpp
--- a/Engine/Source/Runtime/Engine/Private/SkeletalMesh.cpp
+++ b/Engine/Source/Runtime/Engine/Private/SkeletalMesh.cpp
@@ -4852,6 +4852,10 @@ USkinnedMeshComponent::USkinnedMeshComponent(const FObjectInitializer& ObjectIni
 	SectionIndexPreview = -1;
 #endif // WITH_EDITORONLY_DATA
 	bPerBoneMotionBlur = true;
+
+	bDoubleBufferedBlendSpaces = true;
+	CurrentEditableSpaceBases = 0;
+	CurrentReadSpaceBases = 1;
 }
 
 
@@ -4874,7 +4878,8 @@ void USkinnedMeshComponent::Serialize(FArchive& Ar)
 	if(Ar.IsCountingMemory())
 	{
 		// add all native variables - mostly bigger chunks 
-		SpaceBases.CountBytes(Ar);
+		SpaceBasesArray[0].CountBytes(Ar);
+		SpaceBasesArray[1].CountBytes(Ar);
 		MasterBoneMap.CountBytes(Ar);
 	}
 
diff --git a/Engine/Source/Runtime/Engine/Private/SkeletalMeshComponent.cpp b/Engine/Source/Runtime/Engine/Private/SkeletalMeshComponent.cpp
--- a/Engine/Source/Runtime/Engine/Private/SkeletalMeshComponent.cpp
+++ b/Engine/Source/Runtime/Engine/Private/SkeletalMeshComponent.cpp
@@ -634,7 +634,9 @@ void USkeletalMeshComponent::FillSpaceBases(const USkeletalMesh* InSkeletalMesh,
 	for(int32 i=1; i<RequiredBones.Num(); i++)
 	{
 		const int32 BoneIndex = RequiredBones[i];
-		FPlatformMisc::Prefetch(SpaceBasesData + BoneIndex);
+		FTransform* SpaceBase = SpaceBasesData + BoneIndex;
+
+		FPlatformMisc::Prefetch(SpaceBase);
 
 #if (UE_BUILD_DEBUG || UE_BUILD_DEVELOPMENT)
 		// Mark bone as processed
@@ -642,16 +644,17 @@ void USkeletalMeshComponent::FillSpaceBases(const USkeletalMesh* InSkeletalMesh,
 #endif
 		// For all bones below the root, final component-space transform is relative transform * component-space transform of parent.
 		const int32 ParentIndex = InSkeletalMesh->RefSkeleton.GetParentIndex(BoneIndex);
-		FPlatformMisc::Prefetch(SpaceBasesData + ParentIndex);
+		FTransform* ParentSpaceBase = SpaceBasesData + ParentIndex;
+		FPlatformMisc::Prefetch(ParentSpaceBase);
 
 #if (UE_BUILD_DEBUG || UE_BUILD_DEVELOPMENT)
 		// Check the precondition that Parents occur before Children in the RequiredBones array.
 		checkSlow(BoneProcessed[ParentIndex] == 1);
 #endif
-		FTransform::Multiply(SpaceBasesData + BoneIndex, LocalTransformsData + BoneIndex, SpaceBasesData + ParentIndex);
+		FTransform::Multiply(SpaceBase, LocalTransformsData + BoneIndex, ParentSpaceBase);
 
-		checkSlow(DestSpaceBases[BoneIndex].IsRotationNormalized());
-		checkSlow(!DestSpaceBases[BoneIndex].ContainsNaN());
+		checkSlow(SpaceBase->IsRotationNormalized());
+		checkSlow(!SpaceBase->ContainsNaN());
 	}
 
 	/**
@@ -762,7 +765,7 @@ void USkeletalMeshComponent::RecalcRequiredBones(int32 LODIndex)
 	// mirror table/phys body ones has to be calculated
 	if (ShouldUpdateBoneVisibility())
 	{
-		check(BoneVisibilityStates.Num() == SpaceBases.Num());
+		check(BoneVisibilityStates.Num() == GetNumSpaceBases());
 
 		int32 VisibleBoneWriteIndex = 0;
 		for (int32 i = 0; i < RequiredBones.Num(); ++i)
@@ -918,7 +921,7 @@ void USkeletalMeshComponent::RefreshBoneTransforms(FActorComponentTickFunction*
 {
 	SCOPE_CYCLE_COUNTER(STAT_RefreshBoneTransforms);
 
-	if (!SkeletalMesh || SpaceBases.Num() == 0)
+	if (!SkeletalMesh || GetNumSpaceBases() == 0)
 	{
 		return;
 	}
@@ -941,7 +944,7 @@ void USkeletalMeshComponent::RefreshBoneTransforms(FActorComponentTickFunction*
 	const bool bInvalidCachedBones = bDoUpdateRateOptimization &&
 									( (LocalAtoms.Num() != SkeletalMesh->RefSkeleton.GetNum())
 									  || (LocalAtoms.Num() != CachedLocalAtoms.Num())
-									  || (SpaceBases.Num() != CachedSpaceBases.Num()) );
+									  || (GetNumSpaceBases() != CachedSpaceBases.Num()) );
 
 	AnimEvaluationContext.bDoEvaluation = !bDoUpdateRateOptimization || bInvalidCachedBones || !AnimUpdateRateParams.ShouldSkipEvaluation();
 	
@@ -962,7 +965,7 @@ void USkeletalMeshComponent::RefreshBoneTransforms(FActorComponentTickFunction*
 		{
 			// Initialize Parallel Task arrays
 			AnimEvaluationContext.LocalAtoms = LocalAtoms;
-			AnimEvaluationContext.SpaceBases = SpaceBases;
+			AnimEvaluationContext.SpaceBases = GetSpaceBases();
 			AnimEvaluationContext.VertexAnims = ActiveVertexAnims;
 		}
 
@@ -986,13 +989,13 @@ void USkeletalMeshComponent::RefreshBoneTransforms(FActorComponentTickFunction*
 			}
 			else
 			{
-				PerformAnimationEvaluation(SkeletalMesh, AnimScriptInstance, SpaceBases, LocalAtoms, ActiveVertexAnims, RootBoneTranslation);
+				PerformAnimationEvaluation(SkeletalMesh, AnimScriptInstance, GetEditableSpaceBases(), LocalAtoms, ActiveVertexAnims, RootBoneTranslation);
 			}
 		}
 		else if (!AnimEvaluationContext.bDoInterpolation)
 		{
 			LocalAtoms = CachedLocalAtoms;
-			SpaceBases = CachedSpaceBases;
+			GetEditableSpaceBases() = CachedSpaceBases;
 		}
 
 		PostAnimEvaluation(AnimEvaluationContext);
@@ -1007,7 +1010,7 @@ void USkeletalMeshComponent::PostAnimEvaluation(FAnimationEvaluationContext& Eva
 	SCOPE_CYCLE_COUNTER(STAT_PostAnimEvaluation);
 	if (EvaluationContext.bDuplicateToCacheBones)
 	{
-		CachedSpaceBases = SpaceBases;
+		CachedSpaceBases = GetEditableSpaceBases();
 		CachedLocalAtoms = LocalAtoms;
 	}
 
@@ -1017,14 +1020,14 @@ void USkeletalMeshComponent::PostAnimEvaluation(FAnimationEvaluationContext& Eva
 
 		const float Alpha = 0.25f + (1.f / float(FMath::Max(AnimUpdateRateParams.GetEvaluationRate(), 2) * 2));
 		FAnimationRuntime::LerpBoneTransforms(LocalAtoms, CachedLocalAtoms, Alpha, RequiredBones);
-		FAnimationRuntime::LerpBoneTransforms(SpaceBases, CachedSpaceBases, Alpha, RequiredBones);
+		FAnimationRuntime::LerpBoneTransforms(GetEditableSpaceBases(), CachedSpaceBases, Alpha, RequiredBones);
 	}
 
 	// Transforms updated, cached local bounds are now out of date.
 	InvalidateCachedBounds();
 
 	// update physics data from animated data
-	UpdateKinematicBonesToPhysics(false, true);
+	UpdateKinematicBonesToPhysics(GetEditableSpaceBases(), false, true);
 	UpdateRBJointMotors();
 
 	// @todo anim : hack TTP 224385	ANIM: Skeletalmesh double buffer
@@ -1036,6 +1039,9 @@ void USkeletalMeshComponent::PostAnimEvaluation(FAnimationEvaluationContext& Eva
 	{
 		SCOPE_CYCLE_COUNTER(STAT_UpdateLocalToWorldAndOverlaps);
 
+		// Updated last good bone positions
+		FlipEditableSpaceBases();
+
 		// New bone positions need to be sent to render thread
 		UpdateComponentToWorld();
 
@@ -1322,7 +1328,7 @@ void USkeletalMeshComponent::DebugDrawBones(UCanvas* Canvas, bool bSimpleBones)
 		{
 			int32 BoneIndex = RequiredBones[Index];
 			int32 ParentIndex = SkeletalMesh->RefSkeleton.GetParentIndex(BoneIndex);
-			FTransform BoneTM = (SpaceBases[BoneIndex] * ComponentToWorld);
+			FTransform BoneTM = (GetSpaceBases()[BoneIndex] * ComponentToWorld);
 			FVector Start, End;
 			FLinearColor LineColor;
 
@@ -1330,7 +1336,7 @@ void USkeletalMeshComponent::DebugDrawBones(UCanvas* Canvas, bool bSimpleBones)
 
 			if (ParentIndex >=0)
 			{
-				Start = (SpaceBases[ParentIndex] * ComponentToWorld).GetLocation();
+				Start = (GetSpaceBases()[ParentIndex] * ComponentToWorld).GetLocation();
 				LineColor = FLinearColor::White;
 			}
 			else
diff --git a/Engine/Source/Runtime/Engine/Private/SkeletalMeshComponentPhysics.cpp b/Engine/Source/Runtime/Engine/Private/SkeletalMeshComponentPhysics.cpp
--- a/Engine/Source/Runtime/Engine/Private/SkeletalMeshComponentPhysics.cpp
+++ b/Engine/Source/Runtime/Engine/Private/SkeletalMeshComponentPhysics.cpp
@@ -1293,7 +1293,7 @@ void USkeletalMeshComponent::OnUpdateTransform(bool bSkipPhysicsMove)
 	// Always send new transform to physics
 	if(bPhysicsStateCreated && !bSkipPhysicsMove )
 	{
-		UpdateKinematicBonesToPhysics(false, false, true);
+		UpdateKinematicBonesToPhysics(GetSpaceBases(), false, false, true);
 	}
 
 #if WITH_APEX_CLOTHING
@@ -3058,6 +3058,7 @@ void USkeletalMeshComponent::PreClothTick(float DeltaTime)
 	// if physics is disabled on dedicated server, no reason to be here. 
 	if (!bEnablePhysicsOnDedicatedServer && IsRunningDedicatedServer())
 	{
+		FlipEditableSpaceBases();
 		return;
 	}
 
@@ -3079,6 +3080,12 @@ void USkeletalMeshComponent::PreClothTick(float DeltaTime)
 		BlendInPhysics();
 	}
 
+	if (IsSimulatingPhysics())
+	{
+		// If we aren't simulating we will have already flipped this
+		FlipEditableSpaceBases();
+	}
+
 	//TODO: move this into pre physics tick
 
 #if WITH_APEX_CLOTHING
@@ -3449,14 +3456,9 @@ void USkeletalMeshComponent::UpdateClothState(float DeltaTime)
 		return;
 	}
 
-	TArray<FTransform>* BoneTransforms = &SpaceBases;
-
-	if(MasterPoseComponent.IsValid())
-	{
-		BoneTransforms = &(MasterPoseComponent.Get()->SpaceBases);
-	}
+	const TArray<FTransform>& BoneTransforms = MasterPoseComponent.IsValid() ? MasterPoseComponent.Get()->GetSpaceBases() : GetSpaceBases();
 
-	if(BoneTransforms->Num() == 0)
+	if(BoneTransforms.Num() == 0)
 	{
 		return;
 	}
@@ -3503,7 +3505,7 @@ void USkeletalMeshComponent::UpdateClothState(float DeltaTime)
  
 					// If ParentBoneIndex is valid, grab matrix from MasterPoseComponent.
 					if( MasterBoneIndex != INDEX_NONE && 
-						MasterBoneIndex < MasterPoseComponent->SpaceBases.Num())
+						MasterBoneIndex < MasterPoseComponent->GetNumSpaceBases())
 					{
 						BoneIndex = MasterBoneIndex;
 					}
@@ -3512,7 +3514,7 @@ void USkeletalMeshComponent::UpdateClothState(float DeltaTime)
 
 		   if(BoneIndex != INDEX_NONE)
 		   {
-			   BoneMatrices[Index] = U2PMatrix((*BoneTransforms)[BoneIndex].ToMatrixWithScale());
+			   BoneMatrices[Index] = U2PMatrix(BoneTransforms[BoneIndex].ToMatrixWithScale());
 		   }
 		   else
 		   {
@@ -4908,7 +4910,7 @@ FTransform USkeletalMeshComponent::GetComponentTransformFromBodyInstance(FBodyIn
 	FTransform BodyTransform = UseBI->GetUnrealWorldTransform();
 	if (RootBodyData.BoneIndex != INDEX_NONE)
 	{
-		FTransform RootTransform(SpaceBases[RootBodyData.BoneIndex]);
+		FTransform RootTransform(GetSpaceBases()[RootBodyData.BoneIndex]);
 		return RootTransform.GetRelativeTransformReverse(BodyTransform);
 	}
 
diff --git a/Engine/Source/Runtime/Engine/Private/SkeletalRender.cpp b/Engine/Source/Runtime/Engine/Private/SkeletalRender.cpp
--- a/Engine/Source/Runtime/Engine/Private/SkeletalRender.cpp
+++ b/Engine/Source/Runtime/Engine/Private/SkeletalRender.cpp
@@ -205,9 +205,14 @@ void UpdateRefToLocalMatrices( TArray<FMatrix>& ReferenceToLocal, const USkinned
 		ReferenceToLocal.AddUninitialized(ThisMesh->RefBasesInvMatrix.Num());
 	}
 
-
 	const bool bIsMasterCompValid = MasterComp && InMeshComponent->MasterBoneMap.Num() == ThisMesh->RefSkeleton.GetNum();
-	const TArray<FBoneIndexType>* RequiredBoneSets[3] = { &LOD.ActiveBoneIndices, ExtraRequiredBoneIndices, NULL };
+
+	TArray<FBoneIndexType> ActiveBoneIndices(LOD.ActiveBoneIndices);
+	ActiveBoneIndices.Sort();
+
+	const TArray<FBoneIndexType>* RequiredBoneSets[3] = { &ActiveBoneIndices, ExtraRequiredBoneIndices, NULL };
+
+	const bool bBoneVisibilityStatesValid = InMeshComponent->BoneVisibilityStates.Num() == InMeshComponent->GetNumSpaceBases();
 
 	// Handle case of using ParentAnimComponent for SpaceBases.
 	for( int32 RequiredBoneSetIndex = 0; RequiredBoneSets[RequiredBoneSetIndex]!=NULL; RequiredBoneSetIndex++ )
@@ -221,38 +226,45 @@ void UpdateRefToLocalMatrices( TArray<FMatrix>& ReferenceToLocal, const USkinned
 
 			if ( ThisMesh->RefBasesInvMatrix.IsValidIndex(ThisBoneIndex) )
 			{
-				FTransform const *BoneTransform = NULL;
+				// On the off chance the parent matrix isn't valid, revert to identity.
+				ReferenceToLocal[ThisBoneIndex] = FMatrix::Identity;
 
 				if( bIsMasterCompValid )
 				{
 					// If valid, use matrix from parent component.
 					const int32 ParentBoneIndex = InMeshComponent->MasterBoneMap[ThisBoneIndex];
-					if ( MasterComp->SpaceBases.IsValidIndex(ParentBoneIndex) )
+					if ( MasterComp->GetSpaceBases().IsValidIndex(ParentBoneIndex) )
 					{
-						BoneTransform = &MasterComp->SpaceBases[ParentBoneIndex];
+						checkSlow(MasterComp->GetSpaceBases()[ParentBoneIndex].IsRotationNormalized());
+						ReferenceToLocal[ThisBoneIndex] = MasterComp->GetSpaceBases()[ParentBoneIndex].ToMatrixWithScale();
 					}
 				}
 				else
 				{
-					if ( InMeshComponent->SpaceBases.IsValidIndex(ThisBoneIndex) )
+					if (InMeshComponent->GetSpaceBases().IsValidIndex(ThisBoneIndex))
 					{
 						// If we can't find this bone in the parent, we just use the reference pose.
-						BoneTransform = &InMeshComponent->SpaceBases[ThisBoneIndex];
+						if (bBoneVisibilityStatesValid)
+						{
+							const int32 ParentIndex = ThisMesh->RefSkeleton.GetParentIndex(ThisBoneIndex);
+							bool bNeedToHideBone = InMeshComponent->BoneVisibilityStates[ThisBoneIndex] != BVS_Visible;
+							if (bNeedToHideBone && ParentIndex != INDEX_NONE)
+							{
+								ReferenceToLocal[ThisBoneIndex] = ReferenceToLocal[ParentIndex].ApplyScale(0.f);
+							}
+							else
+							{
+								checkSlow(InMeshComponent->GetSpaceBases()[ThisBoneIndex].IsRotationNormalized());
+								ReferenceToLocal[ThisBoneIndex] = InMeshComponent->GetSpaceBases()[ThisBoneIndex].ToMatrixWithScale();
+							}
+						}
+						else
+						{
+							checkSlow(InMeshComponent->GetSpaceBases()[ThisBoneIndex].IsRotationNormalized());
+							ReferenceToLocal[ThisBoneIndex] = InMeshComponent->GetSpaceBases()[ThisBoneIndex].ToMatrixWithScale();
+						}
 					}
 				}
-
-				if ( BoneTransform )
-				{
-					// we need to convert to Matrix before applying refBases since the way BoneTransform applies transform isn't same as Matrix
-					// from this point, it is all Matrix till render thread
-					ReferenceToLocal[ThisBoneIndex] = ThisMesh->RefBasesInvMatrix[ThisBoneIndex] * BoneTransform->ToMatrixWithScale();
-					checkSlow( BoneTransform->IsRotationNormalized() );
-				}
-				else
-				{
-					// On the off chance the parent matrix isn't valid, revert to identity.
-					ReferenceToLocal[ThisBoneIndex] = FMatrix::Identity;
-				}
 			}
 			// removed else statement to set ReferenceToLocal[ThisBoneIndex] = FTransform::Identity;
 			// since it failed in ( ThisMesh->RefBasesInvMatrix.IsValidIndex(ThisBoneIndex) ), ReferenceToLocal is not valid either
@@ -260,6 +272,11 @@ void UpdateRefToLocalMatrices( TArray<FMatrix>& ReferenceToLocal, const USkinned
 			// if(ReferenceToLocal.Num() != ThisMesh->RefBasesInvMatrix.Num())
 		}
 	}
+
+	for (int32 ThisBoneIndex = 0; ThisBoneIndex < ReferenceToLocal.Num(); ++ThisBoneIndex)
+	{
+		ReferenceToLocal[ThisBoneIndex] = ThisMesh->RefBasesInvMatrix[ThisBoneIndex] * ReferenceToLocal[ThisBoneIndex];
+	}
 }
 
 /**
@@ -297,19 +314,19 @@ void UpdateCustomLeftRightVectors( TArray<FTwoVectors>& OutVectors, const USkinn
 			else
 			{
 				int32 SpaceBasesBoneIndex = ThisMesh->RefSkeleton.FindBoneIndex(CustomLeftRightBoneName);
-				const USkinnedMeshComponent* SpaceBasesComp = InMeshComponent;
+				const TArray<FTransform>* SpaceBases = &InMeshComponent->GetSpaceBases();
 				
 				// Handle case of using MasterPoseComponent for SpaceBases.
 				if( MasterComp && InMeshComponent->MasterBoneMap.Num() == ThisMesh->RefSkeleton.GetNum() && SpaceBasesBoneIndex != INDEX_NONE )
 				{
 					// If valid, use matrix from parent component.
 					SpaceBasesBoneIndex = InMeshComponent->MasterBoneMap[SpaceBasesBoneIndex];
-					SpaceBasesComp = MasterComp;
+					SpaceBases = &MasterComp->GetSpaceBases();
 				}
 
-				if ( SpaceBasesComp->SpaceBases.IsValidIndex(SpaceBasesBoneIndex) )
+				if (SpaceBases->IsValidIndex(SpaceBasesBoneIndex))
 				{
-					const FMatrix BoneMatrix = SpaceBasesComp->SpaceBases[SpaceBasesBoneIndex].ToMatrixWithScale();
+					const FMatrix BoneMatrix = (*SpaceBases)[SpaceBasesBoneIndex].ToMatrixWithScale();
 					OutVectors[SectionIndex].v1 = BoneMatrix.GetOrigin();
 
 					EAxis::Type MatrixAxis = EAxis::X;
diff --git a/Engine/Source/Runtime/Engine/Private/SkeletalRenderCPUSkin.cpp b/Engine/Source/Runtime/Engine/Private/SkeletalRenderCPUSkin.cpp
--- a/Engine/Source/Runtime/Engine/Private/SkeletalRenderCPUSkin.cpp
+++ b/Engine/Source/Runtime/Engine/Private/SkeletalRenderCPUSkin.cpp
@@ -452,7 +452,7 @@ FDynamicSkelMeshObjectDataCPUSkin::FDynamicSkelMeshObjectDataCPUSkin(
 	UpdateCustomLeftRightVectors( CustomLeftRightVectors, InMeshComponent, InSkeletalMeshResource, LODIndex );
 
 #if !(UE_BUILD_SHIPPING || UE_BUILD_TEST)
-	MeshSpaceBases = InMeshComponent->SpaceBases;
+	MeshSpaceBases = InMeshComponent->GetSpaceBases();
 #endif
 }
 
diff --git a/Engine/Source/Runtime/Engine/Private/SkeletalRenderGPUSkin.cpp b/Engine/Source/Runtime/Engine/Private/SkeletalRenderGPUSkin.cpp
--- a/Engine/Source/Runtime/Engine/Private/SkeletalRenderGPUSkin.cpp
+++ b/Engine/Source/Runtime/Engine/Private/SkeletalRenderGPUSkin.cpp
@@ -991,7 +991,7 @@ FDynamicSkelMeshObjectDataGPUSkin::FDynamicSkelMeshObjectDataGPUSkin(
 
 
 #if !(UE_BUILD_SHIPPING || UE_BUILD_TEST)
-	MeshSpaceBases = InMeshComponent->SpaceBases;
+	MeshSpaceBases = InMeshComponent->GetSpaceBases();
 #endif
 
 	// find number of morphs that are currently weighted and will affect the mesh
diff --git a/Engine/Source/Runtime/Engine/Private/SkinnedMeshComponent.cpp b/Engine/Source/Runtime/Engine/Private/SkinnedMeshComponent.cpp
--- a/Engine/Source/Runtime/Engine/Private/SkinnedMeshComponent.cpp
+++ b/Engine/Source/Runtime/Engine/Private/SkinnedMeshComponent.cpp
@@ -161,26 +161,6 @@ void USkinnedMeshComponent::SendRenderDynamicData_Concurrent()
 			ActiveVertexAnims.Empty();
 		}
 
-		if( BoneVisibilityStates.Num() == SpaceBases.Num() )
-		{
-			// for invisible bones, I'll still need to update the transform, so that rendering can use it for skinning
-			uint8 const * BoneVisibilityState = BoneVisibilityStates.GetData();
-			FTransform* SpaceBase = SpaceBases.GetData();
-			for (int32 BoneIndex=0; BoneIndex<BoneVisibilityStates.Num(); ++BoneIndex, ++BoneVisibilityState, ++SpaceBase)
-			{
-				if (*BoneVisibilityState!=BVS_Visible)
-				{
-					if (BoneIndex != 0 )
-					{
-						// since they're invisible, copy parent transform to itself and scale set to be 0
-						int32 const ParentIndex =  SkeletalMesh->RefSkeleton.GetParentIndex(BoneIndex);
-						*SpaceBase = SpaceBases[ParentIndex];
-					}
-					SpaceBase->SetScale3D(FVector::ZeroVector);
-				}
-			}
-		}
-
 		MeshObject->Update(UseLOD,this,ActiveVertexAnims);  // send to rendering thread
 		MeshObject->bHasBeenUpdatedAtLeastOnce = true;
 		
@@ -475,7 +455,7 @@ FBoxSphereBounds USkinnedMeshComponent::CalcMeshBound(const FVector& RootOffset,
 	// Can only use the PhysicsAsset to calculate the bounding box if we are not non-uniformly scaling the mesh.
 	const bool bCanUsePhysicsAsset = DrawScale.IsUniform() && (SkeletalMesh != NULL)
 		// either space base exists or child component
-		&& ( (SpaceBases.Num() == SkeletalMesh->RefSkeleton.GetNum()) || (MasterPhysicsAsset) );
+		&& ( (GetNumSpaceBases() == SkeletalMesh->RefSkeleton.GetNum()) || (MasterPhysicsAsset) );
 
 	const bool bDetailModeAllowsRendering = (DetailMode <= GetCachedScalabilityCVars().DetailMode);
 	const bool bVisible = ( bDetailModeAllowsRendering && (ShouldRender() || bCastHiddenShadow));
@@ -563,9 +543,9 @@ FMatrix USkinnedMeshComponent::GetBoneMatrix(int32 BoneIdx) const
 
 			// If ParentBoneIndex is valid, grab matrix from MasterPoseComponent.
 			if(	ParentBoneIndex != INDEX_NONE && 
-				ParentBoneIndex < MasterPoseComponent->SpaceBases.Num())
+				ParentBoneIndex < MasterPoseComponent->GetNumSpaceBases())
 			{
-				return MasterPoseComponent->SpaceBases[ParentBoneIndex].ToMatrixWithScale() * ComponentToWorld.ToMatrixWithScale();
+				return MasterPoseComponent->GetSpaceBases()[ParentBoneIndex].ToMatrixWithScale() * ComponentToWorld.ToMatrixWithScale();
 			}
 			else
 			{
@@ -581,9 +561,9 @@ FMatrix USkinnedMeshComponent::GetBoneMatrix(int32 BoneIdx) const
 	}
 	else
 	{
-		if( SpaceBases.Num() && BoneIdx < SpaceBases.Num() )
+		if( GetNumSpaceBases() && BoneIdx < GetNumSpaceBases() )
 		{
-			return SpaceBases[BoneIdx].ToMatrixWithScale() * ComponentToWorld.ToMatrixWithScale();
+			return GetSpaceBases()[BoneIdx].ToMatrixWithScale() * ComponentToWorld.ToMatrixWithScale();
 		}
 		else
 		{
@@ -612,9 +592,9 @@ FTransform USkinnedMeshComponent::GetBoneTransform(int32 BoneIdx) const
 
 			// If ParentBoneIndex is valid, grab matrix from MasterPoseComponent.
 			if(	ParentBoneIndex != INDEX_NONE && 
-				ParentBoneIndex < MasterPoseComponent->SpaceBases.Num())
+				ParentBoneIndex < MasterPoseComponent->GetNumSpaceBases())
 			{
-				return MasterPoseComponent->SpaceBases[ParentBoneIndex] * ComponentToWorld;
+				return MasterPoseComponent->GetSpaceBases()[ParentBoneIndex] * ComponentToWorld;
 			}
 			else
 			{
@@ -630,9 +610,9 @@ FTransform USkinnedMeshComponent::GetBoneTransform(int32 BoneIdx) const
 	}
 	else
 	{
-		if( SpaceBases.Num() && BoneIdx < SpaceBases.Num() )
+		if( GetNumSpaceBases() && BoneIdx < GetNumSpaceBases() )
 		{
-			return SpaceBases[BoneIdx] * ComponentToWorld;
+			return GetSpaceBases()[BoneIdx] * ComponentToWorld;
 		}
 		else
 		{
@@ -787,14 +767,17 @@ bool USkinnedMeshComponent::AllocateTransformData()
 	// Allocate transforms if not present.
 	if ( SkeletalMesh != NULL && MasterPoseComponent == NULL )
 	{
-		if( SpaceBases.Num() != SkeletalMesh->RefSkeleton.GetNum() )
+		if( GetNumSpaceBases() != SkeletalMesh->RefSkeleton.GetNum() )
 		{
-			SpaceBases.Empty( SkeletalMesh->RefSkeleton.GetNum() );
-			SpaceBases.AddUninitialized( SkeletalMesh->RefSkeleton.GetNum() );
-
-			for(int32 I=0; I<SkeletalMesh->RefSkeleton.GetNum(); ++I)
+			for (int32 BaseIndex = 0; BaseIndex < 2; ++BaseIndex)
 			{
-				SpaceBases[I].SetIdentity();
+				SpaceBasesArray[BaseIndex].Empty(SkeletalMesh->RefSkeleton.GetNum());
+				SpaceBasesArray[BaseIndex].AddUninitialized(SkeletalMesh->RefSkeleton.GetNum());
+
+				for (int32 I = 0; I < SkeletalMesh->RefSkeleton.GetNum(); ++I)
+				{
+					SpaceBasesArray[BaseIndex][I].SetIdentity();
+				}
 			}
  
 			BoneVisibilityStates.Empty( SkeletalMesh->RefSkeleton.GetNum() );
@@ -813,14 +796,16 @@ bool USkinnedMeshComponent::AllocateTransformData()
 	}
 	
 	// Reset the animation stuff when changing mesh.
-	SpaceBases.Empty();		
+	SpaceBasesArray[0].Empty();
+	SpaceBasesArray[1].Empty();
 
 	return false;
 }
 
 void USkinnedMeshComponent::DeallocateTransformData()
 {
-	SpaceBases.Empty();
+	SpaceBasesArray[0].Empty();
+	SpaceBasesArray[1].Empty();
 	BoneVisibilityStates.Empty();
 }
 
@@ -1075,9 +1060,9 @@ FQuat USkinnedMeshComponent::GetBoneQuaternion(FName BoneName, EBoneSpaces::Type
 				int32 ParentBoneIndex = MasterBoneMap[BoneIndex];
 				// If ParentBoneIndex is valid, grab matrix from MasterPoseComponent.
 				if(	ParentBoneIndex != INDEX_NONE && 
-					ParentBoneIndex < MasterPoseComponent->SpaceBases.Num())
+					ParentBoneIndex < MasterPoseComponent->GetNumSpaceBases())
 				{
-					BoneTransform = MasterPoseComponent->SpaceBases[ParentBoneIndex];
+					BoneTransform = MasterPoseComponent->GetSpaceBases()[ParentBoneIndex];
 				}
 				else
 				{
@@ -1091,7 +1076,7 @@ FQuat USkinnedMeshComponent::GetBoneQuaternion(FName BoneName, EBoneSpaces::Type
 		}
 		else
 		{
-			BoneTransform = SpaceBases[BoneIndex];
+			BoneTransform = GetSpaceBases()[BoneIndex];
 		}
 	}
 	else
@@ -1122,9 +1107,9 @@ FVector USkinnedMeshComponent::GetBoneLocation(FName BoneName, EBoneSpaces::Type
 				int32 ParentBoneIndex = MasterBoneMap[BoneIndex];
 				// If ParentBoneIndex is valid, grab transform from MasterPoseComponent.
 				if(	ParentBoneIndex != INDEX_NONE && 
-					ParentBoneIndex < MasterPoseComponent->SpaceBases.Num())
+					ParentBoneIndex < MasterPoseComponent->GetNumSpaceBases())
 				{
-					return MasterPoseComponent->SpaceBases[ParentBoneIndex].GetLocation();
+					return MasterPoseComponent->GetSpaceBases()[ParentBoneIndex].GetLocation();
 				}
 			}
 			
@@ -1133,7 +1118,7 @@ FVector USkinnedMeshComponent::GetBoneLocation(FName BoneName, EBoneSpaces::Type
 		}
 		else
 		{
-			return SpaceBases[BoneIndex].GetLocation();
+			return GetSpaceBases()[BoneIndex].GetLocation();
 		}
 	}
 	else if (Space == EBoneSpaces::WorldSpace)
@@ -1261,11 +1246,11 @@ FName USkinnedMeshComponent::FindClosestBone(FVector TestLocation, FVector* Bone
 		float IgnoreScaleSquared = FMath::Square(IgnoreScale);
 		float BestDistSquared = BIG_NUMBER;
 		int32 BestIndex = -1;
-		for (int32 i = 0; i < SpaceBases.Num(); i++)
+		for (int32 i = 0; i < GetNumSpaceBases(); i++)
 		{
-			if (IgnoreScale < 0.f || SpaceBases[i].GetScaledAxis( EAxis::X ).SizeSquared() > IgnoreScaleSquared)
+			if (IgnoreScale < 0.f || GetSpaceBases()[i].GetScaledAxis( EAxis::X ).SizeSquared() > IgnoreScaleSquared)
 			{
-				float DistSquared = (TestLocation - SpaceBases[i].GetLocation()).SizeSquared();
+				float DistSquared = (TestLocation - GetSpaceBases()[i].GetLocation()).SizeSquared();
 				if (DistSquared < BestDistSquared)
 				{
 					BestIndex = i;
@@ -1287,7 +1272,7 @@ FName USkinnedMeshComponent::FindClosestBone(FVector TestLocation, FVector* Bone
 			// transform the bone location into world space
 			if (BoneLocation != NULL)
 			{
-				*BoneLocation = (SpaceBases[BestIndex] * ComponentToWorld).GetLocation();
+				*BoneLocation = (GetSpaceBases()[BestIndex] * ComponentToWorld).GetLocation();
 			}
 			return SkeletalMesh->RefSkeleton.GetBoneName(BestIndex);
 		}
@@ -1398,7 +1383,7 @@ FORCEINLINE FVector USkinnedMeshComponent::GetTypedSkinnedVertexPosition(const F
 				}
 				else
 				{
-					const FMatrix RefToLocal = SkeletalMesh->RefBasesInvMatrix[BoneIndex] * BaseComponent->SpaceBases[BoneIndex].ToMatrixWithScale();
+					const FMatrix RefToLocal = SkeletalMesh->RefBasesInvMatrix[BoneIndex] * BaseComponent->GetSpaceBases()[BoneIndex].ToMatrixWithScale();
 					SkinnedPos += RefToLocal.TransformPosition(VertexBufferGPUSkin.GetVertexPositionFast(SrcSoftVertex)) * Weight;
 				}
 			}
@@ -1423,7 +1408,7 @@ FORCEINLINE FVector USkinnedMeshComponent::GetTypedSkinnedVertexPosition(const F
 		}
 		else
 		{
-			const FMatrix RefToLocal = SkeletalMesh->RefBasesInvMatrix[BoneIndex] * BaseComponent->SpaceBases[BoneIndex].ToMatrixWithScale();
+			const FMatrix RefToLocal = SkeletalMesh->RefBasesInvMatrix[BoneIndex] * BaseComponent->GetSpaceBases()[BoneIndex].ToMatrixWithScale();
 			SkinnedPos = RefToLocal.TransformPosition(VertexBufferGPUSkin.GetVertexPositionFast(SrcRigidVertex));
 		}
 	}
@@ -1483,7 +1468,7 @@ void USkinnedMeshComponent::ComputeSkinnedPositions(TArray<FVector> & OutPositio
 	{
 		for (int32 MatrixIdx = 0; MatrixIdx < RefToLocals.Num(); ++MatrixIdx)
 		{
-			RefToLocals[MatrixIdx] = SkeletalMesh->RefBasesInvMatrix[MatrixIdx] * BaseComponent->SpaceBases[MatrixIdx].ToMatrixWithScale();
+			RefToLocals[MatrixIdx] = SkeletalMesh->RefBasesInvMatrix[MatrixIdx] * BaseComponent->GetSpaceBases()[MatrixIdx].ToMatrixWithScale();
 		}
 	}
 
@@ -1834,6 +1819,19 @@ void USkinnedMeshComponent::AnimUpdateRateSetParams(const bool & bRecentlyRender
 	}
 }
 
+void USkinnedMeshComponent::FlipEditableSpaceBases()
+{
+	if (bDoubleBufferedBlendSpaces)
+	{
+		CurrentReadSpaceBases = CurrentEditableSpaceBases;
+		CurrentEditableSpaceBases = 1 - CurrentEditableSpaceBases;
+	}
+	else
+	{
+		CurrentReadSpaceBases = CurrentEditableSpaceBases = 0;
+	}
+}
+
 void FAnimUpdateRateParameters::Set(class AActor& Owner, const int32& NewUpdateRate, const int32& NewEvaluationRate, const bool & bNewInterpSkippedFrames)
 {
 	UpdateRate = FMath::Max(NewUpdateRate, 1);
