RemovingGRHIShaderPlatform_DEPRECATED use with Instanced geometry. World's feature level is passed down to vertex factory and friends.

diff --git a/Engine/Source/Runtime/Engine/Private/InstancedStaticMesh.cpp b/Engine/Source/Runtime/Engine/Private/InstancedStaticMesh.cpp
--- a/Engine/Source/Runtime/Engine/Private/InstancedStaticMesh.cpp
+++ b/Engine/Source/Runtime/Engine/Private/InstancedStaticMesh.cpp
@@ -124,7 +124,7 @@ class FStaticMeshInstanceBuffer : public FVertexBuffer
 public:
 
 	/** Default constructor. */
-	FStaticMeshInstanceBuffer();
+	FStaticMeshInstanceBuffer(ERHIFeatureLevel::Type InFeatureLevel);
 
 	/** Destructor. */
 	~FStaticMeshInstanceBuffer();
@@ -182,9 +182,10 @@ private:
 };
 
 
-FStaticMeshInstanceBuffer::FStaticMeshInstanceBuffer():
+FStaticMeshInstanceBuffer::FStaticMeshInstanceBuffer(ERHIFeatureLevel::Type InFeatureLevel) : 
 	InstanceData(NULL)
 {
+	SetFeatureLevel(InFeatureLevel);
 }
 
 FStaticMeshInstanceBuffer::~FStaticMeshInstanceBuffer()
@@ -320,7 +321,8 @@ void FStaticMeshInstanceBuffer::AllocateData()
 	// Clear any old VertexData before allocating.
 	CleanUp();
 
-	const bool bInstanced = RHISupportsInstancing(GRHIShaderPlatform_DEPRECATED);
+	check(HasValidFeatureLevel());
+	const bool bInstanced = RHISupportsInstancing(GetFeatureLevelShaderPlatform(GetFeatureLevel()));
 	const bool bNeedsCPUAccess = !bInstanced;
 	InstanceData = new FStaticMeshInstanceData(bNeedsCPUAccess);
 	// Calculate the vertex stride.
@@ -454,7 +456,8 @@ void FInstancedStaticMeshVertexFactory::Copy(const FInstancedStaticMeshVertexFac
 
 void FInstancedStaticMeshVertexFactory::InitRHI()
 {
-	const bool bInstanced = RHISupportsInstancing(GRHIShaderPlatform_DEPRECATED);
+	check(HasValidFeatureLevel());
+	const bool bInstanced = RHISupportsInstancing(GetFeatureLevelShaderPlatform(GetFeatureLevel()));
 
 	// If the vertex buffer containing position is not the same vertex buffer containing the rest of the data,
 	// then initialize PositionStream and PositionDeclaration.
@@ -597,14 +600,17 @@ class FInstancedStaticMeshRenderData
 {
 public:
 
-	FInstancedStaticMeshRenderData(UInstancedStaticMeshComponent* InComponent)
+	FInstancedStaticMeshRenderData(UInstancedStaticMeshComponent* InComponent, ERHIFeatureLevel::Type InFeatureLevel)
 	  : Component(InComponent)
+	  , InstanceBuffer(InFeatureLevel)
 	  , LODModels(Component->StaticMesh->RenderData->LODResources)
+	  , FeatureLevel(InFeatureLevel)
 	{
 		// Allocate the vertex factories for each LOD
 		for( int32 LODIndex=0;LODIndex<LODModels.Num();LODIndex++ )
 		{
-			new(VertexFactories) FInstancedStaticMeshVertexFactory;
+			FInstancedStaticMeshVertexFactory* VertexFactory = new(VertexFactories)FInstancedStaticMeshVertexFactory;
+			VertexFactory->SetFeatureLevel(InFeatureLevel);
 		}
 
 		// Create hit proxies for each instance if the component wants
@@ -691,6 +697,9 @@ public:
 
 	/** Hit proxies for the instances */
 	TArray<TRefCountPtr<HHitProxy> > HitProxies;
+
+	/** Feature level used when creating instance data */
+	ERHIFeatureLevel::Type FeatureLevel;
 };
 
 void FInstancedStaticMeshRenderData::InitStaticMeshVertexFactories(
@@ -698,7 +707,7 @@ void FInstancedStaticMeshRenderData::InitStaticMeshVertexFactories(
 		FInstancedStaticMeshRenderData* InstancedRenderData,
 		UStaticMesh* Parent)
 {
-	const bool bInstanced = RHISupportsInstancing(GRHIShaderPlatform_DEPRECATED);
+	const bool bInstanced = RHISupportsInstancing(GetFeatureLevelShaderPlatform(InstancedRenderData->FeatureLevel));
 
 	for( int32 LODIndex=0;LODIndex<VertexFactories->Num(); LODIndex++ )
 	{
@@ -859,9 +868,9 @@ class FInstancedStaticMeshSceneProxy : public FStaticMeshSceneProxy
 {
 public:
 
-	FInstancedStaticMeshSceneProxy(UInstancedStaticMeshComponent* InComponent)
+	FInstancedStaticMeshSceneProxy(UInstancedStaticMeshComponent* InComponent, ERHIFeatureLevel::Type InFeatureLevel)
 	:	FStaticMeshSceneProxy(InComponent)
-	,	InstancedRenderData(InComponent)
+	, InstancedRenderData(InComponent, InFeatureLevel)
 #if WITH_EDITOR
 	,	bHasSelectedInstances(InComponent->SelectedInstances.Num() > 0)
 #endif
@@ -889,7 +898,7 @@ public:
 
 		check(InstancedRenderData.InstanceBuffer.GetStride() == sizeof(FInstancingUserData::FInstanceStream));
 
-		const bool bInstanced = RHISupportsInstancing(GRHIShaderPlatform_DEPRECATED);
+		const bool bInstanced = RHISupportsInstancing(GetFeatureLevelShaderPlatform(InFeatureLevel));
 
 		// Copy the parameters for LOD - all instances
 		UserData_AllInstances.StartCullDistance = InComponent->InstanceStartCullDistance;
@@ -1127,7 +1136,7 @@ void FInstancedStaticMeshSceneProxy::DrawDynamicElements(FPrimitiveDrawInterface
 
 int32 FInstancedStaticMeshSceneProxy::GetNumMeshBatches() const
 {
-	const bool bInstanced = RHISupportsInstancing(GRHIShaderPlatform_DEPRECATED);
+	const bool bInstanced = RHISupportsInstancing(GetFeatureLevelShaderPlatform(InstancedRenderData.FeatureLevel));
 
 	if (bInstanced)
 	{
@@ -1144,7 +1153,7 @@ int32 FInstancedStaticMeshSceneProxy::GetNumMeshBatches() const
 
 void FInstancedStaticMeshSceneProxy::SetupInstancedMeshBatch(int32 LODIndex, int32 BatchIndex, FMeshBatch& OutMeshBatch) const
 {
-	const bool bInstanced = RHISupportsInstancing(GRHIShaderPlatform_DEPRECATED);
+	const bool bInstanced = RHISupportsInstancing(GetFeatureLevelShaderPlatform(InstancedRenderData.FeatureLevel));
 	OutMeshBatch.VertexFactory = &InstancedRenderData.VertexFactories[LODIndex];
 	const uint32 NumInstances = InstancedRenderData.InstanceBuffer.GetNumInstances();
 	FMeshBatchElement& BatchElement0 = OutMeshBatch.Elements[0];
@@ -1387,7 +1396,7 @@ FPrimitiveSceneProxy* UInstancedStaticMeshComponent::CreateSceneProxy()
 			InstancingRandomSeed = FMath::Rand();
 		}
 
-		return ::new FInstancedStaticMeshSceneProxy(this);
+		return ::new FInstancedStaticMeshSceneProxy(this, GetWorld()->FeatureLevel);
 	}
 	else
 	{
@@ -2151,7 +2160,8 @@ void FInstancedStaticMeshVertexFactoryShaderParameters::SetMesh( FRHICommandList
 		SetShaderValue(RHICmdList, VS, InstancingFadeOutParamsParameter, InstancingFadeOutParams );
 	}
 
-	const bool bInstanced = RHISupportsInstancing(GRHIShaderPlatform_DEPRECATED);
+	auto ShaderPlatform = GetFeatureLevelShaderPlatform(View.GetFeatureLevel());
+	const bool bInstanced = RHISupportsInstancing(ShaderPlatform);
 	if (!bInstanced)
 	{
 		if (CPUInstanceShadowMapBias.IsBound())
diff --git a/Engine/Source/Runtime/Engine/Private/World.cpp b/Engine/Source/Runtime/Engine/Private/World.cpp
--- a/Engine/Source/Runtime/Engine/Private/World.cpp
+++ b/Engine/Source/Runtime/Engine/Private/World.cpp
@@ -5351,7 +5351,6 @@ void UWorld::ChangeFeatureLevel(ERHIFeatureLevel::Type InFeatureLevel)
 			BeginUpdateResourceRHI(*It);
 		}
 
-		// try iterating over all worlds here
 		FeatureLevel = InFeatureLevel;
 
 		if (Scene)
