Fix ConstraintInstance angular rotation offset ignoring the orientation of the object

diff --git a/Engine/Source/Runtime/Engine/Private/PhysicsEngine/PhysicsConstraintComponent.cpp b/Engine/Source/Runtime/Engine/Private/PhysicsEngine/PhysicsConstraintComponent.cpp
--- a/Engine/Source/Runtime/Engine/Private/PhysicsEngine/PhysicsConstraintComponent.cpp
+++ b/Engine/Source/Runtime/Engine/Private/PhysicsEngine/PhysicsConstraintComponent.cpp
@@ -402,9 +402,13 @@ void UPhysicsConstraintComponent::UpdateConstraintFrames()
 	ConstraintInstance.PriAxis1 = A1Transform.InverseTransformVectorNoScale(WPri);
 	ConstraintInstance.SecAxis1 = A1Transform.InverseTransformVectorNoScale(WOrth);
 
+	const FVector RotatedX = ConstraintInstance.AngularRotationOffset.RotateVector(FVector(1,0,0));
+	const FVector RotatedY = ConstraintInstance.AngularRotationOffset.RotateVector(FVector(0,1,0));
+	const FVector WPri2 = ComponentToWorld.TransformVector(RotatedX);
+	const FVector WOrth2 = ComponentToWorld.TransformVector(RotatedY);
 	ConstraintInstance.Pos2 = A2Transform.InverseTransformPosition(WPos);
-	ConstraintInstance.PriAxis2 = A2Transform.InverseTransformVectorNoScale(ConstraintInstance.AngularRotationOffset.RotateVector(WPri));
-	ConstraintInstance.SecAxis2 = A2Transform.InverseTransformVectorNoScale(ConstraintInstance.AngularRotationOffset.RotateVector(WOrth));
+	ConstraintInstance.PriAxis2 = A2Transform.InverseTransformVectorNoScale(WPri2);
+	ConstraintInstance.SecAxis2 = A2Transform.InverseTransformVectorNoScale(WOrth2);
 }
 
 void UPhysicsConstraintComponent::SetConstraintReferenceFrame(EConstraintFrame::Type Frame, const FTransform& RefFrame)
