Allow destructible components to enable/disable gravity (PR-562)

diff --git a/Engine/Source/Runtime/Engine/Classes/Components/DestructibleComponent.h b/Engine/Source/Runtime/Engine/Classes/Components/DestructibleComponent.h
--- a/Engine/Source/Runtime/Engine/Classes/Components/DestructibleComponent.h
+++ b/Engine/Source/Runtime/Engine/Classes/Components/DestructibleComponent.h
@@ -137,6 +137,7 @@ public:
 
 	virtual bool LineTraceComponent( FHitResult& OutHit, const FVector Start, const FVector End, const FCollisionQueryParams& Params ) override;
 	virtual bool SweepComponent( FHitResult& OutHit, const FVector Start, const FVector End, const FCollisionShape& CollisionShape, bool bTraceComplex=false) override;
+	virtual void SetEnableGravity(bool bGravityEnabled) override;
 	// End UPrimitiveComponent interface.
 
 	// Begin SkinnedMeshComponent interface.
diff --git a/Engine/Source/Runtime/Engine/Private/DestructibleComponent.cpp b/Engine/Source/Runtime/Engine/Private/DestructibleComponent.cpp
--- a/Engine/Source/Runtime/Engine/Private/DestructibleComponent.cpp
+++ b/Engine/Source/Runtime/Engine/Private/DestructibleComponent.cpp
@@ -1198,8 +1198,25 @@ void UDestructibleComponent::ResetFakeBodyInstance( FFakeBodyInstanceState& Prev
 	BodyInstance.RigidActorAsync = PrevState.ActorAsync;
 	BodyInstance.InstanceBodyIndex = PrevState.InstanceIndex;
 }
+
 #endif
 
+void UDestructibleComponent::SetEnableGravity(bool bGravityEnabled)
+{
+	Super::SetEnableGravity(bGravityEnabled);
+	
+#if WITH_APEX
+	for (FDestructibleChunkInfo& ChunkInfo : ChunkInfos)
+	 {
+		physx::PxRigidDynamic* Actor = ChunkInfo.Actor;
+		if (Actor)
+		{
+			Actor->setActorFlag(PxActorFlag::eDISABLE_GRAVITY, !bGravityEnabled);
+		}
+	}
+#endif //WITH_APEX
+}
+
 FBodyInstance* UDestructibleComponent::GetBodyInstance( FName BoneName /*= NAME_None*/, bool) const
 {
 #if WITH_APEX
