backup/restoreglobalshadermap now stores and restores shader data for all available featurelevels.
removal of GRHIShaderPlatform_DEPRECATED from RecompileRendererModule.

diff --git a/Engine/Source/Runtime/Engine/Classes/Materials/Material.h b/Engine/Source/Runtime/Engine/Classes/Materials/Material.h
--- a/Engine/Source/Runtime/Engine/Classes/Materials/Material.h
+++ b/Engine/Source/Runtime/Engine/Classes/Materials/Material.h
@@ -1023,7 +1023,7 @@ public:
 	/** 
 	 * Recreates FShaders for FMaterialShaderMap's from the serialized data.  Shader maps may not be complete after this due to changes in the shader keys.
 	 */
-	ENGINE_API static void RestoreMaterialShadersFromMemory(EShaderPlatform ShaderPlatform, const TMap<FMaterialShaderMap*, TScopedPointer<TArray<uint8> > >& ShaderMapToSerializedShaderData);
+	ENGINE_API static void RestoreMaterialShadersFromMemory(const TMap<FMaterialShaderMap*, TScopedPointer<TArray<uint8> > >& ShaderMapToSerializedShaderData);
 
 	/** Builds a map from UMaterialInterface name to the shader maps that are needed for rendering on the given platform. */
 	ENGINE_API static void CompileMaterialsForRemoteRecompile(
diff --git a/Engine/Source/Runtime/Engine/Private/GlobalShader.cpp b/Engine/Source/Runtime/Engine/Private/GlobalShader.cpp
--- a/Engine/Source/Runtime/Engine/Private/GlobalShader.cpp
+++ b/Engine/Source/Runtime/Engine/Private/GlobalShader.cpp
@@ -140,26 +140,41 @@ FGlobalShader::FGlobalShader(const ShaderMetaType::CompiledShaderInitializerType
 {
 }
 
-TArray<uint8>* BackupGlobalShaderMap(EShaderPlatform Platform)
+void BackupGlobalShaderMap(FGlobalShaderBackupData& OutGlobalShaderBackup)
 {
-	TArray<uint8>* ShaderData = new TArray<uint8>();
-	FMemoryWriter Ar(*ShaderData);
-	GGlobalShaderMap[Platform]->SerializeInline(Ar, true, true);
-	GGlobalShaderMap[Platform]->Empty();
+	for (int32 i = (int32)ERHIFeatureLevel::ES2; i < (int32)ERHIFeatureLevel::Num; ++i)
+	{
+		EShaderPlatform ShaderPlatform = GetFeatureLevelShaderPlatform((ERHIFeatureLevel::Type)i);
+		if (ShaderPlatform < EShaderPlatform::SP_NumPlatforms && GGlobalShaderMap[ShaderPlatform] != nullptr)
+		{
+			TArray<uint8>* ShaderData = new TArray<uint8>();
+			FMemoryWriter Ar(*ShaderData);
+			GGlobalShaderMap[ShaderPlatform]->SerializeInline(Ar, true, true);
+			GGlobalShaderMap[ShaderPlatform]->Empty();
+			OutGlobalShaderBackup.FeatureLevelShaderData[i] = ShaderData;
+		}
+	}
 
 	// Remove cached references to global shaders
-	for(TLinkedList<FGlobalBoundShaderStateResource*>::TIterator It(FGlobalBoundShaderStateResource::GetGlobalBoundShaderStateList());It;It.Next())
+	for (TLinkedList<FGlobalBoundShaderStateResource*>::TIterator It(FGlobalBoundShaderStateResource::GetGlobalBoundShaderStateList()); It; It.Next())
 	{
 		BeginUpdateResourceRHI(*It);
 	}
-
-	return ShaderData;
 }
 
-void RestoreGlobalShaderMap(EShaderPlatform Platform, const TArray<uint8>& ShaderData)
+void RestoreGlobalShaderMap(const FGlobalShaderBackupData& GlobalShaderBackup)
 {
-	FMemoryReader Ar(ShaderData);
-	GGlobalShaderMap[Platform]->SerializeInline(Ar, true, true);
+	for (int32 i = (int32)ERHIFeatureLevel::ES2; i < (int32)ERHIFeatureLevel::Num; ++i)
+	{
+		EShaderPlatform ShaderPlatform = GetFeatureLevelShaderPlatform((ERHIFeatureLevel::Type)i);		
+		if (GlobalShaderBackup.FeatureLevelShaderData[i] != nullptr
+			&& ShaderPlatform < EShaderPlatform::SP_NumPlatforms
+			&& GGlobalShaderMap[ShaderPlatform] != nullptr)
+		{
+			FMemoryReader Ar(*GlobalShaderBackup.FeatureLevelShaderData[i]);
+			GGlobalShaderMap[ShaderPlatform]->SerializeInline(Ar, true, true);
+		}
+	}
 }
 
 /**
diff --git a/Engine/Source/Runtime/Engine/Private/Materials/Material.cpp b/Engine/Source/Runtime/Engine/Private/Materials/Material.cpp
--- a/Engine/Source/Runtime/Engine/Private/Materials/Material.cpp
+++ b/Engine/Source/Runtime/Engine/Private/Materials/Material.cpp
@@ -2937,7 +2937,7 @@ void UMaterial::BackupMaterialShadersToMemory(TMap<FMaterialShaderMap*, TScopedP
 	FMaterial::BackupEditorLoadedMaterialShadersToMemory(ShaderMapToSerializedShaderData);
 }
 
-void UMaterial::RestoreMaterialShadersFromMemory(EShaderPlatform ShaderPlatform, const TMap<FMaterialShaderMap*, TScopedPointer<TArray<uint8> > >& ShaderMapToSerializedShaderData)
+void UMaterial::RestoreMaterialShadersFromMemory(const TMap<FMaterialShaderMap*, TScopedPointer<TArray<uint8> > >& ShaderMapToSerializedShaderData)
 {
 	// Process FMaterialShaderMap's referenced by UObjects (UMaterial, UMaterialInstance)
 	for (TObjectIterator<UMaterialInterface> It; It; ++It)
diff --git a/Engine/Source/Runtime/Engine/Private/RendererSupport.cpp b/Engine/Source/Runtime/Engine/Private/RendererSupport.cpp
--- a/Engine/Source/Runtime/Engine/Private/RendererSupport.cpp
+++ b/Engine/Source/Runtime/Engine/Private/RendererSupport.cpp
@@ -19,7 +19,7 @@
 void ClearReferencesToRendererModuleClasses(
 	TMap<UWorld*, bool>& WorldsToUpdate, 
 	TMap<FMaterialShaderMap*, TScopedPointer<TArray<uint8> > >& ShaderMapToSerializedShaderData,
-	TScopedPointer<TArray<uint8> >& GlobalShaderData,
+	FGlobalShaderBackupData& GlobalShaderBackup,
 	TMap<FShaderType*, FString>& ShaderTypeNames,
 	TMap<FVertexFactoryType*, FString>& VertexFactoryTypeNames)
 {
@@ -48,13 +48,13 @@ void ClearReferencesToRendererModuleClasses(
 			World->Scene = NULL;
 		}
 	}
-
-	// Save off shaders by serializing them into memory, and remove all shader map references to FShaders
-	GlobalShaderData = TScopedPointer<TArray<uint8> >(BackupGlobalShaderMap(GRHIShaderPlatform_DEPRECATED));
+		
+	// For each feature level save off its shaders by serializing them into memory, and remove all shader map references to FShaders
+	BackupGlobalShaderMap(GlobalShaderBackup);
 	UMaterial::BackupMaterialShadersToMemory(ShaderMapToSerializedShaderData);
 
 	// Verify no FShaders still in memory
-	for(TLinkedList<FShaderType*>::TIterator It(FShaderType::GetTypeList()); It; It.Next())
+	for (TLinkedList<FShaderType*>::TIterator It(FShaderType::GetTypeList()); It; It.Next())
 	{
 		FShaderType* ShaderType = *It;
 		check(ShaderType->GetNumShaders() == 0);
@@ -116,7 +116,7 @@ void RecompileRendererModule()
 void RestoreReferencesToRendererModuleClasses(
 	const TMap<UWorld*, bool>& WorldsToUpdate, 
 	const TMap<FMaterialShaderMap*, TScopedPointer<TArray<uint8> > >& ShaderMapToSerializedShaderData,
-	const TScopedPointer<TArray<uint8> >& GlobalShaderData,
+	const FGlobalShaderBackupData& GlobalShaderBackup,
 	const TMap<FShaderType*, FString>& ShaderTypeNames,
 	const TMap<FVertexFactoryType*, FString>& VertexFactoryTypeNames)
 {
@@ -145,9 +145,18 @@ void RestoreReferencesToRendererModuleClasses(
 
 	// Restore FShaders from the serialized memory blobs
 	// Shader maps may still not be complete after this due to code changes picked up in the recompile
-	RestoreGlobalShaderMap(GRHIShaderPlatform_DEPRECATED, *GlobalShaderData);
-	UMaterial::RestoreMaterialShadersFromMemory(GRHIShaderPlatform_DEPRECATED, ShaderMapToSerializedShaderData);
-	FMaterialShaderMap::FixupShaderTypes(GRHIShaderPlatform_DEPRECATED, ShaderTypeNames, VertexFactoryTypeNames);
+	RestoreGlobalShaderMap(GlobalShaderBackup);
+	UMaterial::RestoreMaterialShadersFromMemory(ShaderMapToSerializedShaderData);
+
+	for (int32 i = (int32)ERHIFeatureLevel::ES2; i < (int32)ERHIFeatureLevel::Num; ++i)
+	{
+		if (GlobalShaderBackup.FeatureLevelShaderData[i] != nullptr)
+		{
+			EShaderPlatform ShaderPlatform = GetFeatureLevelShaderPlatform((ERHIFeatureLevel::Type)i);
+			check(ShaderPlatform < EShaderPlatform::SP_NumPlatforms);
+			FMaterialShaderMap::FixupShaderTypes(ShaderPlatform, ShaderTypeNames, VertexFactoryTypeNames);
+		}
+	}
 
 	TArray<FShaderType*> OutdatedShaderTypes;
 	TArray<const FVertexFactoryType*> OutdatedFactoryTypes;
@@ -156,7 +165,8 @@ void RestoreReferencesToRendererModuleClasses(
 	// Recompile any missing shaders
 	UMaterialInterface::IterateOverActiveFeatureLevels([&](ERHIFeatureLevel::Type FeatureLevel) 
 	{
-		auto ShaderPlatform = GShaderPlatformForFeatureLevel[FeatureLevel];
+		auto ShaderPlatform = GetFeatureLevelShaderPlatform(FeatureLevel);
+		check(ShaderPlatform < EShaderPlatform::SP_NumPlatforms);
 		BeginRecompileGlobalShaders(OutdatedShaderTypes, ShaderPlatform);
 		UMaterial::UpdateMaterialShaders(OutdatedShaderTypes, OutdatedFactoryTypes, ShaderPlatform);
 	});
@@ -187,11 +197,11 @@ void RecompileRenderer(const TArray<FString>& Args)
 
 		TMap<UWorld*, bool> WorldsToUpdate;
 		TMap<FMaterialShaderMap*, TScopedPointer<TArray<uint8> > > ShaderMapToSerializedShaderData;
-		TScopedPointer<TArray<uint8> > GlobalShaderData;
+		FGlobalShaderBackupData GlobalShaderBackup;
 		TMap<FShaderType*, FString> ShaderTypeNames;
 		TMap<FVertexFactoryType*, FString> VertexFactoryTypeNames;
 
-		ClearReferencesToRendererModuleClasses(WorldsToUpdate, ShaderMapToSerializedShaderData, GlobalShaderData, ShaderTypeNames, VertexFactoryTypeNames);
+		ClearReferencesToRendererModuleClasses(WorldsToUpdate, ShaderMapToSerializedShaderData, GlobalShaderBackup, ShaderTypeNames, VertexFactoryTypeNames);
 
 		EndShutdownTime = FPlatformTime::Seconds();
 		UE_LOG(LogShaders, Warning, TEXT("Shutdown complete %.1fs"),(float)(EndShutdownTime - StartTime));
@@ -201,7 +211,7 @@ void RecompileRenderer(const TArray<FString>& Args)
 		EndRecompileTime = FPlatformTime::Seconds();
 		UE_LOG(LogShaders, Warning, TEXT("Recompile complete %.1fs"),(float)(EndRecompileTime - EndShutdownTime));
 
-		RestoreReferencesToRendererModuleClasses(WorldsToUpdate, ShaderMapToSerializedShaderData, GlobalShaderData, ShaderTypeNames, VertexFactoryTypeNames);
+		RestoreReferencesToRendererModuleClasses(WorldsToUpdate, ShaderMapToSerializedShaderData, GlobalShaderBackup, ShaderTypeNames, VertexFactoryTypeNames);
 	}
 
 #if WITH_EDITOR
diff --git a/Engine/Source/Runtime/Engine/Public/GlobalShader.h b/Engine/Source/Runtime/Engine/Public/GlobalShader.h
--- a/Engine/Source/Runtime/Engine/Public/GlobalShader.h
+++ b/Engine/Source/Runtime/Engine/Public/GlobalShader.h
@@ -150,11 +150,20 @@ public:
 	}
 };
 
+/**
+* Container for Backup/RestoreGlobalShaderMap functions.
+* Includes shader data from any populated feature levels.
+*/
+struct FGlobalShaderBackupData
+{
+	TScopedPointer<TArray<uint8>> FeatureLevelShaderData[ERHIFeatureLevel::Num];
+};
+
 /** Backs up all global shaders to memory through serialization, and removes all references to FShaders from the global shader map. */
-extern ENGINE_API TArray<uint8>* BackupGlobalShaderMap(EShaderPlatform Platform);
+extern ENGINE_API void BackupGlobalShaderMap(FGlobalShaderBackupData& OutGlobalShaderBackup);
 
 /** Recreates shaders in the global shader map from the serialized memory. */
-extern ENGINE_API void RestoreGlobalShaderMap(EShaderPlatform Platform, const TArray<uint8>& ShaderData);
+extern ENGINE_API void RestoreGlobalShaderMap(const FGlobalShaderBackupData& GlobalShaderData);
 
 /**
  * Makes sure all global shaders are loaded and/or compiled for the passed in platform.
