DataTable: fix crash when UDS in unload, minor fixes.

diff --git a/Engine/Source/Runtime/Engine/Classes/Engine/DataTable.h b/Engine/Source/Runtime/Engine/Classes/Engine/DataTable.h
--- a/Engine/Source/Runtime/Engine/Classes/Engine/DataTable.h
+++ b/Engine/Source/Runtime/Engine/Classes/Engine/DataTable.h
@@ -89,7 +89,7 @@ class UDataTable : public UObject
 	/** Returns the column property where PropertyName matches the name of the column property. Returns NULL if no match is found or the match is not a supported table property */
 	ENGINE_API UProperty* FindTableProperty(const FName& PropertyName) const;
 
-	void* FindRowUnchecked(FName RowName, bool MustExist=false) const
+	uint8* FindRowUnchecked(FName RowName, bool MustExist=false) const
 	{
 		if(RowStruct == NULL)
 		{
@@ -110,7 +110,7 @@ class UDataTable : public UObject
 			return NULL;
 		}
 
-		void* RowData = *RowDataPtr;
+		uint8* RowData = *RowDataPtr;
 		check(RowData);
 
 		return RowData;
@@ -119,6 +119,8 @@ class UDataTable : public UObject
 	/** Empty the table info (will not clear RowStruct) */
 	ENGINE_API void EmptyTable();
 
+	ENGINE_API TArray<FName> GetRowNames() const;
+
 #if WITH_EDITOR
 
 private:
diff --git a/Engine/Source/Runtime/Engine/Private/DataTable.cpp b/Engine/Source/Runtime/Engine/Private/DataTable.cpp
--- a/Engine/Source/Runtime/Engine/Private/DataTable.cpp
+++ b/Engine/Source/Runtime/Engine/Private/DataTable.cpp
@@ -151,6 +151,15 @@ void UDataTable::Serialize( FArchive& Ar )
 {
 	Super::Serialize(Ar); // When loading, this should load our RowStruct!	
 
+	if (RowStruct && RowStruct->HasAnyFlags(RF_NeedLoad))
+	{
+		auto RowStructLinker = RowStruct->GetLinker();
+		if (RowStructLinker)
+		{
+			RowStructLinker->Preload(RowStruct);
+		}
+	}
+
 	if(Ar.IsLoading())
 	{
 		LoadStructData(Ar);
@@ -670,4 +679,11 @@ TArray< TArray<FString> > UDataTable::GetTableData() const
 
 }
 
-#endif //WITH_EDITOR
\ No newline at end of file
+#endif //WITH_EDITOR
+
+TArray<FName> UDataTable::GetRowNames() const
+{
+	TArray<FName> Keys;
+	RowMap.GetKeys(Keys);
+	return Keys;
+}
\ No newline at end of file
