Fix build error

diff --git a/Engine/Source/Runtime/Engine/Classes/Components/SkinnedMeshComponent.h b/Engine/Source/Runtime/Engine/Classes/Components/SkinnedMeshComponent.h
--- a/Engine/Source/Runtime/Engine/Classes/Components/SkinnedMeshComponent.h
+++ b/Engine/Source/Runtime/Engine/Classes/Components/SkinnedMeshComponent.h
@@ -579,12 +579,11 @@ public:
 
 		if (bDoubleBufferedBlendSpaces)
 		{
-			CurrentReadSpaceBases = CurrentReadSpaceBases;
 			CurrentEditableSpaceBases = 1 - CurrentReadSpaceBases;
 		}
 		else
 		{
-			CurrentReadSpaceBases = CurrentEditableSpaceBases = CurrentReadSpaceBases;
+			CurrentEditableSpaceBases = CurrentReadSpaceBases;
 		}
 	}
 
