Fix welding code from updating shapes on a root that is being garbage collected (UE-5231,5109)

diff --git a/Engine/Source/Runtime/Engine/Private/PhysicsEngine/BodyInstance.cpp b/Engine/Source/Runtime/Engine/Private/PhysicsEngine/BodyInstance.cpp
--- a/Engine/Source/Runtime/Engine/Private/PhysicsEngine/BodyInstance.cpp
+++ b/Engine/Source/Runtime/Engine/Private/PhysicsEngine/BodyInstance.cpp
@@ -897,6 +897,7 @@ void FBodyInstance::UpdatePhysicsFilterData(bool bForceSimpleAsComplex)
 #endif
 }
 
+
 #if UE_WITH_PHYSICS
 void FBodyInstance::InitBody(UBodySetup* Setup, const FTransform& Transform, UPrimitiveComponent* PrimComp, FPhysScene* InRBScene, PxAggregate* InAggregate)
 {
@@ -1279,6 +1280,7 @@ void FBodyInstance::InitBody(UBodySetup* Setup, const FTransform& Transform, UPr
 		check(FPhysxUserData::Get<FBodyInstance>(PNewActorAsync->userData) == this && FPhysxUserData::Get<FBodyInstance>(PNewActorAsync->userData)->OwnerComponent != NULL);
 	}
 
+
 	// If we added no shapes, generate warning, destroy actor and bail out (don't add to scene).
 	if ((PNewActorSync && PNewActorSync->getNbShapes() == 0) || ((PNewActorAsync && PNewActorAsync->getNbShapes() == 0)))
 	{
@@ -1454,7 +1456,6 @@ TArray<int32> FBodyInstance::AddCollisionNotifyInfo(const FBodyInstance* Body0,
 	return PairNotifyMapping;
 }
 
-
 //helper function for TermBody to avoid code duplication between scenes
 void TermBodyHelper(int32& SceneIndex, PxRigidActor*& PRigidActor, FBodyInstance* BodyInstance)
 {
@@ -1469,6 +1470,8 @@ void TermBodyHelper(int32& SceneIndex, PxRigidActor*& PRigidActor, FBodyInstance
 
 			if (PRigidActor)
 			{
+
+
 				// Let FPhysScene know
 				FPhysScene* PhysScene = FPhysxUserData::Get<FPhysScene>(PScene->userData);
 				if (PhysScene)
@@ -1625,6 +1628,7 @@ void FBodyInstance::UnWeld(FBodyInstance* TheirBI)
 	int32 NumSyncShapes = 0;
 	TArray<PxShape *> PShapes = GetAllShapes(NumSyncShapes);
 
+	bool bShapesChanged = false;
 	bool bNeedsNotification = false;
 
 	PxScene * PSyncScene = RigidActorSync ? RigidActorSync->getScene() : NULL;
@@ -1644,6 +1648,7 @@ void FBodyInstance::UnWeld(FBodyInstance* TheirBI)
 			{
 				PShape->userData = NULL;
 				RigidActorSync->detachShape(*PShape);
+				bShapesChanged = true;
 			}
 		}
 	}
@@ -1659,11 +1664,15 @@ void FBodyInstance::UnWeld(FBodyInstance* TheirBI)
 			{
 				PShape->userData = NULL;
 				RigidActorSync->detachShape(*PShape);
+				bShapesChanged = true;
 			}
 		}
 	}
 
-	PostShapeChange();
+	if (bShapesChanged)
+	{
+		PostShapeChange();
+	}
 #endif
 }
 
diff --git a/Engine/Source/Runtime/Engine/Private/PrimitiveComponentPhysics.cpp b/Engine/Source/Runtime/Engine/Private/PrimitiveComponentPhysics.cpp
--- a/Engine/Source/Runtime/Engine/Private/PrimitiveComponentPhysics.cpp
+++ b/Engine/Source/Runtime/Engine/Private/PrimitiveComponentPhysics.cpp
@@ -652,6 +652,8 @@ void UPrimitiveComponent::WeldTo(USceneComponent* InParent, FName InSocketName /
 
 void UPrimitiveComponent::UnWeldFromParent()
 {
+
+
 	FBodyInstance* NewRootBI = GetBodyInstance(NAME_None, false);
 	UWorld* CurrentWorld = GetWorld();
 	if (NewRootBI == NULL || NewRootBI->bWelded == false || CurrentWorld == nullptr || IsPendingKill())
@@ -666,8 +668,13 @@ void UPrimitiveComponent::UnWeldFromParent()
 	{
 		if (FBodyInstance* RootBI = RootComponent->GetBodyInstance(SocketName, false))
 		{
-			//create new root
-			RootBI->UnWeld(NewRootBI);
+			bool bRootIsBeingDeleted = RootComponent->HasAnyFlags(RF_PendingKill) || RootComponent->HasAnyFlags(RF_Unreachable);
+			if (!bRootIsBeingDeleted)
+			{
+				//create new root
+				RootBI->UnWeld(NewRootBI);	//don't bother fixing up shapes if RootComponent is about to be deleted
+			}
+			
 			NewRootBI->bWelded = false;
 			NewRootBI->WeldParent = NULL;
 
@@ -692,7 +699,11 @@ void UPrimitiveComponent::UnWeldFromParent()
 				FBodyInstance* ChildBI = ChildrenBodies[ChildIdx];
 				if (ChildBI != NewRootBI)
 				{
-					RootBI->UnWeld(NewRootBI);
+					if (!bRootIsBeingDeleted)
+					{
+						RootBI->UnWeld(ChildBI);
+					}
+					
 					if (bHasBodySetup)
 					{
 						NewRootBI->Weld(ChildBI, ChildBI->OwnerComponent->GetSocketTransform(ChildrenLabels[ChildIdx]));
