#vr Removing HMD funcitonality in the editor until we add an explicit mode for it

diff --git a/Engine/Source/Runtime/Engine/Private/UnrealEngine.cpp b/Engine/Source/Runtime/Engine/Private/UnrealEngine.cpp
--- a/Engine/Source/Runtime/Engine/Private/UnrealEngine.cpp
+++ b/Engine/Source/Runtime/Engine/Private/UnrealEngine.cpp
@@ -1727,23 +1727,26 @@ public:
 
 bool UEngine::InitializeHMDDevice()
 {
-	if (FParse::Param(FCommandLine::Get(),TEXT("emulatestereo")))
+	if (!GIsEditor)
 	{
-		TSharedPtr<FFakeStereoRenderingDevice> FakeStereoDevice(new FFakeStereoRenderingDevice());
-		StereoRenderingDevice = FakeStereoDevice;
-	}
-	// No reason to connect an HMD on a dedicated server.  Also fixes dedicated servers stealing the oculus connection.
-	else if(!HMDDevice.IsValid() && !FParse::Param(FCommandLine::Get(),TEXT("nohmd")) && !IsRunningDedicatedServer())
-	{
-		// Get a list of plugins that implement this feature
-		TArray<IHeadMountedDisplayModule*> HMDImplementations = IModularFeatures::Get().GetModularFeatureImplementations<IHeadMountedDisplayModule>( IHeadMountedDisplayModule::GetModularFeatureName() );
-		HMDImplementations.Sort(FHMDPluginSorter());
-		for( auto HMDModuleIt = HMDImplementations.CreateIterator(); HMDModuleIt && !HMDDevice.IsValid(); ++HMDModuleIt )
+		if (FParse::Param(FCommandLine::Get(), TEXT("emulatestereo")))
+		{
+			TSharedPtr<FFakeStereoRenderingDevice> FakeStereoDevice(new FFakeStereoRenderingDevice());
+			StereoRenderingDevice = FakeStereoDevice;
+		}
+		// No reason to connect an HMD on a dedicated server.  Also fixes dedicated servers stealing the oculus connection.
+		else if (!HMDDevice.IsValid() && !FParse::Param(FCommandLine::Get(), TEXT("nohmd")) && !IsRunningDedicatedServer())
 		{
-			HMDDevice = (*HMDModuleIt)->CreateHeadMountedDisplay();
-			if( HMDDevice.IsValid() )
+			// Get a list of plugins that implement this feature
+			TArray<IHeadMountedDisplayModule*> HMDImplementations = IModularFeatures::Get().GetModularFeatureImplementations<IHeadMountedDisplayModule>(IHeadMountedDisplayModule::GetModularFeatureName());
+			HMDImplementations.Sort(FHMDPluginSorter());
+			for (auto HMDModuleIt = HMDImplementations.CreateIterator(); HMDModuleIt && !HMDDevice.IsValid(); ++HMDModuleIt)
 			{
-				StereoRenderingDevice = HMDDevice;
+				HMDDevice = (*HMDModuleIt)->CreateHeadMountedDisplay();
+				if (HMDDevice.IsValid())
+				{
+					StereoRenderingDevice = HMDDevice;
+				}
 			}
 		}
 	}
