UberGraphPersistentFrame fix for cyclic dependency

diff --git a/Engine/Source/Editor/KismetCompiler/Private/KismetCompiler.cpp b/Engine/Source/Editor/KismetCompiler/Private/KismetCompiler.cpp
--- a/Engine/Source/Editor/KismetCompiler/Private/KismetCompiler.cpp
+++ b/Engine/Source/Editor/KismetCompiler/Private/KismetCompiler.cpp
@@ -3395,6 +3395,12 @@ void FKismetCompilerContext::Compile()
 		}
 	}
 
+	// It's necessary to tell if UberGraphFunction is ready to create frame.
+	if (NewClass->UberGraphFunction)
+	{
+		NewClass->UberGraphFunction->SetFlags(RF_LoadCompleted);
+	}
+
 	{ BP_SCOPED_COMPILER_EVENT_STAT(EKismetCompilerStats_FinalizationWork);
 
 		// Set any final flags and seal the class, build a CDO, etc...
diff --git a/Engine/Source/Runtime/Engine/Private/BlueprintGeneratedClass.cpp b/Engine/Source/Runtime/Engine/Private/BlueprintGeneratedClass.cpp
--- a/Engine/Source/Runtime/Engine/Private/BlueprintGeneratedClass.cpp
+++ b/Engine/Source/Runtime/Engine/Private/BlueprintGeneratedClass.cpp
@@ -411,7 +411,8 @@ uint8* UBlueprintGeneratedClass::GetPersistentUberGraphFrame(UObject* Obj, UFunc
 		if (UberGraphFunction == FuncToCheck)
 		{
 			auto PointerToUberGraphFrame = UberGraphFramePointerProperty->ContainerPtrToValuePtr<FPointerToUberGraphFrame>(Obj);
-			checkSlow(PointerToUberGraphFrame && PointerToUberGraphFrame->RawPointer);
+			checkSlow(PointerToUberGraphFrame);
+			ensure(PointerToUberGraphFrame->RawPointer);
 			return PointerToUberGraphFrame->RawPointer;
 		}
 	}
@@ -425,11 +426,21 @@ void UBlueprintGeneratedClass::CreatePersistentUberGraphFrame(UObject* Obj) cons
 	checkSlow(!UberGraphFramePointerProperty == !UberGraphFunction);
 	if (Obj && UsePersistentUberGraphFrame() && UberGraphFramePointerProperty && UberGraphFunction)
 	{
-		auto FrameMemory = (uint8*)FMemory::Malloc(UberGraphFunction->GetStructureSize());
-		FMemory::Memzero(FrameMemory, UberGraphFunction->GetStructureSize());
-		for (UProperty* Property = UberGraphFunction->PropertyLink; Property; Property = Property->PropertyLinkNext)
+		uint8* FrameMemory = NULL;
+		const bool bUberGraphFunctionIsReady = !WITH_EDITOR // no cyclic dependency problems
+			|| UberGraphFunction->HasAllFlags(RF_LoadCompleted); // is fully loaded
+		if (bUberGraphFunctionIsReady)
 		{
-			Property->InitializeValue_InContainer(FrameMemory);
+			FrameMemory = (uint8*)FMemory::Malloc(UberGraphFunction->GetStructureSize());
+			FMemory::Memzero(FrameMemory, UberGraphFunction->GetStructureSize());
+			for (UProperty* Property = UberGraphFunction->PropertyLink; Property; Property = Property->PropertyLinkNext)
+			{
+				Property->InitializeValue_InContainer(FrameMemory);
+			}
+		}
+		else
+		{
+			UE_LOG(LogBlueprint, Log, TEXT("Function '%s' is not ready to create frame."), *GetPathNameSafe(UberGraphFunction));
 		}
 		
 		auto PointerToUberGraphFrame = UberGraphFramePointerProperty->ContainerPtrToValuePtr<FPointerToUberGraphFrame>(Obj);
@@ -448,15 +459,21 @@ void UBlueprintGeneratedClass::DestroyPersistentUberGraphFrame(UObject* Obj) con
 	if (Obj && UsePersistentUberGraphFrame() && UberGraphFramePointerProperty && UberGraphFunction)
 	{
 		auto PointerToUberGraphFrame = UberGraphFramePointerProperty->ContainerPtrToValuePtr<FPointerToUberGraphFrame>(Obj);
-		checkSlow(PointerToUberGraphFrame && PointerToUberGraphFrame->RawPointer);
+		checkSlow(PointerToUberGraphFrame);
 		auto FrameMemory = PointerToUberGraphFrame->RawPointer;
 		PointerToUberGraphFrame->RawPointer = NULL;
-
-		for (UProperty* Property = UberGraphFunction->PropertyLink; Property; Property = Property->PropertyLinkNext)
+		if (FrameMemory)
+		{
+			for (UProperty* Property = UberGraphFunction->PropertyLink; Property; Property = Property->PropertyLinkNext)
+			{
+				Property->DestroyValue_InContainer(FrameMemory);
+			}
+			FMemory::Free(FrameMemory);
+		}
+		else
 		{
-			Property->DestroyValue_InContainer(FrameMemory);
+			UE_LOG(LogBlueprint, Log, TEXT("Object '%s' had no Uber Graph Persistent Frame"), *GetPathNameSafe(Obj));
 		}
-		FMemory::Free(FrameMemory);
 	}
 
 	auto ParentClass = GetSuperClass();
