Making sure UWorld::RemovePIEPrefix does not return uninitialized value.

diff --git a/Engine/Source/Runtime/Engine/Private/World.cpp b/Engine/Source/Runtime/Engine/Private/World.cpp
--- a/Engine/Source/Runtime/Engine/Private/World.cpp
+++ b/Engine/Source/Runtime/Engine/Private/World.cpp
@@ -5009,6 +5009,10 @@ FString UWorld::RemovePIEPrefix(const FString &Source)
 		
 		FixedName += RightS;
 	}
+	else
+	{
+		FixedName = Source;
+	}
 
 	return FixedName;
 }
