Fixed UE-5108 "StrategyGame: fails to launch on or package"

diff --git a/Engine/Source/Runtime/Engine/Public/VisualLogger/VisualLogger.inl b/Engine/Source/Runtime/Engine/Public/VisualLogger/VisualLogger.inl
--- a/Engine/Source/Runtime/Engine/Public/VisualLogger/VisualLogger.inl
+++ b/Engine/Source/Runtime/Engine/Public/VisualLogger/VisualLogger.inl
@@ -1,6 +1,34 @@
 // Copyright 1998-2014 Epic Games, Inc. All Rights Reserved.
 
-#include "Runtime/Core/Private/Misc/VarargsHelper.h"
+// Unfortunately needs to be a #define since it uses GET_VARARGS_RESULT which uses the va_list stuff which operates on the
+// current function, so we can't easily call a function
+#define COLLAPSED_LOGF(SerializeFunc) \
+	int32	BufferSize	= 1024; \
+	TCHAR*	Buffer		= NULL; \
+	int32	Result		= -1; \
+	/* allocate some stack space to use on the first pass, which matches most strings */ \
+	TCHAR	StackBuffer[512]; \
+	TCHAR*	AllocatedBuffer = NULL; \
+	\
+	/* first, try using the stack buffer */ \
+	Buffer = StackBuffer; \
+	GET_VARARGS_RESULT( Buffer, ARRAY_COUNT(StackBuffer), ARRAY_COUNT(StackBuffer) - 1, Fmt, Fmt, Result ); \
+	\
+	/* if that fails, then use heap allocation to make enough space */ \
+	while(Result == -1) \
+	{ \
+		FMemory::SystemFree(AllocatedBuffer); \
+		/* We need to use malloc here directly as GMalloc might not be safe. */ \
+		Buffer = AllocatedBuffer = (TCHAR*) FMemory::SystemMalloc( BufferSize * sizeof(TCHAR) ); \
+		GET_VARARGS_RESULT( Buffer, BufferSize, BufferSize-1, Fmt, Fmt, Result ); \
+		BufferSize *= 2; \
+	}; \
+	Buffer[Result] = 0; \
+	; \
+	\
+	SerializeFunc; \
+	FMemory::SystemFree(AllocatedBuffer);
+
 
 FORCEINLINE
 bool CheckVisualLogInputInternal(const class UObject* Object, const struct FLogCategoryBase& Category, ELogVerbosity::Type Verbosity, UWorld **World, FVisualLogEntry **CurrentEntry)
@@ -130,7 +158,7 @@ VARARG_BODY(void, FVisualLogger::CategorizedLogf, const TCHAR*, VARARG_EXTRA(con
 		return;
 	}
 
-	GROWABLE_LOGF(
+	COLLAPSED_LOGF(
 		CurrentEntry->AddText(Buffer, Category.GetCategoryName());
 	);
 }
@@ -146,7 +174,7 @@ VARARG_BODY(void, FVisualLogger::GeometryShapeLogf, const TCHAR*, VARARG_EXTRA(c
 		return;
 	}
 
-	GROWABLE_LOGF(
+	COLLAPSED_LOGF(
 		CurrentEntry->AddElement(Start, End, Category.GetCategoryName(), Color, Buffer);
 	);
 }
@@ -162,7 +190,7 @@ VARARG_BODY(void, FVisualLogger::GeometryShapeLogf, const TCHAR*, VARARG_EXTRA(c
 		return;
 	}
 
-	GROWABLE_LOGF(
+	COLLAPSED_LOGF(
 		CurrentEntry->AddElement(Location, Category.GetCategoryName(), Color, Buffer, Radius);
 	);
 }
@@ -178,7 +206,7 @@ VARARG_BODY(void, FVisualLogger::GeometryShapeLogf, const TCHAR*, VARARG_EXTRA(c
 		return;
 	}
 
-	GROWABLE_LOGF(
+	COLLAPSED_LOGF(
 		CurrentEntry->AddElement(Box, Category.GetCategoryName(), Color, Buffer);
 	);
 }
@@ -194,7 +222,7 @@ VARARG_BODY(void, FVisualLogger::GeometryShapeLogf, const TCHAR*, VARARG_EXTRA(c
 		return;
 	}
 
-	GROWABLE_LOGF(
+	COLLAPSED_LOGF(
 		CurrentEntry->AddElement(Orgin, Direction, Length, Angle, Angle, Category.GetCategoryName(), Color, Buffer);
 	);
 }
@@ -210,7 +238,7 @@ VARARG_BODY(void, FVisualLogger::GeometryShapeLogf, const TCHAR*, VARARG_EXTRA(c
 		return;
 	}
 
-	GROWABLE_LOGF(
+	COLLAPSED_LOGF(
 		CurrentEntry->AddElement(Start, End, Radius, Category.GetCategoryName(), Color, Buffer);
 	);
 }
@@ -226,7 +254,7 @@ VARARG_BODY(void, FVisualLogger::GeometryShapeLogf, const TCHAR*, VARARG_EXTRA(c
 		return;
 	}
 
-	GROWABLE_LOGF(
+	COLLAPSED_LOGF(
 		CurrentEntry->AddElement(Center, HalfHeight, Radius, Rotation, Category.GetCategoryName(), Color, Buffer);
 	);
 }
@@ -242,11 +270,13 @@ VARARG_BODY(void, FVisualLogger::HistogramDataLogf, const TCHAR*, VARARG_EXTRA(c
 		return;
 	}
 
-	GROWABLE_LOGF(
+	COLLAPSED_LOGF(
 		CurrentEntry->AddHistogramData(Data, Category.GetCategoryName(), GraphName, DataName);
 	);
 }
 
+#undef COLLAPSED_LOGF
+
 FORCEINLINE
 void FVisualLogger::EventLog(const class UObject* Object, const FName EventTag1, const FVisualLogEventBase& Event1, const FVisualLogEventBase& Event2, const FVisualLogEventBase& Event3, const FVisualLogEventBase& Event4, const FVisualLogEventBase& Event5, const FVisualLogEventBase& Event6)
 {
