Remove UPROPERTY(AdvancedDisplay) from PlayerStateClass
[https://github.com/EpicGames/UnrealEngine/pull/556] (slonopotamus)

diff --git a/Engine/Source/Runtime/Engine/Classes/GameFramework/GameMode.h b/Engine/Source/Runtime/Engine/Classes/GameFramework/GameMode.h
--- a/Engine/Source/Runtime/Engine/Classes/GameFramework/GameMode.h
+++ b/Engine/Source/Runtime/Engine/Classes/GameFramework/GameMode.h
@@ -218,7 +218,7 @@ public:
 	TSubclassOf<class ASpectatorPawn> SpectatorClass;
 
 	/** A PlayerState of this class will be associated with every player to replicate relevant player information to all clients. */
-	UPROPERTY(EditAnywhere, noclear, BlueprintReadWrite, Category=Classes, AdvancedDisplay)
+	UPROPERTY(EditAnywhere, noclear, BlueprintReadWrite, Category=Classes)
 	TSubclassOf<class APlayerState> PlayerStateClass;
 
 	/** Class of GameState associated with this GameMode. */
