Remove GameMode.h from EnginePrivate.h

diff --git a/Engine/Source/Runtime/Engine/Private/Controller.cpp b/Engine/Source/Runtime/Engine/Private/Controller.cpp
--- a/Engine/Source/Runtime/Engine/Private/Controller.cpp
+++ b/Engine/Source/Runtime/Engine/Private/Controller.cpp
@@ -10,6 +10,7 @@
 #include "NetworkingDistanceConstants.h"
 #include "VisualLogger/VisualLogger.h"
 #include "GameFramework/Character.h"
+#include "GameFramework/GameMode.h"
 #include "GameFramework/PawnMovementComponent.h"
 #include "GameFramework/CharacterMovementComponent.h"
 
diff --git a/Engine/Source/Runtime/Engine/Private/EnginePrivate.h b/Engine/Source/Runtime/Engine/Private/EnginePrivate.h
--- a/Engine/Source/Runtime/Engine/Private/EnginePrivate.h
+++ b/Engine/Source/Runtime/Engine/Private/EnginePrivate.h
@@ -25,7 +25,6 @@ EnginePrivate.h: Unreal engine private header file.
 #include "AI/Navigation/NavigationSystem.h"
 #include "Camera/PlayerCameraManager.h"
 #include "GameFramework/Info.h"
-#include "GameFramework/GameMode.h"
 #include "Components/InputComponent.h"
 #include "GameFramework/PlayerController.h"
 #include "Engine/Canvas.h"
diff --git a/Engine/Source/Runtime/Engine/Private/GameEngine.cpp b/Engine/Source/Runtime/Engine/Private/GameEngine.cpp
--- a/Engine/Source/Runtime/Engine/Private/GameEngine.cpp
+++ b/Engine/Source/Runtime/Engine/Private/GameEngine.cpp
@@ -35,6 +35,7 @@
 #include "Components/ReflectionCaptureComponent.h"
 #include "Engine/GameEngine.h"
 #include "GameFramework/GameUserSettings.h"
+#include "GameFramework/GameMode.h"
 
 DEFINE_LOG_CATEGORY_STATIC(LogEngine, Log, All);
 
diff --git a/Engine/Source/Runtime/Engine/Private/GameMode.cpp b/Engine/Source/Runtime/Engine/Private/GameMode.cpp
--- a/Engine/Source/Runtime/Engine/Private/GameMode.cpp
+++ b/Engine/Source/Runtime/Engine/Private/GameMode.cpp
@@ -18,6 +18,7 @@
 #include "GameFramework/PlayerState.h"
 #include "GameFramework/GameSession.h"
 #include "GameFramework/CheatManager.h"
+#include "GameFramework/GameMode.h"
 #include "Engine/ChildConnection.h"
 
 DEFINE_LOG_CATEGORY_STATIC(LogGameMode, Log, All);
@@ -1465,7 +1466,7 @@ void AGameMode::RemovePlayerStart(APlayerStart* RemovedPlayerStart)
 	PlayerStarts.Remove(RemovedPlayerStart);
 }
 
-AActor* AGameMode::ChoosePlayerStart( AController* Player )
+AActor* AGameMode::ChoosePlayerStart_Implementation( AController* Player )
 {
 	// Choose a player start
 	APlayerStart* FoundPlayerStart = NULL;
diff --git a/Engine/Source/Runtime/Engine/Private/GameNetworkManager.cpp b/Engine/Source/Runtime/Engine/Private/GameNetworkManager.cpp
--- a/Engine/Source/Runtime/Engine/Private/GameNetworkManager.cpp
+++ b/Engine/Source/Runtime/Engine/Private/GameNetworkManager.cpp
@@ -6,6 +6,7 @@
 
 #include "EnginePrivate.h"
 #include "GameFramework/GameNetworkManager.h"
+#include "GameFramework/GameMode.h"
 
 DEFINE_LOG_CATEGORY_STATIC(LogGameNetworkManager, Log, All);
 
diff --git a/Engine/Source/Runtime/Engine/Private/GameSession.cpp b/Engine/Source/Runtime/Engine/Private/GameSession.cpp
--- a/Engine/Source/Runtime/Engine/Private/GameSession.cpp
+++ b/Engine/Source/Runtime/Engine/Private/GameSession.cpp
@@ -10,6 +10,7 @@
 #include "Kismet/GameplayStatics.h"
 #include "GameFramework/PlayerState.h"
 #include "GameFramework/GameSession.h"
+#include "GameFramework/GameMode.h"
 
 DEFINE_LOG_CATEGORY_STATIC(LogGameSession, Log, All);
 
diff --git a/Engine/Source/Runtime/Engine/Private/GameState.cpp b/Engine/Source/Runtime/Engine/Private/GameState.cpp
--- a/Engine/Source/Runtime/Engine/Private/GameState.cpp
+++ b/Engine/Source/Runtime/Engine/Private/GameState.cpp
@@ -8,6 +8,7 @@
 #include "Net/UnrealNetwork.h"
 #include "GameFramework/GameState.h"
 #include "GameFramework/PlayerState.h"
+#include "GameFramework/GameMode.h"
 
 DEFINE_LOG_CATEGORY_STATIC(LogGameState, Log, All);
 
diff --git a/Engine/Source/Runtime/Engine/Private/GameplayStatics.cpp b/Engine/Source/Runtime/Engine/Private/GameplayStatics.cpp
--- a/Engine/Source/Runtime/Engine/Private/GameplayStatics.cpp
+++ b/Engine/Source/Runtime/Engine/Private/GameplayStatics.cpp
@@ -13,6 +13,7 @@
 #include "Sound/SoundBase.h"
 #include "Sound/SoundCue.h"
 #include "Engine/GameInstance.h"
+#include "GameFramework/GameMode.h"
 #include "GameFramework/SaveGame.h"
 #include "GameFramework/DamageType.h"
 #include "PhysicalMaterials/PhysicalMaterial.h"
diff --git a/Engine/Source/Runtime/Engine/Private/HUD.cpp b/Engine/Source/Runtime/Engine/Private/HUD.cpp
--- a/Engine/Source/Runtime/Engine/Private/HUD.cpp
+++ b/Engine/Source/Runtime/Engine/Private/HUD.cpp
@@ -7,6 +7,7 @@
 #include "EnginePrivate.h"
 #include "Components/LineBatchComponent.h"
 #include "GameFramework/HUD.h"
+#include "GameFramework/GameMode.h"
 #include "Net/UnrealNetwork.h"
 #include "MessageLog.h"
 #include "UObjectToken.h"
diff --git a/Engine/Source/Runtime/Engine/Private/LevelActor.cpp b/Engine/Source/Runtime/Engine/Private/LevelActor.cpp
--- a/Engine/Source/Runtime/Engine/Private/LevelActor.cpp
+++ b/Engine/Source/Runtime/Engine/Private/LevelActor.cpp
@@ -12,6 +12,7 @@
 #endif // WITH_PHYSX
 #include "Components/SphereComponent.h"
 #include "Components/CapsuleComponent.h"
+#include "GameFramework/GameMode.h"
 
 // CVars
 static TAutoConsoleVariable<float> CVarEncroachEpsilon(
diff --git a/Engine/Source/Runtime/Engine/Private/LocalPlayer.cpp b/Engine/Source/Runtime/Engine/Private/LocalPlayer.cpp
--- a/Engine/Source/Runtime/Engine/Private/LocalPlayer.cpp
+++ b/Engine/Source/Runtime/Engine/Private/LocalPlayer.cpp
@@ -20,6 +20,7 @@
 #include "DataChannel.h"
 #include "GameFramework/OnlineSession.h"
 #include "GameFramework/PlayerInput.h"
+#include "GameFramework/GameMode.h"
 
 DEFINE_LOG_CATEGORY(LogPlayerManagement);
 DEFINE_LOG_CATEGORY_STATIC(LogEngine, Log, All);
diff --git a/Engine/Source/Runtime/Engine/Private/NetConnection.cpp b/Engine/Source/Runtime/Engine/Private/NetConnection.cpp
--- a/Engine/Source/Runtime/Engine/Private/NetConnection.cpp
+++ b/Engine/Source/Runtime/Engine/Private/NetConnection.cpp
@@ -11,6 +11,7 @@
 #include "Engine/ActorChannel.h"
 #include "DataChannel.h"
 #include "Engine/PackageMapClient.h"
+#include "GameFramework/GameMode.h"
 
 /*-----------------------------------------------------------------------------
 	UNetConnection implementation.
diff --git a/Engine/Source/Runtime/Engine/Private/NetworkDriver.cpp b/Engine/Source/Runtime/Engine/Private/NetworkDriver.cpp
--- a/Engine/Source/Runtime/Engine/Private/NetworkDriver.cpp
+++ b/Engine/Source/Runtime/Engine/Private/NetworkDriver.cpp
@@ -17,6 +17,7 @@
 #include "DataChannel.h"
 #include "Engine/PackageMapClient.h"
 #include "GameFramework/PlayerState.h"
+#include "GameFramework/GameMode.h"
 
 // Default net driver stats
 DEFINE_STAT(STAT_Ping);
diff --git a/Engine/Source/Runtime/Engine/Private/Player.cpp b/Engine/Source/Runtime/Engine/Private/Player.cpp
--- a/Engine/Source/Runtime/Engine/Private/Player.cpp
+++ b/Engine/Source/Runtime/Engine/Private/Player.cpp
@@ -18,6 +18,7 @@
 #include "GameFramework/CheatManager.h"
 #include "GameFramework/Pawn.h"
 #include "GameFramework/GameState.h"
+#include "GameFramework/GameMode.h"
 
 //////////////////////////////////////////////////////////////////////////
 // UPlayer
diff --git a/Engine/Source/Runtime/Engine/Private/PlayerController.cpp b/Engine/Source/Runtime/Engine/Private/PlayerController.cpp
--- a/Engine/Source/Runtime/Engine/Private/PlayerController.cpp
+++ b/Engine/Source/Runtime/Engine/Private/PlayerController.cpp
@@ -34,6 +34,7 @@
 #include "GameFramework/InputSettings.h"
 #include "GameFramework/PlayerState.h"
 #include "GameFramework/GameState.h"
+#include "GameFramework/GameMode.h"
 #include "Engine/ChildConnection.h"
 
 DEFINE_LOG_CATEGORY(LogPlayerController);
diff --git a/Engine/Source/Runtime/Engine/Private/PlayerStart.cpp b/Engine/Source/Runtime/Engine/Private/PlayerStart.cpp
--- a/Engine/Source/Runtime/Engine/Private/PlayerStart.cpp
+++ b/Engine/Source/Runtime/Engine/Private/PlayerStart.cpp
@@ -3,6 +3,7 @@
 #include "EnginePrivate.h"
 #include "Components/CapsuleComponent.h"
 #include "GameFramework/PlayerStart.h"
+#include "GameFramework/GameMode.h"
 
 APlayerStart::APlayerStart(const FObjectInitializer& ObjectInitializer)
 	: Super(ObjectInitializer)
diff --git a/Engine/Source/Runtime/Engine/Private/UnrealEngine.cpp b/Engine/Source/Runtime/Engine/Private/UnrealEngine.cpp
--- a/Engine/Source/Runtime/Engine/Private/UnrealEngine.cpp
+++ b/Engine/Source/Runtime/Engine/Private/UnrealEngine.cpp
@@ -48,6 +48,7 @@
 #include "Animation/SkeletalMeshActor.h"
 #include "GameFramework/HUD.h"
 #include "GameFramework/Character.h"
+#include "GameFramework/GameMode.h"
 #include "Engine/LevelStreamingVolume.h"
 #include "Engine/LevelScriptActor.h"
 #include "Vehicles/TireType.h"
diff --git a/Engine/Source/Runtime/Engine/Private/World.cpp b/Engine/Source/Runtime/Engine/Private/World.cpp
--- a/Engine/Source/Runtime/Engine/Private/World.cpp
+++ b/Engine/Source/Runtime/Engine/Private/World.cpp
@@ -37,6 +37,7 @@
 #include "Camera/CameraActor.h"
 #include "Engine/DemoNetDriver.h"
 #include "Layers/Layer.h"
+#include "GameFramework/GameMode.h"
 #include "GameFramework/GameState.h"
 #include "GameFramework/PlayerState.h"
 
