Fix error in CL# 2358443

diff --git a/Engine/Source/Runtime/Engine/Private/GameMode.cpp b/Engine/Source/Runtime/Engine/Private/GameMode.cpp
--- a/Engine/Source/Runtime/Engine/Private/GameMode.cpp
+++ b/Engine/Source/Runtime/Engine/Private/GameMode.cpp
@@ -1466,7 +1466,7 @@ void AGameMode::RemovePlayerStart(APlayerStart* RemovedPlayerStart)
 	PlayerStarts.Remove(RemovedPlayerStart);
 }
 
-AActor* AGameMode::ChoosePlayerStart_Implementation( AController* Player )
+AActor* AGameMode::ChoosePlayerStart( AController* Player )
 {
 	// Choose a player start
 	APlayerStart* FoundPlayerStart = NULL;
