Move function body out of header

diff --git a/Engine/Source/Runtime/Engine/Classes/Components/SkinnedMeshComponent.h b/Engine/Source/Runtime/Engine/Classes/Components/SkinnedMeshComponent.h
--- a/Engine/Source/Runtime/Engine/Classes/Components/SkinnedMeshComponent.h
+++ b/Engine/Source/Runtime/Engine/Classes/Components/SkinnedMeshComponent.h
@@ -573,19 +573,7 @@ public:
 	/** Flip the editable space base buffer */
 	void FlipEditableSpaceBases();
 
-	void SetSpaceBaseDoubleBuffering(bool bInDoubleBufferedBlendSpaces)
-	{
-		bDoubleBufferedBlendSpaces = bInDoubleBufferedBlendSpaces;
-
-		if (bDoubleBufferedBlendSpaces)
-		{
-			CurrentEditableSpaceBases = 1 - CurrentReadSpaceBases;
-		}
-		else
-		{
-			CurrentEditableSpaceBases = CurrentReadSpaceBases;
-		}
-	}
+	void SetSpaceBaseDoubleBuffering(bool bInDoubleBufferedBlendSpaces);
 
 protected:
 	/**
diff --git a/Engine/Source/Runtime/Engine/Private/SkinnedMeshComponent.cpp b/Engine/Source/Runtime/Engine/Private/SkinnedMeshComponent.cpp
--- a/Engine/Source/Runtime/Engine/Private/SkinnedMeshComponent.cpp
+++ b/Engine/Source/Runtime/Engine/Private/SkinnedMeshComponent.cpp
@@ -1832,6 +1832,20 @@ void USkinnedMeshComponent::FlipEditableSpaceBases()
 	}
 }
 
+void USkinnedMeshComponent::SetSpaceBaseDoubleBuffering(bool bInDoubleBufferedBlendSpaces)
+{
+	bDoubleBufferedBlendSpaces = bInDoubleBufferedBlendSpaces;
+
+	if (bDoubleBufferedBlendSpaces)
+	{
+		CurrentEditableSpaceBases = 1 - CurrentReadSpaceBases;
+	}
+	else
+	{
+		CurrentEditableSpaceBases = CurrentReadSpaceBases;
+	}
+}
+
 void FAnimUpdateRateParameters::Set(class AActor& Owner, const int32& NewUpdateRate, const int32& NewEvaluationRate, const bool & bNewInterpSkippedFrames)
 {
 	UpdateRate = FMath::Max(NewUpdateRate, 1);
