Fix "Ignore Root Motion" mode not properly extracting root motion from root bone

diff --git a/Engine/Source/Runtime/Engine/Private/Animation/AnimInstance.cpp b/Engine/Source/Runtime/Engine/Private/Animation/AnimInstance.cpp
--- a/Engine/Source/Runtime/Engine/Private/Animation/AnimInstance.cpp
+++ b/Engine/Source/Runtime/Engine/Private/Animation/AnimInstance.cpp
@@ -135,7 +135,7 @@ void UAnimInstance::SequenceEvaluatePose(UAnimSequenceBase* Sequence, FA2Pose& P
 	USkeletalMeshComponent* Component = GetSkelMeshComponent();
 
 	FAnimExtractContext NewExtractContext(ExtractionContext);
-	NewExtractContext.bExtractRootMotion = RootMotionMode == ERootMotionMode::RootMotionFromEverything;
+	NewExtractContext.bExtractRootMotion = RootMotionMode == ERootMotionMode::RootMotionFromEverything || RootMotionMode == ERootMotionMode::IgnoreRootMotion;
 
 	if(const UAnimSequence* AnimSequence = Cast<const UAnimSequence>(Sequence))
 	{
@@ -1580,7 +1580,7 @@ void UAnimInstance::Montage_Advance(float DeltaSeconds)
 				RootMotionParams = (RootMotionMode != ERootMotionMode::IgnoreRootMotion) ? &ExtractedRootMotion : &LocalExtractedRootMotion;
 			}
 
-			MontageInstance->Advance(DeltaSeconds, bExtractRootMotion ? &ExtractedRootMotion : NULL, bUsingBlendedRootMotion);
+			MontageInstance->Advance(DeltaSeconds, RootMotionParams, bUsingBlendedRootMotion);
 
 			if (!MontageInstance->IsValid())
 			{
