Fix build (missing Box2D guard).

diff --git a/Engine/Source/Runtime/Engine/Private/PhysicsEngine/BodyInstance.cpp b/Engine/Source/Runtime/Engine/Private/PhysicsEngine/BodyInstance.cpp
--- a/Engine/Source/Runtime/Engine/Private/PhysicsEngine/BodyInstance.cpp
+++ b/Engine/Source/Runtime/Engine/Private/PhysicsEngine/BodyInstance.cpp
@@ -3845,8 +3845,11 @@ bool FBodyInstance::UseAsyncScene() const
 
 void FBodyInstance::SetUseAsyncScene(bool bNewUseAsyncScene)
 {
+#if WITH_BOX2D
 	// Invalid to call this if the body has been created
 	check(BodyInstancePtr == nullptr);
+#endif // WITH_BOX2D
+
 	// Set flag
 	bUseAsyncScene = bNewUseAsyncScene;
 }
