defer all component registration until after SCS execution during UCS

- removed mobility override patch from UCS attachment code and replaced w/ fixup code
- removed mobility override patch from SCS as it is already fixed up by the editor during node attachment

diff --git a/Engine/Source/Runtime/Engine/Private/ActorConstruction.cpp b/Engine/Source/Runtime/Engine/Private/ActorConstruction.cpp
--- a/Engine/Source/Runtime/Engine/Private/ActorConstruction.cpp
+++ b/Engine/Source/Runtime/Engine/Private/ActorConstruction.cpp
@@ -423,34 +423,26 @@ void AActor::ExecuteConstruction(const FTransform& Transform, const FComponentIn
 	OnConstruction(Transform);
 }
 
-
-void FDeferRegisterComponents::DeferComponentRegistration(AActor* Actor, USceneComponent* Component, EComponentMobility::Type OriginalMobility)
+void FDeferRegisterComponents::DeferComponentRegistration(AActor* Actor, UActorComponent* Component)
 {
-	TArray<FDeferredComponentInfo>* ActorComponentsToRegister = ComponentsToRegister.Find(Actor);
+	TArray<UActorComponent*>* ActorComponentsToRegister = ComponentsToRegister.Find(Actor);
 	if (!ActorComponentsToRegister)
 	{
-		ActorComponentsToRegister = &ComponentsToRegister.Add(Actor, TArray<FDeferredComponentInfo>());
+		ActorComponentsToRegister = &ComponentsToRegister.Add(Actor, TArray<UActorComponent*>());
 	}
-	ActorComponentsToRegister->Add(FDeferredComponentInfo(Component, OriginalMobility));
 	
-	// Components with Mobility set to EComponentMobility::Static or EComponentMobility::Stationary 
-	// can't be properly set up from a construction script (all changes will be rejected due to 
-	// EComponentMobility::Static flag) so we're going to temporarily change the flag and defer 
-	// the registration until the construction-script has finished.
-	Component->Mobility = EComponentMobility::Movable;
+	ActorComponentsToRegister->Add(Component);
 }
 
 void FDeferRegisterComponents::RegisterComponents(AActor* Actor)
 {
-	TArray<FDeferredComponentInfo>* ActorComponentsToRegister = ComponentsToRegister.Find(Actor);
+	TArray<UActorComponent*>* ActorComponentsToRegister = ComponentsToRegister.Find(Actor);
 	if (ActorComponentsToRegister)
 	{
 		for (int32 ComponentIndex = 0; ComponentIndex < ActorComponentsToRegister->Num(); ++ComponentIndex)
 		{
-			const FDeferredComponentInfo& SavedInfo = (*ActorComponentsToRegister)[ComponentIndex];
-			USceneComponent* Component = SavedInfo.Component;
-			// Restore saved mobility just before registering this component.
-			Component->Mobility = SavedInfo.SavedMobility;
+			UActorComponent* Component = (*ActorComponentsToRegister)[ComponentIndex];
+			
 			Component->RegisterComponent();
 		}
 		ComponentsToRegister.Remove(Actor);
@@ -460,13 +452,13 @@ void FDeferRegisterComponents::RegisterComponents(AActor* Actor)
 // FGCObject interface
 void FDeferRegisterComponents::AddReferencedObjects(FReferenceCollector& Collector)
 {
-	for (TMap<AActor*, TArray<FDeferredComponentInfo> >::TIterator It(ComponentsToRegister); It; ++It)
+	for (TMap<AActor*, TArray<UActorComponent*> >::TIterator It(ComponentsToRegister); It; ++It)
 	{
 		Collector.AddReferencedObject(It.Key());
-		TArray<FDeferredComponentInfo>& Components = It.Value();
+		TArray<UActorComponent*>& Components = It.Value();
 		for (int32 ComponentIndex = 0; ComponentIndex < Components.Num(); ++ComponentIndex)
 		{
-			Collector.AddReferencedObject(Components[ComponentIndex].Component);
+			Collector.AddReferencedObject(Components[ComponentIndex]);
 		}
 	}
 }
@@ -514,12 +506,12 @@ UActorComponent* AActor::CreateComponentFromTemplate(UActorComponent* Template,
 
 UActorComponent* AActor::AddComponent(FName TemplateName, bool bManualAttachment, const FTransform& RelativeTransform, const UObject* ComponentTemplateContext)
 {
-	UActorComponent* Template = NULL;
-	UBlueprintGeneratedClass* BlueprintGeneratedClass = Cast<UBlueprintGeneratedClass>((ComponentTemplateContext != NULL) ? ComponentTemplateContext->GetClass() : GetClass());
-	while(BlueprintGeneratedClass != NULL)
+	UActorComponent* Template = nullptr;
+	UBlueprintGeneratedClass* BlueprintGeneratedClass = Cast<UBlueprintGeneratedClass>((ComponentTemplateContext != nullptr) ? ComponentTemplateContext->GetClass() : GetClass());
+	while(BlueprintGeneratedClass != nullptr)
 	{
 		Template = BlueprintGeneratedClass->FindComponentTemplateByName(TemplateName);
-		if(NULL != Template)
+		if(nullptr != Template)
 		{
 			break;
 		}
@@ -527,33 +519,36 @@ UActorComponent* AActor::AddComponent(FName TemplateName, bool bManualAttachment
 	}
 
 	UActorComponent* NewActorComp = CreateComponentFromTemplate(Template);
-	if(NewActorComp != NULL)
+	if(NewActorComp != nullptr)
 	{
 		// The user has the option of doing attachment manually where they have complete control or via the automatic rule
 		// that the first component added becomes the root component, with subsequent components attached to the root.
 		USceneComponent* NewSceneComp = Cast<USceneComponent>(NewActorComp);
-		bool bDeferRegisterStaticComponent = false;
-		EComponentMobility::Type OriginalMobility = EComponentMobility::Movable;
-
-		if(NewSceneComp != NULL)
+		if(NewSceneComp != nullptr)
 		{
-			// Components with Mobility set to EComponentMobility::Static or EComponentMobility::Stationary can't be properly set up in UCS (all changes will be rejected
-			// due to EComponentMobility::Static flag) so we're going to temporarily change the flag and defer the registration until UCS has finished.
-			bDeferRegisterStaticComponent = bRunningUserConstructionScript && NewSceneComp->Mobility != EComponentMobility::Movable;
-			OriginalMobility = NewSceneComp->Mobility;
-			if (bDeferRegisterStaticComponent)
-			{
-				NewSceneComp->Mobility = EComponentMobility::Movable;
-			}
-
 			if (!bManualAttachment)
 			{
-				if (RootComponent == NULL)
+				if (RootComponent == nullptr)
 				{
 					RootComponent = NewSceneComp;
 				}
 				else
 				{
+					// The root component can't be more mobile than its children, so we check for that here and adjust as needed before attaching.
+					if(RootComponent->Mobility > NewSceneComp->Mobility)
+					{
+						if(NewSceneComp->IsA<UStaticMeshComponent>())
+						{
+							// SMCs can't be stationary, so always set them to be movable
+							NewSceneComp->SetMobility(EComponentMobility::Movable);
+						}
+						else
+						{
+							// Set the new component to be at least as mobile as the root
+							NewSceneComp->SetMobility(RootComponent->Mobility);
+						}
+					}
+
 					NewSceneComp->AttachTo(RootComponent);
 				}
 			}
@@ -564,10 +559,10 @@ UActorComponent* AActor::AddComponent(FName TemplateName, bool bManualAttachment
 		// Call function to notify component it has been created
 		NewActorComp->OnComponentCreated();
 
-		if (bDeferRegisterStaticComponent)
+		if(bRunningUserConstructionScript)
 		{
 			// Defer registration until after UCS has completed.
-			FDeferRegisterComponents::Get().DeferComponentRegistration(this, NewSceneComp, OriginalMobility);
+			FDeferRegisterComponents::Get().DeferComponentRegistration(this, NewActorComp);
 		}
 		else
 		{
diff --git a/Engine/Source/Runtime/Engine/Private/DeferRegisterComponents.h b/Engine/Source/Runtime/Engine/Private/DeferRegisterComponents.h
--- a/Engine/Source/Runtime/Engine/Private/DeferRegisterComponents.h
+++ b/Engine/Source/Runtime/Engine/Private/DeferRegisterComponents.h
@@ -2,19 +2,6 @@
 
 #pragma once
 
-
-class FDeferredComponentInfo
-{
-public:
-	USceneComponent* Component;
-	EComponentMobility::Type SavedMobility;
-
-	FDeferredComponentInfo(USceneComponent* InComponent, EComponentMobility::Type OriginalMobility) :
-		Component(InComponent),
-		SavedMobility(OriginalMobility)
-	{}
-};
-
 /**
  * Helper class to store components which registration should be deferred until
  * the construction script has completed... generally, all components get stuck here
@@ -24,7 +11,7 @@ public:
 class FDeferRegisterComponents : public FGCObject
 {
 	/** Map of actors and their components to register. */
-	TMap<AActor*, TArray<FDeferredComponentInfo> > ComponentsToRegister;
+	TMap<AActor*, TArray<UActorComponent*> > ComponentsToRegister;
 
 public:
 
@@ -34,7 +21,7 @@ public:
 	 * @param Actor Actor this component belongs to.
 	 * @param Component Component to register.
 	 */
-	void DeferComponentRegistration(AActor* Actor, USceneComponent* Component, EComponentMobility::Type OriginalMobility);
+	void DeferComponentRegistration(AActor* Actor, UActorComponent* Component);
 
 	/**
 	 * Registers all deferred components that belong to the specified actor.
diff --git a/Engine/Source/Runtime/Engine/Private/SCS_Node.cpp b/Engine/Source/Runtime/Engine/Private/SCS_Node.cpp
--- a/Engine/Source/Runtime/Engine/Private/SCS_Node.cpp
+++ b/Engine/Source/Runtime/Engine/Private/SCS_Node.cpp
@@ -28,12 +28,12 @@ USCS_Node::USCS_Node(const FObjectInitializer& ObjectInitializer)
 
 void USCS_Node::ExecuteNodeOnActor(AActor* Actor, USceneComponent* ParentComponent, const FTransform* RootTransform, bool bIsDefaultTransform)
 {
-	check(Actor != NULL);
-	check((ParentComponent != NULL && !ParentComponent->IsPendingKill()) || (RootTransform != NULL)); // must specify either a parent component or a world transform
+	check(Actor != nullptr);
+	check((ParentComponent != nullptr && !ParentComponent->IsPendingKill()) || (RootTransform != nullptr)); // must specify either a parent component or a world transform
 
 	// Create a new component instance based on the template
 	UActorComponent* NewActorComp = Actor->CreateComponentFromTemplate(ComponentTemplate, VariableName.ToString());
-	if(NewActorComp != NULL)
+	if(NewActorComp != nullptr)
 	{
 		bool bDeferRegisterStaticComponent = false;
 		EComponentMobility::Type OriginalMobility = EComponentMobility::Movable;
@@ -43,19 +43,10 @@ void USCS_Node::ExecuteNodeOnActor(AActor* Actor, USceneComponent* ParentCompone
 
 		// Special handling for scene components
 		USceneComponent* NewSceneComp = Cast<USceneComponent>(NewActorComp);
-		if (NewSceneComp != NULL)
+		if (NewSceneComp != nullptr)
 		{
-			// Components with Mobility set to EComponentMobility::Static or EComponentMobility::Stationary can't be properly set up in SCS (all changes will be rejected
-			// due to EComponentMobility::Static flag) so we're going to temporarily change the flag and defer the registration until SCS has finished.
-			bDeferRegisterStaticComponent = NewSceneComp->Mobility != EComponentMobility::Movable;
-			OriginalMobility = NewSceneComp->Mobility;
-			if (bDeferRegisterStaticComponent)
-			{
-				NewSceneComp->Mobility = EComponentMobility::Movable;
-			}
-
 			// If NULL is passed in, we are the root, so set transform and assign as RootComponent on Actor
-			if (ParentComponent == NULL || (ParentComponent && ParentComponent->IsPendingKill()))
+			if (ParentComponent == nullptr || (ParentComponent && ParentComponent->IsPendingKill()))
 			{
 				FTransform WorldTransform = *RootTransform;
 				if(bIsDefaultTransform)
@@ -76,16 +67,9 @@ void USCS_Node::ExecuteNodeOnActor(AActor* Actor, USceneComponent* ParentCompone
 		// Call function to notify component it has been created
 		NewActorComp->OnComponentCreated();
 
-		if (bDeferRegisterStaticComponent)
-		{
-			// need to defer component registration until after the construction script
-			// is ran, since the construction script can mutate the object (for collision, etc.)
-			FDeferRegisterComponents::Get().DeferComponentRegistration(Actor, NewSceneComp, OriginalMobility);
-		}
-		else
-		{
-			NewActorComp->RegisterComponent();
-		}
+		// need to defer component registration until after the construction script
+		// is ran, since the construction script can mutate the object (for collision, etc.)
+		FDeferRegisterComponents::Get().DeferComponentRegistration(Actor, NewActorComp);
 
 		if (NewActorComp->GetIsReplicated())
 		{
@@ -108,7 +92,7 @@ void USCS_Node::ExecuteNodeOnActor(AActor* Actor, USceneComponent* ParentCompone
 #if WITH_EDITOR
 				// If we're constructing editable components in the SCS editor, set the component instance corresponding to this node for editing purposes
 				USimpleConstructionScript* SCS = GetSCS();
-				if(SCS != NULL && (SCS->IsConstructingEditorComponents() || SCS->GetComponentEditorActorInstance() == Actor))
+				if(SCS != nullptr && (SCS->IsConstructingEditorComponents() || SCS->GetComponentEditorActorInstance() == Actor))
 				{
 					EditorComponentInstance = NewSceneComp;
 				}
@@ -117,14 +101,14 @@ void USCS_Node::ExecuteNodeOnActor(AActor* Actor, USceneComponent* ParentCompone
 		}
 
 		// Determine the parent component for our children (it's still our parent if we're a non-scene component)
-		USceneComponent* ParentSceneComponentOfChildren = (NewSceneComp != NULL) ? NewSceneComp : ParentComponent;
+		USceneComponent* ParentSceneComponentOfChildren = (NewSceneComp != nullptr) ? NewSceneComp : ParentComponent;
 
 		// If we made a component, go ahead and process our children
 		for (int32 NodeIdx = 0; NodeIdx < ChildNodes.Num(); NodeIdx++)
 		{
 			USCS_Node* Node = ChildNodes[NodeIdx];
-			check(Node != NULL);
-			Node->ExecuteNodeOnActor(Actor, ParentSceneComponentOfChildren, NULL, false);
+			check(Node != nullptr);
+			Node->ExecuteNodeOnActor(Actor, ParentSceneComponentOfChildren, nullptr, false);
 		}
 	}
 }
