Fix Mac Monolithic build

diff --git a/Engine/Source/Runtime/Engine/Private/CheatManager.cpp b/Engine/Source/Runtime/Engine/Private/CheatManager.cpp
--- a/Engine/Source/Runtime/Engine/Private/CheatManager.cpp
+++ b/Engine/Source/Runtime/Engine/Private/CheatManager.cpp
@@ -16,6 +16,7 @@
 #include "Components/CapsuleComponent.h"
 #include "Engine/DebugCameraController.h"
 #include "GameFramework/PlayerState.h"
+#include "GameFramework/GameMode.h"
 #include "Components/BrushComponent.h"
 
 DEFINE_LOG_CATEGORY_STATIC(LogCheatManager, Log, All);
