#Animation Editing Feature

- If you select bone and modify, you can add to the additive layer track
- You can also move curve keys during SIE

diff --git a/Engine/Source/Developer/AssetTools/Private/AssetTypeActions/AssetTypeActions_Skeleton.cpp b/Engine/Source/Developer/AssetTools/Private/AssetTypeActions/AssetTypeActions_Skeleton.cpp
--- a/Engine/Source/Developer/AssetTools/Private/AssetTypeActions/AssetTypeActions_Skeleton.cpp
+++ b/Engine/Source/Developer/AssetTools/Private/AssetTypeActions/AssetTypeActions_Skeleton.cpp
@@ -826,14 +826,11 @@ void FAssetTypeActions_Skeleton::RetargetSkeleton(TArray<FAssetToRemapSkeleton>&
 					UAnimSequenceBase * SequenceBase = Cast<UAnimSequenceBase>(AnimAsset);
 					if (SequenceBase)
 					{
-						// Copy curve data from source asset, preserving data in the target if present.
-						FSmartNameMapping* OldNameMapping = OldSkeleton->SmartNames.GetContainer(USkeleton::AnimCurveMappingName);
-						FSmartNameMapping* NewNameMapping = NewSkeleton->SmartNames.GetContainer(USkeleton::AnimCurveMappingName);
-						SequenceBase->RawCurveData.UpdateLastObservedNames(OldNameMapping);
-
-						for(FFloatCurve& Curve : SequenceBase->RawCurveData.FloatCurves)
+						EditorAnimUtils::CopyAnimCurves(OldSkeleton, NewSkeleton, SequenceBase, USkeleton::AnimCurveMappingName, FRawCurveTracks::FloatType);
+						
+						if (UAnimSequence * Sequence = Cast<UAnimSequence>(SequenceBase))
 						{
-							NewNameMapping->AddName(Curve.LastObservedName, Curve.CurveUid);
+							EditorAnimUtils::CopyAnimCurves(OldSkeleton, NewSkeleton, Sequence, USkeleton::AnimTrackCurveMappingName, FRawCurveTracks::TransformType);
 						}
 					}
 					
diff --git a/Engine/Source/Editor/AnimGraph/Classes/AnimPreviewInstance.h b/Engine/Source/Editor/AnimGraph/Classes/AnimPreviewInstance.h
--- a/Engine/Source/Editor/AnimGraph/Classes/AnimPreviewInstance.h
+++ b/Engine/Source/Editor/AnimGraph/Classes/AnimPreviewInstance.h
@@ -33,6 +33,10 @@ class ANIMGRAPH_API UAnimPreviewInstance : public UAnimSingleNodeInstance
 	UPROPERTY(transient)
 	TArray<FAnimNode_ModifyBone> BoneControllers;
 
+	/** Curve modifiers */
+	UPROPERTY(transient)
+	TArray<FAnimNode_ModifyBone> CurveBoneControllers;
+
 	/** Shared parameters for previewing blendspace or animsequence **/
 	UPROPERTY(transient)
 	float SkeletalControlAlpha;
@@ -88,29 +92,95 @@ class ANIMGRAPH_API UAnimPreviewInstance : public UAnimSingleNodeInstance
 	 * @param	InBoneName	The name of the bone modification to find
 	 * @return the bone modification or NULL if no current modification was found
 	 */
-	FAnimNode_ModifyBone* FindModifiedBone(const FName& InBoneName);
+	FAnimNode_ModifyBone* FindModifiedBone(const FName& InBoneName, bool bCurveController=false);
 
 	/** 
 	 * Modifies a single bone. Create a new FAnimNode_ModifyBone if one does not exist for the passed-in bone.
 	 * @param	InBoneName	The name of the bone to modify
 	 * @return the new or existing bone modification
 	 */
-	FAnimNode_ModifyBone& ModifyBone(const FName& InBoneName);
+	FAnimNode_ModifyBone& ModifyBone(const FName& InBoneName, bool bCurveController=false);
 
 	/**
 	 * Removes an existing bone modification
 	 * @param	InBoneName	The name of the existing modification to remove
 	 */
-	void RemoveBoneModification(const FName& InBoneName);
+	void RemoveBoneModification(const FName& InBoneName, bool bCurveController=false);
 
 	/**
 	 * Reset all bone modified
 	 */
-	void ResetModifiedBone();
+	void ResetModifiedBone(bool bCurveController=false);
 
 #if WITH_EDITORONLY_DATA
 	bool bForceRetargetBasePose;
 #endif
+
+	/**
+	 * Convert current modified bone transforms (BoneControllers) to transform curves (CurveControllers)
+	 * it does based on CurrentTime. This function does not set key directly here. 
+	 * It does wait until next update, and it gets the delta of transform before applying curves, and 
+	 * creates curves from it, so you'll need delegate if you'd like to do something after
+	 * 
+	 * @param Delegate To be called once set key is completed
+	 */
+	void SetKey(FSimpleDelegate InOnSetKeyCompleteDelegate);
+
+	/** 
+	 * Refresh Curve Bone Controllers based on TransformCurves from Animation data
+	 */
+	void RefreshCurveBoneControllers();
+
+	/**
+	 * Apply all Transform Curves to the RawAnimationData of the animation
+	 */
+	void BakeAnimation();
+
+	/** 
+	 * Enable Controllers
+	 * This is used by when editing, when controller has to be disabled
+	 */
+	void EnableControllers(bool bEnable);
+
+private:
+	/** 
+	 * Apply Bone Controllers to the Outpose
+	 *
+	 * @param	Component	Component to apply bone controller to
+	 * @param	BoneControllers	 List of Bone Controllers to apply
+	 * @param 	OutMeshPose	Outpose in Mesh Space once applied
+	 */
+	void ApplyBoneControllers(USkeletalMeshComponent* Component, TArray<FAnimNode_ModifyBone> &BoneControllers, FA2CSPose& OutMeshPose);
+	/** 
+	 * Update CurveControllers based on TransformCurves of Animation
+	 */
+	void UpdateCurveController();
+
+	/* 
+	 * Set Key Implementation function
+	 * It gets Pre Controller Local Space and gets Post Controller Local Space, and add the key to the curve 
+	 */
+	void SetKeyImplementation(const TArray<FTransform>& PreControllerInLocalSpace, const TArray<FTransform>& PostControllerInLocalSpace);
+	/** 
+	 * Add Key to the Sequence
+	 * Now Additive Key is generated, add to the curves
+	 */
+	void AddKeyToSequence(UAnimSequence* Sequence, float Time, const FName& BoneName, const FTransform& AdditiveTransform);
+	/* 
+	 * When this flag is true, it sets key
+	 */
+	bool bSetKey;
+
+	/*
+	 * Used to determine if controller has to be applied or not
+	 * Used to disable controller during editing
+	 */
+	bool bEnableControllers;
+	
+	/**
+	 * Delegate to call after Key is set
+	 */
+	FSimpleDelegate OnSetKeyCompleteDelegate;
 };
 
 
diff --git a/Engine/Source/Editor/AnimGraph/Private/AnimPreviewInstance.cpp b/Engine/Source/Editor/AnimGraph/Private/AnimPreviewInstance.cpp
--- a/Engine/Source/Editor/AnimGraph/Private/AnimPreviewInstance.cpp
+++ b/Engine/Source/Editor/AnimGraph/Private/AnimPreviewInstance.cpp
@@ -4,69 +4,101 @@
 #include "AnimGraphPrivatePCH.h"
 #include "AnimPreviewInstance.h"
 
+#if WITH_EDITOR
+#include "ScopedTransaction.h"
+#endif 
+
+#define LOCTEXT_NAMESPACE "AnimPreviewInstance"
+
 UAnimPreviewInstance::UAnimPreviewInstance(const FObjectInitializer& ObjectInitializer)
 	: Super(ObjectInitializer)
 	, SkeletalControlAlpha(1.0f)
 #if WITH_EDITORONLY_DATA
 	, bForceRetargetBasePose(false)
 #endif
+	, bSetKey(false)
 {
 	RootMotionMode = ERootMotionMode::RootMotionFromEverything;
+	bEnableControllers = true;
 }
 
 void UAnimPreviewInstance::NativeInitializeAnimation()
 {
 	// Cache our play state from the previous animation otherwise set to play
 	bool bCachedIsPlaying = (CurrentAsset != NULL) ? bPlaying : true;
-
+	bSetKey = false;
 	Super::NativeInitializeAnimation();
 
 	SetPlaying(bCachedIsPlaying);
+
+	RefreshCurveBoneControllers();
 }
 
-void UAnimPreviewInstance::ResetModifiedBone()
+void UAnimPreviewInstance::ResetModifiedBone(bool bCurveController/*=false*/)
 {
-	BoneControllers.Empty();
+	TArray<FAnimNode_ModifyBone>& Controllers = (bCurveController)?CurveBoneControllers : BoneControllers;
+	Controllers.Empty();
 }
 
-FAnimNode_ModifyBone* UAnimPreviewInstance::FindModifiedBone(const FName& InBoneName)
+FAnimNode_ModifyBone* UAnimPreviewInstance::FindModifiedBone(const FName& InBoneName, bool bCurveController/*=false*/)
 {
-	return BoneControllers.FindByPredicate(
+	TArray<FAnimNode_ModifyBone>& Controllers = (bCurveController)?CurveBoneControllers : BoneControllers;
+
+	return Controllers.FindByPredicate(
 		[InBoneName](const FAnimNode_ModifyBone& InController) -> bool
-		{
-			return InController.BoneToModify.BoneName == InBoneName;
-		}
+	{
+		return InController.BoneToModify.BoneName == InBoneName;
+	}
 	);
 }
 
-FAnimNode_ModifyBone& UAnimPreviewInstance::ModifyBone(const FName& InBoneName)
+FAnimNode_ModifyBone& UAnimPreviewInstance::ModifyBone(const FName& InBoneName, bool bCurveController/*=false*/)
 {
-	FAnimNode_ModifyBone* SingleBoneController = FindModifiedBone(InBoneName);
+	FAnimNode_ModifyBone* SingleBoneController = FindModifiedBone(InBoneName, bCurveController);
+	TArray<FAnimNode_ModifyBone>& Controllers = (bCurveController)?CurveBoneControllers : BoneControllers;
 
 	if(SingleBoneController == nullptr)
 	{
-		int32 NewIndex = BoneControllers.Add(FAnimNode_ModifyBone());
-		SingleBoneController = &BoneControllers[NewIndex];
+		int32 NewIndex = Controllers.Add(FAnimNode_ModifyBone());
+		SingleBoneController = &Controllers[NewIndex];
 	}
 
 	SingleBoneController->BoneToModify.BoneName = InBoneName;
 
-	SingleBoneController->TranslationMode = BMM_Replace;
-	SingleBoneController->TranslationSpace = BCS_BoneSpace;
+	if (bCurveController)
+	{
+		SingleBoneController->TranslationMode = BMM_Additive;
+		SingleBoneController->TranslationSpace = BCS_BoneSpace;
+
+		SingleBoneController->RotationMode = BMM_Additive;
+		SingleBoneController->RotationSpace = BCS_BoneSpace;
+
+		SingleBoneController->ScaleMode = BMM_Additive;
+		SingleBoneController->ScaleSpace = BCS_BoneSpace;
+	}
+	else
+	{
+		SingleBoneController->TranslationMode = BMM_Replace;
+		SingleBoneController->TranslationSpace = BCS_BoneSpace;
+
+		SingleBoneController->RotationMode = BMM_Replace;
+		SingleBoneController->RotationSpace = BCS_BoneSpace;
 
-	SingleBoneController->RotationMode = BMM_Replace;
-	SingleBoneController->RotationSpace = BCS_BoneSpace;
+		SingleBoneController->ScaleMode = BMM_Replace;
+		SingleBoneController->ScaleSpace = BCS_BoneSpace;
+	}
 
 	return *SingleBoneController;
 }
 
-void UAnimPreviewInstance::RemoveBoneModification(const FName& InBoneName)
+void UAnimPreviewInstance::RemoveBoneModification(const FName& InBoneName, bool bCurveController/*=false*/)
 {
-	BoneControllers.RemoveAll(
+	TArray<FAnimNode_ModifyBone>& Controllers = (bCurveController)?CurveBoneControllers : BoneControllers;
+	Controllers.RemoveAll(
 		[InBoneName](const FAnimNode_ModifyBone& InController)
-		{
-			return InController.BoneToModify.BoneName == InBoneName;
-		}
+	{
+		return InController.BoneToModify.BoneName == InBoneName;
+	}
 	);
 }
 
@@ -88,7 +120,7 @@ bool UAnimPreviewInstance::NativeEvaluateAnimation(FPoseContext& Output)
 #if WITH_EDITORONLY_DATA
 	if(bForceRetargetBasePose)
 	{
-		USkeletalMeshComponent * MeshComponent = GetSkelMeshComponent();
+		USkeletalMeshComponent* MeshComponent = GetSkelMeshComponent();
 		if(MeshComponent && MeshComponent->SkeletalMesh)
 		{
 			FAnimationRuntime::FillWithRetargetBaseRefPose(Output.Pose.Bones, GetSkelMeshComponent()->SkeletalMesh, RequiredBones);
@@ -105,37 +137,257 @@ bool UAnimPreviewInstance::NativeEvaluateAnimation(FPoseContext& Output)
 		Super::NativeEvaluateAnimation(Output);
 	}
 
-	USkeletalMeshComponent* Component = GetSkelMeshComponent();
-
-	if (Component && CurrentSkeleton)
+	if (bEnableControllers)
 	{
-		if(BoneControllers.Num() > 0)
+		UDebugSkelMeshComponent* Component = Cast<UDebugSkelMeshComponent>(GetSkelMeshComponent());
+
+		if(Component && CurrentSkeleton)
 		{
-			FA2CSPose OutMeshPose;
-			OutMeshPose.AllocateLocalPoses(RequiredBones, Output.Pose);
+			// update curve controllers
+			UpdateCurveController();
 
-			for(auto& SingleBoneController : BoneControllers)
+			// create bone controllers from 
+			if(BoneControllers.Num() > 0 || CurveBoneControllers.Num() > 0)
 			{
-				SingleBoneController.BoneToModify.BoneIndex = RequiredBones.GetPoseBoneIndexForBoneName(SingleBoneController.BoneToModify.BoneName);
-				if (SingleBoneController.BoneToModify.BoneIndex != INDEX_NONE)
+				TArray<FTransform> PreController, PostController;
+				// if set key is true, we should save pre controller local space transform 
+				// so that we can calculate the delta correctly
+				if(bSetKey)
 				{
-					TArray<FBoneTransform> BoneTransforms;
-					SingleBoneController.EvaluateBoneTransforms(Component, RequiredBones, OutMeshPose, BoneTransforms);
-					if (BoneTransforms.Num() > 0)
-					{
-						OutMeshPose.LocalBlendCSBoneTransforms(BoneTransforms, 1.0f);
-					}
+					PreController = Output.Pose.Bones;
+				}
+
+				FA2CSPose OutMeshPose;
+				OutMeshPose.AllocateLocalPoses(RequiredBones, Output.Pose);
+
+				// apply curve data first
+				ApplyBoneControllers(Component, CurveBoneControllers, OutMeshPose);
+
+				// and now apply bone controllers data
+				// it is possible they can be overlapping, but then bone controllers will overwrite
+				ApplyBoneControllers(Component, BoneControllers, OutMeshPose);
+
+				// convert back to local @todo check this
+				OutMeshPose.ConvertToLocalPoses(Output.Pose);
+
+				if(bSetKey)
+				{
+					// now we have post controller, and calculate delta now
+					PostController = Output.Pose.Bones;
+					SetKeyImplementation(PreController, PostController);
+				}
+			}
+			// if any other bone is selected, still go for set key even if nothing changed
+			else if(Component->BonesOfInterest.Num() > 0)
+			{
+				if(bSetKey)
+				{
+					// in this case, pose is same
+					SetKeyImplementation(Output.Pose.Bones, Output.Pose.Bones);
 				}
 			}
+		}
 
-			// convert back to local @todo check this
-			OutMeshPose.ConvertToLocalPoses(Output.Pose);
+		// we should unset here, just in case somebody clicks the key when it's not valid
+		if(bSetKey)
+		{
+			bSetKey = false;
 		}
 	}
 
 	return true;
 }
 
+void UAnimPreviewInstance::ApplyBoneControllers(USkeletalMeshComponent* Component, TArray<FAnimNode_ModifyBone> &InBoneControllers, FA2CSPose& OutMeshPose)
+{
+	for(auto& SingleBoneController : InBoneControllers)
+	{
+		SingleBoneController.BoneToModify.BoneIndex = RequiredBones.GetPoseBoneIndexForBoneName(SingleBoneController.BoneToModify.BoneName);
+		if(SingleBoneController.BoneToModify.BoneIndex != INDEX_NONE)
+		{
+			TArray<FBoneTransform> BoneTransforms;
+			SingleBoneController.EvaluateBoneTransforms(Component, RequiredBones, OutMeshPose, BoneTransforms);
+			if(BoneTransforms.Num() > 0)
+			{
+				OutMeshPose.LocalBlendCSBoneTransforms(BoneTransforms, 1.0f);
+			}
+		}
+	}
+}
+
+void UAnimPreviewInstance::SetKeyImplementation(const TArray<FTransform>& PreControllerInLocalSpace, const TArray<FTransform>& PostControllerInLocalSpace)
+{
+#if WITH_EDITOR
+	// evaluate the curve data first
+	UAnimSequence* CurrentSequence = Cast<UAnimSequence>(CurrentAsset);
+	UDebugSkelMeshComponent* Component = Cast<UDebugSkelMeshComponent> (GetSkelMeshComponent());
+
+	if(CurrentSequence && CurrentSkeleton && Component && Component->SkeletalMesh)
+	{
+		FScopedTransaction ScopedTransaction(LOCTEXT("SetKey", "Set Key"));
+		CurrentSequence->Modify(true);
+		Modify();
+
+		TArray<FName> BonesToModify;
+		// need to get component transform first. Depending on when this gets called, the transform is not up-to-date. 
+		// first look at the bonecontrollers, and convert each bone controller to transform curve key
+		// and add new curvebonecontrollers with additive data type
+		// clear bone controller data
+		for(auto& SingleBoneController : BoneControllers)
+		{
+			// find bone name, and just get transform of the bone in local space
+			// and get the additive data
+			// find if this already exists, then just add curve data only
+			FName BoneName = SingleBoneController.BoneToModify.BoneName;
+			// now convert data
+			const int32 BoneIndex = Component->GetBoneIndex(BoneName);
+			FTransform  LocalTransform = PostControllerInLocalSpace[BoneIndex];
+
+			// now we have LocalTransform and get additive data
+			FTransform AdditiveTransform = LocalTransform.GetRelativeTransform(PreControllerInLocalSpace[BoneIndex]);
+			AddKeyToSequence(CurrentSequence, CurrentTime, BoneName, AdditiveTransform);
+
+			BonesToModify.Add(BoneName);
+		}
+
+		// see if the bone is selected right now and if that is added - if bone is selected, we should add identity key to it. 
+		if ( Component->BonesOfInterest.Num() > 0 )
+		{
+			// if they're selected, we should add to the modifyBone list even if they're not modified, so that they can key that point. 
+			// first make sure those are added 
+			// if not added, make sure to set the key for them
+			for (const auto& BoneIndex : Component->BonesOfInterest)
+			{
+				FName BoneName = Component->GetBoneName(BoneIndex);
+				// if it's not on BonesToModify, add identity here. 
+				if (!BonesToModify.Contains(BoneName))
+				{
+					AddKeyToSequence(CurrentSequence, CurrentTime, BoneName, FTransform::Identity);
+				}
+			}
+		}
+
+		ResetModifiedBone(false);
+
+		OnSetKeyCompleteDelegate.ExecuteIfBound();
+	}
+#endif
+}
+
+void UAnimPreviewInstance::AddKeyToSequence(UAnimSequence* Sequence, float Time, const FName& BoneName, const FTransform& AdditiveTransform) 
+{
+	Sequence->AddKeyToSequence(Time, BoneName, AdditiveTransform);
+
+	// now add to the controller
+	// find if it exists in CurveBoneController
+	// make sure you add it there
+	ModifyBone(BoneName, true);
+
+	RequiredBones.SetUseSourceData(true);
+}
+
+void UAnimPreviewInstance::SetKey(FSimpleDelegate InOnSetKeyCompleteDelegate)
+{
+#if WITH_EDITOR
+	bSetKey = true;
+	OnSetKeyCompleteDelegate = InOnSetKeyCompleteDelegate;
+#endif
+}
+
+void UAnimPreviewInstance::RefreshCurveBoneControllers()
+{
+	// go through all curves and see if it has Transform Curve
+	// if so, find what bone that belong to and create BoneMOdifier for them
+	UAnimSequence* CurrentSequence = Cast<UAnimSequence>(CurrentAsset);
+
+	CurveBoneControllers.Empty();
+
+	// do not apply if BakedAnimation is on
+	if(CurrentSequence)
+	{
+		// make sure if this needs source update
+		if ( !CurrentSequence->DoesContainTransformCurves() )
+		{
+			return;
+		}
+
+		RequiredBones.SetUseSourceData(true);
+
+		TArray<FTransformCurve>& Curves = CurrentSequence->RawCurveData.TransformCurves;
+		FSmartNameMapping* NameMapping = CurrentSkeleton->SmartNames.GetContainer(USkeleton::AnimTrackCurveMappingName);
+
+		for (auto& Curve : Curves)
+		{
+			// skip if disabled
+			if (Curve.GetCurveTypeFlag(ACF_Disabled))
+			{
+				continue;
+			}
+
+			// add bone modifier
+			FName CurveName;
+			NameMapping->GetName(Curve.CurveUid, CurveName);
+
+			// @TODO: this is going to be issue. If they don't save skeleton with it, we don't have name at all?
+ 			if (CurveName == NAME_None)
+ 			{
+				FSmartNameMapping::UID NewUID;
+ 				NameMapping->AddOrFindName(Curve.LastObservedName, NewUID);
+				Curve.CurveUid = NewUID;
+
+				CurveName = Curve.LastObservedName;
+ 			}
+
+			FName BoneName = CurveName;
+			if (BoneName != NAME_None && CurrentSkeleton->GetReferenceSkeleton().FindBoneIndex(BoneName) != INDEX_NONE)
+			{
+				ModifyBone(BoneName, true);
+			}
+		}
+	}
+}
+
+void UAnimPreviewInstance::UpdateCurveController()
+{
+	// evaluate the curve data first
+	UAnimSequenceBase* CurrentSequence = Cast<UAnimSequenceBase>(CurrentAsset);
+
+	if (CurrentSequence && CurrentSkeleton)
+	{
+		TMap<FName, FTransform> ActiveCurves;
+		CurrentSequence->RawCurveData.EvaluateTransformCurveData(CurrentSkeleton, ActiveCurves, CurrentTime, 1.f);
+
+		// make sure those curves exists in the bone controller, otherwise problem
+		if ( ActiveCurves.Num() > 0 )
+		{
+			for(auto& SingleBoneController : CurveBoneControllers)
+			{
+				// make sure the curve exists
+				FName CurveName = SingleBoneController.BoneToModify.BoneName;
+
+				// we should add extra key to front and back whenever animation length changes or so. 
+				// animation length change requires to bake down animation first 
+				// this will make sure all the keys that were embedded at the start/end will automatically be backed to the data
+				const FTransform* Value = ActiveCurves.Find(CurveName);
+				if (Value)
+				{
+					// apply this change
+					SingleBoneController.Translation = Value->GetTranslation();
+					SingleBoneController.Scale = Value->GetScale3D();
+					// sasd we're converting twice
+					SingleBoneController.Rotation = Value->GetRotation().Rotator();
+				}
+			}
+		}
+		else
+		{
+			// should match
+			ensure (CurveBoneControllers.Num() == 0);
+			CurveBoneControllers.Empty();
+		}
+	}
+}
+
 /** Set SkeletalControl Alpha**/
 void UAnimPreviewInstance::SetSkeletalControlAlpha(float InSkeletalControlAlpha)
 {
@@ -163,8 +415,14 @@ void UAnimPreviewInstance::RestartMontage(UAnimMontage* Montage, FName FromSecti
 
 void UAnimPreviewInstance::SetAnimationAsset(UAnimationAsset* NewAsset, bool bIsLooping, float InPlayRate)
 {
+	// make sure to turn that off before setting new asset
+	RequiredBones.SetUseSourceData(false);
+
 	Super::SetAnimationAsset(NewAsset, bIsLooping, InPlayRate);
 	RootMotionMode = Cast<UAnimMontage>(CurrentAsset) != NULL ? ERootMotionMode::RootMotionFromMontagesOnly : ERootMotionMode::RootMotionFromEverything;
+
+	// should re sync up curve bone controllers from new asset
+	RefreshCurveBoneControllers();
 }
 
 void UAnimPreviewInstance::MontagePreview_SetLooping(bool bIsLooping)
@@ -217,7 +475,7 @@ void UAnimPreviewInstance::MontagePreview_SetReverse(bool bInReverse)
 {
 	Super::SetReverse(bInReverse);
 
-	if (FAnimMontageInstance * CurMontageInstance = GetActiveMontageInstance())
+	if (FAnimMontageInstance* CurMontageInstance = GetActiveMontageInstance())
 	{
 		// copy the current playrate
 		CurMontageInstance->SetPlayRate(PlayRate);
@@ -710,3 +968,19 @@ int32 UAnimPreviewInstance::MontagePreview_FindLastSection(int32 StartSectionIdx
 	}
 	return ResultIdx;
 }
+
+void UAnimPreviewInstance::BakeAnimation()
+{
+	if (UAnimSequence* Sequence = Cast<UAnimSequence>(CurrentAsset))
+	{
+		FScopedTransaction ScopedTransaction(LOCTEXT("BakeAnimation", "Bake Animation"));
+		Sequence->Modify(true);
+		Sequence->BakeTrackCurvesToRawAnimation();
+	}
+}
+
+void UAnimPreviewInstance::EnableControllers(bool bEnable)
+{
+	bEnableControllers = bEnable;
+}
+#undef LOCTEXT_NAMESPACE
\ No newline at end of file
diff --git a/Engine/Source/Editor/EditorStyle/Private/SlateEditorStyle.cpp b/Engine/Source/Editor/EditorStyle/Private/SlateEditorStyle.cpp
--- a/Engine/Source/Editor/EditorStyle/Private/SlateEditorStyle.cpp
+++ b/Engine/Source/Editor/EditorStyle/Private/SlateEditorStyle.cpp
@@ -4462,6 +4462,10 @@ void FSlateEditorStyle::FStyle::SetupPersonaStyle()
 		Set("Persona.StopRecordAnimation.Small", new IMAGE_BRUSH(TEXT("Icons/icon_Persona_StopRecord_40x"), Icon20x20));
 		Set("Persona.StopRecordAnimation_Alt", new IMAGE_BRUSH(TEXT("Icons/icon_Persona_StopRecord_Alt_40x"), Icon40x40));
 		Set("Persona.StopRecordAnimation_Alt.Small", new IMAGE_BRUSH(TEXT("Icons/icon_Persona_StopRecord_Alt_40x"), Icon20x20));
+		Set("Persona.SetKey", new IMAGE_BRUSH(TEXT("Icons/icon_Persona_SetKey_40x"), Icon40x40));
+		Set("Persona.SetKey.Small", new IMAGE_BRUSH(TEXT("Icons/icon_Persona_SetKey_40x"), Icon20x20));
+		Set("Persona.ApplyAnimation", new IMAGE_BRUSH(TEXT("Icons/icon_Persona_BakeAnim_40x"), Icon40x40));
+		Set("Persona.ApplyAnimation.Small", new IMAGE_BRUSH(TEXT("Icons/icon_Persona_BakeAnim_40x"), Icon20x20));
 
 		// persona extras
 		Set("Persona.ConvertAnimationGraph", new IMAGE_BRUSH("Old/Graph/ConvertIcon", Icon40x40));
diff --git a/Engine/Source/Editor/Persona/Private/AnimViewportShowCommands.cpp b/Engine/Source/Editor/Persona/Private/AnimViewportShowCommands.cpp
--- a/Engine/Source/Editor/Persona/Private/AnimViewportShowCommands.cpp
+++ b/Engine/Source/Editor/Persona/Private/AnimViewportShowCommands.cpp
@@ -31,6 +31,8 @@ void FAnimViewportShowCommands::RegisterCommands()
 	UI_COMMAND( ShowRawAnimation, "Uncompressed Animation", "Display Skeleton With Uncompressed Animation Data", EUserInterfaceActionType::ToggleButton, FInputGesture() );
 	UI_COMMAND( ShowNonRetargetedAnimation, "NonRetargeted Animation", "Display Skeleton With non retargeted Animation Data", EUserInterfaceActionType::ToggleButton, FInputGesture() );
 	UI_COMMAND( ShowAdditiveBaseBones, "Additive Base", "Display Skeleton In Additive Base Pose", EUserInterfaceActionType::ToggleButton, FInputGesture() );
+	UI_COMMAND( ShowSourceRawAnimation, "Source Animation", "Display Skeleton In Source Raw Animation if you have Track Curves Modified.", EUserInterfaceActionType::ToggleButton, FInputGesture() );
+	UI_COMMAND( ShowBakedAnimation, "Baked Animation", "Display Skeleton In Baked Raw Animation if you have Track Curves Modified.", EUserInterfaceActionType::ToggleButton, FInputGesture() );
 	UI_COMMAND( ShowSockets, "Sockets", "Display socket hitpoints", EUserInterfaceActionType::ToggleButton, FInputGesture() );
 
 	UI_COMMAND( ShowLocalAxesNone, "None", "Hides all local hierarchy axis", EUserInterfaceActionType::RadioButton, FInputGesture() );
diff --git a/Engine/Source/Editor/Persona/Private/AnimViewportShowCommands.h b/Engine/Source/Editor/Persona/Private/AnimViewportShowCommands.h
--- a/Engine/Source/Editor/Persona/Private/AnimViewportShowCommands.h
+++ b/Engine/Source/Editor/Persona/Private/AnimViewportShowCommands.h
@@ -60,6 +60,12 @@ public:
 	/** Show additive base pose */
 	TSharedPtr< FUICommandInfo > ShowAdditiveBaseBones;
 
+	/** Show non retargeted animation. */
+	TSharedPtr< FUICommandInfo > ShowSourceRawAnimation;
+
+	/** Show non retargeted animation. */
+	TSharedPtr< FUICommandInfo > ShowBakedAnimation;
+
 	/** Show skeletal mesh bone names */
 	TSharedPtr< FUICommandInfo > ShowBoneNames;
 
diff --git a/Engine/Source/Editor/Persona/Private/AnimationEditorViewportClient.cpp b/Engine/Source/Editor/Persona/Private/AnimationEditorViewportClient.cpp
--- a/Engine/Source/Editor/Persona/Private/AnimationEditorViewportClient.cpp
+++ b/Engine/Source/Editor/Persona/Private/AnimationEditorViewportClient.cpp
@@ -425,7 +425,14 @@ void FAnimationViewportClient::Draw(const FSceneView* View, FPrimitiveDrawInterf
 		{
 			DrawMeshBonesAdditiveBasePose(PreviewSkelMeshComp.Get(), PDI);
 		}
-
+		if(PreviewSkelMeshComp->bDisplayBakedAnimation)
+		{
+			DrawMeshBonesBakedAnimation(PreviewSkelMeshComp.Get(), PDI);
+		}
+		if(PreviewSkelMeshComp->bDisplaySourceAnimation)
+		{
+			DrawMeshBonesSourceRawAnimation(PreviewSkelMeshComp.Get(), PDI);
+		}
 		// Display normal vectors of each simulation vertex
 		if ( PreviewSkelMeshComp->bDisplayClothingNormals )
 		{
@@ -771,6 +778,20 @@ void FAnimationViewportClient::DisplayInfo(FCanvas* Canvas, FSceneView* View, bo
 		}
 	}
 
+	UAnimPreviewInstance* PreviewInstance = PreviewSkelMeshComp->PreviewInstance;
+	if( PreviewInstance )
+	{
+		// see if you have anim sequence that has transform curves
+		UAnimSequence* Sequence = Cast<UAnimSequence>(PreviewInstance->CurrentAsset);
+		if ( Sequence && Sequence->DoesNeedRebake() )
+		{
+			FColor SubHeadlineColour(202, 66, 0);
+			InfoString = TEXT("Animation is being edited. To apply to raw animation data, click \"Apply\"");
+			Canvas->DrawShadowedString(CurXOffset, CurYOffset, *InfoString, GEngine->GetSmallFont(), SubHeadlineColour);
+			CurYOffset += YL + 2;
+		}
+	}
+
 	if (PreviewSkelMeshComp->IsUsingInGameBounds())
 	{
 		if (!PreviewSkelMeshComp->CheckIfBoundsAreCorrrect())
@@ -1478,6 +1499,22 @@ void FAnimationViewportClient::DrawMeshBonesAdditiveBasePose(UDebugSkelMeshCompo
 	}
 }
 
+void FAnimationViewportClient::DrawMeshBonesSourceRawAnimation(UDebugSkelMeshComponent * MeshComponent, FPrimitiveDrawInterface* PDI) const
+{
+	if(MeshComponent && MeshComponent->SkeletalMesh)
+	{
+		DrawBonesFromTransforms(MeshComponent->SourceAnimationPoses, MeshComponent, PDI, FColor(195, 195, 195, 255), FColor(195, 159, 195, 255));
+	}
+}
+
+void FAnimationViewportClient::DrawMeshBonesBakedAnimation(UDebugSkelMeshComponent * MeshComponent, FPrimitiveDrawInterface* PDI) const
+{
+	if(MeshComponent && MeshComponent->SkeletalMesh)
+	{
+		DrawBonesFromTransforms(MeshComponent->BakedAnimationPoses, MeshComponent, PDI, FColor(0, 128, 192, 255), FColor(0, 128, 192, 255));
+	}
+}
+
 void FAnimationViewportClient::DrawMeshBones(USkeletalMeshComponent * MeshComponent, FPrimitiveDrawInterface* PDI) const
 {
 	if ( MeshComponent && MeshComponent->SkeletalMesh )
@@ -1516,7 +1553,7 @@ void FAnimationViewportClient::DrawMeshBones(USkeletalMeshComponent * MeshCompon
 	}
 }
 
-void FAnimationViewportClient::DrawBones(const USkeletalMeshComponent* MeshComponent, const TArray<FBoneIndexType> & RequiredBones, const TArray<FTransform> & WorldTransforms, FPrimitiveDrawInterface* PDI, const TArray<FLinearColor> BoneColours, float LineThickness/*=0.f*/) const
+void FAnimationViewportClient::DrawBones(const USkeletalMeshComponent* MeshComponent, const TArray<FBoneIndexType>& RequiredBones, const TArray<FTransform>& WorldTransforms, FPrimitiveDrawInterface* PDI, const TArray<FLinearColor> BoneColours, float LineThickness/*=0.f*/) const
 {
 	check ( MeshComponent && MeshComponent->SkeletalMesh );
 
diff --git a/Engine/Source/Editor/Persona/Private/AnimationEditorViewportClient.h b/Engine/Source/Editor/Persona/Private/AnimationEditorViewportClient.h
--- a/Engine/Source/Editor/Persona/Private/AnimationEditorViewportClient.h
+++ b/Engine/Source/Editor/Persona/Private/AnimationEditorViewportClient.h
@@ -298,6 +298,10 @@ private:
 	void DrawMeshBonesNonRetargetedAnimation(UDebugSkelMeshComponent * MeshComponent, FPrimitiveDrawInterface* PDI) const;
 	/** Draws Bones for Additive Base Pose */
 	void DrawMeshBonesAdditiveBasePose(UDebugSkelMeshComponent * MeshComponent, FPrimitiveDrawInterface* PDI) const;
+	/** Draw Bones for non retargeted animation. */
+	void DrawMeshBonesSourceRawAnimation(UDebugSkelMeshComponent * MeshComponent, FPrimitiveDrawInterface* PDI) const;
+	/** Draw Bones for non retargeted animation. */
+	void DrawMeshBonesBakedAnimation(UDebugSkelMeshComponent * MeshComponent, FPrimitiveDrawInterface* PDI) const;
 	/** Draws Bones for RequiredBones with WorldTransform **/
 	void DrawBones(const USkeletalMeshComponent* MeshComponent, const TArray<FBoneIndexType> & RequiredBones, const TArray<FTransform> & WorldTransforms, FPrimitiveDrawInterface* PDI, const TArray<FLinearColor> BoneColours, float LineThickness=0.f) const;
 	/** Draw Sub set of Bones **/
diff --git a/Engine/Source/Editor/Persona/Private/Persona.cpp b/Engine/Source/Editor/Persona/Private/Persona.cpp
--- a/Engine/Source/Editor/Persona/Private/Persona.cpp
+++ b/Engine/Source/Editor/Persona/Private/Persona.cpp
@@ -731,15 +731,117 @@ void FPersona::InitPersona(const EToolkitMode::Type Mode, const TSharedPtr< clas
 	
 }
 	
+void FPersona::CreateAnimation(const TArray<UObject*> NewAssets, int32 Option) 
+{
+	bool bResult = true;
+	if (NewAssets.Num() > 0)
+	{
+		USkeletalMeshComponent * MeshComponent = GetPreviewMeshComponent();
+		UAnimSequence * Sequence = Cast<UAnimSequence> (GetPreviewAnimationAsset());
+
+		for (auto NewAsset : NewAssets)
+		{
+			UAnimSequence * NewAnimSequence = Cast<UAnimSequence>(NewAsset);
+			if (NewAnimSequence)
+			{
+				switch (Option)
+				{
+				case 0:
+					bResult &= NewAnimSequence->CreateAnimation(MeshComponent->SkeletalMesh);
+					break;
+				case 1:
+					bResult &= NewAnimSequence->CreateAnimation(MeshComponent);
+					break;
+				case 2:
+					bResult &= NewAnimSequence->CreateAnimation(Sequence);
+					break;
+				}
+			}
+		}
+
+		// if it contains error, warn them
+		if (bResult)
+		{
+			OnAssetCreated(NewAssets);
+
+			// if it created based on current mesh component, 
+			if (Option == 1)
+			{
+				PreviewComponent->PreviewInstance->ResetModifiedBone();
+			}
+		}
+		else
+		{
+			// give warning
+		}
+	}
+}
+
+void FPersona::FillCreateAnimationMenu(FMenuBuilder& MenuBuilder) const
+{
+	TArray<TWeakObjectPtr<USkeleton>> Skeletons;
+
+	Skeletons.Add(TargetSkeleton);
+
+	// create rig
+	MenuBuilder.BeginSection("CreateAnimationSubMenu", LOCTEXT("CreateAnimationSubMenuHeading", "Create Animation"));
+	{
+		MenuBuilder.AddMenuEntry(
+			LOCTEXT("CreateAnimation_RefPose", "From Reference Pose"),
+			LOCTEXT("CreateAnimation_RefPose_Tooltip", "Create Animation from reference pose."),
+			FSlateIcon(),
+			FUIAction(
+			FExecuteAction::CreateStatic(&AnimationEditorUtils::ExecuteNewAnimAsset<UAnimSequenceFactory, UAnimSequence>, Skeletons, FString("_Sequence"), FAnimAssetCreated::CreateSP(this, &FPersona::CreateAnimation, 0), false),
+			FCanExecuteAction()
+			)
+			);
+
+		MenuBuilder.AddMenuEntry(
+			LOCTEXT("CreateAnimation_CurrentPose", "From Current Pose"),
+			LOCTEXT("CreateAnimation_CurrentPose_Tooltip", "Create Animation from current pose."),
+			FSlateIcon(),
+			FUIAction(
+			FExecuteAction::CreateStatic(&AnimationEditorUtils::ExecuteNewAnimAsset<UAnimSequenceFactory, UAnimSequence>, Skeletons, FString("_Sequence"), FAnimAssetCreated::CreateSP(this, &FPersona::CreateAnimation, 1), false),
+			FCanExecuteAction()
+			)
+			);
+
+		MenuBuilder.AddMenuEntry(
+			LOCTEXT("CreateAnimation_CurrentAnimation", "From Current Animation"),
+			LOCTEXT("CreateAnimation_CurrentAnimation_Tooltip", "Create Animation from current animation."),
+			FSlateIcon(),
+			FUIAction(
+			FExecuteAction::CreateStatic(&AnimationEditorUtils::ExecuteNewAnimAsset<UAnimSequenceFactory, UAnimSequence>, Skeletons, FString("_Sequence"), FAnimAssetCreated::CreateSP(this, &FPersona::CreateAnimation, 2), false),
+			FCanExecuteAction::CreateSP(this, &FPersona::HasValidAnimationSequencePlaying)
+			)
+			);
+	}
+	MenuBuilder.EndSection();
+}
+
 TSharedRef< SWidget > FPersona::GenerateCreateAssetMenu( USkeleton* Skeleton ) const
 {
 	const bool bShouldCloseWindowAfterMenuSelection = true;
 	FMenuBuilder MenuBuilder(bShouldCloseWindowAfterMenuSelection, NULL);
 
+	// Create Animation menu
+	MenuBuilder.BeginSection("CreateAnimation", LOCTEXT("CreateAnimationMenuHeading", "Animation"));
+	{
+		// create menu
+		MenuBuilder.AddSubMenu(
+				LOCTEXT("CreateAnimationSubmenu", "Create Animation"),
+				LOCTEXT("CreateAnimationSubmenu_ToolTip", "Create Animation for this skeleton"),
+				FNewMenuDelegate::CreateSP(this, &FPersona::FillCreateAnimationMenu),
+				false,
+				FSlateIcon(FEditorStyle::GetStyleSetName(), "Persona.AssetActions.CreateAnimAsset")
+				);
+	}
+	MenuBuilder.EndSection();
+
 	TArray<TWeakObjectPtr<USkeleton>> Skeletons;
-	
+
 	Skeletons.Add(Skeleton);
-	
+
 	AnimationEditorUtils::FillCreateAssetMenu(MenuBuilder, Skeletons, FAnimAssetCreated::CreateSP(this, &FPersona::OnAssetCreated), false);
 
 	return MenuBuilder.MakeWidget();
@@ -815,6 +917,13 @@ void FPersona::ExtendDefaultPersonaToolbar()
 				ToolbarBuilder.AddToolBarButton(FPersonaCommands::Get().ApplyCompression, NAME_None, LOCTEXT("Toolbar_ApplyCompression", "Compression"));
 			}
 			ToolbarBuilder.EndSection();
+
+			ToolbarBuilder.BeginSection("Editing");
+			{
+				ToolbarBuilder.AddToolBarButton(FPersonaCommands::Get().SetKey, NAME_None, LOCTEXT("Toolbar_SetKey", "Key"));
+				ToolbarBuilder.AddToolBarButton(FPersonaCommands::Get().ApplyAnimation, NAME_None, LOCTEXT("Toolbar_ApplyAnimation", "Apply"));
+			}
+			ToolbarBuilder.EndSection();
 		}
 	};
 
@@ -1032,6 +1141,20 @@ void FPersona::CreateDefaultCommands()
 		FIsActionButtonVisible::CreateSP(this, &FPersona::IsInPersonaMode, FPersonaModes::AnimationEditMode)
 		);
 
+	ToolkitCommands->MapAction(FPersonaCommands::Get().SetKey,
+		FExecuteAction::CreateSP(this, &FPersona::OnSetKey),
+		FCanExecuteAction::CreateSP(this, &FPersona::CanSetKey),
+		FIsActionChecked(),
+		FIsActionButtonVisible::CreateSP(this, &FPersona::IsInPersonaMode, FPersonaModes::AnimationEditMode)
+		);
+
+	ToolkitCommands->MapAction(FPersonaCommands::Get().ApplyAnimation,
+			FExecuteAction::CreateSP(this, &FPersona::OnBakeAnimation),
+			FCanExecuteAction::CreateSP(this, &FPersona::CanBakeAnimation),
+			FIsActionChecked(),
+			FIsActionButtonVisible::CreateSP(this, &FPersona::IsInPersonaMode, FPersonaModes::AnimationEditMode)
+			);
+
 	ToolkitCommands->MapAction(FPersonaCommands::Get().ExportToFBX,
 		FExecuteAction::CreateSP(this, &FPersona::OnExportToFBX),
 		FCanExecuteAction::CreateSP(this, &FPersona::HasValidAnimationSequencePlaying),
@@ -1713,6 +1836,15 @@ void FPersona::OnPropertyChanged(UObject* ObjectBeingModified, FPropertyChangedE
 	//@TODO: Should we still do this?
 }
 
+void FPersona::RefreshPreviewInstanceTrackCurves()
+{
+	// need to refresh the preview mesh
+	if(PreviewComponent->PreviewInstance)
+	{
+		PreviewComponent->PreviewInstance->RefreshCurveBoneControllers();
+	}
+}
+
 void FPersona::PostUndo(bool bSuccess)
 {
 	DocumentManager->RefreshAllTabs();
@@ -1722,6 +1854,8 @@ void FPersona::PostUndo(bool bSuccess)
 	// PostUndo broadcast
 	OnPostUndo.Broadcast();	
 
+	RefreshPreviewInstanceTrackCurves();
+
 	// clear up preview anim notify states
 	// animnotify states are saved in AnimInstance
 	// if those are undoed or redoed, they have to be 
@@ -2405,6 +2539,44 @@ void FPersona::RecordAnimation()
 	}
 }
 
+void FPersona::OnSetKeyCompleted()
+{
+	OnTrackCurvesChanged.Broadcast();
+}
+
+bool FPersona::CanSetKey() const
+{
+	return ( HasValidAnimationSequencePlaying() && PreviewComponent->BonesOfInterest.Num() > 0);
+}
+
+void FPersona::OnSetKey()
+{
+	UAnimSequence * AnimSequence = Cast<UAnimSequence> (GetAnimationAssetBeingEdited());
+	if (AnimSequence)
+	{
+		UDebugSkelMeshComponent * Component = GetPreviewMeshComponent();
+		Component->PreviewInstance->SetKey(FSimpleDelegate::CreateSP(this, &FPersona::OnSetKeyCompleted));
+	}
+}
+
+bool FPersona::CanBakeAnimation() const
+{
+	UAnimSequence * AnimSequence = Cast<UAnimSequence> (GetAnimationAssetBeingEdited());
+	// ideally would be great if we can only show if something changed
+	return (AnimSequence && AnimSequence->DoesNeedRebake());
+}
+
+void FPersona::OnBakeAnimation()
+{
+	UAnimSequence * AnimSequence = Cast<UAnimSequence>(GetAnimationAssetBeingEdited());
+	if(AnimSequence)
+	{
+		UDebugSkelMeshComponent * Component = GetPreviewMeshComponent();
+		// now bake
+		Component->PreviewInstance->BakeAnimation();
+	}
+}
+
 void FPersona::OnApplyCompression()
 {
 	UAnimSequence * AnimSequence = Cast<UAnimSequence> (GetPreviewAnimationAsset());
@@ -2441,13 +2613,13 @@ void FPersona::OnAddLoopingInterpolation()
 	}
 }
 
-void FPersona::ApplyCompression(TArray<TWeakObjectPtr<UAnimSequence>> & AnimSequences)
+void FPersona::ApplyCompression(TArray<TWeakObjectPtr<UAnimSequence>>& AnimSequences)
 {
 	FDlgAnimCompression AnimCompressionDialog(AnimSequences);
 	AnimCompressionDialog.ShowModal();
 }
 
-void FPersona::ExportToFBX(TArray<TWeakObjectPtr<UAnimSequence>> & AnimSequences)
+void FPersona::ExportToFBX(TArray<TWeakObjectPtr<UAnimSequence>>& AnimSequences)
 {
 	IDesktopPlatform* DesktopPlatform = FDesktopPlatformModule::Get();
 
@@ -2581,7 +2753,7 @@ void FPersona::ExportToFBX(TArray<TWeakObjectPtr<UAnimSequence>> & AnimSequences
 	}
 }
 
-void FPersona::AddLoopingInterpolation(TArray<TWeakObjectPtr<UAnimSequence>> & AnimSequences)
+void FPersona::AddLoopingInterpolation(TArray<TWeakObjectPtr<UAnimSequence>>& AnimSequences)
 {
 	FText WarningMessage = LOCTEXT("AddLoopiingInterpolation", "This will add an extra first frame at the end of the animation to create a better looping interpolation. This action cannot be undone. Would you like to proceed?");
 
diff --git a/Engine/Source/Editor/Persona/Private/Persona.h b/Engine/Source/Editor/Persona/Private/Persona.h
--- a/Engine/Source/Editor/Persona/Private/Persona.h
+++ b/Engine/Source/Editor/Persona/Private/Persona.h
@@ -43,7 +43,7 @@ public:
 	/** Set the current preview viewport for Persona */
 	void SetViewport(TWeakPtr<class SAnimationEditorViewportTabBody> NewViewport);
 	/** Set Sequence browser */
-	void SetSequenceBrowser(class SAnimationSequenceBrowser * SequenceBrowser);
+	void SetSequenceBrowser(class SAnimationSequenceBrowser* SequenceBrowser);
 
 	/** Refresh viewport */
 	void RefreshViewport();
@@ -200,6 +200,8 @@ public:
 	const FSlateBrush* GetDirtyImageForMode(FName Mode) const;
 
 	TSharedRef<SKismetInspector> GetPreviewEditor() {return PreviewEditor.ToSharedRef();}
+	/** Refresh Preview Instance Track Curves **/
+	void RefreshPreviewInstanceTrackCurves();
 
 protected:
 	// FBlueprintEditor interface
@@ -294,34 +296,19 @@ public:
 	class IDetailLayoutBuilder* PersonaMeshDetailLayout;
 private:
 	// called when animation asset has been changed
-	DECLARE_MULTICAST_DELEGATE_OneParam( FOnAnimChangedMulticaster, UAnimationAsset* );
-	// Called after an undo is performed to give child widgets a chance to refresh
-	DECLARE_MULTICAST_DELEGATE( FOnPostUndoMulticaster );
+	DECLARE_MULTICAST_DELEGATE_OneParam( FOnAnimChangedMulticaster, UAnimationAsset* )
 	// Called when the preview mesh has been changed
-	DECLARE_MULTICAST_DELEGATE_OneParam(FOnPreviewMeshChangedMulticaster, class USkeletalMesh*);
-
+	DECLARE_MULTICAST_DELEGATE_OneParam(FOnPreviewMeshChangedMulticaster, class USkeletalMesh*)
 	// Called when a socket is selected
 	DECLARE_MULTICAST_DELEGATE_OneParam( FOnSelectSocket, const struct FSelectedSocketInfo& )
 	// Called when a bone is selected
 	DECLARE_MULTICAST_DELEGATE_OneParam( FOnSelectBone, const FName& )
-	// Called when selection is cleared
-	DECLARE_MULTICAST_DELEGATE( FOnDeselectAll )
-	// Called when the skeleton tree has been changed (socket added/deleted/etc)
-	DECLARE_MULTICAST_DELEGATE( FOnChangeSkeletonTree )
-	// Called when the notifies of the current animation are changed
-	DECLARE_MULTICAST_DELEGATE( FOnChangeAnimNotifies )
-	// Called when the curve panel is changed / updated
-	DECLARE_MULTICAST_DELEGATE( FOnChangeCurves )
 	// Called when the preview viewport is created
 	DECLARE_MULTICAST_DELEGATE_OneParam(FOnCreateViewport, TWeakPtr<class SAnimationEditorViewportTabBody>)
-	// Called when generic delete happens
-	DECLARE_MULTICAST_DELEGATE( FOnGenericDelete );
-	// Called when Persona refreshes
-	DECLARE_MULTICAST_DELEGATE( FOnPersonaRefreshMulticaster );
 public:
 
-	// Persona refreshed
-	typedef FOnPersonaRefreshMulticaster::FDelegate FOnPersonaRefresh;
+	// Called when Persona refreshes
+	typedef FSimpleMulticastDelegate::FDelegate FOnPersonaRefresh;
 
 	/** Registers a delegate to be called when Persona is Refreshed */
 	void RegisterOnPersonaRefresh(const FOnPersonaRefresh& Delegate)
@@ -350,8 +337,8 @@ public:
 		OnAnimChanged.RemoveAll(Widget);
 	}
 
-	// post undo 
-	typedef FOnPostUndoMulticaster::FDelegate FOnPostUndo;
+	// Called after an undo is performed to give child widgets a chance to refresh
+	typedef FSimpleMulticastDelegate::FDelegate FOnPostUndo;
 
 	/** Registers a delegate to be called after an Undo operation */
 	void RegisterOnPostUndo(const FOnPostUndo& Delegate)
@@ -410,8 +397,8 @@ public:
 		OnSocketSelected.RemoveAll( Widget );
 	}
 
-	// Deselected all
-	typedef FOnDeselectAll::FDelegate FOnAllDeselected;
+	// Called when selection is cleared
+	typedef FSimpleMulticastDelegate::FDelegate FOnAllDeselected;
 
 	/** Registers a delegate to be called when all bones/sockets are deselected */
 	void RegisterOnDeselectAll(const FOnAllDeselected& Delegate)
@@ -425,8 +412,8 @@ public:
 		OnAllDeselected.RemoveAll( Widget );
 	}
 
-	// Skeleton tree changed (socket duplicated, etc)
-	typedef FOnChangeSkeletonTree::FDelegate FOnSkeletonTreeChanged;
+	// Called when the skeleton tree has been changed (socket added/deleted/etc)
+	typedef FSimpleMulticastDelegate::FDelegate FOnSkeletonTreeChanged;
 
 	/** Registers a delegate to be called when the skeleton tree has changed */
 	void RegisterOnChangeSkeletonTree(const FOnSkeletonTreeChanged& Delegate)
@@ -440,8 +427,8 @@ public:
 		OnSkeletonTreeChanged.RemoveAll( Widget );
 	}
 
-	// Animation Notifies Changed
-	typedef FOnChangeAnimNotifies::FDelegate FOnAnimNotifiesChanged;
+	// Called when the notifies of the current animation are changed
+	typedef FSimpleMulticastDelegate::FDelegate FOnAnimNotifiesChanged;
 
 	/** Registers a delegate to be called when the skeleton anim notifies have been changed */
 	void RegisterOnChangeAnimNotifies(const FOnAnimNotifiesChanged& Delegate)
@@ -456,10 +443,10 @@ public:
 	}
 
 	/** Delegate for when the skeletons animation notifies have been changed */
-	FOnChangeAnimNotifies OnAnimNotifiesChanged;
+	FSimpleMulticastDelegate OnAnimNotifiesChanged;
 
-	// Curve changed
-	typedef FOnChangeCurves::FDelegate FOnCurvesChanged;
+	// Called when the curve panel is changed / updated
+	typedef FSimpleMulticastDelegate::FDelegate FOnCurvesChanged;
 
 	/** Registers delegate for changing / updating of curves panel */
 	void RegisterOnChangeCurves(const FOnCurvesChanged& Delegate)
@@ -474,7 +461,22 @@ public:
 	}
 
 	/** Delegate for changing / updating of curves panel */
-	FOnChangeCurves OnCurvesChanged;
+	FSimpleMulticastDelegate OnCurvesChanged;
+
+	// Called when the track curve is changed / updated
+	typedef FSimpleMulticastDelegate::FDelegate FOnTrackCurvesChanged;
+
+	/** Registers delegate for changing / updating of track curves panel */
+	void RegisterOnChangeTrackCurves(const FOnTrackCurvesChanged& Delegate)
+	{
+		OnTrackCurvesChanged.Add(Delegate);
+	}
+
+	/** Unregisters delegate for changing / updating of track curves panel */
+	void UnregisterOnChangeTrackCurves(SWidget* Widget)
+	{
+		OnTrackCurvesChanged.RemoveAll(Widget);
+	}
 
 	// Viewport Created
 	typedef FOnCreateViewport::FDelegate FOnViewportCreated;
@@ -491,7 +493,8 @@ public:
 		OnViewportCreated.RemoveAll( Widget );
 	}
 
-	typedef FOnGenericDelete::FDelegate FOnDeleteGeneric;
+	// Called when generic delete happens
+	typedef FSimpleMulticastDelegate::FDelegate FOnDeleteGeneric;
 
 	/** Registers a delegate to be called when Persona receives a generic delete command */
 	void RegisterOnGenericDelete(const FOnDeleteGeneric& Delegate)
@@ -506,11 +509,11 @@ public:
 	}
 
 	/** Apply Compression to list of animations */
-	void ApplyCompression(TArray<TWeakObjectPtr<UAnimSequence>> & AnimSequences);
+	void ApplyCompression(TArray<TWeakObjectPtr<UAnimSequence>>& AnimSequences);
 	/** Export to FBX files of the list of animations */
-	void ExportToFBX(TArray<TWeakObjectPtr<UAnimSequence>> & AnimSequences);
+	void ExportToFBX(TArray<TWeakObjectPtr<UAnimSequence>>& AnimSequences);
 	/** Add looping interpolation to the list of animations */
-	void AddLoopingInterpolation(TArray<TWeakObjectPtr<UAnimSequence>> & AnimSequences);
+	void AddLoopingInterpolation(TArray<TWeakObjectPtr<UAnimSequence>>& AnimSequences);
 
 protected:
 	/** Undo Action**/
@@ -521,10 +524,10 @@ protected:
 protected:
 
 	/** Called when persona is refreshed through an external action (reimport etc) */
-	FOnPersonaRefreshMulticaster OnPersonaRefresh;
+	FSimpleMulticastDelegate OnPersonaRefresh;
 
 	/** Delegate called after an undo operation for child widgets to refresh */
-	FOnPostUndoMulticaster OnPostUndo;	
+	FSimpleMulticastDelegate OnPostUndo;	
 
 	/**	Broadcasts whenever the animation changes */
 	FOnAnimChangedMulticaster OnAnimChanged;
@@ -539,21 +542,24 @@ protected:
 	FOnSelectBone OnBoneSelected;
 
 	/** Delegate for clearing the current skeleton bone/socket selection */
-	FOnDeselectAll OnAllDeselected;
+	FSimpleMulticastDelegate OnAllDeselected;
 
 	/** Delegate for when the skeleton tree has changed (e.g. a socket has been duplicated in the viewport) */
-	FOnChangeSkeletonTree OnSkeletonTreeChanged;
+	FSimpleMulticastDelegate OnSkeletonTreeChanged;
 
 	/** Delegate for when the preview viewport is created */
 	FOnCreateViewport OnViewportCreated;
 
+	/** Delegate for changing / updating of track curves panel */
+	FSimpleMulticastDelegate OnTrackCurvesChanged;
+
 	/**
 	 * Delegate for handling generic delete command
 	 * Register to this if you want the generic delete command (delete key) from the editor. Also
 	 * remember that the user may be using another part of the editor when you get the broadcast so
 	 * make sure not to delete unless the user intends it.
 	 */
-	FOnGenericDelete OnGenericDelete;
+	FSimpleMulticastDelegate OnGenericDelete;
 
 	/** Persona toolbar builder class */
 	TSharedPtr<class FPersonaToolbar> PersonaToolbar;
@@ -575,6 +581,13 @@ private:
 	bool HasValidAnimationSequencePlaying() const;
 	bool IsInPersonaMode(const FName InPersonaMode) const;
 
+	/** Animation Editing Features **/
+	void OnSetKey();
+	bool CanSetKey() const;
+	void OnSetKeyCompleted();
+	void OnBakeAnimation();
+	bool CanBakeAnimation() const;
+	
 	/** Change skeleton preview mesh functions */
 	void ChangeSkeletonPreviewMesh();
 	bool CanChangeSkeletonPreviewMesh() const;
@@ -591,6 +604,8 @@ private:
 	void OnReimportAnimation();
 	void OnAssetCreated(const TArray<UObject*> NewAssets);
 	TSharedRef< SWidget > GenerateCreateAssetMenu( USkeleton* Skeleton ) const;
+	void FillCreateAnimationMenu(FMenuBuilder& MenuBuilder) const;
+	void CreateAnimation(const TArray<UObject*> NewAssets, int32 Option);
 
 	/** Extend menu and toolbar */
 	void ExtendMenu();
diff --git a/Engine/Source/Editor/Persona/Private/PersonaCommands.cpp b/Engine/Source/Editor/Persona/Private/PersonaCommands.cpp
--- a/Engine/Source/Editor/Persona/Private/PersonaCommands.cpp
+++ b/Engine/Source/Editor/Persona/Private/PersonaCommands.cpp
@@ -29,6 +29,8 @@ void FPersonaCommands::RegisterCommands()
 	UI_COMMAND( RecordAnimation, "Record to new Animation", "Create new animation from currently playing", EUserInterfaceActionType::Button, FInputGesture() );
 	UI_COMMAND( ExportToFBX, "Export to FBX", "Export current animation to FBX", EUserInterfaceActionType::Button, FInputGesture() );
 	UI_COMMAND( AddLoopingInterpolation, "Add Looping Interpolation", "Add an extra first frame at the end of the animation to create interpolation when looping", EUserInterfaceActionType::Button, FInputGesture() );
+	UI_COMMAND( SetKey, "Set Key", "Add Bone Transform to Additive Layer Tracks", EUserInterfaceActionType::Button, FInputGesture(EKeys::S) );
+	UI_COMMAND( ApplyAnimation, "Apply Animation", "Apply Additive Layer Tracks to Runtime Animation Data", EUserInterfaceActionType::Button, FInputGesture() );
 
 	UI_COMMAND( UpdateSkeletonRefPose, "Update Skeleton RefPose", "Update Skeleton ref pose based on current preview mesh", EUserInterfaceActionType::Button, FInputGesture() );
 }
diff --git a/Engine/Source/Editor/Persona/Private/PersonaCommands.h b/Engine/Source/Editor/Persona/Private/PersonaCommands.h
--- a/Engine/Source/Editor/Persona/Private/PersonaCommands.h
+++ b/Engine/Source/Editor/Persona/Private/PersonaCommands.h
@@ -52,6 +52,10 @@ public:
 	TSharedPtr<FUICommandInfo> ExportToFBX;
 	// Add looping interpolation
 	TSharedPtr<FUICommandInfo> AddLoopingInterpolation;
+	// set key for bone track
+	TSharedPtr<FUICommandInfo> SetKey;
+	// bake bone track curves to animation
+	TSharedPtr<FUICommandInfo> ApplyAnimation;
 	// Create Asset command
 	TSharedPtr<FUICommandInfo> CreateAimOffset1D;
 	TSharedPtr<FUICommandInfo> CreateAimOffset2D;
diff --git a/Engine/Source/Editor/Persona/Private/SAnimCurvePanel.cpp b/Engine/Source/Editor/Persona/Private/SAnimCurvePanel.cpp
--- a/Engine/Source/Editor/Persona/Private/SAnimCurvePanel.cpp
+++ b/Engine/Source/Editor/Persona/Private/SAnimCurvePanel.cpp
@@ -5,7 +5,7 @@
 
 #include "SAnimCurvePanel.h"
 #include "ScopedTransaction.h"
-#include "SCurveEditor.h"
+#include "SAnimCurveEd.h"
 #include "Editor/KismetWidgets/Public/SScrubWidget.h"
 #include "AssetRegistryModule.h"
 #include "Kismet2NameValidators.h"
@@ -104,177 +104,6 @@ public:
 	{
 	}
 };
-
-//////////////////////////////////////////////////////////////////////////
-//  SAnimCurveEd : anim curve editor
-
-class SAnimCurveEd : public SCurveEditor
-{
-public:
-	SLATE_BEGIN_ARGS( SAnimCurveEd )
-		: _ViewMinInput(0.0f)
-		, _ViewMaxInput(10.0f)
-		, _TimelineLength(5.0f)
-		, _DrawCurve(true)
-		, _HideUI(true)
-		, _OnGetScrubValue()
-	{}
-	
-		SLATE_ATTRIBUTE( float, ViewMinInput )
-		SLATE_ATTRIBUTE( float, ViewMaxInput )
-		SLATE_ATTRIBUTE( TOptional<float>, DataMinInput )
-		SLATE_ATTRIBUTE( TOptional<float>, DataMaxInput )
-		SLATE_ATTRIBUTE( float, TimelineLength )
-		SLATE_ATTRIBUTE( int32, NumberOfKeys)
-		SLATE_ATTRIBUTE( FVector2D, DesiredSize )
-		SLATE_ARGUMENT( bool, DrawCurve )
-		SLATE_ARGUMENT( bool, HideUI )
-		SLATE_EVENT( FOnGetScrubValue, OnGetScrubValue )
-		SLATE_EVENT( FOnSetInputViewRange, OnSetInputViewRange )
-	SLATE_END_ARGS()
-
-	void Construct(const FArguments& InArgs);
-
-protected:
-	// SWidget interface
-	virtual int32 OnPaint(const FPaintArgs& Args, const FGeometry& AllottedGeometry, const FSlateRect& MyClippingRect, FSlateWindowElementList& OutDrawElements, int32 LayerId, const FWidgetStyle& InWidgetStyle, bool bParentEnabled) const override;
-	virtual FReply OnMouseWheel(const FGeometry& MyGeometry, const FPointerEvent& MouseEvent);
-	virtual FCursorReply OnCursorQuery(const FGeometry& MyGeometry, const FPointerEvent& CursorEvent) const override;
-	// SWidget interface
-
-	// SCurveEditor interface
-	virtual void SetDefaultOutput(const float MinZoomRange) override;
-	virtual float GetTimeStep(FTrackScaleInfo &ScaleInfo) const override;
-	// SCurveEditor interface
-	
-private:
-	// scrub value grabber
-	FOnGetScrubValue	OnGetScrubValue;
-	// @todo redo this code, all mess and dangerous
-	TAttribute<int32>	NumberOfKeys;
-};
-
-//////////////////////////////////////////////////////////////////////////
-//  SAnimCurveEd : anim curve editor
-
-float SAnimCurveEd::GetTimeStep(FTrackScaleInfo &ScaleInfo)const
-{
-	if(NumberOfKeys.Get())
-	{
-		int32 Divider = SScrubWidget::GetDivider( ViewMinInput.Get(), ViewMaxInput.Get(), ScaleInfo.WidgetSize, TimelineLength.Get(), NumberOfKeys.Get());
-
-		float TimePerKey;
-
-		if ( NumberOfKeys.Get() != 0.f )
-		{
-			TimePerKey = TimelineLength.Get()/(float)NumberOfKeys.Get();
-		}
-		else
-		{
-			TimePerKey = 1.f;
-		}
-
-		return TimePerKey * Divider;
-	}
-
-	return 0.0f;
-}
-
-int32 SAnimCurveEd::OnPaint(const FPaintArgs& Args, const FGeometry& AllottedGeometry, const FSlateRect& MyClippingRect, FSlateWindowElementList& OutDrawElements, int32 LayerId, const FWidgetStyle& InWidgetStyle, bool bParentEnabled) const
-{
-	int32 NewLayerId = SCurveEditor::OnPaint(Args, AllottedGeometry, MyClippingRect, OutDrawElements, LayerId, InWidgetStyle, bParentEnabled) + 1;
-
-	float Value = 0.f;
-
-	if ( OnGetScrubValue.IsBound() )
-	{
-		Value = OnGetScrubValue.Execute();
-	}
-
-	FPaintGeometry MyGeometry = AllottedGeometry.ToPaintGeometry();
-
-	// scale info
-	FTrackScaleInfo ScaleInfo(ViewMinInput.Get(), ViewMaxInput.Get(), 0.f, 0.f, AllottedGeometry.Size);
-	float XPos = ScaleInfo.InputToLocalX(Value);
-
-	TArray<FVector2D> LinePoints;
-	LinePoints.Add(FVector2D(XPos-1, 0.f));
-	LinePoints.Add(FVector2D(XPos+1, AllottedGeometry.Size.Y));
-
-
-	FSlateDrawElement::MakeLines( 
-		OutDrawElements,
-		NewLayerId,
-		MyGeometry,
-		LinePoints,
-		MyClippingRect,
-		ESlateDrawEffect::None,
-		FLinearColor::Red
-		);
-
-	// now draw scrub with new layer ID + 1;
-	return NewLayerId;
-}
-
-FReply SAnimCurveEd::OnMouseWheel(const FGeometry& MyGeometry, const FPointerEvent& MouseEvent)
-{
-	const float ZoomDelta = -0.1f * MouseEvent.GetWheelDelta();
-
-	const FVector2D WidgetSpace = MyGeometry.AbsoluteToLocal(MouseEvent.GetScreenSpacePosition());
-	const float ZoomRatio = FMath::Clamp((WidgetSpace.X / MyGeometry.Size.X), 0.f, 1.f);
-
-	{
-		const float InputViewSize = ViewMaxInput.Get() - ViewMinInput.Get();
-		const float InputChange = InputViewSize * ZoomDelta;
-
-		float NewViewMinInput = ViewMinInput.Get() - (InputChange * ZoomRatio);
-		float NewViewMaxInput = ViewMaxInput.Get() + (InputChange * (1.f - ZoomRatio));
-
-		SetInputMinMax(NewViewMinInput, NewViewMaxInput);
-	}
-
-	return FReply::Handled();
-}
-
-FCursorReply SAnimCurveEd::OnCursorQuery(const FGeometry& MyGeometry, const FPointerEvent& CursorEvent) const
-{
-	if (ViewMinInput.Get() > 0.f || ViewMaxInput.Get() < TimelineLength.Get())
-	{
-		return FCursorReply::Cursor(EMouseCursor::GrabHand);
-	}
-
-	return FCursorReply::Unhandled();
-}
-
-void SAnimCurveEd::Construct(const FArguments& InArgs)
-{
-	OnGetScrubValue = InArgs._OnGetScrubValue;
-	NumberOfKeys = InArgs._NumberOfKeys;
-
-	SCurveEditor::Construct( SCurveEditor::FArguments()
-		.ViewMinInput(InArgs._ViewMinInput)
-		.ViewMaxInput(InArgs._ViewMaxInput)
-		.DataMinInput(InArgs._DataMinInput)
-		.DataMaxInput(InArgs._DataMaxInput)
-		.ViewMinOutput(0.f)
-		.ViewMaxOutput(1.f)
-		.ZoomToFitVertical(true)
-		.ZoomToFitHorizontal(false)
-		.TimelineLength(InArgs._TimelineLength)
-		.DrawCurve(InArgs._DrawCurve)
-		.HideUI(InArgs._HideUI)
-		.AllowZoomOutput(false)
-		.DesiredSize(InArgs._DesiredSize)
-		.OnSetInputViewRange(InArgs._OnSetInputViewRange));
-}
-
-void SAnimCurveEd::SetDefaultOutput(const float MinZoomRange)
-{
-	const float NewMinOutput = (ViewMinOutput.Get());
-	const float NewMaxOutput = (ViewMaxOutput.Get() + MinZoomRange);
-
-	SetOutputMinMax(NewMinOutput, NewMaxOutput);
-}
 //////////////////////////////////////////////////////////////////////////
 //  SCurveEd Track : each track for curve editing 
 
@@ -381,8 +210,7 @@ void SCurveEdTrack::Construct(const FArguments& InArgs)
 		SAssignNew(InnerBox, SHorizontalBox)
 	];
 	
-	FFloatCurve* CurveData = (FFloatCurve*)Curve;
-	bool bIsMetadata = CurveData->GetCurveTypeFlag(ACF_Metadata);
+	bool bIsMetadata = Curve->GetCurveTypeFlag(ACF_Metadata);
 	if(!bIsMetadata)
 	{
 		InnerBox->AddSlot()
@@ -796,7 +624,7 @@ FReply SAnimCurvePanel::DuplicateTrack(USkeleton::AnimCurveUID Uid)
 		// Use the validator to pick a reasonable name for the duplicated curve.
 		FString NewCurveName = CurveNameToCopy.ToString();
 		Validator->FindValidString(NewCurveName);
-		if(NameMapping->AddName(*NewCurveName, NewUid))
+		if(NameMapping->AddOrFindName(*NewCurveName, NewUid))
 		{
 			if(Sequence->RawCurveData.DuplicateCurveData(Uid, NewUid))
 			{
@@ -1111,7 +939,7 @@ ESlateCheckBoxState::Type SAnimCurvePanel::IsCurveEditable(USkeleton::AnimCurveU
 {
 	if ( Sequence )
 	{
-		const FFloatCurve* Curve = Sequence->RawCurveData.GetCurveData(Uid);
+		const FFloatCurve* Curve = static_cast<const FFloatCurve *>(Sequence->RawCurveData.GetCurveData(Uid, FRawCurveTracks::FloatType));
 		if ( Curve )
 		{
 			return Curve->GetCurveTypeFlag(ACF_Editable)? ESlateCheckBoxState::Checked : ESlateCheckBoxState::Unchecked;
@@ -1127,7 +955,7 @@ void SAnimCurvePanel::ToggleEditability(ESlateCheckBoxState::Type NewType, USkel
 
 	if ( Sequence )
 	{
-		FFloatCurve * Curve = Sequence->RawCurveData.GetCurveData(Uid);
+		FFloatCurve * Curve = static_cast<FFloatCurve *>(Sequence->RawCurveData.GetCurveData(Uid, FRawCurveTracks::FloatType));
 		if ( Curve )
 		{
 			Curve->SetCurveTypeFlag(ACF_Editable, bEdit);
@@ -1244,7 +1072,7 @@ void SAnimCurvePanel::AddMetadataEntry(USkeleton::AnimCurveUID Uid)
 	if(Sequence->RawCurveData.AddCurveData(Uid))
 	{
 		Sequence->Modify(true);
-		FFloatCurve* Curve = Sequence->RawCurveData.GetCurveData(Uid);
+		FFloatCurve* Curve = static_cast<FFloatCurve *>(Sequence->RawCurveData.GetCurveData(Uid, FRawCurveTracks::FloatType));
 		Curve->FloatCurve.AddKey(0.0f, 1.0f);
 		Curve->SetCurveTypeFlag(ACF_Metadata, true);
 		RefreshPanel();
diff --git a/Engine/Source/Editor/Persona/Private/SAnimTrackCurvePanel.cpp b/Engine/Source/Editor/Persona/Private/SAnimTrackCurvePanel.cpp
new file mode 100644
--- /dev/null
+++ b/Engine/Source/Editor/Persona/Private/SAnimTrackCurvePanel.cpp
@@ -0,0 +1,785 @@
+// Copyright 1998-2014 Epic Games, Inc. All Rights Reserved.
+
+
+#include "PersonaPrivatePCH.h"
+
+#include "SAnimTrackCurvePanel.h"
+#include "ScopedTransaction.h"
+#include "SAnimCurveEd.h"
+#include "Editor/KismetWidgets/Public/SScrubWidget.h"
+#include "AssetRegistryModule.h"
+#include "Kismet2NameValidators.h"
+#include "SExpandableArea.h"
+#include "STextEntryPopup.h"
+
+#define LOCTEXT_NAMESPACE "AnimTrackCurvePanel"
+
+//////////////////////////////////////////////////////////////////////////
+//  FAnimTrackCurveInterface interface 
+
+/** Interface you implement if you want the CurveEditors to be able to edit curves on you */
+class FAnimTrackCurveBaseInterface : public FCurveOwnerInterface
+{
+public:
+	TWeakObjectPtr<UAnimSequence>	BaseSequence;
+	FAnimCurveBase*	CurveData;
+	bool			bRotationCurve;
+
+public:
+	FAnimTrackCurveBaseInterface(UAnimSequence * BaseSeq, FAnimCurveBase*	InData, bool bInRotationCurve)
+		: BaseSequence(BaseSeq)
+		, CurveData(InData)
+		, bRotationCurve(bInRotationCurve)
+	{
+		// they should be valid
+		check (BaseSequence.IsValid());
+		check (CurveData);
+	}
+
+	FName GetDisplayCurveName(int32 Index) const
+	{
+		if(bRotationCurve)
+		{
+			switch(Index)
+			{
+			case 1:
+			return FName(TEXT("Pitch"));
+			case 2:
+			return FName(TEXT("Yaw"));
+			default:
+			return FName(TEXT("Roll"));
+			}
+		}
+		else
+		{
+			switch(Index)
+			{
+			case 1:
+			return FName(TEXT("Y"));
+			case 2:
+			return FName(TEXT("Z"));
+			default:
+			return FName(TEXT("X"));
+			}
+		}
+	}
+
+	/** Returns set of curves to edit. Must not release the curves while being edited. */
+	virtual TArray<FRichCurveEditInfoConst> GetCurves() const
+	{
+		TArray<FRichCurveEditInfoConst> Curves;
+		FVectorCurve * VectorCurveData = (FVectorCurve*)(CurveData);
+		Curves.Add(FRichCurveEditInfoConst(&VectorCurveData->FloatCurves[0], GetDisplayCurveName(0)));
+		Curves.Add(FRichCurveEditInfoConst(&VectorCurveData->FloatCurves[1], GetDisplayCurveName(1)));
+		Curves.Add(FRichCurveEditInfoConst(&VectorCurveData->FloatCurves[2], GetDisplayCurveName(2)));
+		return Curves;
+	}
+	/** Returns set of curves to query. Must not release the curves while being edited. */
+	virtual TArray<FRichCurveEditInfo> GetCurves()
+	{
+		TArray<FRichCurveEditInfo> Curves;
+		FVectorCurve * VectorCurveData = (FVectorCurve*)(CurveData);
+		Curves.Add(FRichCurveEditInfo(&VectorCurveData->FloatCurves[0], GetDisplayCurveName(0)));
+		Curves.Add(FRichCurveEditInfo(&VectorCurveData->FloatCurves[1], GetDisplayCurveName(1)));
+		Curves.Add(FRichCurveEditInfo(&VectorCurveData->FloatCurves[2], GetDisplayCurveName(2)));
+		return Curves;
+	}
+
+	virtual UObject* GetOwner()
+	{
+		if (BaseSequence.IsValid())
+		{
+			return BaseSequence.Get();
+		}
+
+		return NULL;
+	}
+
+	/** Called to modify the owner of the curve */
+	virtual void ModifyOwner()
+	{
+		if (BaseSequence.IsValid())
+		{
+			// need to rebake
+			BaseSequence.Get()->bNeedsRebake = true;
+			BaseSequence.Get()->Modify(true);
+		}
+	}
+
+	/** Called to make curve owner transactional */
+	virtual void MakeTransactional()
+	{
+		if (BaseSequence.IsValid())
+		{
+			BaseSequence.Get()->SetFlags(BaseSequence.Get()->GetFlags() | RF_Transactional);
+		}
+	}
+
+	/** Called to get the name of a curve back from the animation skeleton */
+	virtual FText GetCurveName(USkeleton::AnimCurveUID Uid) const
+	{
+		if(BaseSequence.IsValid())
+		{
+			FSmartNameMapping* NameMapping = BaseSequence.Get()->GetSkeleton()->SmartNames.GetContainer(USkeleton::AnimTrackCurveMappingName);
+			if(NameMapping)
+			{
+				FName CurveName;
+				if(NameMapping->GetName(Uid, CurveName))
+				{
+					FName DisplayName = CurveName;
+
+					return FText::FromString(FString::Printf(TEXT("%s(%c)"), *DisplayName.ToString(), (bRotationCurve)? TCHAR('R') : TCHAR('T')));
+				}
+			}
+		}
+
+		return FText::GetEmpty();
+	}
+
+	virtual void OnCurveChanged() override
+	{
+	}
+};
+
+//////////////////////////////////////////////////////////////////////////
+//  SCurveEd Track : each track for curve editing 
+
+/** Widget for editing a single track of animation curve - this includes CurveEditors */
+class STransformCurveEdTrack : public SCompoundWidget
+{
+private:
+	/** Pointer to notify panel for drawing*/
+	TSharedPtr<class SCurveEditor>			CurveEditors[2];
+
+	/** Name of curve it's editing - CurveName should be unique within this tracks**/
+	FAnimTrackCurveBaseInterface	*		CurveInterfaces[2];
+
+	/** Curve Panel Ptr **/
+	TWeakPtr<SAnimTrackCurvePanel>			PanelPtr;
+
+	enum EType
+	{
+		Translation = 0,
+		Rotation,
+	};
+
+private:
+	USkeleton::AnimCurveUID	CurveUid;
+
+public:
+	SLATE_BEGIN_ARGS( STransformCurveEdTrack )
+		: _AnimTrackCurvePanel()
+		, _Sequence()
+		, _CurveUid()
+		, _WidgetWidth()
+		, _ViewInputMin()
+		, _ViewInputMax()
+		, _OnSetInputViewRange()
+		, _OnGetScrubValue()
+	{}
+	SLATE_ARGUMENT( TSharedPtr<SAnimTrackCurvePanel>, AnimTrackCurvePanel)
+	// editing related variables
+	SLATE_ARGUMENT( class UAnimSequence*, Sequence )
+	SLATE_ARGUMENT( USkeleton::AnimCurveUID, CurveUid )
+	// widget viewing related variables
+	SLATE_ARGUMENT( float, WidgetWidth ) // @todo do I need this?
+	SLATE_ATTRIBUTE( float, ViewInputMin )
+	SLATE_ATTRIBUTE( float, ViewInputMax )
+	SLATE_EVENT( FOnSetInputViewRange, OnSetInputViewRange )
+	SLATE_EVENT( FOnGetScrubValue, OnGetScrubValue )
+	SLATE_END_ARGS()
+
+	void Construct(const FArguments& InArgs);
+	virtual ~STransformCurveEdTrack();
+
+
+	/**
+	 * Build and display curve track context menu.
+	 *
+	 */
+	FReply OnContextMenu();
+
+	// expand editor mode 
+	ESlateCheckBoxState::Type IsEditorExpanded() const;
+	void ToggleExpandEditor(ESlateCheckBoxState::Type NewType);
+	const FSlateBrush* GetExpandContent() const;
+	FVector2D GetDesiredSize() const;
+
+	// Bound to attribute for curve name, uses curve interface to request from skeleton
+	FText GetCurveName(USkeleton::AnimCurveUID Uid, EType Type) const;
+};
+
+//////////////////////////////////////////////////////////////////////////
+// STransformCurveEdTrack
+
+void STransformCurveEdTrack::Construct(const FArguments& InArgs)
+{
+	TSharedRef<SAnimTrackCurvePanel> PanelRef = InArgs._AnimTrackCurvePanel.ToSharedRef();
+	PanelPtr = InArgs._AnimTrackCurvePanel;
+	
+	// now create CurveInterfaces, 
+	// find which curve this belongs to
+	UAnimSequence * Sequence = InArgs._Sequence;
+	check (Sequence);
+
+	// get the curve data
+	FTransformCurve* Curve = static_cast<FTransformCurve*> (Sequence->RawCurveData.GetCurveData(InArgs._CurveUid, FRawCurveTracks::TransformType));
+	check (Curve);
+
+	CurveUid = InArgs._CurveUid;
+
+	CurveInterfaces[EType::Translation] = new FAnimTrackCurveBaseInterface(Sequence, &Curve->TranslationCurve, false);
+	CurveInterfaces[EType::Rotation] = new FAnimTrackCurveBaseInterface(Sequence, &Curve->RotationCurve, true);
+	int32 NumberOfKeys = Sequence->GetNumberOfFrames();
+	//////////////////////////////
+	
+	this->ChildSlot
+	[
+		SNew(SBorder)
+		.Padding(FMargin(2.0f, 2.0f))
+		[
+			SNew(SVerticalBox)
+
+			+SVerticalBox::Slot()
+			[
+				SNew(SHorizontalBox)
+
+				+SHorizontalBox::Slot()
+				.FillWidth(1)
+				[
+					// Notification editor panel
+					SAssignNew(CurveEditors[EType::Translation], SAnimCurveEd)
+					.ViewMinInput(InArgs._ViewInputMin)
+					.ViewMaxInput(InArgs._ViewInputMax)
+					.DataMinInput(0.f)
+					.DataMaxInput(Sequence->SequenceLength)
+					// @fixme fix this to delegate
+					.TimelineLength(Sequence->SequenceLength)
+					.NumberOfKeys(NumberOfKeys)
+					.DesiredSize(this, &STransformCurveEdTrack::GetDesiredSize)
+					.OnSetInputViewRange(InArgs._OnSetInputViewRange)
+					.OnGetScrubValue(InArgs._OnGetScrubValue)
+				]
+
+				+SHorizontalBox::Slot()
+				.AutoWidth()
+				.Padding(2.f)
+				[
+					SNew(SBox)
+					.WidthOverride(InArgs._WidgetWidth)
+					[
+						SNew(SHorizontalBox)
+
+						+SHorizontalBox::Slot()
+						.FillWidth(1)
+						.VAlign(VAlign_Center)
+						.HAlign(HAlign_Left)
+						.Padding(5, 0, 0, 0)
+						[
+							// Name of track
+							SNew(SEditableText)
+							.MinDesiredWidth(64.0f)
+							.IsEnabled(true)
+							.Font(FEditorStyle::GetFontStyle("CurveEd.InfoFont"))
+							.SelectAllTextWhenFocused(true)
+							.IsReadOnly(true)
+							.Text(this, &STransformCurveEdTrack::GetCurveName, Curve->CurveUid, EType::Translation)
+						]
+
+						+SHorizontalBox::Slot()
+						.Padding(FMargin(0.0f, 5.0f, 0.0f, 5.0f))
+						.AutoWidth()
+						.VAlign(VAlign_Top)
+						[
+							SNew(SButton)
+							.ToolTipText(LOCTEXT("DisplayTrackOptionsMenuTooltip", "Display track options menu"))
+							.OnClicked(this, &STransformCurveEdTrack::OnContextMenu)
+							.Content()
+							[
+								SNew(SImage)
+								.Image(FEditorStyle::GetBrush("ComboButton.Arrow"))
+								.ColorAndOpacity(FSlateColor::UseForeground())
+							]
+						]
+					]
+				]
+			]
+
+			+SVerticalBox::Slot()
+			[
+				SNew(SHorizontalBox)
+
+				+SHorizontalBox::Slot()
+				.FillWidth(1)
+				[
+					// Notification editor panel
+					SAssignNew(CurveEditors[EType::Rotation], SAnimCurveEd)
+					.ViewMinInput(InArgs._ViewInputMin)
+					.ViewMaxInput(InArgs._ViewInputMax)
+					.DataMinInput(0.f)
+					.DataMaxInput(Sequence->SequenceLength)
+					// @fixme fix this to delegate
+					.TimelineLength(Sequence->SequenceLength)
+					.NumberOfKeys(NumberOfKeys)
+					.DesiredSize(this, &STransformCurveEdTrack::GetDesiredSize)
+					.OnSetInputViewRange(InArgs._OnSetInputViewRange)
+					.OnGetScrubValue(InArgs._OnGetScrubValue)
+				]
+
+				+SHorizontalBox::Slot()
+				.AutoWidth()
+				.Padding(2.f)
+				[
+					SNew(SBox)
+					.WidthOverride(InArgs._WidgetWidth)
+					[
+
+						SNew(SHorizontalBox)
+
+						+SHorizontalBox::Slot()
+						.FillWidth(1)
+						.VAlign(VAlign_Center)
+						.HAlign(HAlign_Left)
+						.Padding(5, 0, 0, 0)
+						[
+							// Name of track
+							SNew(SEditableText)
+							.MinDesiredWidth(64.0f)
+							.IsEnabled(true)
+							.Font(FEditorStyle::GetFontStyle("CurveEd.InfoFont"))
+							.SelectAllTextWhenFocused(true)
+							.IsReadOnly(true)
+							.Text(this, &STransformCurveEdTrack::GetCurveName, Curve->CurveUid, EType::Rotation)
+						]
+					]
+				]
+			]
+
+		]
+	];
+	
+	//Inform track widget about the curve and whether it is editable or not.
+	CurveEditors[EType::Translation]->SetCurveOwner(CurveInterfaces[EType::Translation], true);
+	CurveEditors[EType::Rotation]->SetCurveOwner(CurveInterfaces[EType::Rotation], true);
+}
+
+/** return a widget */
+const FSlateBrush* STransformCurveEdTrack::GetExpandContent() const
+{
+	return FEditorStyle::GetBrush("Kismet.VariableList.HideForInstance");
+}
+
+STransformCurveEdTrack::~STransformCurveEdTrack()
+{
+	// @fixme - check -is this okay way of doing it?
+	// @fix this - funny
+	delete CurveInterfaces[EType::Translation];
+	delete CurveInterfaces[EType::Rotation];
+}
+
+FReply STransformCurveEdTrack::OnContextMenu()
+{
+	TSharedPtr<SAnimTrackCurvePanel> PanelShared = PanelPtr.Pin();
+
+	if(PanelShared.IsValid())
+	{
+		FSlateApplication::Get().PushMenu(SharedThis(this),
+										  PanelShared->CreateCurveContextMenu(CurveUid),
+										  FSlateApplication::Get().GetCursorPos(),
+										  FPopupTransitionEffect(FPopupTransitionEffect::TypeInPopup));
+	}
+
+	return FReply::Handled();
+}
+
+FVector2D STransformCurveEdTrack::GetDesiredSize() const
+{
+	return FVector2D(128.f, 128.f);
+}
+
+FText STransformCurveEdTrack::GetCurveName(USkeleton::AnimCurveUID Uid, EType Type) const
+{
+	return CurveInterfaces[Type]->GetCurveName(Uid);
+}
+
+//////////////////////////////////////////////////////////////////////////
+// SAnimTrackCurvePanel
+
+void SAnimTrackCurvePanel::Construct(const FArguments& InArgs)
+{
+	SAnimTrackPanel::Construct( SAnimTrackPanel::FArguments()
+		.WidgetWidth(InArgs._WidgetWidth)
+		.ViewInputMin(InArgs._ViewInputMin)
+		.ViewInputMax(InArgs._ViewInputMax)
+		.InputMin(InArgs._InputMin)
+		.InputMax(InArgs._InputMax)
+		.OnSetInputViewRange(InArgs._OnSetInputViewRange));
+
+	WeakPersona = InArgs._Persona;
+	Sequence = InArgs._Sequence;
+	WidgetWidth = InArgs._WidgetWidth;
+	OnGetScrubValue = InArgs._OnGetScrubValue;
+
+	this->ChildSlot
+	[
+		SNew(SVerticalBox)
+		+SVerticalBox::Slot()
+		.FillHeight(1)
+		[
+			SNew( SExpandableArea )
+			.AreaTitle( LOCTEXT("TransformCurve_Title", "Tracks") )
+			.BodyContent()
+			[
+				SNew(SVerticalBox)
+				+SVerticalBox::Slot()
+				.AutoHeight()
+				[
+					SNew(SHorizontalBox)
+					+SHorizontalBox::Slot()
+					.AutoWidth()
+					.Padding(5)
+					[
+						SNew(STextBlock)
+						.Text(LOCTEXT("AdditiveLayerTrackList_Title", "Additive Layer Tracks"))
+						.Font( FEditorStyle::GetFontStyle( "PropertyWindow.BoldFont" ) )
+						.ShadowOffset( FVector2D(1.0f, 1.0f) )
+					]
+
+					+SHorizontalBox::Slot()
+					.AutoWidth()
+					[
+						SNew(SComboButton)
+						.ContentPadding(FMargin(2.0f))
+						.OnGetMenuContent(this, &SAnimTrackCurvePanel::GenerateCurveList)
+					]
+				]
+
+				+SVerticalBox::Slot()
+				.Padding(FMargin(0.0f, 5.0f, 0.0f, 0.0f))
+				.AutoHeight()
+				[
+					SAssignNew(PanelSlot, SSplitter)
+					.Orientation(Orient_Vertical)
+				]
+			]
+		]
+	];
+
+	UpdatePanel();
+}
+
+void SAnimTrackCurvePanel::DeleteTrack(USkeleton::AnimCurveUID Uid)
+{
+	if(Sequence->RawCurveData.GetCurveData(Uid, FRawCurveTracks::TransformType))
+	{
+		const FScopedTransaction Transaction( LOCTEXT("AnimCurve_DeleteTrack", "Delete Curve") );
+		Sequence->Modify(true);
+		Sequence->bNeedsRebake = true;
+		Sequence->RawCurveData.DeleteCurveData(Uid, FRawCurveTracks::TransformType);
+		UpdatePanel();
+		WeakPersona.Pin()->RefreshPreviewInstanceTrackCurves();
+	}
+}
+
+void SAnimTrackCurvePanel::UpdatePanel()
+{
+	if(Sequence != NULL)
+	{
+		USkeleton* CurrentSkeleton = Sequence->GetSkeleton();
+		FSmartNameMapping* MetadataNameMap = CurrentSkeleton->SmartNames.GetContainer(USkeleton::AnimTrackCurveMappingName);
+		// Sort the raw curves before setting up display
+		Sequence->RawCurveData.TransformCurves.Sort([MetadataNameMap](const FTransformCurve& A, const FTransformCurve& B)
+		{
+			bool bAMeta = A.GetCurveTypeFlag(ACF_Metadata);
+			bool bBMeta = B.GetCurveTypeFlag(ACF_Metadata);
+			
+			if(bAMeta != bBMeta)
+			{
+				return !bAMeta;
+			}
+
+			FName AName;
+			FName BName;
+			MetadataNameMap->GetName(A.CurveUid, AName);
+			MetadataNameMap->GetName(B.CurveUid, BName);
+
+			return AName < BName;
+		});
+
+		// see if we need to clear or not
+		FChildren * VariableChildren = PanelSlot->GetChildren();
+		for (int32 Id=VariableChildren->Num()-1; Id>=0; --Id)
+		{
+			PanelSlot->RemoveAt(Id);
+		}
+
+		// Clear all tracks as we're re-adding them all anyway.
+		Tracks.Empty();
+
+		// Updating new tracks
+		FSmartNameMapping* NameMapping = CurrentSkeleton->SmartNames.GetContainer(USkeleton::AnimTrackCurveMappingName);
+
+		const int32 TotalCurve = Sequence->RawCurveData.TransformCurves.Num();
+		for(int32 CurrentIt = 0 ; CurrentIt < TotalCurve ; ++CurrentIt)
+		{
+			FTransformCurve&  Curve = Sequence->RawCurveData.TransformCurves[CurrentIt];
+
+			const bool bEditable = Curve.GetCurveTypeFlag(ACF_Editable);
+			FName CurveName;
+
+			// if editable, add to the list
+			if(bEditable && NameMapping->GetName(Curve.CurveUid, CurveName))
+			{
+				TSharedPtr<STransformCurveEdTrack> CurrentTrack;
+				PanelSlot->AddSlot()
+				.SizeRule(SSplitter::SizeToContent)
+				[
+					SNew(SVerticalBox)
+					+ SVerticalBox::Slot()
+					.AutoHeight()
+					.VAlign(VAlign_Center)
+					[
+						SAssignNew(CurrentTrack, STransformCurveEdTrack)
+						.Sequence(Sequence)
+						.CurveUid(Curve.CurveUid)
+						.AnimTrackCurvePanel(SharedThis(this))
+						.WidgetWidth(WidgetWidth)
+						.ViewInputMin(ViewInputMin)
+						.ViewInputMax(ViewInputMax)
+						.OnGetScrubValue(OnGetScrubValue)
+						.OnSetInputViewRange(OnSetInputViewRange)
+					]
+				];
+				Tracks.Add(CurrentTrack);
+			}
+		}
+
+		TSharedPtr<FPersona> SharedPersona = WeakPersona.Pin();
+		if(SharedPersona.IsValid())
+		{
+			SharedPersona->OnCurvesChanged.Broadcast();
+		}
+	}
+}
+
+TSharedRef<SWidget> SAnimTrackCurvePanel::GenerateCurveList()
+{
+	TSharedPtr<SVerticalBox> MainBox, ListBox;
+	TSharedRef<SWidget> NewWidget = SAssignNew(MainBox, SVerticalBox);
+
+	if ( Sequence && Sequence->RawCurveData.TransformCurves.Num() > 0)
+	{
+		MainBox->AddSlot()
+			.AutoHeight()
+			.MaxHeight(300)
+			[
+				SNew( SScrollBox )
+				+SScrollBox::Slot() 
+				[
+					SAssignNew(ListBox, SVerticalBox)
+				]
+			];
+
+		// Mapping to retrieve curve names
+		FSmartNameMapping* NameMapping = Sequence->GetSkeleton()->SmartNames.GetContainer(USkeleton::AnimTrackCurveMappingName);
+		check(NameMapping);
+
+		for (auto Iter=Sequence->RawCurveData.TransformCurves.CreateConstIterator(); Iter; ++Iter)
+		{
+			const FTransformCurve& Curve= *Iter;
+
+			FName CurveName;
+			NameMapping->GetName(Curve.CurveUid, CurveName);
+
+			ListBox->AddSlot()
+				.AutoHeight()
+				.VAlign(VAlign_Center)
+				.Padding( 2.0f, 2.0f )
+				[
+					SNew( SCheckBox )
+					.IsChecked(this, &SAnimTrackCurvePanel::IsCurveEditable, Curve.CurveUid)
+					.OnCheckStateChanged(this, &SAnimTrackCurvePanel::ToggleEditability, Curve.CurveUid)
+					.ToolTipText( LOCTEXT("Show Curves", "Show or Hide Curves") )
+					.IsEnabled( true )
+					[
+						SNew( STextBlock )
+						.Text(FText::FromName(CurveName))
+					]
+				];
+		}
+
+		MainBox->AddSlot()
+			.AutoHeight()
+			.VAlign(VAlign_Center)
+			.Padding( 2.0f, 2.0f )
+			[
+				SNew( SButton )
+				.VAlign( VAlign_Center )
+				.HAlign( HAlign_Center )
+				.OnClicked( this, &SAnimTrackCurvePanel::RefreshPanel )
+				[
+					SNew( STextBlock )
+					.Text( LOCTEXT("RefreshCurve", "Refresh") )
+				]
+			];
+
+		MainBox->AddSlot()
+			.AutoHeight()
+			.VAlign(VAlign_Center)
+			.Padding( 2.0f, 2.0f )
+			[
+				SNew( SButton )
+				.VAlign( VAlign_Center )
+				.HAlign( HAlign_Center )
+				.OnClicked( this, &SAnimTrackCurvePanel::ShowAll, true )
+				[
+					SNew( STextBlock )
+					.Text( LOCTEXT("ShowAll", "Show All") )
+				]
+			];
+
+		MainBox->AddSlot()
+			.AutoHeight()
+			.VAlign(VAlign_Center)
+			.Padding( 2.0f, 2.0f )
+			[
+				SNew( SButton )
+				.VAlign( VAlign_Center )
+				.HAlign( HAlign_Center )
+				.OnClicked( this, &SAnimTrackCurvePanel::ShowAll, false )
+				[
+					SNew( STextBlock )
+					.Text( LOCTEXT("HideAll", "Hide All") )
+				]
+			];
+	}
+	else
+	{
+		MainBox->AddSlot()
+			.AutoHeight()
+			.VAlign(VAlign_Center)
+			.Padding( 2.0f, 2.0f )
+			[
+				SNew( STextBlock )
+				.Text(LOCTEXT("Not Available", "No curve exists"))
+			];
+	}
+
+	return NewWidget;
+}
+
+ESlateCheckBoxState::Type SAnimTrackCurvePanel::IsCurveEditable(USkeleton::AnimCurveUID Uid) const
+{
+	if ( Sequence )
+	{
+		const FTransformCurve* Curve = static_cast<const FTransformCurve *>(Sequence->RawCurveData.GetCurveData(Uid, FRawCurveTracks::TransformType));
+		if ( Curve )
+		{
+			return Curve->GetCurveTypeFlag(ACF_Editable)? ESlateCheckBoxState::Checked : ESlateCheckBoxState::Unchecked;
+		}
+	}
+
+	return ESlateCheckBoxState::Undetermined;
+}
+
+void SAnimTrackCurvePanel::ToggleEditability(ESlateCheckBoxState::Type NewType, USkeleton::AnimCurveUID Uid)
+{
+	bool bEdit = (NewType == ESlateCheckBoxState::Checked);
+
+	if ( Sequence )
+	{
+		FTransformCurve * Curve = static_cast<FTransformCurve *>(Sequence->RawCurveData.GetCurveData(Uid, FRawCurveTracks::TransformType));
+		if ( Curve )
+		{
+			Curve->SetCurveTypeFlag(ACF_Editable, bEdit);
+		}
+	}
+}
+
+FReply		SAnimTrackCurvePanel::RefreshPanel()
+{
+	UpdatePanel();
+	return FReply::Handled();
+}
+
+FReply		SAnimTrackCurvePanel::ShowAll(bool bShow)
+{
+	if ( Sequence )
+	{
+		for (auto Iter = Sequence->RawCurveData.TransformCurves.CreateIterator(); Iter; ++Iter)
+		{
+			FTransformCurve & Curve = *Iter;
+			Curve.SetCurveTypeFlag(ACF_Editable, bShow);
+		}
+
+		UpdatePanel();
+	}
+
+	return FReply::Handled();
+}
+
+ESlateCheckBoxState::Type SAnimTrackCurvePanel::GetCurveFlagAsCheckboxState(USkeleton::AnimCurveUID CurveUid, EAnimCurveFlags InFlag) const
+{
+	FAnimCurveBase* Curve = Sequence->RawCurveData.GetCurveData(CurveUid, FRawCurveTracks::TransformType);
+	return Curve && Curve->GetCurveTypeFlag(InFlag) ? ESlateCheckBoxState::Checked : ESlateCheckBoxState::Unchecked;
+}
+
+void SAnimTrackCurvePanel::SetCurveFlagFromCheckboxState(ESlateCheckBoxState::Type CheckState, USkeleton::AnimCurveUID CurveUid, EAnimCurveFlags InFlag)
+{
+	bool Enabled = CheckState == ESlateCheckBoxState::Checked;
+	FAnimCurveBase* Curve = Sequence->RawCurveData.GetCurveData(CurveUid, FRawCurveTracks::TransformType);
+	if (Curve)
+	{
+		Curve->SetCurveTypeFlag(InFlag, Enabled);
+
+		if (InFlag == ACF_Disabled)
+		{
+			// needs to rebake
+			Sequence->bNeedsRebake = true;
+			// need to update curves, otherwise they're not disabled
+			WeakPersona.Pin()->RefreshPreviewInstanceTrackCurves();		
+		}
+	}
+}
+
+TSharedRef<SWidget> SAnimTrackCurvePanel::CreateCurveContextMenu(USkeleton::AnimCurveUID CurveUid) const
+{
+	FMenuBuilder MenuBuilder(true, NULL);
+
+	// get the curve data
+	FAnimCurveBase* Curve = Sequence->RawCurveData.GetCurveData(CurveUid, FRawCurveTracks::TransformType);
+	if (Curve)
+	{
+		MenuBuilder.BeginSection("AnimTrackCurvePanelCurveTypes", LOCTEXT("CurveTypesHeading", "Curve Types"));
+		{
+			MenuBuilder.AddWidget(
+				SNew(SCheckBox)
+				.IsChecked(this, &SAnimTrackCurvePanel::GetCurveFlagAsCheckboxState, CurveUid, ACF_Disabled)
+				.OnCheckStateChanged(this, &SAnimTrackCurvePanel::SetCurveFlagFromCheckboxState, CurveUid, ACF_Disabled)
+				.ToolTipText(LOCTEXT("DisableCurveTooltip", "Disable Track"))
+				[
+					SNew(STextBlock)
+					.Text(LOCTEXT("DisableCurveTextLabel", "Disable Curve"))
+				],
+				FText()
+				);
+		}
+		MenuBuilder.EndSection();
+
+		MenuBuilder.BeginSection("AnimTrackCurvePanelTrackOptions", LOCTEXT("TrackOptionsHeading", "Track Options"));
+		{
+			FUIAction NewAction;
+
+			NewAction.ExecuteAction.BindSP(this, &SAnimTrackCurvePanel::DeleteTrack, CurveUid);
+			MenuBuilder.AddMenuEntry(
+				LOCTEXT("RemoveTrack", "Remove Track"),
+				LOCTEXT("RemoveTrackTooltip", "Remove this track"),
+				FSlateIcon(),
+				NewAction);
+		}
+		MenuBuilder.EndSection();
+	}
+
+	return MenuBuilder.MakeWidget();
+}
+#undef LOCTEXT_NAMESPACE
diff --git a/Engine/Source/Editor/Persona/Private/SAnimTrackCurvePanel.h b/Engine/Source/Editor/Persona/Private/SAnimTrackCurvePanel.h
new file mode 100644
--- /dev/null
+++ b/Engine/Source/Editor/Persona/Private/SAnimTrackCurvePanel.h
@@ -0,0 +1,116 @@
+// Copyright 1998-2014 Epic Games, Inc. All Rights Reserved.
+
+
+#pragma once
+
+#include "Persona.h"
+#include "GraphEditor.h"
+#include "SNodePanel.h"
+#include "SAnimTrackPanel.h"
+#include "SCurveEditor.h"
+#include "Animation/AnimSequenceBase.h"
+
+//////////////////////////////////////////////////////////////////////////
+// SAnimTrackCurvePanel
+
+class SAnimTrackCurvePanel: public SAnimTrackPanel
+{
+public:
+	SLATE_BEGIN_ARGS( SAnimTrackCurvePanel )
+		: _Sequence()
+		, _ViewInputMin()
+		, _ViewInputMax()
+		, _InputMin()
+		, _InputMax()
+		, _OnSetInputViewRange()
+		, _OnGetScrubValue()
+	{}
+
+	/**
+	 * Weak ptr to our Persona instance
+	 */
+	SLATE_ARGUMENT(TWeakPtr<FPersona>, Persona)
+	/**
+	 * AnimSequenceBase to be used for this panel
+	 */
+	SLATE_ARGUMENT( class UAnimSequence*, Sequence)
+	/**
+	 * Right side of widget width (outside of curve)
+	 */
+	SLATE_ARGUMENT( float, WidgetWidth )
+	/**
+	 * Viewable Range control variables
+	 */
+	SLATE_ATTRIBUTE( float, ViewInputMin )
+	SLATE_ATTRIBUTE( float, ViewInputMax )
+	SLATE_ATTRIBUTE( float, InputMin )
+	SLATE_ATTRIBUTE( float, InputMax )
+	SLATE_EVENT( FOnSetInputViewRange, OnSetInputViewRange )
+	/**
+	 * Get current value
+	 */
+	SLATE_EVENT( FOnGetScrubValue, OnGetScrubValue )
+	SLATE_END_ARGS()
+
+	void Construct(const FArguments& InArgs);
+	/**
+	 * Delete Track
+	 *
+	 */
+	void DeleteTrack(USkeleton::AnimCurveUID Uid);
+
+	/**
+	 * Sets the specified flag value to State for the provided curve
+	 */
+	void SetCurveFlag(FAnimCurveBase* CurveInterface, bool State, EAnimCurveFlags FlagToSet);
+
+	/**
+	* Update Panel
+	* Used internally and by sequence editor to refresh the panel contents
+	* @todo this has to be more efficient. Right now it refreshes the entire panel
+	*/
+	void UpdatePanel();
+
+	/** Get Context Menu Per Track **/
+	TSharedRef<SWidget> CreateCurveContextMenu(USkeleton::AnimCurveUID CurveUid) const;
+
+private:
+	TWeakPtr<FPersona> WeakPersona;
+	TSharedPtr<SSplitter> PanelSlot;
+	class UAnimSequence* Sequence;
+	TAttribute<float> CurrentPosition;
+	FOnGetScrubValue OnGetScrubValue;
+	TArray<TWeakPtr<class STransformCurveEdTrack>> Tracks;
+
+	/**
+	 * This is to control visibility of the curves, so you can edit or not
+	 * Get Widget that shows all curve list and edit
+	 */
+	TSharedRef<SWidget>		GenerateCurveList();
+	/**
+	 * Returns true if this curve is editable
+	 */
+	ESlateCheckBoxState::Type	IsCurveEditable(USkeleton::AnimCurveUID Uid) const;
+	/**
+	 * Toggle curve visibility
+	 */
+	void		ToggleEditability(ESlateCheckBoxState::Type NewType, USkeleton::AnimCurveUID Uid);
+	/**
+	 * Refresh Panel
+	 */
+	FReply		RefreshPanel();
+	/**
+	 *	Show All Curves
+	 */
+	FReply		ShowAll(bool bShow);
+
+	/**
+	 * Convert the requested flag bool value into a checkbox state
+	 */
+	ESlateCheckBoxState::Type GetCurveFlagAsCheckboxState(USkeleton::AnimCurveUID CurveUid, EAnimCurveFlags InFlag) const;
+
+	/**
+	 * Convert a given checkbox state into a flag value in the provided curve
+	 */
+	void SetCurveFlagFromCheckboxState(ESlateCheckBoxState::Type CheckState, USkeleton::AnimCurveUID CurveUid, EAnimCurveFlags InFlag);
+};
\ No newline at end of file
diff --git a/Engine/Source/Editor/Persona/Private/SAnimViewportToolBar.cpp b/Engine/Source/Editor/Persona/Private/SAnimViewportToolBar.cpp
--- a/Engine/Source/Editor/Persona/Private/SAnimViewportToolBar.cpp
+++ b/Engine/Source/Editor/Persona/Private/SAnimViewportToolBar.cpp
@@ -470,6 +470,8 @@ TSharedRef<SWidget> SAnimViewportToolBar::GenerateShowMenu() const
 				ShowMenuBuilder.AddMenuEntry( Actions.ShowRawAnimation );
 				ShowMenuBuilder.AddMenuEntry( Actions.ShowNonRetargetedAnimation );
 				ShowMenuBuilder.AddMenuEntry( Actions.ShowAdditiveBaseBones );
+				ShowMenuBuilder.AddMenuEntry( Actions.ShowSourceRawAnimation);
+				ShowMenuBuilder.AddMenuEntry( Actions.ShowBakedAnimation );
 			}
 			ShowMenuBuilder.EndSection();
 
diff --git a/Engine/Source/Editor/Persona/Private/SAnimationEditorViewport.cpp b/Engine/Source/Editor/Persona/Private/SAnimationEditorViewport.cpp
--- a/Engine/Source/Editor/Persona/Private/SAnimationEditorViewport.cpp
+++ b/Engine/Source/Editor/Persona/Private/SAnimationEditorViewport.cpp
@@ -513,6 +513,18 @@ void SAnimationEditorViewportTabBody::BindCommands()
 		FCanExecuteAction::CreateSP(this, &SAnimationEditorViewportTabBody::IsPreviewingAnimation),
 		FIsActionChecked::CreateSP(this, &SAnimationEditorViewportTabBody::IsShowingAdditiveBase));
 
+	CommandList.MapAction(
+		ViewportShowMenuCommands.ShowSourceRawAnimation,
+		FExecuteAction::CreateSP(this, &SAnimationEditorViewportTabBody::OnShowSourceRawAnimation),
+		FCanExecuteAction(),
+		FIsActionChecked::CreateSP(this, &SAnimationEditorViewportTabBody::IsShowingSourceRawAnimation));
+
+	CommandList.MapAction(
+		ViewportShowMenuCommands.ShowBakedAnimation,
+		FExecuteAction::CreateSP(this, &SAnimationEditorViewportTabBody::OnShowBakedAnimation),
+		FCanExecuteAction(),
+		FIsActionChecked::CreateSP(this, &SAnimationEditorViewportTabBody::IsShowingBakedAnimation));
+
 	//Display info
 	CommandList.MapAction( 
 		ViewportShowMenuCommands.ShowDisplayInfoBasic,
@@ -809,6 +821,26 @@ void SAnimationEditorViewportTabBody::OnShowNonRetargetedAnimation()
 	}
 }
 
+void SAnimationEditorViewportTabBody::OnShowSourceRawAnimation()
+{
+	UDebugSkelMeshComponent* PreviewComponent = PersonaPtr.Pin()->PreviewComponent;
+	if(PreviewComponent != NULL)
+	{
+		PreviewComponent->bDisplaySourceAnimation = !PreviewComponent->bDisplaySourceAnimation;
+		PreviewComponent->MarkRenderStateDirty();
+	}
+}
+
+void SAnimationEditorViewportTabBody::OnShowBakedAnimation()
+{
+	UDebugSkelMeshComponent* PreviewComponent = PersonaPtr.Pin()->PreviewComponent;
+	if(PreviewComponent != NULL)
+	{
+		PreviewComponent->bDisplayBakedAnimation = !PreviewComponent->bDisplayBakedAnimation;
+		PreviewComponent->MarkRenderStateDirty();
+	}
+}
+
 void SAnimationEditorViewportTabBody::OnShowAdditiveBase()
 {
 	UDebugSkelMeshComponent* PreviewComponent = PersonaPtr.Pin()->PreviewComponent;
@@ -855,6 +887,18 @@ bool SAnimationEditorViewportTabBody::IsShowingAdditiveBase() const
 	return PreviewComponent != NULL && PreviewComponent->bDisplayAdditiveBasePose;
 }
 
+bool SAnimationEditorViewportTabBody::IsShowingSourceRawAnimation() const
+{
+	UDebugSkelMeshComponent* PreviewComponent = PersonaPtr.Pin()->PreviewComponent;
+	return PreviewComponent != NULL && PreviewComponent->bDisplaySourceAnimation;
+}
+
+bool SAnimationEditorViewportTabBody::IsShowingBakedAnimation() const
+{
+	UDebugSkelMeshComponent* PreviewComponent = PersonaPtr.Pin()->PreviewComponent;
+	return PreviewComponent != NULL && PreviewComponent->bDisplayBakedAnimation;
+}
+
 void SAnimationEditorViewportTabBody::OnShowDisplayInfo(int32 DisplayInfoMode)
 {
 	GetAnimationViewportClient()->OnSetShowMeshStats(DisplayInfoMode);
diff --git a/Engine/Source/Editor/Persona/Private/SAnimationEditorViewport.h b/Engine/Source/Editor/Persona/Private/SAnimationEditorViewport.h
--- a/Engine/Source/Editor/Persona/Private/SAnimationEditorViewport.h
+++ b/Engine/Source/Editor/Persona/Private/SAnimationEditorViewport.h
@@ -213,6 +213,17 @@ private:
 
 	bool IsShowingNonRetargetedPose() const;
 
+	/** Show non retargeted animation. */
+	void OnShowSourceRawAnimation();
+
+	bool IsShowingSourceRawAnimation() const;
+
+	/** Show non retargeted animation. */
+	void OnShowBakedAnimation();
+
+	bool IsShowingBakedAnimation() const;
+
+
 	/** Additive Base Pose on top of full animation **/
 	void OnShowAdditiveBase();
 
diff --git a/Engine/Source/Editor/Persona/Private/SSequenceEditor.cpp b/Engine/Source/Editor/Persona/Private/SSequenceEditor.cpp
--- a/Engine/Source/Editor/Persona/Private/SSequenceEditor.cpp
+++ b/Engine/Source/Editor/Persona/Private/SSequenceEditor.cpp
@@ -10,6 +10,7 @@
 #include "SAnimationScrubPanel.h"
 #include "SAnimNotifyPanel.h"
 #include "SAnimCurvePanel.h"
+#include "SAnimTrackCurvePanel.h"
 
 #define LOCTEXT_NAMESPACE "AnimSequenceEditor"
 
@@ -26,6 +27,7 @@ void SSequenceEditor::Construct(const FArguments& InArgs)
 		);
 
 	PersonaPtr.Pin()->RegisterOnPostUndo(FPersona::FOnPostUndo::CreateSP( this, &SSequenceEditor::PostUndo ) );
+	PersonaPtr.Pin()->RegisterOnChangeTrackCurves(FPersona::FOnTrackCurvesChanged::CreateSP( this, &SSequenceEditor::OnTrackCurveChanged ) );
 
 	EditorPanels->AddSlot()
 	.AutoHeight()
@@ -59,29 +61,35 @@ void SSequenceEditor::Construct(const FArguments& InArgs)
 		.OnSetInputViewRange(this, &SAnimEditorBase::SetInputViewRange)
 		.OnGetScrubValue(this, &SAnimEditorBase::GetScrubValue)
 	];
-}
 
-SSequenceEditor::~SSequenceEditor()
-{
-	if (PersonaPtr.IsValid())
+	UAnimSequence * AnimSeq = Cast<UAnimSequence>(SequenceObj);
+	if (AnimSeq)
 	{
-		PersonaPtr.Pin()->UnregisterOnPostUndo(this);
+		EditorPanels->AddSlot()
+		.AutoHeight()
+		.Padding(0, 10)
+		[
+			SAssignNew(AnimTrackCurvePanel, SAnimTrackCurvePanel)
+			.Persona(InArgs._Persona)
+			.Sequence(AnimSeq)
+			.WidgetWidth(S2ColumnWidget::DEFAULT_RIGHT_COLUMN_WIDTH)
+			.ViewInputMin(this, &SAnimEditorBase::GetViewMinInput)
+			.ViewInputMax(this, &SAnimEditorBase::GetViewMaxInput)
+			.InputMin(this, &SAnimEditorBase::GetMinInput)
+			.InputMax(this, &SAnimEditorBase::GetMaxInput)
+			.OnSetInputViewRange(this, &SAnimEditorBase::SetInputViewRange)
+			.OnGetScrubValue(this, &SAnimEditorBase::GetScrubValue)
+		];
 	}
 }
 
-void SSequenceEditor::SetSequenceObj(UAnimSequenceBase * NewSequence)
+SSequenceEditor::~SSequenceEditor()
 {
-	SequenceObj = NewSequence;
-
-	if( SequenceObj )
+	if (PersonaPtr.IsValid())
 	{
-		SetInputViewRange(0, SequenceObj->SequenceLength);
+		PersonaPtr.Pin()->UnregisterOnPostUndo(this);
+		PersonaPtr.Pin()->UnregisterOnChangeTrackCurves(this);
 	}
-
-	AnimNotifyPanel->SetSequence(NewSequence);
-	AnimCurvePanel->SetSequence(NewSequence);
-	// sequence editor locks the sequence, so it doesn't get replaced by clicking 
-	AnimScrubPanel->ReplaceLockedSequence(NewSequence);
 }
 
 void SSequenceEditor::PostUndo()
@@ -98,6 +106,18 @@ void SSequenceEditor::PostUndo()
 
 		AnimNotifyPanel->Update();
 		AnimCurvePanel->UpdatePanel();
+		if (AnimTrackCurvePanel.IsValid())
+		{
+			AnimTrackCurvePanel->UpdatePanel();
+		}
+	}
+}
+
+void SSequenceEditor::OnTrackCurveChanged()
+{
+	if (AnimTrackCurvePanel.IsValid())
+	{
+		AnimTrackCurvePanel->UpdatePanel();
 	}
 }
 
diff --git a/Engine/Source/Editor/Persona/Private/SSequenceEditor.h b/Engine/Source/Editor/Persona/Private/SSequenceEditor.h
--- a/Engine/Source/Editor/Persona/Private/SSequenceEditor.h
+++ b/Engine/Source/Editor/Persona/Private/SSequenceEditor.h
@@ -29,16 +29,14 @@ private:
 	/** Persona reference **/
 	TSharedPtr<class SAnimNotifyPanel>	AnimNotifyPanel;
 	TSharedPtr<class SAnimCurvePanel>	AnimCurvePanel;
+	TSharedPtr<class SAnimTrackCurvePanel>	AnimTrackCurvePanel;
 	TSharedPtr<class SAnimationScrubPanel> AnimScrubPanel;
 
 public:
 	void Construct(const FArguments& InArgs);
 	virtual ~SSequenceEditor();
 
-	void SetSequenceObj(UAnimSequenceBase * NewSequence);
-	UAnimSequenceBase * GetSequenceObj() const { return SequenceObj; }
-
-	virtual UAnimSequenceBase* GetEditorObject() const override { return GetSequenceObj(); }
+	virtual UAnimSequenceBase* GetEditorObject() const override { return SequenceObj; }
 
 private:
 	/** Pointer to the animation sequence being edited */
@@ -46,4 +44,5 @@ private:
 
 	/** Post undo **/
 	void PostUndo();
+	void OnTrackCurveChanged();
 };
diff --git a/Engine/Source/Editor/Persona/Private/Shared/AnimationRecorder.cpp b/Engine/Source/Editor/Persona/Private/Shared/AnimationRecorder.cpp
--- a/Engine/Source/Editor/Persona/Private/Shared/AnimationRecorder.cpp
+++ b/Engine/Source/Editor/Persona/Private/Shared/AnimationRecorder.cpp
@@ -129,7 +129,7 @@ void FAnimationRecorder::StopRecord(bool bShowMessage)
 		AnimationObject->NumFrames = NumFrames;
 
 		// can't use TimePassed. That is just total time that has been passed, not necessarily match with frame count
-		AnimationObject->SequenceLength = NumFrames * IntervalTime;
+		AnimationObject->SequenceLength = (NumFrames>1)? (NumFrames-1) * IntervalTime : MINIMUM_ANIMATION_LENGTH;
 		AnimationObject->PostProcessSequence();
 		AnimationObject->MarkPackageDirty();
 
diff --git a/Engine/Source/Editor/UnrealEd/Classes/Animation/DebugSkelMeshComponent.h b/Engine/Source/Editor/UnrealEd/Classes/Animation/DebugSkelMeshComponent.h
--- a/Engine/Source/Editor/UnrealEd/Classes/Animation/DebugSkelMeshComponent.h
+++ b/Engine/Source/Editor/UnrealEd/Classes/Animation/DebugSkelMeshComponent.h
@@ -95,6 +95,14 @@ class UDebugSkelMeshComponent : public USkeletalMeshComponent
 	UPROPERTY(transient)
 	uint32 bDisplayAdditiveBasePose:1;
 
+	/** Display baked animation pose */
+	UPROPERTY(Transient)
+	uint32 bDisplayBakedAnimation:1;
+
+	/** Display source animation pose */
+	UPROPERTY(Transient)
+	uint32 bDisplaySourceAnimation:1;
+
 	/** Display Bound **/
 	UPROPERTY(transient)
 	uint32 bDisplayBound:1;
@@ -111,6 +119,12 @@ class UDebugSkelMeshComponent : public USkeletalMeshComponent
 	/** Storage for non retargeted pose. */
 	TArray<FTransform> NonRetargetedSpaceBases;
 
+	/** Storage of Baked Animation Pose for when bDisplayBakedAnimation == true, as they have to be calculated */
+	TArray<FTransform> BakedAnimationPoses;
+
+	/** Storage of Source Animation Pose for when bDisplaySourceAnimation == true, as they have to be calculated */
+	TArray<FTransform> SourceAnimationPoses;
+
 	/** Color render mode enum value - 0 - none, 1 - tangent, 2 - normal, 3 - mirror, 4 - bone weighting */
 	//var native transient int ColorRenderMode;
 	
diff --git a/Engine/Source/Editor/UnrealEd/Classes/Factories/AnimSequenceFactory.h b/Engine/Source/Editor/UnrealEd/Classes/Factories/AnimSequenceFactory.h
new file mode 100644
--- /dev/null
+++ b/Engine/Source/Editor/UnrealEd/Classes/Factories/AnimSequenceFactory.h
@@ -0,0 +1,29 @@
+// Copyright 1998-2014 Epic Games, Inc. All Rights Reserved.
+
+/**
+ *
+ */
+
+#pragma once
+#include "AnimSequenceFactory.generated.h"
+
+UCLASS(HideCategories=Object,MinimalAPI)
+class UAnimSequenceFactory : public UFactory
+{
+	GENERATED_UCLASS_BODY()
+
+	UPROPERTY()
+	class USkeleton* TargetSkeleton;
+
+	// Begin UFactory Interface
+	virtual bool ConfigureProperties() override;
+	virtual UObject* FactoryCreateNew(UClass* Class,UObject* InParent,FName Name,EObjectFlags Flags,UObject* Context,FFeedbackContext* Warn) override;
+	// Begin UFactory Interface	
+
+private:
+	void OnTargetSkeletonSelected(const FAssetData& SelectedAsset);
+
+private:
+	TSharedPtr<SWindow> PickerWindow;
+};
+
diff --git a/Engine/Source/Editor/UnrealEd/Private/Animation/DebugSkelMeshComponent.cpp b/Engine/Source/Editor/UnrealEd/Private/Animation/DebugSkelMeshComponent.cpp
--- a/Engine/Source/Editor/UnrealEd/Private/Animation/DebugSkelMeshComponent.cpp
+++ b/Engine/Source/Editor/UnrealEd/Private/Animation/DebugSkelMeshComponent.cpp
@@ -436,6 +436,40 @@ void UDebugSkelMeshComponent::GenSpaceBases(TArray<FTransform>& OutSpaceBases)
 
 void UDebugSkelMeshComponent::RefreshBoneTransforms(FActorComponentTickFunction* TickFunction)
 {
+	const bool bIsPreviewInstance = (PreviewInstance && PreviewInstance == AnimScriptInstance);
+
+	BakedAnimationPoses.Empty();
+	if(bDisplayBakedAnimation && bIsPreviewInstance && PreviewInstance->RequiredBones.IsValid())
+	{
+		if(UAnimSequence* Sequence = Cast<UAnimSequence>(PreviewInstance->CurrentAsset))
+		{
+			BakedAnimationPoses.AddUninitialized(PreviewInstance->RequiredBones.GetNumBones());
+			bool bSavedUseSourceData = AnimScriptInstance->RequiredBones.ShouldUseSourceData();
+			AnimScriptInstance->RequiredBones.SetUseRAWData(true);
+			AnimScriptInstance->RequiredBones.SetUseSourceData(false);
+			PreviewInstance->EnableControllers(false);
+			GenSpaceBases(BakedAnimationPoses);
+			AnimScriptInstance->RequiredBones.SetUseRAWData(false);
+			AnimScriptInstance->RequiredBones.SetUseSourceData(bSavedUseSourceData);
+			PreviewInstance->EnableControllers(true);
+		}
+	}
+
+	SourceAnimationPoses.Empty();
+	if(bDisplaySourceAnimation && bIsPreviewInstance && PreviewInstance->RequiredBones.IsValid())
+	{
+		if(UAnimSequence* Sequence = Cast<UAnimSequence>(PreviewInstance->CurrentAsset))
+		{
+			SourceAnimationPoses.AddUninitialized(PreviewInstance->RequiredBones.GetNumBones());
+			bool bSavedUseSourceData = AnimScriptInstance->RequiredBones.ShouldUseSourceData();
+			AnimScriptInstance->RequiredBones.SetUseSourceData(true);
+			PreviewInstance->EnableControllers(false);
+			GenSpaceBases(SourceAnimationPoses);
+			AnimScriptInstance->RequiredBones.SetUseSourceData(bSavedUseSourceData);
+			PreviewInstance->EnableControllers(true);
+		}
+	}
+
 	bool bGenerateRAWAnimation = bDisplayRawAnimation && AnimScriptInstance && AnimScriptInstance->RequiredBones.IsValid();
 
 	if (bGenerateRAWAnimation)
@@ -455,6 +489,7 @@ void UDebugSkelMeshComponent::RefreshBoneTransforms(FActorComponentTickFunction*
 	NonRetargetedSpaceBases.Empty();
 	if( bDisplayNonRetargetedPose && AnimScriptInstance && AnimScriptInstance->RequiredBones.IsValid() )
 	{
+		NonRetargetedSpaceBases.AddUninitialized(PreviewInstance->RequiredBones.GetNumBones());
 		AnimScriptInstance->RequiredBones.SetDisableRetargeting(true);
 		GenSpaceBases(NonRetargetedSpaceBases);
 		AnimScriptInstance->RequiredBones.SetDisableRetargeting(false);
@@ -470,8 +505,6 @@ void UDebugSkelMeshComponent::RefreshBoneTransforms(FActorComponentTickFunction*
 		CompressedSpaceBases.Empty();
 	}
 
-	const bool bIsPreviewInstance = (PreviewInstance && PreviewInstance == AnimScriptInstance);
-
 	// Only works in PreviewInstance, and not for anim blueprint. This is intended.
 	AdditiveBasePoses.Empty();
 	if( bDisplayAdditiveBasePose && bIsPreviewInstance && PreviewInstance->RequiredBones.IsValid() )
@@ -696,4 +729,4 @@ void UDebugSkelMeshComponent::CheckClothTeleport(float DeltaTime)
 	// doesn't need cloth teleport while previewing
 }
 
-#endif // #if WITH_APEX_CLOTHING
\ No newline at end of file
+#endif // #if WITH_APEX_CLOTHING
diff --git a/Engine/Source/Editor/UnrealEd/Private/AnimationEditorUtils.cpp b/Engine/Source/Editor/UnrealEd/Private/AnimationEditorUtils.cpp
--- a/Engine/Source/Editor/UnrealEd/Private/AnimationEditorUtils.cpp
+++ b/Engine/Source/Editor/UnrealEd/Private/AnimationEditorUtils.cpp
@@ -326,47 +326,6 @@ namespace AnimationEditorUtils
 		}
 	}
 
-	template <typename TFactory, typename T>
-	void ExecuteNewAnimAsset(TArray<TWeakObjectPtr<USkeleton>> Objects, const FString InSuffix, FAnimAssetCreated AssetCreated, bool bInContentBrowser )
-	{
-		if(bInContentBrowser && Objects.Num() == 1)
-		{
-			auto Object = Objects[0].Get();
-
-			if(Object)
-			{
-				// Determine an appropriate name for inline-rename
-				FString Name;
-				FString PackageName;
-				CreateUniqueAssetName(Object->GetOutermost()->GetName(), InSuffix, PackageName, Name);
-
-				TFactory* Factory = ConstructObject<TFactory>(TFactory::StaticClass());
-				Factory->TargetSkeleton = Object;
-
-				FContentBrowserModule& ContentBrowserModule = FModuleManager::LoadModuleChecked<FContentBrowserModule>("ContentBrowser");
-				ContentBrowserModule.Get().CreateNewAsset(Name, FPackageName::GetLongPackagePath(PackageName), T::StaticClass(), Factory);
-
-				if(AssetCreated.IsBound())
-				{
-					// @TODO: this doesn't work
-					//FString LongPackagePath = FPackageName::GetLongPackagePath(PackageName);
-					UObject* 	Parent = FindPackage(NULL, *PackageName);
-					UObject* NewAsset = FindObject<UObject>(Parent, *Name, false);
-					if (NewAsset)
-					{
-						TArray<UObject*> NewAssets;
-						NewAssets.Add(NewAsset);
-						AssetCreated.Execute(NewAssets);
-					}
-				}
-			}
-		}
-		else
-		{
-			CreateAnimationAssets(Objects, T::StaticClass(), InSuffix, AssetCreated);
-		}
-	}
-
 	void FillCreateAssetMenu(FMenuBuilder& MenuBuilder, TArray<TWeakObjectPtr<USkeleton>> Skeletons, FAnimAssetCreated AssetCreated, bool bInContentBrowser) 
 	{
 		MenuBuilder.BeginSection("CreateAnimAssets", LOCTEXT("CreateAnimAssetsMenuHeading", "Anim Assets"));
diff --git a/Engine/Source/Editor/UnrealEd/Private/EditorAnimUtils.cpp b/Engine/Source/Editor/UnrealEd/Private/EditorAnimUtils.cpp
--- a/Engine/Source/Editor/UnrealEd/Private/EditorAnimUtils.cpp
+++ b/Engine/Source/Editor/UnrealEd/Private/EditorAnimUtils.cpp
@@ -177,15 +177,10 @@ namespace EditorAnimUtils
 			// Copy curve data from source asset, preserving data in the target if present.
 			if (OldSkeleton)
 			{
-				FSmartNameMapping* OldNameMapping = OldSkeleton->SmartNames.GetContainer(USkeleton::AnimCurveMappingName);
-				FSmartNameMapping* NewNameMapping = NewSkeleton->SmartNames.GetContainer(USkeleton::AnimCurveMappingName);
-				AssetToRetarget->RawCurveData.UpdateLastObservedNames(OldNameMapping);
-
-				for(FFloatCurve& Curve : AssetToRetarget->RawCurveData.FloatCurves)
-				{
-					NewNameMapping->AddName(Curve.LastObservedName, Curve.CurveUid);
-				}
-			}
+				EditorAnimUtils::CopyAnimCurves(OldSkeleton, NewSkeleton, AssetToRetarget, USkeleton::AnimCurveMappingName, FRawCurveTracks::FloatType);
+				// I can't copy transform curves yet because transform curves need retargeting. 
+				//EditorAnimUtils::CopyAnimCurves(OldSkeleton, NewSkeleton, AssetToRetarget, USkeleton::AnimTrackCurveMappingName, FRawCurveTracks::TransformType);
+			}	
 
 			AssetToRetarget->ReplaceReferredAnimations(DuplicatedSequences);
 			AssetToRetarget->ReplaceSkeleton(NewSkeleton, bConvertAnimationDataInComponentSpaces);
@@ -374,6 +369,42 @@ namespace EditorAnimUtils
 			}
 		}
 	}
+
+	void CopyAnimCurves(USkeleton* OldSkeleton, USkeleton* NewSkeleton, UAnimSequenceBase *SequenceBase, const FName ContainerName, FRawCurveTracks::ESupportedCurveType CurveType )
+	{
+		// Copy curve data from source asset, preserving data in the target if present.
+		FSmartNameMapping* OldNameMapping = OldSkeleton->SmartNames.GetContainer(ContainerName);
+		FSmartNameMapping* NewNameMapping = NewSkeleton->SmartNames.GetContainer(ContainerName);
+		SequenceBase->RawCurveData.UpdateLastObservedNames(OldNameMapping, CurveType);
+
+		switch (CurveType)
+		{
+		case FRawCurveTracks::FloatType:
+			{
+				for(FFloatCurve& Curve : SequenceBase->RawCurveData.FloatCurves)
+				{
+					NewNameMapping->AddOrFindName(Curve.LastObservedName, Curve.CurveUid);
+				}
+				break;
+			}
+		case FRawCurveTracks::VectorType:
+			{
+				for(FVectorCurve& Curve : SequenceBase->RawCurveData.VectorCurves)
+				{
+					NewNameMapping->AddOrFindName(Curve.LastObservedName, Curve.CurveUid);
+				}
+				break;
+			}
+		case FRawCurveTracks::TransformType:
+			{
+				for(FTransformCurve& Curve : SequenceBase->RawCurveData.TransformCurves)
+				{
+					NewNameMapping->AddOrFindName(Curve.LastObservedName, Curve.CurveUid);
+				}
+				break;
+			}
+		}
+	}
 }
 
-#undef LOCTEXT_NAMESPACE
\ No newline at end of file
+#undef LOCTEXT_NAMESPACE
diff --git a/Engine/Source/Editor/UnrealEd/Private/Factories/AnimSequenceFactory.cpp b/Engine/Source/Editor/UnrealEd/Private/Factories/AnimSequenceFactory.cpp
new file mode 100644
--- /dev/null
+++ b/Engine/Source/Editor/UnrealEd/Private/Factories/AnimSequenceFactory.cpp
@@ -0,0 +1,86 @@
+// Copyright 1998-2014 Epic Games, Inc. All Rights Reserved.
+
+/*=============================================================================
+	AnimSequenceFactory.cpp: Factory for AnimSequence
+=============================================================================*/
+
+#include "UnrealEd.h"
+
+#include "AssetData.h"
+#include "ContentBrowserModule.h"
+
+#define LOCTEXT_NAMESPACE "AnimSequenceFactory"
+
+UAnimSequenceFactory::UAnimSequenceFactory(const FObjectInitializer& ObjectInitializer)
+	: Super(ObjectInitializer)
+{
+	bCreateNew = true;
+	SupportedClass = UAnimSequence::StaticClass();
+}
+
+bool UAnimSequenceFactory::ConfigureProperties()
+{
+	// Null the skeleton so we can check for selection later
+	TargetSkeleton = NULL;
+
+	// Load the content browser module to display an asset picker
+	FContentBrowserModule& ContentBrowserModule = FModuleManager::LoadModuleChecked<FContentBrowserModule>("ContentBrowser");
+
+	FAssetPickerConfig AssetPickerConfig;
+	
+	/** The asset picker will only show skeletons */
+	AssetPickerConfig.Filter.ClassNames.Add(USkeleton::StaticClass()->GetFName());
+	AssetPickerConfig.Filter.bRecursiveClasses = true;
+
+	/** The delegate that fires when an asset was selected */
+	AssetPickerConfig.OnAssetSelected = FOnAssetSelected::CreateUObject(this, &UAnimSequenceFactory::OnTargetSkeletonSelected);
+
+	/** The default view mode should be a list view */
+	AssetPickerConfig.InitialAssetViewType = EAssetViewType::List;
+
+	/** The default scale for thumbnails. [0-1] range */
+	AssetPickerConfig.ThumbnailScale = 0.25f;
+
+	PickerWindow = SNew(SWindow)
+	.Title(LOCTEXT("CreateAnimSequenceOptions", "Pick Skeleton"))
+	.ClientSize(FVector2D(500, 600))
+	.SupportsMinimize(false) .SupportsMaximize(false)
+	[
+		SNew(SBorder)
+		.BorderImage( FEditorStyle::GetBrush("Menu.Background") )
+		[
+			ContentBrowserModule.Get().CreateAssetPicker(AssetPickerConfig)
+		]
+	];
+
+	GEditor->EditorAddModalWindow(PickerWindow.ToSharedRef());
+	PickerWindow.Reset();
+
+	return TargetSkeleton != NULL;
+}
+
+UObject* UAnimSequenceFactory::FactoryCreateNew(UClass* Class, UObject* InParent, FName Name, EObjectFlags Flags, UObject* Context, FFeedbackContext* Warn)
+{
+	if (TargetSkeleton)
+	{
+		UAnimSequence* AnimSequence = ConstructObject<UAnimSequence>(Class,InParent,Name,Flags);
+
+		// @todo I think this will crash, we should support differentoptions
+		AnimSequence->SequenceLength = 0.f;
+		AnimSequence->NumFrames = 0;
+		
+		AnimSequence->SetSkeleton( TargetSkeleton );
+
+		return AnimSequence;
+	}
+
+	return NULL;
+}
+
+void UAnimSequenceFactory::OnTargetSkeletonSelected(const FAssetData& SelectedAsset)
+{
+	TargetSkeleton = Cast<USkeleton>(SelectedAsset.GetAsset());
+	PickerWindow->RequestDestroyWindow();
+}
+
+#undef LOCTEXT_NAMESPACE
diff --git a/Engine/Source/Editor/UnrealEd/Private/SCurveEditor.cpp b/Engine/Source/Editor/UnrealEd/Private/SCurveEditor.cpp
--- a/Engine/Source/Editor/UnrealEd/Private/SCurveEditor.cpp
+++ b/Engine/Source/Editor/UnrealEd/Private/SCurveEditor.cpp
@@ -2268,7 +2268,6 @@ TArray<SCurveEditor::FSelectedCurveKey> SCurveEditor::GetEditableKeysWithinMarqu
 void SCurveEditor::BeginDragTransaction()
 {
 	TransactionIndex = GEditor->BeginTransaction( LOCTEXT("CurveEditor_Drag", "Mouse Drag") );
-	check( TransactionIndex >= 0 );
 	CurveOwner->ModifyOwner();
 }
 
@@ -2278,9 +2277,9 @@ void SCurveEditor::EndDragTransaction()
 	{
 		GEditor->EndTransaction();
 		TransactionIndex = -1;
-
-		CurveOwner->OnCurveChanged();
 	}
+
+	CurveOwner->OnCurveChanged();
 }
 
 bool SCurveEditor::FSelectedTangent::IsValid() const
diff --git a/Engine/Source/Editor/UnrealEd/Private/SkeletalMeshEdit.cpp b/Engine/Source/Editor/UnrealEd/Private/SkeletalMeshEdit.cpp
--- a/Engine/Source/Editor/UnrealEd/Private/SkeletalMeshEdit.cpp
+++ b/Engine/Source/Editor/UnrealEd/Private/SkeletalMeshEdit.cpp
@@ -350,7 +350,7 @@ void UnFbx::FFbxImporter::FillAndVerifyBoneNames(USkeleton* Skeleton, TArray<Fbx
 	}
 
 	const FReferenceSkeleton& RefSkeleton = Skeleton->GetReferenceSkeleton();
-	const USkeleton::FBoneTreeType & BoneTree = Skeleton->GetBoneTree();
+	const USkeleton::FBoneTreeType& BoneTree = Skeleton->GetBoneTree();
 
 	// make sure at least root bone matches
 	if ( OutRawBoneNames[0] != RefSkeleton.GetBoneName(0) )
@@ -804,7 +804,7 @@ bool UnFbx::FFbxImporter::ImportCurve(const FbxAnimCurve* FbxCurve, FFloatCurve
 			Curve->FloatCurve.SetKeyInterpMode(NewKeyHandle, NewInterpMode);
 			Curve->FloatCurve.SetKeyTangentMode(NewKeyHandle, NewTangentMode);
 			Curve->FloatCurve.SetKeyTangentWeightMode(NewKeyHandle, NewTangentWeightMode);
-			FRichCurveKey & NewKey = Curve->FloatCurve.GetKey(NewKeyHandle);
+			FRichCurveKey& NewKey = Curve->FloatCurve.GetKey(NewKeyHandle);
 
 			// update tangents - I don't know what the unit for Derivatives is
 			// so I'm just inquiring tangent like above
@@ -863,7 +863,7 @@ namespace AnimationTransformDebug
 			// go through all bones and find 
 			for(int32 BoneIndex=0; BoneIndex<TransformDebugData.Num(); ++BoneIndex)
 			{
-				FAnimationTransformDebugData & Data = TransformDebugData[BoneIndex];
+				FAnimationTransformDebugData& Data = TransformDebugData[BoneIndex];
 				int32 ParentIndex = RefSkeleton.GetParentIndex(Data.BoneIndex);
 				int32 ParentTransformDebugDataIndex = 0;
 
@@ -947,14 +947,14 @@ bool UnFbx::FFbxImporter::ImportCurveToAnimSequence(class UAnimSequence * Target
 
 		// Add or retrieve curve
 		USkeleton::AnimCurveUID Uid;
-		NameMapping->AddName(Name, Uid);
+		NameMapping->AddOrFindName(Name, Uid);
 
-		FFloatCurve * CurveToImport = TargetSequence->RawCurveData.GetCurveData(Uid);
+		FFloatCurve * CurveToImport = static_cast<FFloatCurve *>(TargetSequence->RawCurveData.GetCurveData(Uid, FRawCurveTracks::FloatType));
 		if(CurveToImport==NULL)
 		{
 			if(TargetSequence->RawCurveData.AddCurveData(Uid, CurveFlags))
 			{
-				CurveToImport = TargetSequence->RawCurveData.GetCurveData(Uid);
+				CurveToImport = static_cast<FFloatCurve *> (TargetSequence->RawCurveData.GetCurveData(Uid, FRawCurveTracks::FloatType));
 			}
 			else
 			{
@@ -1087,10 +1087,11 @@ bool UnFbx::FFbxImporter::ImportAnimation(USkeleton* Skeleton, UAnimSequence * D
 	}
 
 	// importing custom attribute END
-
+	const bool bSourceDataExists = (DestSeq->SourceRawAnimationData.Num() > 0);
+	TArray<struct FRawAnimSequenceTrack>& RawAnimationData = bSourceDataExists? DestSeq->SourceRawAnimationData : DestSeq->RawAnimationData;
 	DestSeq->TrackToSkeletonMapTable.Empty();
-	DestSeq->RawAnimationData.Empty();
 	DestSeq->AnimationTrackNames.Empty();
+	RawAnimationData.Empty();
 
 	TArray<FName> FbxRawBoneNames;
 	FillAndVerifyBoneNames(Skeleton, SortedLinks, FbxRawBoneNames, FileName);
@@ -1210,7 +1211,7 @@ bool UnFbx::FFbxImporter::ImportAnimation(USkeleton* Skeleton, UAnimSequence * D
 			if (bSuccess)
 			{
 				//add new track
-				int32 NewTrackIdx = DestSeq->RawAnimationData.Add(RawTrack);
+				int32 NewTrackIdx = RawAnimationData.Add(RawTrack);
 				DestSeq->AnimationTrackNames.Add(BoneName);
 
 				NewDebugData.SetTrackData(NewTrackIdx, BoneTreeIndex, BoneName);
@@ -1225,7 +1226,16 @@ bool UnFbx::FFbxImporter::ImportAnimation(USkeleton* Skeleton, UAnimSequence * D
 	}
 
 	DestSeq->NumFrames = TotalNumKeys;
-	DestSeq->PostProcessSequence();
+	// if source data exists, you should bake it to Raw to apply
+	if(bSourceDataExists)
+	{
+		DestSeq->BakeTrackCurvesToRawAnimation();
+	}
+	else
+	{
+		// otherwise just compress
+		DestSeq->PostProcessSequence();
+	}
 
 	// run debug mode
 	AnimationTransformDebug::OutputAnimationTransformDebugData(TransformDebugData, TotalNumKeys, RefSkeleton);
diff --git a/Engine/Source/Editor/UnrealEd/Public/AnimationEditorUtils.h b/Engine/Source/Editor/UnrealEd/Public/AnimationEditorUtils.h
--- a/Engine/Source/Editor/UnrealEd/Public/AnimationEditorUtils.h
+++ b/Engine/Source/Editor/UnrealEd/Public/AnimationEditorUtils.h
@@ -56,7 +56,48 @@ namespace AnimationEditorUtils
 	UNREALED_API void CreateAnimationAssets(const TArray<TWeakObjectPtr<USkeleton>>& Skeletons, TSubclassOf<UAnimationAsset> AssetClass, const FString& InPrefix, FAnimAssetCreated AssetCreated );
 	UNREALED_API void CreateNewAnimBlueprint(TArray<TWeakObjectPtr<USkeleton>> Skeletons, FAnimAssetCreated AssetCreated);
 	UNREALED_API void FillCreateAssetMenu(FMenuBuilder& MenuBuilder, TArray<TWeakObjectPtr<USkeleton>> Skeletons, FAnimAssetCreated AssetCreated, bool bInContentBrowser=true);
+	UNREALED_API void CreateUniqueAssetName(const FString& InBasePackageName, const FString& InSuffix, FString& OutPackageName, FString& OutAssetName);
 
+	template <typename TFactory, typename T>
+	void ExecuteNewAnimAsset(TArray<TWeakObjectPtr<USkeleton>> Objects, const FString InSuffix, FAnimAssetCreated AssetCreated, bool bInContentBrowser)
+	{
+		if(bInContentBrowser && Objects.Num() == 1)
+		{
+			auto Object = Objects[0].Get();
+
+			if(Object)
+			{
+				// Determine an appropriate name for inline-rename
+				FString Name;
+				FString PackageName;
+				CreateUniqueAssetName(Object->GetOutermost()->GetName(), InSuffix, PackageName, Name);
+
+				TFactory* Factory = ConstructObject<TFactory>(TFactory::StaticClass());
+				Factory->TargetSkeleton = Object;
+
+				FContentBrowserModule& ContentBrowserModule = FModuleManager::LoadModuleChecked<FContentBrowserModule>("ContentBrowser");
+				ContentBrowserModule.Get().CreateNewAsset(Name, FPackageName::GetLongPackagePath(PackageName), T::StaticClass(), Factory);
+
+				if(AssetCreated.IsBound())
+				{
+					// @TODO: this doesn't work
+					//FString LongPackagePath = FPackageName::GetLongPackagePath(PackageName);
+					UObject* 	Parent = FindPackage(NULL, *PackageName);
+					UObject* NewAsset = FindObject<UObject>(Parent, *Name, false);
+					if(NewAsset)
+					{
+						TArray<UObject*> NewAssets;
+						NewAssets.Add(NewAsset);
+						AssetCreated.Execute(NewAssets);
+					}
+				}
+			}
+		}
+		else
+		{
+			CreateAnimationAssets(Objects, T::StaticClass(), InSuffix, AssetCreated);
+		}
+	}
 } // namespace AnimationEditorUtils
 
 #endif //__AnimationEditorUtils_h__
diff --git a/Engine/Source/Editor/UnrealEd/Public/EditorAnimUtils.h b/Engine/Source/Editor/UnrealEd/Public/EditorAnimUtils.h
--- a/Engine/Source/Editor/UnrealEd/Public/EditorAnimUtils.h
+++ b/Engine/Source/Editor/UnrealEd/Public/EditorAnimUtils.h
@@ -124,4 +124,7 @@ namespace EditorAnimUtils
 		}
 		return ReturnMap;
 	}
+
+	// utility functions
+	UNREALED_API void CopyAnimCurves(USkeleton* OldSkeleton, USkeleton* NewSkeleton, UAnimSequenceBase *SequenceBase, const FName ContainerName, FRawCurveTracks::ESupportedCurveType CurveType );
 } // namespace EditorAnimUtils
diff --git a/Engine/Source/Runtime/Core/Public/UObject/ObjectVersion.h b/Engine/Source/Runtime/Core/Public/UObject/ObjectVersion.h
--- a/Engine/Source/Runtime/Core/Public/UObject/ObjectVersion.h
+++ b/Engine/Source/Runtime/Core/Public/UObject/ObjectVersion.h
@@ -632,6 +632,7 @@ enum EUnrealEngineObjectUE4Version
 	VER_UE4_SLATE_COMPOSITE_FONTS,
 	// Remove UDEPRECATED_SaveGameSummary, required for UWorld::Serialize
 	VER_UE4_REMOVE_SAVEGAMESUMMARY,
+
 	//Remove bodyseutp serialization from skeletal mesh component
 	VER_UE4_REMOVE_SKELETALMESH_COMPONENT_BODYSETUP_SERIALIZATION,
 	// Made Slate font data use bulk data to store the embedded font data
@@ -646,6 +647,9 @@ enum EUnrealEngineObjectUE4Version
 	VER_UE4_LANDSCAPE_SERIALIZE_PHYSICS_MATERIALS,
 	// Rename Visiblity on widgets to Visibility
 	VER_UE4_RENAME_WIDGET_VISIBILITY,
+	// add track curves for animation
+	VER_UE4_ANIMATION_ADD_TRACKCURVES,
+
 	// -----<new versions can be added before this line>-------------------------------------------------
 	// - this needs to be the last line (see note below)
 	VER_UE4_AUTOMATIC_VERSION_PLUS_ONE,
diff --git a/Engine/Source/Runtime/Engine/Classes/Animation/AnimSequence.h b/Engine/Source/Runtime/Engine/Classes/Animation/AnimSequence.h
--- a/Engine/Source/Runtime/Engine/Classes/Animation/AnimSequence.h
+++ b/Engine/Source/Runtime/Engine/Classes/Animation/AnimSequence.h
@@ -369,6 +369,11 @@ class UAnimSequence : public UAnimSequenceBase
 	 */
 	UPROPERTY()
 	TArray<FName> AnimationTrackNames;
+
+	/**
+	 * Source RawAnimationData. Only can be overriden by when transform curves are added first time OR imported
+	 */
+	TArray<struct FRawAnimSequenceTrack> SourceRawAnimationData;
 #endif // WITH_EDITORONLY_DATA
 
 	/**
@@ -539,11 +544,16 @@ class UAnimSequence : public UAnimSequenceBase
 	UPROPERTY()
 	FString SourceFileTimestamp_DEPRECATED;
 
+	UPROPERTY(transient)
+	bool bNeedsRebake;
+
 #endif // WITH_EDITORONLY_DATA
 
+public:
 	// Begin UObject interface
 	ENGINE_API virtual void Serialize(FArchive& Ar) override;
 	ENGINE_API virtual void PostLoad() override;
+	virtual void PreSave() override;
 #if WITH_EDITOR
 	ENGINE_API virtual void PostEditChangeProperty(FPropertyChangedEvent& PropertyChangedEvent) override;
 #endif // WITH_EDITOR
@@ -760,6 +770,41 @@ public:
 	 */
 	ENGINE_API bool AddLoopingInterpolation();
 
+	/** 
+	 * Bake Transform Curves.TransformCurves to RawAnimation after making a back up of current RawAnimation
+	 */
+	ENGINE_API void BakeTrackCurvesToRawAnimation();
+
+	/**
+	 * Add Key to Transform Curves
+	 */
+	ENGINE_API void AddKeyToSequence(float Time, const FName& BoneName, const FTransform& AdditiveTransform);
+	/**
+	 * Return true if it needs to re-bake
+	 */
+	ENGINE_API bool DoesNeedRebake() const;
+	/**
+	 * Return true if it contains transform curves
+	 */
+	ENGINE_API bool DoesContainTransformCurves() const;
+
+	/**
+	 * Create Animation Sequence from Reference Pose of the Mesh
+	 */
+	ENGINE_API bool CreateAnimation(USkeletalMesh * Mesh);
+	/**
+	 * Create Animation Sequence from the Mesh Component's current bone trasnform
+	 */
+	ENGINE_API bool CreateAnimation(USkeletalMeshComponent * MeshComponent);
+	/**
+	 * Create Animation Sequence from the given animation
+	 */
+	ENGINE_API bool CreateAnimation(UAnimSequence * Sequence);
+	/**
+	 * Resize Animation to Start/End. If End is bigger, it repeats last frame upto End. 
+	 * @todo implement
+	 */
+	ENGINE_API bool Resize(int32 Start, int32 End);
 #endif
 
 	/** 
@@ -795,12 +840,21 @@ private:
 
 	/** Verify Track Map is valid, if not, fix up */
 	void VerifyTrackMap();
+	/** Reset Animation Data. Called before Creating new Animation data **/
+	void ResetAnimation();
+	/** Refresh Track Map from Animation Track Names **/
+	void RefreshTrackMapFromAnimTrackNames();
 #endif
 
 	/**
 	 * Utility function that helps to remove track, you can't just remove RawAnimationData
 	 */
 	void RemoveTrack(int32 TrackIndex);
+	/**
+	 * Utility function that finds the correct spot to insert track to 
+	 */
+	int32 InsertTrack(const FName& BoneName);
+
 	friend class UAnimationAsset;
 };
 
diff --git a/Engine/Source/Runtime/Engine/Classes/Animation/AnimSequenceBase.h b/Engine/Source/Runtime/Engine/Classes/Animation/AnimSequenceBase.h
--- a/Engine/Source/Runtime/Engine/Classes/Animation/AnimSequenceBase.h
+++ b/Engine/Source/Runtime/Engine/Classes/Animation/AnimSequenceBase.h
@@ -159,6 +159,30 @@
 	}
 };
 
+/** native enum for curve types **/
+enum EAnimCurveFlags
+{
+	// Used as morph target curve
+	ACF_DrivesMorphTarget	= 0x00000001,
+	// Used as triggering event
+	ACF_TriggerEvent		= 0x00000002,
+	// Is editable in Sequence Editor
+	ACF_Editable			= 0x00000004,
+	// Used as a material curve
+	ACF_DrivesMaterial		= 0x00000008,
+	// Is a metadata 'curve'
+	ACF_Metadata			= 0x00000010,
+	// motifies bone track
+	ACF_DriveTrack			= 0x00000020,
+	// disabled, right now it's used by track
+	ACF_Disabled			= 0x00000040,
+
+	// default flag when created
+	ACF_DefaultCurve		= ACF_TriggerEvent | ACF_Editable,
+	// curves created from Morph Target
+	ACF_MorphTargetCurve	= ACF_DrivesMorphTarget
+};
+
 /**
  * Float curve data for one track
  */
@@ -174,17 +198,27 @@
 	UPROPERTY()
 	FName		LastObservedName;
 
+	// For smart naming - management purpose - i.e. rename/delete
 	FSmartNameMapping::UID CurveUid;
 
+private:
+
+	/** Curve Type Flags */
+	UPROPERTY()
+	int32		CurveTypeFlags;
+
+public:
 	FAnimCurveBase(){}
 
-	FAnimCurveBase(FName InName)
+	FAnimCurveBase(FName InName, int32 InCurveTypeFlags)
 		: LastObservedName(InName)
+		, CurveTypeFlags(InCurveTypeFlags)
 	{	
 	}
 
-	FAnimCurveBase(USkeleton::AnimCurveUID Uid)
+	FAnimCurveBase(USkeleton::AnimCurveUID Uid, int32 InCurveTypeFlags)
 		: CurveUid(Uid)
+		, CurveTypeFlags(InCurveTypeFlags)
 	{}
 
 	// To be able to use typedef'd types we need to serialize manually
@@ -195,49 +229,6 @@
 			Ar << CurveUid;
 		}
 	}
-};
 
-/** native enum for curve types **/
-enum EAnimCurveFlags
-{
-	// Used as morph target curve
-	ACF_DrivesMorphTarget	= 0x00000001,
-	// Used as triggering event
-	ACF_TriggerEvent		= 0x00000002,
-	// Is editable in Sequence Editor
-	ACF_Editable			= 0x00000004,
-	// Used as a material curve
-	ACF_DrivesMaterial		= 0x00000008,
-	// Is a metadata 'curve'
-	ACF_Metadata		= 0x00000010,
-
-	// default flag when created
-	ACF_DefaultCurve		= ACF_TriggerEvent | ACF_Editable,
-	// curves created from Morph Target
-	ACF_MorphTargetCurve	= ACF_DrivesMorphTarget
-};
-
-USTRUCT()
-struct FFloatCurve : public FAnimCurveBase
-{
-	GENERATED_USTRUCT_BODY()
-
-	/** Curve data for float. */
-	UPROPERTY()
-	FRichCurve	FloatCurve;
-
-	FFloatCurve(){}
-
-	FFloatCurve(FName InName, int32 InCurveTypeFlags)
-		: FAnimCurveBase(InName)
-		, CurveTypeFlags(InCurveTypeFlags)
-	{
-	}
-
-	FFloatCurve(USkeleton::AnimCurveUID Uid, int32 InCurveTypeFlags)
-		: FAnimCurveBase(Uid)
-		, CurveTypeFlags(InCurveTypeFlags)
-	{}
-
 	/**
 	 * Set InFlag to bValue
@@ -263,19 +254,118 @@
 	 * returns CurveTypeFlags
 	 */
 	int32 GetCurveTypeFlags() const;
+};
+
+USTRUCT()
+struct FFloatCurve : public FAnimCurveBase
+{
+	GENERATED_USTRUCT_BODY()
+
+	/** Curve data for float. */
+	UPROPERTY()
+	FRichCurve	FloatCurve;
+
+	FFloatCurve(){}
+
+	FFloatCurve(FName InName, int32 InCurveTypeFlags)
+		: FAnimCurveBase(InName, InCurveTypeFlags)
+	{
+	}
+
+	FFloatCurve(USkeleton::AnimCurveUID Uid, int32 InCurveTypeFlags)
+		: FAnimCurveBase(Uid, InCurveTypeFlags)
+	{}
 
 	virtual void Serialize(FArchive& Ar) override
 	{
 		FAnimCurveBase::Serialize(Ar);
 	}
 
-private:
-
-	/** Curve Type Flags */
-	UPROPERTY()
-	int32		CurveTypeFlags;
+	// we don't want to have = operator. This only copies curves, but leaving naming and everything else intact. 
+	void CopyCurve(FFloatCurve& SourceCurve);
+	float Evaluate(float CurrentTime, float BlendWeight) const;
+	void UpdateOrAddKey(float NewKey, float CurrentTime);
 };
 
+USTRUCT()
+struct FVectorCurve : public FAnimCurveBase
+{
+	GENERATED_USTRUCT_BODY()
+
+	enum EIndex
+	{
+		X = 0, 
+		Y, 
+		Z, 
+		Max
+	};
+
+	/** Curve data for float. */
+	UPROPERTY()
+	FRichCurve	FloatCurves[3];
+
+	FVectorCurve(){}
+
+	FVectorCurve(FName InName, int32 InCurveTypeFlags)
+		: FAnimCurveBase(InName, InCurveTypeFlags)
+	{
+	}
+
+	FVectorCurve(USkeleton::AnimCurveUID Uid, int32 InCurveTypeFlags)
+		: FAnimCurveBase(Uid, InCurveTypeFlags)
+	{}
+
+	virtual void Serialize(FArchive& Ar) override
+	{
+		FAnimCurveBase::Serialize(Ar);
+	}
+
+	// we don't want to have = operator. This only copies curves, but leaving naming and everything else intact. 
+	void CopyCurve(FVectorCurve& SourceCurve);
+	FVector Evaluate(float CurrentTime, float BlendWeight) const;
+	void UpdateOrAddKey(const FVector& NewKey, float CurrentTime);
+};
+
+USTRUCT()
+struct FTransformCurve: public FAnimCurveBase
+{
+	GENERATED_USTRUCT_BODY()
+
+	/** Curve data for each transform. */
+	UPROPERTY()
+	FVectorCurve	TranslationCurve;
+
+	/** Rotation curve - right now we use euler because quat also doesn't provide linear interpolation - curve editor can't handle quat interpolation
+	 * If you hit gimbal lock, you should add extra key to fix it. This will cause gimbal lock. 
+	 * @TODO: Eventually we'll need FRotationCurve that would contain rotation curve - that will interpolate as slerp or as quaternion 
+	 */
+	UPROPERTY()
+	FVectorCurve	RotationCurve;
+
+	UPROPERTY()
+	FVectorCurve	ScaleCurve;
+
+	FTransformCurve(){}
+
+	FTransformCurve(FName InName, int32 InCurveTypeFlags)
+		: FAnimCurveBase(InName, InCurveTypeFlags)
+	{
+	}
+
+	FTransformCurve(USkeleton::AnimCurveUID Uid, int32 InCurveTypeFlags)
+		: FAnimCurveBase(Uid, InCurveTypeFlags)
+	{}
+
+	virtual void Serialize(FArchive& Ar) override
+	{
+		FAnimCurveBase::Serialize(Ar);
+	}
+
+	// we don't want to have = operator. This only copies curves, but leaving naming and everything else intact. 
+	void CopyCurve(FTransformCurve& SourceCurve);
+	FTransform Evaluate(float CurrentTime, float BlendWeight) const;
+	ENGINE_API void UpdateOrAddKey(const FTransform& NewKey, float CurrentTime);
+};
 /**
  * Raw Curve data for serialization
  */
@@ -284,38 +374,118 @@
 {
 	GENERATED_USTRUCT_BODY()
 
+	enum ESupportedCurveType
+	{
+		FloatType,
+		VectorType,
+		TransformType,
+		Max, 
+	};
+
 	UPROPERTY()
 	TArray<FFloatCurve>		FloatCurves;
 
+#if WITH_EDITORONLY_DATA
 	/**
-	 * Evaluate curve data at the time CurrentTime, and add to Instance
+	 * @note : Currently VectorCurves are not evaluated or used for anything else but transient data for modifying bone track
+	 *			Note that it doesn't have UPROPERTY tag yet. In the future, we'd like this to be serialized, but not for now
+	 **/
+	UPROPERTY(transient)
+	TArray<FVectorCurve>	VectorCurves;
+
+	/**
+	 * @note : Currently VectorCurves are not evaluated or used for anything else but transient data for modifying bone track
+	 *			Note that it doesn't have UPROPERTY tag yet. In the future, we'd like this to be serialized, but not for now
+	 **/
+	UPROPERTY()
+	TArray<FTransformCurve>	TransformCurves;
+#endif // #if WITH_EDITOR_DATA
+	/**
+	 * Evaluate curve data at the time CurrentTime, and add to Instance. It only evaluates Float Curve for now
 	 */
 	void EvaluateCurveData(class UAnimInstance* Instance, float CurrentTime, float BlendWeight ) const;
+
+#if WITH_EDITOR
+	/**
+	 *	Evaluate transform curves 
+	 */
+	ENGINE_API void EvaluateTransformCurveData(USkeleton * Skeleton, TMap<FName, FTransform>&OutCurves, float CurrentTime, float BlendWeight) const;
+#endif // WITH_EDITOR
 	/**
 	 * Find curve data based on the curve UID
 	 */
-	ENGINE_API FFloatCurve * GetCurveData(USkeleton::AnimCurveUID Uid);
+	ENGINE_API FAnimCurveBase * GetCurveData(USkeleton::AnimCurveUID Uid, ESupportedCurveType SupportedCurveType = FloatType);
 	/**
 	 * Add new curve from the provided UID and return true if success
 	 * bVectorInterpCurve == true, then it will create FVectorCuve, otherwise, FFloatCurve
 	 */
-	ENGINE_API bool AddCurveData(USkeleton::AnimCurveUID Uid, int32 CurveFlags = ACF_DefaultCurve);
+	ENGINE_API bool AddCurveData(USkeleton::AnimCurveUID Uid, int32 CurveFlags = ACF_DefaultCurve, ESupportedCurveType SupportedCurveType = FloatType);
+
 	/**
 	 * Delete curve data 
 	 */
-	ENGINE_API bool DeleteCurveData(USkeleton::AnimCurveUID Uid);
+	ENGINE_API bool DeleteCurveData(USkeleton::AnimCurveUID Uid, ESupportedCurveType SupportedCurveType = FloatType);
 	/**
 	 * Duplicate curve data
 	 * 
 	 */
-	ENGINE_API bool DuplicateCurveData(USkeleton::AnimCurveUID ToCopyUid, USkeleton::AnimCurveUID NewUid);
+	ENGINE_API bool DuplicateCurveData(USkeleton::AnimCurveUID ToCopyUid, USkeleton::AnimCurveUID NewUid, ESupportedCurveType SupportedCurveType = FloatType);
 
 	/**
 	 * Updates the LastObservedName field of the curves from the provided name container
 	 */
-	ENGINE_API void UpdateLastObservedNames(FSmartNameMapping* NameMapping);
+	ENGINE_API void UpdateLastObservedNames(FSmartNameMapping* NameMapping, ESupportedCurveType SupportedCurveType = FloatType);
 
+	/** 
+	 * Serialize
+	 */
 	void Serialize(FArchive& Ar);
+
+	/** 
+	 * Clear all keys
+	 */
+	void Empty()
+	{
+		FloatCurves.Empty();
+#if WITH_EDITORONLY_DATA
+		VectorCurves.Empty();
+		TransformCurves.Empty();
+#endif
+	}
+private:
+	/** 
+	 * Adding vector curve support - this is all transient data for now. This does not save and all these data will be baked into RawAnimationData
+	 */
+
+	/**
+	 * Find curve data based on the curve UID
+	 */
+	template <typename DataType>
+	DataType * GetCurveDataImpl(TArray<DataType>& Curves, USkeleton::AnimCurveUID Uid);
+
+	/**
+	 * Add new curve from the provided UID and return true if success
+	 * bVectorInterpCurve == true, then it will create FVectorCuve, otherwise, FFloatCurve
+	 */
+	template <typename DataType>
+	bool AddCurveDataImpl(TArray<DataType>& Curves, USkeleton::AnimCurveUID Uid, int32 CurveFlags);
+	/**
+	 * Delete curve data 
+	 */
+	template <typename DataType>
+	bool DeleteCurveDataImpl(TArray<DataType>& Curves, USkeleton::AnimCurveUID Uid);
+	/**
+	 * Duplicate curve data
+	 * 
+	 */
+	template <typename DataType>
+	bool DuplicateCurveDataImpl(TArray<DataType>& Curves, USkeleton::AnimCurveUID ToCopyUid, USkeleton::AnimCurveUID NewUid);
+
+	/**
+	 * Updates the LastObservedName field of the curves from the provided name container
+	 */
+	template <typename DataType>
+	void UpdateLastObservedNamesImpl(TArray<DataType>& Curves, FSmartNameMapping* NameMapping);
 };
 
 UENUM()
@@ -368,7 +538,7 @@
 	 * Supports playing backwards (DeltaTime<0).
 	 * Returns notifies between StartTime (exclusive) and StartTime+DeltaTime (inclusive)
 	 */
-	void GetAnimNotifies(const float& StartTime, const float& DeltaTime, const bool bAllowLooping, TArray<const FAnimNotifyEvent *> & OutActiveNotifies) const;
+	void GetAnimNotifies(const float& StartTime, const float& DeltaTime, const bool bAllowLooping, TArray<const FAnimNotifyEvent *>& OutActiveNotifies) const;
 
 	/** 
 	 * Retrieves AnimNotifies between two time positions. ]PreviousPosition, CurrentPosition]
@@ -376,7 +546,7 @@
 	 * Supports playing backwards (CurrentPosition<PreviousPosition).
 	 * Only supports contiguous range, does NOT support looping and wrapping over.
 	 */
-	void GetAnimNotifiesFromDeltaPositions(const float& PreviousPosition, const float& CurrentPosition, TArray<const FAnimNotifyEvent *> & OutActiveNotifies) const;
+	void GetAnimNotifiesFromDeltaPositions(const float& PreviousPosition, const float& CurrentPosition, TArray<const FAnimNotifyEvent *>& OutActiveNotifies) const;
 
 	/** Evaluate curve data to Instance at the time of CurrentTime **/
 	ENGINE_API virtual void EvaluateCurveData(class UAnimInstance* Instance, float CurrentTime, float BlendWeight ) const;
@@ -442,4 +612,8 @@
 	ENGINE_API void UnregisterOnNotifyChanged(void* Unregister);
 	ENGINE_API virtual bool IsValidToPlay() const { return true; }
 #endif
+
+protected:
+	template <typename DataType>
+	void VerifyCurveNames(USkeleton* Skeleton, const FName& NameContainer, TArray<DataType>& CurveList);
 };
diff --git a/Engine/Source/Runtime/Engine/Classes/Animation/Skeleton.h b/Engine/Source/Runtime/Engine/Classes/Animation/Skeleton.h
--- a/Engine/Source/Runtime/Engine/Classes/Animation/Skeleton.h
+++ b/Engine/Source/Runtime/Engine/Classes/Animation/Skeleton.h
@@ -270,6 +270,9 @@ public:
 	// Names for smartname mappings, if you're adding a new category of smartnames add a new name here
 	static ENGINE_API const FName AnimCurveMappingName;
 
+	// Names for smartname mappings, if you're adding a new category of smartnames add a new name here
+	static ENGINE_API const FName AnimTrackCurveMappingName;
+
 	// Container for smart name mappings
 	UPROPERTY()
 	FSmartNameContainer SmartNames;
diff --git a/Engine/Source/Runtime/Engine/Classes/Animation/SmartName.h b/Engine/Source/Runtime/Engine/Classes/Animation/SmartName.h
--- a/Engine/Source/Runtime/Engine/Classes/Animation/SmartName.h
+++ b/Engine/Source/Runtime/Engine/Classes/Animation/SmartName.h
@@ -23,7 +23,7 @@ struct ENGINE_API FSmartNameMapping
 	// @param Name - The name to add/get
 	// @param OUT OutUid - The UID of the newly created or retrieved name
 	// @return bool - true if the name was added, false if it existed (OutUid will be correctly set still in this case)
-	bool AddName(FName Name, UID& OutUid);
+	bool AddOrFindName(FName Name, UID& OutUid);
 
 	// Get a name from the mapping
 	// @param Uid - UID of the name to retrieve
diff --git a/Engine/Source/Runtime/Engine/Private/Animation/AnimSequence.cpp b/Engine/Source/Runtime/Engine/Private/Animation/AnimSequence.cpp
--- a/Engine/Source/Runtime/Engine/Private/Animation/AnimSequence.cpp
+++ b/Engine/Source/Runtime/Engine/Private/Animation/AnimSequence.cpp
@@ -199,6 +199,12 @@ void UAnimSequence::Serialize(FArchive& Ar)
 	if( !StripFlags.IsEditorDataStripped() )
 	{
 		Ar << RawAnimationData;
+#if WITH_EDITORONLY_DATA
+		if (Ar.UE4Ver() >= VER_UE4_ANIMATION_ADD_TRACKCURVES)
+		{
+			Ar << SourceRawAnimationData;
+		}
+#endif // WITH_EDITORONLY_DATA
 	}
 
 	if ( Ar.IsLoading() )
@@ -276,6 +282,22 @@ bool UAnimSequence::IsValidToPlay() const
 	return ( SequenceLength > 0.f && RawAnimationData.Num() > 0 && CompressedTrackOffsets.Num() > 0 );
 }
 
+void UAnimSequence::PreSave()
+{
+#if WITH_EDITOR
+	// we have to bake it if it's not baked
+	if (DoesNeedRebake())
+	{
+		BakeTrackCurvesToRawAnimation();
+	}
+
+	// make sure if it does contain transform curvesm it contains source data
+	check (!DoesContainTransformCurves() || SourceRawAnimationData.Num() != 0);
+#endif
+
+	Super::PreSave();
+}
+
 void UAnimSequence::PostLoad()
 {
 #if WITH_EDITOR
@@ -431,23 +453,30 @@ void UAnimSequence::PostLoad()
  		}
 	}
 
-	for(FAnimNotifyEvent& Notify : Notifies)
+#if WITH_EDITOR
+	if (USkeleton * CurrentSkeleton = GetSkeleton())
 	{
-		if(Notify.DisplayTime_DEPRECATED != 0.0f)
-		{
-			Notify.Clear();
-			Notify.LinkSequence(this, Notify.DisplayTime_DEPRECATED);
-		}
-		else
+		// Get the name mapping object for curves
+		FSmartNameMapping* NameMapping = CurrentSkeleton->SmartNames.GetContainer(USkeleton::AnimTrackCurveMappingName);
+
+		TArray<FTransformCurve*> UnlinkedCurves;
+		for(FTransformCurve& Curve : RawCurveData.TransformCurves)
 		{
-			Notify.LinkSequence(this, Notify.GetTime());
+			if(!NameMapping->Exists(Curve.LastObservedName))
+			{
+				// The skeleton doesn't know our name. Use the last observed name that was saved with the
+				// curve to create a new name. This can happen if a user saves an animation but not a skeleton
+				// either when updating the assets or editing the curves within.
+				UnlinkedCurves.Add(&Curve);
+			}
 		}
 
-		if(Notify.Duration != 0.0f)
+		for(FTransformCurve* Curve : UnlinkedCurves)
 		{
-			Notify.EndLink.LinkSequence(this, Notify.GetTime() + Notify.Duration);
+			NameMapping->AddOrFindName(Curve->LastObservedName, Curve->CurveUid);
 		}
 	}
+#endif // WITH_EDITOR
 }
 
 #if WITH_EDITOR
@@ -458,7 +487,7 @@ void UAnimSequence::VerifyTrackMap()
 	{
 		UE_LOG(LogAnimation, Warning, TEXT("RESAVE ANIMATION NEEDED(%s): Fixing track names."), *GetName());
 
-		const TArray<FBoneNode> & BoneTree = MySkeleton->GetBoneTree();
+		const TArray<FBoneNode>& BoneTree = MySkeleton->GetBoneTree();
 		AnimationTrackNames.Empty();
 		AnimationTrackNames.AddUninitialized(TrackToSkeletonMapTable.Num());
 		for(int32 I=0; I<TrackToSkeletonMapTable.Num(); ++I)
@@ -491,7 +520,7 @@ void UAnimSequence::VerifyTrackMap()
 		{
 			UE_LOG(LogAnimation, Warning, TEXT("RESAVE ANIMATION NEEDED(%s): Fixing track index."), *GetName());
 
-			const TArray<FBoneNode> & BoneTree = MySkeleton->GetBoneTree();
+			const TArray<FBoneNode>& BoneTree = MySkeleton->GetBoneTree();
 			for(int32 I=0; I<NumTracks; ++I)
 			{
 				TrackToSkeletonMapTable[I].BoneTreeIndex = MySkeleton->GetReferenceSkeleton().FindBoneIndex(AnimationTrackNames[I]);
@@ -598,7 +627,7 @@ void UAnimSequence::GetBoneTransform(FTransform& OutAtom, int32 TrackIndex, floa
 	ExtractBoneTransform(RawAnimationData, OutAtom, TrackIndex, Time);
 }
 
-void UAnimSequence::ExtractBoneTransform(const TArray<struct FRawAnimSequenceTrack> & InRawAnimationData, FTransform& OutAtom, int32 TrackIndex, float Time) const
+void UAnimSequence::ExtractBoneTransform(const TArray<struct FRawAnimSequenceTrack>& InRawAnimationData, FTransform& OutAtom, int32 TrackIndex, float Time) const
 {
 	// Bail out if the animation data doesn't exists (e.g. was stripped by the cooker).
 	if(InRawAnimationData.Num() == 0)
@@ -709,7 +738,7 @@ FTransform UAnimSequence::ExtractRootTrackTransform(float Pos, const FBoneContai
 	// Fallback to root bone from reference skeleton.
 	if( RequiredBones )
 	{
-		const FReferenceSkeleton & RefSkeleton = RequiredBones->GetReferenceSkeleton();
+		const FReferenceSkeleton& RefSkeleton = RequiredBones->GetReferenceSkeleton();
 		if( RefSkeleton.GetNum() > 0 )
 		{
 			return RefSkeleton.GetRefBonePose()[0];
@@ -786,7 +815,7 @@ FTransform UAnimSequence::ExtractRootMotionFromRange(float StartTrackPosition, f
 	return EndTransform.GetRelativeTransform(StartTransform);
 }
 
-void UAnimSequence::GetAnimationPose(FTransformArrayA2 & OutAtoms, const FBoneContainer& RequiredBones, const FAnimExtractContext& ExtractionContext) const
+void UAnimSequence::GetAnimationPose(FTransformArrayA2& OutAtoms, const FBoneContainer& RequiredBones, const FAnimExtractContext& ExtractionContext) const
 {
 	// @todo anim: if compressed and baked in the future, we don't have to do this 
 	if( IsValidAdditive() ) 
@@ -806,7 +835,7 @@ void UAnimSequence::GetAnimationPose(FTransformArrayA2 & OutAtoms, const FBoneCo
 	}
 }
 
-void UAnimSequence::ResetRootBoneForRootMotion(FTransformArrayA2 & BoneTransforms, const FBoneContainer & RequiredBones, ERootMotionRootLock::Type RootMotionRootLock) const
+void UAnimSequence::ResetRootBoneForRootMotion(FTransformArrayA2& BoneTransforms, const FBoneContainer& RequiredBones, ERootMotionRootLock::Type RootMotionRootLock) const
 {
 	switch (RootMotionRootLock)
 	{
@@ -817,7 +846,7 @@ void UAnimSequence::ResetRootBoneForRootMotion(FTransformArrayA2 & BoneTransform
 	}
 }
 
-void UAnimSequence::GetBonePose(FTransformArrayA2 & OutAtoms, const FBoneContainer& RequiredBones, const FAnimExtractContext& ExtractionContext) const
+void UAnimSequence::GetBonePose(FTransformArrayA2& OutAtoms, const FBoneContainer& RequiredBones, const FAnimExtractContext& ExtractionContext) const
 {
 	USkeleton* MySkeleton = GetSkeleton();
 	if (!MySkeleton)
@@ -830,9 +859,9 @@ void UAnimSequence::GetBonePose(FTransformArrayA2 & OutAtoms, const FBoneContain
 	bool const bDisableRetargeting = RequiredBones.GetDisableRetargeting();
 	if (bDisableRetargeting)
 	{
-		TArray<FTransform> const & AuthoredOnRefSkeleton = MySkeleton->GetRefLocalPoses(RetargetSource);
-		TArray<FBoneIndexType> const & RequireBonesIndexArray = RequiredBones.GetBoneIndicesArray();
-		TArray<int32> const & PoseToSkeletonBoneIndexArray = RequiredBones.GetPoseToSkeletonBoneIndexArray();
+		TArray<FTransform> const& AuthoredOnRefSkeleton = MySkeleton->GetRefLocalPoses(RetargetSource);
+		TArray<FBoneIndexType> const& RequireBonesIndexArray = RequiredBones.GetBoneIndicesArray();
+		TArray<int32> const& PoseToSkeletonBoneIndexArray = RequiredBones.GetPoseToSkeletonBoneIndexArray();
 	
 		// Allocate output and make sure it's the right size.
 		int32 const NumPoseBones = RequiredBones.GetNumBones();
@@ -842,8 +871,8 @@ void UAnimSequence::GetBonePose(FTransformArrayA2 & OutAtoms, const FBoneContain
 		int32 const NumRequiredBones = RequireBonesIndexArray.Num();
 		for (int32 ArrayIndex = 0; ArrayIndex < NumRequiredBones; ArrayIndex++)
 		{
-			int32 const & PoseBoneIndex = RequireBonesIndexArray[ArrayIndex];
-			int32 const & SkeletonBoneIndex = PoseToSkeletonBoneIndexArray[PoseBoneIndex];
+			int32 const& PoseBoneIndex = RequireBonesIndexArray[ArrayIndex];
+			int32 const& SkeletonBoneIndex = PoseToSkeletonBoneIndexArray[PoseBoneIndex];
 
 			// Pose bone index should always exist in Skeleton
 			checkSlow(SkeletonBoneIndex != INDEX_NONE);
@@ -862,10 +891,23 @@ void UAnimSequence::GetBonePose(FTransformArrayA2 & OutAtoms, const FBoneContain
 		return;
 	}
 
+#if WITH_EDITOR
+	// this happens only with editor data
 	// Slower path for disable retargeting, that's only used in editor and for debugging.
-	if (RequiredBones.ShouldUseRawData() || bDisableRetargeting)
+	if (bDisableRetargeting || RequiredBones.ShouldUseRawData() || RequiredBones.ShouldUseSourceData())
 	{
-		TArray<int32> const & SkeletonToPoseBoneIndexArray = RequiredBones.GetSkeletonToPoseBoneIndexArray();
+		TArray<int32> const& SkeletonToPoseBoneIndexArray = RequiredBones.GetSkeletonToPoseBoneIndexArray();
+		const TArray<FRawAnimSequenceTrack>* AnimationData = NULL;
+		
+		if (RequiredBones.ShouldUseSourceData() && SourceRawAnimationData.Num() > 0)
+		{
+			AnimationData = &SourceRawAnimationData;
+		}
+		else
+		{
+			AnimationData = &RawAnimationData;
+		}
+
 		for(int32 TrackIndex=0; TrackIndex<NumTracks; TrackIndex++)
 		{
 			const int32 SkeletonBoneIndex = GetSkeletonIndexFromTrackIndex(TrackIndex);
@@ -875,7 +917,7 @@ void UAnimSequence::GetBonePose(FTransformArrayA2 & OutAtoms, const FBoneContain
 			if( PoseBoneIndex != INDEX_NONE )
 			{
 				// extract animation
-				GetBoneTransform(OutAtoms[PoseBoneIndex], TrackIndex, ExtractionContext.CurrentTime, true);
+				ExtractBoneTransform(*AnimationData, OutAtoms[PoseBoneIndex], TrackIndex, ExtractionContext.CurrentTime);
 
 				if (!bDisableRetargeting)
 				{
@@ -890,9 +932,10 @@ void UAnimSequence::GetBonePose(FTransformArrayA2 & OutAtoms, const FBoneContain
 		}
 		return;
 	}
+#endif // WITH_EDITOR
 
-	TArray<FBoneNode> const & BoneTree = MySkeleton->GetBoneTree();
-	TArray<int32> const & SkeletonToPoseBoneIndexArray = RequiredBones.GetSkeletonToPoseBoneIndexArray();
+	TArray<FBoneNode> const& BoneTree = MySkeleton->GetBoneTree();
+	TArray<int32> const& SkeletonToPoseBoneIndexArray = RequiredBones.GetSkeletonToPoseBoneIndexArray();
 
 	//@TODO:@ANIMATION: These should be memstack allocated - very heavy
 	BoneTrackArray RotationScalePairs;
@@ -976,14 +1019,14 @@ void UAnimSequence::GetBonePose(FTransformArrayA2 & OutAtoms, const FBoneContain
 	int32 const NumBonesToScaleRetarget = AnimScaleRetargetingPairs.Num();
 	if (NumBonesToScaleRetarget > 0)
 	{
-		TArray<FTransform> const & AuthoredOnRefSkeleton = MySkeleton->GetRefLocalPoses(RetargetSource);
-		TArray<FTransform> const & PlayingOnRefSkeleton = RequiredBones.GetRefPoseArray();
+		TArray<FTransform> const& AuthoredOnRefSkeleton = MySkeleton->GetRefLocalPoses(RetargetSource);
+		TArray<FTransform> const& PlayingOnRefSkeleton = RequiredBones.GetRefPoseArray();
 
 		for (int32 Index = 0; Index<NumBonesToScaleRetarget; Index++)
 		{
-			BoneTrackPair const & BonePair = AnimScaleRetargetingPairs[Index];
-			int32 const & PoseBoneIndex = BonePair.AtomIndex;
-			int32 const & SkeletonBoneIndex = BonePair.TrackIndex; 
+			BoneTrackPair const& BonePair = AnimScaleRetargetingPairs[Index];
+			int32 const& PoseBoneIndex = BonePair.AtomIndex;
+			int32 const& SkeletonBoneIndex = BonePair.TrackIndex; 
 
 			// @todo - precache that in FBoneContainer when we have SkeletonIndex->TrackIndex mapping. So we can just apply scale right away.
 			float const SourceTranslationLength = AuthoredOnRefSkeleton[SkeletonBoneIndex].GetTranslation().Size();
@@ -999,14 +1042,14 @@ void UAnimSequence::GetBonePose(FTransformArrayA2 & OutAtoms, const FBoneContain
 	int32 const NumBonesToRelativeRetarget = AnimRelativeRetargetingPairs.Num();
 	if (NumBonesToRelativeRetarget > 0)
 	{
-		TArray<FTransform> const & AuthoredOnRefSkeleton = MySkeleton->GetRefLocalPoses(RetargetSource);
-		TArray<FTransform> const & PlayingOnRefSkeleton = RequiredBones.GetRefPoseArray();
+		TArray<FTransform> const& AuthoredOnRefSkeleton = MySkeleton->GetRefLocalPoses(RetargetSource);
+		TArray<FTransform> const& PlayingOnRefSkeleton = RequiredBones.GetRefPoseArray();
 
 		for (int32 Index = 0; Index < NumBonesToRelativeRetarget; Index++)
 		{
-			BoneTrackPair const & BonePair = AnimRelativeRetargetingPairs[Index];
-			int32 const & PoseBoneIndex = BonePair.AtomIndex;
-			int32 const & SkeletonBoneIndex = BonePair.TrackIndex;
+			BoneTrackPair const& BonePair = AnimRelativeRetargetingPairs[Index];
+			int32 const& PoseBoneIndex = BonePair.AtomIndex;
+			int32 const& SkeletonBoneIndex = BonePair.TrackIndex;
 
 			// Apply the retargeting as if it were an additive difference between the current skeleton and the retarget skeleton. 
 			OutAtoms[PoseBoneIndex].SetRotation(OutAtoms[PoseBoneIndex].GetRotation() * AuthoredOnRefSkeleton[SkeletonBoneIndex].GetRotation().Inverse() * PlayingOnRefSkeleton[PoseBoneIndex].GetRotation() );
@@ -1017,7 +1060,7 @@ void UAnimSequence::GetBonePose(FTransformArrayA2 & OutAtoms, const FBoneContain
 	}
 }
 
-void UAnimSequence::GetBonePose_Additive(FTransformArrayA2 & OutAtoms, const FBoneContainer& RequiredBones, const FAnimExtractContext& ExtractionContext) const
+void UAnimSequence::GetBonePose_Additive(FTransformArrayA2& OutAtoms, const FBoneContainer& RequiredBones, const FAnimExtractContext& ExtractionContext) const
 {
 	if( !IsValidAdditive() )
 	{
@@ -1036,7 +1079,7 @@ void UAnimSequence::GetBonePose_Additive(FTransformArrayA2 & OutAtoms, const FBo
 	FAnimationRuntime::ConvertPoseToAdditive(OutAtoms, BasePoseAtoms, RequiredBones);
 }
 
-void UAnimSequence::GetAdditiveBasePose(FTransformArrayA2 & OutAtoms, const FBoneContainer& RequiredBones, const FAnimExtractContext& ExtractionContext) const
+void UAnimSequence::GetAdditiveBasePose(FTransformArrayA2& OutAtoms, const FBoneContainer& RequiredBones, const FAnimExtractContext& ExtractionContext) const
 {
 	switch (RefPoseType)
 	{
@@ -1098,28 +1141,28 @@ void UAnimSequence::GetBonePose_AdditiveMeshRotationOnly(FTransformArrayA2& OutA
 void UAnimSequence::RetargetBoneTransform(FTransform& BoneTransform, const int32& SkeletonBoneIndex, const int32& PoseBoneIndex, const FBoneContainer& RequiredBones) const
 {
 	const USkeleton* MySkeleton = GetSkeleton();
-	const TArray<FBoneNode> & BoneTree = MySkeleton->GetBoneTree();
+	const TArray<FBoneNode>& BoneTree = MySkeleton->GetBoneTree();
 	if( BoneTree[SkeletonBoneIndex].TranslationRetargetingMode == EBoneTranslationRetargetingMode::Skeleton )
 	{
-		const TArray<FTransform> & TargetRefPoseArray = RequiredBones.GetRefPoseArray();
+		const TArray<FTransform>& TargetRefPoseArray = RequiredBones.GetRefPoseArray();
 		BoneTransform.SetTranslation( TargetRefPoseArray[PoseBoneIndex].GetTranslation() );
 	}
 	else if( BoneTree[SkeletonBoneIndex].TranslationRetargetingMode == EBoneTranslationRetargetingMode::AnimationScaled )
 	{
 		// @todo - precache that in FBoneContainer when we have SkeletonIndex->TrackIndex mapping. So we can just apply scale right away.
-		const TArray<FTransform> & SkeletonRefPoseArray = GetSkeleton()->GetRefLocalPoses(RetargetSource);
+		const TArray<FTransform>& SkeletonRefPoseArray = GetSkeleton()->GetRefLocalPoses(RetargetSource);
 		const float SourceTranslationLength = SkeletonRefPoseArray[SkeletonBoneIndex].GetTranslation().Size();
 		if( FMath::Abs(SourceTranslationLength) > KINDA_SMALL_NUMBER )
 		{
-			const TArray<FTransform> & TargetRefPoseArray = RequiredBones.GetRefPoseArray();
+			const TArray<FTransform>& TargetRefPoseArray = RequiredBones.GetRefPoseArray();
 			const float TargetTranslationLength = TargetRefPoseArray[PoseBoneIndex].GetTranslation().Size();
 			BoneTransform.ScaleTranslation(TargetTranslationLength / SourceTranslationLength);
 		}
 	}
 	else if (BoneTree[SkeletonBoneIndex].TranslationRetargetingMode == EBoneTranslationRetargetingMode::AnimationRelative)
 	{
-		const TArray<FTransform> & AuthoredOnRefSkeleton = GetSkeleton()->GetRefLocalPoses(RetargetSource);
-		const TArray<FTransform> & PlayingOnRefSkeleton = RequiredBones.GetRefPoseArray();
+		const TArray<FTransform>& AuthoredOnRefSkeleton = GetSkeleton()->GetRefLocalPoses(RetargetSource);
+		const TArray<FTransform>& PlayingOnRefSkeleton = RequiredBones.GetRefPoseArray();
 
 		// Apply the retargeting as if it were an additive difference between the current skeleton and the retarget skeleton. 
 		BoneTransform.SetRotation(BoneTransform.GetRotation() * AuthoredOnRefSkeleton[SkeletonBoneIndex].GetRotation().Inverse() * PlayingOnRefSkeleton[PoseBoneIndex].GetRotation());
@@ -1230,7 +1273,7 @@ bool UAnimSequence::CropRawAnimData( float CurrentTime, bool bFromStart )
 
 		for ( auto CurveIter = RawCurveData.FloatCurves.CreateIterator(); CurveIter; ++CurveIter )
 		{
-			FFloatCurve & Curve = *CurveIter;
+			FFloatCurve& Curve = *CurveIter;
 			// fix up curve before deleting
 			// add these two values to keep the previous curve shape
 			// it's possible I'm adding same values in the same location
@@ -1413,7 +1456,7 @@ bool UAnimSequence::CompressRawAnimData(float MaxPosDiff, float MaxAngleDiff)
 		// go through remove keys if not needed
 		for ( int32 TrackIndex=0; TrackIndex<RawAnimationData.Num(); TrackIndex++ )
 		{
-			FRawAnimSequenceTrack const & RawData = RawAnimationData[TrackIndex];
+			FRawAnimSequenceTrack const& RawData = RawAnimationData[TrackIndex];
 			if ( RawData.ScaleKeys.Num() > 0 )
 			{
 				// if scale key exists, see if we can just empty it
@@ -1432,7 +1475,7 @@ bool UAnimSequence::CompressRawAnimData(float MaxPosDiff, float MaxAngleDiff)
 			// then remove all scale keys
 			for ( int32 TrackIndex=0; TrackIndex<RawAnimationData.Num(); TrackIndex++ )
 			{
-				FRawAnimSequenceTrack & RawData = RawAnimationData[TrackIndex];
+				FRawAnimSequenceTrack& RawData = RawAnimationData[TrackIndex];
 				RawData.ScaleKeys.Empty();
 			}
 		}
@@ -1556,7 +1599,7 @@ bool UAnimSequence::CopyNotifies(UAnimSequence* SourceAnimSeq, UAnimSequence* De
 		{
 			// If a notify is found which occurs off the end of the destination sequence, prompt the user to continue.
 			const FAnimNotifyEvent& SrcNotifyEvent = SourceAnimSeq->Notifies[NotifyIndex];
-			if( SrcNotifyEvent.GetTime() > DestAnimSeq->SequenceLength )
+			if( SrcNotifyEvent.DisplayTime_DEPRECATED > DestAnimSeq->SequenceLength )
 			{
 				const bool bProceed = EAppReturnType::Yes == FMessageDialog::Open( EAppMsgType::YesNo, NSLOCTEXT("UnrealEd", "SomeNotifiesWillNotBeCopiedQ", "Some notifies will not be copied because the destination sequence is not long enough.  Proceed?") );
 				if( !bProceed )
@@ -1593,7 +1636,7 @@ bool UAnimSequence::CopyNotifies(UAnimSequence* SourceAnimSeq, UAnimSequence* De
 		const FAnimNotifyEvent& SrcNotifyEvent = SourceAnimSeq->Notifies[NotifyIndex];
 
 		// Skip notifies which occur at times later than the destination sequence is long.
-		if( SrcNotifyEvent.GetTime() > DestAnimSeq->SequenceLength )
+		if( SrcNotifyEvent.DisplayTime_DEPRECATED > DestAnimSeq->SequenceLength )
 		{
 			continue;
 		}
@@ -1602,7 +1645,7 @@ bool UAnimSequence::CopyNotifies(UAnimSequence* SourceAnimSeq, UAnimSequence* De
 		// to insert the new notify.
 		int32 NewNotifyIndex = 0;
 		while( NewNotifyIndex < DestAnimSeq->Notifies.Num()
-			&& DestAnimSeq->Notifies[NewNotifyIndex].GetTime() <= SrcNotifyEvent.GetTime() )
+			&& DestAnimSeq->Notifies[NewNotifyIndex].DisplayTime_DEPRECATED <= SrcNotifyEvent.DisplayTime_DEPRECATED )
 		{
 			++NewNotifyIndex;
 		}
@@ -1615,10 +1658,10 @@ bool UAnimSequence::CopyNotifies(UAnimSequence* SourceAnimSeq, UAnimSequence* De
 
 		// Copy time and comment.
 		FAnimNotifyEvent& Notify = DestAnimSeq->Notifies[NewNotifyIndex];
-		Notify.SetTime(SrcNotifyEvent.GetTime());
-		Notify.TriggerTimeOffset = GetTriggerTimeOffsetForType( DestAnimSeq->CalculateOffsetForNotify(Notify.GetTime()));
+		Notify.DisplayTime_DEPRECATED = SrcNotifyEvent.DisplayTime_DEPRECATED;
+		Notify.TriggerTimeOffset = GetTriggerTimeOffsetForType( DestAnimSeq->CalculateOffsetForNotify(Notify.DisplayTime_DEPRECATED));
 		Notify.NotifyName = SrcNotifyEvent.NotifyName;
-		Notify.SetDuration(SrcNotifyEvent.GetDuration());
+		Notify.Duration = SrcNotifyEvent.Duration;
 
 		// Copy the notify itself, and point the new one at it.
 		if( SrcNotifyEvent.Notify )
@@ -1670,13 +1713,13 @@ void UAnimSequence::UpgradeMorphTargetCurves()
 			FSmartNameMapping* NameMapping = GetSkeleton()->SmartNames.GetContainer(USkeleton::AnimCurveMappingName);
 			check(NameMapping); // This name mapping should always be present
 			USkeleton::AnimCurveUID NewId;
-			NameMapping->AddName(CurveTrack.CurveName, NewId);
+			NameMapping->AddOrFindName(CurveTrack.CurveName, NewId);
 
 			// add it first, it won't add if duplicate
 			RawCurveData.AddCurveData(NewId);
 			
 			// get the last one I just added
-			FFloatCurve * NewCurve = RawCurveData.GetCurveData(NewId);
+			FFloatCurve * NewCurve = static_cast<FFloatCurve*>(RawCurveData.GetCurveData(NewId, FRawCurveTracks::FloatType));
 			check (NewCurve);
 
 			// since the import code change, this might not match all the time
@@ -1740,7 +1783,7 @@ int32 FindMeshBoneIndexFromBoneName(USkeleton * Skeleton, const FName &BoneName)
 
 	return BoneIndex;
 }
-void FillUpTransformBasedOnRig(USkeleton* Skeleton, TArray<FTransform> & NodeSpaceBases, TArray<FTransform> &Rotations, TArray<FTransform> & Translations)
+void FillUpTransformBasedOnRig(USkeleton* Skeleton, TArray<FTransform>& NodeSpaceBases, TArray<FTransform> &Rotations, TArray<FTransform>& Translations)
 {
 	TArray<FTransform> SpaceBases;
 	FAnimationRuntime::FillUpSpaceBasesRetargetBasePose(Skeleton, SpaceBases);
@@ -1789,7 +1832,7 @@ void FillUpTransformBasedOnRig(USkeleton* Skeleton, TArray<FTransform> & NodeSpa
 				if (TransformBase != NULL)
 				{
 					// orientation constraint			
-					const auto & RotConstraint = TransformBase->Constraints[EControlConstraint::Type::Orientation];
+					const auto& RotConstraint = TransformBase->Constraints[EControlConstraint::Type::Orientation];
 
 					if (RotConstraint.TransformConstraints.Num() > 0)
 					{
@@ -1803,7 +1846,7 @@ void FillUpTransformBasedOnRig(USkeleton* Skeleton, TArray<FTransform> & NodeSpa
 					}
 
 					// translation constraint
-					const auto & TransConstraint = TransformBase->Constraints[EControlConstraint::Type::Translation];
+					const auto& TransConstraint = TransformBase->Constraints[EControlConstraint::Type::Translation];
 
 					if (TransConstraint.TransformConstraints.Num() > 0)
 					{
@@ -1881,7 +1924,7 @@ void UAnimSequence::RemapTracksToNewSkeleton( USkeleton* NewSkeleton, bool bConv
 				// saves the translation conversion datao
 				OldToNewTranslationRatio.AddUninitialized(NumNodes);
 
-				const TArray<FNode> & Nodes = Rig->GetNodes();
+				const TArray<FNode>& Nodes = Rig->GetNodes();
 				// calculate the relative transform to new skeleton
 				// so that we can apply the delta in component space
 				for (int32 NodeIndex = 0; NodeIndex < NumNodes; ++NodeIndex)
@@ -1937,7 +1980,7 @@ void UAnimSequence::RemapTracksToNewSkeleton( USkeleton* NewSkeleton, bool bConv
 			{
 				int32 NodeIndex = Rig->FindNode(SrcValidNodeNames[SrcTrackIndex]);
 				check (NodeIndex != INDEX_NONE);
-				auto & RawAnimation = RiggingAnimationData.AnimationTracks[SrcTrackIndex];
+				auto& RawAnimation = RiggingAnimationData.AnimationTracks[SrcTrackIndex];
 
 				// find rotation parent node
 				int32 RotParentTrackIndex = FindValidTransformParentTrack(Rig, NodeIndex, false, SrcValidNodeNames);
@@ -2014,7 +2057,7 @@ void UAnimSequence::RemapTracksToNewSkeleton( USkeleton* NewSkeleton, bool bConv
 					}
 				}
 
-				auto & RawAnimation = NewRawAnimationData[SrcTrackIndex];
+				auto& RawAnimation = NewRawAnimationData[SrcTrackIndex];
 				RawAnimation.PosKeys.Empty(NumKeys);
 				RawAnimation.PosKeys.AddUninitialized(NumKeys);
 				RawAnimation.RotKeys.Empty(NumKeys);
@@ -2092,8 +2135,8 @@ void UAnimSequence::RemapTracksToNewSkeleton( USkeleton* NewSkeleton, bool bConv
 			FAnimationRuntime::FillUpSpaceBasesRefPose(NewSkeleton, NewSpaceBaseRefPose);
 			FAnimationRuntime::FillUpSpaceBasesRefPose(OldSkeleton, OldSpaceBaseRefPose);
 
-			const TArray<FTransform> & OldRefPose = OldSkeleton->GetReferenceSkeleton().GetRefBonePose();
-			const TArray<FTransform> & NewRefPose = NewSkeleton->GetReferenceSkeleton().GetRefBonePose();
+			const TArray<FTransform>& OldRefPose = OldSkeleton->GetReferenceSkeleton().GetRefBonePose();
+			const TArray<FTransform>& NewRefPose = NewSkeleton->GetReferenceSkeleton().GetRefBonePose();
 
 			// now we'd like to get the relative transform from old to new ref pose in component space
 			// PK2*K2 = PK1*K1*theta where theta => P1*R1*theta = P2*R2 
@@ -2166,7 +2209,7 @@ void UAnimSequence::RemapTracksToNewSkeleton( USkeleton* NewSkeleton, bool bConv
 
 				if(TrackIndex != INDEX_NONE)
 				{
-					auto & RawAnimation = RawAnimationData[TrackIndex];
+					auto& RawAnimation = RawAnimationData[TrackIndex];
 					// fill up keys - calculate PK1 * K1
 					for(int32 Key=0; Key<NumKeys; ++Key)
 					{
@@ -2246,7 +2289,7 @@ void UAnimSequence::RemapTracksToNewSkeleton( USkeleton* NewSkeleton, bool bConv
 				// now save back to animation data
 				if(TrackIndex != INDEX_NONE)
 				{
-					auto & RawAnimation = NewRawAnimationData[TrackIndex];
+					auto& RawAnimation = NewRawAnimationData[TrackIndex];
 					RawAnimation.PosKeys.Empty(NumKeys);
 					RawAnimation.PosKeys.AddUninitialized(NumKeys);
 					RawAnimation.RotKeys.Empty(NumKeys);
@@ -2308,6 +2351,8 @@ void UAnimSequence::RemapTracksToNewSkeleton( USkeleton* NewSkeleton, bool bConv
 
 void UAnimSequence::PostProcessSequence()
 {
+	// pre process before compress raw animation data
+
 	// if scale is too small, zero it out. Cause it hard to retarget when compress
 	// inverse scale is applied to translation, and causing translation to be huge to retarget, but
 	// compression can't handle that much precision. 
@@ -2331,6 +2376,16 @@ void UAnimSequence::PostProcessSequence()
 				Scale3D.Z = 0.f;
 			}
 		}
+
+		// make sure Rotation part is normalized before compress
+		for(auto RotIter = RawAnim.RotKeys.CreateIterator(); RotIter; ++RotIter)
+		{
+			FQuat& Rotation = *RotIter;
+			if( !Rotation.IsNormalized() )
+			{
+				Rotation.Normalize();
+			}
+		}
 	}
 
 	CompressRawAnimData();
@@ -2396,11 +2451,108 @@ void UAnimSequence::RemoveTrack(int32 TrackIndex)
 		RawAnimationData.RemoveAt(TrackIndex);
 		AnimationTrackNames.RemoveAt(TrackIndex);
 		TrackToSkeletonMapTable.RemoveAt(TrackIndex);
+		SourceRawAnimationData.RemoveAt(TrackIndex);
 
 		check (RawAnimationData.Num() == AnimationTrackNames.Num() && AnimationTrackNames.Num() == TrackToSkeletonMapTable.Num() );
 	}
 }
 
+int32 FindFirstChildTrack(const USkeleton* MySkeleton, const FReferenceSkeleton& RefSkeleton, const TArray<FName>& AnimationTrackNames, FName BoneName)
+{
+	const int32 BoneIndex = RefSkeleton.FindBoneIndex(BoneName);
+	if(BoneIndex == INDEX_NONE)
+	{
+		// get out, nothing to do
+		return INDEX_NONE;
+	}
+
+	// find children
+	TArray<int32> Childs;
+	if(MySkeleton->GetChildBones(BoneIndex, Childs) > 0)
+	{
+		// first look for direct children
+		for(auto ChildIndex : Childs)
+		{
+			FName ChildBoneName = RefSkeleton.GetBoneName(ChildIndex);
+			int32 ChildTrackIndex = AnimationTrackNames.Find(ChildBoneName);
+			if(ChildTrackIndex != INDEX_NONE)
+			{
+				// found the new track
+				return ChildTrackIndex;
+			}
+		}
+
+		int32 BestGrandChildIndex = INDEX_NONE;
+		// if you didn't find yet, now you have to go through all children
+		for(auto ChildIndex : Childs)
+		{
+			FName ChildBoneName = RefSkeleton.GetBoneName(ChildIndex);
+			// now I have to go through all childrewn and find who is earliest since I don't know which one might be the closest one
+			int32 GrandChildIndex = FindFirstChildTrack(MySkeleton, RefSkeleton, AnimationTrackNames, ChildBoneName);
+			if (GrandChildIndex != INDEX_NONE)
+			{
+				if (BestGrandChildIndex == INDEX_NONE)
+				{
+					BestGrandChildIndex = GrandChildIndex;
+				}
+				else if (BestGrandChildIndex > GrandChildIndex)
+				{
+					// best should be earlier track index
+					BestGrandChildIndex = GrandChildIndex;
+				}
+			}
+		}
+
+		return BestGrandChildIndex;
+	}
+	else
+	{
+		// there is no child, just add at the end
+		return AnimationTrackNames.Num();
+	}
+}
+
+int32 UAnimSequence::InsertTrack(const FName& BoneName)
+{
+	// first verify if it doesn't exists, if it does, return
+	int32 CurrentTrackIndex = AnimationTrackNames.Find(BoneName);
+	if (CurrentTrackIndex != INDEX_NONE)
+	{
+		return CurrentTrackIndex;
+	}
+
+	USkeleton * MySkeleton = GetSkeleton();
+	// should not call this if skeleton was empty
+	if (ensure(MySkeleton) == false)
+	{
+		return INDEX_NONE;
+	}
+
+	const FReferenceSkeleton& RefSkeleton = MySkeleton->GetReferenceSkeleton();
+	int32 NewTrackIndex = FindFirstChildTrack(MySkeleton, RefSkeleton, AnimationTrackNames, BoneName);
+	int32 BoneIndex = RefSkeleton.FindBoneIndex(BoneName);
+	if (NewTrackIndex != INDEX_NONE)
+	{
+		const TArray<FTransform>& RefPose = RefSkeleton.GetRefBonePose();
+
+		FRawAnimSequenceTrack RawTrack;
+		RawTrack.PosKeys.Add(RefPose[BoneIndex].GetTranslation());
+		RawTrack.RotKeys.Add(RefPose[BoneIndex].GetRotation());
+		RawTrack.ScaleKeys.Add(RefPose[BoneIndex].GetScale3D());
+
+		// now insert to the track
+		RawAnimationData.Insert(RawTrack, NewTrackIndex);
+		AnimationTrackNames.Insert(BoneName, NewTrackIndex);
+		SourceRawAnimationData.Insert(RawTrack, NewTrackIndex);
+
+		RefreshTrackMapFromAnimTrackNames();
+
+		check(RawAnimationData.Num() == AnimationTrackNames.Num() && AnimationTrackNames.Num() == TrackToSkeletonMapTable.Num());
+	}
+
+	return NewTrackIndex;
+}
+
 bool UAnimSequence::GetAllAnimationSequencesReferred(TArray<UAnimSequence*>& AnimationSequences)
 {
 	if (RefPoseSeq)
@@ -2436,7 +2588,7 @@ bool UAnimSequence::AddLoopingInterpolation()
 		// now I need to calculate back to new animation data
 		for(int32 TrackIndex=0; TrackIndex<NumTracks; ++TrackIndex)
 		{
-			auto & RawAnimation = RawAnimationData[TrackIndex];
+			auto& RawAnimation = RawAnimationData[TrackIndex];
 			if (RawAnimation.PosKeys.Num() > 1)
 			{
 				FVector FirstKey = RawAnimation.PosKeys[0];
@@ -2473,7 +2625,7 @@ int32 FindParentNodeIndex(URig* Rig, USkeleton* Skeleton, FName ParentNodeName)
 	return Skeleton->GetReferenceSkeleton().FindBoneIndex(ParentBoneName);
 }
 
-int32 UAnimSequence::GetSpaceBasedAnimationData(TArray< TArray<FTransform> > & AnimationDataInComponentSpace, FAnimSequenceTrackContainer * RiggingAnimationData) const
+int32 UAnimSequence::GetSpaceBasedAnimationData(TArray< TArray<FTransform> >& AnimationDataInComponentSpace, FAnimSequenceTrackContainer * RiggingAnimationData) const
 {
 	USkeleton* Skeleton = GetSkeleton();
 
@@ -2652,7 +2804,7 @@ int32 UAnimSequence::GetSpaceBasedAnimationData(TArray< TArray<FTransform> > & A
 
 			if (TrackIndex != INDEX_NONE)
 			{
-				auto & RawAnimation = RawAnimationData[TrackIndex];
+				auto& RawAnimation = RawAnimationData[TrackIndex];
 				// fill up keys - calculate PK1 * K1
 				for (int32 Key = 0; Key < NumKeys; ++Key)
 				{
@@ -2693,7 +2845,7 @@ int32 UAnimSequence::GetSpaceBasedAnimationData(TArray< TArray<FTransform> > & A
 	return AnimationDataInComponentSpace.Num();
 }
 
-bool UAnimSequence::ConvertAnimationDataToRiggingData(FAnimSequenceTrackContainer & RiggingAnimationData)
+bool UAnimSequence::ConvertAnimationDataToRiggingData(FAnimSequenceTrackContainer& RiggingAnimationData)
 {
 	USkeleton* Skeleton = GetSkeleton();
 	if (Skeleton && Skeleton->GetRig())
@@ -2711,7 +2863,7 @@ bool UAnimSequence::ConvertAnimationDataToRiggingData(FAnimSequenceTrackContaine
 			// first we copy all space bases back to it
 			for (int32 NodeIndex = 0; NodeIndex < NumNodes; ++NodeIndex)
 			{
-				struct FRawAnimSequenceTrack & Track = RiggingAnimationData.AnimationTracks[NodeIndex];
+				struct FRawAnimSequenceTrack& Track = RiggingAnimationData.AnimationTracks[NodeIndex];
 				const FName& NodeName = ValidNodeNames[NodeIndex];
 				const FName& BoneName = Skeleton->GetRigBoneMapping(NodeName);
 				const int32& BoneIndex = Skeleton->GetReferenceSkeleton().FindBoneIndex(BoneName);
@@ -2732,7 +2884,7 @@ bool UAnimSequence::ConvertAnimationDataToRiggingData(FAnimSequenceTrackContaine
 						const auto* RigConstraint = Rig->GetTransformBase(RigConstraintIndex);
 
 						// apply orientation - for now only one
-						const TArray<FRigTransformConstraint> & RotationTransformConstraint = RigConstraint->Constraints[EControlConstraint::Type::Orientation].TransformConstraints;
+						const TArray<FRigTransformConstraint>& RotationTransformConstraint = RigConstraint->Constraints[EControlConstraint::Type::Orientation].TransformConstraints;
 
 						if (RotationTransformConstraint.Num() > 0)
 						{
@@ -2768,7 +2920,7 @@ bool UAnimSequence::ConvertAnimationDataToRiggingData(FAnimSequenceTrackContaine
 						}
 
 						// apply translation - for now only one
-						const TArray<FRigTransformConstraint> & TranslationTransformConstraint = RigConstraint->Constraints[EControlConstraint::Type::Translation].TransformConstraints;
+						const TArray<FRigTransformConstraint>& TranslationTransformConstraint = RigConstraint->Constraints[EControlConstraint::Type::Translation].TransformConstraints;
 
 						if (TranslationTransformConstraint.Num() > 0)
 						{
@@ -2824,7 +2976,7 @@ bool UAnimSequence::ConvertAnimationDataToRiggingData(FAnimSequenceTrackContaine
 	return false;
 }
 
-bool UAnimSequence::ConvertRiggingDataToAnimationData(FAnimSequenceTrackContainer & RiggingAnimationData)
+bool UAnimSequence::ConvertRiggingDataToAnimationData(FAnimSequenceTrackContainer& RiggingAnimationData)
 {
 	if (RiggingAnimationData.GetNum() > 0)
 	{
@@ -2868,7 +3020,7 @@ bool UAnimSequence::ConvertRiggingDataToAnimationData(FAnimSequenceTrackContaine
 				AnimationTrackNames[NodeIndex] = BoneName;
 
 				// update bone trasfnrom
-				FRawAnimSequenceTrack & Track = RawAnimationData[NodeIndex];
+				FRawAnimSequenceTrack& Track = RawAnimationData[NodeIndex];
 
 				Track.PosKeys.Empty();
 				Track.RotKeys.Empty();
@@ -2920,9 +3072,345 @@ bool UAnimSequence::ConvertRiggingDataToAnimationData(FAnimSequenceTrackContaine
 	return false;
 }
 
+void UAnimSequence::BakeTrackCurvesToRawAnimation()
+{
+	// now bake the curves to the RawAnimationData
+	if(NumFrames == 0)
+	{
+		// fail error?
+		return;
+	}
+
+	if (!DoesContainTransformCurves())
+	{
+		if (SourceRawAnimationData.Num() > 0)
+		{
+			// if curve doesn't exists, we just bring back Source to Raw, and clears Source
+			RawAnimationData = MoveTemp(SourceRawAnimationData);
+			PostProcessSequence();
+		}
+	}
+	else
+	{
+		if(SourceRawAnimationData.Num() == 0)
+		{
+			// if source data is empty, this is first time
+			// copies the data
+			SourceRawAnimationData = RawAnimationData;
+		}
+		else
+		{
+			// we copy SourceRawAnimationData because we'd need to create additive on top of current one
+			RawAnimationData = SourceRawAnimationData;
+		}
+
+		USkeleton * CurSkeleton = GetSkeleton();
+		check(CurSkeleton);
+
+		FSmartNameMapping* NameMapping = CurSkeleton->SmartNames.GetContainer(USkeleton::AnimTrackCurveMappingName);
+		// if no mapping, that means there is no transform curves
+		if(!NameMapping)
+		{
+			// if no name mapping is found but curve exists, we should verify curve namex
+			VerifyCurveNames(CurSkeleton, USkeleton::AnimTrackCurveMappingName, RawCurveData.TransformCurves);
+			NameMapping = CurSkeleton->SmartNames.GetContainer(USkeleton::AnimTrackCurveMappingName);
+		}
+
+		for(const auto& Curve : RawCurveData.TransformCurves)
+		{
+			// find curves first, and then see what is index of this curve
+			FName BoneName;
+
+			if(Curve.GetCurveTypeFlag(ACF_Disabled)== false &&
+				ensure(NameMapping->GetName(Curve.CurveUid, BoneName)))
+			{
+				int32 TrackIndex = AnimationTrackNames.Find(BoneName);
+
+				// the animation data doesn't have this track, so insert it
+				if(TrackIndex == INDEX_NONE)
+				{
+					TrackIndex = InsertTrack(BoneName);
+					// if it still didn't find, something went horribly wrong
+					if(ensure(TrackIndex != INDEX_NONE) == false)
+					{
+						UE_LOG(LogAnimation, Warning, TEXT("Animation Baking : Error adding %s track."), *BoneName.ToString());
+						// I can't do anything about it
+						continue;
+					}
+				}
+
+				// now modify data
+				auto& RawTrack = RawAnimationData[TrackIndex];
+
+				// since now we're editing keys, 
+				// if 1 (which meant constant), just expands to # of frames
+				if(RawTrack.PosKeys.Num() == 1)
+				{
+					FVector OneKey = RawTrack.PosKeys[0];
+					RawTrack.PosKeys.Init(OneKey, NumFrames);
+				}
+				else
+				{
+					ensure(RawTrack.PosKeys.Num() == NumFrames);
+				}
+
+				if(RawTrack.RotKeys.Num() == 1)
+				{
+					FQuat OneKey = RawTrack.RotKeys[0];
+					RawTrack.RotKeys.Init(OneKey, NumFrames);
+				}
+				else
+				{
+					ensure(RawTrack.RotKeys.Num() == NumFrames);
+				}
+
+				// although we don't allow edit of scale
+				// it is important to consider scale when apply transform
+				// so make sure this also is included
+				if(RawTrack.ScaleKeys.Num() == 1)
+				{
+					FVector OneKey = RawTrack.ScaleKeys[0];
+					RawTrack.ScaleKeys.Init(OneKey, NumFrames);
+				}
+				// scale key can be 0
+				else if(RawTrack.ScaleKeys.Num() == 0)
+				{
+					FVector OneKey = FVector(1.f);
+					RawTrack.ScaleKeys.Init(OneKey, NumFrames);
+				}
+				else
+				{
+					ensure(RawTrack.ScaleKeys.Num() == NumFrames);
+				}
+
+				// NumFrames can't be zero (filtered earlier)
+				float Interval = (NumFrames > 1)? SequenceLength/(NumFrames-1) : MINIMUM_ANIMATION_LENGTH;
+
+				// now we have all data ready to apply
+				for(int32 KeyIndex=0; KeyIndex < NumFrames; ++KeyIndex)
+				{
+					// now evaluate
+					FTransformCurve* TransformCurve = static_cast<FTransformCurve*>(RawCurveData.GetCurveData(Curve.CurveUid, FRawCurveTracks::TransformType));
+
+					if(ensure(TransformCurve))
+					{
+						FTransform AdditiveTransform = TransformCurve->Evaluate(KeyIndex * Interval, 1.0);
+						FTransform LocalTransform(RawTrack.RotKeys[KeyIndex], RawTrack.PosKeys[KeyIndex], RawTrack.ScaleKeys[KeyIndex]);
+						//  						LocalTransform = LocalTransform * AdditiveTransform;
+						//  						RawTrack.RotKeys[KeyIndex] = LocalTransform.GetRotation();
+						//  						RawTrack.PosKeys[KeyIndex] = LocalTransform.GetTranslation();
+						//  						RawTrack.ScaleKeys[KeyIndex] = LocalTransform.GetScale3D();
+
+						RawTrack.RotKeys[KeyIndex] = LocalTransform.GetRotation() * AdditiveTransform.GetRotation();
+						RawTrack.PosKeys[KeyIndex] = LocalTransform.TransformPosition(AdditiveTransform.GetTranslation());
+						RawTrack.ScaleKeys[KeyIndex] = LocalTransform.GetScale3D() * AdditiveTransform.GetScale3D();
+					}
+					else
+					{
+						UE_LOG(LogAnimation, Warning, TEXT("Animation Baking : Missing Curve for %s."), *BoneName.ToString());
+					}
+				}
+			}
+		}
+
+		PostProcessSequence();
+	}
+
+	bNeedsRebake = false;
+}
+
+bool UAnimSequence::DoesNeedRebake() const
+{
+	return (bNeedsRebake);
+}
+
+bool UAnimSequence::DoesContainTransformCurves() const
+{
+	return (RawCurveData.TransformCurves.Num() > 0);
+}
+
+void UAnimSequence::AddKeyToSequence(float Time, const FName& BoneName, const FTransform& AdditiveTransform)
+{
+	// find if this already exists, then just add curve data only
+	FName CurveName = BoneName;
+	USkeleton * CurrentSkeleton = GetSkeleton();
+	check (CurrentSkeleton);
+
+	FSmartNameMapping* NameMapping = CurrentSkeleton->SmartNames.GetContainer(USkeleton::AnimTrackCurveMappingName);
+	if(NameMapping == NULL)
+	{
+		CurrentSkeleton->Modify(true);
+		CurrentSkeleton->SmartNames.AddContainer(USkeleton::AnimTrackCurveMappingName);
+		NameMapping = CurrentSkeleton->SmartNames.GetContainer(USkeleton::AnimTrackCurveMappingName);
+		check(NameMapping);
+	}
+
+	FSmartNameMapping::UID NewUID;
+	NameMapping->AddOrFindName(CurveName, NewUID);
+
+	// add curve - this won't add duplicate curve
+	RawCurveData.AddCurveData(NewUID, ACF_DriveTrack | ACF_Editable, FRawCurveTracks::TransformType);
+
+	//Add this curve
+	FTransformCurve* TransformCurve = static_cast<FTransformCurve*>(RawCurveData.GetCurveData(NewUID, FRawCurveTracks::TransformType));
+	check(TransformCurve);
+
+	TransformCurve->UpdateOrAddKey(AdditiveTransform, Time);
+
+	bNeedsRebake = true;
+}
+
+void UAnimSequence::ResetAnimation()
+{
+	// clear everything. Making new animation, so need to reset all the things that belong here
+	NumFrames = 0;
+	SequenceLength = 0.f;
+	RawAnimationData.Empty();
+	SourceRawAnimationData.Empty();
+	AnimationTrackNames.Empty();
+	TrackToSkeletonMapTable.Empty();
+	CompressedTrackOffsets.Empty();
+	CompressedScaleOffsets.Empty();
+	CompressedByteStream.Empty();
+
+	Notifies.Empty();
+	AnimNotifyTracks.Empty();
+	RawCurveData.Empty();
+	RateScale = 1.f;
+}
+
+void UAnimSequence::RefreshTrackMapFromAnimTrackNames()
+{
+	TrackToSkeletonMapTable.Empty();
+
+	const USkeleton * MySkeleton = GetSkeleton();
+	const FReferenceSkeleton& RefSkeleton = MySkeleton->GetReferenceSkeleton();
+	const int32 NumBones = AnimationTrackNames.Num();
+	TrackToSkeletonMapTable.AddUninitialized(NumBones);
+	for(int32 BoneIndex=0; BoneIndex<AnimationTrackNames.Num(); ++BoneIndex)
+	{
+		int32 BoneTreeIndex = RefSkeleton.FindBoneIndex(AnimationTrackNames[BoneIndex]);
+		TrackToSkeletonMapTable[BoneIndex] = FTrackToSkeletonMap(BoneTreeIndex);
+	}
+}
+
+bool UAnimSequence::CreateAnimation(USkeletalMesh * Mesh)
+{
+	// create animation from Mesh's ref pose
+	if (Mesh)
+	{
+		ResetAnimation();
+
+		const FReferenceSkeleton& RefSkeleton = Mesh->RefSkeleton;
+		SequenceLength = MINIMUM_ANIMATION_LENGTH;
+		NumFrames = 1;
+
+		const int32 NumBones = RefSkeleton.GetNum();
+		RawAnimationData.AddZeroed(NumBones);
+		AnimationTrackNames.AddUninitialized(NumBones);
+
+		const TArray<FTransform>& RefBonePose = RefSkeleton.GetRefBonePose();
+
+		check (RefBonePose.Num() == NumBones);
+
+		for (int32 BoneIndex=0; BoneIndex<NumBones; ++BoneIndex)
+		{
+			AnimationTrackNames[BoneIndex] = RefSkeleton.GetBoneName(BoneIndex);
+
+			FRawAnimSequenceTrack& RawTrack = RawAnimationData[BoneIndex];
+
+			RawTrack.PosKeys.Add(RefBonePose[BoneIndex].GetTranslation());
+			RawTrack.RotKeys.Add(RefBonePose[BoneIndex].GetRotation());
+			RawTrack.ScaleKeys.Add(RefBonePose[BoneIndex].GetScale3D());
+		}
+
+		// refresh TrackToskeletonMapIndex
+		RefreshTrackMapFromAnimTrackNames();
+
+		// should recreate track map
+		PostProcessSequence();
+		return true;
+	}
+
+	return false;
+}
+
+bool UAnimSequence::CreateAnimation(USkeletalMeshComponent * MeshComponent)
+{
+	if(MeshComponent && MeshComponent->SkeletalMesh)
+	{
+		USkeletalMesh * Mesh = MeshComponent->SkeletalMesh;
+
+		ResetAnimation();
+
+		const FReferenceSkeleton& RefSkeleton = Mesh->RefSkeleton;
+		SequenceLength = MINIMUM_ANIMATION_LENGTH;
+		NumFrames = 1;
+
+		const int32 NumBones = RefSkeleton.GetNum();
+		RawAnimationData.AddZeroed(NumBones);
+		AnimationTrackNames.AddUninitialized(NumBones);
+
+		const TArray<FTransform>& LocalAtoms = MeshComponent->LocalAtoms;
+
+		check(LocalAtoms.Num() == NumBones);
+
+		for(int32 BoneIndex=0; BoneIndex<NumBones; ++BoneIndex)
+		{
+			AnimationTrackNames[BoneIndex] = RefSkeleton.GetBoneName(BoneIndex);
+
+			FRawAnimSequenceTrack& RawTrack = RawAnimationData[BoneIndex];
+
+			RawTrack.PosKeys.Add(LocalAtoms[BoneIndex].GetTranslation());
+			RawTrack.RotKeys.Add(LocalAtoms[BoneIndex].GetRotation());
+			RawTrack.ScaleKeys.Add(LocalAtoms[BoneIndex].GetScale3D());
+		}
+
+		// refresh TrackToskeletonMapIndex
+		RefreshTrackMapFromAnimTrackNames();
+
+		// should recreate track map
+		PostProcessSequence();
+		return true;
+	}
+
+	return false;
+}
+
+bool UAnimSequence::CreateAnimation(UAnimSequence * Sequence)
+{
+	if(Sequence)
+	{
+		ResetAnimation();
+
+		SequenceLength = Sequence->SequenceLength;
+		NumFrames = Sequence->NumFrames;
+
+		RawAnimationData = Sequence->RawAnimationData;
+		AnimationTrackNames = Sequence->AnimationTrackNames;
+
+		Notifies = Sequence->Notifies;
+		AnimNotifyTracks = Sequence->AnimNotifyTracks;
+		RawCurveData = Sequence->RawCurveData;
+
+		// refresh TrackToskeletonMapIndex
+		RefreshTrackMapFromAnimTrackNames();
+
+		// should recreate track map
+		PostProcessSequence();
+		return true;
+	}
+	
+	return false;
+}
+
+bool UAnimSequence::Resize(int32 Start, int32 End)
+{
+	return false;
+}
 #endif
 /*-----------------------------------------------------------------------------
-	AnimNotify & subclasses
+	AnimNotify& subclasses
 -----------------------------------------------------------------------------*/
 
 
diff --git a/Engine/Source/Runtime/Engine/Private/Animation/AnimSequenceBase.cpp b/Engine/Source/Runtime/Engine/Private/Animation/AnimSequenceBase.cpp
--- a/Engine/Source/Runtime/Engine/Private/Animation/AnimSequenceBase.cpp
+++ b/Engine/Source/Runtime/Engine/Private/Animation/AnimSequenceBase.cpp
@@ -86,7 +86,7 @@ void FAnimNotifyEvent::SetTime(float NewTime, EAnimLinkMethod::Type ReferenceFra
 /////////////////////////////////////////////////////
 // FFloatCurve
 
-void FFloatCurve::SetCurveTypeFlag(EAnimCurveFlags InFlag, bool bValue)
+void FAnimCurveBase::SetCurveTypeFlag(EAnimCurveFlags InFlag, bool bValue)
 {
 	if (bValue)
 	{
@@ -98,81 +98,211 @@ void FFloatCurve::SetCurveTypeFlag(EAnimCurveFlags InFlag, bool bValue)
 	}
 }
 
-void FFloatCurve::ToggleCurveTypeFlag(EAnimCurveFlags InFlag)
+void FAnimCurveBase::ToggleCurveTypeFlag(EAnimCurveFlags InFlag)
 {
 	bool Current = GetCurveTypeFlag(InFlag);
 	SetCurveTypeFlag(InFlag, !Current);
 }
 
-bool FFloatCurve::GetCurveTypeFlag(EAnimCurveFlags InFlag) const
+bool FAnimCurveBase::GetCurveTypeFlag(EAnimCurveFlags InFlag) const
 {
 	return (CurveTypeFlags & InFlag) != 0;
 }
 
 
-void FFloatCurve::SetCurveTypeFlags(int32 NewCurveTypeFlags)
+void FAnimCurveBase::SetCurveTypeFlags(int32 NewCurveTypeFlags)
 {
 	CurveTypeFlags = NewCurveTypeFlags;
 }
 
-int32 FFloatCurve::GetCurveTypeFlags() const
+int32 FAnimCurveBase::GetCurveTypeFlags() const
 {
 	return CurveTypeFlags;
 }
 
+////////////////////////////////////////////////////
+//  FFloatCurve
+
+// we don't want to have = operator. This only copies curves, but leaving naming and everything else intact. 
+void FFloatCurve::CopyCurve(FFloatCurve& SourceCurve)
+{
+	FloatCurve = SourceCurve.FloatCurve;
+}
+
+float FFloatCurve::Evaluate(float CurrentTime, float BlendWeight) const
+{
+	return FloatCurve.Eval(CurrentTime)*BlendWeight;
+}
+
+void FFloatCurve::UpdateOrAddKey(float NewKey, float CurrentTime)
+{
+	FloatCurve.UpdateOrAddKey(CurrentTime, NewKey);
+}
+
+////////////////////////////////////////////////////
+//  FVectorCurve
+
+// we don't want to have = operator. This only copies curves, but leaving naming and everything else intact. 
+void FVectorCurve::CopyCurve(FVectorCurve& SourceCurve)
+{
+	FMemory::MemCopy(FloatCurves, SourceCurve.FloatCurves);
+}
+
+FVector FVectorCurve::Evaluate(float CurrentTime, float BlendWeight) const
+{
+	FVector Value;
+
+	Value.X = FloatCurves[X].Eval(CurrentTime)*BlendWeight;
+	Value.Y = FloatCurves[Y].Eval(CurrentTime)*BlendWeight;
+	Value.Z = FloatCurves[Z].Eval(CurrentTime)*BlendWeight;
+
+	return Value;
+}
+
+void FVectorCurve::UpdateOrAddKey(const FVector& NewKey, float CurrentTime)
+{
+	FloatCurves[X].UpdateOrAddKey(CurrentTime, NewKey.X);
+	FloatCurves[Y].UpdateOrAddKey(CurrentTime, NewKey.Y);
+	FloatCurves[Z].UpdateOrAddKey(CurrentTime, NewKey.Z);
+}
+////////////////////////////////////////////////////
+//  FTransformCurve
+
+// we don't want to have = operator. This only copies curves, but leaving naming and everything else intact. 
+void FTransformCurve::CopyCurve(FTransformCurve& SourceCurve)
+{
+	TranslationCurve.CopyCurve(SourceCurve.TranslationCurve);
+	RotationCurve.CopyCurve(SourceCurve.RotationCurve);
+	ScaleCurve.CopyCurve(SourceCurve.ScaleCurve);
+}
+
+FTransform FTransformCurve::Evaluate(float CurrentTime, float BlendWeight) const
+{
+	FTransform Value;
+	Value.SetTranslation(TranslationCurve.Evaluate(CurrentTime, BlendWeight));
+	Value.SetScale3D(ScaleCurve.Evaluate(CurrentTime, BlendWeight));
+
+	// blend rotation float curve
+	FVector RotationAsVector = RotationCurve.Evaluate(CurrentTime, BlendWeight);
+	// pitch, yaw, roll order - please check AddKey function
+	FRotator Rotator(RotationAsVector.Y, RotationAsVector.Z, RotationAsVector.X);
+	Value.SetRotation(FQuat(Rotator));
+
+	return Value;
+}
+
+void FTransformCurve::UpdateOrAddKey(const FTransform& NewKey, float CurrentTime)
+{
+	TranslationCurve.UpdateOrAddKey(NewKey.GetTranslation(), CurrentTime);
+	// pitch, yaw, roll order - please check Evaluate function
+	FVector RotationAsVector;
+	FRotator Rotator = NewKey.GetRotation().Rotator();
+	RotationAsVector.X = Rotator.Roll;
+	RotationAsVector.Y = Rotator.Pitch;
+	RotationAsVector.Z = Rotator.Yaw;
+
+	RotationCurve.UpdateOrAddKey(RotationAsVector, CurrentTime);
+	ScaleCurve.UpdateOrAddKey(NewKey.GetScale3D(), CurrentTime);
+}
 /////////////////////////////////////////////////////
 // FRawCurveTracks
 
 void FRawCurveTracks::EvaluateCurveData(class UAnimInstance* Instance, float CurrentTime, float BlendWeight ) const
 {
 	// evaluate the curve data at the CurrentTime and add to Instance
-	for (auto CurveIter = FloatCurves.CreateConstIterator(); CurveIter; ++CurveIter)
+	for(auto CurveIter = FloatCurves.CreateConstIterator(); CurveIter; ++CurveIter)
 	{
 		const FFloatCurve& Curve = *CurveIter;
 
-		Instance->AddCurveValue( Curve.CurveUid, Curve.FloatCurve.Eval(CurrentTime)*BlendWeight, Curve.GetCurveTypeFlags() );
+		Instance->AddCurveValue(Curve.CurveUid, Curve.Evaluate(CurrentTime, BlendWeight), Curve.GetCurveTypeFlags());
 	}
 }
 
-FFloatCurve * FRawCurveTracks::GetCurveData(USkeleton::AnimCurveUID Uid)
+#if WITH_EDITOR
+/**
+ * Since we don't care about blending, we just change this decoration to OutCurves
+ * @TODO : Fix this if we're saving vectorcurves and blending
+ */
+void FRawCurveTracks::EvaluateTransformCurveData(USkeleton * Skeleton, TMap<FName, FTransform>&OutCurves, float CurrentTime, float BlendWeight) const
 {
-	for(FFloatCurve& Curve : FloatCurves)
+	check (Skeleton);
+	// evaluate the curve data at the CurrentTime and add to Instance
+	for(auto CurveIter = TransformCurves.CreateConstIterator(); CurveIter; ++CurveIter)
 	{
-		if(Curve.CurveUid == Uid)
+		const FTransformCurve& Curve = *CurveIter;
+
+		// if disabled, do not handle
+		if (Curve.GetCurveTypeFlag(ACF_Disabled))
 		{
-			return &Curve;
+			continue;
 		}
-	}
 
-	return NULL;
-}
+		FSmartNameMapping* NameMapping = Skeleton->SmartNames.GetContainer(USkeleton::AnimTrackCurveMappingName);
 
-bool FRawCurveTracks::DeleteCurveData(USkeleton::AnimCurveUID Uid)
-{
-	for( int32 Idx = 0; Idx < FloatCurves.Num(); ++Idx )
-	{
-		if( FloatCurves[Idx].CurveUid == Uid )
+		// Add or retrieve curve
+		FName CurveName;
+		
+		// make sure it was added
+		if (ensure (NameMapping->GetName(Curve.CurveUid, CurveName)))
 		{
-			FloatCurves.RemoveAt(Idx);
-			return true;
+			// note we're not checking Curve.GetCurveTypeFlags() yet
+			FTransform & Value = OutCurves.FindOrAdd(CurveName);
+			Value = Curve.Evaluate(CurrentTime, BlendWeight);
 		}
 	}
+}
+#endif
+FAnimCurveBase * FRawCurveTracks::GetCurveData(USkeleton::AnimCurveUID Uid, ESupportedCurveType SupportedCurveType /*= FloatType*/)
+{
+	switch (SupportedCurveType)
+	{
+#if WITH_EDITOR
+	case VectorType:
+		return GetCurveDataImpl<FVectorCurve>(VectorCurves, Uid);
+	case TransformType:
+		return GetCurveDataImpl<FTransformCurve>(TransformCurves, Uid);
+#endif // WITH_EDITOR
+	case FloatType:
+	default:
+		return GetCurveDataImpl<FFloatCurve>(FloatCurves, Uid);
+	}
+}
 
-	return false;
+bool FRawCurveTracks::DeleteCurveData(USkeleton::AnimCurveUID Uid, ESupportedCurveType SupportedCurveType /*= FloatType*/)
+{
+	switch(SupportedCurveType)
+	{
+#if WITH_EDITOR
+	case VectorType:
+		return DeleteCurveDataImpl<FVectorCurve>(VectorCurves, Uid);
+	case TransformType:
+		return DeleteCurveDataImpl<FTransformCurve>(TransformCurves, Uid);
+#endif // WITH_EDITOR
+	case FloatType:
+	default:
+		return DeleteCurveDataImpl<FFloatCurve>(FloatCurves, Uid);
+	}
 }
 
-bool FRawCurveTracks::AddCurveData(USkeleton::AnimCurveUID Uid, int32 CurveFlags /*= ACF_DefaultCurve*/)
+bool FRawCurveTracks::AddCurveData(USkeleton::AnimCurveUID Uid, int32 CurveFlags /*= ACF_DefaultCurve*/, ESupportedCurveType SupportedCurveType /*= FloatType*/)
 {
-	if(GetCurveData(Uid) == NULL)
+	switch(SupportedCurveType)
 	{
-		FloatCurves.Add(FFloatCurve(Uid, CurveFlags));
-		return true;
+#if WITH_EDITOR
+	case VectorType:
+		return AddCurveDataImpl<FVectorCurve>(VectorCurves, Uid, CurveFlags);
+	case TransformType:
+		return AddCurveDataImpl<FTransformCurve>(TransformCurves, Uid, CurveFlags);
+#endif // WITH_EDITOR
+	case FloatType:
+	default:
+		return AddCurveDataImpl<FFloatCurve>(FloatCurves, Uid, CurveFlags);
 	}
-	return false;
 }
 
 void FRawCurveTracks::Serialize(FArchive& Ar)
 {
+	// @TODO: If we're about to serialize vector curve, add here
 	if(Ar.UE4Ver() >= VER_UE4_SKELETON_ADD_SMARTNAMES)
 	{
 		for(FFloatCurve& Curve : FloatCurves)
@@ -180,28 +310,123 @@ void FRawCurveTracks::Serialize(FArchive& Ar)
 			Curve.Serialize(Ar);
 		}
 	}
+#if WITH_EDITOR
+	if(Ar.UE4Ver() >= VER_UE4_ANIMATION_ADD_TRACKCURVES)
+	{
+		for(FTransformCurve& Curve : TransformCurves)
+		{
+			Curve.Serialize(Ar);
+		}
+	}
+#endif // WITH_EDITOR
+}
+
+void FRawCurveTracks::UpdateLastObservedNames(FSmartNameMapping* NameMapping, ESupportedCurveType SupportedCurveType /*= FloatType*/)
+{
+	switch(SupportedCurveType)
+	{
+#if WITH_EDITOR
+	case VectorType:
+		UpdateLastObservedNamesImpl<FVectorCurve>(VectorCurves, NameMapping);
+		break;
+	case TransformType:
+		UpdateLastObservedNamesImpl<FTransformCurve>(TransformCurves, NameMapping);
+		break;
+#endif // WITH_EDITOR
+	case FloatType:
+	default:
+		UpdateLastObservedNamesImpl<FFloatCurve>(FloatCurves, NameMapping);
+	}
 }
 
-ENGINE_API void FRawCurveTracks::UpdateLastObservedNames(FSmartNameMapping* NameMapping)
+bool FRawCurveTracks::DuplicateCurveData(USkeleton::AnimCurveUID ToCopyUid, USkeleton::AnimCurveUID NewUid, ESupportedCurveType SupportedCurveType /*= FloatType*/)
+{
+	switch(SupportedCurveType)
+	{
+#if WITH_EDITOR
+	case VectorType:
+		return DuplicateCurveDataImpl<FVectorCurve>(VectorCurves, ToCopyUid, NewUid);
+	case TransformType:
+		return DuplicateCurveDataImpl<FTransformCurve>(TransformCurves, ToCopyUid, NewUid);
+#endif // WITH_EDITOR
+	case FloatType:
+	default:
+		return DuplicateCurveDataImpl<FFloatCurve>(FloatCurves, ToCopyUid, NewUid);
+	}
+}
+
+///////////////////////////////////
+// @TODO: REFACTOR THIS IF WE'RE SERIALIZING VECTOR CURVES
+//
+// implementation template functions to accomodate FloatCurve and VectorCurve
+// for now vector curve isn't used in run-time, so it's useless outside of editor
+// so just to reduce cost of run-time, functionality is split. 
+// this split worries me a bit because if the name conflict happens this will break down w.r.t. smart naming
+// currently vector curve is not saved and not evaluated, so it will be okay since the name doesn't matter much, 
+// but this has to be refactored once we'd like to move onto serialize
+///////////////////////////////////
+template <typename DataType>
+DataType * FRawCurveTracks::GetCurveDataImpl(TArray<DataType> & Curves, USkeleton::AnimCurveUID Uid)
+{
+	for(DataType& Curve : Curves)
+	{
+		if(Curve.CurveUid == Uid)
+		{
+			return &Curve;
+		}
+	}
+
+	return NULL;
+}
+
+template <typename DataType>
+bool FRawCurveTracks::DeleteCurveDataImpl(TArray<DataType> & Curves, USkeleton::AnimCurveUID Uid)
+{
+	for(int32 Idx = 0; Idx < Curves.Num(); ++Idx)
+	{
+		if(Curves[Idx].CurveUid == Uid)
+		{
+			Curves.RemoveAt(Idx);
+			return true;
+		}
+	}
+
+	return false;
+}
+
+template <typename DataType>
+bool FRawCurveTracks::AddCurveDataImpl(TArray<DataType> & Curves, USkeleton::AnimCurveUID Uid, int32 CurveFlags)
+{
+	if(GetCurveDataImpl<DataType>(Curves, Uid) == NULL)
+	{
+		Curves.Add(DataType(Uid, CurveFlags));
+		return true;
+	}
+	return false;
+}
+
+template <typename DataType>
+void FRawCurveTracks::UpdateLastObservedNamesImpl(TArray<DataType> & Curves, FSmartNameMapping* NameMapping)
 {
 	if(NameMapping)
 	{
-		for(FFloatCurve& Curve : FloatCurves)
+		for(DataType& Curve : Curves)
 		{
 			NameMapping->GetName(Curve.CurveUid, Curve.LastObservedName);
 		}
 	}
 }
 
-ENGINE_API bool FRawCurveTracks::DuplicateCurveData(USkeleton::AnimCurveUID ToCopyUid, USkeleton::AnimCurveUID NewUid)
+template <typename DataType>
+bool FRawCurveTracks::DuplicateCurveDataImpl(TArray<DataType> & Curves, USkeleton::AnimCurveUID ToCopyUid, USkeleton::AnimCurveUID NewUid)
 {
-	FFloatCurve* ExistingCurve = GetCurveData(ToCopyUid);
-	if(ExistingCurve && GetCurveData(NewUid) == NULL)
+	DataType* ExistingCurve = GetCurveDataImpl<DataType>(Curves, ToCopyUid);
+	if(ExistingCurve && GetCurveDataImpl<DataType>(Curves, NewUid) == NULL)
 	{
 		// Add the curve to the track and set its data to the existing curve
-		FloatCurves.Add(FFloatCurve(NewUid, ExistingCurve->GetCurveTypeFlags()));
-		FloatCurves.Last().FloatCurve = ExistingCurve->FloatCurve;
-		
+		Curves.Add(DataType(NewUid, ExistingCurve->GetCurveTypeFlags()));
+		Curves.Last().CopyCurve(*ExistingCurve);
+
 		return true;
 	}
 	return false;
@@ -215,6 +440,40 @@ UAnimSequenceBase::UAnimSequenceBase(const FObjectInitializer& ObjectInitializer
 {
 }
 
+template <typename DataType>
+void UAnimSequenceBase::VerifyCurveNames(USkeleton* Skeleton, const FName& NameContainer, TArray<DataType>& CurveList)
+{
+	FSmartNameMapping* NameMapping = Skeleton->SmartNames.GetContainer(NameContainer);
+
+	// since this is verify function that makes sure it exists after loaded
+	// we should add it if it doesn't exist
+	if (!NameMapping)
+	{
+		// if it doens't exists, we should add it
+		Skeleton->Modify(true);
+		Skeleton->SmartNames.AddContainer(NameContainer);
+		NameMapping = Skeleton->SmartNames.GetContainer(NameContainer);
+		check(NameMapping);
+	}
+
+	TArray<DataType*> UnlinkedCurves;
+	for(DataType& Curve : CurveList)
+	{
+		if(!NameMapping->Exists(Curve.LastObservedName))
+		{
+			// The skeleton doesn't know our name. Use the last observed name that was saved with the
+			// curve to create a new name. This can happen if a user saves an animation but not a skeleton
+			// either when updating the assets or editing the curves within.
+			UnlinkedCurves.Add(&Curve);
+		}
+	}
+
+	for(DataType* Curve : UnlinkedCurves)
+	{
+		NameMapping->AddOrFindName(Curve->LastObservedName, Curve->CurveUid);
+	}
+}
+
 void UAnimSequenceBase::PostLoad()
 {
 	Super::PostLoad();
@@ -264,37 +523,25 @@ void UAnimSequenceBase::PostLoad()
 
 	if(USkeleton* Skeleton = GetSkeleton())
 	{
-		// Get the name mapping object for curves
-		FSmartNameMapping* NameMapping = Skeleton->SmartNames.GetContainer(USkeleton::AnimCurveMappingName);
-		
 		// Fix up the existing curves to work with smartnames
 		if(GetLinkerUE4Version() < VER_UE4_SKELETON_ADD_SMARTNAMES)
 		{
+			// Get the name mapping object for curves
+			FSmartNameMapping* NameMapping = Skeleton->SmartNames.GetContainer(USkeleton::AnimCurveMappingName);
 			for(FFloatCurve& Curve : RawCurveData.FloatCurves)
 			{
 				// Add the names of the curves into the smartname mapping and store off the curve uid which will be saved next time the sequence saves.
-				NameMapping->AddName(Curve.LastObservedName, Curve.CurveUid);
+				NameMapping->AddOrFindName(Curve.LastObservedName, Curve.CurveUid);
 			}
 		}
 		else
 		{
-			TArray<FFloatCurve*> UnlinkedCurves;
-			for(FFloatCurve& Curve : RawCurveData.FloatCurves)
-			{
-				if(!NameMapping->Exists(Curve.LastObservedName))
-				{
-					// The skeleton doesn't know our name. Use the last observed name that was saved with the
-					// curve to create a new name. This can happen if a user saves an animation but not a skeleton
-					// either when updating the assets or editing the curves within.
-					UnlinkedCurves.Add(&Curve);
-				}
-			}
-
-			for(FFloatCurve* Curve : UnlinkedCurves)
-			{
-				NameMapping->AddName(Curve->LastObservedName, Curve->CurveUid);
-			}
+			VerifyCurveNames<FFloatCurve>(Skeleton, USkeleton::AnimCurveMappingName, RawCurveData.FloatCurves);
 		}
+
+#if WITH_EDITOR
+		VerifyCurveNames<FTransformCurve>(Skeleton, USkeleton::AnimTrackCurveMappingName, RawCurveData.TransformCurves);
+#endif
 	}
 }
 
@@ -650,6 +897,26 @@ void UAnimSequenceBase::Serialize(FArchive& Ar)
 			RawCurveData.UpdateLastObservedNames(Mapping);
 		}
 	}
+
+#if WITH_EDITORONLY_DATA
+	if(Ar.ArIsSaving && Ar.UE4Ver() >= VER_UE4_ANIMATION_ADD_TRACKCURVES)
+	{
+		if(USkeleton* Skeleton = GetSkeleton())
+		{
+			// we don't add track curve container unless it has been editied. 
+			if ( RawCurveData.TransformCurves.Num() > 0 )
+			{
+				FSmartNameMapping* Mapping = GetSkeleton()->SmartNames.GetContainer(USkeleton::AnimTrackCurveMappingName);
+				// this name might not exists because it's only available if you edit animation
+				if (Mapping)
+				{
+					RawCurveData.UpdateLastObservedNames(Mapping, FRawCurveTracks::TransformType);
+				}
+			}
+		}
+	}
+#endif //WITH_EDITORONLY_DATA
+
 	RawCurveData.Serialize(Ar);
 }
 
diff --git a/Engine/Source/Runtime/Engine/Private/Animation/Skeleton.cpp b/Engine/Source/Runtime/Engine/Private/Animation/Skeleton.cpp
--- a/Engine/Source/Runtime/Engine/Private/Animation/Skeleton.cpp
+++ b/Engine/Source/Runtime/Engine/Private/Animation/Skeleton.cpp
@@ -26,6 +26,7 @@ const FName USkeleton::RigTag = FName(TEXT("Rig"));
 
 // Names of smartname containers for skeleton properties
 const FName USkeleton::AnimCurveMappingName = FName(TEXT("AnimationCurves"));
+const FName USkeleton::AnimTrackCurveMappingName = FName(TEXT("AnimationTrackCurves"));
 
 const FName FAnimSlotGroup::DefaultGroupName = FName(TEXT("DefaultGroup"));
 const FName FAnimSlotGroup::DefaultSlotName = FName(TEXT("DefaultSlot"));
@@ -914,7 +915,7 @@ bool USkeleton::AddSmartnameAndModify(FName ContainerName, FName NewName, FSmart
 	FSmartNameMapping* RequestedMapping = SmartNames.GetContainer(ContainerName);
 	if(RequestedMapping)
 	{
-		if(RequestedMapping->AddName(NewName, NewUid))
+		if(RequestedMapping->AddOrFindName(NewName, NewUid))
 		{
 			Modify(true);
 			Successful = true;
diff --git a/Engine/Source/Runtime/Engine/Private/Animation/SmartName.cpp b/Engine/Source/Runtime/Engine/Private/Animation/SmartName.cpp
--- a/Engine/Source/Runtime/Engine/Private/Animation/SmartName.cpp
+++ b/Engine/Source/Runtime/Engine/Private/Animation/SmartName.cpp
@@ -7,7 +7,7 @@ FSmartNameMapping::FSmartNameMapping()
 : NextUid(0)
 {}
 
-bool FSmartNameMapping::AddName(FName Name, UID& OutUid)
+bool FSmartNameMapping::AddOrFindName(FName Name, UID& OutUid)
 {
 	check(Name.IsValid());
 
diff --git a/Engine/Source/Runtime/Engine/Public/AnimationRuntime.h b/Engine/Source/Runtime/Engine/Public/AnimationRuntime.h
--- a/Engine/Source/Runtime/Engine/Public/AnimationRuntime.h
+++ b/Engine/Source/Runtime/Engine/Public/AnimationRuntime.h
@@ -263,12 +263,13 @@ public:
 	static FTransform GetSpaceTransform(FA2CSPose& Pose, int32 Index);
 	static void SetSpaceTransform(FA2Pose& Pose, int32 Index, FTransform& NewTransform);
 	static void SetSpaceTransform(FA2CSPose& Pose, int32 Index, FTransform& NewTransform);
-
 	// space bases
-	static void FillUpSpaceBasesRefPose(const USkeleton* Skeleton, TArray<FTransform> &SpaceBaseRefPose);
 #if WITH_EDITOR
+	static void FillUpSpaceBases(const FReferenceSkeleton & RefSkeleton, const TArray<FTransform> &LocalAtoms, TArray<FTransform> &SpaceBases);
+	static void FillUpSpaceBasesRefPose(const USkeleton* Skeleton, TArray<FTransform> &SpaceBaseRefPose);
 	static void FillUpSpaceBasesRetargetBasePose(const USkeleton* Skeleton, TArray<FTransform> &SpaceBaseRefPose);
 #endif
+
 private:
 	/** 
 	* Blend Poses per bone weights : The BasePoses + BlendPoses(SourceIndex) * Blend Weights(BoneIndex)
diff --git a/Engine/Source/Runtime/Engine/Public/BoneContainer.h b/Engine/Source/Runtime/Engine/Public/BoneContainer.h
--- a/Engine/Source/Runtime/Engine/Public/BoneContainer.h
+++ b/Engine/Source/Runtime/Engine/Public/BoneContainer.h
@@ -39,6 +39,8 @@ private:
 	bool bDisableRetargeting;
 	/** Disable animation compression, use RAW data instead. */
 	bool bUseRAWData;
+	/** Use Source Data that is imported that are not compressed. */
+	bool bUseSourceData;
 
 public:
 
@@ -97,6 +99,18 @@ public:
 		return bUseRAWData;
 	}
 
+	/** use Source data instead.*/
+	void SetUseSourceData(bool InbUseSourceData)
+	{
+		bUseSourceData = InbUseSourceData;
+	}
+
+	/** True if we're requesting Source data instead of RawAnimationData. For debugging. */
+	bool ShouldUseSourceData() const
+	{
+		return bUseSourceData;
+	}
+
 	/**
 	* returns Required Bone Indices Array
 	*/
