potential compile fix

diff --git a/Engine/Source/Runtime/Engine/Classes/Animation/AnimSequence.h b/Engine/Source/Runtime/Engine/Classes/Animation/AnimSequence.h
--- a/Engine/Source/Runtime/Engine/Classes/Animation/AnimSequence.h
+++ b/Engine/Source/Runtime/Engine/Classes/Animation/AnimSequence.h
@@ -791,15 +791,15 @@ public:
 	/**
 	 * Create Animation Sequence from Reference Pose of the Mesh
 	 */
-	ENGINE_API bool CreateAnimation(USkeletalMesh * Mesh);
+	ENGINE_API bool CreateAnimation(class USkeletalMesh * Mesh);
 	/**
 	 * Create Animation Sequence from the Mesh Component's current bone trasnform
 	 */
-	ENGINE_API bool CreateAnimation(USkeletalMeshComponent * MeshComponent);
+	ENGINE_API bool CreateAnimation(class USkeletalMeshComponent * MeshComponent);
 	/**
 	 * Create Animation Sequence from the given animation
 	 */
-	ENGINE_API bool CreateAnimation(UAnimSequence * Sequence);
+	ENGINE_API bool CreateAnimation(class UAnimSequence * Sequence);
 	/**
 	 * Resize Animation to Start/End. If End is bigger, it repeats last frame upto End. 
 	 * @todo implement
