Adding ability to set amount of loops on looping node

diff --git a/Engine/Source/Runtime/Engine/Classes/Sound/SoundNodeLooping.h b/Engine/Source/Runtime/Engine/Classes/Sound/SoundNodeLooping.h
--- a/Engine/Source/Runtime/Engine/Classes/Sound/SoundNodeLooping.h
+++ b/Engine/Source/Runtime/Engine/Classes/Sound/SoundNodeLooping.h
@@ -13,6 +13,14 @@ class USoundNodeLooping : public USoundNode
 {
 	GENERATED_UCLASS_BODY()
 
+	/* The amount of times to loop */
+	UPROPERTY(EditAnywhere, Category = Looping, meta = (ClampMin = 1))
+	int32 LoopCount;
+
+	/* If enabled, the node will continue to loop indefinitely regardless of the Loop Count value. */
+	UPROPERTY(EditAnywhere, Category = Looping)
+	uint32 bLoopIndefinitely : 1;
+
 public:	
 	// Begin USoundNode interface. 
 	virtual bool NotifyWaveInstanceFinished( struct FWaveInstance* WaveInstance ) override;
diff --git a/Engine/Source/Runtime/Engine/Private/SoundNodeLooping.cpp b/Engine/Source/Runtime/Engine/Private/SoundNodeLooping.cpp
--- a/Engine/Source/Runtime/Engine/Private/SoundNodeLooping.cpp
+++ b/Engine/Source/Runtime/Engine/Private/SoundNodeLooping.cpp
@@ -11,11 +11,23 @@
 -----------------------------------------------------------------------------*/
 USoundNodeLooping::USoundNodeLooping(const FObjectInitializer& ObjectInitializer)
 	: Super(ObjectInitializer)
+	, LoopCount(1)
+	, bLoopIndefinitely(true)
 {
 }
 
 void USoundNodeLooping::ParseNodes( FAudioDevice* AudioDevice, const UPTRINT NodeWaveInstanceHash, FActiveSound& ActiveSound, const FSoundParseParameters& ParseParams, TArray<FWaveInstance*>& WaveInstances )
 {
+	RETRIEVE_SOUNDNODE_PAYLOAD(sizeof(int32));
+	DECLARE_SOUNDNODE_ELEMENT(int32, CurrentLoopCount);
+
+	if (*RequiresInitialization)
+	{
+		CurrentLoopCount = 0;
+
+		*RequiresInitialization = false;
+	}
+
 	FSoundParseParameters UpdatedParams = ParseParams;
 	UpdatedParams.NotifyBufferFinishedHooks.AddNotify(this, NodeWaveInstanceHash);
 
@@ -24,70 +36,79 @@ void USoundNodeLooping::ParseNodes( FAudioDevice* AudioDevice, const UPTRINT Nod
 
 bool USoundNodeLooping::NotifyWaveInstanceFinished( FWaveInstance* InWaveInstance )
 {
-	FActiveSound* ActiveSound = InWaveInstance->ActiveSound;
+	FActiveSound& ActiveSound = *InWaveInstance->ActiveSound;
+	const UPTRINT NodeWaveInstanceHash = InWaveInstance->NotifyBufferFinishedHooks.GetHashForNode(this);
+	RETRIEVE_SOUNDNODE_PAYLOAD(sizeof(int32));
+	DECLARE_SOUNDNODE_ELEMENT(int32, CurrentLoopCount);
+	check(*RequiresInitialization == 0);
 
-	struct FNodeHashPairs
+	if (bLoopIndefinitely == 1 || CurrentLoopCount < LoopCount)
 	{
-		USoundNode* Node;
-		UPTRINT NodeWaveInstanceHash;
-
-		FNodeHashPairs(USoundNode* InNode, const UPTRINT InHash)
-			: Node(InNode)
-			, NodeWaveInstanceHash(InHash)
+		struct FNodeHashPairs
 		{
-		}
-	};
+			USoundNode* Node;
+			UPTRINT NodeWaveInstanceHash;
 
-	TArray<FNodeHashPairs> NodesToReset;
+			FNodeHashPairs(USoundNode* InNode, const UPTRINT InHash)
+				: Node(InNode)
+				, NodeWaveInstanceHash(InHash)
+			{
+			}
+		};
 
-	for (int32 ChildNodeIndex = 0; ChildNodeIndex < ChildNodes.Num(); ++ChildNodeIndex)
-	{
-		USoundNode* ChildNode = ChildNodes[ChildNodeIndex];
-		if (ChildNode)
+		TArray<FNodeHashPairs> NodesToReset;
+
+		for (int32 ChildNodeIndex = 0; ChildNodeIndex < ChildNodes.Num(); ++ChildNodeIndex)
 		{
-			NodesToReset.Add(FNodeHashPairs(ChildNode, GetNodeWaveInstanceHash(InWaveInstance->NotifyBufferFinishedHooks.GetHashForNode(this), ChildNode, ChildNodeIndex)));
+			USoundNode* ChildNode = ChildNodes[ChildNodeIndex];
+			if (ChildNode)
+			{
+				NodesToReset.Add(FNodeHashPairs(ChildNode, GetNodeWaveInstanceHash(NodeWaveInstanceHash, ChildNode, ChildNodeIndex)));
+			}
 		}
-	}
 
-	// GetAllNodes includes current node so we have to start at Index 1.
-	for( int32 ResetNodeIndex = 0; ResetNodeIndex < NodesToReset.Num(); ++ResetNodeIndex )
-	{
-		const FNodeHashPairs& NodeHashPair = NodesToReset[ResetNodeIndex];
-
-		// Reset all child nodes so they are initialized again.
-		uint32* Offset = ActiveSound->SoundNodeOffsetMap.Find( NodeHashPair.NodeWaveInstanceHash );
-		if( Offset )
+		// GetAllNodes includes current node so we have to start at Index 1.
+		for (int32 ResetNodeIndex = 0; ResetNodeIndex < NodesToReset.Num(); ++ResetNodeIndex)
 		{
-			bool* bRequiresInitialization = ( bool* )&ActiveSound->SoundNodeData[ *Offset ];
-			*bRequiresInitialization = true;
-		}
+			const FNodeHashPairs& NodeHashPair = NodesToReset[ResetNodeIndex];
 
-		USoundNode* ResetNode = NodeHashPair.Node;
+			// Reset all child nodes so they are initialized again.
+			uint32* Offset = ActiveSound.SoundNodeOffsetMap.Find(NodeHashPair.NodeWaveInstanceHash);
+			if (Offset)
+			{
+				bool* bRequiresInitialization = (bool*)&ActiveSound.SoundNodeData[*Offset];
+				*bRequiresInitialization = true;
+			}
 
-		if (ResetNode->ChildNodes.Num())
-		{
-			for (int32 ResetChildIndex = 0; ResetChildIndex < ResetNode->ChildNodes.Num(); ++ResetChildIndex)
+			USoundNode* ResetNode = NodeHashPair.Node;
+
+			if (ResetNode->ChildNodes.Num())
 			{
-				USoundNode* ResetChildNode = ResetNode->ChildNodes[ResetChildIndex];
-				if (ResetChildNode)
+				for (int32 ResetChildIndex = 0; ResetChildIndex < ResetNode->ChildNodes.Num(); ++ResetChildIndex)
 				{
-					NodesToReset.Add(FNodeHashPairs(ResetChildNode, GetNodeWaveInstanceHash(NodeHashPair.NodeWaveInstanceHash, ResetChildNode, ResetChildIndex)));
+					USoundNode* ResetChildNode = ResetNode->ChildNodes[ResetChildIndex];
+					if (ResetChildNode)
+					{
+						NodesToReset.Add(FNodeHashPairs(ResetChildNode, GetNodeWaveInstanceHash(NodeHashPair.NodeWaveInstanceHash, ResetChildNode, ResetChildIndex)));
+					}
 				}
 			}
-		}
-		else if (ResetNode->IsA<USoundNodeWavePlayer>())
-		{
-			FWaveInstance* WaveInstance = ActiveSound->FindWaveInstance(NodeHashPair.NodeWaveInstanceHash);
-			if (WaveInstance)
+			else if (ResetNode->IsA<USoundNodeWavePlayer>())
 			{
-				WaveInstance->bAlreadyNotifiedHook = true;
+				FWaveInstance* WaveInstance = ActiveSound.FindWaveInstance(NodeHashPair.NodeWaveInstanceHash);
+				if (WaveInstance)
+				{
+					WaveInstance->bAlreadyNotifiedHook = true;
+				}
 			}
 		}
-	}
 
-	// Reset wave instances that notified us of completion.
-	InWaveInstance->bIsStarted = false;
-	InWaveInstance->bIsFinished = false;
+		// Reset wave instances that notified us of completion.
+		InWaveInstance->bIsStarted = false;
+		InWaveInstance->bIsFinished = false;
+
+		CurrentLoopCount++;
+	}
 
 	return true;
 }
