added navigation lock for creating preview actors when dragging from class viewer

diff --git a/Engine/Source/Editor/UnrealEd/Private/LevelEditorViewport.cpp b/Engine/Source/Editor/UnrealEd/Private/LevelEditorViewport.cpp
--- a/Engine/Source/Editor/UnrealEd/Private/LevelEditorViewport.cpp
+++ b/Engine/Source/Editor/UnrealEd/Private/LevelEditorViewport.cpp
@@ -1178,6 +1178,8 @@ bool FLevelEditorViewportClient::DropObjectsAtCoordinates(int32 MouseX, int32 Mo
 
 			if (TargetActor != NULL)
 			{
+				FNavigationLockContext LockNavigationUpdates(TargetActor->GetWorld(), ENavigationLockReason::SpawnOnDragEnter, bCreateDropPreview);
+
 				// if the target actor is selected, we should drop onto all selected items
 				// otherwise, we should drop only onto this object
 				bool bDropOntoSelected = TargetActor->IsSelected();
diff --git a/Engine/Source/Runtime/Engine/Classes/AI/Navigation/NavigationSystem.h b/Engine/Source/Runtime/Engine/Classes/AI/Navigation/NavigationSystem.h
--- a/Engine/Source/Runtime/Engine/Classes/AI/Navigation/NavigationSystem.h
+++ b/Engine/Source/Runtime/Engine/Classes/AI/Navigation/NavigationSystem.h
@@ -83,22 +83,29 @@ namespace ENavigationLockReason
 		MaterialUpdate			= 1 << 2,
 		LightingUpdate			= 1 << 3,
 		ContinuousEditorMove	= 1 << 4,
+		SpawnOnDragEnter		= 1 << 5,
 	};
 }
 
 class ENGINE_API FNavigationLockContext
 {
 public:
-	FNavigationLockContext(ENavigationLockReason::Type Reason = ENavigationLockReason::Unknown) 
-		: MyWorld(NULL), LockReason(Reason), bSingleWorld(false)
+	FNavigationLockContext(ENavigationLockReason::Type Reason = ENavigationLockReason::Unknown, bool bApplyLock = true) 
+		: MyWorld(NULL), LockReason(Reason), bSingleWorld(false), bIsLocked(false)
 	{
-		LockUpdates(); 
+		if (bApplyLock)
+		{
+			LockUpdates();
+		}
 	}
 
-	FNavigationLockContext(UWorld* InWorld, ENavigationLockReason::Type Reason = ENavigationLockReason::Unknown) 
-		: MyWorld(InWorld), LockReason(Reason), bSingleWorld(true)
+	FNavigationLockContext(UWorld* InWorld, ENavigationLockReason::Type Reason = ENavigationLockReason::Unknown, bool bApplyLock = true)
+		: MyWorld(InWorld), LockReason(Reason), bSingleWorld(true), bIsLocked(false)
 	{
-		LockUpdates(); 
+		if (bApplyLock)
+		{
+			LockUpdates();
+		}
 	}
 
 	~FNavigationLockContext()
@@ -110,6 +117,7 @@ private:
 	UWorld* MyWorld;
 	uint8 LockReason;
 	uint8 bSingleWorld : 1;
+	uint8 bIsLocked : 1;
 
 	void LockUpdates();
 	void UnlockUpdates();
diff --git a/Engine/Source/Runtime/Engine/Private/AI/Navigation/NavigationSystem.cpp b/Engine/Source/Runtime/Engine/Private/AI/Navigation/NavigationSystem.cpp
--- a/Engine/Source/Runtime/Engine/Private/AI/Navigation/NavigationSystem.cpp
+++ b/Engine/Source/Runtime/Engine/Private/AI/Navigation/NavigationSystem.cpp
@@ -113,6 +113,8 @@ namespace NavigationDebugDrawing
 void FNavigationLockContext::LockUpdates()
 {
 #if WITH_EDITOR
+	bIsLocked = true;
+
 	if (bSingleWorld)
 	{
 		UNavigationSystem* NavSys = UNavigationSystem::GetCurrent(MyWorld);
@@ -138,6 +140,11 @@ void FNavigationLockContext::LockUpdates()
 void FNavigationLockContext::UnlockUpdates()
 {
 #if WITH_EDITOR
+	if (!bIsLocked)
+	{
+		return;
+	}
+
 	if (bSingleWorld)
 	{
 		UNavigationSystem* NavSys = UNavigationSystem::GetCurrent(MyWorld);
