UE-5228 - Shootergame character T-Posing in reload animation. The reload asset is incorrectly set up and a recent change to montage editing has exposed the broken section of the montage. Added some further reconciliation for montage lengths to auto patch up the montage to play correctly and log a message telling the user that the lengths differ and that the asset may need fixing.

diff --git a/Engine/Source/Editor/Persona/Private/SMontageEditor.cpp b/Engine/Source/Editor/Persona/Private/SMontageEditor.cpp
--- a/Engine/Source/Editor/Persona/Private/SMontageEditor.cpp
+++ b/Engine/Source/Editor/Persona/Private/SMontageEditor.cpp
@@ -392,16 +392,7 @@ void SMontageEditor::SortAndUpdateMontage()
 
 float SMontageEditor::CalculateSequenceLengthOfEditorObject() const
 {
-	float NewSequenceLength = 0.f;
-	for (auto Iter = MontageObj->SlotAnimTracks.CreateIterator(); Iter; ++Iter)
-	{
-		FSlotAnimationTrack& SlotAnimTrack = (*Iter);
-		if ( SlotAnimTrack.AnimTrack.AnimSegments.Num() > 0 )
-		{
-			NewSequenceLength = FMath::Max(NewSequenceLength, SlotAnimTrack.AnimTrack.GetLength());
-		}
-	}
-	return NewSequenceLength;
+	return MontageObj->CalculateSequenceLength();
 }
 
 /** Sort Segments by starting time */
diff --git a/Engine/Source/Runtime/Engine/Classes/Animation/AnimMontage.h b/Engine/Source/Runtime/Engine/Classes/Animation/AnimMontage.h
--- a/Engine/Source/Runtime/Engine/Classes/Animation/AnimMontage.h
+++ b/Engine/Source/Runtime/Engine/Classes/Animation/AnimMontage.h
@@ -312,6 +312,10 @@ class UAnimMontage : public UAnimCompositeBase
 public:
 	// Begin UObject Interface
 	virtual void PostLoad() override;
+
+	// Gets the sequence length of the montage by calculating it from the lengths of the segments in the montage
+	ENGINE_API float CalculateSequenceLength();
+
 #if WITH_EDITOR
 	virtual void PostEditChangeProperty(FPropertyChangedEvent& PropertyChangedEvent) override;
 #endif // WITH_EDITOR
diff --git a/Engine/Source/Runtime/Engine/Private/Animation/AnimMontage.cpp b/Engine/Source/Runtime/Engine/Private/Animation/AnimMontage.cpp
--- a/Engine/Source/Runtime/Engine/Private/Animation/AnimMontage.cpp
+++ b/Engine/Source/Runtime/Engine/Private/Animation/AnimMontage.cpp
@@ -360,7 +360,15 @@ void UAnimMontage::PostLoad()
 			}
 		}
 		Track.ValidateSegmentTimes();
-			}
+
+		float CurrentCalculatedLength = CalculateSequenceLength();
+
+		if(CurrentCalculatedLength != SequenceLength)
+		{
+			UE_LOG(LogAnimation, Warning, TEXT("UAnimMontage::PostLoad: The actual sequence length for montage %s does not match the length stored in the asset, please resave the montage asset."), *GetName());
+			SequenceLength = CurrentCalculatedLength;
+		}
+	}
 
 	int32 Ver = GetLinker()->UE4Ver();
 
@@ -1324,6 +1332,20 @@ void FAnimMontageInstance::HandleEvents(float PreviousTrackPos, float CurrentTra
 	}
 }
 
+float UAnimMontage::CalculateSequenceLength()
+{
+	float CalculatedSequenceLength = 0.f;
+	for(auto Iter = SlotAnimTracks.CreateIterator(); Iter; ++Iter)
+	{
+		FSlotAnimationTrack& SlotAnimTrack = (*Iter);
+		if(SlotAnimTrack.AnimTrack.AnimSegments.Num() > 0)
+		{
+			CalculatedSequenceLength = FMath::Max(CalculatedSequenceLength, SlotAnimTrack.AnimTrack.GetLength());
+		}
+	}
+	return CalculatedSequenceLength;
+}
+
 #if WITH_EDITOR
 bool UAnimMontage::GetAllAnimationSequencesReferred(TArray<UAnimSequence*>& AnimationSequences)
 {
