UDS: no GUID is visiable in DisplayName of property.

diff --git a/Engine/Source/Runtime/CoreUObject/Private/UObject/Class.cpp b/Engine/Source/Runtime/CoreUObject/Private/UObject/Class.cpp
--- a/Engine/Source/Runtime/CoreUObject/Private/UObject/Class.cpp
+++ b/Engine/Source/Runtime/CoreUObject/Private/UObject/Class.cpp
@@ -195,14 +195,24 @@ struct FDisplayNameHelper
 {
 	static FString Get(const UObject& Object)
 	{
-		FString Name = Object.GetName();
 		const UClass* Class = dynamic_cast<const UClass*>(&Object);
 		if (Class && !Class->HasAnyClassFlags(CLASS_Native))
 		{
+			FString Name = Object.GetName();
 			Name.RemoveFromEnd(TEXT("_C"));
 			Name.RemoveFromStart(TEXT("SKEL_"));
+			return Name;
 		}
-		return Name;
+
+		if (auto Property = dynamic_cast<const UProperty*>(&Object))
+		{
+			if (auto OwnerStruct = Property->GetOwnerStruct())
+			{
+				return OwnerStruct->PropertyNameToDisplayName(Property->GetFName());
+			}
+		}
+
+		return Object.GetName();
 	}
 };
 
diff --git a/Engine/Source/Runtime/CoreUObject/Public/UObject/Class.h b/Engine/Source/Runtime/CoreUObject/Public/UObject/Class.h
--- a/Engine/Source/Runtime/CoreUObject/Public/UObject/Class.h
+++ b/Engine/Source/Runtime/CoreUObject/Public/UObject/Class.h
@@ -372,6 +372,11 @@ public:
 		Children = Child;
 	}
 
+	virtual FString PropertyNameToDisplayName(FName Name) const 
+	{ 
+		return Name.ToString(); 
+	}
+
 #if WITH_EDITOR
 	/** Try and find boolean metadata with the given key. If not found on this class, work up hierarchy looking for it. */
 	bool GetBoolMetaDataHierarchical(const FName& Key) const;
diff --git a/Engine/Source/Runtime/Engine/Classes/Engine/UserDefinedStruct.h b/Engine/Source/Runtime/Engine/Classes/Engine/UserDefinedStruct.h
--- a/Engine/Source/Runtime/Engine/Classes/Engine/UserDefinedStruct.h
+++ b/Engine/Source/Runtime/Engine/Classes/Engine/UserDefinedStruct.h
@@ -51,6 +51,7 @@ public:
 #endif	// WITH_EDITOR
 
 	virtual void SerializeTaggedProperties(FArchive& Ar, uint8* Data, UStruct* DefaultsStruct, uint8* Defaults, const UObject* BreakRecursionIfFullyLoad=NULL) const override;
+	virtual FString PropertyNameToDisplayName(FName Name) const override;
 
 	// UScriptStruct interface.
 	virtual void RecursivelyPreload() override;
diff --git a/Engine/Source/Runtime/Engine/Private/UserDefinedStruct.cpp b/Engine/Source/Runtime/Engine/Private/UserDefinedStruct.cpp
--- a/Engine/Source/Runtime/Engine/Private/UserDefinedStruct.cpp
+++ b/Engine/Source/Runtime/Engine/Private/UserDefinedStruct.cpp
@@ -61,6 +61,29 @@ void UUserDefinedStruct::InitializeDefaultValue(uint8* StructData) const
 
 #endif	// WITH_EDITOR
 
+FString UUserDefinedStruct::PropertyNameToDisplayName(FName Name) const
+{
+#if WITH_EDITOR
+	auto Guid = FStructureEditorUtils::GetGuidFromPropertyName(Name);
+	return FStructureEditorUtils::GetVariableDisplayName(this, Guid);
+#endif	// WITH_EDITOR
+
+	const int32 GuidStrLen = 32;
+	const int32 MinimalPostfixlen = GuidStrLen + 3;
+	const FString OriginalName = Name.ToString();
+	if (OriginalName.Len() > MinimalPostfixlen)
+	{
+		FString DisplayName = OriginalName.LeftChop(GuidStrLen + 1);
+		int FirstCharToRemove = -1;
+		const bool bCharFound = DisplayName.FindLastChar(TCHAR('_'), FirstCharToRemove);
+		if(bCharFound && (FirstCharToRemove > 0))
+		{
+			return DisplayName.Mid(0, FirstCharToRemove);
+		}
+	}
+	return OriginalName;
+}
+
 void UUserDefinedStruct::SerializeTaggedProperties(FArchive& Ar, uint8* Data, UStruct* DefaultsStruct, uint8* Defaults, const UObject* BreakRecursionIfFullyLoad) const
 {
 #if WITH_EDITOR
