Prevent crash when streaming very short sounds

When a sound is so small that it fits into the double buffer provided, the notify that the sound has finished is missed. Made sure that system gracefully handles asking to read more data when we are already at the end.

diff --git a/Engine/Source/Runtime/Engine/Private/OpusAudioInfo.cpp b/Engine/Source/Runtime/Engine/Private/OpusAudioInfo.cpp
--- a/Engine/Source/Runtime/Engine/Private/OpusAudioInfo.cpp
+++ b/Engine/Source/Runtime/Engine/Private/OpusAudioInfo.cpp
@@ -127,11 +127,12 @@ FOpusAudioInfo::~FOpusAudioInfo(void)
 
 size_t FOpusAudioInfo::Read(void *ptr, uint32 size)
 {
-	check(SrcBufferOffset < SrcBufferDataSize);
-
 	size_t BytesToRead = FMath::Min(size, SrcBufferDataSize - SrcBufferOffset);
-	FMemory::Memcpy(ptr, SrcBufferData + SrcBufferOffset, BytesToRead);
-	SrcBufferOffset += BytesToRead;
+	if (BytesToRead > 0)
+	{
+		FMemory::Memcpy(ptr, SrcBufferData + SrcBufferOffset, BytesToRead);
+		SrcBufferOffset += BytesToRead;
+	}
 	return(BytesToRead);
 }
 
