Prevent instanced meshes from going through weld code. Not supported (UE-5421)

diff --git a/Engine/Source/Runtime/Engine/Private/InstancedStaticMesh.cpp b/Engine/Source/Runtime/Engine/Private/InstancedStaticMesh.cpp
--- a/Engine/Source/Runtime/Engine/Private/InstancedStaticMesh.cpp
+++ b/Engine/Source/Runtime/Engine/Private/InstancedStaticMesh.cpp
@@ -1433,6 +1433,7 @@ void UInstancedStaticMeshComponent::InitInstanceBody(int32 InstanceIdx, FBodyIns
 	// Aggregates aren't used for static objects
 	auto* Aggregate = (Mobility == EComponentMobility::Movable) ? Aggregates[FMath::DivideAndRoundDown<int32>(InstanceIdx, AggregateMaxSize)] : nullptr;
 	check(Mobility != EComponentMobility::Movable || Aggregate->getNbActors() < Aggregate->getMaxNbActors());
+	InstanceBodyInstance->bAutoWeld = false;	//We don't support this for instanced meshes.
 	InstanceBodyInstance->InitBody(BodySetup, InstanceTransform, this, GetWorld()->GetPhysicsScene(), Aggregate);
 #endif //WITH_PHYSX
 }
