UberGraphFrame crash

diff --git a/Engine/Source/Runtime/Engine/Private/BlueprintGeneratedClass.cpp b/Engine/Source/Runtime/Engine/Private/BlueprintGeneratedClass.cpp
--- a/Engine/Source/Runtime/Engine/Private/BlueprintGeneratedClass.cpp
+++ b/Engine/Source/Runtime/Engine/Private/BlueprintGeneratedClass.cpp
@@ -528,10 +528,12 @@ void UBlueprintGeneratedClass::AddReferencedObjectsInUbergraphFrame(UObject* InT
 			{
 				checkSlow(BPGC->UberGraphFunction);
 				auto PointerToUberGraphFrame = BPGC->UberGraphFramePointerProperty->ContainerPtrToValuePtr<FPointerToUberGraphFrame>(InThis);
-				checkSlow(PointerToUberGraphFrame && PointerToUberGraphFrame->RawPointer);
-
-				FSimpleObjectReferenceCollectorArchive ObjectReferenceCollector(InThis, Collector);
-				BPGC->UberGraphFunction->SerializeBin(ObjectReferenceCollector, PointerToUberGraphFrame->RawPointer, 0);
+				checkSlow(PointerToUberGraphFrame)
+				if (PointerToUberGraphFrame->RawPointer)
+				{
+					FSimpleObjectReferenceCollectorArchive ObjectReferenceCollector(InThis, Collector);
+					BPGC->UberGraphFunction->SerializeBin(ObjectReferenceCollector, PointerToUberGraphFrame->RawPointer, 0);
+				}
 			}
 		}
 		else if (CurrentClass->HasAllClassFlags(CLASS_Native))
