Sanity check that we are closing a channel only if we actually opened it (own it locally)

diff --git a/Engine/Source/Runtime/Engine/Private/DataChannel.cpp b/Engine/Source/Runtime/Engine/Private/DataChannel.cpp
--- a/Engine/Source/Runtime/Engine/Private/DataChannel.cpp
+++ b/Engine/Source/Runtime/Engine/Private/DataChannel.cpp
@@ -65,7 +65,9 @@ void UChannel::SetClosingFlag()
 
 void UChannel::Close()
 {
+	check(OpenedLocally);		// We are only allowed to close channels that we opened locally
 	check(Connection->Channels[ChIndex]==this);
+
 	if ( !Closing && ( Connection->State == USOCK_Open || Connection->State == USOCK_Pending ) )
 	{
 		if ( ChIndex == 0 )
diff --git a/Engine/Source/Runtime/Engine/Private/NetConnection.cpp b/Engine/Source/Runtime/Engine/Private/NetConnection.cpp
--- a/Engine/Source/Runtime/Engine/Private/NetConnection.cpp
+++ b/Engine/Source/Runtime/Engine/Private/NetConnection.cpp
@@ -203,7 +203,7 @@ void UNetConnection::Close()
 			*LowLevelGetRemoteAddress(true),
 			*FDateTime::UtcNow().ToString(TEXT("%Y.%m.%d-%H.%M.%S")));
 
-		if (Channels[0] != NULL)
+		if (Channels[0] != NULL && Channels[0]->OpenedLocally)
 		{
 			Channels[0]->Close();
 		}
