Don't assert when actor fails to create, we should be resilient to this, and log an error

diff --git a/Engine/Source/Runtime/Engine/Private/PackageMapClient.cpp b/Engine/Source/Runtime/Engine/Private/PackageMapClient.cpp
--- a/Engine/Source/Runtime/Engine/Private/PackageMapClient.cpp
+++ b/Engine/Source/Runtime/Engine/Private/PackageMapClient.cpp
@@ -284,7 +284,6 @@ bool UPackageMapClient::SerializeNewActor(FArchive& Ar, class UActorChannel *Cha
 				else
 				{
 					UE_LOG(LogNetPackageMap, Error, TEXT("UPackageMapClient::SerializeNewActor Unable to read Archetype for NetGUID %s / %s"), *NetGUID.ToString(), *ArchetypeNetGUID.ToString() );
-					ensure(Archetype);
 				}
 			}
 		}
