fixing non-unity compile

diff --git a/Engine/Source/Runtime/Engine/Private/NiagaraConstantSet.cpp b/Engine/Source/Runtime/Engine/Private/NiagaraConstantSet.cpp
--- a/Engine/Source/Runtime/Engine/Private/NiagaraConstantSet.cpp
+++ b/Engine/Source/Runtime/Engine/Private/NiagaraConstantSet.cpp
@@ -1,4 +1,5 @@
 #include "EnginePrivate.h"
+#include "Niagara/NiagaraConstantSet.h"
 
 template<>
 struct TStructOpsTypeTraits<FNiagaraConstantMap> : public TStructOpsTypeTraitsBase
