Ensure that DefaultTextValue is also cleared when resetting a GraphPin
Compartmentalize value resetting in to a function.

diff --git a/Engine/Source/Runtime/Engine/Classes/EdGraph/EdGraphPin.h b/Engine/Source/Runtime/Engine/Classes/EdGraph/EdGraphPin.h
--- a/Engine/Source/Runtime/Engine/Classes/EdGraph/EdGraphPin.h
+++ b/Engine/Source/Runtime/Engine/Classes/EdGraph/EdGraphPin.h
@@ -284,6 +284,13 @@ public:
 	/** Returns how the name of the pin should be displayed in the UI */
 	ENGINE_API FText GetDisplayName() const;
 
+	void ResetDefaultValue()
+	{
+		DefaultValue.Empty();
+		DefaultObject = nullptr;
+		DefaultTextValue = FText::GetEmpty();
+	}
+
 	void ResetToDefaults()
 	{
 		check(LinkedTo.Num() == 0);
@@ -292,9 +299,8 @@ public:
 
 		PinName.Empty();
 		PinFriendlyName = FText::GetEmpty();
-		DefaultValue.Empty();
 		AutogeneratedDefaultValue.Empty();
-		DefaultObject = NULL;
+		ResetDefaultValue();
 
 #if WITH_EDITORONLY_DATA
 		bHidden = false;
