Add ability to record demo as URL option

diff --git a/Engine/Source/Runtime/Engine/Private/UnrealEngine.cpp b/Engine/Source/Runtime/Engine/Private/UnrealEngine.cpp
--- a/Engine/Source/Runtime/Engine/Private/UnrealEngine.cpp
+++ b/Engine/Source/Runtime/Engine/Private/UnrealEngine.cpp
@@ -9025,6 +9025,15 @@ bool UEngine::LoadMap( FWorldContext& WorldContext, FURL URL, class UPendingNetG
 	// RedrawViewports() may have added a dummy playerstart location. Remove all views to start from fresh the next Tick().
 	IStreamingManager::Get().RemoveStreamingViews( RemoveStreamingViews_All );
 	
+	// See if we need to record network demos
+	const TCHAR* DemoRecName = URL.GetOption( TEXT( "DemoRec=" ), NULL );
+
+	if ( DemoRecName != NULL )
+	{
+		// Play the demo
+		GEngine->Exec( WorldContext.World(), *FString::Printf( TEXT("DEMOREC %s"), DemoRecName ) );
+	}
+
 	MALLOC_PROFILER( FMallocProfiler::SnapshotMemoryLoadMapEnd( URL.Map ); )
 
 	// Successfully started local level.
