Added BP getter for NavigationSystem #UE4

diff --git a/Engine/Source/Runtime/Engine/Classes/AI/Navigation/NavigationSystem.h b/Engine/Source/Runtime/Engine/Classes/AI/Navigation/NavigationSystem.h
--- a/Engine/Source/Runtime/Engine/Classes/AI/Navigation/NavigationSystem.h
+++ b/Engine/Source/Runtime/Engine/Classes/AI/Navigation/NavigationSystem.h
@@ -203,8 +203,12 @@ private:
 
 public:
 	//----------------------------------------------------------------------//
-	// Kismet functions
+	// Blueprint functions
 	//----------------------------------------------------------------------//
+
+	UFUNCTION(BlueprintPure, Category = "AI|Navigation", meta = (WorldContext = "WorldContext"))
+	static UNavigationSystem* GetNavigationSystem(UObject* WorldContext);
+
 	/** Project a point onto the NavigationData */
 	UFUNCTION(BlueprintPure, Category="AI|Navigation", meta=(WorldContext="WorldContext" ) )
 	static FVector ProjectPointToNavigation(UObject* WorldContext, const FVector& Point, ANavigationData* NavData = NULL, TSubclassOf<UNavigationQueryFilter> FilterClass = NULL);
@@ -614,7 +618,7 @@ public:
 
 	static UNavigationSystem* GetCurrent(UWorld* World);
 	static UNavigationSystem* GetCurrent(UObject* WorldContextObject);
-
+	
 	/** try to create and setup navigation system */
 	static void InitializeForWorld(UWorld* World, FNavigationSystem::EMode Mode);
 
diff --git a/Engine/Source/Runtime/Engine/Private/AI/Navigation/NavigationSystem.cpp b/Engine/Source/Runtime/Engine/Private/AI/Navigation/NavigationSystem.cpp
--- a/Engine/Source/Runtime/Engine/Private/AI/Navigation/NavigationSystem.cpp
+++ b/Engine/Source/Runtime/Engine/Private/AI/Navigation/NavigationSystem.cpp
@@ -2990,8 +2990,13 @@ void UNavigationSystem::CleanUp(ECleanupMode Mode)
 }
 
 //----------------------------------------------------------------------//
-// Kismet functions
+// Blueprint functions
 //----------------------------------------------------------------------//
+UNavigationSystem* UNavigationSystem::GetNavigationSystem(UObject* WorldContext)
+{
+	return GetCurrent(WorldContext);
+}
+
 FVector UNavigationSystem::ProjectPointToNavigation(UObject* WorldContextObject, const FVector& Point, ANavigationData* NavData, TSubclassOf<UNavigationQueryFilter> FilterClass)
 {
 	FNavLocation ProjectedPoint(Point);
