UE-3049 Change default color for HUD funcs to have alpha of 1

diff --git a/Engine/Source/Runtime/Engine/Classes/GameFramework/HUD.h b/Engine/Source/Runtime/Engine/Classes/GameFramework/HUD.h
--- a/Engine/Source/Runtime/Engine/Classes/GameFramework/HUD.h
+++ b/Engine/Source/Runtime/Engine/Classes/GameFramework/HUD.h
@@ -196,7 +196,7 @@ public:
 	 * @param Scale				Scale multiplier to control size of the text.
 	 * @param bScalePosition	Whether the "Scale" parameter should also scale the position of this draw call.
 	 */
-	UFUNCTION(BlueprintCallable, Category=HUD)
+	UFUNCTION(BlueprintCallable, Category=HUD, meta=(TextColor="(R=0,G=0,B=0,A=1)"))
 	void DrawText(const FString& Text, FLinearColor TextColor, float ScreenX, float ScreenY, UFont* Font=NULL, float Scale=1.f, bool bScalePosition=false);
 
 	/**
@@ -207,7 +207,7 @@ public:
 	 * @param EndScreenY		Screen-space Y coordinate of end of the line.
 	 * @param LineColor			Color to draw line
 	 */
-	UFUNCTION(BlueprintCallable, Category=HUD)
+	UFUNCTION(BlueprintCallable, Category=HUD, meta=(LineColor="(R=0,G=0,B=0,A=1)"))
 	void DrawLine(float StartScreenX, float StartScreenY, float EndScreenX, float EndScreenY, FLinearColor LineColor);
 
 	/**
@@ -218,7 +218,7 @@ public:
 	 * @param ScreenW			Screen-space width of the quad (in pixels).
 	 * @param ScreenH			Screen-space height of the quad (in pixels).
 	 */
-	UFUNCTION(BlueprintCallable, Category=HUD)
+	UFUNCTION(BlueprintCallable, Category=HUD, meta=(RectColor="(R=0,G=0,B=0,A=1)"))
 	void DrawRect(FLinearColor RectColor, float ScreenX, float ScreenY, float ScreenW, float ScreenH);
 
 	/**
