UE-2010 Add SetMassScale and SetAllMassScale to PrimitiveComponent

diff --git a/Engine/Source/Runtime/Engine/Classes/Components/PrimitiveComponent.h b/Engine/Source/Runtime/Engine/Classes/Components/PrimitiveComponent.h
--- a/Engine/Source/Runtime/Engine/Classes/Components/PrimitiveComponent.h
+++ b/Engine/Source/Runtime/Engine/Classes/Components/PrimitiveComponent.h
@@ -1371,6 +1371,14 @@ public:
 	UFUNCTION(BlueprintCallable, Category="Physics")
 	virtual float GetAngularDamping() const;
 
+	/** Change the mass scale used to calculate the mass of a single physics body */
+	UFUNCTION(BlueprintCallable, Category="Physics")
+	virtual void SetMassScale(FName BoneName = NAME_None, float InMassScale = 1.f);
+
+	/** Change the mass scale used fo all bodies in this component */
+	UFUNCTION(BlueprintCallable, Category="Physics")
+	virtual void SetAllMassScale(float InMassScale = 1.f);
+
 	/** Returns the mass of this component in kg. */
 	UFUNCTION(BlueprintCallable, Category="Physics")
 	virtual float GetMass() const;
diff --git a/Engine/Source/Runtime/Engine/Classes/Components/SkeletalMeshComponent.h b/Engine/Source/Runtime/Engine/Classes/Components/SkeletalMeshComponent.h
--- a/Engine/Source/Runtime/Engine/Classes/Components/SkeletalMeshComponent.h
+++ b/Engine/Source/Runtime/Engine/Classes/Components/SkeletalMeshComponent.h
@@ -835,6 +835,7 @@ public:
 	virtual void AddRadialImpulse(FVector Origin, float Radius, float Strength, ERadialImpulseFalloff Falloff, bool bVelChange=false) override;
 	virtual void AddRadialForce(FVector Origin, float Radius, float Strength, ERadialImpulseFalloff Falloff) override;
 	virtual void SetAllPhysicsLinearVelocity(FVector NewVel,bool bAddToCurrent = false) override;
+	virtual void SetAllMassScale(float InMassScale = 1.f) override;
 	virtual float GetMass() const override;
 	virtual float CalculateMass(FName BoneName = NAME_None) override;
 #if WITH_EDITOR
diff --git a/Engine/Source/Runtime/Engine/Classes/PhysicsEngine/BodyInstance.h b/Engine/Source/Runtime/Engine/Classes/PhysicsEngine/BodyInstance.h
--- a/Engine/Source/Runtime/Engine/Classes/PhysicsEngine/BodyInstance.h
+++ b/Engine/Source/Runtime/Engine/Classes/PhysicsEngine/BodyInstance.h
@@ -541,6 +541,9 @@ public:
 	/** return true if it uses Collision Profile System. False otherwise*/
 	bool DoesUseCollisionProfile() const;
 
+	/** Modify the mass scale of this body */
+	void SetMassScale(float InMassScale=1.f);
+
 	/** Update instance's mass properties (mass, inertia and center-of-mass offset) based on MassScale, InstanceMassScale and COMNudge. */
 	void UpdateMassProperties();
 
diff --git a/Engine/Source/Runtime/Engine/Private/PhysicsEngine/BodyInstance.cpp b/Engine/Source/Runtime/Engine/Private/PhysicsEngine/BodyInstance.cpp
--- a/Engine/Source/Runtime/Engine/Private/PhysicsEngine/BodyInstance.cpp
+++ b/Engine/Source/Runtime/Engine/Private/PhysicsEngine/BodyInstance.cpp
@@ -513,6 +513,12 @@ bool FBodyInstance::DoesUseCollisionProfile() const
 	return IsValidCollisionProfileName(CollisionProfileName);
 }
 
+void FBodyInstance::SetMassScale(float InMassScale)
+{
+	MassScale = InMassScale;
+	UpdateMassProperties();
+}
+
 void FBodyInstance::SetCollisionEnabled(ECollisionEnabled::Type NewType, bool bUpdatePhysicsFilterData)
 {
 	if (CollisionEnabled != NewType)
diff --git a/Engine/Source/Runtime/Engine/Private/PrimitiveComponentPhysics.cpp b/Engine/Source/Runtime/Engine/Private/PrimitiveComponentPhysics.cpp
--- a/Engine/Source/Runtime/Engine/Private/PrimitiveComponentPhysics.cpp
+++ b/Engine/Source/Runtime/Engine/Private/PrimitiveComponentPhysics.cpp
@@ -397,6 +397,20 @@ float UPrimitiveComponent::GetAngularDamping() const
 	return 0.f;
 }
 
+void UPrimitiveComponent::SetMassScale(FName BoneName, float InMassScale)
+{
+	FBodyInstance* BI = GetBodyInstance(BoneName);
+	if (BI)
+	{
+		BI->SetMassScale(InMassScale);
+	}
+}
+
+void UPrimitiveComponent::SetAllMassScale(float InMassScale)
+{
+	SetMassScale(NAME_None, InMassScale);
+}
+
 float UPrimitiveComponent::GetMass() const
 {
 	FBodyInstance* BI = GetBodyInstance();
diff --git a/Engine/Source/Runtime/Engine/Private/SkeletalMeshComponentPhysics.cpp b/Engine/Source/Runtime/Engine/Private/SkeletalMeshComponentPhysics.cpp
--- a/Engine/Source/Runtime/Engine/Private/SkeletalMeshComponentPhysics.cpp
+++ b/Engine/Source/Runtime/Engine/Private/SkeletalMeshComponentPhysics.cpp
@@ -4829,6 +4829,19 @@ void USkeletalMeshComponent::DrawClothingPhysicalMeshWire(FPrimitiveDrawInterfac
 #endif // #if WITH_APEX_CLOTHING
 }
 
+void USkeletalMeshComponent::SetAllMassScale(float InMassScale)
+{
+	// Apply mass scale to each child body
+	for(FBodyInstance* BI : Bodies)
+	{
+		if (BI->IsValidBodyInstance())
+		{
+			BI->SetMassScale(InMassScale);
+		}
+	}
+}
+
+
 float USkeletalMeshComponent::GetMass() const
 {
 	float Mass = 0.0f;
