Niagara - few bug fixes and optimizations

diff --git a/Engine/Source/Editor/NiagaraEditor/Private/NiagaraEditor.cpp b/Engine/Source/Editor/NiagaraEditor/Private/NiagaraEditor.cpp
--- a/Engine/Source/Editor/NiagaraEditor/Private/NiagaraEditor.cpp
+++ b/Engine/Source/Editor/NiagaraEditor/Private/NiagaraEditor.cpp
@@ -270,9 +270,7 @@ void FNiagaraEditor::ExtendToolbar()
 
 FReply FNiagaraEditor::OnCompileClicked()
 {
-	FNiagaraEditorModule& NiagaraEditorModule = FModuleManager::LoadModuleChecked<FNiagaraEditorModule>( "NiagaraEditor" );
-	NiagaraEditorModule.CompileScript(Script);
-
+	Script->Source->Compile();
 	return FReply::Handled();
 }
 
diff --git a/Engine/Source/Editor/NiagaraEditor/Private/SNiagaraEffectEditorWidget.cpp b/Engine/Source/Editor/NiagaraEditor/Private/SNiagaraEffectEditorWidget.cpp
--- a/Engine/Source/Editor/NiagaraEditor/Private/SNiagaraEffectEditorWidget.cpp
+++ b/Engine/Source/Editor/NiagaraEditor/Private/SNiagaraEffectEditorWidget.cpp
@@ -90,30 +90,6 @@ void SNiagaraEffectEditorWidget::Construct(const FArguments& InArgs)
 				.AutoWidth()
 				[
 					SAssignNew(OverlayWidget, SOverlay)
-					/*
-					// details view on the left
-					+ SOverlay::Slot()
-					.Expose(DetailsViewSlot)
-					.HAlign(HAlign_Left)
-					.VAlign(VAlign_Top)
-					[
-					SNew(SBorder)
-					.Padding(4)
-					[
-					EmitterDetails.ToSharedRef()
-					]
-					]
-					*/
-
-					// Title bar - optional
-					/*
-					+ SOverlay::Slot()
-					.VAlign(VAlign_Top)
-					[
-					InArgs._TitleBar.IsValid() ? InArgs._TitleBar.ToSharedRef() : (TSharedRef<SWidget>)SNullWidget::NullWidget
-					]
-					*/
-
 
 					+SOverlay::Slot()
 					.Padding(10)
diff --git a/Engine/Source/Editor/UnrealEd/Classes/Niagara/NiagaraScriptSource.h b/Engine/Source/Editor/UnrealEd/Classes/Niagara/NiagaraScriptSource.h
--- a/Engine/Source/Editor/UnrealEd/Classes/Niagara/NiagaraScriptSource.h
+++ b/Engine/Source/Editor/UnrealEd/Classes/Niagara/NiagaraScriptSource.h
@@ -17,7 +17,7 @@ class UNiagaraScriptSource : public UNiagaraScriptSourceBase
 
 	// UObject interface.
 	virtual void PostLoad() override;
-	void Compile();
+	virtual void Compile() override;
 
 	UNREALED_API void GetParticleAttributes(TArray<FName>& VectorOutputs);
 	UNREALED_API void GetEmitterAttributes(TArray<FName>& VectorInputs, TArray<FName>& MatrixInputs);
diff --git a/Engine/Source/Editor/UnrealEd/Private/Niagara/NiagaraScriptSource.cpp b/Engine/Source/Editor/UnrealEd/Private/Niagara/NiagaraScriptSource.cpp
--- a/Engine/Source/Editor/UnrealEd/Private/Niagara/NiagaraScriptSource.cpp
+++ b/Engine/Source/Editor/UnrealEd/Private/Niagara/NiagaraScriptSource.cpp
@@ -50,6 +50,8 @@ void UNiagaraScriptSource::Compile()
 	FNiagaraEditorModule& NiagaraEditorModule = FModuleManager::Get().LoadModuleChecked<FNiagaraEditorModule>(TEXT("NiagaraEditor"));
 	NiagaraEditorModule.CompileScript(ScriptOwner);
 
+	ExposedVectorConstants.Empty();
+
 	// grab all constant nodes that are exposed to the editor
 	TArray<UNiagaraNodeConstant*> ConstNodes;
 	UpdateGraph->GetNodesOfClass<UNiagaraNodeConstant>(ConstNodes);
diff --git a/Engine/Source/Runtime/Engine/Classes/Engine/NiagaraScriptSourceBase.h b/Engine/Source/Runtime/Engine/Classes/Engine/NiagaraScriptSourceBase.h
--- a/Engine/Source/Runtime/Engine/Classes/Engine/NiagaraScriptSourceBase.h
+++ b/Engine/Source/Runtime/Engine/Classes/Engine/NiagaraScriptSourceBase.h
@@ -17,4 +17,5 @@ class UNiagaraScriptSourceBase : public UObject
 	GENERATED_UCLASS_BODY()
 
 	TArray<TSharedPtr<EditorExposedVectorConstant> > ExposedVectorConstants;
+	virtual void Compile() {};
 };
\ No newline at end of file
diff --git a/Engine/Source/Runtime/Engine/Private/NiagaraEffectRenderer.cpp b/Engine/Source/Runtime/Engine/Private/NiagaraEffectRenderer.cpp
--- a/Engine/Source/Runtime/Engine/Private/NiagaraEffectRenderer.cpp
+++ b/Engine/Source/Runtime/Engine/Private/NiagaraEffectRenderer.cpp
@@ -339,9 +339,10 @@ FNiagaraDynamicDataBase *NiagaraEffectRendererSprites::GenerateVertexData(const
 	const FVector4 *RotPtr = Data.GetAttributeData("Rotation");
 
 	float ParticleId = 0.0f, IdInc = 1.0f / Data.GetNumParticles();
+	RenderData.AddUninitialized(Data.GetNumParticles());
 	for (uint32 ParticleIndex = 0; ParticleIndex < Data.GetNumParticles(); ParticleIndex++)
 	{
-		FParticleSpriteVertex& NewVertex = *new(RenderData)FParticleSpriteVertex;
+		FParticleSpriteVertex& NewVertex = RenderData[ParticleIndex];
 		NewVertex.Position = PosPtr[ParticleIndex];
 		NewVertex.OldPosition = NewVertex.Position;
 		NewVertex.Color = FLinearColor(ColPtr[ParticleIndex]);
@@ -352,6 +353,11 @@ FNiagaraDynamicDataBase *NiagaraEffectRendererSprites::GenerateVertexData(const
 		NewVertex.Rotation = RotPtr[ParticleIndex].X;
 		NewVertex.SubImageIndex = 0.f;
 
+		FPlatformMisc::Prefetch(PosPtr + ParticleIndex+1);
+		FPlatformMisc::Prefetch(RotPtr + ParticleIndex + 1);
+		FPlatformMisc::Prefetch(ColPtr + ParticleIndex + 1);		
+		FPlatformMisc::Prefetch(AgePtr + ParticleIndex + 1);
+
 		//CachedBounds += NewVertex.Position;
 	}
 	//CachedBounds.ExpandBy(MaxSize);
