New node/function to enable/disable actor ticking.

diff --git a/Engine/Source/Runtime/Engine/Classes/GameFramework/Actor.h b/Engine/Source/Runtime/Engine/Classes/GameFramework/Actor.h
--- a/Engine/Source/Runtime/Engine/Classes/GameFramework/Actor.h
+++ b/Engine/Source/Runtime/Engine/Classes/GameFramework/Actor.h
@@ -934,6 +934,10 @@ public:
 	UFUNCTION(BlueprintCallable, Category="Utilities")
 	void SetTickableWhenPaused(bool bTickableWhenPaused);
 
+	/** Disables/Enables ticking for this actor (CanEverTick() has to be true, for tick enabling to work). */
+	UFUNCTION(BlueprintCallable, Category = "Utilities")
+	void SetTickEnabled(bool bEnable = true);
+
 	/** Allocate a MID for a given parent material. */
 	UFUNCTION(BlueprintCallable, meta=(DeprecatedFunction, DeprecationMessage="Use PrimitiveComponent.CreateAndSetMaterialInstanceDynamic instead.", BlueprintProtected = "true"), Category="Rendering|Material")
 	class UMaterialInstanceDynamic* MakeMIDForMaterial(class UMaterialInterface* Parent);
diff --git a/Engine/Source/Runtime/Engine/Private/Actor.cpp b/Engine/Source/Runtime/Engine/Private/Actor.cpp
--- a/Engine/Source/Runtime/Engine/Private/Actor.cpp
+++ b/Engine/Source/Runtime/Engine/Private/Actor.cpp
@@ -388,6 +388,18 @@ void AActor::SetTickableWhenPaused(bool bTickableWhenPaused)
 	PrimaryActorTick.bTickEvenWhenPaused = bTickableWhenPaused;
 }
 
+void AActor::SetTickEnabled(bool bEnable)
+{
+	if (CanEverTick())
+	{
+		PrimaryActorTick.SetTickFunctionEnable(bEnable);
+	}
+	else if (bEnable)
+	{
+		UE_LOG(LogActor, Warning, TEXT("Attempting to enable ticking for '%s' when the actor is explicitly prevented from doing so (check CanEverTick)."), *GetName());
+	}
+}
+
 void AActor::AddControllingMatineeActor( AMatineeActor& InMatineeActor )
 {
 	if (RootComponent)
