- Instanced mesh emitters now have their stride set early enough for Metal

diff --git a/Engine/Source/Runtime/Engine/Private/MeshParticleVertexFactory.cpp b/Engine/Source/Runtime/Engine/Private/MeshParticleVertexFactory.cpp
--- a/Engine/Source/Runtime/Engine/Private/MeshParticleVertexFactory.cpp
+++ b/Engine/Source/Runtime/Engine/Private/MeshParticleVertexFactory.cpp
@@ -94,6 +94,7 @@ void FMeshParticleVertexFactory::InitRHI()
 		{
 			// Stream 0 - Instance data
 			{
+				checkf(DynamicVertexStride != -1, TEXT("FMeshParticleVertexFactory does not have a valid DynamicVertexStride - likely an empty one was made, but SetStrides was not called"));
 				FVertexStream VertexStream;
 				VertexStream.VertexBuffer = NULL;
 				VertexStream.Stride = 0;
@@ -101,26 +102,28 @@ void FMeshParticleVertexFactory::InitRHI()
 				Streams.Add(VertexStream);
 	
 				// @todo metal: this will need a valid stride when we get to instanced meshes!
-				Elements.Add(FVertexElement(0, Data.TransformComponent[0].Offset, Data.TransformComponent[0].Type, 8, 0xFFFF, Data.TransformComponent[0].bUseInstanceIndex));
-				Elements.Add(FVertexElement(0, Data.TransformComponent[1].Offset, Data.TransformComponent[1].Type, 9, 0xFFFF, Data.TransformComponent[1].bUseInstanceIndex));
-				Elements.Add(FVertexElement(0, Data.TransformComponent[2].Offset, Data.TransformComponent[2].Type, 10, 0xFFFF, Data.TransformComponent[2].bUseInstanceIndex));
+				Elements.Add(FVertexElement(0, Data.TransformComponent[0].Offset, Data.TransformComponent[0].Type, 8, DynamicVertexStride, Data.TransformComponent[0].bUseInstanceIndex));
+				Elements.Add(FVertexElement(0, Data.TransformComponent[1].Offset, Data.TransformComponent[1].Type, 9, DynamicVertexStride, Data.TransformComponent[1].bUseInstanceIndex));
+				Elements.Add(FVertexElement(0, Data.TransformComponent[2].Offset, Data.TransformComponent[2].Type, 10, DynamicVertexStride, Data.TransformComponent[2].bUseInstanceIndex));
 	
-				Elements.Add(FVertexElement(0, Data.SubUVs.Offset, Data.SubUVs.Type, 11, 0, Data.SubUVs.bUseInstanceIndex));
-				Elements.Add(FVertexElement(0, Data.SubUVLerpAndRelTime.Offset, Data.SubUVLerpAndRelTime.Type, 12, 0xFFFF, Data.SubUVLerpAndRelTime.bUseInstanceIndex));
+				Elements.Add(FVertexElement(0, Data.SubUVs.Offset, Data.SubUVs.Type, 11, DynamicVertexStride, Data.SubUVs.bUseInstanceIndex));
+				Elements.Add(FVertexElement(0, Data.SubUVLerpAndRelTime.Offset, Data.SubUVLerpAndRelTime.Type, 12, DynamicVertexStride, Data.SubUVLerpAndRelTime.bUseInstanceIndex));
 	
-				Elements.Add(FVertexElement(0, Data.ParticleColorComponent.Offset, Data.ParticleColorComponent.Type, 14, 0xFFFF, Data.ParticleColorComponent.bUseInstanceIndex));
-				Elements.Add(FVertexElement(0, Data.VelocityComponent.Offset, Data.VelocityComponent.Type, 15, 0xFFFF, Data.VelocityComponent.bUseInstanceIndex));
+				Elements.Add(FVertexElement(0, Data.ParticleColorComponent.Offset, Data.ParticleColorComponent.Type, 14, DynamicVertexStride, Data.ParticleColorComponent.bUseInstanceIndex));
+				Elements.Add(FVertexElement(0, Data.VelocityComponent.Offset, Data.VelocityComponent.Type, 15, DynamicVertexStride, Data.VelocityComponent.bUseInstanceIndex));
 			}
 
 			// Stream 1 - Dynamic parameter
 			{
+				checkf(DynamicParameterVertexStride != -1, TEXT("FMeshParticleVertexFactory does not have a valid DynamicParameterVertexStride - likely an empty one was made, but SetStrides was not called"));
+				
 				FVertexStream VertexStream;
 				VertexStream.VertexBuffer = NULL;
 				VertexStream.Stride = 0;
 				VertexStream.Offset = 0;
 				Streams.Add(VertexStream);
 	
-				Elements.Add(FVertexElement(1, 0, VET_Float4, 13, 0xFFFF, true));
+				Elements.Add(FVertexElement(1, 0, VET_Float4, 13, DynamicParameterVertexStride, true));
 			}
 		}
 
diff --git a/Engine/Source/Runtime/Engine/Private/Particles/ParticleSystemRender.cpp b/Engine/Source/Runtime/Engine/Private/Particles/ParticleSystemRender.cpp
--- a/Engine/Source/Runtime/Engine/Private/Particles/ParticleSystemRender.cpp
+++ b/Engine/Source/Runtime/Engine/Private/Particles/ParticleSystemRender.cpp
@@ -263,7 +263,7 @@ void ENGINE_API FinishOneFrameParticleStats()
  *
  *	@return	FParticleVertexFactoryBase*	The created VF; NULL if invalid InType
  */
-FParticleVertexFactoryBase* CreateParticleVertexFactory(EParticleVertexFactoryType InType, ERHIFeatureLevel::Type InFeatureLevel)
+static FParticleVertexFactoryBase* CreateParticleVertexFactory(EParticleVertexFactoryType InType, ERHIFeatureLevel::Type InFeatureLevel, const FDynamicSpriteEmitterDataBase* ParticleData)
 {
 	FParticleVertexFactoryBase* NewVertexFactory = NULL;
 	switch (InType)
@@ -275,7 +275,7 @@ FParticleVertexFactoryBase* CreateParticleVertexFactory(EParticleVertexFactoryTy
 		NewVertexFactory = new FParticleBeamTrailVertexFactory(PVFT_BeamTrail, InFeatureLevel);
 		break;
 	case PVFT_Mesh:
-		NewVertexFactory = new FMeshParticleVertexFactory(PVFT_Mesh, InFeatureLevel);
+		NewVertexFactory = new FMeshParticleVertexFactory(PVFT_Mesh, InFeatureLevel, ParticleData->GetDynamicVertexStride(InFeatureLevel), ParticleData->GetDynamicParameterVertexStride());
 		break;
 	default:
 		break;
@@ -292,7 +292,7 @@ FParticleOrderPool GParticleOrderPool;
 // Particle vertex factory pool
 FParticleVertexFactoryPool GParticleVertexFactoryPool;
 
-FParticleVertexFactoryBase* FParticleVertexFactoryPool::GetParticleVertexFactory(EParticleVertexFactoryType InType, ERHIFeatureLevel::Type InFeatureLevel)
+FParticleVertexFactoryBase* FParticleVertexFactoryPool::GetParticleVertexFactory(EParticleVertexFactoryType InType, ERHIFeatureLevel::Type InFeatureLevel, const FDynamicSpriteEmitterDataBase* ParticleData)
 {
 	SCOPE_CYCLE_COUNTER(STAT_ParticlePoolTime);
 	check(InType < PVFT_MAX);
@@ -302,7 +302,7 @@ FParticleVertexFactoryBase* FParticleVertexFactoryPool::GetParticleVertexFactory
 	if (AvailableFactories.Num() == 0)
 	{
 		// If there are none in the pool, create a new one, add it to the in use list and return it
-		VertexFactory = CreateParticleVertexFactory(InType, InFeatureLevel);
+		VertexFactory = CreateParticleVertexFactory(InType, InFeatureLevel, ParticleData);
 		VertexFactories.Add(VertexFactory);
 	}
 	else
@@ -1563,7 +1563,7 @@ int32 FDynamicSpriteEmitterData::RenderEmitter(FParticleSystemSceneProxy* Proxy,
 
 FParticleVertexFactoryBase* FDynamicSpriteEmitterData::BuildVertexFactory(const FParticleSystemSceneProxy* InOwnerProxy)
 {
-	return GParticleVertexFactoryPool.GetParticleVertexFactory(PVFT_Sprite, InOwnerProxy->GetScene().GetFeatureLevel());
+	return GParticleVertexFactoryPool.GetParticleVertexFactory(PVFT_Sprite, InOwnerProxy->GetScene().GetFeatureLevel(), this);
 }
 
 void FDynamicSpriteEmitterData::UpdateRenderThreadResourcesEmitter(const FParticleSystemSceneProxy* InOwnerProxy)
@@ -1776,7 +1776,7 @@ void FDynamicMeshEmitterData::Init( bool bInSelected,
 
 FParticleVertexFactoryBase* FDynamicMeshEmitterData::BuildVertexFactory(const FParticleSystemSceneProxy* InOwnerProxy)
 {
-	FParticleVertexFactoryBase* VertexFactory = GParticleVertexFactoryPool.GetParticleVertexFactory(PVFT_Mesh, InOwnerProxy->GetScene().GetFeatureLevel());
+	FParticleVertexFactoryBase* VertexFactory = GParticleVertexFactoryPool.GetParticleVertexFactory(PVFT_Mesh, InOwnerProxy->GetScene().GetFeatureLevel(), this);
 	SetupVertexFactory((FMeshParticleVertexFactory*)VertexFactory, StaticMesh->RenderData->LODResources[0]);
 	return VertexFactory;
 }
@@ -1858,6 +1858,7 @@ void FDynamicMeshEmitterData::GetDynamicMeshElementsEmitter(const FParticleSyste
 			FDynamicMeshEmitterCollectorResources& CollectorResources = Collector.AllocateOneFrameResource<FDynamicMeshEmitterCollectorResources>();
 			CollectorResources.VertexFactory.SetParticleFactoryType(PVFT_Mesh);
 			CollectorResources.VertexFactory.SetFeatureLevel(FeatureLevel);
+			CollectorResources.VertexFactory.SetStrides(GetDynamicVertexStride(FeatureLevel), GetDynamicParameterVertexStride());
 			CollectorResources.VertexFactory.InitResource();
 			SetupVertexFactory(&CollectorResources.VertexFactory, StaticMesh->RenderData->LODResources[0]);
 
@@ -2876,7 +2877,7 @@ void FDynamicBeam2EmitterData::Init( bool bInSelected )
 
 FParticleVertexFactoryBase* FDynamicBeam2EmitterData::BuildVertexFactory(const FParticleSystemSceneProxy* InOwnerProxy)
 {
-	return (FParticleBeamTrailVertexFactory*)(GParticleVertexFactoryPool.GetParticleVertexFactory(PVFT_BeamTrail, InOwnerProxy->GetScene().GetFeatureLevel()));
+	return (FParticleBeamTrailVertexFactory*)(GParticleVertexFactoryPool.GetParticleVertexFactory(PVFT_BeamTrail, InOwnerProxy->GetScene().GetFeatureLevel(), this));
 }
 
 /** Perform the actual work of filling the buffer, often called from another thread 
@@ -5819,7 +5820,7 @@ void FDynamicTrailsEmitterData::Init(bool bInSelected)
 
 FParticleVertexFactoryBase* FDynamicTrailsEmitterData::BuildVertexFactory(const FParticleSystemSceneProxy* InOwnerProxy)
 {
-	return (FParticleBeamTrailVertexFactory*)(GParticleVertexFactoryPool.GetParticleVertexFactory(PVFT_BeamTrail, InOwnerProxy->GetScene().GetFeatureLevel()));
+	return (FParticleBeamTrailVertexFactory*)(GParticleVertexFactoryPool.GetParticleVertexFactory(PVFT_BeamTrail, InOwnerProxy->GetScene().GetFeatureLevel(), this));
 }
 
 void FDynamicTrailsEmitterData::GetDynamicMeshElementsEmitter(const FParticleSystemSceneProxy* Proxy, const FSceneView* View, const FSceneViewFamily& ViewFamily, int32 ViewIndex, FMeshElementCollector& Collector) const
diff --git a/Engine/Source/Runtime/Engine/Public/MeshParticleVertexFactory.h b/Engine/Source/Runtime/Engine/Public/MeshParticleVertexFactory.h
--- a/Engine/Source/Runtime/Engine/Public/MeshParticleVertexFactory.h
+++ b/Engine/Source/Runtime/Engine/Public/MeshParticleVertexFactory.h
@@ -75,12 +75,16 @@ public:
 	};
 
 	/** Default constructor. */
-	FMeshParticleVertexFactory(EParticleVertexFactoryType InType, ERHIFeatureLevel::Type InFeatureLevel)
+	FMeshParticleVertexFactory(EParticleVertexFactoryType InType, ERHIFeatureLevel::Type InFeatureLevel, int32 InDynamicVertexStride, int32 InDynamicParameterVertexStride)
 		: FParticleVertexFactoryBase(InType, InFeatureLevel)
+		, DynamicVertexStride(InDynamicVertexStride)
+		, DynamicParameterVertexStride(InDynamicParameterVertexStride)
 	{}
 
 	FMeshParticleVertexFactory()
 		: FParticleVertexFactoryBase(PVFT_MAX, ERHIFeatureLevel::Num)
+		, DynamicVertexStride(-1)
+		, DynamicParameterVertexStride(-1)
 	{}
 
 	/**
@@ -127,6 +131,15 @@ public:
 		return MeshParticleUniformBuffer;
 	}
 
+	/**
+	 * Update the data strides (MUST HAPPEN BEFORE InitRHI is called)
+	 */
+	void SetStrides(int32 InDynamicVertexStride, int32 InDynamicParameterVertexStride)
+	{
+		DynamicVertexStride = InDynamicVertexStride;
+		DynamicParameterVertexStride = InDynamicParameterVertexStride;
+	}
+	
 	 /**
 	 * Set the source vertex buffer that contains particle instance data.
 	 */
@@ -152,6 +165,11 @@ public:
 
 protected:
 	DataType Data;
+	
+	/** Stride information for instanced mesh particles */
+	int32 DynamicVertexStride;
+	int32 DynamicParameterVertexStride;
+	
 
 	/** Uniform buffer with mesh particle paramters. */
 	FUniformBufferRHIParamRef MeshParticleUniformBuffer;
diff --git a/Engine/Source/Runtime/Engine/Public/ParticleHelper.h b/Engine/Source/Runtime/Engine/Public/ParticleHelper.h
--- a/Engine/Source/Runtime/Engine/Public/ParticleHelper.h
+++ b/Engine/Source/Runtime/Engine/Public/ParticleHelper.h
@@ -1098,7 +1098,7 @@ public:
 		ClearPoolInternal();
 	}
 
-	FParticleVertexFactoryBase* GetParticleVertexFactory(EParticleVertexFactoryType InType, ERHIFeatureLevel::Type InFeatureLevel);
+	FParticleVertexFactoryBase* GetParticleVertexFactory(EParticleVertexFactoryType InType, ERHIFeatureLevel::Type InFeatureLevel, const struct FDynamicSpriteEmitterDataBase* ParticleData);
 
 	bool ReturnParticleVertexFactory(FParticleVertexFactoryBase* InVertexFactory);
 
