#ue4 - Don't assert in UCharacterMovementComponent::CapsuleTouched() if bEnablePhysicsInteraction is disabled, since it can be changed dynamically after the function is registered. Also don't gate removal of the delegate on this variable for the same reason.

UDN: https://udn.unrealengine.com/questions/223979/player-to-rigid-body-interaction.html

diff --git a/Engine/Source/Runtime/Engine/Private/CharacterMovementComponent.cpp b/Engine/Source/Runtime/Engine/Private/CharacterMovementComponent.cpp
--- a/Engine/Source/Runtime/Engine/Private/CharacterMovementComponent.cpp
+++ b/Engine/Source/Runtime/Engine/Private/CharacterMovementComponent.cpp
@@ -372,7 +372,7 @@ void UCharacterMovementComponent::SetUpdatedComponent(UPrimitiveComponent* NewUp
 	bDeferUpdateMoveComponent = false;
 	DeferredUpdatedMoveComponent = NULL;
 
-	if (UpdatedComponent != NULL && UpdatedComponent->OnComponentBeginOverlap.IsBound() && bEnablePhysicsInteraction)
+	if (UpdatedComponent != NULL && UpdatedComponent->OnComponentBeginOverlap.IsBound())
 	{
 		UpdatedComponent->OnComponentBeginOverlap.RemoveDynamic(this, &UCharacterMovementComponent::CapsuleTouched);
 	}
@@ -6188,7 +6188,10 @@ void UCharacterMovementComponent::ClientAckGoodMove_Implementation(float TimeSta
 
 void UCharacterMovementComponent::CapsuleTouched( AActor* Other, UPrimitiveComponent* OtherComp, int32 OtherBodyIndex, bool bFromSweep, const FHitResult& SweepResult )
 {
-	check(bEnablePhysicsInteraction);
+	if (!bEnablePhysicsInteraction)
+	{
+		return;
+	}
 
 	if (OtherComp != NULL && OtherComp->IsAnySimulatingPhysics())
 	{
