Integrate BlackTusk handle based timer optimizations and bug fixes

diff --git a/Engine/Source/Runtime/Engine/Private/TimerManager.cpp b/Engine/Source/Runtime/Engine/Private/TimerManager.cpp
--- a/Engine/Source/Runtime/Engine/Private/TimerManager.cpp
+++ b/Engine/Source/Runtime/Engine/Private/TimerManager.cpp
@@ -107,7 +107,7 @@ int32 FTimerManager::FindTimerInList(const TArray<FTimerData> &SearchArray, FTim
 {
 	for (int32 Idx=0; Idx<SearchArray.Num(); ++Idx)
 	{
-		if (SearchArray[Idx].TimerDelegate == InDelegate && !SearchArray[Idx].TimerHandle.IsValid())
+		if (!SearchArray[Idx].TimerHandle.IsValid() && SearchArray[Idx].TimerDelegate == InDelegate)
 		{
 			return Idx;
 		}
@@ -119,16 +119,14 @@ int32 FTimerManager::FindTimerInList(const TArray<FTimerData> &SearchArray, FTim
 /** Will find the given timer in the given TArray and return its index. */
 int32 FTimerManager::FindTimerInList(const TArray<FTimerData> &SearchArray, FTimerHandle const& InHandle) const
 {
-	if (!(InHandle.IsValid()))
+	if (InHandle.IsValid())
 	{
-		return INDEX_NONE;
-	}
-
-	for (int32 Idx = 0; Idx < SearchArray.Num(); ++Idx)
-	{
-		if (SearchArray[Idx].TimerHandle == InHandle)
+		for (int32 Idx = 0; Idx < SearchArray.Num(); ++Idx)
 		{
-			return Idx;
+			if (SearchArray[Idx].TimerHandle == InHandle)
+			{
+				return Idx;
+			}
 		}
 	}
 
@@ -181,7 +179,7 @@ void FTimerManager::InternalSetTimer(FTimerHandle& InOutHandle, FTimerUnifiedDel
 
 void FTimerManager::InternalSetTimer(FTimerData& NewTimerData, float InRate, bool InbLoop, float InFirstDelay)
 {
-	if (NewTimerData.TimerDelegate.IsBound() || NewTimerData.TimerHandle.IsValid())
+	if (NewTimerData.TimerHandle.IsValid() || NewTimerData.TimerDelegate.IsBound())
 	{
 		NewTimerData.Rate = InRate;
 		NewTimerData.bLoop = InbLoop;
@@ -239,7 +237,7 @@ void FTimerManager::InternalClearTimer(FTimerUnifiedDelegate const& InDelegate)
 	{
 		// Edge case. We're currently handling this timer when it got cleared.  Unbind it to prevent it firing again
 		// in case it was scheduled to fire multiple times.
-		if (CurrentlyExecutingTimer.TimerDelegate == InDelegate)
+		if (!CurrentlyExecutingTimer.TimerHandle.IsValid() && CurrentlyExecutingTimer.TimerDelegate == InDelegate)
 		{
 			CurrentlyExecutingTimer.TimerDelegate.Unbind();
 		}
@@ -251,6 +249,12 @@ void FTimerManager::InternalClearTimer(FTimerHandle const& InHandle)
 	// not currently threadsafe
 	check(IsInGameThread());
 
+	// Skip if the handle is invalid as it  would not be found by FindTimer and unbind the current handler if it also used INDEX_NONE. 
+	if (!InHandle.IsValid())
+	{
+		return;
+	}
+
 	int32 TimerIdx;
 	FTimerData const* const TimerData = FindTimer(InHandle, &TimerIdx);
 	if (TimerData)
@@ -270,6 +274,7 @@ void FTimerManager::InternalClearTimer(FTimerHandle const& InHandle)
 		if (CurrentlyExecutingTimer.TimerHandle == InHandle)
 		{
 			CurrentlyExecutingTimer.TimerDelegate.Unbind();
+			CurrentlyExecutingTimer.TimerHandle.Invalidate();
 		}
 	}
 }
@@ -466,18 +471,19 @@ void FTimerManager::Tick(float DeltaTime)
 				CurrentlyExecutingTimer.TimerDelegate.Execute();
 
 				// If timer was cleared in the delegate execution, don't execute further 
-				if( !CurrentlyExecutingTimer.TimerDelegate.IsBound() && !CurrentlyExecutingTimer.TimerHandle.IsValid() )
+				if( !CurrentlyExecutingTimer.TimerHandle.IsValid() && !CurrentlyExecutingTimer.TimerDelegate.IsBound() )
 				{
 					break;
 				}
 			}
 
 			if( CurrentlyExecutingTimer.bLoop && 
-				(CurrentlyExecutingTimer.TimerDelegate.IsBound() ||
-				 CurrentlyExecutingTimer.TimerHandle.IsValid()) && 							// did not get cleared during execution
-				((FindTimer(CurrentlyExecutingTimer.TimerDelegate) == NULL) || 
-				 (CurrentlyExecutingTimer.TimerHandle.IsValid() && FindTimer(CurrentlyExecutingTimer.TimerHandle) == NULL))		// did not get manually re-added during execution
-			  )
+				(CurrentlyExecutingTimer.TimerHandle.IsValid() ||
+				 CurrentlyExecutingTimer.TimerDelegate.IsBound()) && 							// did not get cleared during execution
+				(CurrentlyExecutingTimer.TimerHandle.IsValid() ? 
+					(FindTimer(CurrentlyExecutingTimer.TimerHandle) == nullptr) : 
+					(FindTimer(CurrentlyExecutingTimer.TimerDelegate) == nullptr)) // did not get manually re-added during execution			  
+				)
 			{
 				// Put this timer back on the heap
 				CurrentlyExecutingTimer.ExpireTime += CallCount * CurrentlyExecutingTimer.Rate;
diff --git a/Engine/Source/Runtime/Engine/Private/TimerManagerTests.cpp b/Engine/Source/Runtime/Engine/Private/TimerManagerTests.cpp
--- a/Engine/Source/Runtime/Engine/Private/TimerManagerTests.cpp
+++ b/Engine/Source/Runtime/Engine/Private/TimerManagerTests.cpp
@@ -390,6 +390,101 @@ bool TimerManagerTest_ValidTimer_HandleNoDelegate(UWorld* World, FAutomationTest
 	return true;
 }
 
+struct FLoopingTestFunc
+{
+	static FTimerManager* TimerManager;
+	static FTimerHandle* Handle;
+	static int TimerCalled;
+	static float NewTime;
+	static void TimerExecute()
+	{
+		TimerCalled++;
+		if (TimerCalled == 1)
+		{
+			TimerManager->SetTimer(*Handle, FTimerDelegate::CreateStatic(FLoopingTestFunc::TimerExecute), NewTime, true);
+		}
+		else
+		{
+			TimerManager->ClearTimer(*Handle);
+		}
+	}
+};
+
+FTimerManager* FLoopingTestFunc::TimerManager = nullptr;
+FTimerHandle* FLoopingTestFunc::Handle = nullptr;
+int FLoopingTestFunc::TimerCalled = 0;
+float FLoopingTestFunc::NewTime = 1.0f;
+
+bool TimerManagerTest_ValidTimer_HandleLoopingSetDuringExecute(UWorld* World, FAutomationTestBase* Test)
+{
+	FTimerManager& TimerManager = World->GetTimerManager();
+	FTimerHandle Handle;
+	const float Rate = 3.f;
+
+	FLoopingTestFunc::TimerManager = &TimerManager;
+	FLoopingTestFunc::Handle = &Handle;
+	FLoopingTestFunc::TimerCalled = 0;
+
+	Test->TestTrue(TIMER_TEST_TEXT("Timer called count starts at 0"), FLoopingTestFunc::TimerCalled == 0);
+
+	TimerManager.SetTimer(Handle, FTimerDelegate::CreateStatic(FLoopingTestFunc::TimerExecute), Rate, true);
+
+	// small tick to move the timer from the pending list to the active list, the timer will start counting time after this tick
+	TimerTest_TickWorld(World, KINDA_SMALL_NUMBER);
+
+	TimerTest_TickWorld(World, 3.0f);
+	Test->TestTrue(TIMER_TEST_TEXT("Timer was called first time"), FLoopingTestFunc::TimerCalled == 1);
+	Test->TestTrue(TIMER_TEST_TEXT("Timer was readded"), TimerManager.IsTimerActive(Handle));
+	Test->TestTrue(TIMER_TEST_TEXT("Timer was readded with correct time"), FMath::IsNearlyEqual(TimerManager.GetTimerRemaining(Handle), FLoopingTestFunc::NewTime, 1e-2f));
+
+	TimerTest_TickWorld(World, 1.1f);
+	Test->TestTrue(TIMER_TEST_TEXT("Timer was called second time"), FLoopingTestFunc::TimerCalled == 2);
+	Test->TestFalse(TIMER_TEST_TEXT("Timer handle no longer active"), TimerManager.IsTimerActive(Handle));
+
+	return true;
+}
+
+bool TimerManagerTest_LoopingTimers_DifferentHandles(UWorld* World, FAutomationTestBase* Test)
+{
+	FTimerManager& TimerManager = World->GetTimerManager();
+	FTimerHandle HandleOne, HandleTwo;
+
+	int32 CallCount = 0;
+	auto Func = [](int* CallCount){ (*CallCount)++; };
+
+	FTimerDelegate Delegate = FTimerDelegate::CreateStatic(Func, &CallCount);
+
+	TimerManager.SetTimer(Delegate, 1.0f, false);
+	TimerTest_TickWorld(World, KINDA_SMALL_NUMBER);
+
+	Test->TestTrue(TIMER_TEST_TEXT("First delegate time remaining is 1.0f"),
+		FMath::IsNearlyEqual(TimerManager.GetTimerRemaining(Delegate), 1.0f, KINDA_SMALL_NUMBER));
+	
+	TimerManager.SetTimer(Delegate, 5.0f, false);
+	TimerTest_TickWorld(World, KINDA_SMALL_NUMBER);
+	Test->TestTrue(TIMER_TEST_TEXT("Reset delegate time remaining is 5.0f"),
+		FMath::IsNearlyEqual(TimerManager.GetTimerRemaining(Delegate), 5.0f, KINDA_SMALL_NUMBER));
+	
+	TimerManager.SetTimer(HandleOne, Delegate, 1.0f, true);
+	TimerManager.SetTimer(HandleTwo, Delegate, 1.5f, true);
+	TimerTest_TickWorld(World, KINDA_SMALL_NUMBER);
+
+	Test->TestTrue(TIMER_TEST_TEXT("Handle One is active"), TimerManager.IsTimerActive(HandleOne));
+	Test->TestTrue(TIMER_TEST_TEXT("Handle Two is active"), TimerManager.IsTimerActive(HandleTwo));
+
+	TimerTest_TickWorld(World, 1.0f);
+
+	Test->TestTrue(TIMER_TEST_TEXT("Handle One is active after tick"), TimerManager.IsTimerActive(HandleOne));
+	Test->TestTrue(TIMER_TEST_TEXT("Handle Two is active after tick"), TimerManager.IsTimerActive(HandleTwo));
+
+	Test->TestTrue(TIMER_TEST_TEXT("Handle One has 0 seconds remaining after tick"), 
+		FMath::IsNearlyEqual(TimerManager.GetTimerRemaining(HandleOne), 0.0f, 1e-2f) );
+	Test->TestTrue(TIMER_TEST_TEXT("Handle Two has 0.5 seconds remaining after tick"),
+		FMath::IsNearlyEqual(TimerManager.GetTimerRemaining(HandleTwo), 0.5f, 1e-2f));
+
+	return true;
+}
+
 bool FTimerManagerTest::RunTest(const FString& Parameters)
 {
 	UWorld *World = UWorld::CreateWorld(EWorldType::Game, false);
@@ -405,6 +500,8 @@ bool FTimerManagerTest::RunTest(const FString& Parameters)
 	TimerManagerTest_ValidTimer_Delegate(World, this);
 	TimerManagerTest_ValidTimer_HandleWithDelegate(World, this);
 	TimerManagerTest_ValidTimer_HandleNoDelegate(World, this);
+	TimerManagerTest_ValidTimer_HandleLoopingSetDuringExecute(World, this);
+	TimerManagerTest_LoopingTimers_DifferentHandles(World, this);
 
 	GEngine->DestroyWorldContext(World);
 	World->DestroyWorld(false);
