Replace actor filtering with component filtering for physics scene queries. codereview James.Golding

diff --git a/Engine/Source/Runtime/Engine/Private/Collision/Collision.cpp b/Engine/Source/Runtime/Engine/Private/Collision/Collision.cpp
--- a/Engine/Source/Runtime/Engine/Private/Collision/Collision.cpp
+++ b/Engine/Source/Runtime/Engine/Private/Collision/Collision.cpp
@@ -131,9 +131,14 @@ FCollisionQueryParams::FCollisionQueryParams(FName InTraceTag, bool bInTraceComp
 
 void FCollisionQueryParams::AddIgnoredActor(const AActor* InIgnoreActor)
 {
-	if (InIgnoreActor != NULL)
+	if (InIgnoreActor)
 	{
-		IgnoreActors.AddUnique(InIgnoreActor->GetUniqueID());
+		TArray<UPrimitiveComponent*> PrimComponents;
+		InIgnoreActor->GetComponents(PrimComponents);
+		for (UPrimitiveComponent* PrimComponent : PrimComponents)
+		{
+			IgnoreComponents.AddUnique(PrimComponent->GetUniqueID());
+		}
 	}
 }
 
@@ -141,10 +146,9 @@ void FCollisionQueryParams::AddIgnoredActors(const TArray<AActor*>& InIgnoreActo
 {
 	for (int32 Idx = 0; Idx < InIgnoreActors.Num(); ++Idx)
 	{
-		AActor const* const A = InIgnoreActors[Idx];
-		if (A)
+		if (AActor const* const A = InIgnoreActors[Idx])
 		{
-			IgnoreActors.Add(A->GetUniqueID());
+			AddIgnoredActor(A);
 		}
 	}
 }
@@ -153,10 +157,37 @@ void FCollisionQueryParams::AddIgnoredActors(const TArray<TWeakObjectPtr<AActor>
 {
 	for (int32 Idx = 0; Idx < InIgnoreActors.Num(); ++Idx)
 	{
-		AActor const* const A = InIgnoreActors[Idx].Get();
-		if (A)
+		if (AActor const* const A = InIgnoreActors[Idx].Get())
+		{
+			AddIgnoredActor(A);
+		}
+	}
+}
+
+
+void FCollisionQueryParams::AddIgnoredComponent(const UPrimitiveComponent* InIgnoreComponent)
+{
+	if (InIgnoreComponent)
+	{
+		IgnoreComponents.AddUnique(InIgnoreComponent->GetUniqueID());
+	}
+}
+
+void FCollisionQueryParams::AddIgnoredComponents(const TArray<UPrimitiveComponent*>& InIgnoreComponents)
+{
+	for (const UPrimitiveComponent* IgnoreComponent : InIgnoreComponents)
+	{
+		AddIgnoredComponent(IgnoreComponent);
+	}
+}
+
+void FCollisionQueryParams::AddIgnoredComponents(const TArray<TWeakObjectPtr<UPrimitiveComponent> >& InIgnoreComponents)
+{
+	for (TWeakObjectPtr<UPrimitiveComponent> IgnoreComponent : InIgnoreComponents)
+	{
+		if (IgnoreComponent.IsValid())
 		{
-			IgnoreActors.Add(A->GetUniqueID());
+			AddIgnoredComponent(IgnoreComponent.Get());
 		}
 	}
 }
diff --git a/Engine/Source/Runtime/Engine/Private/Collision/PhysXCollision.cpp b/Engine/Source/Runtime/Engine/Private/Collision/PhysXCollision.cpp
--- a/Engine/Source/Runtime/Engine/Private/Collision/PhysXCollision.cpp
+++ b/Engine/Source/Runtime/Engine/Private/Collision/PhysXCollision.cpp
@@ -175,14 +175,13 @@ FRaycastHelperBox2D_Base::EResponse FRaycastHelperBox2D_Base::CalcQueryHitType(b
 		Result = ResponseNone;
 	}
 
-	// Check to see if the actor ultimately owning the fixture should be ignored
+	// Check to see if the component ultimately owning the fixture should be ignored
 	if (Result != ResponseNone)
 	{
 		FBodyInstance* OwnerBodyInstance = static_cast<FBodyInstance*>(Fixture->GetBody()->GetUserData());
 		UPrimitiveComponent* OwnerComponent = OwnerBodyInstance->OwnerComponent.Get();
-		AActor* Owner = OwnerComponent ? OwnerComponent->GetOwner() : nullptr;
-
-		if ((Owner != nullptr) && Params.IgnoreActors.Contains(Owner->GetUniqueID()))
+		
+		if ((OwnerComponent != nullptr) && Params.IgnoreComponents.Contains(OwnerComponent->GetUniqueID()))
 		{
 			Result = ResponseNone;
 		}
@@ -595,7 +594,7 @@ PxSceneQueryHitType::Enum FPxQueryFilterCallback::preFilter(const PxFilterData&
 #endif // ENABLE_PREFILTER_LOGGING
 
 	// See if we are ignoring the actor this shape belongs to (word0 of shape filterdata is actorID)
-	if(IgnoreActors.Contains(ShapeFilter.word0))
+	if(IgnoreComponents.Contains(ShapeFilter.word0))
 	{
 		//UE_LOG(LogTemp, Log, TEXT("Ignoring Actor: %d"), ShapeFilter.word0);
 		return (PrefilterReturnValue = PxSceneQueryHitType::eNONE);
@@ -790,7 +789,7 @@ bool RaycastTest(const UWorld* World, const FVector Start, const FVector End, EC
 			PxFilterData PFilter = CreateQueryFilterData(TraceChannel, Params.bTraceComplex, ResponseParams.CollisionResponse, ObjectParams, false);
 			PxSceneQueryFilterData PQueryFilterData(PFilter, PxSceneQueryFilterFlag::eSTATIC | PxSceneQueryFilterFlag::eDYNAMIC | PxSceneQueryFilterFlag::ePREFILTER);
 			PxSceneQueryFlags POutputFlags = PxHitFlags();
-			FPxQueryFilterCallback PQueryCallback(Params.IgnoreActors);
+			FPxQueryFilterCallback PQueryCallback(Params.IgnoreComponents);
 			PQueryCallback.bSingleQuery = true;
 
 			FPhysScene* PhysScene = World->GetPhysicsScene();
@@ -862,7 +861,7 @@ bool RaycastSingle(const UWorld* World, struct FHitResult& OutHit, const FVector
 			PxFilterData PFilter = CreateQueryFilterData(TraceChannel, Params.bTraceComplex, ResponseParams.CollisionResponse, ObjectParams, false);
 			PxSceneQueryFilterData PQueryFilterData(PFilter, PxSceneQueryFilterFlag::eSTATIC | PxSceneQueryFilterFlag::eDYNAMIC | PxSceneQueryFilterFlag::ePREFILTER);
 			PxSceneQueryFlags POutputFlags = PxSceneQueryFlag::ePOSITION | PxSceneQueryFlag::eNORMAL | PxSceneQueryFlag::eDISTANCE | PxSceneQueryFlag::eMTD;
-			FPxQueryFilterCallback PQueryCallback(Params.IgnoreActors);
+			FPxQueryFilterCallback PQueryCallback(Params.IgnoreComponents);
 			PQueryCallback.bSingleQuery = true; // this needs to be set for raycastSingle so we return eNONE instead of eTOUCH
 
 			PxVec3 PStart = U2PVector(Start);
@@ -966,7 +965,7 @@ bool RaycastMulti(const UWorld* World, TArray<struct FHitResult>& OutHits, const
 		PxFilterData PFilter = CreateQueryFilterData(TraceChannel, Params.bTraceComplex, ResponseParams.CollisionResponse, ObjectParams, true);
 		PxSceneQueryFilterData PQueryFilterData(PFilter, PxSceneQueryFilterFlag::eSTATIC | PxSceneQueryFilterFlag::eDYNAMIC | PxSceneQueryFilterFlag::ePREFILTER);
 		PxSceneQueryFlags POutputFlags = PxSceneQueryFlag::ePOSITION | PxSceneQueryFlag::eNORMAL | PxSceneQueryFlag::eDISTANCE | PxSceneQueryFlag::eMTD;
-		FPxQueryFilterCallback PQueryCallback(Params.IgnoreActors);
+		FPxQueryFilterCallback PQueryCallback(Params.IgnoreComponents);
 
 		// Create buffer for hits. Note: memory is not initialized (for perf reasons), since API does not require it.
 		TTypeCompatibleBytes<PxRaycastHit> RawPHits[HIT_BUFFER_SIZE];
@@ -1126,7 +1125,7 @@ bool GeomSweepTest(const UWorld* World, const struct FCollisionShape& CollisionS
 		PxSceneQueryFilterData PQueryFilterData(PFilter, PxSceneQueryFilterFlag::eSTATIC | PxSceneQueryFilterFlag::eDYNAMIC | PxSceneQueryFilterFlag::ePREFILTER | PxSceneQueryFilterFlag::ePOSTFILTER);
 		PxSceneQueryFlags POutputFlags; 
 
-		FPxQueryFilterCallbackSweep PQueryCallbackSweep(Params.IgnoreActors);
+		FPxQueryFilterCallbackSweep PQueryCallbackSweep(Params.IgnoreComponents);
 		PQueryCallbackSweep.DiscardInitialOverlaps = !Params.bFindInitialOverlaps;
 		PQueryCallbackSweep.bSingleQuery = true;
 
@@ -1198,7 +1197,7 @@ bool GeomSweepSingle(const UWorld* World, const struct FCollisionShape& Collisio
 		//UE_LOG(LogCollision, Log, TEXT("PFilter: %x %x %x %x"), PFilter.word0, PFilter.word1, PFilter.word2, PFilter.word3);
 		PxSceneQueryFilterData PQueryFilterData(PFilter, PxSceneQueryFilterFlag::eSTATIC | PxSceneQueryFilterFlag::eDYNAMIC | PxSceneQueryFilterFlag::ePREFILTER);
 		PxSceneQueryFlags POutputFlags = PxSceneQueryFlag::ePOSITION | PxSceneQueryFlag::eNORMAL | PxSceneQueryFlag::eDISTANCE | PxSceneQueryFlag::eMTD;
-		FPxQueryFilterCallbackSweep PQueryCallbackSweep(Params.IgnoreActors);
+		FPxQueryFilterCallbackSweep PQueryCallbackSweep(Params.IgnoreComponents);
 		PQueryCallbackSweep.bSingleQuery = true; //LOC_MOD need to do this to return eNONE instead of eTOUCH for r
 		PQueryCallbackSweep.DiscardInitialOverlaps = !Params.bFindInitialOverlaps;	
 
@@ -1301,7 +1300,7 @@ bool GeomSweepMulti_PhysX(const UWorld* World, const PxGeometry& PGeom, const Px
 	PxFilterData PFilter = CreateQueryFilterData(TraceChannel, Params.bTraceComplex, ResponseParams.CollisionResponse, ObjectParams, true);
 	PxSceneQueryFilterData PQueryFilterData(PFilter, PxSceneQueryFilterFlag::eSTATIC | PxSceneQueryFilterFlag::eDYNAMIC | PxSceneQueryFilterFlag::ePREFILTER | PxSceneQueryFilterFlag::ePOSTFILTER);
 	PxSceneQueryFlags POutputFlags = PxSceneQueryFlag::ePOSITION | PxSceneQueryFlag::eNORMAL | PxSceneQueryFlag::eDISTANCE | PxSceneQueryFlag::eMTD;
-	FPxQueryFilterCallbackSweep PQueryCallbackSweep(Params.IgnoreActors);
+	FPxQueryFilterCallbackSweep PQueryCallbackSweep(Params.IgnoreComponents);
 	PQueryCallbackSweep.DiscardInitialOverlaps = !Params.bFindInitialOverlaps;
 
 	const FVector Delta = End - Start;
@@ -1478,7 +1477,7 @@ bool GeomOverlapTest(const UWorld* World, const struct FCollisionShape& Collisio
 	// Create filter data used to filter collisions
 	PxFilterData PFilter = CreateQueryFilterData(TraceChannel, Params.bTraceComplex, ResponseParams.CollisionResponse, ObjectParams, false);
 	PxSceneQueryFilterData PQueryFilterData(PFilter, PxSceneQueryFilterFlag::eSTATIC | PxSceneQueryFilterFlag::eDYNAMIC | PxSceneQueryFilterFlag::ePREFILTER);
-	FPxQueryFilterCallback PQueryCallback(Params.IgnoreActors);
+	FPxQueryFilterCallback PQueryCallback(Params.IgnoreComponents);
 
 	PxOverlapHit POverlapResult;
 	bool bHit = false;
@@ -1552,7 +1551,7 @@ bool GeomOverlapSingle(const UWorld* World, const struct FCollisionShape& Collis
 	// Create filter data used to filter collisions
 	PxFilterData PFilter = CreateQueryFilterData(TraceChannel, Params.bTraceComplex, ResponseParams.CollisionResponse, ObjectParams, false);
 	PxSceneQueryFilterData PQueryFilterData(PFilter, PxSceneQueryFilterFlag::eSTATIC | PxSceneQueryFilterFlag::eDYNAMIC | PxSceneQueryFilterFlag::ePREFILTER);
-	FPxQueryFilterCallback PQueryCallback(Params.IgnoreActors);
+	FPxQueryFilterCallback PQueryCallback(Params.IgnoreComponents);
 
 	// Enable scene locks, in case they are required
 	FPhysScene* PhysScene = World->GetPhysicsScene();
@@ -1620,7 +1619,7 @@ bool GeomOverlapMulti_PhysX(const UWorld* World, const PxGeometry& PGeom, const
 		// Create filter data used to filter collisions
 		PxFilterData PFilter = CreateQueryFilterData(TraceChannel, Params.bTraceComplex, ResponseParams.CollisionResponse, ObjectParams, true);
 		PxSceneQueryFilterData PQueryFilterData(PFilter, PxSceneQueryFilterFlag::eSTATIC | PxSceneQueryFilterFlag::eDYNAMIC | PxSceneQueryFilterFlag::ePREFILTER);
-		FPxQueryFilterCallback PQueryCallback(Params.IgnoreActors);
+		FPxQueryFilterCallback PQueryCallback(Params.IgnoreComponents);
 
 		// Enable scene locks, in case they are required
 		FPhysScene* PhysScene = World->GetPhysicsScene();
diff --git a/Engine/Source/Runtime/Engine/Private/Collision/PhysXCollision.h b/Engine/Source/Runtime/Engine/Private/Collision/PhysXCollision.h
--- a/Engine/Source/Runtime/Engine/Private/Collision/PhysXCollision.h
+++ b/Engine/Source/Runtime/Engine/Private/Collision/PhysXCollision.h
@@ -79,7 +79,7 @@ class FPxQueryFilterCallback : public PxSceneQueryFilterCallback
 {
 public:
 	/** List of ActorIds for this query to ignore */
-	TArray<uint32, TInlineAllocator<1> >	IgnoreActors;
+	TArray<uint32, TInlineAllocator<1> >	IgnoreComponents;
 	PxSceneQueryHitType::Enum				PrefilterReturnValue;
 
 	/** Whether this is a raycastSingle or sweepSingle, because in 3.3 you can't return eTOUCH for single queries */
@@ -92,10 +92,11 @@ public:
 		bSingleQuery = false;
 	}
 
-	FPxQueryFilterCallback(const TArray<uint32, TInlineAllocator<1> > InIgnoreActors)
+	FPxQueryFilterCallback(const TArray<uint32, TInlineAllocator<1> > InIgnoreComponents)
 	{
 		PrefilterReturnValue = PxSceneQueryHitType::eNONE;
-		IgnoreActors = InIgnoreActors;
+		
+		IgnoreComponents = InIgnoreComponents;
 		bSingleQuery = false;
 	}
 
@@ -125,8 +126,8 @@ class FPxQueryFilterCallbackSweep : public FPxQueryFilterCallback
 public:
 	bool DiscardInitialOverlaps;
 
-	FPxQueryFilterCallbackSweep(const TArray<uint32, TInlineAllocator<1> > InIgnoreActors)
-		: FPxQueryFilterCallback(InIgnoreActors)
+	FPxQueryFilterCallbackSweep(const TArray<uint32, TInlineAllocator<1> > InIgnoreComponents)
+		: FPxQueryFilterCallback(InIgnoreComponents)
 	{
 		DiscardInitialOverlaps = false;
 	}
diff --git a/Engine/Source/Runtime/Engine/Private/Collision/PhysicsFiltering.h b/Engine/Source/Runtime/Engine/Private/Collision/PhysicsFiltering.h
--- a/Engine/Source/Runtime/Engine/Private/Collision/PhysicsFiltering.h
+++ b/Engine/Source/Runtime/Engine/Private/Collision/PhysicsFiltering.h
@@ -33,16 +33,16 @@ struct FPhysicsFilterBuilder
 		}
 	}
 
-	inline void GetQueryData(uint32 ActorID, uint32& OutWord0, uint32& OutWord1, uint32& OutWord2, uint32& OutWord3) const
+	inline void GetQueryData(uint32 ComponentID, uint32& OutWord0, uint32& OutWord1, uint32& OutWord2, uint32& OutWord3) const
 	{
 		/**
 		 * Format for QueryData : 
-		 *		word0 (actor ID)
+		 *		word0 (object ID)
 		 *		word1 (blocking channels)
 		 *		word2 (touching channels)
 		 *		word3 (MyChannel (top 8) as ECollisionChannel + Flags (lower 24))
 		 */
-		OutWord0 = ActorID;
+		OutWord0 = ComponentID;
 		OutWord1 = BlockingBits;
 		OutWord2 = TouchingBits;
 		OutWord3 = Word3;
@@ -80,7 +80,7 @@ private:
 #if WITH_PHYSX
 inline void CreateShapeFilterData(
 	const uint8 MyChannel,
-	const int32 ActorID,
+	const int32 ComponentID,
 	const FCollisionResponseContainer& ResponseToChannels,
 	uint32 SkelMeshCompID,
 	uint16 BodyIndex,
@@ -97,7 +97,7 @@ inline void CreateShapeFilterData(
 
 	OutQueryData.setToDefault();
 	OutSimData.setToDefault();
-	Builder.GetQueryData(ActorID, OutQueryData.word0, OutQueryData.word1, OutQueryData.word2, OutQueryData.word3);
+	Builder.GetQueryData(ComponentID, OutQueryData.word0, OutQueryData.word1, OutQueryData.word2, OutQueryData.word3);
 	Builder.GetSimData(BodyIndex, SkelMeshCompID, OutSimData.word0, OutSimData.word1, OutSimData.word2, OutSimData.word3);
 }
 #endif //WITH_PHYSX
diff --git a/Engine/Source/Runtime/Engine/Private/DestructibleComponent.cpp b/Engine/Source/Runtime/Engine/Private/DestructibleComponent.cpp
--- a/Engine/Source/Runtime/Engine/Private/DestructibleComponent.cpp
+++ b/Engine/Source/Runtime/Engine/Private/DestructibleComponent.cpp
@@ -252,7 +252,7 @@ void UDestructibleComponent::CreatePhysicsState()
 
 	// Passing AssetInstanceID = 0 so we'll have self-collision
 	AActor* Owner = GetOwner();
-	CreateShapeFilterData(MoveChannel, (Owner ? Owner->GetUniqueID() : 0), CollResponse, 0, 0, PQueryFilterData, PSimFilterData, BodyInstance.bUseCCD, TheDestructibleMesh->DefaultDestructibleParameters.DamageParameters.bEnableImpactDamage, false);
+	CreateShapeFilterData(MoveChannel, GetUniqueID(), CollResponse, 0, 0, PQueryFilterData, PSimFilterData, BodyInstance.bUseCCD, TheDestructibleMesh->DefaultDestructibleParameters.DamageParameters.bEnableImpactDamage, false);
 
 	// Build filterData variations for complex and simple
 	PSimFilterData.word3 |= EPDF_SimpleCollision | EPDF_ComplexCollision;
@@ -1261,7 +1261,7 @@ void UDestructibleComponent::SetCollisionResponseForActor(PxRigidDynamic* Actor,
 		bool bLargeChunk = IsChunkLarge(ChunkIdx);
 		const FCollisionResponse& ColResponse = bLargeChunk ? LargeChunkCollisionResponse : SmallChunkCollisionResponse;
 		//TODO: we currently assume chunks will not have impact damage as it's very expensive. Should look into exposing this a bit more
-		CreateShapeFilterData(MoveChannel, (Owner ? Owner->GetUniqueID() : 0), ColResponse.GetResponseContainer(), 0, ChunkIdxToBoneIdx(ChunkIdx), PQueryFilterData, PSimFilterData, BodyInstance.bUseCCD, false, false);
+		CreateShapeFilterData(MoveChannel, GetUniqueID(), ColResponse.GetResponseContainer(), 0, ChunkIdxToBoneIdx(ChunkIdx), PQueryFilterData, PSimFilterData, BodyInstance.bUseCCD, false, false);
 		
 		PQueryFilterData.word3 |= EPDF_SimpleCollision | EPDF_ComplexCollision;
 
diff --git a/Engine/Source/Runtime/Engine/Private/PhysicsEngine/BodyInstance.cpp b/Engine/Source/Runtime/Engine/Private/PhysicsEngine/BodyInstance.cpp
--- a/Engine/Source/Runtime/Engine/Private/PhysicsEngine/BodyInstance.cpp
+++ b/Engine/Source/Runtime/Engine/Private/PhysicsEngine/BodyInstance.cpp
@@ -655,8 +655,8 @@ void FBodyInstance::UpdatePhysicsShapeFilterData(uint32 SkelMeshCompID, bool bUs
 			bool bUseNotify = bNotifyOverride ? *bNotifyOverride : BI->bNotifyRigidBodyCollision;
 
 
-			AActor* Owner = BI->OwnerComponent.IsValid() ? BI->OwnerComponent.Get()->GetOwner() : NULL;
-			int32 OwnerID = (Owner != NULL) ? Owner->GetUniqueID() : 0;
+			UPrimitiveComponent* OwnerComponent = BI->OwnerComponent.IsValid() ? BI->OwnerComponent.Get() : nullptr;
+			int32 CompID = (OwnerComponent != nullptr) ? OwnerComponent->GetUniqueID() : 0;
 
 			// Create the filterdata structs
 			PxFilterData PSimFilterData;
@@ -664,7 +664,7 @@ void FBodyInstance::UpdatePhysicsShapeFilterData(uint32 SkelMeshCompID, bool bUs
 			PxFilterData PComplexQueryData;
 			if (UseCollisionEnabled != ECollisionEnabled::NoCollision)
 			{
-				CreateShapeFilterData(BI->ObjectType, OwnerID, UseResponse, SkelMeshCompID, BI->InstanceBodyIndex, PSimpleQueryData, PSimFilterData, bUseCCD && !bPhysicsStatic, bUseNotify, bPhysicsStatic);	//CCD is determined by root body in case of welding
+				CreateShapeFilterData(BI->ObjectType, CompID, UseResponse, SkelMeshCompID, BI->InstanceBodyIndex, PSimpleQueryData, PSimFilterData, bUseCCD && !bPhysicsStatic, bUseNotify, bPhysicsStatic);	//CCD is determined by root body in case of welding
 				PComplexQueryData = PSimpleQueryData;
 
 				// Build filterdata variations for complex and simple
diff --git a/Engine/Source/Runtime/Engine/Private/Vehicles/WheeledVehicleMovementComponent.cpp b/Engine/Source/Runtime/Engine/Private/Vehicles/WheeledVehicleMovementComponent.cpp
--- a/Engine/Source/Runtime/Engine/Private/Vehicles/WheeledVehicleMovementComponent.cpp
+++ b/Engine/Source/Runtime/Engine/Private/Vehicles/WheeledVehicleMovementComponent.cpp
@@ -244,7 +244,7 @@ void UWheeledVehicleMovementComponent::SetupVehicleShapes()
 		CollisionResponse.SetAllChannels( ECR_Ignore );
 
 		PxFilterData PWheelQueryFilterData, PDummySimData;
-		CreateShapeFilterData( ECC_Vehicle, GetOwner()->GetUniqueID(), CollisionResponse, 0, 0, PWheelQueryFilterData, PDummySimData, false, false, false );
+		CreateShapeFilterData(ECC_Vehicle, UpdatedComponent->GetUniqueID(), CollisionResponse, 0, 0, PWheelQueryFilterData, PDummySimData, false, false, false );
 
 		//// Give suspension raycasts the same group ID as the chassis so that they don't hit each other
 		PWheelShape->setQueryFilterData( PWheelQueryFilterData );
diff --git a/Engine/Source/Runtime/Engine/Public/CollisionQueryParams.h b/Engine/Source/Runtime/Engine/Public/CollisionQueryParams.h
--- a/Engine/Source/Runtime/Engine/Public/CollisionQueryParams.h
+++ b/Engine/Source/Runtime/Engine/Public/CollisionQueryParams.h
@@ -35,8 +35,8 @@ struct ENGINE_API FCollisionQueryParams
 	/** Only fill in the PhysMaterial field of  */
 	bool bReturnPhysicalMaterial;
 
-	/** Set of actors to ignore during the trace */
-	TArray<uint32, TInlineAllocator<1> > IgnoreActors;
+	/** Set of components to ignore during the trace */
+	TArray<uint32, TInlineAllocator<1> > IgnoreComponents;
 
 	// Constructors
 	FCollisionQueryParams(bool bInTraceComplex=false)
@@ -62,6 +62,15 @@ struct ENGINE_API FCollisionQueryParams
 	/** Variant that uses an array of TWeakObjectPtrs */
 	void AddIgnoredActors(const TArray<TWeakObjectPtr<AActor> >& InIgnoreActors);
 
+	/** Add a component for this trace to ignore */
+	void AddIgnoredComponent(const UPrimitiveComponent* InIgnoreComponent);
+
+	/** Add a collection of componnets for this trace to ignore */
+	void AddIgnoredComponents(const TArray<UPrimitiveComponent*>& InIgnoreComponents);
+	
+	/** Varient that uses an array of TWeakObjectPtrs */
+	void AddIgnoredComponents(const TArray<TWeakObjectPtr<UPrimitiveComponent>>& InIgnoreComponents);
+
 	FString ToString() const
 	{
 		return FString::Printf(TEXT("[%s:%s] TraceAsync(%d), TraceComplex(%d)"), *OwnerTag.ToString(), *TraceTag.ToString(), bTraceAsyncScene, bTraceComplex );
diff --git a/Engine/Source/Runtime/Landscape/Private/LandscapeCollision.cpp b/Engine/Source/Runtime/Landscape/Private/LandscapeCollision.cpp
--- a/Engine/Source/Runtime/Landscape/Private/LandscapeCollision.cpp
+++ b/Engine/Source/Runtime/Landscape/Private/LandscapeCollision.cpp
@@ -162,7 +162,7 @@ void ULandscapeHeightfieldCollisionComponent::CreatePhysicsState()
 
 				// Setup filtering
 				PxFilterData PQueryFilterData, PSimFilterData;
-				CreateShapeFilterData(GetCollisionObjectType(), GetOwner()->GetUniqueID(), GetCollisionResponseToChannels(), 0, 0, PQueryFilterData, PSimFilterData, true, false, true);
+				CreateShapeFilterData(GetCollisionObjectType(), GetUniqueID(), GetCollisionResponseToChannels(), 0, 0, PQueryFilterData, PSimFilterData, true, false, true);
 
 				// Heightfield is used for simple and complex collision
 				PQueryFilterData.word3 |= (EPDF_SimpleCollision | EPDF_ComplexCollision);
@@ -188,7 +188,7 @@ void ULandscapeHeightfieldCollisionComponent::CreatePhysicsState()
 						CollisionResponse.SetAllChannels(ECollisionResponse::ECR_Ignore);
 						CollisionResponse.SetResponse(ECollisionChannel::ECC_Visibility, ECR_Block);
 						PxFilterData PQueryFilterDataEd, PSimFilterDataEd;
-						CreateShapeFilterData(ECollisionChannel::ECC_Visibility, GetOwner()->GetUniqueID(), CollisionResponse, 0, 0, PQueryFilterDataEd, PSimFilterDataEd, true, false, true);
+						CreateShapeFilterData(ECollisionChannel::ECC_Visibility, GetUniqueID(), CollisionResponse, 0, 0, PQueryFilterDataEd, PSimFilterDataEd, true, false, true);
 
 						PQueryFilterDataEd.word3 |= (EPDF_SimpleCollision | EPDF_ComplexCollision);
 						HeightFieldEdShapeSync->setQueryFilterData(PQueryFilterDataEd);
@@ -726,7 +726,7 @@ void ULandscapeMeshCollisionComponent::CreatePhysicsState()
 
 				// Setup filtering
 				PxFilterData PQueryFilterData, PSimFilterData;
-				CreateShapeFilterData(GetCollisionObjectType(), GetOwner()->GetUniqueID(), GetCollisionResponseToChannels(), 0, 0, PQueryFilterData, PSimFilterData, false, false, true);
+				CreateShapeFilterData(GetCollisionObjectType(), GetUniqueID(), GetCollisionResponseToChannels(), 0, 0, PQueryFilterData, PSimFilterData, false, false, true);
 
 				// Heightfield is used for simple and complex collision
 				PQueryFilterData.word3 |= (EPDF_SimpleCollision | EPDF_ComplexCollision);
@@ -774,7 +774,7 @@ void ULandscapeMeshCollisionComponent::CreatePhysicsState()
 						CollisionResponse.SetAllChannels(ECollisionResponse::ECR_Ignore);
 						CollisionResponse.SetResponse(ECollisionChannel::ECC_Visibility, ECR_Block);
 						PxFilterData PQueryFilterDataEd, PSimFilterDataEd;
-						CreateShapeFilterData(ECollisionChannel::ECC_Visibility, GetOwner()->GetUniqueID(), CollisionResponse, 0, 0, PQueryFilterDataEd, PSimFilterDataEd, true, false, true);
+						CreateShapeFilterData(ECollisionChannel::ECC_Visibility, GetUniqueID(), CollisionResponse, 0, 0, PQueryFilterDataEd, PSimFilterDataEd, true, false, true);
 
 						PQueryFilterDataEd.word3 |= (EPDF_SimpleCollision | EPDF_ComplexCollision);
 						MeshShapeEdSync->setQueryFilterData(PQueryFilterDataEd);
