Added more information to log output.  Help track down material recompilation requests.

diff --git a/Engine/Source/Runtime/Engine/Private/Materials/MaterialShared.cpp b/Engine/Source/Runtime/Engine/Private/Materials/MaterialShared.cpp
--- a/Engine/Source/Runtime/Engine/Private/Materials/MaterialShared.cpp
+++ b/Engine/Source/Runtime/Engine/Private/Materials/MaterialShared.cpp
@@ -1350,7 +1350,7 @@ bool FMaterial::CacheShaders(const FMaterialShaderMapId& ShaderMapId, EShaderPla
 			{
 				ShaderMapCondition = TEXT("Missing");
 			}
-			UE_LOG(LogMaterial, Log, TEXT("%s cached shader map for material %s, compiling."),ShaderMapCondition,*GetFriendlyName());
+			UE_LOG(LogMaterial, Log, TEXT("%s cached shader map for material %s, compiling. %s"),ShaderMapCondition,*GetFriendlyName(), IsSpecialEngineMaterial() ? TEXT("Is special engine material.") : TEXT("") );
 
 			// If there's no cached shader map for this material, compile a new one.
 			// This is just kicking off the async compile, GameThreadShaderMap will not be complete yet
