OverlapMulti automatically ignores self component. (UE-3947) codereview (Lina.Halper)

diff --git a/Engine/Source/Runtime/Engine/Private/PrimitiveComponent.cpp b/Engine/Source/Runtime/Engine/Private/PrimitiveComponent.cpp
--- a/Engine/Source/Runtime/Engine/Private/PrimitiveComponent.cpp
+++ b/Engine/Source/Runtime/Engine/Private/PrimitiveComponent.cpp
@@ -2298,8 +2298,10 @@ void UPrimitiveComponent::UpdateOverlaps(TArray<FOverlapInfo> const* PendingOver
 
 bool UPrimitiveComponent::ComponentOverlapMulti(TArray<struct FOverlapResult>& OutOverlaps, const UWorld* World, const FVector& Pos, const FRotator& Rot, ECollisionChannel TestChannel, const struct FComponentQueryParams& Params, const struct FCollisionObjectQueryParams& ObjectQueryParams) const
 {
+	FComponentQueryParams ParamsWithSelf = Params;
+	ParamsWithSelf.AddIgnoredComponent(this);
 	OutOverlaps.Reset();
-	return BodyInstance.OverlapMulti(OutOverlaps, World, /*pWorldToComponent=*/ nullptr, Pos, Rot, TestChannel, Params, FCollisionResponseParams(GetCollisionResponseToChannels()), ObjectQueryParams);
+	return BodyInstance.OverlapMulti(OutOverlaps, World, /*pWorldToComponent=*/ nullptr, Pos, Rot, TestChannel, ParamsWithSelf, FCollisionResponseParams(GetCollisionResponseToChannels()), ObjectQueryParams);
 }
 
 void UPrimitiveComponent::UpdatePhysicsVolume( bool bTriggerNotifiers )
diff --git a/Engine/Source/Runtime/Engine/Private/SkeletalMeshComponentPhysics.cpp b/Engine/Source/Runtime/Engine/Private/SkeletalMeshComponentPhysics.cpp
--- a/Engine/Source/Runtime/Engine/Private/SkeletalMeshComponentPhysics.cpp
+++ b/Engine/Source/Runtime/Engine/Private/SkeletalMeshComponentPhysics.cpp
@@ -1876,11 +1876,14 @@ bool USkeletalMeshComponent::ComponentOverlapMulti(TArray<struct FOverlapResult>
 	const FTransform WorldToComponent(ComponentToWorld.Inverse());
 	const FCollisionResponseParams ResponseParams(GetCollisionResponseToChannels());
 
+	FComponentQueryParams ParamsWithSelf = Params;
+	ParamsWithSelf.AddIgnoredComponent(this);
+
 	bool bHaveBlockingHit = false;
 	for (const FBodyInstance* Body : Bodies)
 	{
 		checkSlow(Body);
-		if (Body->OverlapMulti(/*inout*/ OutOverlaps, World, &WorldToComponent, Pos, Rot, TestChannel, Params, ResponseParams, ObjectQueryParams))
+		if (Body->OverlapMulti(OutOverlaps, World, &WorldToComponent, Pos, Rot, TestChannel, ParamsWithSelf, ResponseParams, ObjectQueryParams))
 		{
 			bHaveBlockingHit = true;
 		}
