#ue4 - Fix issue with SetJumpAllowed() not affecting the result of CanJump(). CanJump() was erroneously checking CanEverJump() instead of the current state (IsJumpAllowed()).

- Also added comments to CanJump() and CanJumpInternal() explaining how to override the behavior (since BlueprintNativeEvents auto-generate a virtual function).

diff --git a/Engine/Source/Runtime/Engine/Classes/GameFramework/Character.h b/Engine/Source/Runtime/Engine/Classes/GameFramework/Character.h
--- a/Engine/Source/Runtime/Engine/Classes/GameFramework/Character.h
+++ b/Engine/Source/Runtime/Engine/Classes/GameFramework/Character.h
@@ -397,7 +397,8 @@ public:
 	/**
 	 * Check if the character can jump in the current state.
 	 *
-	 * The default implementation may be overridden or extended by implementing the custom CanJump event.
+	 * The default implementation may be overridden or extended by implementing the custom CanJump event in Blueprints.
+	 * C++ code can override CanJumpInternal_Implementation() for custom behavior.
 	 * 
 	 * @Return Whether the character can jump in the current state. 
 	 */
@@ -414,6 +415,8 @@ protected:
 	 * As well as returning true when on the ground, it also returns true when GetMaxJumpTime is more
 	 * than zero and IsJumping returns true.
 	 * 
+	 * (This generates a virtual CanJumpInternal_Implementation function).
+	 *
 	 * @Return Whether the character can jump in the current state. 
 	 */
 
diff --git a/Engine/Source/Runtime/Engine/Classes/GameFramework/NavMovementComponent.h b/Engine/Source/Runtime/Engine/Classes/GameFramework/NavMovementComponent.h
--- a/Engine/Source/Runtime/Engine/Classes/GameFramework/NavMovementComponent.h
+++ b/Engine/Source/Runtime/Engine/Classes/GameFramework/NavMovementComponent.h
@@ -100,7 +100,7 @@ public:
 	FORCEINLINE bool CanEverFly() const { return NavAgentProps.bCanFly; }
 
 	/** @return true if component is allowed to jump */
-	FORCEINLINE bool IsJumpAllowed() const { return NavAgentProps.bCanJump && MovementState.bCanJump; }
+	FORCEINLINE bool IsJumpAllowed() const { return CanEverJump() && MovementState.bCanJump; }
 
 	/** @param bAllowed true if component is allowed to jump */
 	FORCEINLINE void SetJumpAllowed(bool bAllowed) { MovementState.bCanJump = bAllowed; }
diff --git a/Engine/Source/Runtime/Engine/Private/Character.cpp b/Engine/Source/Runtime/Engine/Private/Character.cpp
--- a/Engine/Source/Runtime/Engine/Private/Character.cpp
+++ b/Engine/Source/Runtime/Engine/Private/Character.cpp
@@ -219,7 +219,7 @@ bool ACharacter::CanJumpInternal_Implementation() const
 {
 	const bool bCanHoldToJumpHigher = (GetJumpMaxHoldTime() > 0.0f) && IsJumpProvidingForce();
 
-	return !bIsCrouched && CharacterMovement && (CharacterMovement->IsMovingOnGround() || bCanHoldToJumpHigher) && CharacterMovement->CanEverJump() && !CharacterMovement->bWantsToCrouch;
+	return !bIsCrouched && CharacterMovement && (CharacterMovement->IsMovingOnGround() || bCanHoldToJumpHigher) && CharacterMovement->IsJumpAllowed() && !CharacterMovement->bWantsToCrouch;
 }
 
 void ACharacter::OnJumped_Implementation()
