Fixed issue with not showing collision profile correctly when checking out first time due to reinitializing from config

diff --git a/Engine/Source/Runtime/Engine/Classes/Engine/CollisionProfile.h b/Engine/Source/Runtime/Engine/Classes/Engine/CollisionProfile.h
--- a/Engine/Source/Runtime/Engine/Classes/Engine/CollisionProfile.h
+++ b/Engine/Source/Runtime/Engine/Classes/Engine/CollisionProfile.h
@@ -204,6 +204,9 @@ public:
 	/* custom collision profile name that you can modify what you'd like */
 	ENGINE_API static FName CustomCollisionProfileName;
 
+#if WITH_EDITOR
+	virtual void PostEditChangeProperty(FPropertyChangedEvent& PropertyChangedEvent) override;
+#endif //WITH_EDITOR
 private:
 
 	/** 
diff --git a/Engine/Source/Runtime/Engine/Private/Collision/CollisionProfile.cpp b/Engine/Source/Runtime/Engine/Private/Collision/CollisionProfile.cpp
--- a/Engine/Source/Runtime/Engine/Private/Collision/CollisionProfile.cpp
+++ b/Engine/Source/Runtime/Engine/Private/Collision/CollisionProfile.cpp
@@ -704,4 +704,13 @@ ETraceTypeQuery UCollisionProfile::ConvertToTraceType(ECollisionChannel Collisio
 	}
 
 	return TraceTypeQuery_MAX;
-}
\ No newline at end of file
+}
+
+#if WITH_EDITOR
+void UCollisionProfile::PostEditChangeProperty(FPropertyChangedEvent& PropertyChangedEvent)
+{
+	LoadProfileConfig(false);
+
+	Super::PostEditChangeProperty(PropertyChangedEvent);
+}
+#endif
\ No newline at end of file
