Fixing possible crash.

Expression was evaluated as (NavData && NavData->bRebuildAtRuntime) || (GIsEditor && !bIsInGame), so expression could still be true even if NavData were null.

diff --git a/Engine/Source/Runtime/Engine/Private/AI/Navigation/NavigationSystem.cpp b/Engine/Source/Runtime/Engine/Private/AI/Navigation/NavigationSystem.cpp
--- a/Engine/Source/Runtime/Engine/Private/AI/Navigation/NavigationSystem.cpp
+++ b/Engine/Source/Runtime/Engine/Private/AI/Navigation/NavigationSystem.cpp
@@ -2781,7 +2781,7 @@ void UNavigationSystem::RebuildAll()
 	{
 		ANavigationData* NavData = NavDataSet[NavDataIndex];
 
-		if (NavData && NavData->bRebuildAtRuntime || (GIsEditor && !bIsInGame))
+		if (NavData && (NavData->bRebuildAtRuntime || (GIsEditor && !bIsInGame)))
 		{
 			NavData->RebuildAll();
 		}
