Fixing || and && being in the same scope.

Not changing the existing behavior, although unsure whether it matches the intent in some cases.
Rule of thumb: when || and && are in the same scope, they are evaluated as if && was in parentheses
   E.g.
      Pointer && Pointer->Foo() || Pointer->Bar()
    is evaluated as
     (Pointer && Pointer->Foo()) || Pointer->Bar()

(I fixed such cases in changelists before this one as they were certainly unsafe).

diff --git a/Engine/Source/Developer/FriendsAndChat/Private/FriendsAndChatManager.cpp b/Engine/Source/Developer/FriendsAndChat/Private/FriendsAndChatManager.cpp
--- a/Engine/Source/Developer/FriendsAndChat/Private/FriendsAndChatManager.cpp
+++ b/Engine/Source/Developer/FriendsAndChat/Private/FriendsAndChatManager.cpp
@@ -793,7 +793,7 @@ void FFriendsAndChatManager::PreProcessList(const FString& ListName)
 				TSharedPtr< IFriendItem > ExistingFriend = FindUser(Friend->GetUserId().Get());
 				if ( ExistingFriend.IsValid() )
 				{
-					if (Friend->GetInviteStatus() != ExistingFriend->GetOnlineFriend()->GetInviteStatus() || ExistingFriend->IsPendingAccepted() && Friend->GetInviteStatus() == EInviteStatus::Accepted)
+					if (Friend->GetInviteStatus() != ExistingFriend->GetOnlineFriend()->GetInviteStatus() || (ExistingFriend->IsPendingAccepted() && Friend->GetInviteStatus() == EInviteStatus::Accepted))
 					{
 						ExistingFriend->SetOnlineFriend(Friend);
 					}
diff --git a/Engine/Source/Editor/BlueprintGraph/Private/K2ActionMenuBuilder.cpp b/Engine/Source/Editor/BlueprintGraph/Private/K2ActionMenuBuilder.cpp
--- a/Engine/Source/Editor/BlueprintGraph/Private/K2ActionMenuBuilder.cpp
+++ b/Engine/Source/Editor/BlueprintGraph/Private/K2ActionMenuBuilder.cpp
@@ -2482,7 +2482,7 @@ void FK2ActionMenuBuilder::GetMacroTools(FBlueprintPaletteListBuilder& ActionMen
 	{
 		UBlueprint* MacroBP = *BlueprintIt;
 		if ((ActionMenuBuilder.Blueprint == MacroBP) || //add local Macros
-			(MacroBP->BlueprintType == BPTYPE_MacroLibrary) && (ActionMenuBuilder.Blueprint->ParentClass->IsChildOf(MacroBP->ParentClass)))
+			((MacroBP->BlueprintType == BPTYPE_MacroLibrary) && (ActionMenuBuilder.Blueprint->ParentClass->IsChildOf(MacroBP->ParentClass))))
 		{
 			for (TArray<UEdGraph*>::TIterator GraphIt(MacroBP->MacroGraphs); GraphIt; ++GraphIt)
 			{
@@ -2521,7 +2521,7 @@ void FK2ActionMenuBuilder::GetPinAllowedMacros(FBlueprintGraphActionListBuilder&
 	{
 		UBlueprint* MacroBP = *BlueprintIt;
 		if ((MacroBP == ContextMenuBuilder.Blueprint) || //let local macros be added
-			(MacroBP->BlueprintType == BPTYPE_MacroLibrary) && (ContextMenuBuilder.Blueprint->ParentClass->IsChildOf(MacroBP->ParentClass)))
+			((MacroBP->BlueprintType == BPTYPE_MacroLibrary) && (ContextMenuBuilder.Blueprint->ParentClass->IsChildOf(MacroBP->ParentClass))))
 		{
 			// getting 'top-level' of the macros
 			for (TArray<UEdGraph*>::TIterator GraphIt(MacroBP->MacroGraphs); GraphIt; ++GraphIt)
diff --git a/Engine/Source/Editor/Kismet/Private/BlueprintDetailsCustomization.cpp b/Engine/Source/Editor/Kismet/Private/BlueprintDetailsCustomization.cpp
--- a/Engine/Source/Editor/Kismet/Private/BlueprintDetailsCustomization.cpp
+++ b/Engine/Source/Editor/Kismet/Private/BlueprintDetailsCustomization.cpp
@@ -1603,7 +1603,7 @@ EVisibility FBlueprintVarActionDetails::IsTooltipEditVisible() const
 	UProperty* VariableProperty = SelectionAsProperty();
 	if (VariableProperty)
 	{
-		if (IsABlueprintVariable(VariableProperty) && !IsAComponentVariable(VariableProperty) || IsALocalVariable(VariableProperty))
+		if ((IsABlueprintVariable(VariableProperty) && !IsAComponentVariable(VariableProperty)) || IsALocalVariable(VariableProperty))
 		{
 			return EVisibility::Visible;
 		}
diff --git a/Engine/Source/Editor/PlacementMode/Private/PlacementMode.cpp b/Engine/Source/Editor/PlacementMode/Private/PlacementMode.cpp
--- a/Engine/Source/Editor/PlacementMode/Private/PlacementMode.cpp
+++ b/Engine/Source/Editor/PlacementMode/Private/PlacementMode.cpp
@@ -146,7 +146,7 @@ bool FPlacementMode::MouseLeave( FEditorViewportClient* ViewportClient, FViewpor
 
 bool FPlacementMode::AllowPreviewActors( FEditorViewportClient* ViewportClient ) const
 {
-	return IsCurrentlyPlacing() && ( ViewportClient->IsTracking() && AllowPreviewActorsWhileTracking ) || !ViewportClient->IsTracking();
+	return (IsCurrentlyPlacing() && ( ViewportClient->IsTracking() && AllowPreviewActorsWhileTracking )) || !ViewportClient->IsTracking();
 }
 
 bool FPlacementMode::MouseMove( FEditorViewportClient* ViewportClient, FViewport* Viewport, int32 x, int32 y )
diff --git a/Engine/Source/Editor/PropertyEditor/Private/SDetailsViewBase.cpp b/Engine/Source/Editor/PropertyEditor/Private/SDetailsViewBase.cpp
--- a/Engine/Source/Editor/PropertyEditor/Private/SDetailsViewBase.cpp
+++ b/Engine/Source/Editor/PropertyEditor/Private/SDetailsViewBase.cpp
@@ -964,7 +964,7 @@ static bool IsVisibleStandaloneProperty(const FPropertyNode& PropertyNode, const
 			// Do not add this child node to the current map if its a single object property in a category (serves no purpose for UI)
 			bIsVisibleStandalone = !ParentArrayProperty && (PropertyNode.GetNumChildNodes() == 0 || PropertyNode.GetNumChildNodes() > 1);
 		}
-		else if (Property->IsA(UArrayProperty::StaticClass()) || Property->ArrayDim > 1 && PropertyNode.GetArrayIndex() == INDEX_NONE)
+		else if (Property->IsA(UArrayProperty::StaticClass()) || (Property->ArrayDim > 1 && PropertyNode.GetArrayIndex() == INDEX_NONE))
 		{
 			// Base array properties are always visible
 			bIsVisibleStandalone = true;
diff --git a/Engine/Source/Editor/PropertyEditor/Private/SSingleProperty.cpp b/Engine/Source/Editor/PropertyEditor/Private/SSingleProperty.cpp
--- a/Engine/Source/Editor/PropertyEditor/Private/SSingleProperty.cpp
+++ b/Engine/Source/Editor/PropertyEditor/Private/SSingleProperty.cpp
@@ -131,7 +131,7 @@ void SSingleProperty::SetObject( UObject* InObject )
 
 		bIsAcceptableProperty = true;
 		// not an array property (dynamic or static)
-		bIsAcceptableProperty &= !( Property->IsA( UArrayProperty::StaticClass() ) || Property->ArrayDim > 1 && ValueNode->GetArrayIndex() == INDEX_NONE );
+		bIsAcceptableProperty &= !( Property->IsA( UArrayProperty::StaticClass() ) || (Property->ArrayDim > 1 && ValueNode->GetArrayIndex() == INDEX_NONE) );
 		// not a struct property unless its a built in type like a vector
 		bIsAcceptableProperty &= ( !Property->IsA( UStructProperty::StaticClass() ) || PropertyEditorHelpers::IsBuiltInStructProperty( Property ) );
 	}
diff --git a/Engine/Source/Editor/UnrealEd/Private/MouseDeltaTracker.cpp b/Engine/Source/Editor/UnrealEd/Private/MouseDeltaTracker.cpp
--- a/Engine/Source/Editor/UnrealEd/Private/MouseDeltaTracker.cpp
+++ b/Engine/Source/Editor/UnrealEd/Private/MouseDeltaTracker.cpp
@@ -51,8 +51,8 @@ void FMouseDeltaTracker::DetermineCurrentAxis(FEditorViewportClient* InViewportC
 
 	const bool bIsRotateObjectMode = InViewportClient->IsOrtho() && ControlDown && RightMouseButtonDown;
 	// Ctrl + LEFT/RIGHT mouse button acts the same as dragging the most appropriate widget handle.
-	if( (!InViewportClient->ShouldOrbitCamera() && bIsRotateObjectMode ) || (!bIsRotateObjectMode && ControlDown && !AltDown ) &&
-		(LeftMouseButtonDown || RightMouseButtonDown) )
+	if( (!InViewportClient->ShouldOrbitCamera() && bIsRotateObjectMode ) || ((!bIsRotateObjectMode && ControlDown && !AltDown ) &&
+		(LeftMouseButtonDown || RightMouseButtonDown)) )
 	{
 		// Only try to pick an axis if we're not dragging by widget handle.
 		if ( InViewportClient->GetCurrentWidgetAxis() == EAxisList::None )
diff --git a/Engine/Source/Runtime/Engine/Private/DirectionalLightComponent.cpp b/Engine/Source/Runtime/Engine/Private/DirectionalLightComponent.cpp
--- a/Engine/Source/Runtime/Engine/Private/DirectionalLightComponent.cpp
+++ b/Engine/Source/Runtime/Engine/Private/DirectionalLightComponent.cpp
@@ -699,8 +699,8 @@ bool UDirectionalLightComponent::CanEditChange(const UProperty* InProperty) cons
 		{
 			return CastShadows 
 				&& CastDynamicShadows 
-				&& (DynamicShadowDistanceMovableLight > 0 && Mobility == EComponentMobility::Movable
-					|| DynamicShadowDistanceStationaryLight > 0 && Mobility == EComponentMobility::Stationary);
+				&& ((DynamicShadowDistanceMovableLight > 0 && Mobility == EComponentMobility::Movable)
+					|| (DynamicShadowDistanceStationaryLight > 0 && Mobility == EComponentMobility::Stationary));
 		}
 
 		if (PropertyName == GET_MEMBER_NAME_STRING_CHECKED(UDirectionalLightComponent, DistanceFieldShadowDistance)
diff --git a/Engine/Source/Runtime/Engine/Private/Materials/HLSLMaterialTranslator.h b/Engine/Source/Runtime/Engine/Private/Materials/HLSLMaterialTranslator.h
--- a/Engine/Source/Runtime/Engine/Private/Materials/HLSLMaterialTranslator.h
+++ b/Engine/Source/Runtime/Engine/Private/Materials/HLSLMaterialTranslator.h
@@ -704,7 +704,7 @@ protected:
 
 		checkf(Index >= 0 && Index < CodeChunks[MaterialProperty][ShaderFrequency].Num(), TEXT("Index %d/%d, Platform=%d"), Index, CodeChunks[MaterialProperty][ShaderFrequency].Num(), Platform);
 		const FShaderCodeChunk& CodeChunk = CodeChunks[MaterialProperty][ShaderFrequency][Index];
-		if(CodeChunk.UniformExpression && CodeChunk.UniformExpression->IsConstant() || CodeChunk.bInline)
+		if((CodeChunk.UniformExpression && CodeChunk.UniformExpression->IsConstant()) || CodeChunk.bInline)
 		{
 			// Constant uniform expressions and code chunks which are marked to be inlined are accessed via Definition
 			return CodeChunk.Definition;
@@ -2864,10 +2864,10 @@ protected:
 
 		EMaterialValueType	VectorType = GetParameterType(Vector);
 
-		if(	A && (VectorType & MCT_Float) < MCT_Float4 ||
-			B && (VectorType & MCT_Float) < MCT_Float3 ||
-			G && (VectorType & MCT_Float) < MCT_Float2 ||
-			R && (VectorType & MCT_Float) < MCT_Float1)
+		if(	(A && (VectorType & MCT_Float) < MCT_Float4) ||
+			(B && (VectorType & MCT_Float) < MCT_Float3) ||
+			(G && (VectorType & MCT_Float) < MCT_Float2) ||
+			(R && (VectorType & MCT_Float) < MCT_Float1))
 		{
 			return Errorf(TEXT("Not enough components in (%s: %s) for component mask %u%u%u%u"),*GetParameterCode(Vector),DescribeType(GetParameterType(Vector)),R,G,B,A);
 		}
diff --git a/Engine/Source/Runtime/Engine/Private/Materials/Material.cpp b/Engine/Source/Runtime/Engine/Private/Materials/Material.cpp
--- a/Engine/Source/Runtime/Engine/Private/Materials/Material.cpp
+++ b/Engine/Source/Runtime/Engine/Private/Materials/Material.cpp
@@ -2339,7 +2339,7 @@ void UMaterial::PostEditChangeProperty(FPropertyChangedEvent& PropertyChangedEve
 	}
 
 	// Check if the material is masked and uses a custom opacity (that's not 1.0f).
-	bIsMasked = ((EBlendMode(BlendMode) == BLEND_Masked) && (OpacityMask.Expression || (OpacityMask.UseConstant && OpacityMask.Constant < 0.999f)) || bUseMaterialAttributes);
+	bIsMasked = (((EBlendMode(BlendMode) == BLEND_Masked) && (OpacityMask.Expression || (OpacityMask.UseConstant && OpacityMask.Constant < 0.999f))) || bUseMaterialAttributes);
 
 	bool bRequiresCompilation = true;
 	if( PropertyThatChanged ) 
diff --git a/Engine/Source/Runtime/Engine/Private/ModelLight.cpp b/Engine/Source/Runtime/Engine/Private/ModelLight.cpp
--- a/Engine/Source/Runtime/Engine/Private/ModelLight.cpp
+++ b/Engine/Source/Runtime/Engine/Private/ModelLight.cpp
@@ -525,7 +525,7 @@ void UModel::GroupAllNodes(const ULevel* Level, const TArray<ULightComponentBase
 			ULightComponent* Light = Cast<ULightComponent>(LightBase);
 
 			// Only add enabled lights and lights that can potentially be enabled at runtime (toggleable)
-			if (Light && (Light->bVisible || !Light->HasStaticLighting() && Light->AffectsPrimitive(Level->ModelComponents[ComponentIndex])))
+			if (Light && (Light->bVisible || (!Light->HasStaticLighting() && Light->AffectsPrimitive(Level->ModelComponents[ComponentIndex]))))
 			{
 				RelevantLights->Add(Light);
 			}
diff --git a/Engine/Source/Runtime/Landscape/Private/Landscape.cpp b/Engine/Source/Runtime/Landscape/Private/Landscape.cpp
--- a/Engine/Source/Runtime/Landscape/Private/Landscape.cpp
+++ b/Engine/Source/Runtime/Landscape/Private/Landscape.cpp
@@ -1275,7 +1275,7 @@ void ALandscapeProxy::PostLoad()
 	}
 
 #if WITH_EDITOR
-	if (GetLinker() && (GetLinker()->UE4Ver() < VER_UE4_LANDSCAPE_COMPONENT_LAZY_REFERENCES) || LandscapeComponents.Num() != CollisionComponents.Num())
+	if ((GetLinker() && (GetLinker()->UE4Ver() < VER_UE4_LANDSCAPE_COMPONENT_LAZY_REFERENCES)) || LandscapeComponents.Num() != CollisionComponents.Num())
 	{
 		// Need to clean up invalid collision components
 		RecreateCollisionComponents();
diff --git a/Engine/Source/Runtime/Landscape/Private/LandscapeSplines.cpp b/Engine/Source/Runtime/Landscape/Private/LandscapeSplines.cpp
--- a/Engine/Source/Runtime/Landscape/Private/LandscapeSplines.cpp
+++ b/Engine/Source/Runtime/Landscape/Private/LandscapeSplines.cpp
@@ -910,12 +910,12 @@ void ULandscapeSplineControlPoint::UpdateSplinePoints(bool bUpdateCollision, boo
 {
 	ULandscapeSplinesComponent* OuterSplines = GetOuterULandscapeSplinesComponent();
 
-	if (MeshComponent == NULL && Mesh != NULL ||
-		MeshComponent != NULL &&
+	if ((MeshComponent == NULL && Mesh != NULL) ||
+		(MeshComponent != NULL &&
 			(MeshComponent->StaticMesh != Mesh ||
 			MeshComponent->RelativeLocation != Location ||
 			MeshComponent->RelativeRotation != Rotation ||
-			MeshComponent->RelativeScale3D  != MeshScale)
+			MeshComponent->RelativeScale3D  != MeshScale))
 		)
 	{
 		OuterSplines->Modify();
