Fixed versionning for Visual Logger binary files #ue4
Added histogram data serialization to visual logger files and increased version for file format #ue4
Added getters for categories stored in vlogs to FVisualLoggerHelpers class #ue4

diff --git a/Engine/Source/Runtime/Engine/Private/VisualLogger/VisualLoggerTypes.cpp b/Engine/Source/Runtime/Engine/Private/VisualLogger/VisualLoggerTypes.cpp
--- a/Engine/Source/Runtime/Engine/Private/VisualLogger/VisualLoggerTypes.cpp
+++ b/Engine/Source/Runtime/Engine/Private/VisualLogger/VisualLoggerTypes.cpp
@@ -1,11 +1,13 @@
 // Copyright 1998-2014 Epic Games, Inc. All Rights Reserved.
 
 #include "EnginePrivate.h"
+#include "ObjectBase.h"
 #include "VisualLogger/VisualLogger.h"
 
 #if ENABLE_VISUAL_LOG
 
-#define VISUAL_LOGGER_MAGIC_NUMBER 0xFAFAAFAF
+#define DEPRECATED_VISUAL_LOGGER_MAGIC_NUMBER 0xFAFAAFAF
+#define VISUAL_LOGGER_MAGIC_NUMBER 0xAFAFFAFA
 
 bool FVisualLogStatusCategory::GetDesc(int32 Index, FString& Key, FString& Value) const
 {
@@ -314,6 +316,12 @@ FArchive& operator<<(FArchive& Ar, FVisualLogEntry& LogEntry)
 	Ar << LogEntry.ElementsToDraw;
 	Ar << LogEntry.DataBlocks;
 
+	const int32 VLogsVer = Ar.CustomVer(EVisualLoggerVersion::GUID);
+	if (VLogsVer > EVisualLoggerVersion::Initial)
+	{
+		Ar << LogEntry.HistogramSamples;
+	}
+
 	return Ar;
 }
 
@@ -360,17 +368,29 @@ FArchive& FVisualLoggerHelpers::Serialize(FArchive& Ar, TArray<FVisualLogDevice:
 
 	if (Ar.IsLoading())
 	{
-		TArray<FVisualLogDevice::FVisualLogEntryItem> CurrentFrame;;
-
+		TArray<FVisualLogDevice::FVisualLogEntryItem> CurrentFrame;
 		while (Ar.AtEnd() == false)
 		{
 			int32 FrameTag = VISUAL_LOGGER_MAGIC_NUMBER;
 			Ar << FrameTag;
-			if (FrameTag != VISUAL_LOGGER_MAGIC_NUMBER)
+			if (FrameTag != DEPRECATED_VISUAL_LOGGER_MAGIC_NUMBER && FrameTag != VISUAL_LOGGER_MAGIC_NUMBER)
 			{
 				break;
 			}
 
+			if (FrameTag == DEPRECATED_VISUAL_LOGGER_MAGIC_NUMBER)
+			{
+				Ar.SetCustomVersion(EVisualLoggerVersion::GUID, EVisualLoggerVersion::Initial, TEXT("VisualLogger"));
+			}
+			else
+			{
+				int32 ArchiveVer = -1;
+				Ar << ArchiveVer;
+				check(ArchiveVer >= EVisualLoggerVersion::Initial);
+
+				Ar.SetCustomVersion(EVisualLoggerVersion::GUID, ArchiveVer, TEXT("VisualLogger"));
+			}
+
 			Ar << CurrentFrame;
 			RecordedLogs.Append(CurrentFrame);
 			CurrentFrame.Reset();
@@ -380,11 +400,56 @@ FArchive& FVisualLoggerHelpers::Serialize(FArchive& Ar, TArray<FVisualLogDevice:
 	{
 		int32 FrameTag = VISUAL_LOGGER_MAGIC_NUMBER;
 		Ar << FrameTag;
+
+		int32 ArchiveVer = Ar.CustomVer(EVisualLoggerVersion::GUID);
+		Ar << ArchiveVer;
 		Ar << RecordedLogs;
 	}
 
+	int32 CustomVer = Ar.CustomVer(EVisualLoggerVersion::GUID);
+
 	return Ar;
 }
 
+void FVisualLoggerHelpers::GetCategories(const FVisualLogEntry& EntryItem, TArray<FVisualLoggerCategoryVerbosityPair>& OutCategories)
+{
+	for (const auto& CurrentEvent : EntryItem.Events)
+	{
+		OutCategories.Add(FVisualLoggerCategoryVerbosityPair(*CurrentEvent.Name, CurrentEvent.Verbosity));
+	}
+
+	for (const auto& CurrentLine : EntryItem.LogLines)
+	{
+		OutCategories.Add(FVisualLoggerCategoryVerbosityPair(CurrentLine.Category, CurrentLine.Verbosity));
+	}
+
+	for (const auto& CurrentElement : EntryItem.ElementsToDraw)
+	{
+		OutCategories.Add(FVisualLoggerCategoryVerbosityPair(CurrentElement.Category, CurrentElement.Verbosity));
+	}
+
+	for (const auto& CurrentSample : EntryItem.HistogramSamples)
+	{
+		OutCategories.Add(FVisualLoggerCategoryVerbosityPair(CurrentSample.Category, CurrentSample.Verbosity));
+	}
+
+	for (const auto& CurrentBlock : EntryItem.DataBlocks)
+	{
+		OutCategories.Add(FVisualLoggerCategoryVerbosityPair(CurrentBlock.Category, CurrentBlock.Verbosity));
+	}
+}
+
+void FVisualLoggerHelpers::GetHistogramCategories(const FVisualLogEntry& EntryItem, TMap<FString, TArray<FString> >& OutCategories)
+{
+	for (const auto& CurrentSample : EntryItem.HistogramSamples)
+	{
+		auto& DataNames = OutCategories.FindOrAdd(CurrentSample.GraphName.ToString());
+		if (DataNames.Find(CurrentSample.DataName.ToString()) == INDEX_NONE)
+		{
+			DataNames.Add(CurrentSample.DataName.ToString());
+		}
+	}
+}
+
 #endif //ENABLE_VISUAL_LOG
 
diff --git a/Engine/Source/Runtime/Engine/Public/VisualLogger/VisualLoggerTypes.h b/Engine/Source/Runtime/Engine/Public/VisualLogger/VisualLoggerTypes.h
--- a/Engine/Source/Runtime/Engine/Public/VisualLogger/VisualLoggerTypes.h
+++ b/Engine/Source/Runtime/Engine/Public/VisualLogger/VisualLoggerTypes.h
@@ -26,6 +26,7 @@ namespace EVisualLoggerVersion
 	enum Type
 	{
 		Initial = 0,
+		HistogramGraphsSerialization = 1,
 		// -----<new versions can be added before this line>-------------------------------------------------
 		// - this needs to be the last line (see note below)
 		VersionPlusOne,
@@ -47,6 +48,7 @@ enum class EVisualLoggerShapeElement : uint8
 	Cone,
 	Cylinder,
 	Capsule,
+	Polygon,
 	// note that in order to remain backward compatibility in terms of log
 	// serialization new enum values need to be added at the end
 };
@@ -230,12 +232,22 @@ public:
 	virtual bool HasFlags(int32 InFlags) const { return !!(InFlags & EVisualLoggerDeviceFlags::NoFlags); }
 };
 
+struct ENGINE_API FVisualLoggerCategoryVerbosityPair
+{
+	FVisualLoggerCategoryVerbosityPair(FName Category, ELogVerbosity::Type InVerbosity) : CategoryName(Category), Verbosity(InVerbosity) {}
+
+	FName CategoryName;
+	ELogVerbosity::Type Verbosity;
+};
+
 struct ENGINE_API FVisualLoggerHelpers
 {
 	static FString GenerateTemporaryFilename(const FString& FileExt);
 	static FString GenerateFilename(const FString& TempFileName, const FString& Prefix, float StartRecordingTime, float EndTimeStamp);
 	static FArchive& Serialize(FArchive& Ar, FName& Name);
 	static FArchive& Serialize(FArchive& Ar, TArray<FVisualLogDevice::FVisualLogEntryItem>& RecordedLogs);
+	static void GetCategories(const FVisualLogEntry& RecordedLogs, TArray<FVisualLoggerCategoryVerbosityPair>& OutCategories);
+	static void GetHistogramCategories(const FVisualLogEntry& RecordedLogs, TMap<FString, TArray<FString> >& OutCategories);
 };
 
 ENGINE_API  FArchive& operator<<(FArchive& Ar, FVisualLogDevice::FVisualLogEntryItem& FrameCacheItem);
