Particle UV flipping

Enum in Required module defining the UV flip behavior. None, UV, UOnly, VOnly, RandomUV, RandomU, RandomV.
Only done for sprites and GPU spirtes atm but would be easy to add for other types once requirements were defined.
No ability to flip during particle lifetime but would be trivial to add if desired.

diff --git a/Engine/Shaders/ParticleGPUSpriteVertexFactory.usf b/Engine/Shaders/ParticleGPUSpriteVertexFactory.usf
--- a/Engine/Shaders/ParticleGPUSpriteVertexFactory.usf
+++ b/Engine/Shaders/ParticleGPUSpriteVertexFactory.usf
@@ -228,10 +228,16 @@ FVertexFactoryIntermediates GetVertexFactoryIntermediates( FVertexFactoryInput I
 	float4 MiscCurve = MiscCurveSample * EmitterUniforms.MiscScale + EmitterUniforms.MiscBias;
 
 	// Compute the size of the sprite. Note it is (0,0) if the sprite is dead.
-	float2 InitialSize = AttributeSample.xy;
+	// Reconstruct the real size after we've encoded the UV flip mode as sizes > 0.5;
+	float2 UVFlipOffset = float2( AttributeSample.x < 0.5 ? 0.0f : -0.5, AttributeSample.y < 0.5 ? 0.0f : -0.5);
+	float2 InitialSize = (AttributeSample.xy + UVFlipOffset) * float2(2.0,2.0);
 	float2 SizeScale = MiscCurve.xy * EmitterDynamicUniforms.LocalToWorldScale;
 	float2 SizeScaleBySpeed = clamp(EmitterUniforms.SizeBySpeed.xy * ParticleSpeed,float2(1,1), EmitterUniforms.SizeBySpeed.zw);
 	float2 Size = InitialSize * SizeScale * SizeScaleBySpeed * IsAlive.xx;
+	
+	float2 FlippedUV = Input.TexCoord.xy;
+	FlippedUV.x = UVFlipOffset.x == 0.0 ? 1.0 - FlippedUV.x : FlippedUV.x;
+	FlippedUV.y = UVFlipOffset.y == 0.0 ? 1.0 - FlippedUV.y : FlippedUV.y;
 
 	// SubUV.
 	float SubImageIndex = MiscCurve.z;
@@ -245,8 +251,8 @@ FVertexFactoryIntermediates GetVertexFactoryIntermediates( FVertexFactoryInput I
 
 	// Compute the final texture coordinates for both subimages.
 	float4 TexCoord;
-	TexCoord.xy = (float2( SubImageAH, SubImageAV ) + Input.TexCoord.xy) * EmitterUniforms.SubImageSize.zw;
-	TexCoord.zw = (float2( SubImageBH, SubImageBV ) + Input.TexCoord.xy) * EmitterUniforms.SubImageSize.zw;
+	TexCoord.xy = (float2( SubImageAH, SubImageAV ) + FlippedUV) * EmitterUniforms.SubImageSize.zw;
+	TexCoord.zw = (float2( SubImageBH, SubImageBV ) + FlippedUV) * EmitterUniforms.SubImageSize.zw;
 
 	// Current rotation of the sprite. Map [0,1] to one full rotation.
 	float RotationRate = AttributeSample.w * EmitterUniforms.RotationRateScale;
diff --git a/Engine/Shaders/ParticleSimulationShader.usf b/Engine/Shaders/ParticleSimulationShader.usf
--- a/Engine/Shaders/ParticleSimulationShader.usf
+++ b/Engine/Shaders/ParticleSimulationShader.usf
@@ -377,7 +377,7 @@ void PixelMain(
 	float4 MiscCurve = MiscCurveSample * Simulation.MiscScale + Simulation.MiscBias;
 
 	// Compute the size of the sprite. Note it is (0,0) if the sprite is dead.
-	float2 InitialSize = RenderAttributeSample.xy;
+	float2 InitialSize = abs(RenderAttributeSample.xy);
 	float2 SizeScale = MiscCurve.xy;
 	float2 Size = InitialSize * SizeScale * LocalToWorldScale;
 
diff --git a/Engine/Shaders/ParticleSpriteVertexFactory.usf b/Engine/Shaders/ParticleSpriteVertexFactory.usf
--- a/Engine/Shaders/ParticleSpriteVertexFactory.usf
+++ b/Engine/Shaders/ParticleSpriteVertexFactory.usf
@@ -393,10 +393,16 @@ FVertexFactoryIntermediates GetVertexFactoryIntermediates(FVertexFactoryInput In
 
 	// Vertex position.
 	float4 VertexWorldPosition = float4(ParticleWorldPosition.xyz + View.PreViewTranslation,1);
-	float2 Size = Input.SizeRotSubImage.xy;
+	float2 Size = abs(Input.SizeRotSubImage.xy);
 	VertexWorldPosition += Size.x * (Input.TexCoord.x + SpriteVF.PivotOffset.x) * float4(Right,0);
 	VertexWorldPosition += Size.y * (Input.TexCoord.y + SpriteVF.PivotOffset.y) * float4(Up,0);
 	Intermediates.VertexWorldPosition = VertexWorldPosition.xyz;
+	
+	//We encode the UV flip in the sign of the size data.
+	float2 UVFlip = sign(Input.SizeRotSubImage.xy);
+	float2 FlippedUV = Input.TexCoord.xy;
+	FlippedUV.x = UVFlip.x < 0.0 ? 1.0 - FlippedUV.x : FlippedUV.x;
+	FlippedUV.y = UVFlip.y < 0.0 ? 1.0 - FlippedUV.y : FlippedUV.y;
 
 	// SubUV.
 	float SubImageIndex = Input.SizeRotSubImage.w;
@@ -407,8 +413,8 @@ FVertexFactoryIntermediates GetVertexFactoryIntermediates(FVertexFactoryInput In
 	float SubImageBH = fmod( SubImageB, SpriteVF.SubImageSize.x );
 	float SubImageAV = floor( SubImageA * SpriteVF.SubImageSize.z );
 	float SubImageBV = floor( SubImageB * SpriteVF.SubImageSize.z );
-	Intermediates.TexCoord.xy = (float2( SubImageAH, SubImageAV ) + Input.TexCoord.xy) * SpriteVF.SubImageSize.zw;
-	Intermediates.TexCoord.zw = (float2( SubImageBH, SubImageBV ) + Input.TexCoord.xy) * SpriteVF.SubImageSize.zw;
+	Intermediates.TexCoord.xy = (float2( SubImageAH, SubImageAV ) + FlippedUV) * SpriteVF.SubImageSize.zw ;
+	Intermediates.TexCoord.zw = (float2( SubImageBH, SubImageBV ) + FlippedUV) * SpriteVF.SubImageSize.zw ;
 	Intermediates.SubImageLerp = SubImageLerp;
 
 	Intermediates.Color = Input.Color;
diff --git a/Engine/Source/Runtime/Engine/Classes/Particles/ParticleModuleRequired.h b/Engine/Source/Runtime/Engine/Classes/Particles/ParticleModuleRequired.h
--- a/Engine/Source/Runtime/Engine/Classes/Particles/ParticleModuleRequired.h
+++ b/Engine/Source/Runtime/Engine/Classes/Particles/ParticleModuleRequired.h
@@ -9,6 +9,40 @@
 
 class UInterpCurveEdSetup;
 
+UENUM()
+enum class EParticleUVFlipMode : uint8
+{
+	/** Flips UV on all particles. */
+	None,
+	/** Flips UV on all particles. */
+	FlipUV,
+	/** Flips U only on all particles. */
+	FlipUOnly,
+	/** Flips V only on all particles. */
+	FlipVOnly,
+	/** Flips UV randomly for each particle on spawn. */
+	RandomFlipUV,
+	/** Flips U only randomly for each particle on spawn. */
+	RandomFlipUOnly,
+	/** Flips V only randomly for each particle on spawn. */
+	RandomFlipVOnly,
+};
+
+/** Flips the sign of a particle's base size based on it's UV flip mode. */
+FORCEINLINE void AdjustParticleBaseSizeForUVFlipping(FVector& OutSize, EParticleUVFlipMode FlipMode)
+{
+	static const int32 HalfRandMax = RAND_MAX / 2;
+	switch (FlipMode)
+	{
+	case EParticleUVFlipMode::FlipUV:			OutSize = -OutSize;			return;
+	case EParticleUVFlipMode::FlipUOnly:		OutSize.X = -OutSize.X;		return;
+	case EParticleUVFlipMode::FlipVOnly:		OutSize.Y = -OutSize.Y;		return;
+	case EParticleUVFlipMode::RandomFlipUV:		OutSize = FMath::Rand() > HalfRandMax ? -OutSize : OutSize;			return;
+	case EParticleUVFlipMode::RandomFlipUOnly:	OutSize.X = FMath::Rand() > HalfRandMax ? -OutSize.X : OutSize.X;	return;
+	case EParticleUVFlipMode::RandomFlipVOnly:	OutSize.Y = FMath::Rand() > HalfRandMax ? -OutSize.Y : OutSize.Y;	return;
+	};
+}
+
 UENUM()
 enum EParticleSortMode
 {
@@ -279,6 +313,12 @@ class UParticleModuleRequired : public UParticleModule
 	UPROPERTY(EditAnywhere, Category=Emitter)
 	uint32 bOrbitModuleAffectsVelocityAlignment:1;
 
+	/**
+	* Controls UV Flipping for this emitter.
+	*/
+	UPROPERTY(EditAnywhere, Category = Rendering)
+	EParticleUVFlipMode UVFlippingMode;
+
 	/** 
 	*	Named material overrides for this emitter. 
 	*	Overrides this emitter's material(s) with those in the correspondingly named slot(s) of the owning system.
diff --git a/Engine/Source/Runtime/Engine/Private/Particles/ParticleEmitterInstances.cpp b/Engine/Source/Runtime/Engine/Private/Particles/ParticleEmitterInstances.cpp
--- a/Engine/Source/Runtime/Engine/Private/Particles/ParticleEmitterInstances.cpp
+++ b/Engine/Source/Runtime/Engine/Private/Particles/ParticleEmitterInstances.cpp
@@ -1437,7 +1437,7 @@ void FParticleEmitterInstance::ResetParticleParameters(float DeltaTime)
 	{
 		DECLARE_PARTICLE(Particle, ParticleData + ParticleStride * ParticleIndices[ParticleIndex]);
 		Particle.Velocity		= Particle.BaseVelocity;
-		Particle.Size			= Particle.BaseSize;
+		Particle.Size = GetParticleBaseSize(Particle);
 		Particle.RotationRate	= Particle.BaseRotationRate;
 		Particle.Color			= Particle.BaseColor;
 		Particle.RelativeTime	+= Particle.OneOverMaxLifetime * DeltaTime;
diff --git a/Engine/Source/Runtime/Engine/Private/Particles/ParticleGpuSimulation.cpp b/Engine/Source/Runtime/Engine/Private/Particles/ParticleGpuSimulation.cpp
--- a/Engine/Source/Runtime/Engine/Private/Particles/ParticleGpuSimulation.cpp
+++ b/Engine/Source/Runtime/Engine/Private/Particles/ParticleGpuSimulation.cpp
@@ -3772,8 +3772,14 @@ private:
 			NewParticle->Position = TempParticle->Location + InterpFraction * EmitterDelta + SpawnInfo.StartTime * NewParticle->Velocity + EmitterInfo.OrbitOffsetBase + EmitterInfo.OrbitOffsetRange * RandomOrbit;
 			NewParticle->RelativeTime = TempParticle->RelativeTime;
 			NewParticle->TimeScale = FMath::Max<float>(TempParticle->OneOverMaxLifetime, 0.001f);
-			NewParticle->Size.X = TempParticle->BaseSize.X * EmitterInfo.InvMaxSize.X;
-			NewParticle->Size.Y = bSquare ? (NewParticle->Size.X) : (TempParticle->BaseSize.Y * EmitterInfo.InvMaxSize.Y);
+
+			//So here I'm reducing the size to 0-0.5 range and using < 0.5 to indicate flipped UVs.
+			FVector BaseSize = GetParticleBaseSize(*TempParticle, true);
+			FVector2D UVFlipSizeOffset = FVector2D(BaseSize.X < 0.0f ? 0.0f : 0.5f, BaseSize.Y < 0.0f ? 0.0f : 0.5f);
+			NewParticle->Size.X = (FMath::Abs(BaseSize.X) * EmitterInfo.InvMaxSize.X * 0.5f);
+			NewParticle->Size.Y = bSquare ? (NewParticle->Size.X) : (FMath::Abs(BaseSize.Y) * EmitterInfo.InvMaxSize.Y * 0.5f);
+			NewParticle->Size += UVFlipSizeOffset;
+
 			NewParticle->Rotation = FMath::Fractional( TempParticle->Rotation * OneOverTwoPi );
 			NewParticle->RelativeRotationRate = TempParticle->BaseRotationRate * OneOverTwoPi * EmitterInfo.InvRotationRateScale / NewParticle->TimeScale;
 			NewParticle->RandomOrbit = RandomOrbit;
diff --git a/Engine/Source/Runtime/Engine/Private/Particles/ParticleModules_Size.cpp b/Engine/Source/Runtime/Engine/Private/Particles/ParticleModules_Size.cpp
--- a/Engine/Source/Runtime/Engine/Private/Particles/ParticleModules_Size.cpp
+++ b/Engine/Source/Runtime/Engine/Private/Particles/ParticleModules_Size.cpp
@@ -96,8 +96,10 @@ void UParticleModuleSize::SpawnEx(FParticleEmitterInstance* Owner, int32 Offset,
 {
 	SPAWN_INIT;
 	FVector Size		 = StartSize.GetValue(Owner->EmitterTime, Owner->Component, 0, InRandomStream);
-	Particle.Size		+= Size;
-	Particle.BaseSize	+= Size;
+	Particle.Size	+= Size;
+
+	AdjustParticleBaseSizeForUVFlipping(Size, Owner->CurrentLODLevel->RequiredModule->UVFlippingMode);
+	Particle.BaseSize += Size;
 }
 
 /*-----------------------------------------------------------------------------
@@ -406,14 +408,14 @@ void UParticleModuleSizeScale::Spawn(FParticleEmitterInstance* Owner, int32 Offs
 {
 	SPAWN_INIT;
 	FVector ScaleFactor = SizeScale.GetValue(Particle.RelativeTime, Owner->Component);
-	Particle.Size = Particle.BaseSize * ScaleFactor;
+	Particle.Size = GetParticleBaseSize(*ParticleBase) * ScaleFactor;
 }
 
 void UParticleModuleSizeScale::Update(FParticleEmitterInstance* Owner, int32 Offset, float DeltaTime)
 {
 	BEGIN_UPDATE_LOOP;
 		FVector ScaleFactor = SizeScale.GetValue(Particle.RelativeTime, Owner->Component);
-		Particle.Size = Particle.BaseSize * ScaleFactor;
+		Particle.Size = GetParticleBaseSize(Particle) * ScaleFactor;
 	END_UPDATE_LOOP;
 }
 
@@ -447,7 +449,7 @@ void UParticleModuleSizeScaleBySpeed::Update(FParticleEmitterInstance* Owner, in
 		FVector Size = Scale * Particle.Velocity.Size();
 		Size = Size.ComponentMax(FVector(1.0f));
 		Size = Size.ComponentMin(ScaleMax);
-		Particle.Size = Particle.BaseSize * Size;
+		Particle.Size = GetParticleBaseSize(Particle) * Size;
 	END_UPDATE_LOOP;
 }
 
diff --git a/Engine/Source/Runtime/Engine/Private/Particles/ParticleSystemRender.cpp b/Engine/Source/Runtime/Engine/Private/Particles/ParticleSystemRender.cpp
--- a/Engine/Source/Runtime/Engine/Private/Particles/ParticleSystemRender.cpp
+++ b/Engine/Source/Runtime/Engine/Private/Particles/ParticleSystemRender.cpp
@@ -864,7 +864,7 @@ bool FDynamicSpriteEmitterData::GetVertexAndIndexData(void* VertexData, void* Dy
 		FillVertex->OldPosition	= ParticleOldPosition;
 		// Create a floating point particle ID from the counter, map into approximately 0-1
 		FillVertex->ParticleId = (Particle.Flags & STATE_CounterMask) / 10000.0f;
-		FillVertex->Size		= Size;
+		FillVertex->Size = GetParticleSizeWithUVFlipInSign(Particle, Size);
 		FillVertex->Rotation	= Particle.Rotation;
 		FillVertex->SubImageIndex = SubImageIndex;
 		FillVertex->Color		= Particle.Color;
@@ -996,7 +996,7 @@ bool FDynamicSpriteEmitterData::GetVertexAndIndexDataNonInstanced(void* VertexDa
 			FillVertex[I].OldPosition	= ParticleOldPosition;
 			// Create a floating point particle ID from the counter, map into approximately 0-1
 			FillVertex[I].ParticleId = (Particle.Flags & STATE_CounterMask) / 10000.0f;
-			FillVertex[I].Size		= Size;
+			FillVertex[I].Size = GetParticleSizeWithUVFlipInSign(Particle, Size);
 			FillVertex[I].Rotation	= Particle.Rotation;
 			FillVertex[I].SubImageIndex = SubImageIndex;
 			FillVertex[I].Color		= Particle.Color;
diff --git a/Engine/Source/Runtime/Engine/Public/ParticleHelper.h b/Engine/Source/Runtime/Engine/Public/ParticleHelper.h
--- a/Engine/Source/Runtime/Engine/Public/ParticleHelper.h
+++ b/Engine/Source/Runtime/Engine/Public/ParticleHelper.h
@@ -2789,3 +2789,14 @@ public:
 /** The global null color vertex buffer, which is set with a stride of 0 on meshes without a color component. */
 extern TGlobalResource<FNullDynamicParameterVertexBuffer> GNullDynamicParameterVertexBuffer;
 
+FORCEINLINE FVector GetParticleBaseSize(const FBaseParticle& Particle, bool bKeepFlipScale = false)
+{
+	return bKeepFlipScale ? Particle.BaseSize : FVector(FMath::Abs(Particle.BaseSize.X), FMath::Abs(Particle.BaseSize.Y), FMath::Abs(Particle.BaseSize.Z));
+}
+
+FORCEINLINE FVector2D GetParticleSizeWithUVFlipInSign(const FBaseParticle& Particle, const FVector2D& ScaledSize)
+{
+	return FVector2D(
+		Particle.BaseSize.X >= 0.0f ? ScaledSize.X : -ScaledSize.X,
+		Particle.BaseSize.Y >= 0.0f ? ScaledSize.Y : -ScaledSize.Y);
+}
\ No newline at end of file
