Minor extension to UV flip to allow random independent flipping of U and V.

diff --git a/Engine/Source/Runtime/Engine/Classes/Particles/ParticleModuleRequired.h b/Engine/Source/Runtime/Engine/Classes/Particles/ParticleModuleRequired.h
--- a/Engine/Source/Runtime/Engine/Classes/Particles/ParticleModuleRequired.h
+++ b/Engine/Source/Runtime/Engine/Classes/Particles/ParticleModuleRequired.h
@@ -26,6 +26,8 @@ enum class EParticleUVFlipMode : uint8
 	RandomFlipUOnly,
 	/** Flips V only randomly for each particle on spawn. */
 	RandomFlipVOnly,
+	/** Flips U and V independently at random for each particle on spawn. */
+	RandomFlipUVIndependent,
 };
 
 /** Flips the sign of a particle's base size based on it's UV flip mode. */
@@ -34,12 +36,18 @@ FORCEINLINE void AdjustParticleBaseSizeForUVFlipping(FVector& OutSize, EParticle
 	static const int32 HalfRandMax = RAND_MAX / 2;
 	switch (FlipMode)
 	{
-	case EParticleUVFlipMode::FlipUV:			OutSize = -OutSize;			return;
-	case EParticleUVFlipMode::FlipUOnly:		OutSize.X = -OutSize.X;		return;
-	case EParticleUVFlipMode::FlipVOnly:		OutSize.Y = -OutSize.Y;		return;
-	case EParticleUVFlipMode::RandomFlipUV:		OutSize = FMath::Rand() > HalfRandMax ? -OutSize : OutSize;			return;
-	case EParticleUVFlipMode::RandomFlipUOnly:	OutSize.X = FMath::Rand() > HalfRandMax ? -OutSize.X : OutSize.X;	return;
-	case EParticleUVFlipMode::RandomFlipVOnly:	OutSize.Y = FMath::Rand() > HalfRandMax ? -OutSize.Y : OutSize.Y;	return;
+	case EParticleUVFlipMode::FlipUV:						OutSize = -OutSize;			return;
+	case EParticleUVFlipMode::FlipUOnly:					OutSize.X = -OutSize.X;		return;
+	case EParticleUVFlipMode::FlipVOnly:					OutSize.Y = -OutSize.Y;		return;
+	case EParticleUVFlipMode::RandomFlipUV:					OutSize = FMath::Rand() > HalfRandMax ? -OutSize : OutSize;			return;
+	case EParticleUVFlipMode::RandomFlipUOnly:				OutSize.X = FMath::Rand() > HalfRandMax ? -OutSize.X : OutSize.X;	return;
+	case EParticleUVFlipMode::RandomFlipVOnly:				OutSize.Y = FMath::Rand() > HalfRandMax ? -OutSize.Y : OutSize.Y;	return;
+	case EParticleUVFlipMode::RandomFlipUVIndependent:
+	{
+		OutSize.X = FMath::Rand() > HalfRandMax ? -OutSize.X : OutSize.X;		
+		OutSize.Y = FMath::Rand() > HalfRandMax ? -OutSize.Y : OutSize.Y;
+		return;
+	}
 	};
 }
 
