Initialising UV flip to None.

diff --git a/Engine/Source/Runtime/Engine/Private/Particles/ParticleModules.cpp b/Engine/Source/Runtime/Engine/Private/Particles/ParticleModules.cpp
--- a/Engine/Source/Runtime/Engine/Private/Particles/ParticleModules.cpp
+++ b/Engine/Source/Runtime/Engine/Private/Particles/ParticleModules.cpp
@@ -1083,6 +1083,7 @@ UParticleModuleRequired::UParticleModuleRequired(const FObjectInitializer& Objec
 	NormalsSphereCenter = FVector(0.0f, 0.0f, 100.0f);
 	NormalsCylinderDirection = FVector(0.0f, 0.0f, 1.0f);
 	bUseLegacyEmitterTime = true;
+	UVFlippingMode = EParticleUVFlipMode::None;
 }
 
 void UParticleModuleRequired::InitializeDefaults()
