Some minor fixes to Text pin handling in Blueprints.

diff --git a/Engine/Source/Runtime/Engine/Private/EdGraph/EdGraphPin.cpp b/Engine/Source/Runtime/Engine/Private/EdGraph/EdGraphPin.cpp
--- a/Engine/Source/Runtime/Engine/Private/EdGraph/EdGraphPin.cpp
+++ b/Engine/Source/Runtime/Engine/Private/EdGraph/EdGraphPin.cpp
@@ -145,15 +145,10 @@ void UEdGraphPin::CopyPersistentDataFromOldPin(const UEdGraphPin& SourcePin)
 	// The PinType, Direction, and bNotConnectable are properties generated from the schema
 
 	// Only move the default value if it was modified; inherit the new default value otherwise
-	if (SourcePin.DefaultValue != SourcePin.AutogeneratedDefaultValue || SourcePin.DefaultObject != NULL)
+	if (SourcePin.DefaultValue != SourcePin.AutogeneratedDefaultValue || SourcePin.DefaultObject != NULL || SourcePin.DefaultTextValue.ToString() != SourcePin.AutogeneratedDefaultValue)
 	{
 		DefaultObject = SourcePin.DefaultObject;
 		DefaultValue = SourcePin.DefaultValue;
-	}
-
-	const bool bTextValueWasModified = (SourcePin.DefaultTextValue.ToString() != SourcePin.AutogeneratedDefaultValue);
-	if (bTextValueWasModified)
-	{
 		DefaultTextValue = SourcePin.DefaultTextValue;
 	}
 
@@ -224,6 +219,10 @@ FString UEdGraphPin::GetDefaultAsString() const
 	{
 		return DefaultObject->GetPathName();
 	}
+	else if(!DefaultTextValue.IsEmpty())
+	{
+		return DefaultTextValue.ToString();
+	}
 	else
 	{
 		return DefaultValue;
