Extended functions for debug drawing (API from SceneManagement.h file):
- Added "Thickness" parameter to some functions (only for wired shapes rendering and with 0 as default value)
- Added GetHalfSphereMesh, GetConeMesh and GetCapsuleMesh functions

diff --git a/Engine/Source/Editor/ComponentVisualizers/Private/SensingComponentVisualizer.cpp b/Engine/Source/Editor/ComponentVisualizers/Private/SensingComponentVisualizer.cpp
--- a/Engine/Source/Editor/ComponentVisualizers/Private/SensingComponentVisualizer.cpp
+++ b/Engine/Source/Editor/ComponentVisualizers/Private/SensingComponentVisualizer.cpp
@@ -30,7 +30,7 @@ void FSensingComponentVisualizer::DrawVisualization(const UActorComponent* Compo
 			if (Senses->SightRadius > 0.0f)
 			{
 				TArray<FVector> Verts;
-				DrawWireCone(PDI, Transform, Senses->SightRadius, Senses->GetPeripheralVisionAngle(), 10, FColor::Green, SDPG_World, Verts);
+				DrawWireCone(PDI, Verts, Transform, Senses->SightRadius, Senses->GetPeripheralVisionAngle(), 10, FColor::Green, SDPG_World);
 			}
 		}
 	}
diff --git a/Engine/Source/Editor/Persona/Private/AnimationEditorViewportClient.cpp b/Engine/Source/Editor/Persona/Private/AnimationEditorViewportClient.cpp
--- a/Engine/Source/Editor/Persona/Private/AnimationEditorViewportClient.cpp
+++ b/Engine/Source/Editor/Persona/Private/AnimationEditorViewportClient.cpp
@@ -1860,7 +1860,7 @@ void FAnimationViewportClient::DrawBones(const USkeletalMeshComponent* MeshCompo
 			//Render Sphere for bone end point and a cone between it and its parent.
 			PDI->SetHitProxy( new HPersonaBoneProxy( MeshComponent->SkeletalMesh->RefSkeleton.GetBoneName(BoneIndex)) );
 			DrawWireSphere(PDI, End, LineColor, SphereRadius, 10, SDPG_Foreground);
-			DrawWireCone(PDI, FRotationMatrix::MakeFromX(EndToStart)*FTranslationMatrix(End), ConeLength, Angle, 4, LineColor, SDPG_Foreground, Verts);
+			DrawWireCone(PDI, Verts, FRotationMatrix::MakeFromX(EndToStart)*FTranslationMatrix(End), ConeLength, Angle, 4, LineColor, SDPG_Foreground);
 			PDI->SetHitProxy( NULL );
 		
 			// draw gizmo
diff --git a/Engine/Source/Runtime/Engine/Private/DrawDebugHelpers.cpp b/Engine/Source/Runtime/Engine/Private/DrawDebugHelpers.cpp
--- a/Engine/Source/Runtime/Engine/Private/DrawDebugHelpers.cpp
+++ b/Engine/Source/Runtime/Engine/Private/DrawDebugHelpers.cpp
@@ -701,7 +701,7 @@ void DrawDebugFrustum(const UWorld* InWorld, const FMatrix& FrustumToWorld, FCol
 
 
 
-static void DrawHalfCircle(const UWorld* InWorld, const FVector& Base, const FVector& X, const FVector& Y, const FColor& Color, float Radius, int32 NumSides, bool bPersistentLines, float LifeTime, uint8 DepthPriority = 0)
+static void DrawHalfCircle(const UWorld* InWorld, const FVector& Base, const FVector& X, const FVector& Y, const FColor& Color, float Radius, int32 NumSides, bool bPersistentLines, float LifeTime, uint8 DepthPriority, float Thickness)
 {
 	float	AngleDelta = 2.0f * (float)PI / ((float)NumSides);
 	FVector	LastVertex = Base + X * Radius;
@@ -709,12 +709,12 @@ static void DrawHalfCircle(const UWorld* InWorld, const FVector& Base, const FVe
 	for(int32 SideIndex = 0; SideIndex < (NumSides/2); SideIndex++)
 	{
 		FVector	Vertex = Base + (X * FMath::Cos(AngleDelta * (SideIndex + 1)) + Y * FMath::Sin(AngleDelta * (SideIndex + 1))) * Radius;
-		DrawDebugLine(InWorld, LastVertex, Vertex, Color, bPersistentLines, LifeTime, DepthPriority);
+		DrawDebugLine(InWorld, LastVertex, Vertex, Color, bPersistentLines, LifeTime, DepthPriority, Thickness);
 		LastVertex = Vertex;
 	}	
 }
 
-static void DrawCircle(const UWorld* InWorld, const FVector& Base, const FVector& X, const FVector& Y, const FColor& Color, float Radius, int32 NumSides, bool bPersistentLines, float LifeTime, uint8 DepthPriority = 0)
+static void DrawCircle(const UWorld* InWorld, const FVector& Base, const FVector& X, const FVector& Y, const FColor& Color, float Radius, int32 NumSides, bool bPersistentLines, float LifeTime, uint8 DepthPriority, float Thickness)
 {
 	const float	AngleDelta = 2.0f * PI / NumSides;
 	FVector	LastVertex = Base + X * Radius;
@@ -722,12 +722,12 @@ static void DrawCircle(const UWorld* InWorld, const FVector& Base, const FVector
 	for(int32 SideIndex = 0;SideIndex < NumSides;SideIndex++)
 	{
 		const FVector Vertex = Base + (X * FMath::Cos(AngleDelta * (SideIndex + 1)) + Y * FMath::Sin(AngleDelta * (SideIndex + 1))) * Radius;
-		DrawDebugLine(InWorld, LastVertex, Vertex, Color, bPersistentLines, LifeTime, DepthPriority);
+		DrawDebugLine(InWorld, LastVertex, Vertex, Color, bPersistentLines, LifeTime, DepthPriority, Thickness);
 		LastVertex = Vertex;
 	}
 }
 
-void DrawDebugCapsule(const UWorld* InWorld, FVector const& Center, float HalfHeight, float Radius, const FQuat& Rotation, FColor const& Color, bool bPersistentLines, float LifeTime, uint8 DepthPriority)
+void DrawDebugCapsule(const UWorld* InWorld, FVector const& Center, float HalfHeight, float Radius, const FQuat& Rotation, FColor const& Color, bool bPersistentLines, float LifeTime, uint8 DepthPriority, float Thickness)
 {
 	// no debug line drawing on dedicated server
 	if (GEngine->GetNetMode(InWorld) != NM_DedicatedServer)
@@ -745,23 +745,23 @@ void DrawDebugCapsule(const UWorld* InWorld, FVector const& Center, float HalfHe
 		FVector TopEnd = Origin + HalfAxis*ZAxis;
 		FVector BottomEnd = Origin - HalfAxis*ZAxis;
 
-		DrawCircle(InWorld, TopEnd, XAxis, YAxis, Color, Radius, DrawCollisionSides, bPersistentLines, LifeTime, DepthPriority);
-		DrawCircle(InWorld, BottomEnd, XAxis, YAxis, Color, Radius, DrawCollisionSides, bPersistentLines, LifeTime, DepthPriority);
+		DrawCircle(InWorld, TopEnd, XAxis, YAxis, Color, Radius, DrawCollisionSides, bPersistentLines, LifeTime, DepthPriority, Thickness);
+		DrawCircle(InWorld, BottomEnd, XAxis, YAxis, Color, Radius, DrawCollisionSides, bPersistentLines, LifeTime, DepthPriority, Thickness);
 
 		// Draw domed caps
-		DrawHalfCircle(InWorld, TopEnd, YAxis, ZAxis, Color, Radius, DrawCollisionSides, bPersistentLines, LifeTime, DepthPriority);
-		DrawHalfCircle(InWorld, TopEnd, XAxis, ZAxis, Color, Radius, DrawCollisionSides, bPersistentLines, LifeTime, DepthPriority);
+		DrawHalfCircle(InWorld, TopEnd, YAxis, ZAxis, Color, Radius, DrawCollisionSides, bPersistentLines, LifeTime, DepthPriority, Thickness);
+		DrawHalfCircle(InWorld, TopEnd, XAxis, ZAxis, Color, Radius, DrawCollisionSides, bPersistentLines, LifeTime, DepthPriority, Thickness);
 
 		FVector NegZAxis = -ZAxis;
 
-		DrawHalfCircle(InWorld, BottomEnd, YAxis, NegZAxis, Color, Radius, DrawCollisionSides, bPersistentLines, LifeTime, DepthPriority);
-		DrawHalfCircle(InWorld, BottomEnd, XAxis, NegZAxis, Color, Radius, DrawCollisionSides, bPersistentLines, LifeTime, DepthPriority);
+		DrawHalfCircle(InWorld, BottomEnd, YAxis, NegZAxis, Color, Radius, DrawCollisionSides, bPersistentLines, LifeTime, DepthPriority, Thickness);
+		DrawHalfCircle(InWorld, BottomEnd, XAxis, NegZAxis, Color, Radius, DrawCollisionSides, bPersistentLines, LifeTime, DepthPriority, Thickness);
 
-		// Draw connecty lines
-		DrawDebugLine(InWorld, TopEnd + Radius*XAxis, BottomEnd + Radius*XAxis, Color, bPersistentLines, LifeTime, DepthPriority);
-		DrawDebugLine(InWorld, TopEnd - Radius*XAxis, BottomEnd - Radius*XAxis, Color, bPersistentLines, LifeTime, DepthPriority);
-		DrawDebugLine(InWorld, TopEnd + Radius*YAxis, BottomEnd + Radius*YAxis, Color, bPersistentLines, LifeTime, DepthPriority);
-		DrawDebugLine(InWorld, TopEnd - Radius*YAxis, BottomEnd - Radius*YAxis, Color, bPersistentLines, LifeTime, DepthPriority);
+		// Draw connected lines
+		DrawDebugLine(InWorld, TopEnd + Radius*XAxis, BottomEnd + Radius*XAxis, Color, bPersistentLines, LifeTime, DepthPriority, Thickness);
+		DrawDebugLine(InWorld, TopEnd - Radius*XAxis, BottomEnd - Radius*XAxis, Color, bPersistentLines, LifeTime, DepthPriority, Thickness);
+		DrawDebugLine(InWorld, TopEnd + Radius*YAxis, BottomEnd + Radius*YAxis, Color, bPersistentLines, LifeTime, DepthPriority, Thickness);
+		DrawDebugLine(InWorld, TopEnd - Radius*YAxis, BottomEnd - Radius*YAxis, Color, bPersistentLines, LifeTime, DepthPriority, Thickness);
 	}
 }
 
diff --git a/Engine/Source/Runtime/Engine/Private/Particles/ParticleModules_Velocity.cpp b/Engine/Source/Runtime/Engine/Private/Particles/ParticleModules_Velocity.cpp
--- a/Engine/Source/Runtime/Engine/Private/Particles/ParticleModules_Velocity.cpp
+++ b/Engine/Source/Runtime/Engine/Private/Particles/ParticleModules_Velocity.cpp
@@ -624,8 +624,8 @@ void UParticleModuleVelocityCone::Render3DPreview(FParticleEmitterInstance* Owne
 	TArray<FVector> InnerVerts;
 
 	// Draw inner and outer cones
-	DrawWireCone(PDI, Transform, ConeRadius, ConeMinAngle, ConeSides, ModuleEditorColor, SDPG_World, InnerVerts);
-	DrawWireCone(PDI, Transform, ConeRadius, ConeMaxAngle, ConeSides, ModuleEditorColor, SDPG_World, OuterVerts);
+	DrawWireCone(PDI, InnerVerts, Transform, ConeRadius, ConeMinAngle, ConeSides, ModuleEditorColor, SDPG_World);
+	DrawWireCone(PDI, OuterVerts, Transform, ConeRadius, ConeMaxAngle, ConeSides, ModuleEditorColor, SDPG_World);
 
 	// Draw radial spokes
 	for (int32 i = 0; i < ConeSides; ++i)
diff --git a/Engine/Source/Runtime/Engine/Private/PrimitiveDrawingUtils.cpp b/Engine/Source/Runtime/Engine/Private/PrimitiveDrawingUtils.cpp
--- a/Engine/Source/Runtime/Engine/Private/PrimitiveDrawingUtils.cpp
+++ b/Engine/Source/Runtime/Engine/Private/PrimitiveDrawingUtils.cpp
@@ -183,23 +183,23 @@ void DrawBox(FPrimitiveDrawInterface* PDI,const FMatrix& BoxToWorld,const FVecto
 	MeshBuilder.Draw(PDI,FScaleMatrix(Radii) * BoxToWorld,MaterialRenderProxy,DepthPriorityGroup,0.f);
 }
 
-void GetSphereMesh(const FVector& Center,const FVector& Radii,int32 NumSides,int32 NumRings,const FMaterialRenderProxy* MaterialRenderProxy,uint8 DepthPriority,bool bDisableBackfaceCulling,int32 ViewIndex,FMeshElementCollector& Collector)
+void GetHalfSphereMesh(const FVector& Center, const FVector& Radii, int32 NumSides, int32 NumRings, float StartAngle, float EndAngle, const FMaterialRenderProxy* MaterialRenderProxy, uint8 DepthPriority, bool bDisableBackfaceCulling, int32 ViewIndex, FMeshElementCollector& Collector)
 {
 	// Use a mesh builder to draw the sphere.
 	FDynamicMeshBuilder MeshBuilder;
 	{
 		// The first/last arc are on top of each other.
-		int32 NumVerts = (NumSides+1) * (NumRings+1);
-		FDynamicMeshVertex* Verts = (FDynamicMeshVertex*)FMemory::Malloc( NumVerts * sizeof(FDynamicMeshVertex) );
+		int32 NumVerts = (NumSides + 1) * (NumRings + 1);
+		FDynamicMeshVertex* Verts = (FDynamicMeshVertex*)FMemory::Malloc(NumVerts * sizeof(FDynamicMeshVertex));
 
 		// Calculate verts for one arc
-		FDynamicMeshVertex* ArcVerts = (FDynamicMeshVertex*)FMemory::Malloc( (NumRings+1) * sizeof(FDynamicMeshVertex) );
+		FDynamicMeshVertex* ArcVerts = (FDynamicMeshVertex*)FMemory::Malloc((NumRings + 1) * sizeof(FDynamicMeshVertex));
 
-		for(int32 i=0; i<NumRings+1; i++)
+		for (int32 i = 0; i<NumRings + 1; i++)
 		{
 			FDynamicMeshVertex* ArcVert = &ArcVerts[i];
 
-			float angle = ((float)i/NumRings) * PI;
+			float angle = StartAngle + ((float)i / NumRings) * (EndAngle-StartAngle);
 
 			// Note- unit sphere, so position always has mag of one. We can just use it for normal!			
 			ArcVert->Position.X = 0.0f;
@@ -207,32 +207,32 @@ void GetSphereMesh(const FVector& Center,const FVector& Radii,int32 NumSides,int
 			ArcVert->Position.Z = FMath::Cos(angle);
 
 			ArcVert->SetTangents(
-				FVector(1,0,0),
-				FVector(0.0f,-ArcVert->Position.Z,ArcVert->Position.Y),
+				FVector(1, 0, 0),
+				FVector(0.0f, -ArcVert->Position.Z, ArcVert->Position.Y),
 				ArcVert->Position
 				);
 
 			ArcVert->TextureCoordinate.X = 0.0f;
-			ArcVert->TextureCoordinate.Y = ((float)i/NumRings);
+			ArcVert->TextureCoordinate.Y = ((float)i / NumRings);
 		}
 
 		// Then rotate this arc NumSides+1 times.
-		for(int32 s=0; s<NumSides+1; s++)
+		for (int32 s = 0; s<NumSides + 1; s++)
 		{
-			FRotator ArcRotator(0, 360.f * (float)s/NumSides, 0);
-			FRotationMatrix ArcRot( ArcRotator );
-			float XTexCoord = ((float)s/NumSides);
+			FRotator ArcRotator(0, 360.f * (float)s / NumSides, 0);
+			FRotationMatrix ArcRot(ArcRotator);
+			float XTexCoord = ((float)s / NumSides);
 
-			for(int32 v=0; v<NumRings+1; v++)
+			for (int32 v = 0; v<NumRings + 1; v++)
 			{
-				int32 VIx = (NumRings+1)*s + v;
+				int32 VIx = (NumRings + 1)*s + v;
 
-				Verts[VIx].Position = ArcRot.TransformPosition( ArcVerts[v].Position );
+				Verts[VIx].Position = ArcRot.TransformPosition(ArcVerts[v].Position);
 
 				Verts[VIx].SetTangents(
-					ArcRot.TransformVector( ArcVerts[v].TangentX ),
-					ArcRot.TransformVector( ArcVerts[v].GetTangentY() ),
-					ArcRot.TransformVector( ArcVerts[v].TangentZ )
+					ArcRot.TransformVector(ArcVerts[v].TangentX),
+					ArcRot.TransformVector(ArcVerts[v].GetTangentY()),
+					ArcRot.TransformVector(ArcVerts[v].TangentZ)
 					);
 
 				Verts[VIx].TextureCoordinate.X = XTexCoord;
@@ -241,18 +241,18 @@ void GetSphereMesh(const FVector& Center,const FVector& Radii,int32 NumSides,int
 		}
 
 		// Add all of the vertices we generated to the mesh builder.
-		for(int32 VertIdx=0; VertIdx < NumVerts; VertIdx++)
+		for (int32 VertIdx = 0; VertIdx < NumVerts; VertIdx++)
 		{
 			MeshBuilder.AddVertex(Verts[VertIdx]);
 		}
 
 		// Add all of the triangles we generated to the mesh builder.
-		for(int32 s=0; s<NumSides; s++)
+		for (int32 s = 0; s<NumSides; s++)
 		{
-			int32 a0start = (s+0) * (NumRings+1);
-			int32 a1start = (s+1) * (NumRings+1);
+			int32 a0start = (s + 0) * (NumRings + 1);
+			int32 a1start = (s + 1) * (NumRings + 1);
 
-			for(int32 r=0; r<NumRings; r++)
+			for (int32 r = 0; r<NumRings; r++)
 			{
 				MeshBuilder.AddTriangle(a0start + r + 0, a1start + r + 0, a0start + r + 1);
 				MeshBuilder.AddTriangle(a1start + r + 0, a1start + r + 1, a0start + r + 1);
@@ -263,7 +263,12 @@ void GetSphereMesh(const FVector& Center,const FVector& Radii,int32 NumSides,int
 		FMemory::Free(Verts);
 		FMemory::Free(ArcVerts);
 	}
-	MeshBuilder.GetMesh(FScaleMatrix( Radii ) * FTranslationMatrix( Center ), MaterialRenderProxy, DepthPriority, bDisableBackfaceCulling, false, ViewIndex, Collector);
+	MeshBuilder.GetMesh(FScaleMatrix(Radii) * FTranslationMatrix(Center), MaterialRenderProxy, DepthPriority, bDisableBackfaceCulling, false, ViewIndex, Collector);
+}
+
+void GetSphereMesh(const FVector& Center,const FVector& Radii,int32 NumSides,int32 NumRings,const FMaterialRenderProxy* MaterialRenderProxy,uint8 DepthPriority,bool bDisableBackfaceCulling,int32 ViewIndex,FMeshElementCollector& Collector)
+{
+	GetHalfSphereMesh(Center, Radii, NumSides, NumRings, 0, PI, MaterialRenderProxy, DepthPriority, bDisableBackfaceCulling, ViewIndex, Collector);
 }
 
 
@@ -598,6 +603,31 @@ void GetCylinderMesh(const FMatrix& CylToWorld, const FVector& Base, const FVect
 	MeshBuilder.GetMesh(CylToWorld, MaterialRenderProxy, DepthPriority, false, false, ViewIndex, Collector);
 }
 
+void GetConeMesh(const FMatrix& LocalToWorld, float AngleWidth, float AngleHeight, int32 NumSides, const FMaterialRenderProxy* MaterialRenderProxy, uint8 DepthPriority, int32 ViewIndex, FMeshElementCollector& Collector)
+{
+	TArray<FDynamicMeshVertex> MeshVerts;
+	TArray<int32> MeshIndices;
+	BuildConeVerts(AngleWidth * PI / 180, AngleHeight * PI / 180, 1.f, 0.f, 16, MeshVerts, MeshIndices);
+	FDynamicMeshBuilder MeshBuilder;
+	MeshBuilder.AddVertices(MeshVerts);
+	MeshBuilder.AddTriangles(MeshIndices);
+	MeshBuilder.GetMesh(LocalToWorld, MaterialRenderProxy, SDPG_World, false, false, ViewIndex, Collector);
+}
+
+void GetCapsuleMesh(const FVector& Origin, const FVector& XAxis, const FVector& YAxis, const FVector& ZAxis, const FLinearColor& Color, float Radius, float HalfHeight, int32 NumSides, const FMaterialRenderProxy* MaterialRenderProxy, uint8 DepthPriority, bool bDisableBackfaceCulling, int32 ViewIndex, FMeshElementCollector& Collector)
+{
+	const float HalfAxis = FMath::Max<float>(HalfHeight - Radius, 1.f);
+	const FVector BottomEnd = Origin + Radius * ZAxis;
+	const FVector TopEnd = BottomEnd + (2 * HalfAxis) * ZAxis;
+	const float CylinderHalfHeight = (TopEnd - BottomEnd).Size() * 0.5;
+	const FVector CylinderLocation = BottomEnd + CylinderHalfHeight * ZAxis;
+
+	GetHalfSphereMesh(TopEnd, FVector(Radius), NumSides, 16, 0, PI / 2, MaterialRenderProxy, DepthPriority, bDisableBackfaceCulling, ViewIndex, Collector);
+	GetCylinderMesh(CylinderLocation, XAxis, YAxis, ZAxis, Radius, CylinderHalfHeight, NumSides, MaterialRenderProxy, DepthPriority, ViewIndex, Collector);
+	GetHalfSphereMesh(BottomEnd, FVector(Radius), NumSides, 16, PI / 2, PI, MaterialRenderProxy, DepthPriority, bDisableBackfaceCulling, ViewIndex, Collector);
+}
+
+
 void DrawCylinder(FPrimitiveDrawInterface* PDI,const FVector& Base, const FVector& XAxis, const FVector& YAxis, const FVector& ZAxis,
 	float Radius, float HalfHeight, int32 Sides, const FMaterialRenderProxy* MaterialRenderProxy, uint8 DepthPriority)
 {
@@ -779,7 +809,7 @@ void DrawFlatArrow(class FPrimitiveDrawInterface* PDI,const FVector& Base,const
  * @param	Color			Color of the box.
  * @param	DepthPriority	Depth priority for the circle.
  */
-void DrawWireBox(FPrimitiveDrawInterface* PDI,const FBox& Box,const FLinearColor& Color,uint8 DepthPriority)
+void DrawWireBox(FPrimitiveDrawInterface* PDI, const FBox& Box, const FLinearColor& Color, uint8 DepthPriority, float Thickness, float DepthBias, bool bScreenSpace)
 {
 	FVector	B[2],P,Q;
 	int32 i,j;
@@ -792,17 +822,17 @@ void DrawWireBox(FPrimitiveDrawInterface* PDI,const FBox& Box,const FLinearColor
 		P.X=B[i].X; Q.X=B[i].X;
 		P.Y=B[j].Y; Q.Y=B[j].Y;
 		P.Z=B[0].Z; Q.Z=B[1].Z;
-		PDI->DrawLine(P,Q,Color,DepthPriority);
+		PDI->DrawLine(P, Q, Color, DepthPriority, Thickness, DepthBias, bScreenSpace);
 
 		P.Y=B[i].Y; Q.Y=B[i].Y;
 		P.Z=B[j].Z; Q.Z=B[j].Z;
 		P.X=B[0].X; Q.X=B[1].X;
-		PDI->DrawLine(P,Q,Color,DepthPriority);
+		PDI->DrawLine(P, Q, Color, DepthPriority, Thickness, DepthBias, bScreenSpace);
 
 		P.Z=B[i].Z; Q.Z=B[i].Z;
 		P.X=B[j].X; Q.X=B[j].X;
 		P.Y=B[0].Y; Q.Y=B[1].Y;
-		PDI->DrawLine(P,Q,Color,DepthPriority);
+		PDI->DrawLine(P, Q, Color, DepthPriority, Thickness, DepthBias, bScreenSpace);
 	}
 }
 
@@ -819,7 +849,7 @@ void DrawWireBox(FPrimitiveDrawInterface* PDI,const FBox& Box,const FLinearColor
  * @param	NumSides		Numbers of sides that the circle has.
  * @param	DepthPriority	Depth priority for the circle.
  */
-void DrawCircle(FPrimitiveDrawInterface* PDI,const FVector& Base,const FVector& X,const FVector& Y,const FLinearColor& Color,float Radius,int32 NumSides,uint8 DepthPriority)
+void DrawCircle(FPrimitiveDrawInterface* PDI, const FVector& Base, const FVector& X, const FVector& Y, const FLinearColor& Color, float Radius, int32 NumSides, uint8 DepthPriority, float Thickness, float DepthBias, bool bScreenSpace)
 {
 	const float	AngleDelta = 2.0f * PI / NumSides;
 	FVector	LastVertex = Base + X * Radius;
@@ -827,7 +857,7 @@ void DrawCircle(FPrimitiveDrawInterface* PDI,const FVector& Base,const FVector&
 	for(int32 SideIndex = 0;SideIndex < NumSides;SideIndex++)
 	{
 		const FVector Vertex = Base + (X * FMath::Cos(AngleDelta * (SideIndex + 1)) + Y * FMath::Sin(AngleDelta * (SideIndex + 1))) * Radius;
-		PDI->DrawLine(LastVertex,Vertex,Color,DepthPriority);
+		PDI->DrawLine(LastVertex, Vertex, Color, DepthPriority, Thickness, DepthBias, bScreenSpace);
 		LastVertex = Vertex;
 	}
 }
@@ -873,11 +903,11 @@ void DrawArc(FPrimitiveDrawInterface* PDI, const FVector Base, const FVector X,
  * @param	NumSides		Numbers of sides that the circle has.
  * @param	DepthPriority	Depth priority for the circle.
  */
-void DrawWireSphere(class FPrimitiveDrawInterface* PDI, const FVector& Base, const FLinearColor& Color, float Radius, int32 NumSides, uint8 DepthPriority)
+void DrawWireSphere(class FPrimitiveDrawInterface* PDI, const FVector& Base, const FLinearColor& Color, float Radius, int32 NumSides, uint8 DepthPriority, float Thickness, float DepthBias, bool bScreenSpace)
 {
-	DrawCircle(PDI, Base, FVector(1,0,0), FVector(0,1,0), Color, Radius, NumSides, DepthPriority);
-	DrawCircle(PDI, Base, FVector(1,0,0), FVector(0,0,1), Color, Radius, NumSides, DepthPriority);
-	DrawCircle(PDI, Base, FVector(0,1,0), FVector(0,0,1), Color, Radius, NumSides, DepthPriority);
+	DrawCircle(PDI, Base, FVector(1,0,0), FVector(0,1,0), Color, Radius, NumSides, DepthPriority, Thickness, DepthBias, bScreenSpace);
+	DrawCircle(PDI, Base, FVector(1, 0, 0), FVector(0, 0, 1), Color, Radius, NumSides, DepthPriority, Thickness, DepthBias, bScreenSpace);
+	DrawCircle(PDI, Base, FVector(0, 1, 0), FVector(0, 0, 1), Color, Radius, NumSides, DepthPriority, Thickness, DepthBias, bScreenSpace);
 }
 
 /**
@@ -889,26 +919,26 @@ void DrawWireSphere(class FPrimitiveDrawInterface* PDI, const FVector& Base, con
  * @param	Radius			Radius of the sphere.
  * @param	DepthPriority	Depth priority for the circle.
  */
-void DrawWireSphereAutoSides(class FPrimitiveDrawInterface* PDI, const FVector& Base, const FLinearColor& Color, float Radius, uint8 DepthPriority)
+void DrawWireSphereAutoSides(class FPrimitiveDrawInterface* PDI, const FVector& Base, const FLinearColor& Color, float Radius, uint8 DepthPriority, float Thickness, float DepthBias, bool bScreenSpace)
 {
 	// Guess a good number of sides
 	int32 NumSides =  FMath::Clamp<int32>(Radius/4.f, 16, 64);
-	DrawWireSphere(PDI, Base, Color, Radius, NumSides, DepthPriority);
+	DrawWireSphere(PDI, Base, Color, Radius, NumSides, DepthPriority, Thickness, DepthBias, bScreenSpace);
 }
 
-void DrawWireSphere(class FPrimitiveDrawInterface* PDI, const FTransform& Transform, const FLinearColor& Color, float Radius, int32 NumSides, uint8 DepthPriority)
+void DrawWireSphere(class FPrimitiveDrawInterface* PDI, const FTransform& Transform, const FLinearColor& Color, float Radius, int32 NumSides, uint8 DepthPriority, float Thickness, float DepthBias, bool bScreenSpace)
 {
-	DrawCircle( PDI, Transform.GetLocation(), Transform.GetScaledAxis( EAxis::X ), Transform.GetScaledAxis( EAxis::Y ), Color, Radius, NumSides, SDPG_World );
-	DrawCircle( PDI, Transform.GetLocation(), Transform.GetScaledAxis( EAxis::X ), Transform.GetScaledAxis( EAxis::Z ), Color, Radius, NumSides, SDPG_World );
-	DrawCircle( PDI, Transform.GetLocation(), Transform.GetScaledAxis( EAxis::Y ), Transform.GetScaledAxis( EAxis::Z ), Color, Radius, NumSides, SDPG_World );
+	DrawCircle(PDI, Transform.GetLocation(), Transform.GetScaledAxis(EAxis::X), Transform.GetScaledAxis(EAxis::Y), Color, Radius, NumSides, SDPG_World, Thickness, DepthBias, bScreenSpace);
+	DrawCircle(PDI, Transform.GetLocation(), Transform.GetScaledAxis(EAxis::X), Transform.GetScaledAxis(EAxis::Z), Color, Radius, NumSides, SDPG_World, Thickness, DepthBias, bScreenSpace);
+	DrawCircle(PDI, Transform.GetLocation(), Transform.GetScaledAxis(EAxis::Y), Transform.GetScaledAxis(EAxis::Z), Color, Radius, NumSides, SDPG_World, Thickness, DepthBias, bScreenSpace);
 }
 
 
-void DrawWireSphereAutoSides(class FPrimitiveDrawInterface* PDI, const FTransform& Transform, const FLinearColor& Color, float Radius, uint8 DepthPriority)
+void DrawWireSphereAutoSides(class FPrimitiveDrawInterface* PDI, const FTransform& Transform, const FLinearColor& Color, float Radius, uint8 DepthPriority, float Thickness, float DepthBias, bool bScreenSpace)
 {
 	// Guess a good number of sides
 	int32 NumSides =  FMath::Clamp<int32>(Radius/4.f, 16, 64);
-	DrawWireSphere(PDI, Transform, Color, Radius, NumSides, DepthPriority);
+	DrawWireSphere(PDI, Transform, Color, Radius, NumSides, DepthPriority, Thickness, DepthBias, bScreenSpace);
 }
 
 /**
@@ -925,7 +955,7 @@ void DrawWireSphereAutoSides(class FPrimitiveDrawInterface* PDI, const FTransfor
  * @param	NumSides		Numbers of sides that the cylinder has.
  * @param	DepthPriority	Depth priority for the cylinder.
  */
-void DrawWireCylinder(FPrimitiveDrawInterface* PDI,const FVector& Base,const FVector& X,const FVector& Y,const FVector& Z,const FLinearColor& Color,float Radius,float HalfHeight,int32 NumSides,uint8 DepthPriority)
+void DrawWireCylinder(FPrimitiveDrawInterface* PDI, const FVector& Base, const FVector& X, const FVector& Y, const FVector& Z, const FLinearColor& Color, float Radius, float HalfHeight, int32 NumSides, uint8 DepthPriority, float Thickness, float DepthBias, bool bScreenSpace)
 {
 	const float	AngleDelta = 2.0f * PI / NumSides;
 	FVector	LastVertex = Base + X * Radius;
@@ -934,16 +964,16 @@ void DrawWireCylinder(FPrimitiveDrawInterface* PDI,const FVector& Base,const FVe
 	{
 		const FVector Vertex = Base + (X * FMath::Cos(AngleDelta * (SideIndex + 1)) + Y * FMath::Sin(AngleDelta * (SideIndex + 1))) * Radius;
 
-		PDI->DrawLine(LastVertex - Z * HalfHeight,Vertex - Z * HalfHeight,Color,DepthPriority);
-		PDI->DrawLine(LastVertex + Z * HalfHeight,Vertex + Z * HalfHeight,Color,DepthPriority);
-		PDI->DrawLine(LastVertex - Z * HalfHeight,LastVertex + Z * HalfHeight,Color,DepthPriority);
+		PDI->DrawLine(LastVertex - Z * HalfHeight, Vertex - Z * HalfHeight, Color, DepthPriority, Thickness, DepthBias, bScreenSpace);
+		PDI->DrawLine(LastVertex + Z * HalfHeight, Vertex + Z * HalfHeight, Color, DepthPriority, Thickness, DepthBias, bScreenSpace);
+		PDI->DrawLine(LastVertex - Z * HalfHeight, LastVertex + Z * HalfHeight, Color, DepthPriority, Thickness, DepthBias, bScreenSpace);
 
 		LastVertex = Vertex;
 	}
 }
 
 
-static void DrawHalfCircle(FPrimitiveDrawInterface* PDI, const FVector& Base, const FVector& X, const FVector& Y, const FLinearColor& Color, float Radius, int32 NumSides)
+static void DrawHalfCircle(FPrimitiveDrawInterface* PDI, const FVector& Base, const FVector& X, const FVector& Y, const FLinearColor& Color, float Radius, int32 NumSides, float Thickness, float DepthBias, bool bScreenSpace)
 {
 	const float	AngleDelta = (float)PI / ((float)NumSides);
 	FVector	LastVertex = Base + X * Radius;
@@ -951,7 +981,7 @@ static void DrawHalfCircle(FPrimitiveDrawInterface* PDI, const FVector& Base, co
 	for(int32 SideIndex = 0; SideIndex < NumSides; SideIndex++)
 	{
 		const FVector	Vertex = Base + (X * FMath::Cos(AngleDelta * (SideIndex + 1)) + Y * FMath::Sin(AngleDelta * (SideIndex + 1))) * Radius;
-		PDI->DrawLine(LastVertex, Vertex, Color, SDPG_World);
+		PDI->DrawLine(LastVertex, Vertex, Color, SDPG_World, Thickness, DepthBias, bScreenSpace);
 		LastVertex = Vertex;
 	}	
 }
@@ -971,7 +1001,7 @@ static void DrawHalfCircle(FPrimitiveDrawInterface* PDI, const FVector& Base, co
  * @param	NumSides		Numbers of sides that the cylinder has.
  * @param	DepthPriority	Depth priority for the cylinder.
  */
-void DrawWireCapsule(FPrimitiveDrawInterface* PDI,const FVector& Base,const FVector& X,const FVector& Y,const FVector& Z,const FLinearColor& Color,float Radius,float HalfHeight,int32 NumSides,uint8 DepthPriority)
+void DrawWireCapsule(FPrimitiveDrawInterface* PDI, const FVector& Base, const FVector& X, const FVector& Y, const FVector& Z, const FLinearColor& Color, float Radius, float HalfHeight, int32 NumSides, uint8 DepthPriority, float Thickness, float DepthBias, bool bScreenSpace)
 {
 	const FVector Origin = Base;
 	const FVector XAxis = X.SafeNormal();
@@ -992,17 +1022,17 @@ void DrawWireCapsule(FPrimitiveDrawInterface* PDI,const FVector& Base,const FVec
 	const FVector TopEnd = Origin + (HalfHeight * ZAxis);
 	const FVector BottomEnd = Origin - (HalfHeight * ZAxis);
 
-	DrawCircle(PDI, TopEnd, XAxis, YAxis, Color, CapsuleRadius, NumSides, DepthPriority);
-	DrawCircle(PDI, BottomEnd, XAxis, YAxis, Color, CapsuleRadius, NumSides, DepthPriority);
+	DrawCircle(PDI, TopEnd, XAxis, YAxis, Color, CapsuleRadius, NumSides, DepthPriority, Thickness, DepthBias, bScreenSpace);
+	DrawCircle(PDI, BottomEnd, XAxis, YAxis, Color, CapsuleRadius, NumSides, DepthPriority, Thickness, DepthBias, bScreenSpace);
 
 	// Draw domed caps
-	DrawHalfCircle(PDI, TopEnd, YAxis, ZAxis, Color, CapsuleRadius, NumSides/2);
-	DrawHalfCircle(PDI, TopEnd, XAxis, ZAxis, Color, CapsuleRadius, NumSides/2);
+	DrawHalfCircle(PDI, TopEnd, YAxis, ZAxis, Color, CapsuleRadius, NumSides / 2, Thickness, DepthBias, bScreenSpace);
+	DrawHalfCircle(PDI, TopEnd, XAxis, ZAxis, Color, CapsuleRadius, NumSides / 2, Thickness, DepthBias, bScreenSpace);
 
 	const FVector NegZAxis = -ZAxis;
 
-	DrawHalfCircle(PDI, BottomEnd, YAxis, NegZAxis, Color, CapsuleRadius, NumSides/2);
-	DrawHalfCircle(PDI, BottomEnd, XAxis, NegZAxis, Color, CapsuleRadius, NumSides/2);
+	DrawHalfCircle(PDI, BottomEnd, YAxis, NegZAxis, Color, CapsuleRadius, NumSides / 2, Thickness, DepthBias, bScreenSpace);
+	DrawHalfCircle(PDI, BottomEnd, XAxis, NegZAxis, Color, CapsuleRadius, NumSides / 2, Thickness, DepthBias, bScreenSpace);
 
 	// we set NumSides to 4 as it makes a nicer looking capsule as we only draw 2 HalfCircles above
 	const int32 NumCylinderLines = 4;
@@ -1015,7 +1045,7 @@ void DrawWireCapsule(FPrimitiveDrawInterface* PDI,const FVector& Base,const FVec
 	{
 		const FVector Vertex = Base + (XAxis * FMath::Cos(AngleDelta * (SideIndex + 1)) + YAxis * FMath::Sin(AngleDelta * (SideIndex + 1))) * CapsuleRadius;
 
-		PDI->DrawLine(LastVertex - ZAxis * HalfHeight,LastVertex + ZAxis * HalfHeight,Color,DepthPriority);
+		PDI->DrawLine(LastVertex - ZAxis * HalfHeight, LastVertex + ZAxis * HalfHeight, Color, DepthPriority, Thickness, DepthBias, bScreenSpace);
 
 		LastVertex = Vertex;
 	}
@@ -1035,7 +1065,7 @@ void DrawWireCapsule(FPrimitiveDrawInterface* PDI,const FVector& Base,const FVec
  * @param	DepthPriority	Depth priority for the cone.
  * @param	Verts			Out param, the positions of the verts at the cone's base.
  */
-void DrawWireCone(FPrimitiveDrawInterface* PDI, const FTransform& Transform, float ConeRadius, float ConeAngle, int32 ConeSides, const FLinearColor& Color, uint8 DepthPriority, TArray<FVector>& Verts)
+void DrawWireCone(FPrimitiveDrawInterface* PDI, TArray<FVector>& Verts, const FTransform& Transform, float ConeRadius, float ConeAngle, int32 ConeSides, const FLinearColor& Color, uint8 DepthPriority, float Thickness, float DepthBias, bool bScreenSpace)
 {
 	static const float TwoPI = 2.0f * PI;
 	static const float ToRads = PI / 180.0f;
@@ -1066,20 +1096,20 @@ void DrawWireCone(FPrimitiveDrawInterface* PDI, const FTransform& Transform, flo
 	// Draw spokes.
 	for ( int32 i = 0 ; i < Verts.Num(); ++i )
 	{
-		PDI->DrawLine( Transform.GetLocation(), Verts[i], Color, DepthPriority );
+		PDI->DrawLine(Transform.GetLocation(), Verts[i], Color, DepthPriority, Thickness, DepthBias, bScreenSpace);
 	}
 
 	// Draw rim.
 	for ( int32 i = 0 ; i < Verts.Num()-1 ; ++i )
 	{
-		PDI->DrawLine( Verts[i], Verts[i+1], Color, DepthPriority );
+		PDI->DrawLine(Verts[i], Verts[i + 1], Color, DepthPriority, Thickness, DepthBias, bScreenSpace);
 	}
-	PDI->DrawLine( Verts[Verts.Num()-1], Verts[0], Color, DepthPriority );
+	PDI->DrawLine(Verts[Verts.Num() - 1], Verts[0], Color, DepthPriority, Thickness, DepthBias, bScreenSpace);
 }
 
-void DrawWireCone(FPrimitiveDrawInterface* PDI, const FMatrix& Transform, float ConeRadius, float ConeAngle, int32 ConeSides, const FLinearColor& Color, uint8 DepthPriority, TArray<FVector>& Verts)
+void DrawWireCone(FPrimitiveDrawInterface* PDI, TArray<FVector>& Verts, const FMatrix& Transform, float ConeRadius, float ConeAngle, int32 ConeSides, const FLinearColor& Color, uint8 DepthPriority, float Thickness, float DepthBias, bool bScreenSpace)
 {
-	DrawWireCone(PDI, FTransform(Transform), ConeRadius, ConeAngle, ConeSides, Color, DepthPriority, Verts);
+	DrawWireCone(PDI, Verts, FTransform(Transform), ConeRadius, ConeAngle, ConeSides, Color, DepthPriority, Thickness, DepthBias, bScreenSpace);
 }
 
 /**
@@ -1104,7 +1134,7 @@ void DrawWireSphereCappedCone(FPrimitiveDrawInterface* PDI, const FTransform& Tr
 	}
 
 	TArray<FVector> Verts;
-	DrawWireCone(PDI, Transform, ConeRadius, ConeAngle, ConeSides, Color, DepthPriority, Verts);
+	DrawWireCone(PDI, Verts, Transform, ConeRadius, ConeAngle, ConeSides, Color, DepthPriority, 0);
 
 	// Draw arcs
 	int32 ArcCount = (int32)(Verts.Num() / 2);
diff --git a/Engine/Source/Runtime/Engine/Public/DrawDebugHelpers.h b/Engine/Source/Runtime/Engine/Public/DrawDebugHelpers.h
--- a/Engine/Source/Runtime/Engine/Public/DrawDebugHelpers.h
+++ b/Engine/Source/Runtime/Engine/Public/DrawDebugHelpers.h
@@ -35,7 +35,7 @@ ENGINE_API void DrawDebugAltCone(const UWorld* InWorld, FVector const& Origin, F
 ENGINE_API void DrawDebugString(const UWorld* InWorld, FVector const& TextLocation, const FString& Text, class AActor* TestBaseActor = NULL, FColor const& TextColor = FColor::White, float Duration = -1.000000);
 ENGINE_API void DrawDebugFrustum(const UWorld* InWorld, const FMatrix& FrustumToWorld, FColor const& Color, bool bPersistentLines = false, float LifeTime=-1.f, uint8 DepthPriority = 0);
 /** Draw a capsule using the LineBatcher */
-ENGINE_API void DrawDebugCapsule(const UWorld* InWorld, FVector const& Center, float HalfHeight, float Radius, const FQuat& Rotation, FColor const& Color, bool bPersistentLines=false, float LifeTime=-1.f, uint8 DepthPriority = 0);
+ENGINE_API void DrawDebugCapsule(const UWorld* InWorld, FVector const& Center, float HalfHeight, float Radius, const FQuat& Rotation, FColor const& Color, bool bPersistentLines = false, float LifeTime = -1.f, uint8 DepthPriority = 0, float Thickness = 0);
 /** Draw a debug camera shape.  FOV is full angle in degrees. */
 ENGINE_API void DrawDebugCamera(const UWorld* InWorld, FVector const& Location, FRotator const& Rotation, float FOVDeg, float Scale=1.f, FColor const& Color=FColor::White, bool bPersistentLines=false, float LifeTime=-1.f, uint8 DepthPriority = 0);
 ENGINE_API void FlushDebugStrings(const UWorld* InWorld);
diff --git a/Engine/Source/Runtime/Engine/Public/SceneManagement.h b/Engine/Source/Runtime/Engine/Public/SceneManagement.h
--- a/Engine/Source/Runtime/Engine/Public/SceneManagement.h
+++ b/Engine/Source/Runtime/Engine/Public/SceneManagement.h
@@ -1648,29 +1648,34 @@ extern ENGINE_API void DrawCylinder(class FPrimitiveDrawInterface* PDI, const FM
 	float Radius, float HalfHeight, int32 Sides, const FMaterialRenderProxy* MaterialInstance, uint8 DepthPriority);
 
 extern ENGINE_API void GetBoxMesh(const FMatrix& BoxToWorld,const FVector& Radii,const FMaterialRenderProxy* MaterialRenderProxy,uint8 DepthPriority,int32 ViewIndex,FMeshElementCollector& Collector);
-extern ENGINE_API void GetSphereMesh(const FVector& Center,const FVector& Radii,int32 NumSides,int32 NumRings,const FMaterialRenderProxy* MaterialRenderProxy,uint8 DepthPriority,bool bDisableBackfaceCulling,int32 ViewIndex,FMeshElementCollector& Collector);
+extern ENGINE_API void GetSphereMesh(const FVector& Center, const FVector& Radii, int32 NumSides, int32 NumRings, const FMaterialRenderProxy* MaterialRenderProxy, uint8 DepthPriority, bool bDisableBackfaceCulling, int32 ViewIndex, FMeshElementCollector& Collector);
+extern ENGINE_API void GetHalfSphereMesh(const FVector& Center, const FVector& Radii, int32 NumSides, int32 NumRings, float StartAngle, float EndAngle, const FMaterialRenderProxy* MaterialRenderProxy, uint8 DepthPriority, bool bDisableBackfaceCulling, int32 ViewIndex, FMeshElementCollector& Collector);
 extern ENGINE_API void GetCylinderMesh(const FVector& Base, const FVector& XAxis, const FVector& YAxis, const FVector& ZAxis,
 									float Radius, float HalfHeight, int32 Sides, const FMaterialRenderProxy* MaterialInstance, uint8 DepthPriority, int32 ViewIndex, FMeshElementCollector& Collector);
 extern ENGINE_API void GetCylinderMesh(const FMatrix& CylToWorld, const FVector& Base, const FVector& XAxis, const FVector& YAxis, const FVector& ZAxis,
 									float Radius, float HalfHeight, int32 Sides, const FMaterialRenderProxy* MaterialInstance, uint8 DepthPriority, int32 ViewIndex, FMeshElementCollector& Collector);
+extern ENGINE_API void GetConeMesh(const FMatrix& LocalToWorld, float AngleWidth, float AngleHeight, int32 NumSides,
+									const FMaterialRenderProxy* MaterialRenderProxy, uint8 DepthPriority, int32 ViewIndex, FMeshElementCollector& Collector);
+extern ENGINE_API void GetCapsuleMesh(const FVector& Origin, const FVector& XAxis, const FVector& YAxis, const FVector& ZAxis, const FLinearColor& Color, float Radius, float HalfHeight, int32 NumSides,
+									const FMaterialRenderProxy* MaterialRenderProxy, uint8 DepthPriority, bool bDisableBackfaceCulling, int32 ViewIndex, FMeshElementCollector& Collector);
 
 
 extern ENGINE_API void DrawDisc(class FPrimitiveDrawInterface* PDI,const FVector& Base,const FVector& XAxis,const FVector& YAxis,FColor Color,float Radius,int32 NumSides, const FMaterialRenderProxy* MaterialRenderProxy, uint8 DepthPriority);
 extern ENGINE_API void DrawFlatArrow(class FPrimitiveDrawInterface* PDI,const FVector& Base,const FVector& XAxis,const FVector& YAxis,FColor Color,float Length,int32 Width, const FMaterialRenderProxy* MaterialRenderProxy, uint8 DepthPriority);
 
 // Line drawing utility functions.
-extern ENGINE_API void DrawWireBox(class FPrimitiveDrawInterface* PDI,const FBox& Box,const FLinearColor& Color,uint8 DepthPriority);
-extern ENGINE_API void DrawCircle(class FPrimitiveDrawInterface* PDI,const FVector& Base,const FVector& X,const FVector& Y,const FLinearColor& Color,float Radius,int32 NumSides,uint8 DepthPriority);
+extern ENGINE_API void DrawWireBox(class FPrimitiveDrawInterface* PDI, const FBox& Box, const FLinearColor& Color, uint8 DepthPriority, float Thickness = 0, float DepthBias = 0.0f, bool bScreenSpace = false);
+extern ENGINE_API void DrawCircle(class FPrimitiveDrawInterface* PDI, const FVector& Base, const FVector& X, const FVector& Y, const FLinearColor& Color, float Radius, int32 NumSides, uint8 DepthPriority, float Thickness = 0, float DepthBias = 0.0f, bool bScreenSpace = false);
 extern ENGINE_API void DrawArc(FPrimitiveDrawInterface* PDI, const FVector Base, const FVector X, const FVector Y, const float MinAngle, const float MaxAngle, const float Radius, const int32 Sections, const FLinearColor& Color, uint8 DepthPriority);
-extern ENGINE_API void DrawWireSphere(class FPrimitiveDrawInterface* PDI, const FVector& Base, const FLinearColor& Color, float Radius, int32 NumSides, uint8 DepthPriority);
-extern ENGINE_API void DrawWireSphereAutoSides(class FPrimitiveDrawInterface* PDI, const FVector& Base, const FLinearColor& Color, float Radius, uint8 DepthPriority);
-extern ENGINE_API void DrawWireSphere(class FPrimitiveDrawInterface* PDI, const FTransform& Transform, const FLinearColor& Color, float Radius, int32 NumSides, uint8 DepthPriority);
-extern ENGINE_API void DrawWireSphereAutoSides(class FPrimitiveDrawInterface* PDI, const FTransform& Transform, const FLinearColor& Color, float Radius, uint8 DepthPriority);
-extern ENGINE_API void DrawWireCylinder(class FPrimitiveDrawInterface* PDI,const FVector& Base,const FVector& X,const FVector& Y,const FVector& Z,const FLinearColor& Color,float Radius,float HalfHeight,int32 NumSides,uint8 DepthPriority);
-extern ENGINE_API void DrawWireCapsule(class FPrimitiveDrawInterface* PDI,const FVector& Base,const FVector& X,const FVector& Y,const FVector& Z,const FLinearColor& Color,float Radius,float HalfHeight,int32 NumSides,uint8 DepthPriority);
+extern ENGINE_API void DrawWireSphere(class FPrimitiveDrawInterface* PDI, const FVector& Base, const FLinearColor& Color, float Radius, int32 NumSides, uint8 DepthPriority, float Thickness = 0, float DepthBias = 0.0f, bool bScreenSpace = false);
+extern ENGINE_API void DrawWireSphere(class FPrimitiveDrawInterface* PDI, const FTransform& Transform, const FLinearColor& Color, float Radius, int32 NumSides, uint8 DepthPriority, float Thickness = 0, float DepthBias = 0.0f, bool bScreenSpace = false);
+extern ENGINE_API void DrawWireSphereAutoSides(class FPrimitiveDrawInterface* PDI, const FVector& Base, const FLinearColor& Color, float Radius, uint8 DepthPriority, float Thickness = 0, float DepthBias = 0.0f, bool bScreenSpace = false);
+extern ENGINE_API void DrawWireSphereAutoSides(class FPrimitiveDrawInterface* PDI, const FTransform& Transform, const FLinearColor& Color, float Radius, uint8 DepthPriority, float Thickness = 0, float DepthBias = 0.0f, bool bScreenSpace = false);
+extern ENGINE_API void DrawWireCylinder(class FPrimitiveDrawInterface* PDI, const FVector& Base, const FVector& X, const FVector& Y, const FVector& Z, const FLinearColor& Color, float Radius, float HalfHeight, int32 NumSides, uint8 DepthPriority, float Thickness = 0, float DepthBias = 0.0f, bool bScreenSpace = false);
+extern ENGINE_API void DrawWireCapsule(class FPrimitiveDrawInterface* PDI, const FVector& Base, const FVector& X, const FVector& Y, const FVector& Z, const FLinearColor& Color, float Radius, float HalfHeight, int32 NumSides, uint8 DepthPriority, float Thickness = 0, float DepthBias = 0.0f, bool bScreenSpace = false);
 extern ENGINE_API void DrawWireChoppedCone(class FPrimitiveDrawInterface* PDI,const FVector& Base,const FVector& X,const FVector& Y,const FVector& Z,const FLinearColor& Color,float Radius,float TopRadius,float HalfHeight,int32 NumSides,uint8 DepthPriority);
-extern ENGINE_API void DrawWireCone(class FPrimitiveDrawInterface* PDI, const FMatrix& Transform, float ConeRadius, float ConeAngle, int32 ConeSides, const FLinearColor& Color, uint8 DepthPriority, TArray<FVector>& Verts);
-extern ENGINE_API void DrawWireCone(class FPrimitiveDrawInterface* PDI, const FTransform& Transform, float ConeRadius, float ConeAngle, int32 ConeSides, const FLinearColor& Color, uint8 DepthPriority, TArray<FVector>& Verts);
+extern ENGINE_API void DrawWireCone(class FPrimitiveDrawInterface* PDI, TArray<FVector>& Verts, const FMatrix& Transform, float ConeRadius, float ConeAngle, int32 ConeSides, const FLinearColor& Color, uint8 DepthPriority, float Thickness = 0, float DepthBias = 0.0f, bool bScreenSpace = false);
+extern ENGINE_API void DrawWireCone(class FPrimitiveDrawInterface* PDI, TArray<FVector>& Verts, const FTransform& Transform, float ConeRadius, float ConeAngle, int32 ConeSides, const FLinearColor& Color, uint8 DepthPriority, float Thickness = 0, float DepthBias = 0.0f, bool bScreenSpace = false);
 extern ENGINE_API void DrawWireSphereCappedCone(FPrimitiveDrawInterface* PDI, const FTransform& Transform, float ConeRadius, float ConeAngle, int32 ConeSides, int32 ArcFrequency, int32 CapSegments, const FLinearColor& Color, uint8 DepthPriority);
 extern ENGINE_API void DrawOrientedWireBox(class FPrimitiveDrawInterface* PDI,const FVector& Base,const FVector& X,const FVector& Y,const FVector& Z, FVector Extent, const FLinearColor& Color,uint8 DepthPriority);
 extern ENGINE_API void DrawDirectionalArrow(class FPrimitiveDrawInterface* PDI,const FMatrix& ArrowToWorld,const FLinearColor& InColor,float Length,float ArrowSize,uint8 DepthPriority);
