Fixed deferred physics commands not being run when terminating a physics scene, causing out of order releases, added a warning for dangling commands.
[UE-5109]

diff --git a/Engine/Source/Runtime/Engine/Private/PhysicsEngine/PhysCommandHandler.cpp b/Engine/Source/Runtime/Engine/Private/PhysicsEngine/PhysCommandHandler.cpp
--- a/Engine/Source/Runtime/Engine/Private/PhysicsEngine/PhysCommandHandler.cpp
+++ b/Engine/Source/Runtime/Engine/Private/PhysicsEngine/PhysCommandHandler.cpp
@@ -9,6 +9,14 @@
 	#include "../Vehicles/PhysXVehicleManager.h"
 #endif
 
+FPhysCommandHandler::~FPhysCommandHandler()
+{
+	if(PendingCommands.Num() != 0)
+	{
+		UE_LOG(LogPhysics, Warning, TEXT("~FPhysCommandHandler() - Pending command list is not empty. %d item remain."), PendingCommands.Num());
+	}
+}
+
 void FPhysCommandHandler::Flush()
 {
 	check(IsInGameThread());
diff --git a/Engine/Source/Runtime/Engine/Private/PhysicsEngine/PhysScene.cpp b/Engine/Source/Runtime/Engine/Private/PhysicsEngine/PhysScene.cpp
--- a/Engine/Source/Runtime/Engine/Private/PhysicsEngine/PhysScene.cpp
+++ b/Engine/Source/Runtime/Engine/Private/PhysicsEngine/PhysScene.cpp
@@ -1333,6 +1333,9 @@ void FPhysScene::TermPhysScene(uint32 SceneType)
 		// @todo block on any running scene before calling this
 		GPhysCommandHandler->DeferredRelease(PScene);
 
+		// Commands may have accumulated as the scene is terminated - flush any commands for this scene.
+		DeferredCommandHandler.Flush();
+
 		// Remove from the map
 		GPhysXSceneMap.Remove(PhysXSceneIndex[SceneType]);
 	}
diff --git a/Engine/Source/Runtime/Engine/Public/PhysicsPublic.h b/Engine/Source/Runtime/Engine/Public/PhysicsPublic.h
--- a/Engine/Source/Runtime/Engine/Public/PhysicsPublic.h
+++ b/Engine/Source/Runtime/Engine/Public/PhysicsPublic.h
@@ -196,6 +196,8 @@ class FPhysCommandHandler
 {
 public:
 
+	~FPhysCommandHandler();
+	
 	/** Executes pending commands and clears buffer **/
 	void ENGINE_API Flush();
 
