Compartmentalize binding value clearing

diff --git a/Engine/Source/Runtime/Engine/Classes/Components/InputComponent.h b/Engine/Source/Runtime/Engine/Classes/Components/InputComponent.h
--- a/Engine/Source/Runtime/Engine/Classes/Components/InputComponent.h
+++ b/Engine/Source/Runtime/Engine/Classes/Components/InputComponent.h
@@ -629,6 +629,9 @@ public:
 	 */
 	void RemoveActionBinding( const int32 BindingIndex );
 
+	/** Clears all cached binding values. */
+	void ClearBindingValues();
+
 public:
 
 	/**
diff --git a/Engine/Source/Runtime/Engine/Private/InputComponent.cpp b/Engine/Source/Runtime/Engine/Private/InputComponent.cpp
--- a/Engine/Source/Runtime/Engine/Private/InputComponent.cpp
+++ b/Engine/Source/Runtime/Engine/Private/InputComponent.cpp
@@ -227,6 +227,25 @@ void UInputComponent::RemoveActionBinding( const int32 BindingIndex )
 	}
 }
 
+void UInputComponent::ClearBindingValues()
+{
+	for (FInputAxisBinding& AxisBinding : AxisBindings)
+	{
+		AxisBinding.AxisValue = 0.f;
+	}
+	for (FInputAxisKeyBinding& AxisKeyBinding : AxisKeyBindings)
+	{
+		AxisKeyBinding.AxisValue = 0.f;
+	}
+	for (FInputVectorAxisBinding& VectorAxisBinding : VectorAxisBindings)
+	{
+		VectorAxisBinding.AxisValue = FVector::ZeroVector;
+	}
+	for (FInputGestureBinding& GestureBinding : GestureBindings)
+	{
+		GestureBinding.GestureValue = 0.f;
+	}
+}
 
 /* Deprecated functions (needed for Blueprints)
  *****************************************************************************/
diff --git a/Engine/Source/Runtime/Engine/Private/PlayerController.cpp b/Engine/Source/Runtime/Engine/Private/PlayerController.cpp
--- a/Engine/Source/Runtime/Engine/Private/PlayerController.cpp
+++ b/Engine/Source/Runtime/Engine/Private/PlayerController.cpp
@@ -4076,19 +4076,7 @@ bool APlayerController::PopInputComponent(UInputComponent* InputComponent)
 	{
 		if (CurrentInputStack.RemoveSingle(InputComponent) > 0)
 		{
-			for (FInputAxisBinding& AxisBinding : InputComponent->AxisBindings)
-			{
-				AxisBinding.AxisValue = 0.f;
-			}
-			for (FInputAxisKeyBinding& AxisKeyBinding : InputComponent->AxisKeyBindings)
-			{
-				AxisKeyBinding.AxisValue = 0.f;
-			}
-			for (FInputVectorAxisBinding& VectorAxisBinding : InputComponent->VectorAxisBindings)
-			{
-				VectorAxisBinding.AxisValue = FVector::ZeroVector;
-			}
-
+			InputComponent->ClearBindingValues();
 			return true;
 		}
 	}
