Fix non unity compile error.

diff --git a/Engine/Source/Runtime/Engine/Private/Particles/ParticleModules_Size.cpp b/Engine/Source/Runtime/Engine/Private/Particles/ParticleModules_Size.cpp
--- a/Engine/Source/Runtime/Engine/Private/Particles/ParticleModules_Size.cpp
+++ b/Engine/Source/Runtime/Engine/Private/Particles/ParticleModules_Size.cpp
@@ -16,6 +16,7 @@
 #include "Particles/TypeData/ParticleModuleTypeDataGpu.h"
 #include "Particles/ParticleLODLevel.h"
 #include "Particles/ParticleSystemComponent.h"
+#include "Particles/ParticleModuleRequired.h"
 
 UParticleModuleSizeBase::UParticleModuleSizeBase(const FObjectInitializer& ObjectInitializer)
 	: Super(ObjectInitializer)
