Network demo time dilation fixes.

* Don't override time dilation with cvar unless it's > -1 (so game code can control dilation as well)
* Support pausing
* Better handling of spectator time dilation cancelling (use bIgnoreTimeDilation on movement component)

diff --git a/Engine/Source/Runtime/Engine/Private/DemoNetDriver.cpp b/Engine/Source/Runtime/Engine/Private/DemoNetDriver.cpp
--- a/Engine/Source/Runtime/Engine/Private/DemoNetDriver.cpp
+++ b/Engine/Source/Runtime/Engine/Private/DemoNetDriver.cpp
@@ -13,11 +13,12 @@
 #include "GameFramework/SpectatorPawn.h"
 #include "Engine/LevelStreamingKismet.h"
 #include "GameFramework/PlayerStart.h"
+#include "GameFramework/SpectatorPawnMovement.h"
 
 DEFINE_LOG_CATEGORY_STATIC( LogDemo, Log, All );
 
 static TAutoConsoleVariable<float> CVarDemoRecordHz( TEXT( "demo.RecordHz" ), 10, TEXT( "Number of demo frames recorded per second" ) );
-static TAutoConsoleVariable<float> CVarDemoTimeDilation( TEXT( "demo.TimeDilation" ), 1.0f, TEXT( "Override time dilation during demo playback" ) );
+static TAutoConsoleVariable<float> CVarDemoTimeDilation( TEXT( "demo.TimeDilation" ), -1.0f, TEXT( "Override time dilation during demo playback (-1 = don't override)" ) );
 
 static const int32 MAX_DEMO_READ_WRITE_BUFFER = 1024 * 2;
 
@@ -305,7 +306,10 @@ void UDemoNetDriver::TickFlush( float DeltaSeconds )
 				}
 			}
 
-			World->GetWorldSettings()->DemoPlayTimeDilation = CVarDemoTimeDilation.GetValueOnGameThread();
+			if ( CVarDemoTimeDilation.GetValueOnGameThread() >= 0.0f )
+			{
+				World->GetWorldSettings()->DemoPlayTimeDilation = CVarDemoTimeDilation.GetValueOnGameThread();
+			}
 
 			// Clamp time between 1000 hz, and 2 hz 
 			// (this is useful when debugging and you set a breakpoint, you don't want all that time to pass in one frame)
@@ -323,8 +327,15 @@ void UDemoNetDriver::TickFlush( float DeltaSeconds )
 					// Disable collision on the spectator
 					SpectatorController->GetSpectatorPawn()->SetActorEnableCollision( false );
 					
-					// Apply time dilation on spectator to reverse the effects of global dilation
-					SpectatorController->GetSpectatorPawn()->CustomTimeDilation = 1.0f / World->GetWorldSettings()->DemoPlayTimeDilation;
+					SpectatorController->GetSpectatorPawn()->PrimaryActorTick.bTickEvenWhenPaused = true;
+
+					USpectatorPawnMovement* SpectatorMovement = Cast<USpectatorPawnMovement>(SpectatorController->GetSpectatorPawn()->GetMovementComponent());
+
+					if ( SpectatorMovement )
+					{
+						SpectatorMovement->bIgnoreTimeDilation = true;
+						SpectatorMovement->PrimaryComponentTick.bTickEvenWhenPaused = true;
+					}
 				}
 			}
 
@@ -333,7 +344,10 @@ void UDemoNetDriver::TickFlush( float DeltaSeconds )
 				return;
 			}
 
-			TickDemoPlayback( DeltaSeconds );
+			if ( World->GetWorldSettings()->Pauser == NULL )
+			{
+				TickDemoPlayback( DeltaSeconds );
+			}
 		}
 	}
 }
