#ue4 - Fix one more case where a non-normalized normal can get through collision result handling. Renamed 'TraceDir' and 'Direction' variables to be more explicit that they are nor actually normalized.

diff --git a/Engine/Source/Runtime/Engine/Private/Collision/CollisionConversions.cpp b/Engine/Source/Runtime/Engine/Private/Collision/CollisionConversions.cpp
--- a/Engine/Source/Runtime/Engine/Private/Collision/CollisionConversions.cpp
+++ b/Engine/Source/Runtime/Engine/Private/Collision/CollisionConversions.cpp
@@ -82,14 +82,14 @@ static PxVec3 TransformNormalToShapeSpace(const PxMeshScale& meshScale, const Px
 	}
 }
 
-static bool FindSimpleOpposingNormal(const PxLocationHit& PHit, const FVector& Direction, FVector& OutNormal)
+static bool FindSimpleOpposingNormal(const PxLocationHit& PHit, const FVector& TraceDirectionDenorm, FVector& OutNormal)
 {
-	bool bNormalData = PHit.flags & PxHitFlag::eNORMAL;
-	OutNormal = bNormalData ? P2UVector(PHit.normal) : -Direction;
+	const bool bNormalData = PHit.flags & PxHitFlag::eNORMAL;
+	OutNormal = bNormalData ? P2UVector(PHit.normal) : -TraceDirectionDenorm.SafeNormal();
 	return true;
 }
 
-static bool FindHeightFieldOpposingNormal(const PxLocationHit& PHit, const FVector& Direction, FVector& OutNormal)
+static bool FindHeightFieldOpposingNormal(const PxLocationHit& PHit, const FVector& TraceDirectionDenorm, FVector& OutNormal)
 {
 
 	PxHeightFieldGeometry PHeightFieldGeom;
@@ -113,7 +113,7 @@ static bool FindHeightFieldOpposingNormal(const PxLocationHit& PHit, const FVect
 	return false;
 }
 
-static bool FindConvexMeshOpposingNormal(const PxLocationHit& PHit, const FVector& Direction, FVector& OutNormal)
+static bool FindConvexMeshOpposingNormal(const PxLocationHit& PHit, const FVector& TraceDirectionDenorm, FVector& OutNormal)
 {
 	PxConvexMeshGeometry PConvexMeshGeom;
 	bool bSuccess = PHit.shape->getConvexMeshGeometry(PConvexMeshGeom);
@@ -151,7 +151,7 @@ static bool FindConvexMeshOpposingNormal(const PxLocationHit& PHit, const FVecto
 	return false;
 }
 
-static bool FindTriMeshOpposingNormal(const PxLocationHit& PHit, const FVector& Direction, FVector& OutNormal)
+static bool FindTriMeshOpposingNormal(const PxLocationHit& PHit, const FVector& TraceDirectionDenorm, FVector& OutNormal)
 {
 	PxTriangleMeshGeometry PTriMeshGeom;
 	bool bSuccess = PHit.shape->getTriangleMeshGeometry(PTriMeshGeom);
@@ -200,7 +200,7 @@ static bool FindTriMeshOpposingNormal(const PxLocationHit& PHit, const FVector&
 		if (PTriMeshGeom.meshFlags & PxMeshGeometryFlag::eDOUBLE_SIDED)
 		{
 			//double sided mesh so we need to consider direction of query
-			float sign = FVector::DotProduct(OutNormal, Direction) > 0 ? -1.f : 1.f;
+			const float sign = FVector::DotProduct(OutNormal, TraceDirectionDenorm) > 0.f ? -1.f : 1.f;
 			OutNormal *= sign;
 		}
 
@@ -220,11 +220,11 @@ static bool FindTriMeshOpposingNormal(const PxLocationHit& PHit, const FVector&
 /**
  * Util to find the normal of the face that we hit.
  * @param PHit - incoming hit from PhysX
- * @param Direction - direction of sweep test (not normalized)
+ * @param TraceDirectionDenorm - direction of sweep test (not normalized)
  * @param OutNormal - normal we may recompute based on the faceIndex of the hit
  * @return true if we compute a new normal for the geometry.
  */
-static bool FindGeomOpposingNormal(const PxLocationHit& PHit, const FVector& Direction, FVector& OutNormal)
+static bool FindGeomOpposingNormal(const PxLocationHit& PHit, const FVector& TraceDirectionDenorm, FVector& OutNormal)
 {
 	checkf(InvalidQueryHit.faceIndex == 0xFFFFffff, TEXT("Engine code needs fixing: PhysX invalid face index sentinel has changed or is not part of default PxQueryHit!"));
 	if (PHit.faceIndex == InvalidQueryHit.faceIndex)
@@ -237,10 +237,10 @@ static bool FindGeomOpposingNormal(const PxLocationHit& PHit, const FVector& Dir
 	{
 		case PxGeometryType::eSPHERE:
 		case PxGeometryType::eBOX:
-		case PxGeometryType::eCAPSULE:		return FindSimpleOpposingNormal(PHit, Direction, OutNormal);
-		case PxGeometryType::eCONVEXMESH:	return FindConvexMeshOpposingNormal(PHit, Direction, OutNormal);
-		case PxGeometryType::eHEIGHTFIELD:	return FindHeightFieldOpposingNormal(PHit, Direction, OutNormal);
-		case PxGeometryType::eTRIANGLEMESH:	return FindTriMeshOpposingNormal(PHit, Direction, OutNormal);
+		case PxGeometryType::eCAPSULE:		return FindSimpleOpposingNormal(PHit, TraceDirectionDenorm, OutNormal);
+		case PxGeometryType::eCONVEXMESH:	return FindConvexMeshOpposingNormal(PHit, TraceDirectionDenorm, OutNormal);
+		case PxGeometryType::eHEIGHTFIELD:	return FindHeightFieldOpposingNormal(PHit, TraceDirectionDenorm, OutNormal);
+		case PxGeometryType::eTRIANGLEMESH:	return FindTriMeshOpposingNormal(PHit, TraceDirectionDenorm, OutNormal);
 		default: check(false);	//unsupported geom type
 	}
 
@@ -346,13 +346,13 @@ void ConvertQueryImpactHit(const PxLocationHit& PHit, FHitResult& OutResult, flo
 	const float HitTime = PHit.distance/CheckLength;
 
 	// figure out where the the "safe" location for this shape is by moving from the startLoc toward the ImpactPoint
-	const FVector TraceDir = EndLoc - StartLoc;
-	const FVector SafeLocationToFitShape = StartLoc + (HitTime * TraceDir);
+	const FVector TraceStartToEnd = EndLoc - StartLoc;
+	const FVector SafeLocationToFitShape = StartLoc + (HitTime * TraceStartToEnd);
 
 	// Other info
 	OutResult.Location = SafeLocationToFitShape;
 	OutResult.ImpactPoint = (PHit.flags & PxHitFlag::ePOSITION) ? P2UVector(PHit.position) : StartLoc;
-	OutResult.Normal = (PHit.flags & PxHitFlag::eNORMAL) ? P2UVector(PHit.normal) : -TraceDir.SafeNormal();
+	OutResult.Normal = (PHit.flags & PxHitFlag::eNORMAL) ? P2UVector(PHit.normal) : -TraceStartToEnd.SafeNormal();
 	OutResult.ImpactNormal = OutResult.Normal;
 
 	OutResult.TraceStart = StartLoc;
@@ -375,7 +375,7 @@ void ConvertQueryImpactHit(const PxLocationHit& PHit, FHitResult& OutResult, flo
 	// Special handling for swept-capsule results
 	if (GeometryType == PxGeometryType::eCAPSULE || GeometryType == PxGeometryType::eSPHERE)
 	{
-		FindGeomOpposingNormal(PHit, TraceDir, OutResult.ImpactNormal);
+		FindGeomOpposingNormal(PHit, TraceStartToEnd, OutResult.ImpactNormal);
 	}
 	
 	if( PHit.shape->getGeometryType() == PxGeometryType::eHEIGHTFIELD)
@@ -579,7 +579,6 @@ static bool ConvertOverlappedShapeToImpactHit(const PxLocationHit& PHit, const F
 {
 	OutResult.TraceStart = StartLoc;
 	OutResult.TraceEnd = EndLoc;
-	const FVector TraceDir = EndLoc - StartLoc;
 
 	const PxShape* PShape = PHit.shape;
 	const PxRigidActor* PActor = PHit.actor;
