Fixed navmesh not generating on BSP when NavMeshBounds volume is dragged into a navmesh-less level #UE4

- fixes UE-5558

diff --git a/Engine/Source/Runtime/Engine/Private/AI/Navigation/NavigationSystem.cpp b/Engine/Source/Runtime/Engine/Private/AI/Navigation/NavigationSystem.cpp
--- a/Engine/Source/Runtime/Engine/Private/AI/Navigation/NavigationSystem.cpp
+++ b/Engine/Source/Runtime/Engine/Private/AI/Navigation/NavigationSystem.cpp
@@ -2276,6 +2276,7 @@ void UNavigationSystem::PopulateNavOctree()
 	for (int32 LevelIndex = 0; LevelIndex < World->GetNumLevels(); ++LevelIndex) 
 	{
 		ULevel* Level = World->GetLevel(LevelIndex);
+		AddLevelCollisionToOctree(Level);
 
 		for (int32 ActorIndex=0; ActorIndex<Level->Actors.Num(); ActorIndex++)
 		{
