Add priority to InputComponents

diff --git a/Engine/Source/Runtime/Engine/Classes/Components/InputComponent.h b/Engine/Source/Runtime/Engine/Classes/Components/InputComponent.h
--- a/Engine/Source/Runtime/Engine/Classes/Components/InputComponent.h
+++ b/Engine/Source/Runtime/Engine/Classes/Components/InputComponent.h
@@ -550,6 +550,9 @@ class ENGINE_API UInputComponent
 	/** The collection of gesture bindings. */
 	TArray<struct FInputGestureBinding> GestureBindings;
 
+	/** The priority of this input component when pushed in to the stack. */
+	int32 Priority;
+
 	/** Whether any components lower on the input stack should be allowed to receive input. */
 	uint32 bBlockInput:1;
 
diff --git a/Engine/Source/Runtime/Engine/Classes/GameFramework/Actor.h b/Engine/Source/Runtime/Engine/Classes/GameFramework/Actor.h
--- a/Engine/Source/Runtime/Engine/Classes/GameFramework/Actor.h
+++ b/Engine/Source/Runtime/Engine/Classes/GameFramework/Actor.h
@@ -244,6 +244,10 @@ public:
 	UPROPERTY(EditAnywhere, Category=Input)
 	TEnumAsByte<EAutoReceiveInput::Type> AutoReceiveInput;
 
+	/** The priority of this input component when pushed in to the stack. */
+	UPROPERTY(EditAnywhere, Category=Input)
+	int32 InputPriority;
+
 	/** Component that handles input for this actor, if input is enabled. */
 	UPROPERTY()
 	class UInputComponent* InputComponent;
diff --git a/Engine/Source/Runtime/Engine/Private/Actor.cpp b/Engine/Source/Runtime/Engine/Private/Actor.cpp
--- a/Engine/Source/Runtime/Engine/Private/Actor.cpp
+++ b/Engine/Source/Runtime/Engine/Private/Actor.cpp
@@ -2369,6 +2369,7 @@ void AActor::EnableInput(APlayerController* PlayerController)
 			InputComponent = ConstructObject<UInputComponent>(UInputComponent::StaticClass(), this);
 			InputComponent->RegisterComponent();
 			InputComponent->bBlockInput = bBlockInput;
+			InputComponent->Priority = InputPriority;
 
 			// Only do this if this actor is of a blueprint class
 			UBlueprintGeneratedClass* BGClass = Cast<UBlueprintGeneratedClass>(GetClass());
diff --git a/Engine/Source/Runtime/Engine/Private/PlayerController.cpp b/Engine/Source/Runtime/Engine/Private/PlayerController.cpp
--- a/Engine/Source/Runtime/Engine/Private/PlayerController.cpp
+++ b/Engine/Source/Runtime/Engine/Private/PlayerController.cpp
@@ -2346,6 +2346,10 @@ void APlayerController::BuildInputStack(TArray<UInputComponent*>& InputStack)
 		{
 			InputStack.Push(IC);
 		}
+		else
+		{
+			CurrentInputStack.RemoveAt(Idx--);
+		}
 	}
 }
 
@@ -4066,7 +4070,26 @@ void APlayerController::PushInputComponent(UInputComponent* InputComponent)
 {
 	if (InputComponent)
 	{
-		CurrentInputStack.Push(InputComponent);
+		bool bPushed = false;
+		CurrentInputStack.RemoveSingle(InputComponent);
+		for (int32 Index = CurrentInputStack.Num() - 1; Index >= 0; --Index)
+		{
+			UInputComponent* IC = CurrentInputStack[Index].Get();
+			if (IC == nullptr)
+			{
+				CurrentInputStack.RemoveAt(Index);
+			}
+			else if (IC->Priority <= InputComponent->Priority)
+			{
+				CurrentInputStack.Insert(InputComponent, Index + 1);
+				bPushed = true;
+				break;
+			}
+		}
+		if (!bPushed)
+		{
+			CurrentInputStack.Insert(InputComponent, 0);
+		}
 	}
 }
 
