#ue4 - MovementComponent and ProjectileMovementComponent auto-registration moved back to InitializeComponent to avoid registration order issues.

Jira: https://jira.ol.epicgames.net/browse/UE-5548

diff --git a/Engine/Source/Runtime/Engine/Classes/GameFramework/MovementComponent.h b/Engine/Source/Runtime/Engine/Classes/GameFramework/MovementComponent.h
--- a/Engine/Source/Runtime/Engine/Classes/GameFramework/MovementComponent.h
+++ b/Engine/Source/Runtime/Engine/Classes/GameFramework/MovementComponent.h
@@ -135,6 +135,9 @@ public:
 	virtual void PostLoad() override;
 
 	/** Overridden to auto-register the updated component if it starts NULL, and we can find a root component on our owner. */
+	virtual void InitializeComponent() override;
+
+	/** Overridden to update component properties that should be updated while being edited. */	
 	virtual void OnRegister() override;
 
 
@@ -384,6 +387,9 @@ private:
 
 	/** Transient flag indicating whether we are executing OnRegister(). */
 	bool bInOnRegister;
+	
+	/** Transient flag indicating whether we are executing InitializeComponent(). */
+	bool bInInitializeComponent;
 };
 
 
diff --git a/Engine/Source/Runtime/Engine/Classes/GameFramework/ProjectileMovementComponent.h b/Engine/Source/Runtime/Engine/Classes/GameFramework/ProjectileMovementComponent.h
--- a/Engine/Source/Runtime/Engine/Classes/GameFramework/ProjectileMovementComponent.h
+++ b/Engine/Source/Runtime/Engine/Classes/GameFramework/ProjectileMovementComponent.h
@@ -146,7 +146,7 @@ class ENGINE_API UProjectileMovementComponent : public UMovementComponent
 
 	//Begin UMovementComponent Interface
 	virtual float GetMaxSpeed() const override { return MaxSpeed; }
-	virtual void OnRegister() override;
+	virtual void InitializeComponent() override;
 	//End UMovementComponent Interface
 
 	/**
diff --git a/Engine/Source/Runtime/Engine/Private/MovementComponent.cpp b/Engine/Source/Runtime/Engine/Private/MovementComponent.cpp
--- a/Engine/Source/Runtime/Engine/Private/MovementComponent.cpp
+++ b/Engine/Source/Runtime/Engine/Private/MovementComponent.cpp
@@ -34,6 +34,7 @@ UMovementComponent::UMovementComponent(const FObjectInitializer& ObjectInitializ
 	bWantsInitializeComponent = true;
 	bAutoActivate = true;
 	bInOnRegister = false;
+	bInInitializeComponent = false;
 }
 
 
@@ -59,7 +60,7 @@ void UMovementComponent::SetUpdatedComponent(UPrimitiveComponent* NewUpdatedComp
 		UpdatedComponent->bShouldUpdatePhysicsVolume = true;
 		UpdatedComponent->PhysicsVolumeChangedDelegate.AddUniqueDynamic(this, &UMovementComponent::PhysicsVolumeChanged);
 
-		if (!bInOnRegister)
+		if (!bInOnRegister && !bInInitializeComponent)
 		{
 			// UpdateOverlaps() in component registration will take care of this.
 			UpdatedComponent->UpdatePhysicsVolume(true);
@@ -78,10 +79,10 @@ void UMovementComponent::SetUpdatedComponent(UPrimitiveComponent* NewUpdatedComp
 }
 
 
-void UMovementComponent::OnRegister()
+void UMovementComponent::InitializeComponent()
 {
-	TGuardValue<bool> InOnRegisterGuard(bInOnRegister, true);
-	Super::OnRegister();
+	TGuardValue<bool> InInitializeComponentGuard(bInInitializeComponent, true);
+	Super::InitializeComponent();
 
 	UPrimitiveComponent* NewUpdatedComponent = NULL;
 	if (UpdatedComponent != NULL)
@@ -97,7 +98,7 @@ void UMovementComponent::OnRegister()
 			NewUpdatedComponent = Cast<UPrimitiveComponent>(MyActor->GetRootComponent());
 			if (!NewUpdatedComponent)
 			{
-				FMessageLog("PIE").Warning(FText::Format(LOCTEXT("NoRootPrimitiveWarning", "Movement component {0} must update a PrimitiveComponent, but owning actor '{1}' does not have a root PrimitiveComponent. Auto registration failed."), 
+				FMessageLog("PIE").Warning(FText::Format(LOCTEXT("NoRootPrimitiveWarning", "Movement component {0} must update a PrimitiveComponent, but owning actor '{1}' does not have a root PrimitiveComponent. Auto registration failed."),
 					FText::FromString(GetName()),
 					FText::FromString(MyActor->GetName())
 					));
@@ -105,13 +106,26 @@ void UMovementComponent::OnRegister()
 		}
 	}
 
+	SetUpdatedComponent(NewUpdatedComponent);
+}
+
+
+void UMovementComponent::OnRegister()
+{
+	TGuardValue<bool> InOnRegisterGuard(bInOnRegister, true);
+	Super::OnRegister();
+
 	if (PlaneConstraintAxisSetting != EPlaneConstraintAxisSetting::Custom)
 	{
 		SetPlaneConstraintAxisSetting(PlaneConstraintAxisSetting);
 	}
 
 	PlaneConstraintNormal = PlaneConstraintNormal.SafeNormal();
-	SetUpdatedComponent(NewUpdatedComponent);
+
+	if (bSnapToPlaneAtStart)
+	{
+		SnapUpdatedComponentToPlane();
+	}
 }
 
 void UMovementComponent::RegisterComponentTickFunctions(bool bRegister)
diff --git a/Engine/Source/Runtime/Engine/Private/ProjectileMovementComponent.cpp b/Engine/Source/Runtime/Engine/Private/ProjectileMovementComponent.cpp
--- a/Engine/Source/Runtime/Engine/Private/ProjectileMovementComponent.cpp
+++ b/Engine/Source/Runtime/Engine/Private/ProjectileMovementComponent.cpp
@@ -53,9 +53,9 @@ void UProjectileMovementComponent::PostLoad()
 }
 
 
-void UProjectileMovementComponent::OnRegister()
+void UProjectileMovementComponent::InitializeComponent()
 {
-	Super::OnRegister();
+	Super::InitializeComponent();
 
 	if (Velocity.SizeSquared() > 0.f)
 	{
