UE4 - Fix for mobile instanced static meshes not rendering first element of new batches

diff --git a/Engine/Source/Runtime/Engine/Private/InstancedStaticMesh.cpp b/Engine/Source/Runtime/Engine/Private/InstancedStaticMesh.cpp
--- a/Engine/Source/Runtime/Engine/Private/InstancedStaticMesh.cpp
+++ b/Engine/Source/Runtime/Engine/Private/InstancedStaticMesh.cpp
@@ -1168,26 +1168,27 @@ void FInstancedStaticMeshSceneProxy::SetupInstancedMeshBatch(int32 LODIndex, int
 	{
 		const uint32 MaxInstancesPerBatch = FInstancedStaticMeshVertexFactory::NumBitsForVisibilityMask();
 		const uint32 NumBatches = FMath::DivideAndRoundUp(NumInstances, MaxInstancesPerBatch);
-		const uint32 NumInstancesThisBatch = BatchIndex == NumBatches - 1 ? NumInstances % MaxInstancesPerBatch : MaxInstancesPerBatch;
+		uint32 NumInstancesThisBatch = BatchIndex == NumBatches - 1 ? NumInstances % MaxInstancesPerBatch : MaxInstancesPerBatch;
 
-		OutMeshBatch.Elements.Reserve(NumInstances);
+		OutMeshBatch.Elements.Reserve(NumInstancesThisBatch);
 
-		for (uint32 Instance = 0; Instance < NumInstancesThisBatch; ++Instance)
+		int32 InstanceIndex = BatchIndex * MaxInstancesPerBatch;
+		if (NumInstancesThisBatch > 0)
 		{
-			FMeshBatchElement* NewBatchElement; 
+			// BatchElement0 is already inside the array; but Reserve() might have shifted it
+			OutMeshBatch.Elements[0].UserIndex = InstanceIndex;
+			--NumInstancesThisBatch;
+			++InstanceIndex;
 
-			if (Instance == 0)
+			// Add remaining BatchElements 1..n-1
+			while (NumInstancesThisBatch > 0)
 			{
-				NewBatchElement = &BatchElement0;
-			}
-			else
-			{
-				NewBatchElement = new(OutMeshBatch.Elements) FMeshBatchElement();
+				auto* NewBatchElement = new(OutMeshBatch.Elements) FMeshBatchElement();
 				*NewBatchElement = BatchElement0;
+				NewBatchElement->UserIndex = InstanceIndex;
+				++InstanceIndex;
+				--NumInstancesThisBatch;
 			}
-			
-			const int32 InstanceIndex = BatchIndex * MaxInstancesPerBatch + Instance;
-			NewBatchElement->UserIndex = InstanceIndex;
 		}
 	}
 }
