Clearing FPImplRecastNavMesh::DetourNavMesh in case of loading no navmesh data #UE4

diff --git a/Engine/Source/Runtime/Engine/Private/AI/Navigation/PImplRecastNavMesh.cpp b/Engine/Source/Runtime/Engine/Private/AI/Navigation/PImplRecastNavMesh.cpp
--- a/Engine/Source/Runtime/Engine/Private/AI/Navigation/PImplRecastNavMesh.cpp
+++ b/Engine/Source/Runtime/Engine/Private/AI/Navigation/PImplRecastNavMesh.cpp
@@ -18,7 +18,6 @@
 
 #include "AI/Navigation/NavLinkCustomInterface.h"
 #include "VisualLogger/VisualLogger.h"
-#include "AI/Navigation/NavAreas/NavArea.h"
 
 
 //----------------------------------------------------------------------//
@@ -251,6 +250,13 @@ FPImplRecastNavMesh::FPImplRecastNavMesh(ARecastNavMesh* Owner)
 };
 
 FPImplRecastNavMesh::~FPImplRecastNavMesh()
+{
+	ReleaseDetourNavMesh();
+
+	DEC_DWORD_STAT_BY( STAT_NavigationMemory, sizeof(*this) );
+};
+
+void FPImplRecastNavMesh::ReleaseDetourNavMesh()
 {
 	// release navmesh only if we own it
 	if (DetourNavMesh != nullptr)
@@ -258,9 +264,7 @@ FPImplRecastNavMesh::~FPImplRecastNavMesh()
 		dtFreeNavMesh(DetourNavMesh);
 	}
 	DetourNavMesh = nullptr;
-
-	DEC_DWORD_STAT_BY( STAT_NavigationMemory, sizeof(*this) );
-};
+}
 
 /**
  * Serialization.
diff --git a/Engine/Source/Runtime/Engine/Private/AI/Navigation/RecastNavMesh.cpp b/Engine/Source/Runtime/Engine/Private/AI/Navigation/RecastNavMesh.cpp
--- a/Engine/Source/Runtime/Engine/Private/AI/Navigation/RecastNavMesh.cpp
+++ b/Engine/Source/Runtime/Engine/Private/AI/Navigation/RecastNavMesh.cpp
@@ -681,6 +681,8 @@ void ARecastNavMesh::Serialize( FArchive& Ar )
 		{
 			// empty, just skip over this data
 			Ar.Seek( RecastNavMeshSizePos + RecastNavMeshSizeBytes );
+			// if it's not getting filled it's better to just remove it
+			RecastNavMeshImpl->ReleaseDetourNavMesh();
 		}
 	}
 	else
diff --git a/Engine/Source/Runtime/Engine/Public/AI/Navigation/PImplRecastNavMesh.h b/Engine/Source/Runtime/Engine/Public/AI/Navigation/PImplRecastNavMesh.h
--- a/Engine/Source/Runtime/Engine/Public/AI/Navigation/PImplRecastNavMesh.h
+++ b/Engine/Source/Runtime/Engine/Public/AI/Navigation/PImplRecastNavMesh.h
@@ -184,6 +184,7 @@ public:
 public:
 	class dtNavMesh const* GetRecastMesh() const { return DetourNavMesh; };
 	class dtNavMesh* GetRecastMesh() { return DetourNavMesh; };
+	void ReleaseDetourNavMesh();
 
 	bool GetOwnsNavMeshData() const { return bOwnsNavMeshData; }
 
diff --git a/Engine/Source/Runtime/Navmesh/Public/Detour/DetourNavMesh.h b/Engine/Source/Runtime/Navmesh/Public/Detour/DetourNavMesh.h
--- a/Engine/Source/Runtime/Navmesh/Public/Detour/DetourNavMesh.h
+++ b/Engine/Source/Runtime/Navmesh/Public/Detour/DetourNavMesh.h
@@ -726,6 +726,11 @@ public:
 
 	/// Helper for creating links in off-mesh connections
 	void linkOffMeshHelper(dtMeshTile* tile0, unsigned int polyIdx0, dtMeshTile* tile1, unsigned int polyIdx1, unsigned char side, unsigned char edge);
+
+	inline bool isEmpty() const
+	{
+		return (m_tileWidth > 0 && m_tileHeight > 0) == false;
+	}
 	//@UE4 END
 	
 private:
