Added facility to change the game window title at runtime and in project settings.

diff --git a/Engine/Source/Runtime/Engine/Classes/Kismet/KismetSystemLibrary.h b/Engine/Source/Runtime/Engine/Classes/Kismet/KismetSystemLibrary.h
--- a/Engine/Source/Runtime/Engine/Classes/Kismet/KismetSystemLibrary.h
+++ b/Engine/Source/Runtime/Engine/Classes/Kismet/KismetSystemLibrary.h
@@ -206,7 +206,11 @@ class UKismetSystemLibrary : public UBlueprintFunctionLibrary
 	UFUNCTION(BlueprintCallable, meta=(BlueprintInternalUseOnly = "TRUE"))
 	static void PrintWarning(const FString& InString);
 
-	/** 
+	/** Sets the game window title */
+	UFUNCTION(BlueprintCallable, Category = "Utilities")
+	static void SetWindowTitle(const FText& Title);
+
+	/**
 	 * Executes a console command, optionally on a specific controller
 	 * 
 	 * @param	Command			Command to send to the console
diff --git a/Engine/Source/Runtime/Engine/Private/GameEngine.cpp b/Engine/Source/Runtime/Engine/Private/GameEngine.cpp
--- a/Engine/Source/Runtime/Engine/Private/GameEngine.cpp
+++ b/Engine/Source/Runtime/Engine/Private/GameEngine.cpp
@@ -17,6 +17,7 @@
 #include "Engine/GameInstance.h"
 #include "Engine/RendererSettings.h"
 #include "Engine/UserInterfaceSettings.h"
+#include "GeneralProjectSettings.h"
 #include "AVIWriter.h"
 
 #include "SlateBasics.h"
@@ -263,6 +264,7 @@ TSharedRef<SWindow> UGameEngine::CreateGameWindow()
 	Args.Add( TEXT("RHIName"), FText::FromName( LegacyShaderPlatformToShaderFormat( GMaxRHIShaderPlatform ) ) );
 
 	const FText AppName = FText::Format( NSLOCTEXT("UnrealEd", "GameWindowTitle", "{GameName} ({PlatformArchitecture}-bit, {RHIName})"), Args );
+	const FText WindowTitleOverride = GetDefault<UGeneralProjectSettings>()->ProjectDisplayedTitle;
 
 	// Allow optional winX/winY parameters to set initial window position
 	EAutoCenter::Type AutoCenterType = EAutoCenter::PrimaryWorkArea;
@@ -275,7 +277,7 @@ TSharedRef<SWindow> UGameEngine::CreateGameWindow()
 
 	TSharedRef<SWindow> Window = SNew(SWindow)
 	.ClientSize(FVector2D( ResX, ResY ))
-	.Title( AppName )
+	.Title(WindowTitleOverride.IsEmpty() ? AppName : WindowTitleOverride)
 	.AutoCenter(AutoCenterType)
 	.ScreenPosition(FVector2D(WinX, WinY))
 	.FocusWhenFirstShown(true)
diff --git a/Engine/Source/Runtime/Engine/Private/KismetSystemLibrary.cpp b/Engine/Source/Runtime/Engine/Private/KismetSystemLibrary.cpp
--- a/Engine/Source/Runtime/Engine/Private/KismetSystemLibrary.cpp
+++ b/Engine/Source/Runtime/Engine/Private/KismetSystemLibrary.cpp
@@ -1,6 +1,7 @@
 // Copyright 1998-2014 Epic Games, Inc. All Rights Reserved.
 
 #include "EnginePrivate.h"
+#include "Engine/GameEngine.h"
 #include "Kismet/KismetSystemLibrary.h"
 #include "Engine/Console.h"
 #include "LatentActions.h"
@@ -218,6 +219,19 @@ void UKismetSystemLibrary::PrintWarning(const FString& InString)
 	PrintString(NULL, InString, true, true);
 }
 
+void UKismetSystemLibrary::SetWindowTitle(const FText& Title)
+{
+	UGameEngine* GameEngine = Cast<UGameEngine>(GEngine);
+	if (GameEngine != nullptr)
+	{
+		TSharedPtr<SWindow> GameViewportWindow = GameEngine->GameViewportWindow.Pin();
+		if (GameViewportWindow.IsValid())
+		{
+			GameViewportWindow->SetTitle(Title);
+		}
+	}
+}
+
 void UKismetSystemLibrary::ExecuteConsoleCommand(UObject* WorldContextObject, const FString& Command, APlayerController* Player)
 {
 	// First, try routing through the primary player
diff --git a/Engine/Source/Runtime/EngineSettings/Classes/GeneralProjectSettings.h b/Engine/Source/Runtime/EngineSettings/Classes/GeneralProjectSettings.h
--- a/Engine/Source/Runtime/EngineSettings/Classes/GeneralProjectSettings.h
+++ b/Engine/Source/Runtime/EngineSettings/Classes/GeneralProjectSettings.h
@@ -50,4 +50,8 @@ class ENGINESETTINGS_API UGeneralProjectSettings
 	/** The project's support contact information. */
 	UPROPERTY(config, EditAnywhere, Category=Publisher)
 	FString SupportContact;
+
+	/** The project's title as displayed on the window title bar */
+	UPROPERTY(config, EditAnywhere, Category=Displayed)
+	FText ProjectDisplayedTitle;
 };
