Allow surround sounds to be created from single files

Run through existing channels before cooking to check which exist and work out which channels would be required to create an asset we support.
When retrieving data for cooking, add blank channels if they are required.

diff --git a/Engine/Source/Runtime/Engine/Private/AudioDerivedData.cpp b/Engine/Source/Runtime/Engine/Private/AudioDerivedData.cpp
--- a/Engine/Source/Runtime/Engine/Private/AudioDerivedData.cpp
+++ b/Engine/Source/Runtime/Engine/Private/AudioDerivedData.cpp
@@ -718,40 +718,89 @@ void CookSurroundWave( USoundWave* SoundWave, FName FormatName, const IAudioForm
 {
 	check(!Output.Num());
 #if WITH_EDITORONLY_DATA
-	int32						i, ChannelCount;
-	uint32					SampleDataSize;
+	int32					i;
+	uint32					SampleDataSize = 0;
 	FWaveModInfo			WaveInfo;
 	TArray<TArray<uint8> >	SourceBuffers;
+	TArray<int32>			RequiredChannels;
 
 	uint8* RawWaveData = ( uint8* )SoundWave->RawData.Lock( LOCK_READ_ONLY );
 
 	// Front left channel is the master
-	ChannelCount = 1;
 	static_assert(SPEAKER_FrontLeft == 0, "Front-left speaker must be first.");
-	if( WaveInfo.ReadWaveHeader( RawWaveData, SoundWave->ChannelSizes[ SPEAKER_FrontLeft ], SoundWave->ChannelOffsets[ SPEAKER_FrontLeft ] ) )
+
+	// loop through channels to find which have data and which are required
+	for (i = 0; i < SPEAKER_Count; i++)
 	{
+		FWaveModInfo WaveInfoInner;
+
+		// Only mono files allowed
+		if (WaveInfoInner.ReadWaveHeader(RawWaveData, SoundWave->ChannelSizes[i], SoundWave->ChannelOffsets[i])
+			&& *WaveInfoInner.pChannels == 1)
 		{
-			TArray<uint8>& Input = *new (SourceBuffers) TArray<uint8>;
-			Input.AddUninitialized(WaveInfo.SampleDataSize);
-			FMemory::Memcpy(Input.GetData(), WaveInfo.SampleDataStart, WaveInfo.SampleDataSize);
+			if (SampleDataSize == 0)
+			{
+				// keep wave info/size of first channel data we find
+				WaveInfo = WaveInfoInner;
+				SampleDataSize = WaveInfo.SampleDataSize;
+			}
+			switch (i)
+			{
+				case SPEAKER_FrontLeft:
+				case SPEAKER_FrontRight:
+				case SPEAKER_LeftSurround:
+				case SPEAKER_RightSurround:
+					// Must have quadraphonic surround channels
+					RequiredChannels.AddUnique(SPEAKER_FrontLeft);
+					RequiredChannels.AddUnique(SPEAKER_FrontRight);
+					RequiredChannels.AddUnique(SPEAKER_LeftSurround);
+					RequiredChannels.AddUnique(SPEAKER_RightSurround);
+					break;
+				case SPEAKER_FrontCenter:
+				case SPEAKER_LowFrequency:
+					// Must have 5.1 surround channels
+					for (int32 Channel = SPEAKER_FrontLeft; Channel <= SPEAKER_RightSurround; Channel++)
+					{
+						RequiredChannels.AddUnique(Channel);
+					}
+					break;
+				case SPEAKER_LeftBack:
+				case SPEAKER_RightBack:
+					// Must have all previous channels
+					for (int32 Channel = 0; Channel < i; Channel++)
+					{
+						RequiredChannels.AddUnique(Channel);
+					}
+					break;
+				default:
+					// unsupported channel count
+					break;
+			}
 		}
+	}
 
-		SampleDataSize = WaveInfo.SampleDataSize;
-		// Extract all the info for the other channels (may be blank)
-		for( i = 1; i < SPEAKER_Count; i++ )
+	if (SampleDataSize != 0)
+	{
+		int32 ChannelCount = 0;
+		// Extract all the info for channels or insert blank data
+		for( i = 0; i < SPEAKER_Count; i++ )
 		{
 			FWaveModInfo WaveInfoInner;
-			if( WaveInfoInner.ReadWaveHeader( RawWaveData, SoundWave->ChannelSizes[ i ], SoundWave->ChannelOffsets[ i ] ) )
+			if( WaveInfoInner.ReadWaveHeader( RawWaveData, SoundWave->ChannelSizes[ i ], SoundWave->ChannelOffsets[ i ] )
+				&& *WaveInfoInner.pChannels == 1 )
 			{
-				// Only mono files allowed
-				if( *WaveInfoInner.pChannels == 1 )
-				{
-					ChannelCount++;
-					TArray<uint8>& Input = *new (SourceBuffers) TArray<uint8>;
-					Input.AddUninitialized(WaveInfoInner.SampleDataSize);
-					FMemory::Memcpy(Input.GetData(), WaveInfoInner.SampleDataStart, WaveInfoInner.SampleDataSize);
-					SampleDataSize = FMath::Min<uint32>(WaveInfoInner.SampleDataSize, SampleDataSize);
-				}
+				ChannelCount++;
+				TArray<uint8>& Input = *new (SourceBuffers) TArray<uint8>;
+				Input.AddUninitialized(WaveInfoInner.SampleDataSize);
+				FMemory::Memcpy(Input.GetData(), WaveInfoInner.SampleDataStart, WaveInfoInner.SampleDataSize);
+				SampleDataSize = FMath::Min<uint32>(WaveInfoInner.SampleDataSize, SampleDataSize);
+			}
+			else if (RequiredChannels.Contains(i))
+			{
+				// Add an empty channel for cooking
+				ChannelCount++;
+				TArray<uint8>& Input = *new (SourceBuffers) TArray<uint8>;
+				Input.AddZeroed(SampleDataSize);
 			}
 		}
 	
