Allowing Audio Volume Ambient zone settings to be changed in realtime

Updated Listener function so that it checks whether the settings are different as well as the volume.
Added comparison operators to Interior Settings.
Exposed Reverb and Interior settings so that you can create them in blueprints.

diff --git a/Engine/Source/Runtime/Engine/Classes/Sound/AudioVolume.h b/Engine/Source/Runtime/Engine/Classes/Sound/AudioVolume.h
--- a/Engine/Source/Runtime/Engine/Classes/Sound/AudioVolume.h
+++ b/Engine/Source/Runtime/Engine/Classes/Sound/AudioVolume.h
@@ -44,7 +44,7 @@ enum ReverbPreset
 };
 
 /** Struct encapsulating settings for reverb effects. */
-USTRUCT()
+USTRUCT(BlueprintType)
 struct FReverbSettings
 {
 	GENERATED_USTRUCT_BODY()
@@ -95,7 +95,7 @@ struct TStructOpsTypeTraits<FReverbSettings> : public TStructOpsTypeTraitsBase
 
 /** Struct encapsulating settings for interior areas. */
 //@warning: manually mirrored in Components.h
-USTRUCT()
+USTRUCT(BlueprintType)
 struct FInteriorSettings
 {
 	GENERATED_USTRUCT_BODY()
@@ -148,7 +148,9 @@ struct FInteriorSettings
 		, InteriorLPFTime(0.5f)
 		{
 		}
-	
+
+	bool operator==(const FInteriorSettings& Other) const;
+	bool operator!=(const FInteriorSettings& Other) const;
 };
 
 UCLASS(hidecategories=(Advanced, Attachment, Collision, Volume))
diff --git a/Engine/Source/Runtime/Engine/Private/AudioDevice.cpp b/Engine/Source/Runtime/Engine/Private/AudioDevice.cpp
--- a/Engine/Source/Runtime/Engine/Private/AudioDevice.cpp
+++ b/Engine/Source/Runtime/Engine/Private/AudioDevice.cpp
@@ -1340,7 +1340,7 @@ void FAudioDevice::InvalidateCachedInteriorVolumes() const
 
 void FListener::ApplyInteriorSettings( class AAudioVolume* InVolume, const FInteriorSettings& Settings )
 {
-	if( InVolume != Volume )
+	if( InVolume != Volume || Settings != InteriorSettings)
 	{
 		// Use previous/ current interpolation time if we're transitioning to the default worldsettings zone.
 		InteriorStartTime = FApp::GetCurrentTime();
diff --git a/Engine/Source/Runtime/Engine/Private/AudioVolume.cpp b/Engine/Source/Runtime/Engine/Private/AudioVolume.cpp
--- a/Engine/Source/Runtime/Engine/Private/AudioVolume.cpp
+++ b/Engine/Source/Runtime/Engine/Private/AudioVolume.cpp
@@ -120,6 +120,20 @@ void FReverbSettings::PostSerialize(const FArchive& Ar)
 	}
 }
 
+bool FInteriorSettings::operator==(const FInteriorSettings& Other) const
+{
+	return (Other.bIsWorldSettings == bIsWorldSettings)
+		&& (Other.ExteriorVolume == ExteriorVolume) && (Other.ExteriorTime == ExteriorTime)
+		&& (Other.ExteriorLPF == ExteriorLPF) && (Other.ExteriorLPFTime == ExteriorLPFTime)
+		&& (Other.InteriorVolume == InteriorVolume) && (Other.InteriorTime == InteriorTime)
+		&& (Other.InteriorLPF == InteriorLPF) && (Other.InteriorLPFTime == InteriorLPFTime);
+}
+
+bool FInteriorSettings::operator!=(const FInteriorSettings& Other) const
+{
+	return !(*this == Other);
+}
+
 AAudioVolume::AAudioVolume(const FObjectInitializer& ObjectInitializer)
 	: Super(ObjectInitializer)
 {
