UE-4933 Draw wireframe overlay for meshes using 'complex as simple' when showing collision

diff --git a/Engine/Source/Runtime/Engine/Private/StaticMeshRender.cpp b/Engine/Source/Runtime/Engine/Private/StaticMeshRender.cpp
--- a/Engine/Source/Runtime/Engine/Private/StaticMeshRender.cpp
+++ b/Engine/Source/Runtime/Engine/Private/StaticMeshRender.cpp
@@ -574,9 +574,10 @@ void FStaticMeshSceneProxy::GetDynamicMeshElements(const TArray<const FSceneView
 
 	bool bDrawSimpleCollision = false, bDrawComplexCollision = false;
 	const bool bInCollisionView = IsCollisionView(EngineShowFlags, bDrawSimpleCollision, bDrawComplexCollision);
-	// Draw simple collision as wireframe if 'show collision', collision is enabled, and we are not using the complex as the simple
-	const bool bDrawWireframeCollision = (EngineShowFlags.Collision && IsCollisionEnabled() && CollisionTraceFlag != ECollisionTraceFlag::CTF_UseComplexAsSimple);
 	
+	// Should we draw the mesh wireframe to indicate we are using the mesh as collision
+	const bool bDrawComplexWireframeCollision = (EngineShowFlags.Collision && IsCollisionEnabled() && CollisionTraceFlag == ECollisionTraceFlag::CTF_UseComplexAsSimple);
+
 	const bool bDrawMesh = (bInCollisionView) ? (bDrawComplexCollision) : 
 		(	IsRichView(ViewFamily) || HasViewDependentDPG()
 			|| EngineShowFlags.Collision
@@ -703,18 +704,34 @@ void FStaticMeshSceneProxy::GetDynamicMeshElements(const TArray<const FSceneView
 										Collector.RegisterOneFrameMaterialProxy(VertexColorVisualizationMaterialInstance);
 										MeshElement.MaterialRenderProxy = VertexColorVisualizationMaterialInstance;
 									}
-									else if( bInCollisionView && bDrawComplexCollision && AllowDebugViewmodes() )
+									else if( bDrawComplexWireframeCollision || (bInCollisionView && bDrawComplexCollision && AllowDebugViewmodes()) )
 									{
 										if (LODModel.Sections[SectionIndex].bEnableCollision)
 										{
+											UMaterial* MaterialToUse = (bDrawComplexWireframeCollision) ? GEngine->WireframeMaterial : GEngine->ShadedLevelColorationUnlitMaterial;
+
 											// Override the mesh's material with our material that draws the collision color
 											auto CollisionMaterialInstance = new FColoredMaterialRenderProxy(
-												GEngine->ShadedLevelColorationUnlitMaterial->GetRenderProxy(bProxyIsSelected, IsHovered()),
+												MaterialToUse->GetRenderProxy(bProxyIsSelected, IsHovered()),
 												WireframeColor
 												);
 
 											Collector.RegisterOneFrameMaterialProxy(CollisionMaterialInstance);
-											MeshElement.MaterialRenderProxy = CollisionMaterialInstance;
+
+											// If drawing the complex collision in wireframe, we do that in _addition_ to the normal mesh
+											if(bDrawComplexWireframeCollision)
+											{
+												FMeshBatch& CollisionElement = Collector.AllocateMesh();
+												GetMeshElement(LODIndex, BatchIndex, SectionIndex, SDPG_World, bSectionIsSelected, IsHovered(), CollisionElement);
+												CollisionElement.MaterialRenderProxy = CollisionMaterialInstance;
+												Collector.AddMesh(ViewIndex, CollisionElement);
+												INC_DWORD_STAT_BY(STAT_StaticMeshTriangles, CollisionElement.GetNumPrimitives());
+											}
+											// Otherwise just change material on existing mesh
+											else
+											{
+												MeshElement.MaterialRenderProxy = CollisionMaterialInstance;
+											}
 										}
 										else
 										{
@@ -757,7 +774,11 @@ void FStaticMeshSceneProxy::GetDynamicMeshElements(const TArray<const FSceneView
 	{
 		if (VisibilityMap & (1 << ViewIndex))
 		{
-			if((bDrawSimpleCollision || bDrawWireframeCollision) && AllowDebugViewmodes())
+			// Draw simple collision as wireframe if 'show collision', collision is enabled, and we are not using the complex as the simple
+			const bool bDrawSimpleWireframeCollision = (EngineShowFlags.Collision && IsCollisionEnabled() && CollisionTraceFlag != ECollisionTraceFlag::CTF_UseComplexAsSimple);
+
+
+			if((bDrawSimpleCollision || bDrawSimpleWireframeCollision) && AllowDebugViewmodes())
 			{
 				if(BodySetup)
 				{
@@ -769,7 +790,7 @@ void FStaticMeshSceneProxy::GetDynamicMeshElements(const TArray<const FSceneView
 					}
 					else
 					{
-						const bool bDrawSolid = !bDrawWireframeCollision;
+						const bool bDrawSolid = !bDrawSimpleWireframeCollision;
 
 						if(bDrawSolid)
 						{
