Improved material instance base property overrides.

- Improved UI. Removed bOverrideBaseProperties as it was pointless. Altered customization to place check boxes for overrides next to their values like the parameters of instances.
- Overrides now respect the parent chain.
- Mirrored the override data in the MIC render proxy.

diff --git a/Engine/Shaders/ShaderVersion.usf b/Engine/Shaders/ShaderVersion.usf
--- a/Engine/Shaders/ShaderVersion.usf
+++ b/Engine/Shaders/ShaderVersion.usf
@@ -2,4 +2,4 @@
 // Each time the console command is executed it generates a new GUID. As this file is included
 // in common.usf (which should be included in any shader) it allows to invalidate the shader DDC.
 // 
-// GUID = 4C36424F4976EF3266512DACC6E7D40A
\ No newline at end of file
+// GUID = 48B5ECE04311A6B3C8BB15A2A11716CD
\ No newline at end of file
diff --git a/Engine/Source/Editor/MaterialEditor/Private/MaterialEditorInstanceDetailCustomization.cpp b/Engine/Source/Editor/MaterialEditor/Private/MaterialEditorInstanceDetailCustomization.cpp
--- a/Engine/Source/Editor/MaterialEditor/Private/MaterialEditorInstanceDetailCustomization.cpp
+++ b/Engine/Source/Editor/MaterialEditor/Private/MaterialEditorInstanceDetailCustomization.cpp
@@ -49,13 +49,6 @@ void FMaterialInstanceParameterDetails::OnValueCommitted(float NewValue, ETextCo
 
 void FMaterialInstanceParameterDetails::CustomizeDetails(IDetailLayoutBuilder& DetailLayout)
 {
-	// Create a new category for a custom layout for the MIC parameters at the very top
-	FName GroupsCategoryName = TEXT("ParameterGroups");
-	IDetailCategoryBuilder& GroupsCategory = DetailLayout.EditCategory(GroupsCategoryName, LOCTEXT("MICParamGroupsTitle", "Parameter Groups").ToString());
-	TSharedRef<IPropertyHandle> ParameterGroupsProperty = DetailLayout.GetProperty("ParameterGroups");
-
-	CreateGroupsWidget(ParameterGroupsProperty, GroupsCategory);
-
 	// Create default category for class properties
 	const FName DefaultCategoryName = NAME_None;
 	IDetailCategoryBuilder& DefaultCategory = DetailLayout.EditCategory(DefaultCategoryName);
@@ -95,7 +88,6 @@ void FMaterialInstanceParameterDetails::CustomizeDetails(IDetailLayoutBuilder& D
 	DefaultCategory.AddProperty("LightmassSettings");
 	DetailLayout.HideProperty("bUseOldStyleMICEditorGroups");
 	DetailLayout.HideProperty("ParameterGroups");
-	DetailLayout.HideProperty("BasePropertyOverrides");
 
 	{
 		IDetailPropertyRow& PropertyRow = DefaultCategory.AddProperty("RefractionDepthBias");
@@ -112,11 +104,15 @@ void FMaterialInstanceParameterDetails::CustomizeDetails(IDetailLayoutBuilder& D
 		PropertyRow.Visibility(TAttribute<EVisibility>::Create(TAttribute<EVisibility>::FGetter::CreateSP(this, &FMaterialInstanceParameterDetails::ShouldShowSubsurfaceProfile)));
 	}
 
-	//////////////////////////////////////////////////////////////////////////
 	DetailLayout.HideProperty("BasePropertyOverrides");
-	IDetailCategoryBuilder& MaterialCategory = DetailLayout.EditCategory(TEXT("MaterialOverrides"), LOCTEXT("MICMaterialOverridesTitle", "Material Overrides").ToString());
-	MaterialCategory.AddProperty("bOverrideBaseProperties");
-	MaterialCategory.AddProperty("BasePropertyOverrides");
+	CreateBasePropertyOverrideWidgets(DetailLayout);
+
+	// Create a new category for a custom layout for the MIC parameters at the very top
+	FName GroupsCategoryName = TEXT("ParameterGroups");
+	IDetailCategoryBuilder& GroupsCategory = DetailLayout.EditCategory(GroupsCategoryName, LOCTEXT("MICParamGroupsTitle", "Parameter Groups").ToString());
+	TSharedRef<IPropertyHandle> ParameterGroupsProperty = DetailLayout.GetProperty("ParameterGroups");
+
+	CreateGroupsWidget(ParameterGroupsProperty, GroupsCategory);
 }
 
 void FMaterialInstanceParameterDetails::CreateGroupsWidget(TSharedRef<IPropertyHandle> ParameterGroupsProperty, IDetailCategoryBuilder& GroupsCategory)
@@ -442,10 +438,101 @@ EVisibility FMaterialInstanceParameterDetails::ShouldShowMaterialRefractionSetti
 
 EVisibility FMaterialInstanceParameterDetails::ShouldShowSubsurfaceProfile() const
 {
-	EMaterialShadingModel Model = MaterialEditorInstance->SourceInstance->GetMaterial()->GetShadingModel_Internal();
+	EMaterialShadingModel Model = MaterialEditorInstance->SourceInstance->GetShadingModel();
 
 	return (Model == MSM_SubsurfaceProfile) ? EVisibility::Visible : EVisibility::Collapsed;
 }
 
+
+void FMaterialInstanceParameterDetails::CreateBasePropertyOverrideWidgets(IDetailLayoutBuilder& DetailLayout)
+{
+	IDetailCategoryBuilder& DetailCategory = DetailLayout.EditCategory(NAME_None);
+	
+	static FName GroupName(TEXT("BasePropertyOverrideGroup"));
+	IDetailGroup& BasePropertyOverrideGroup = DetailCategory.AddGroup(GroupName, LOCTEXT("BasePropertyOverrideGroup", "Material Property Overrides").ToString(), false, false);
+
+	TAttribute<bool> IsOverrideOpacityClipMaskValueEnabled = TAttribute<bool>::Create(TAttribute<bool>::FGetter::CreateSP(this, &FMaterialInstanceParameterDetails::OverrideOpacityClipMaskValueEnabled));
+	TAttribute<bool> IsOverrideBlendModeEnabled = TAttribute<bool>::Create(TAttribute<bool>::FGetter::CreateSP(this, &FMaterialInstanceParameterDetails::OverrideBlendModeEnabled));
+	TAttribute<bool> IsOverrideShadingModelEnabled = TAttribute<bool>::Create(TAttribute<bool>::FGetter::CreateSP(this, &FMaterialInstanceParameterDetails::OverrideShadingModelEnabled));
+	TAttribute<bool> IsOverrideTwoSidedEnabled = TAttribute<bool>::Create(TAttribute<bool>::FGetter::CreateSP(this, &FMaterialInstanceParameterDetails::OverrideTwoSidedEnabled));
+
+	TSharedRef<IPropertyHandle> BasePropertyOverridePropery = DetailLayout.GetProperty("BasePropertyOverrides");
+	TSharedPtr<IPropertyHandle> OpacityClipMaskValueProperty = BasePropertyOverridePropery->GetChildHandle("OpacityMaskClipValue");
+	TSharedPtr<IPropertyHandle> BlendModeProperty = BasePropertyOverridePropery->GetChildHandle("BlendMode");
+	TSharedPtr<IPropertyHandle> ShadingModelProperty = BasePropertyOverridePropery->GetChildHandle("ShadingModel");
+	TSharedPtr<IPropertyHandle> TwoSidedProperty = BasePropertyOverridePropery->GetChildHandle("TwoSided");
+
+	IDetailPropertyRow& OpacityClipMaskValuePropertyRow = BasePropertyOverrideGroup.AddPropertyRow(OpacityClipMaskValueProperty.ToSharedRef());
+	OpacityClipMaskValuePropertyRow
+		.DisplayName(OpacityClipMaskValueProperty->GetPropertyDisplayName())
+		.ToolTip(OpacityClipMaskValueProperty->GetToolTipText())
+		.EditCondition(IsOverrideOpacityClipMaskValueEnabled, FOnBooleanValueChanged::CreateSP(this, &FMaterialInstanceParameterDetails::OnOverrideOpacityClipMaskValueChanged));
+
+	IDetailPropertyRow& BlendModePropertyRow = BasePropertyOverrideGroup.AddPropertyRow(BlendModeProperty.ToSharedRef());
+	BlendModePropertyRow
+		.DisplayName(BlendModeProperty->GetPropertyDisplayName())
+		.ToolTip(BlendModeProperty->GetToolTipText())
+		.EditCondition(IsOverrideBlendModeEnabled, FOnBooleanValueChanged::CreateSP(this, &FMaterialInstanceParameterDetails::OnOverrideBlendModeChanged));
+
+	IDetailPropertyRow& ShadingModelPropertyRow = BasePropertyOverrideGroup.AddPropertyRow(ShadingModelProperty.ToSharedRef());
+	ShadingModelPropertyRow
+		.DisplayName(ShadingModelProperty->GetPropertyDisplayName())
+		.ToolTip(ShadingModelProperty->GetToolTipText())
+		.EditCondition(IsOverrideShadingModelEnabled, FOnBooleanValueChanged::CreateSP(this, &FMaterialInstanceParameterDetails::OnOverrideShadingModelChanged));
+
+	IDetailPropertyRow& TwoSidedPropertyRow = BasePropertyOverrideGroup.AddPropertyRow(TwoSidedProperty.ToSharedRef());
+	TwoSidedPropertyRow
+		.DisplayName(TwoSidedProperty->GetPropertyDisplayName())
+		.ToolTip(TwoSidedProperty->GetToolTipText())
+		.EditCondition(IsOverrideTwoSidedEnabled, FOnBooleanValueChanged::CreateSP(this, &FMaterialInstanceParameterDetails::OnOverrideTwoSidedChanged));
+}
+
+bool FMaterialInstanceParameterDetails::OverrideOpacityClipMaskValueEnabled() const
+{
+	return MaterialEditorInstance->BasePropertyOverrides.bOverride_OpacityMaskClipValue;
+}
+
+bool FMaterialInstanceParameterDetails::OverrideBlendModeEnabled() const
+{
+	return MaterialEditorInstance->BasePropertyOverrides.bOverride_BlendMode;
+}
+
+bool FMaterialInstanceParameterDetails::OverrideShadingModelEnabled() const
+{
+	return MaterialEditorInstance->BasePropertyOverrides.bOverride_ShadingModel;
+}
+
+bool FMaterialInstanceParameterDetails::OverrideTwoSidedEnabled() const
+{
+	return MaterialEditorInstance->BasePropertyOverrides.bOverride_TwoSided;
+}
+
+void FMaterialInstanceParameterDetails::OnOverrideOpacityClipMaskValueChanged(bool NewValue)
+{
+	MaterialEditorInstance->BasePropertyOverrides.bOverride_OpacityMaskClipValue = NewValue;
+	MaterialEditorInstance->PostEditChange();
+	FEditorSupportDelegates::RedrawAllViewports.Broadcast();
+}
+
+void FMaterialInstanceParameterDetails::OnOverrideBlendModeChanged(bool NewValue)
+{
+	MaterialEditorInstance->BasePropertyOverrides.bOverride_BlendMode = NewValue;
+	MaterialEditorInstance->PostEditChange();
+	FEditorSupportDelegates::RedrawAllViewports.Broadcast();
+}
+
+void FMaterialInstanceParameterDetails::OnOverrideShadingModelChanged(bool NewValue)
+{
+	MaterialEditorInstance->BasePropertyOverrides.bOverride_ShadingModel = NewValue;
+	MaterialEditorInstance->PostEditChange();
+	FEditorSupportDelegates::RedrawAllViewports.Broadcast();
+}
+
+void FMaterialInstanceParameterDetails::OnOverrideTwoSidedChanged(bool NewValue)
+{
+	MaterialEditorInstance->BasePropertyOverrides.bOverride_TwoSided = NewValue;
+	MaterialEditorInstance->PostEditChange();
+	FEditorSupportDelegates::RedrawAllViewports.Broadcast();
+}
 #undef LOCTEXT_NAMESPACE
 
diff --git a/Engine/Source/Editor/MaterialEditor/Private/MaterialEditorInstanceDetailCustomization.h b/Engine/Source/Editor/MaterialEditor/Private/MaterialEditorInstanceDetailCustomization.h
--- a/Engine/Source/Editor/MaterialEditor/Private/MaterialEditorInstanceDetailCustomization.h
+++ b/Engine/Source/Editor/MaterialEditor/Private/MaterialEditorInstanceDetailCustomization.h
@@ -64,6 +64,22 @@ private:
 	/** Returns true if the refraction options should be displayed */
 	EVisibility ShouldShowSubsurfaceProfile() const;
 
+	
+	//Functions supporting BasePropertyOverrides
+
+	/** Creates all the base property override widgets. */
+	void CreateBasePropertyOverrideWidgets(IDetailLayoutBuilder& DetailLayout);
+
+	bool OverrideOpacityClipMaskValueEnabled() const;
+	bool OverrideBlendModeEnabled() const;
+	bool OverrideShadingModelEnabled() const;
+	bool OverrideTwoSidedEnabled() const;
+	void OnOverrideOpacityClipMaskValueChanged(bool NewValue);
+	void OnOverrideBlendModeChanged(bool NewValue);
+	void OnOverrideShadingModelChanged(bool NewValue);
+	void OnOverrideTwoSidedChanged(bool NewValue);
+
+
 private:
 	/** Object that stores all of the possible parameters we can edit */
 	UMaterialEditorInstanceConstant* MaterialEditorInstance;
diff --git a/Engine/Source/Editor/UnrealEd/Classes/MaterialEditor/MaterialEditorInstanceConstant.h b/Engine/Source/Editor/UnrealEd/Classes/MaterialEditor/MaterialEditorInstanceConstant.h
--- a/Engine/Source/Editor/UnrealEd/Classes/MaterialEditor/MaterialEditorInstanceConstant.h
+++ b/Engine/Source/Editor/UnrealEd/Classes/MaterialEditor/MaterialEditorInstanceConstant.h
@@ -227,10 +227,10 @@ class UNREALED_API UMaterialEditorInstanceConstant : public UObject
 	UPROPERTY(EditAnywhere, Category = MaterialEditorInstanceConstant)
 	uint32 bOverrideSubsurfaceProfile : 1;
 
-	UPROPERTY(EditAnywhere, Category = MaterialOverrides)
-	uint32 bOverrideBaseProperties:1;
+	UPROPERTY()
+	uint32 bOverrideBaseProperties_DEPRECATED : 1;
 
-	UPROPERTY(EditAnywhere, Category=MaterialOverrides, meta=(editcondition="bOverrideBaseProperties"))
+	UPROPERTY(EditAnywhere, Category=MaterialOverrides)
 	FMaterialInstanceBasePropertyOverrides BasePropertyOverrides;
 
 	UPROPERTY()
diff --git a/Engine/Source/Editor/UnrealEd/Private/PreviewMaterial.cpp b/Engine/Source/Editor/UnrealEd/Private/PreviewMaterial.cpp
--- a/Engine/Source/Editor/UnrealEd/Private/PreviewMaterial.cpp
+++ b/Engine/Source/Editor/UnrealEd/Private/PreviewMaterial.cpp
@@ -573,9 +573,12 @@ void UMaterialEditorInstanceConstant::CopyToSourceInstance()
 		}
 
 		//Check for changes to see if we need to force a recompile
-		bool bForceRecompile = SourceInstance->BasePropertyOverrides.Update(BasePropertyOverrides);
-		bForceRecompile |= SourceInstance->bOverrideBaseProperties != bOverrideBaseProperties;
-		SourceInstance->bOverrideBaseProperties = bOverrideBaseProperties;
+		bool bForceRecompile = false;
+		if (SourceInstance->BasePropertyOverrides != BasePropertyOverrides)
+		{
+			bForceRecompile = true;
+			SourceInstance->BasePropertyOverrides = BasePropertyOverrides;
+		}
 		
 		FStaticParameterSet NewStaticParameters;
 		BuildStaticParametersForSourceInstance(NewStaticParameters);
@@ -672,9 +675,6 @@ void UMaterialEditorInstanceConstant::SetSourceInstance(UMaterialInstanceConstan
 	Parent = SourceInstance->Parent;
 	PhysMaterial = SourceInstance->PhysMaterial;
 
-	//Init the base property overrides.
-	BasePropertyOverrides.Init(*SourceInstance);
-	bOverrideBaseProperties = SourceInstance->bOverrideBaseProperties;
 	BasePropertyOverrides = SourceInstance->BasePropertyOverrides;
 
 	// Copy the Lightmass settings...
diff --git a/Engine/Source/Runtime/Engine/Classes/Materials/Material.h b/Engine/Source/Runtime/Engine/Classes/Materials/Material.h
--- a/Engine/Source/Runtime/Engine/Classes/Materials/Material.h
+++ b/Engine/Source/Runtime/Engine/Classes/Materials/Material.h
@@ -778,11 +778,11 @@ public:
 		TArray<FName>* OutTextureParamNames, class FStaticParameterSet* InStaticParameterSet) override;
 	ENGINE_API virtual void RecacheUniformExpressions() const override;
 
-	ENGINE_API virtual float GetOpacityMaskClipValue_Internal() const;
-	ENGINE_API virtual EBlendMode GetBlendMode_Internal() const;
-	ENGINE_API virtual EMaterialShadingModel GetShadingModel_Internal() const;
-	ENGINE_API virtual bool IsTwoSided_Internal() const;
-	ENGINE_API virtual bool IsMasked_Internal() const;
+	ENGINE_API virtual float GetOpacityMaskClipValue(bool bIsGameThread = IsInGameThread()) const;
+	ENGINE_API virtual EBlendMode GetBlendMode(bool bIsGameThread = IsInGameThread()) const;
+	ENGINE_API virtual EMaterialShadingModel GetShadingModel(bool bIsGameThread = IsInGameThread()) const;
+	ENGINE_API virtual bool IsTwoSided(bool bIsGameThread = IsInGameThread()) const;
+	ENGINE_API virtual bool IsMasked(bool bIsGameThread = IsInGameThread()) const;
 	ENGINE_API virtual USubsurfaceProfile* GetSubsurfaceProfile_Internal() const;
 
 	ENGINE_API void SetShadingModel(EMaterialShadingModel NewModel) {ShadingModel = NewModel;}
diff --git a/Engine/Source/Runtime/Engine/Classes/Materials/MaterialInstance.h b/Engine/Source/Runtime/Engine/Classes/Materials/MaterialInstance.h
--- a/Engine/Source/Runtime/Engine/Classes/Materials/MaterialInstance.h
+++ b/Engine/Source/Runtime/Engine/Classes/Materials/MaterialInstance.h
@@ -201,10 +201,10 @@ class UMaterialInstance : public UMaterialInterface
 	UPROPERTY(EditAnywhere, BlueprintReadOnly, Category=MaterialInstance)
 	TArray<struct FVectorParameterValue> VectorParameterValues;
 
-	UPROPERTY(EditAnywhere, BlueprintReadOnly, Category=MaterialInstance)
-	bool bOverrideBaseProperties;
+	UPROPERTY()
+	bool bOverrideBaseProperties_DEPRECATED;
 
-	UPROPERTY(EditAnywhere, BlueprintReadOnly, Category=MaterialInstance, meta=(editcondition = "bOverrideBaseProperties"))
+	UPROPERTY(EditAnywhere, BlueprintReadOnly, Category=MaterialInstance)
 	struct FMaterialInstanceBasePropertyOverrides BasePropertyOverrides;
 
 	/** 
@@ -277,24 +277,20 @@ public:
 	virtual ENGINE_API bool GetRefractionSettings(float& OutBiasValue) const override;
 	ENGINE_API virtual void ForceRecompileForRendering() override;
 	
-	ENGINE_API virtual float GetOpacityMaskClipValue_Internal() const;
-	ENGINE_API virtual EBlendMode GetBlendMode_Internal() const;
-	ENGINE_API virtual EMaterialShadingModel GetShadingModel_Internal() const;
-	ENGINE_API virtual bool IsTwoSided_Internal() const;
-	ENGINE_API virtual bool IsMasked_Internal() const;
+	ENGINE_API virtual float GetOpacityMaskClipValue(bool bIsGameThread=IsInGameThread()) const;
+	ENGINE_API virtual EBlendMode GetBlendMode(bool bIsGameThread = IsInGameThread()) const;
+	ENGINE_API virtual EMaterialShadingModel GetShadingModel(bool bIsGameThread = IsInGameThread()) const;
+	ENGINE_API virtual bool IsTwoSided(bool bIsGameThread = IsInGameThread()) const;
+	ENGINE_API virtual bool IsMasked(bool bIsGameThread = IsInGameThread()) const;
+
+	ENGINE_API float RenderThread_GetOpacityMaskClipValue() const;
+	ENGINE_API EBlendMode RenderThread_GetBlendMode() const;
+	ENGINE_API EMaterialShadingModel RenderThread_GetShadingModel() const;
+	ENGINE_API bool RenderThread_IsTwoSided() const;
+	ENGINE_API bool RenderThread_IsMasked() const;
+	
 	ENGINE_API virtual USubsurfaceProfile* GetSubsurfaceProfile_Internal() const;
 
-	/** Returns true and sets Result if the property was overridden in the instance. Otherwise, returns false and the base material property should be used. */
-	ENGINE_API virtual bool GetOpacityMaskClipValueOverride(float& OutResult) const;
-	/** Returns true and sets Result if the property was overridden in the instance. Otherwise, returns false and the base material property should be used. */
-	ENGINE_API virtual bool GetBlendModeOverride(EBlendMode& OutResult) const;
-	/** Returns true and sets Result if the property was overridden in the instance. Otherwise, returns false and the base material property should be used. */
-	ENGINE_API virtual bool GetShadingModelOverride(EMaterialShadingModel& OutResult) const;
-	/** Returns true and sets Result if the property was overridden in the instance. Otherwise, returns false and the base material property should be used. */
-	ENGINE_API virtual bool IsTwoSidedOverride(bool& OutResult) const;
-	/** Returns true and sets Result if the property was overridden in the instance. Otherwise, returns false and the base material property should be used. */
-	ENGINE_API virtual bool IsMaskedOverride(bool& OutResult) const;
-
 	/** Checks to see if an input property should be active, based on the state of the material */
 	ENGINE_API virtual bool IsPropertyActive(EMaterialProperty InProperty) const;
 	/** Allows material properties to be compiled with the option of being overridden by the material attributes input. */
@@ -333,6 +329,8 @@ public:
 	 */
 	ENGINE_API void InitStaticPermutation();
 
+	ENGINE_API void UpdateOverridableBaseProperties();
+
 	/** 
 	 * Cache resource shaders for rendering on the given shader platform. 
 	 * If a matching shader map is not found in memory or the DDC, a new one will be compiled.
@@ -358,6 +356,7 @@ public:
 	void GetAllStaticSwitchParameterValues(TArray<FStaticTerrainLayerWeightParameter> &TerrainLayerWeightParameters, UMaterial* ParentMaterial);
 
 	void GetBasePropertyOverridesHash(FSHAHash& OutHash)const;
+	bool HasOverridenBaseProperties()const;
 
 	// For all materials instances, UMaterialInstance::CacheResourceShadersForRendering
 	static void AllMaterialsCacheResourceShadersForRendering();
diff --git a/Engine/Source/Runtime/Engine/Classes/Materials/MaterialInstanceBasePropertyOverrides.h b/Engine/Source/Runtime/Engine/Classes/Materials/MaterialInstanceBasePropertyOverrides.h
--- a/Engine/Source/Runtime/Engine/Classes/Materials/MaterialInstanceBasePropertyOverrides.h
+++ b/Engine/Source/Runtime/Engine/Classes/Materials/MaterialInstanceBasePropertyOverrides.h
@@ -11,45 +11,39 @@ struct ENGINE_API FMaterialInstanceBasePropertyOverrides
 {
 	GENERATED_USTRUCT_BODY()
 
-	FMaterialInstanceBasePropertyOverrides();
-	void Init(const UMaterialInstance& Instance);
-
 	/** Enables override of the opacity mask clip value. */
-	UPROPERTY(EditAnywhere, Category=Material)
+	UPROPERTY(EditAnywhere, Category = Material)
 	bool bOverride_OpacityMaskClipValue;
 
 	/** Enables override of the blend mode. */
-	UPROPERTY(EditAnywhere, Category=Material)
+	UPROPERTY(EditAnywhere, Category = Material)
 	bool bOverride_BlendMode;
 
 	/** Enables override of the shading model. */
-	UPROPERTY(EditAnywhere, Category=Material)
+	UPROPERTY(EditAnywhere, Category = Material)
 	bool bOverride_ShadingModel;
 
 	/** Enables override of the two sided property. */
-	UPROPERTY(EditAnywhere, Category=Material)
+	UPROPERTY(EditAnywhere, Category = Material)
 	bool bOverride_TwoSided;
 
 	/** If BlendMode is BLEND_Masked, the surface is not rendered where OpacityMask < OpacityMaskClipValue. */
-	UPROPERTY(EditAnywhere, Category=Material, meta=(editcondition = "bOverride_OpacityMaskClipValue", NoSpinbox=true))
+	UPROPERTY(EditAnywhere, Category = Material, meta = (editcondition = "bOverride_OpacityMaskClipValue", NoSpinbox = true))
 	float OpacityMaskClipValue;
 
 	/** The blend mode */
-	UPROPERTY(EditAnywhere, Category=Material, meta=(editcondition = "bOverride_BlendMode"))
+	UPROPERTY(EditAnywhere, Category = Material, meta = (editcondition = "bOverride_BlendMode"))
 	TEnumAsByte<EBlendMode> BlendMode;
 
 	/** The shading model */
-	UPROPERTY(EditAnywhere, Category=Material, meta=(editcondition = "bOverride_ShadingModel"))
-	TEnumAsByte<EMaterialShadingModel> ShadingModel;	
-	
-	/** If the material is two sided or not. */
-	UPROPERTY(EditAnywhere, Category=Material, meta=(editcondition = "bOverride_TwoSided"))
-	uint32 TwoSided:1;
+	UPROPERTY(EditAnywhere, Category = Material, meta = (editcondition = "bOverride_ShadingModel"))
+	TEnumAsByte<EMaterialShadingModel> ShadingModel;
 
-	void UpdateHash(FSHA1& HashState) const;
+	/** If the material is two sided or not. */
+	UPROPERTY(EditAnywhere, Category = Material, meta = (editcondition = "bOverride_TwoSided"))
+	uint32 TwoSided : 1;
 
-	/** Updates the values of this set of overrides. Returns true if this should cause a recompile of the material, false otherwise. */ 
-	bool Update(const FMaterialInstanceBasePropertyOverrides& Updated);
+	FMaterialInstanceBasePropertyOverrides();
 
 	bool operator==(const FMaterialInstanceBasePropertyOverrides& Other)const;
 	bool operator!=(const FMaterialInstanceBasePropertyOverrides& Other)const;
diff --git a/Engine/Source/Runtime/Engine/Classes/Materials/MaterialInterface.h b/Engine/Source/Runtime/Engine/Classes/Materials/MaterialInterface.h
--- a/Engine/Source/Runtime/Engine/Classes/Materials/MaterialInterface.h
+++ b/Engine/Source/Runtime/Engine/Classes/Materials/MaterialInterface.h
@@ -529,21 +529,15 @@ public:
 	ENGINE_API virtual bool GetGroupName(FName ParameterName, FName& GroupName) const;
 
 	/**
-		Allows callers from any thread to access these properties via the interface.
+		Access to overridable properties of the base material.
 	*/
-	ENGINE_API float GetOpacityMaskClipValue() const;
-	ENGINE_API EBlendMode GetBlendMode() const;
-	ENGINE_API EMaterialShadingModel GetShadingModel() const;
-	ENGINE_API bool IsTwoSided() const;
-	ENGINE_API bool IsMasked() const;
-
-	/** Game thread versions of the accessors. On the render thread there are equivelant getters in FMaterial. */
-	ENGINE_API virtual float GetOpacityMaskClipValue_Internal() const;
-	ENGINE_API virtual EBlendMode GetBlendMode_Internal() const;
-	ENGINE_API virtual EMaterialShadingModel GetShadingModel_Internal() const;
-	ENGINE_API virtual bool IsTwoSided_Internal() const;
-	ENGINE_API virtual bool IsMasked_Internal() const;
-	ENGINE_API virtual USubsurfaceProfile* GetSubsurfaceProfile_Internal() const { return 0; }
+	ENGINE_API virtual float GetOpacityMaskClipValue(bool bIsGameThread = IsInGameThread()) const;
+	ENGINE_API virtual EBlendMode GetBlendMode(bool bIsGameThread = IsInGameThread()) const;
+	ENGINE_API virtual EMaterialShadingModel GetShadingModel(bool bIsGameThread = IsInGameThread()) const;
+	ENGINE_API virtual bool IsTwoSided(bool bIsGameThread = IsInGameThread()) const;
+	ENGINE_API virtual bool IsMasked(bool bIsGameThread = IsInGameThread()) const;
+
+	ENGINE_API virtual USubsurfaceProfile* GetSubsurfaceProfile_Internal() const;
 
 	/**
 	 * Force the streaming system to disregard the normal logic for the specified duration and
diff --git a/Engine/Source/Runtime/Engine/Private/Materials/Material.cpp b/Engine/Source/Runtime/Engine/Private/Materials/Material.cpp
--- a/Engine/Source/Runtime/Engine/Private/Materials/Material.cpp
+++ b/Engine/Source/Runtime/Engine/Private/Materials/Material.cpp
@@ -3497,17 +3497,17 @@ static FAutoConsoleCommand CmdListSceneColorMaterials(
 	FConsoleCommandDelegate::CreateStatic(ListSceneColorMaterials)
 	);
 
-float UMaterial::GetOpacityMaskClipValue_Internal() const
+float UMaterial::GetOpacityMaskClipValue(bool bIsInGameThread) const
 {
 	return OpacityMaskClipValue;
 }
 
-EBlendMode UMaterial::GetBlendMode_Internal() const
+EBlendMode UMaterial::GetBlendMode(bool bIsInGameThread) const
 {
 	return BlendMode;
 }
 
-EMaterialShadingModel UMaterial::GetShadingModel_Internal() const
+EMaterialShadingModel UMaterial::GetShadingModel(bool bIsInGameThread) const
 {
 	switch (MaterialDomain)
 	{
@@ -3526,12 +3526,12 @@ EMaterialShadingModel UMaterial::GetShadingModel_Internal() const
 	}
 }
 
-bool UMaterial::IsTwoSided_Internal() const
+bool UMaterial::IsTwoSided(bool bIsInGameThread) const
 {
 	return TwoSided != 0;
 }
 
-bool UMaterial::IsMasked_Internal() const
+bool UMaterial::IsMasked(bool bIsInGameThread) const
 {
 	return bIsMasked != 0;
 }
diff --git a/Engine/Source/Runtime/Engine/Private/Materials/MaterialInstance.cpp b/Engine/Source/Runtime/Engine/Private/Materials/MaterialInstance.cpp
--- a/Engine/Source/Runtime/Engine/Private/Materials/MaterialInstance.cpp
+++ b/Engine/Source/Runtime/Engine/Private/Materials/MaterialInstance.cpp
@@ -77,6 +77,10 @@ FMaterialInstanceResource::FMaterialInstanceResource(UMaterialInstance* InOwner,
 	, Owner(InOwner)
 	, DistanceFieldPenumbraScale(1.0f)
 	, GameThreadParent(NULL)
+	, OpacityMaskClipValue(0.3333333f)
+	, BlendMode(BLEND_Opaque)
+	, ShadingModel(MSM_DefaultLit)
+	, TwoSided(false)
 {
 }
 
@@ -243,13 +247,37 @@ void FMaterialInstanceResource::GameThread_UpdateDistanceFieldPenumbraScale(floa
 	});
 }
 
+void FMaterialInstanceResource::GameThread_UpdateOverridableBaseProperties(const UMaterialInterface* MaterialInterface)
+{
+	ENQUEUE_UNIQUE_RENDER_COMMAND_FOURPARAMETER(
+		UpdateOverridableBaseProperties0,
+		float*, OpacityMaskClipValue, &OpacityMaskClipValue,
+		float, NewOpacityMaskClipValue, MaterialInterface->GetOpacityMaskClipValue(),
+		EBlendMode*, BlendMode, &BlendMode,
+		EBlendMode, NewBlendMode, MaterialInterface->GetBlendMode(),
+		{
+		*OpacityMaskClipValue = NewOpacityMaskClipValue;
+		*BlendMode = NewBlendMode;
+	});
+	ENQUEUE_UNIQUE_RENDER_COMMAND_FOURPARAMETER(
+		UpdateOverrideablBaseProperties1,
+		EMaterialShadingModel*, ShadingModel, &ShadingModel,
+		EMaterialShadingModel, NewShadingModel, MaterialInterface->GetShadingModel(),
+		bool*, TwoSided, &TwoSided,
+		bool, NewTwoSided, MaterialInterface->IsTwoSided(),
+		{
+		*ShadingModel = NewShadingModel;
+		*TwoSided = NewTwoSided;
+	});
+}
+
 void UMaterialInstance::PropagateDataToMaterialProxy()
 {
 	for (int32 i = 0; i < ARRAY_COUNT(Resources); i++)
 	{
 		if (Resources[i])
 		{
-			Resources[i]->GameThread_UpdateDistanceFieldPenumbraScale(GetDistanceFieldPenumbraScale());
+			Resources[i]->GameThread_UpdateDistanceFieldPenumbraScale(GetDistanceFieldPenumbraScale()); 
 
 			UpdateMaterialRenderProxy(*Resources[i]);
 		}
@@ -372,7 +400,6 @@ UMaterialInstance::UMaterialInstance(const FObjectInitializer& ObjectInitializer
 	: Super(ObjectInitializer)
 {
 	bHasStaticPermutationResource = false;
-	bOverrideBaseProperties = false;
 }
 
 void UMaterialInstance::PostInitProperties()	
@@ -440,6 +467,7 @@ void UMaterialInstance::InitResources()
 		if (Resources[CurResourceIndex] != NULL)
 		{
 			Resources[CurResourceIndex]->GameThread_SetParent(SafeParent);
+			Resources[CurResourceIndex]->GameThread_UpdateOverridableBaseProperties(this);
 		}
 	}
 
@@ -1217,6 +1245,22 @@ void UMaterialInstance::InitStaticPermutation()
 		// Cache shaders for the current platform to be used for rendering
 		CacheResourceShadersForRendering();
 	}
+	else
+	{
+		//Need to have this done but it's also done inside CacheResourcesForRendering so don't want to do it twice.
+		UpdateOverridableBaseProperties();
+	}
+}
+
+void UMaterialInstance::UpdateOverridableBaseProperties()
+{
+	for (int32 i = 0; i < ARRAY_COUNT(Resources); i++)
+	{
+		if (Resources[i])
+		{
+			Resources[i]->GameThread_UpdateOverridableBaseProperties(this);
+		}
+	}
 }
 
 void UMaterialInstance::GetAllShaderMaps(TArray<FMaterialShaderMap*>& OutShaderMaps)
@@ -1292,6 +1336,7 @@ void UMaterialInstance::CacheResourceShadersForRendering()
 	}
 
 	UpdatePermutationAllocations();
+	UpdateOverridableBaseProperties();
 
 	if (bHasStaticPermutationResource && FApp::CanEverRender())
 	{
@@ -1627,7 +1672,7 @@ void UMaterialInstance::Serialize(FArchive& Ar)
 		if( Ar.UE4Ver() < VER_UE4_FIX_MATERIAL_PROPERTY_OVERRIDE_SERIALIZE )
 		{
 			// awful old native serialize of FMaterialInstanceBasePropertyOverrides UStruct
-			Ar << bOverrideBaseProperties;
+			Ar << bOverrideBaseProperties_DEPRECATED;
 			bool bHasPropertyOverrides = false;
 			Ar << bHasPropertyOverrides;
 			if( bHasPropertyOverrides )
@@ -1701,7 +1746,7 @@ void UMaterialInstance::PostLoad()
 	PropagateDataToMaterialProxy();
 
 	// Update bHasStaticPermutationResource in case the parent was not found
-	bHasStaticPermutationResource = (!StaticParameters.IsEmpty() || bOverrideBaseProperties) && Parent;
+	bHasStaticPermutationResource = (!StaticParameters.IsEmpty() || HasOverridenBaseProperties()) && Parent;
 
 	STAT(double MaterialLoadTime = 0);
 	{
@@ -1993,7 +2038,7 @@ void UMaterialInstance::UpdateStaticPermutation(const FStaticParameterSet& NewPa
 	TrimToOverriddenOnly(CompareParameters.StaticComponentMaskParameters);
 	TrimToOverriddenOnly(CompareParameters.TerrainLayerWeightParameters);
 
-	const bool bWantsStaticPermutationResource = (!CompareParameters.IsEmpty() || bOverrideBaseProperties || bForceRecompile) && Parent;
+	const bool bWantsStaticPermutationResource = (!CompareParameters.IsEmpty() || HasOverridenBaseProperties() || bForceRecompile) && Parent;
 
 	if (bForceRecompile || bHasStaticPermutationResource != bWantsStaticPermutationResource || StaticParameters != CompareParameters)
 	{
@@ -2405,126 +2450,196 @@ bool UMaterialInstance::IsChildOf(const UMaterialInterface* Parent) const
 
 void UMaterialInstance::GetBasePropertyOverridesHash(FSHAHash& OutHash)const
 {
-	FSHA1 HashState;
+	check(IsInGameThread());
+
+	const UMaterial* Mat = GetMaterial();
+	check(Mat);
 
-	if( bOverrideBaseProperties )
+	FSHA1 Hash;
+	bool bHasOverrides = false;
+
+	float UsedOpacityMaskClipValue = GetOpacityMaskClipValue(true);
+	if (FMath::Abs(UsedOpacityMaskClipValue - Mat->GetOpacityMaskClipValue(true)) > SMALL_NUMBER)
 	{
-		BasePropertyOverrides.UpdateHash(HashState);
+		const FString HashString = TEXT("bOverride_OpacityMaskClipValue");
+		Hash.UpdateWithString(*HashString, HashString.Len());
+		Hash.Update((const uint8*)&UsedOpacityMaskClipValue, sizeof(UsedOpacityMaskClipValue));
+		bHasOverrides = true;
 	}
+
+	EBlendMode UsedBlendMode = GetBlendMode(true);
+	if (UsedBlendMode != Mat->GetBlendMode(true))
+	{
+		const FString HashString = TEXT("bOverride_BlendMode");
+		Hash.UpdateWithString(*HashString, HashString.Len());
+		Hash.Update((const uint8*)&UsedBlendMode, sizeof(UsedBlendMode));
+		bHasOverrides = true;
+ 	}
 	
-	HashState.Final();
-	HashState.GetHash(&OutHash.Hash[0]);
+	EMaterialShadingModel UsedShadingModel = GetShadingModel(true);
+ 	if (UsedShadingModel != Mat->GetShadingModel(true))
+ 	{
+		const FString HashString = TEXT("bOverride_ShadingModel");
+		Hash.UpdateWithString(*HashString, HashString.Len());
+		Hash.Update((const uint8*)&UsedShadingModel, sizeof(UsedShadingModel));
+		bHasOverrides = true;
+	}
+
+ 	bool bUsedIsTwoSided = IsTwoSided(true);
+ 	if (bUsedIsTwoSided != Mat->IsTwoSided(true))
+ 	{
+		const FString HashString = TEXT("bOverride_TwoSided");
+		Hash.UpdateWithString(*HashString, HashString.Len());
+		Hash.Update((uint8*)&bUsedIsTwoSided, sizeof(bUsedIsTwoSided));
+		bHasOverrides = true;
+	}
+
+ 	if (bHasOverrides)
+ 	{
+		Hash.Final();
+		Hash.GetHash(&OutHash.Hash[0]);
+	}
 }
 
-float UMaterialInstance::GetOpacityMaskClipValue_Internal() const
+bool UMaterialInstance::HasOverridenBaseProperties()const
 {
-	checkSlow(IsInGameThread());
-	if( bOverrideBaseProperties && BasePropertyOverrides.bOverride_OpacityMaskClipValue )
+	check(IsInGameThread());
+
+	const UMaterial* Mat = GetMaterial();
+	check(Mat);
+
+	if ((FMath::Abs(GetOpacityMaskClipValue(true) - Mat->GetOpacityMaskClipValue(true)) > SMALL_NUMBER) ||
+		(GetBlendMode(true) != Mat->GetBlendMode(true)) ||
+		(GetShadingModel(true) != Mat->GetShadingModel(true)) ||
+		(IsTwoSided(true) != Mat->IsTwoSided(true))
+		)
 	{
-		return BasePropertyOverrides.OpacityMaskClipValue;
+		return true;
 	}
-	return GetMaterial()->GetOpacityMaskClipValue();
+
+	return false;
 }
-EBlendMode UMaterialInstance::GetBlendMode_Internal() const
+
+float UMaterialInstance::GetOpacityMaskClipValue(bool bIsInGameThread) const
 {
-	checkSlow(IsInGameThread());
-	if( bOverrideBaseProperties && BasePropertyOverrides.bOverride_BlendMode )
+	if (bIsInGameThread)
 	{
-		return BasePropertyOverrides.BlendMode;
+		if (BasePropertyOverrides.bOverride_OpacityMaskClipValue)
+		{
+			return BasePropertyOverrides.OpacityMaskClipValue;
+		}
+		// go up the chain if possible
+		return Parent ? Parent->GetOpacityMaskClipValue(true) : 0.0f;
 	}
-	return GetMaterial()->GetBlendMode();
+	
+	//Get the value mirrored in the render proxy.
+	return RenderThread_GetOpacityMaskClipValue();
 }
 
-EMaterialShadingModel UMaterialInstance::GetShadingModel_Internal() const
+float UMaterialInstance::RenderThread_GetOpacityMaskClipValue() const
 {
-	checkSlow(IsInGameThread());
-	if( bOverrideBaseProperties && BasePropertyOverrides.bOverride_ShadingModel )
-	{
-		return BasePropertyOverrides.ShadingModel;
-	}
-	return GetMaterial()->GetShadingModel();
+	FMaterialInstanceResource* Proxy = ((FMaterialInstanceResource*)GetRenderProxy(0));
+	return Proxy ? Proxy->GetOpacityMaskClipValue() : 0.0f;
 }
 
-bool UMaterialInstance::IsTwoSided_Internal() const
+EBlendMode UMaterialInstance::GetBlendMode(bool bIsInGameThread) const
 {
-	checkSlow(IsInGameThread());
-	if( bOverrideBaseProperties && BasePropertyOverrides.bOverride_TwoSided )
+	if (bIsInGameThread)
 	{
-		return BasePropertyOverrides.TwoSided != 0;
+		if (BasePropertyOverrides.bOverride_BlendMode)
+		{
+			return BasePropertyOverrides.BlendMode;
+		}
+		// go up the chain if possible
+		return Parent ? Parent->GetBlendMode(true) : BLEND_Opaque;
 	}
-	return GetMaterial()->IsTwoSided();
+
+	//Get the value mirrored in the render proxy.
+	return RenderThread_GetBlendMode();
 }
 
-bool UMaterialInstance::IsMasked_Internal() const
+EBlendMode UMaterialInstance::RenderThread_GetBlendMode() const
 {
-	checkSlow(IsInGameThread());
-	if (bOverrideBaseProperties && BasePropertyOverrides.bOverride_BlendMode)
-	{
-		return BasePropertyOverrides.BlendMode == EBlendMode::BLEND_Masked;
-	}
-	return GetMaterial()->IsMasked();
+	FMaterialInstanceResource* Proxy = ((FMaterialInstanceResource*)GetRenderProxy(0));
+	return Proxy ? Proxy->GetBlendMode() : BLEND_Opaque;
 }
 
-USubsurfaceProfile* UMaterialInstance::GetSubsurfaceProfile_Internal() const
+EMaterialShadingModel UMaterialInstance::GetShadingModel(bool bIsInGameThread) const
 {
-	checkSlow(IsInGameThread());
-	if (bOverrideSubsurfaceProfile)
+	if (bIsInGameThread)
 	{
-		return SubsurfaceProfile;
+		if (BasePropertyOverrides.bOverride_ShadingModel)
+		{
+			return BasePropertyOverrides.ShadingModel;
+		}
+		// go up the chain if possible
+		return Parent ? Parent->GetShadingModel(true) : MSM_DefaultLit;
 	}
 
-	// go up the chain if possible
-	return Parent ? Parent->GetSubsurfaceProfile_Internal() : 0;
+	//Get the value mirrored in the render proxy.
+	return RenderThread_GetShadingModel();
 }
 
-bool UMaterialInstance::GetOpacityMaskClipValueOverride(float& OutResult) const
+EMaterialShadingModel UMaterialInstance::RenderThread_GetShadingModel()const 
 {
-	if( bOverrideBaseProperties && BasePropertyOverrides.bOverride_OpacityMaskClipValue )
-	{
-		OutResult = BasePropertyOverrides.OpacityMaskClipValue;
-		return true;
-	}
-	return false;
+	FMaterialInstanceResource* Proxy = ((FMaterialInstanceResource*)GetRenderProxy(0));
+	return Proxy ? Proxy->GetShadingModel() : MSM_DefaultLit;
 }
 
-bool UMaterialInstance::GetBlendModeOverride(EBlendMode& OutResult) const
+bool UMaterialInstance::IsTwoSided(bool bIsInGameThread) const
 {
-	if( bOverrideBaseProperties && BasePropertyOverrides.bOverride_BlendMode )
+	if (bIsInGameThread)
 	{
-		OutResult = BasePropertyOverrides.BlendMode;
-		return true;
+		if (BasePropertyOverrides.bOverride_TwoSided)
+		{
+			return BasePropertyOverrides.TwoSided != 0;
+		}
+		// go up the chain if possible
+		return Parent ? Parent->IsTwoSided(true) : false;
 	}
-	return false;
+
+	//Get the value mirrored in the render proxy.
+	return RenderThread_IsTwoSided();
 }
 
-bool UMaterialInstance::GetShadingModelOverride(EMaterialShadingModel& OutResult) const
+bool UMaterialInstance::RenderThread_IsTwoSided() const
 {
-	if( bOverrideBaseProperties && BasePropertyOverrides.bOverride_ShadingModel )
-	{
-		OutResult = BasePropertyOverrides.ShadingModel;
-		return true;
-	}
-	return false;
+	FMaterialInstanceResource* Proxy = ((FMaterialInstanceResource*)GetRenderProxy(0));
+	return Proxy ? Proxy->IsTwoSided() : false;
 }
 
-bool UMaterialInstance::IsTwoSidedOverride(bool& OutResult) const
+bool UMaterialInstance::IsMasked(bool bIsInGameThread) const
 {
-	if( bOverrideBaseProperties && BasePropertyOverrides.bOverride_TwoSided )
+	if (bIsInGameThread)
 	{
-		OutResult = BasePropertyOverrides.TwoSided != 0;
-		return true;
+		if (BasePropertyOverrides.bOverride_BlendMode)
+		{
+			return BasePropertyOverrides.BlendMode == EBlendMode::BLEND_Masked;
+		}
+		// go up the chain if possible
+		return Parent ? Parent->IsMasked(true) : false;
 	}
-	return false;
+
+	//Get the value mirrored in the render proxy.
+	return RenderThread_IsMasked();
+}
+
+bool UMaterialInstance::RenderThread_IsMasked() const
+{
+	FMaterialInstanceResource* Proxy = ((FMaterialInstanceResource*)GetRenderProxy(0));
+	return Proxy ? Proxy->GetBlendMode() == EBlendMode::BLEND_Masked : false;
 }
 
-bool UMaterialInstance::IsMaskedOverride(bool& OutResult) const
+USubsurfaceProfile* UMaterialInstance::GetSubsurfaceProfile_Internal() const
 {
-	if (bOverrideBaseProperties && BasePropertyOverrides.BlendMode)
+	checkSlow(IsInGameThread());
+	if (bOverrideSubsurfaceProfile)
 	{
-		OutResult = BasePropertyOverrides.BlendMode == EBlendMode::BLEND_Masked;
-		return true;
+		return SubsurfaceProfile;
 	}
-	return false;
+
+	// go up the chain if possible
+	return Parent ? Parent->GetSubsurfaceProfile_Internal() : 0;
 }
 
 /** Checks to see if an input property should be active, based on the state of the material */
diff --git a/Engine/Source/Runtime/Engine/Private/Materials/MaterialInstanceSupport.h b/Engine/Source/Runtime/Engine/Private/Materials/MaterialInstanceSupport.h
--- a/Engine/Source/Runtime/Engine/Private/Materials/MaterialInstanceSupport.h
+++ b/Engine/Source/Runtime/Engine/Private/Materials/MaterialInstanceSupport.h
@@ -97,6 +97,9 @@ public:
 
 	/** Called from the game thread to update DistanceFieldPenumbraScale. */
 	void GameThread_UpdateDistanceFieldPenumbraScale(float NewDistanceFieldPenumbraScale);
+	
+	/** Called from the game thread to update the overridable base properties in the proxy. */
+	void GameThread_UpdateOverridableBaseProperties(const UMaterialInterface* MaterialInterface);
 
 	/**
 	 * Clears all parameters set on this material instance.
@@ -152,6 +155,11 @@ public:
 		return NULL;
 	}
 
+	float GetOpacityMaskClipValue()const{ return OpacityMaskClipValue; }
+	EBlendMode GetBlendMode()const{ return BlendMode; }
+	EMaterialShadingModel GetShadingModel()const{ return ShadingModel; }
+	bool IsTwoSided()const{ return TwoSided; }
+
 private:
 	/**
 	 * Retrieves the array of values for a given type.
@@ -177,6 +185,15 @@ private:
 	TArray<TNamedParameter<float> > ScalarParameterArray;
 	/** Texture parameters for this material instance. */
 	TArray<TNamedParameter<const UTexture*> > TextureParameterArray;
+
+	/** 
+		Potentially overridden properties of the base material. 
+		Cached here from the game thread so that the render thread doesn't have to traverse up the parent chain.
+	*/
+	float OpacityMaskClipValue;
+	EBlendMode BlendMode;
+	EMaterialShadingModel ShadingModel;
+	bool TwoSided;
 };
 
 template <> FORCEINLINE TArray<FMaterialInstanceResource::TNamedParameter<float> >& FMaterialInstanceResource::GetValueArray() { return ScalarParameterArray; }
diff --git a/Engine/Source/Runtime/Engine/Private/Materials/MaterialInterface.cpp b/Engine/Source/Runtime/Engine/Private/Materials/MaterialInterface.cpp
--- a/Engine/Source/Runtime/Engine/Private/Materials/MaterialInterface.cpp
+++ b/Engine/Source/Runtime/Engine/Private/Materials/MaterialInterface.cpp
@@ -235,87 +235,36 @@ bool DoesMaterialUseTexture(const UMaterialInterface* Material,const UTexture* C
 	return false;
 }
 
-float UMaterialInterface::GetOpacityMaskClipValue() const
-{
-	if( IsInGameThread() )
-	{
-		return GetOpacityMaskClipValue_Internal();
-	}
-
-	//We're on the render thread so get it from the proxy.
-	return GetRenderProxy(0)->GetMaterial(GMaxRHIFeatureLevel)->GetOpacityMaskClipValue();
-}
-
-EBlendMode UMaterialInterface::GetBlendMode() const
-{
-	if( IsInGameThread() )
-	{
-		return GetBlendMode_Internal();
-	}
-
-	//We're on the render thread so get it from the proxy.
-	return GetRenderProxy(0)->GetMaterial(GMaxRHIFeatureLevel)->GetBlendMode();
-}
-
-bool UMaterialInterface::IsTwoSided() const
-{
-	if( IsInGameThread() )
-	{
-		return IsTwoSided_Internal();
-	}
-
-	//We're on the render thread so get it from the proxy.
-	return GetRenderProxy(0)->GetMaterial(GMaxRHIFeatureLevel)->IsTwoSided();
-}
-
-bool UMaterialInterface::IsMasked() const
-{
-	if (IsInGameThread())
-	{
-		return IsMasked_Internal();
-	}
-
-	//We're on the render thread so get it from the proxy.
-	return GetRenderProxy(0)->GetMaterial(GMaxRHIFeatureLevel)->IsMasked();
-}
-
-EMaterialShadingModel UMaterialInterface::GetShadingModel() const
-{
-	if( IsInGameThread() )
-	{
-		return GetShadingModel_Internal();
-	}
-
-	//We're on the render thread so get it from the proxy.
-	return GetRenderProxy(0)->GetMaterial(GMaxRHIFeatureLevel)->GetShadingModel();
-
-}
-
-float UMaterialInterface::GetOpacityMaskClipValue_Internal() const
+float UMaterialInterface::GetOpacityMaskClipValue(bool bIsInGameThread) const
 {
 	return 0.0f;
 }
 
-EBlendMode UMaterialInterface::GetBlendMode_Internal() const
+EBlendMode UMaterialInterface::GetBlendMode(bool bIsInGameThread) const
 {
 	return BLEND_Opaque;
 }
 
-bool UMaterialInterface::IsTwoSided_Internal() const
+bool UMaterialInterface::IsTwoSided(bool bIsInGameThread) const
 {
 	return false;
 }
 
-bool UMaterialInterface::IsMasked_Internal() const
+bool UMaterialInterface::IsMasked(bool bIsInGameThread) const
 {
 	return false;
 }
 
-EMaterialShadingModel UMaterialInterface::GetShadingModel_Internal() const
+EMaterialShadingModel UMaterialInterface::GetShadingModel(bool bIsInGameThread) const
 {
 	return MSM_DefaultLit;
 }
 
+USubsurfaceProfile* UMaterialInterface::GetSubsurfaceProfile_Internal() const
+{
+	return NULL;
+}
+
 void UMaterialInterface::SetFeatureLevelToCompile(ERHIFeatureLevel::Type FeatureLevel, bool bShouldCompile)
 {
 	uint32 FeatureLevelBit = (1 << FeatureLevel);
@@ -356,7 +305,7 @@ void UMaterialInterface::UpdateMaterialRenderProxy(FMaterialRenderProxy& Proxy)
 	// no 0 pointer
 	check(&Proxy);
 
-	EMaterialShadingModel MaterialShadingModel = GetShadingModel_Internal();
+	EMaterialShadingModel MaterialShadingModel = GetShadingModel(true);
 
 	// for better performance we only update SubsurfaceProfileRT if the feature is used
 	if (MaterialShadingModel == MSM_SubsurfaceProfile)
diff --git a/Engine/Source/Runtime/Engine/Private/Materials/MaterialShared.cpp b/Engine/Source/Runtime/Engine/Private/Materials/MaterialShared.cpp
--- a/Engine/Source/Runtime/Engine/Private/Materials/MaterialShared.cpp
+++ b/Engine/Source/Runtime/Engine/Private/Materials/MaterialShared.cpp
@@ -828,52 +828,27 @@ ETranslucencyLightingMode FMaterialResource::GetTranslucencyLightingMode() const
 
 float FMaterialResource::GetOpacityMaskClipValue() const 
 {
-	float Ret;
-	if( !(MaterialInstance && MaterialInstance->GetOpacityMaskClipValueOverride(Ret)) )
-	{
-		Ret = Material->OpacityMaskClipValue;
-	}
-	return Ret;
+	return MaterialInstance ? MaterialInstance->GetOpacityMaskClipValue() : Material->GetOpacityMaskClipValue();
 }
 
 EBlendMode FMaterialResource::GetBlendMode() const 
 {
-	EBlendMode Ret;
-	if( !(MaterialInstance && MaterialInstance->GetBlendModeOverride(Ret)) )
-	{
-		Ret = Material->BlendMode;
-	}
-	return Ret;
+	return MaterialInstance ? MaterialInstance->GetBlendMode() : Material->GetBlendMode();
 }
 
 EMaterialShadingModel FMaterialResource::GetShadingModel() const 
-{ 
-	EMaterialShadingModel Ret;
-	if( !(MaterialInstance && MaterialInstance->GetShadingModelOverride(Ret)) )
-	{
-		Ret = Material->GetShadingModel_Internal();
-	}
-	return Ret;
+{
+	return MaterialInstance ? MaterialInstance->GetShadingModel() : Material->GetShadingModel();
 }
 
 bool FMaterialResource::IsTwoSided() const 
 {
-	bool Ret;
-	if( !(MaterialInstance && MaterialInstance->IsTwoSidedOverride(Ret)) )
-	{
-		Ret = Material->TwoSided != 0;
-	}
-	return Ret;
+	return MaterialInstance ? MaterialInstance->IsTwoSided() : Material->IsTwoSided();
 }
 
 bool FMaterialResource::IsMasked() const 
-{ 
-	bool Ret;
-	if (!(MaterialInstance && MaterialInstance->IsMaskedOverride(Ret)))
-	{
-		Ret = Material->bIsMasked != 0;
-	}
-	return Ret;
+{
+	return MaterialInstance ? MaterialInstance->IsMasked() : Material->IsMasked();
 }
 
 bool FMaterialResource::IsDistorted() const { return Material->bUsesDistortion; }
@@ -2359,12 +2334,14 @@ void DoMaterialAttributeReorder(FExpressionInput* Input, int32 UE4Ver)
 	}
 }
 
+//////////////////////////////////////////////////////////////////////////
+
 FMaterialInstanceBasePropertyOverrides::FMaterialInstanceBasePropertyOverrides()
 	:bOverride_OpacityMaskClipValue(false)
 	,bOverride_BlendMode(false)
 	,bOverride_ShadingModel(false)
 	,bOverride_TwoSided(false)
-	,OpacityMaskClipValue(0.0f)
+	,OpacityMaskClipValue(.333333f)
 	,BlendMode(BLEND_Opaque)
 	,ShadingModel(MSM_DefaultLit)
 	,TwoSided(0)
@@ -2372,71 +2349,24 @@ FMaterialInstanceBasePropertyOverrides::FMaterialInstanceBasePropertyOverrides()
 
 }
 
-void FMaterialInstanceBasePropertyOverrides::Init(const UMaterialInstance& Instance)
+bool FMaterialInstanceBasePropertyOverrides::operator==(const FMaterialInstanceBasePropertyOverrides& Other)const
 {
-	OpacityMaskClipValue = Instance.GetOpacityMaskClipValue();
-	BlendMode = Instance.GetBlendMode();
-	ShadingModel = Instance.GetShadingModel();
-	TwoSided = (uint32)Instance.IsTwoSided();
+	return	bOverride_OpacityMaskClipValue == Other.bOverride_OpacityMaskClipValue &&
+			bOverride_BlendMode == Other.bOverride_BlendMode &&
+			bOverride_ShadingModel == Other.bOverride_ShadingModel &&
+			bOverride_TwoSided == Other.bOverride_TwoSided &&
+			OpacityMaskClipValue == Other.OpacityMaskClipValue &&
+			BlendMode == Other.BlendMode &&
+			ShadingModel == Other.ShadingModel &&
+			TwoSided == Other.TwoSided;
 }
 
-void FMaterialInstanceBasePropertyOverrides::UpdateHash(FSHA1& HashState) const
+bool FMaterialInstanceBasePropertyOverrides::operator!=(const FMaterialInstanceBasePropertyOverrides& Other)const
 {
-	if(bOverride_OpacityMaskClipValue)
-	{
-		const FString HashString = TEXT("bOverride_OpacityMaskClipValue");
-		HashState.UpdateWithString(*HashString, HashString.Len());
-		HashState.Update((const uint8*)&OpacityMaskClipValue, sizeof(OpacityMaskClipValue));
-	}
-
-	if(bOverride_BlendMode)
-	{
-		const FString HashString = TEXT("bOverride_BlendMode");
-		HashState.UpdateWithString(*HashString, HashString.Len());
-		HashState.Update((const uint8*)&BlendMode, sizeof(BlendMode));
-	}
-
-	if(bOverride_ShadingModel)
-	{
-		const FString HashString = TEXT("bOverride_ShadingModel");
-		HashState.UpdateWithString(*HashString, HashString.Len());
-		HashState.Update((const uint8*)&ShadingModel, sizeof(ShadingModel));
-	}
-
-	//This does seem like it needs to be in the hash but due to some shaders being added to when two sided is enabled
-	//it causes a recompile anyway.
-// 	if(bOverride_TwoSided)
-// 	{
-// 		HashString = TEXT("bOverride_TwoSided");
-// 		HashState.UpdateWithString(*HashString, HashString.Len());
-// 		bool bIsTwoSided = TwoSided;
-// 		HashState.Update((uint8*)&bIsTwoSided, sizeof(bIsTwoSided));
-// 	}
-
-	//Some properties may not need to be in the shader map ID so don't add them unnecessarily.
+	return !(*this == Other);
 }
 
-bool FMaterialInstanceBasePropertyOverrides::Update(const FMaterialInstanceBasePropertyOverrides& Updated)
-{
-	bool bRet = false;
-	//Work out if we need a recompile.
-	bRet |= (Updated.bOverride_OpacityMaskClipValue != bOverride_OpacityMaskClipValue);
-	bRet |= Updated.bOverride_OpacityMaskClipValue && (OpacityMaskClipValue != Updated.OpacityMaskClipValue);
-
-	bRet |= (Updated.bOverride_BlendMode != bOverride_BlendMode);
-	bRet |= Updated.bOverride_BlendMode && (BlendMode != Updated.BlendMode);
-
-	bRet |= (Updated.bOverride_ShadingModel != bOverride_ShadingModel);
-	bRet |= Updated.bOverride_ShadingModel && (ShadingModel != Updated.ShadingModel);
-
-	bRet |= (Updated.bOverride_TwoSided != bOverride_TwoSided);
-	bRet |= Updated.bOverride_TwoSided && (TwoSided != Updated.TwoSided);
-
-	//Update the data.
-	*this = Updated;
-
-	return bRet;
-}
+//////////////////////////////////////////////////////////////////////////
 
 bool FMaterialShaderMapId::ContainsShaderType(const FShaderType* ShaderType) const
 {
diff --git a/Engine/Source/Runtime/Engine/Private/StaticMeshComponent.cpp b/Engine/Source/Runtime/Engine/Private/StaticMeshComponent.cpp
--- a/Engine/Source/Runtime/Engine/Private/StaticMeshComponent.cpp
+++ b/Engine/Source/Runtime/Engine/Private/StaticMeshComponent.cpp
@@ -1241,7 +1241,7 @@ bool UStaticMeshComponent::UsesOnlyUnlitMaterials() const
 				UMaterialInterface*	MaterialInterface	= GetMaterial(LOD.Sections[ElementIndex].MaterialIndex);
 				UMaterial*			Material			= MaterialInterface ? MaterialInterface->GetMaterial() : NULL;
 
-				bUsesOnlyUnlitMaterials = Material && Material->GetShadingModel_Internal() == MSM_Unlit;
+				bUsesOnlyUnlitMaterials = Material && Material->GetShadingModel() == MSM_Unlit;
 			}
 		}
 		return bUsesOnlyUnlitMaterials;
