Add ability to notify game code when demo playback fails

diff --git a/Engine/Source/Runtime/Engine/Classes/Engine/EngineBaseTypes.h b/Engine/Source/Runtime/Engine/Classes/Engine/EngineBaseTypes.h
--- a/Engine/Source/Runtime/Engine/Classes/Engine/EngineBaseTypes.h
+++ b/Engine/Source/Runtime/Engine/Classes/Engine/EngineBaseTypes.h
@@ -541,6 +541,43 @@ enum ETravelType
 	TRAVEL_MAX,
 };
 
+/** Types of demo play failures broadcast from the engine */
+UENUM(BlueprintType)
+namespace EDemoPlayFailure
+{
+	enum Type
+	{
+		/** A Generic failure */
+		Generic,
+		/** Demo was not found */
+		DemoNotFound,
+		/** Demo is corrupt */
+		Corrupt,
+		/** Invalid version */
+		InvalidVersion,
+	};
+}
+
+namespace EDemoPlayFailure
+{
+	inline const TCHAR* ToString(EDemoPlayFailure::Type FailureType)
+	{
+		switch (FailureType)
+		{
+		case Generic:
+			return TEXT("Gneric");
+		case DemoNotFound:
+			return TEXT("DemoNotFound");
+		case Corrupt:
+			return TEXT("Corrupt");
+		case InvalidVersion:
+			return TEXT("InvalidVersion");
+		}
+
+		return TEXT("Unknown EDemoPlayFailure error occurred.");
+	}
+}
+
 //URL structure.
 USTRUCT()
 struct ENGINE_API FURL
diff --git a/Engine/Source/Runtime/Engine/Classes/Engine/GameInstance.h b/Engine/Source/Runtime/Engine/Classes/Engine/GameInstance.h
--- a/Engine/Source/Runtime/Engine/Classes/Engine/GameInstance.h
+++ b/Engine/Source/Runtime/Engine/Classes/Engine/GameInstance.h
@@ -153,5 +153,8 @@ public:
 	TArray<class ULocalPlayer*>::TConstIterator	GetLocalPlayerIterator();
 	const TArray<class ULocalPlayer*> &			GetLocalPlayers();
 
+	/** Called when demo playback fails for any reason */
+	virtual void			HandleDemoPlaybackFailure( EDemoPlayFailure::Type FailureType, const FString& ErrorString = TEXT("") ) { }
+
 	static void				AddReferencedObjects(UObject* InThis, FReferenceCollector& Collector);
 };
\ No newline at end of file
diff --git a/Engine/Source/Runtime/Engine/Private/DemoNetDriver.cpp b/Engine/Source/Runtime/Engine/Private/DemoNetDriver.cpp
--- a/Engine/Source/Runtime/Engine/Private/DemoNetDriver.cpp
+++ b/Engine/Source/Runtime/Engine/Private/DemoNetDriver.cpp
@@ -14,6 +14,7 @@
 #include "Engine/LevelStreamingKismet.h"
 #include "GameFramework/PlayerStart.h"
 #include "GameFramework/SpectatorPawnMovement.h"
+#include "Engine/GameInstance.h"
 
 DEFINE_LOG_CATEGORY_STATIC( LogDemo, Log, All );
 
@@ -70,9 +71,24 @@ FString UDemoNetDriver::LowLevelGetNetworkNumber()
 
 bool UDemoNetDriver::InitConnect( FNetworkNotify* InNotify, const FURL& ConnectURL, FString& Error )
 {
+	if ( GetWorld() == nullptr )
+	{
+		UE_LOG( LogDemo, Error, TEXT( "GetWorld() == nullptr" ) );
+		return false;
+	}
+
+	if ( GetWorld()->GetGameInstance() == nullptr )
+	{
+		UE_LOG( LogDemo, Error, TEXT( "GetWorld()->GetGameInstance() == nullptr" ) );
+		return false;
+	}
+
+	UGameInstance* GameInstance = GetWorld()->GetGameInstance();
+
 	// handle default initialization
 	if ( !InitBase( true, InNotify, ConnectURL, false, Error ) )
 	{
+		GameInstance->HandleDemoPlaybackFailure( EDemoPlayFailure::Generic, FString( TEXT( "InitBase FAILED" ) ) );
 		return false;
 	}
 
@@ -83,6 +99,7 @@ bool UDemoNetDriver::InitConnect( FNetworkNotify* InNotify, const FURL& ConnectU
 	{
 		Error = FString::Printf( TEXT( "Couldn't open demo file %s for reading" ), *DemoFilename );
 		UE_LOG( LogDemo, Error, TEXT( "UDemoNetDriver::InitConnect: %s" ), *Error );
+		GameInstance->HandleDemoPlaybackFailure( EDemoPlayFailure::DemoNotFound, FString( EDemoPlayFailure::ToString( EDemoPlayFailure::DemoNotFound ) ) );
 		return false;
 	}
 
@@ -115,16 +132,17 @@ bool UDemoNetDriver::InitConnect( FNetworkNotify* InNotify, const FURL& ConnectU
 	FURL DemoURL;
 	DemoURL.Map = LevelName;
 
-	FWorldContext * WorldContext = GEngine->GetWorldContextFromWorld( World );
+	FWorldContext * WorldContext = GEngine->GetWorldContextFromWorld( GetWorld() );
 
 	if ( WorldContext == NULL )
 	{
 		Error = FString::Printf( TEXT( "No world context" ), *DemoFilename );
 		UE_LOG( LogDemo, Error, TEXT( "UDemoNetDriver::InitConnect: %s" ), *Error );
+		GameInstance->HandleDemoPlaybackFailure( EDemoPlayFailure::Generic, FString( TEXT( "No world context" ) ) );
 		return false;
 	}
 
-	World->DemoNetDriver = NULL;
+	GetWorld()->DemoNetDriver = NULL;
 	SetWorld( NULL );
 
 	UDemoPendingNetGame * NewPendingNetGame = new UDemoPendingNetGame( FObjectInitializer() );
@@ -137,6 +155,7 @@ bool UDemoNetDriver::InitConnect( FNetworkNotify* InNotify, const FURL& ConnectU
 	{
 		Error = LoadMapError;
 		UE_LOG( LogDemo, Error, TEXT( "UDemoNetDriver::InitConnect: LoadMap failed: failed: %s" ), *Error );
+		GameInstance->HandleDemoPlaybackFailure( EDemoPlayFailure::Generic, FString( TEXT( "LoadMap failed" ) ) );
 		return false;
 	}
 
diff --git a/Engine/Source/Runtime/Engine/Private/World.cpp b/Engine/Source/Runtime/Engine/Private/World.cpp
--- a/Engine/Source/Runtime/Engine/Private/World.cpp
+++ b/Engine/Source/Runtime/Engine/Private/World.cpp
@@ -72,6 +72,7 @@
 #include "PhysicsEngine/PhysicsSettings.h"
 #include "Engine/Polys.h"
 #include "Engine/LightMapTexture2D.h"
+#include "Engine/GameInstance.h"
 
 DEFINE_LOG_CATEGORY_STATIC(LogWorld, Log, All);
 DEFINE_LOG_CATEGORY(LogSpawn);
@@ -2687,13 +2688,20 @@ bool UWorld::HandleDemoPlayCommand( const TCHAR* Cmd, FOutputDevice& Ar, UWorld*
 			Ar.Logf( TEXT( "Demo playback failed: %s" ), *Error );
 			DestroyDemoNetDriver();
 		}
-
-		FCoreUObjectDelegates::PostDemoPlay.Broadcast();
+		else
+		{
+			FCoreUObjectDelegates::PostDemoPlay.Broadcast();
+		}
 #endif
 	}
 	else
 	{
 		Ar.Log( TEXT( "You must specify a filename" ) );
+
+		if ( GetGameInstance() != nullptr )
+		{
+			GetGameInstance()->HandleDemoPlaybackFailure( EDemoPlayFailure::Generic, FString( TEXT( "You must specify a filename " ) ) );
+		}
 	}
 
 	return true;
