Collision traces by profile

* Added ability to do collision traces using a collision profile instead of a trace channel and response params.  The UWorld trace functions are the same except they end with "ByProfile" (ex. SweepSingleByProfile).

* Added a custom detail widget that allows collision profiles to be set as BP properties.

* Added a custom graph pin widget for setting collision profiles as pins on a node.  It also supports the property being set to expose on spawn.

* Added FCollisionProfileName struct.  This is a BlueprintType used by the custom widgets.

* Made the custom widgets work with profile redirects.  The property will be fixed up next time the BP is compiled.

* Registered a custom setter for FCollisionProfileName pin to work when set to expose on spawn.

diff --git a/Engine/Source/Editor/BlueprintGraph/Private/EdGraphSchema_K2.cpp b/Engine/Source/Editor/BlueprintGraph/Private/EdGraphSchema_K2.cpp
--- a/Engine/Source/Editor/BlueprintGraph/Private/EdGraphSchema_K2.cpp
+++ b/Engine/Source/Editor/BlueprintGraph/Private/EdGraphSchema_K2.cpp
@@ -3959,6 +3959,11 @@ UFunction* UEdGraphSchema_K2::FindSetVariableByNameFunction(const FEdGraphPinTyp
 		static FName SetTransformName(GET_FUNCTION_NAME_CHECKED(UKismetSystemLibrary, SetTransformPropertyByName));
 		SetFunctionName = SetTransformName;
 	}
+	else if (PinType.PinCategory == K2Schema->PC_Struct && PinType.PinSubCategoryObject == FCollisionProfileName::StaticStruct())
+	{
+		static FName SetStructureName(GET_FUNCTION_NAME_CHECKED(UKismetSystemLibrary, SetCollisionProfileNameProperty));
+		SetFunctionName = SetStructureName;
+	}
 	else if (PinType.PinCategory == K2Schema->PC_Struct && FIsCustomStructureParamHelper::Is(PinType.PinSubCategoryObject.Get()))
 	{
 		static FName SetStructureName(GET_FUNCTION_NAME_CHECKED(UKismetSystemLibrary, SetStructurePropertyByName));
diff --git a/Engine/Source/Editor/DetailCustomizations/Private/CollisionProfileNameCustomization.cpp b/Engine/Source/Editor/DetailCustomizations/Private/CollisionProfileNameCustomization.cpp
new file mode 100644
--- /dev/null
+++ b/Engine/Source/Editor/DetailCustomizations/Private/CollisionProfileNameCustomization.cpp
@@ -0,0 +1,181 @@
+// Copyright 1998-2014 Epic Games, Inc. All Rights Reserved.
+
+#include "DetailCustomizationsPrivatePCH.h"
+#include "CollisionProfileNameCustomization.h"
+#include "SNameComboBox.h"
+#include "Engine/CollisionProfile.h"
+
+
+#define LOCTEXT_NAMESPACE "CollisionProfileNameCustomization"
+
+
+FCollisionProfileNameCustomization::FCollisionProfileNameCustomization()
+{
+	UCollisionProfile::GetProfileNames(NameList);
+}
+
+
+TSharedRef<IPropertyTypeCustomization> FCollisionProfileNameCustomization::MakeInstance()
+{
+	return MakeShareable(new FCollisionProfileNameCustomization);
+}
+
+
+void FCollisionProfileNameCustomization::CustomizeHeader(TSharedRef<IPropertyHandle> StructPropertyHandle, class FDetailWidgetRow& HeaderRow, IPropertyTypeCustomizationUtils& StructCustomizationUtils)
+{
+}
+
+
+void FCollisionProfileNameCustomization::CustomizeChildren(TSharedRef<class IPropertyHandle> StructPropertyHandle, class IDetailChildrenBuilder& StructBuilder, IPropertyTypeCustomizationUtils& StructCustomizationUtils)
+{
+	NameHandle = StructPropertyHandle->GetChildHandle(GET_MEMBER_NAME_CHECKED(FCollisionProfileName, Name));
+	check(NameHandle.IsValid());
+
+	TSharedPtr<FName> InitialSelectedName = GetSelectedName();
+	if (InitialSelectedName.IsValid())
+	{
+		SetPropertyWithName(*InitialSelectedName.Get());
+	}
+
+	IDetailGroup& CollisionGroup = StructBuilder.AddChildGroup(TEXT("Collision"), LOCTEXT("CollisionPresetName", "Collision Preset").ToString());
+	CollisionGroup.HeaderRow()
+	.NameContent()
+	[
+		SNew(STextBlock)
+		.Text(StructPropertyHandle->GetPropertyDisplayName())
+		.Font(IDetailLayoutBuilder::GetDetailFont())
+	]
+	.ValueContent()
+	[
+		SAssignNew(NameComboBox, SComboBox<TSharedPtr<FName>>)
+		.OptionsSource(&NameList)
+		.OnGenerateWidget(this, &FCollisionProfileNameCustomization::OnGenerateWidget)
+		.OnSelectionChanged(this, &FCollisionProfileNameCustomization::OnSelectionChanged, &CollisionGroup)
+		.OnComboBoxOpening(this, &FCollisionProfileNameCustomization::OnComboBoxOpening)
+		.InitiallySelectedItem(InitialSelectedName)
+		.IsEnabled(FSlateApplication::Get().GetNormalExecutionAttribute())
+		.ContentPadding(FMargin(2.0f, 2.0f))
+		.Content()
+		[
+			SNew(STextBlock)
+			.Text(this, &FCollisionProfileNameCustomization::GetProfileComboBoxContent)
+			.Font(IDetailLayoutBuilder::GetDetailFont())
+			.ToolTipText(this, &FCollisionProfileNameCustomization::GetProfileComboBoxToolTip)
+		]
+	];
+}
+
+
+TSharedRef<SWidget> FCollisionProfileNameCustomization::OnGenerateWidget(TSharedPtr<FName> InItem)
+{
+	return
+		SNew(STextBlock)
+		.Text(InItem->ToString())
+		.Font(IDetailLayoutBuilder::GetDetailFont());
+}
+
+
+void FCollisionProfileNameCustomization::OnSelectionChanged(TSharedPtr<FName> NameItem, ESelectInfo::Type SelectInfo, IDetailGroup* CollisionGroup)
+{
+	if (NameItem.IsValid())
+	{
+		SetPropertyWithName(*NameItem);
+	}
+}
+
+
+void FCollisionProfileNameCustomization::OnComboBoxOpening()
+{
+	TSharedPtr<FName> SelectedName = GetSelectedName();
+	if (SelectedName.IsValid())
+	{
+		check(NameComboBox.IsValid());
+		NameComboBox->SetSelectedItem(SelectedName);
+	}
+}
+
+
+TSharedPtr<FName> FCollisionProfileNameCustomization::GetSelectedName() const
+{
+	int32 NameCount = NameList.Num();
+	if (NameCount <= 0)
+	{
+		return NULL;
+	}
+
+	FName Name;
+	GetPropertyAsName(Name);
+
+	for (int32 NameIndex = 0; NameIndex < NameCount; ++NameIndex)
+	{
+		if (Name == *NameList[NameIndex].Get())
+		{
+			return NameList[NameIndex];
+		}
+	}
+
+	return NameList[0];
+}
+
+
+void FCollisionProfileNameCustomization::SetPropertyWithName(const FName& Name)
+{
+	check(NameHandle.IsValid());
+
+	FName OldName;
+	NameHandle->GetValue(OldName);
+
+	if (OldName != Name)
+	{
+		NameHandle->SetValue(Name);
+	}
+}
+
+
+void FCollisionProfileNameCustomization::GetPropertyAsName(FName& OutName) const
+{
+	check(NameHandle.IsValid());
+	NameHandle->GetValue(OutName);
+
+	UCollisionProfile* CollisionProfile = UCollisionProfile::Get();
+	check(CollisionProfile);
+
+	const FName* RedirectName = CollisionProfile->LookForProfileRedirect(OutName);
+	if (RedirectName)
+	{
+		OutName = *RedirectName;
+	}
+}
+
+
+FString FCollisionProfileNameCustomization::GetProfileComboBoxContent() const
+{
+	TSharedPtr<FName> SelectedName = GetSelectedName();
+	if (SelectedName.IsValid())
+	{
+		return (*SelectedName->ToString());
+	}
+
+	return TEXT("Invalid");
+}
+
+
+FString FCollisionProfileNameCustomization::GetProfileComboBoxToolTip() const
+{
+	UCollisionProfile* CollisionProfile = UCollisionProfile::Get();
+	check(CollisionProfile);
+
+	FName ProfileName;
+	GetPropertyAsName(ProfileName);
+
+	FCollisionResponseTemplate ProfileTemplate;
+	if (CollisionProfile->GetProfileTemplate(ProfileName, ProfileTemplate))
+	{
+		return ProfileTemplate.HelpMessage;
+	}
+
+	return TEXT("Invalid");
+}
+
+
+#undef LOCTEXT_NAMESPACE
diff --git a/Engine/Source/Editor/DetailCustomizations/Private/CollisionProfileNameCustomization.h b/Engine/Source/Editor/DetailCustomizations/Private/CollisionProfileNameCustomization.h
new file mode 100644
--- /dev/null
+++ b/Engine/Source/Editor/DetailCustomizations/Private/CollisionProfileNameCustomization.h
@@ -0,0 +1,44 @@
+// Copyright 1998-2014 Epic Games, Inc. All Rights Reserved.
+
+#pragma once
+
+
+class SNameComboBox;
+
+
+/**
+ * Customizes a CollisionProfileName property to use a dropdown
+ */
+class FCollisionProfileNameCustomization : public IPropertyTypeCustomization
+{
+public:
+
+	FCollisionProfileNameCustomization();
+
+	static TSharedRef<IPropertyTypeCustomization> MakeInstance();
+
+	/** IPropertyTypeCustomization interface */
+	virtual void CustomizeHeader(TSharedRef<IPropertyHandle> StructPropertyHandle, class FDetailWidgetRow& HeaderRow, IPropertyTypeCustomizationUtils& StructCustomizationUtils) override;
+	virtual void CustomizeChildren(TSharedRef<class IPropertyHandle> StructPropertyHandle, class IDetailChildrenBuilder& StructBuilder, IPropertyTypeCustomizationUtils& StructCustomizationUtils) override;
+
+protected:
+
+	TSharedRef<SWidget> OnGenerateWidget(TSharedPtr<FName> InItem);
+
+	void OnSelectionChanged(TSharedPtr<FName> NameItem, ESelectInfo::Type SelectInfo, IDetailGroup* CollisionGroup);
+	void OnComboBoxOpening();
+
+	TSharedPtr<FName> GetSelectedName() const;
+
+	void SetPropertyWithName(const FName& Name);
+	void GetPropertyAsName(FName& OutName) const;
+
+	FString GetProfileComboBoxContent() const;
+	FString GetProfileComboBoxToolTip() const;
+
+protected:
+
+	TSharedPtr<IPropertyHandle> NameHandle;
+	TArray<TSharedPtr<FName>> NameList;
+	TSharedPtr<SComboBox<TSharedPtr<FName>>> NameComboBox;
+};
diff --git a/Engine/Source/Editor/DetailCustomizations/Private/DetailCustomizations.cpp b/Engine/Source/Editor/DetailCustomizations/Private/DetailCustomizations.cpp
--- a/Engine/Source/Editor/DetailCustomizations/Private/DetailCustomizations.cpp
+++ b/Engine/Source/Editor/DetailCustomizations/Private/DetailCustomizations.cpp
@@ -183,6 +183,7 @@ void FDetailCustomizationsModule::RegisterPropertyTypeCustomizations()
 	RegisterCustomPropertyTypeLayout("RuntimeCurveLinearColor", FOnGetPropertyTypeCustomizationInstance::CreateStatic(&FCurveColorCustomization::MakeInstance));
 	RegisterCustomPropertyTypeLayout("ParticleSysParam", FOnGetPropertyTypeCustomizationInstance::CreateStatic(&FParticleSysParamStructCustomization::MakeInstance));
 	RegisterCustomPropertyTypeLayout("RawDistributionVector", FOnGetPropertyTypeCustomizationInstance::CreateStatic(&FRawDistributionVectorStructCustomization::MakeInstance));
+	RegisterCustomPropertyTypeLayout("CollisionProfileName", FOnGetPropertyTypeCustomizationInstance::CreateStatic(&FCollisionProfileNameCustomization::MakeInstance));
 }
 
 
diff --git a/Engine/Source/Editor/GraphEditor/Private/KismetPins/SGraphPinCollisionProfile.cpp b/Engine/Source/Editor/GraphEditor/Private/KismetPins/SGraphPinCollisionProfile.cpp
new file mode 100644
--- /dev/null
+++ b/Engine/Source/Editor/GraphEditor/Private/KismetPins/SGraphPinCollisionProfile.cpp
@@ -0,0 +1,126 @@
+// Copyright 1998-2014 Epic Games, Inc. All Rights Reserved.
+
+#include "GraphEditorCommon.h"
+#include "SNameComboBox.h"
+#include "SGraphPinCollisionProfile.h"
+
+
+void SGraphPinCollisionProfile::Construct(const FArguments& InArgs, UEdGraphPin* InGraphPinObj)
+{
+	SGraphPin::Construct(SGraphPin::FArguments(), InGraphPinObj);
+}
+
+
+TSharedRef<SWidget>	SGraphPinCollisionProfile::GetDefaultValueWidget()
+{
+	UCollisionProfile::GetProfileNames(NameList);
+
+	TSharedPtr<FName> InitialSelectedName = GetSelectedName();
+	if (InitialSelectedName.IsValid())
+	{
+		SetPropertyWithName(*InitialSelectedName.Get());
+	}
+
+	return SAssignNew(NameComboBox, SNameComboBox)
+		.ContentPadding(FMargin(6.0f, 2.0f))
+		.OptionsSource(&NameList)
+		.InitiallySelectedItem(InitialSelectedName)
+		.OnSelectionChanged(this, &SGraphPinCollisionProfile::OnSelectionChanged)
+		.OnComboBoxOpening(this, &SGraphPinCollisionProfile::OnComboBoxOpening)
+		.Visibility(this, &SGraphPin::GetDefaultValueVisibility);
+}
+
+
+void SGraphPinCollisionProfile::OnSelectionChanged(TSharedPtr<FName> NameItem, ESelectInfo::Type SelectInfo)
+{
+	if (NameItem.IsValid())
+	{
+		SetPropertyWithName(*NameItem);
+	}
+}
+
+
+void SGraphPinCollisionProfile::OnComboBoxOpening()
+{
+	TSharedPtr<FName> SelectedName = GetSelectedName();
+	if (SelectedName.IsValid())
+	{
+		check(NameComboBox.IsValid());
+		NameComboBox->SetSelectedItem(SelectedName);
+	}
+}
+
+
+TSharedPtr<FName> SGraphPinCollisionProfile::GetSelectedName() const
+{
+	int32 NameCount = NameList.Num();
+	if (NameCount <= 0)
+	{
+		return NULL;
+	}
+
+	FName Name;
+	GetPropertyAsName(Name);
+
+	for (int32 NameIndex = 0; NameIndex < NameCount; ++NameIndex)
+	{
+		if (Name == *NameList[NameIndex].Get())
+		{
+			return NameList[NameIndex];
+		}
+	}
+
+	return NameList[0];
+}
+
+
+void SGraphPinCollisionProfile::SetPropertyWithName(const FName& Name)
+{
+	check(GraphPinObj);
+	check(GraphPinObj->PinType.PinSubCategoryObject == FCollisionProfileName::StaticStruct());
+
+	FString PinString = TEXT("(Name=\"");
+	PinString += *Name.ToString();
+	PinString += TEXT("\")");
+
+	if (PinString != GraphPinObj->GetDefaultAsString())
+	{
+		GraphPinObj->GetSchema()->TrySetDefaultValue(*GraphPinObj, PinString);
+	}
+}
+
+
+void SGraphPinCollisionProfile::GetPropertyAsName(FName& OutName) const
+{
+	check(GraphPinObj);
+	check(GraphPinObj->PinType.PinSubCategoryObject == FCollisionProfileName::StaticStruct());
+
+	FString PinString = GraphPinObj->GetDefaultAsString();
+
+	if (PinString.StartsWith(TEXT("(")) && PinString.EndsWith(TEXT(")")))
+	{
+		PinString = PinString.LeftChop(1);
+		PinString = PinString.RightChop(1);
+		PinString.Split("=", NULL, &PinString);
+
+		if (PinString.StartsWith(TEXT("\"")) && PinString.EndsWith(TEXT("\"")))
+		{
+			PinString = PinString.LeftChop(1);
+			PinString = PinString.RightChop(1);
+		}
+	}
+
+	if (!PinString.IsEmpty())
+	{
+		OutName = *PinString;
+
+		UCollisionProfile* CollisionProfile = UCollisionProfile::Get();
+		check(CollisionProfile);
+
+		const FName* RedirectName = CollisionProfile->LookForProfileRedirect(OutName);
+		if (RedirectName)
+		{
+			OutName = *RedirectName;
+		}
+	}
+}
diff --git a/Engine/Source/Editor/GraphEditor/Private/KismetPins/SGraphPinCollisionProfile.h b/Engine/Source/Editor/GraphEditor/Private/KismetPins/SGraphPinCollisionProfile.h
new file mode 100644
--- /dev/null
+++ b/Engine/Source/Editor/GraphEditor/Private/KismetPins/SGraphPinCollisionProfile.h
@@ -0,0 +1,37 @@
+// Copyright 1998-2014 Epic Games, Inc. All Rights Reserved.
+
+#pragma once
+
+
+class SNameComboBox;
+
+
+/**
+ * Customizes a CollisionProfileName graph pin to use a dropdown
+ */
+class SGraphPinCollisionProfile : public SGraphPin
+{
+public:
+
+	SLATE_BEGIN_ARGS(SGraphPinCollisionProfile) {}
+	SLATE_END_ARGS()
+
+	void Construct(const FArguments& InArgs, UEdGraphPin* InGraphPinObj);
+
+protected:
+
+	virtual TSharedRef<SWidget>	GetDefaultValueWidget() override;
+
+	void OnSelectionChanged(TSharedPtr<FName> NameItem, ESelectInfo::Type SelectInfo);
+	void OnComboBoxOpening();
+
+	TSharedPtr<FName> GetSelectedName() const;
+
+	void SetPropertyWithName(const FName& Name);
+	void GetPropertyAsName(FName& OutName) const;
+
+protected:
+
+	TArray<TSharedPtr<FName>> NameList;
+	TSharedPtr<SNameComboBox> NameComboBox;
+};
diff --git a/Engine/Source/Editor/GraphEditor/Private/NodeFactory.cpp b/Engine/Source/Editor/GraphEditor/Private/NodeFactory.cpp
--- a/Engine/Source/Editor/GraphEditor/Private/NodeFactory.cpp
+++ b/Engine/Source/Editor/GraphEditor/Private/NodeFactory.cpp
@@ -64,6 +64,7 @@
 #include "KismetPins/SGraphPinVector2D.h"
 #include "NiagaraPins/SGraphPinVector4.h"
 #include "KismetPins/SGraphPinIndex.h"
+#include "KismetPins/SGraphPinCollisionProfile.h"
 
 #include "SoundNodes/SGraphNodeSoundBase.h"
 #include "SoundNodes/SGraphNodeSoundResult.h"
@@ -328,6 +329,10 @@ TSharedPtr<SGraphPin> FNodeFactory::CreatePinWidget(UEdGraphPin* InPin)
 			{
 				return SNew(SGraphPinPose, InPin);
 			}
+			else if (InPin->PinType.PinSubCategoryObject == FCollisionProfileName::StaticStruct())
+			{
+				return SNew(SGraphPinCollisionProfile, InPin);
+			}
 		}
 		else if (InPin->PinType.PinCategory == K2Schema->PC_Byte)
 		{
diff --git a/Engine/Source/Runtime/Engine/Classes/Engine/CollisionProfile.h b/Engine/Source/Runtime/Engine/Classes/Engine/CollisionProfile.h
--- a/Engine/Source/Runtime/Engine/Classes/Engine/CollisionProfile.h
+++ b/Engine/Source/Runtime/Engine/Classes/Engine/CollisionProfile.h
@@ -10,6 +10,20 @@
 #include "CollisionProfile.generated.h"
 
 
+USTRUCT(BlueprintType)
+struct FCollisionProfileName
+{
+	GENERATED_USTRUCT_BODY()
+
+	FCollisionProfileName()
+		: Name(NAME_None)
+	{}
+
+	UPROPERTY(EditAnywhere, Category = Collision)
+	FName Name;
+};
+
+
 /**
  * Structure for collision response templates.
  */
@@ -160,12 +174,18 @@ public:
 	/** Accessor and initializer **/
 	ENGINE_API static UCollisionProfile* Get();
 
+	/** Fill up the array with the profile names **/
+	ENGINE_API static void GetProfileNames(TArray<TSharedPtr<FName>>& OutNameList);
+
 	/** Fill up the loaded config of the profile name to the BodyInstance **/
 	bool ReadConfig(FName ProfileName, struct FBodyInstance& BodyInstance) const;
 
 	/** Fill up the loaded config of the profile name to the BodyInstance **/
 	ENGINE_API bool GetProfileTemplate(FName ProfileName, struct FCollisionResponseTemplate& ProfileData) const;
 
+	/** Check if this profile name has been redirected **/
+	ENGINE_API const FName* LookForProfileRedirect(FName ProfileName) const;
+
 	/** Accessor for UI customization **/
 	int32 GetNumOfProfiles() const { return Profiles.Num(); }
 
diff --git a/Engine/Source/Runtime/Engine/Classes/Engine/World.h b/Engine/Source/Runtime/Engine/Classes/Engine/World.h
--- a/Engine/Source/Runtime/Engine/Classes/Engine/World.h
+++ b/Engine/Source/Runtime/Engine/Classes/Engine/World.h
@@ -1029,6 +1029,16 @@ public:
 	 */
 	bool LineTraceTest(const FVector& Start,const FVector& End,const struct FCollisionQueryParams& Params, const struct FCollisionObjectQueryParams& ObjectQueryParams) const;
 
+	/**
+	 *  Trace a ray against the world using a specific profile and return if a blocking hit is found.
+	 *  @param  Start           Start location of the ray
+	 *  @param  End             End location of the ray
+	 *  @param  ProfileName     The 'profile' used to determine which components to hit
+	 *  @param  Params          Additional parameters used for the trace
+	 *  @return TRUE if a blocking hit is found
+	 */
+	bool LineTraceTestByProfile(const FVector& Start, const FVector& End, FName ProfileName, const struct FCollisionQueryParams& Params) const;
+
 	/**
 	 *  Trace a ray against the world using a specific channel and return the first blocking hit
 	 *  @param  OutHit          First blocking hit found
@@ -1052,6 +1062,17 @@ public:
 	 */
 	bool LineTraceSingle(struct FHitResult& OutHit,const FVector& Start,const FVector& End,const struct FCollisionQueryParams& Params, const struct FCollisionObjectQueryParams& ObjectQueryParams) const;
 
+	/**
+	 *  Trace a ray against the world using a specific profile and return the first blocking hit
+	 *  @param  OutHit          First blocking hit found
+	 *  @param  Start           Start location of the ray
+	 *  @param  End             End location of the ray
+	 *  @param  ProfileName     The 'profile' used to determine which components to hit
+	 *  @param  Params          Additional parameters used for the trace
+	 *  @return TRUE if a blocking hit is found
+	 */
+	bool LineTraceSingleByProfile(struct FHitResult& OutHit, const FVector& Start, const FVector& End, FName ProfileName, const struct FCollisionQueryParams& Params) const;
+
 	/**
 	 *  Trace a ray against the world using a specific channel and return overlapping hits and then first blocking hit
 	 *  Results are sorted, so a blocking hit (if found) will be the last element of the array
@@ -1079,6 +1100,19 @@ public:
 	 */
 	bool LineTraceMulti(TArray<struct FHitResult>& OutHits,const FVector& Start,const FVector& End,const struct FCollisionQueryParams& Params, const struct FCollisionObjectQueryParams& ObjectQueryParams) const;
 
+	/**
+	 *  Trace a ray against the world using a specific profile and return overlapping hits and then first blocking hit
+	 *  Results are sorted, so a blocking hit (if found) will be the last element of the array
+	 *  Only the single closest blocking result will be generated, no tests will be done after that
+	 *  @param  OutHits         Array of hits found between ray and the world
+	 *  @param  Start           Start location of the ray
+	 *  @param  End             End location of the ray
+	 *  @param  ProfileName     The 'profile' used to determine which components to hit
+	 *  @param  Params          Additional parameters used for the trace
+	 *  @return TRUE if OutHits contains any blocking hit entries
+	 */
+	bool LineTraceMultiByProfile(TArray<struct FHitResult>& OutHits, const FVector& Start, const FVector& End, FName ProfileName, const struct FCollisionQueryParams& Params) const;
+
 	/**
 	 *  Sweep a sphere against the world using a specific channel and return if a blocking hit is found.
 	 *  @param  Start           Start location of the sphere
@@ -1102,6 +1136,17 @@ public:
 	 */
 	bool SweepTest(const FVector& Start, const FVector& End, const FQuat& Rot, const struct FCollisionShape& CollisionShape, const struct FCollisionQueryParams& Params, const struct FCollisionObjectQueryParams& ObjectQueryParams) const;
 
+	/**
+	 *  Sweep a sphere against the world using a specific profile and return if a blocking hit is found.
+	 *  @param  Start           Start location of the sphere
+	 *  @param  End             End location of the sphere
+	 *  @param  ProfileName     The 'profile' used to determine which components to hit
+	 *  @param	CollisionShape	CollisionShape - supports Box, Sphere, Capsule
+	 *  @param  Params          Additional parameters used for the trace
+	 *  @return TRUE if a blocking hit is found
+	 */
+	bool SweepTestByProfile(const FVector& Start, const FVector& End, const FQuat& Rot, FName ProfileName, const struct FCollisionShape& CollisionShape, const struct FCollisionQueryParams& Params) const;
+
 	/**
 	 *  Sweep a sphere against the world and return the first blocking hit using a specific channel
 	 *  @param  OutHit          First blocking hit found
@@ -1127,6 +1172,18 @@ public:
 	 */
 	bool SweepSingle(struct FHitResult& OutHit, const FVector& Start, const FVector& End, const FQuat& Rot, const struct FCollisionShape& CollisionShape, const struct FCollisionQueryParams& Params, const struct FCollisionObjectQueryParams& ObjectQueryParams) const;
 
+	/**
+	 *  Sweep a sphere against the world and return the first blocking hit using a specific profile
+	 *  @param  OutHit          First blocking hit found
+	 *  @param  Start           Start location of the sphere
+	 *  @param  End             End location of the sphere
+	 *  @param  ProfileName     The 'profile' used to determine which components to hit
+	 *  @param	CollisionShape	CollisionShape - supports Box, Sphere, Capsule
+	 *  @param  Params          Additional parameters used for the trace
+	 *  @return TRUE if OutHits contains any blocking hit entries
+	 */
+	bool SweepSingleByProfile(struct FHitResult& OutHit, const FVector& Start, const FVector& End, const FQuat& Rot, FName ProfileName, const struct FCollisionShape& CollisionShape, const struct FCollisionQueryParams& Params) const;
+
 	/**
 	 *  Sweep a sphere against the world and return all initial overlaps using a specific channel (including blocking) if requested, then overlapping hits and then first blocking hit
 	 *  Results are sorted, so a blocking hit (if found) will be the last element of the array
@@ -1156,6 +1213,20 @@ public:
 	 */
 	bool SweepMulti(TArray<struct FHitResult>& OutHits, const FVector& Start, const FVector& End, const FQuat& Rot, const struct FCollisionShape& CollisionShape, const struct FCollisionQueryParams& Params, const struct FCollisionObjectQueryParams& ObjectQueryParams) const;
 
+	/**
+	 *  Sweep a sphere against the world and return all initial overlaps using a specific profile, then overlapping hits and then first blocking hit
+	 *  Results are sorted, so a blocking hit (if found) will be the last element of the array
+	 *  Only the single closest blocking result will be generated, no tests will be done after that
+	 *  @param  OutHits         Array of hits found between ray and the world
+	 *  @param  Start           Start location of the sphere
+	 *  @param  End             End location of the sphere
+	 *  @param  ProfileName     The 'profile' used to determine which components to hit
+	 *  @param	CollisionShape	CollisionShape - supports Box, Sphere, Capsule
+	 *  @param  Params          Additional parameters used for the trace
+	 *  @return TRUE if OutHits contains any blocking hit entries
+	 */
+	bool SweepMultiByProfile(TArray<FHitResult>& OutHits, const FVector& Start, const FVector& End, const FQuat& Rot, FName ProfileName, const struct FCollisionShape& CollisionShape, const struct FCollisionQueryParams& Params) const;
+
 	/**
 	 *  Test the collision of an AABB at the supplied location using a specific channel, and return if any blocking overlap is found
 	 *  @param  Pos             Location of center of box to test against the world
@@ -1177,6 +1248,16 @@ public:
 	 */
 	bool OverlapTest(const FVector& Pos, const FQuat& Rot, const struct FCollisionShape& CollisionShape, const struct FCollisionQueryParams& Params, const struct FCollisionObjectQueryParams& ObjectQueryParams) const;
 
+	/**
+	 *  Test the collision of an AABB at the supplied location using a specific profile, and return if any blocking overlap is found
+	 *  @param  Pos             Location of center of box to test against the world
+	 *  @param  ProfileName     The 'profile' used to determine which components to hit
+	 *  @param	CollisionShape	CollisionShape - supports Box, Sphere, Capsule
+	 *  @param  Params          Additional parameters used for the trace
+	 *  @return TRUE if any blocking results are found
+	 */
+	bool OverlapTestByProfile(const FVector& Pos, const FQuat& Rot, FName ProfileName, const struct FCollisionShape& CollisionShape, const struct FCollisionQueryParams& Params) const;
+
 	/**
 	 *  Test the collision of a sphere at the supplied location using a specific channel, and determine the set of components that it overlaps
 	 *  @param  OutOverlaps     Array of components found to overlap supplied box
@@ -1200,6 +1281,17 @@ public:
 	 */
 	bool OverlapSingle(struct FOverlapResult& OutOverlap, const FVector& Pos, const FQuat& Rot, const struct FCollisionShape& CollisionShape, const struct FCollisionQueryParams& Params, const struct FCollisionObjectQueryParams& ObjectQueryParams) const;
 
+	/**
+	 *  Test the collision of a sphere at the supplied location using a specific profile, and determine the set of components that it overlaps
+	 *  @param  OutOverlaps     Array of components found to overlap supplied box
+	 *  @param  Pos             Location of center of sphere to test against the world
+	 *  @param  ProfileName     The 'profile' used to determine which components to hit
+	 *  @param	CollisionShape	CollisionShape - supports Box, Sphere, Capsule
+	 *  @param  Params          Additional parameters used for the trace
+	 *  @return TRUE if OutOverlaps contains any blocking results
+	 */
+	bool OverlapSingleByProfile(struct FOverlapResult& OutOverlap, const FVector& Pos, const FQuat& Rot, FName ProfileName, const struct FCollisionShape& CollisionShape, const struct FCollisionQueryParams& Params) const;
+
 	/**
 	 *  Test the collision of a sphere at the supplied location using a specific channel, and determine the set of components that it overlaps
 	 *  @param  OutOverlaps     Array of components found to overlap supplied box
@@ -1223,6 +1315,17 @@ public:
 	 */
 	bool OverlapMulti(TArray<struct FOverlapResult>& OutOverlaps, const FVector& Pos, const FQuat& Rot, const struct FCollisionShape& CollisionShape, const struct FCollisionQueryParams& Params, const struct FCollisionObjectQueryParams& ObjectQueryParams) const;
 	
+	/**
+	 *  Test the collision of a sphere at the supplied location using a specific profile, and determine the set of components that it overlaps
+	 *  @param  OutOverlaps     Array of components found to overlap supplied box
+	 *  @param  Pos             Location of center of sphere to test against the world
+	 *  @param  ProfileName     The 'profile' used to determine which components to hit
+	 *  @param	CollisionShape	CollisionShape - supports Box, Sphere, Capsule
+	 *  @param  Params          Additional parameters used for the trace
+	 *  @return TRUE if OutOverlaps contains any blocking results
+	 */
+	bool OverlapMultiByProfile(TArray<struct FOverlapResult>& OutOverlaps, const FVector& Pos, const FQuat& Rot, FName ProfileName, const struct FCollisionShape& CollisionShape, const struct FCollisionQueryParams& Params) const;
+
 	// COMPONENT SWEEP
 
 	/**
diff --git a/Engine/Source/Runtime/Engine/Classes/Kismet/KismetSystemLibrary.h b/Engine/Source/Runtime/Engine/Classes/Kismet/KismetSystemLibrary.h
--- a/Engine/Source/Runtime/Engine/Classes/Kismet/KismetSystemLibrary.h
+++ b/Engine/Source/Runtime/Engine/Classes/Kismet/KismetSystemLibrary.h
@@ -4,6 +4,7 @@
 
 #include "Engine/LatentActionManager.h"
 #include "BlueprintFunctionLibrary.h"
+#include "Engine/CollisionProfile.h"
 #include "KismetSystemLibrary.generated.h"
 
 /** It's just a proxy for FTimerDynamicDelegate */
@@ -387,6 +388,23 @@ class UKismetSystemLibrary : public UBlueprintFunctionLibrary
 	UFUNCTION(BlueprintCallable, meta=(BlueprintInternalUseOnly = "true", AutoCreateRefTerm = "Value" ))
 	static void SetTransformPropertyByName(UObject* Object, FName PropertyName, const FTransform& Value);
 
+	/** Set a CollisionProfileName property by name */
+	UFUNCTION(BlueprintCallable, CustomThunk, meta = (BlueprintInternalUseOnly = "true", AutoCreateRefTerm = "Value"))
+	static void SetCollisionProfileNameProperty(UObject* Object, FName PropertyName, const FCollisionProfileName& Value);
+
+	DECLARE_FUNCTION(execSetCollisionProfileNameProperty)
+	{
+		P_GET_OBJECT(UObject, OwnerObject);
+		P_GET_PROPERTY(UNameProperty, StructPropertyName);
+
+		Stack.StepCompiledIn<UStructProperty>(NULL);
+		void* SrcStructAddr = Stack.MostRecentPropertyAddress;
+
+		P_FINISH;
+
+		Generic_SetStructurePropertyByName(OwnerObject, StructPropertyName, SrcStructAddr);
+	}
+
 	/** Set a custom structure property by name */
 	UFUNCTION(BlueprintCallable, CustomThunk, meta = (BlueprintInternalUseOnly = "true", CustomStructureParam = "Value"))
 	static void SetStructurePropertyByName(UObject* Object, FName PropertyName, const FGenericStruct& Value);
diff --git a/Engine/Source/Runtime/Engine/Private/Collision/CollisionProfile.cpp b/Engine/Source/Runtime/Engine/Private/Collision/CollisionProfile.cpp
--- a/Engine/Source/Runtime/Engine/Private/Collision/CollisionProfile.cpp
+++ b/Engine/Source/Runtime/Engine/Private/Collision/CollisionProfile.cpp
@@ -60,6 +60,24 @@ UCollisionProfile* UCollisionProfile::Get()
 	return CollisionProfile;
 }
 
+void UCollisionProfile::GetProfileNames(TArray<TSharedPtr<FName>>& OutNameList)
+{
+	UCollisionProfile* CollisionProfile = UCollisionProfile::Get();
+	check(CollisionProfile);
+
+	int32 NumProfiles = CollisionProfile->GetNumOfProfiles();
+
+	OutNameList.Empty(NumProfiles);
+
+	for (int32 ProfileId = 0; ProfileId < NumProfiles; ++ProfileId)
+	{
+		const FCollisionResponseTemplate* ProfileTemplate = CollisionProfile->GetProfileByIndex(ProfileId);
+		check(ProfileTemplate);
+
+		OutNameList.Add(MakeShareable(new FName(ProfileTemplate->Name)));
+	}
+}
+
 bool UCollisionProfile::GetProfileTemplate(FName ProfileName, struct FCollisionResponseTemplate& ProfileData) const
 {
 	// verify if it is in redirect first
@@ -97,6 +115,11 @@ bool UCollisionProfile::CheckRedirect(FName ProfileName, FBodyInstance& BodyInst
 	return false;
 }
 
+const FName* UCollisionProfile::LookForProfileRedirect(FName ProfileName) const
+{
+	return ProfileRedirectsMap.Find(ProfileName);
+}
+
 bool UCollisionProfile::FindProfileData(const TArray<FCollisionResponseTemplate>& ProfileList, FName ProfileName, struct FCollisionResponseTemplate& ProfileData) const
 {
 	if ( ProfileName != NAME_None )
diff --git a/Engine/Source/Runtime/Engine/Private/Collision/WorldCollision.cpp b/Engine/Source/Runtime/Engine/Private/Collision/WorldCollision.cpp
--- a/Engine/Source/Runtime/Engine/Private/Collision/WorldCollision.cpp
+++ b/Engine/Source/Runtime/Engine/Private/Collision/WorldCollision.cpp
@@ -294,6 +294,119 @@ bool UWorld::OverlapSingle(struct FOverlapResult& OutOverlap, const FVector& Pos
 #endif
 }
 
+// profile interfaces
+static void GetCollisionProfileChannelAndResponseParams(FName ProfileName, ECollisionChannel &CollisionChannel, FCollisionResponseParams &ResponseParams)
+{
+	const UCollisionProfile* CollisionProfile = UCollisionProfile::Get();
+	check(CollisionProfile != NULL);
+
+	FCollisionResponseTemplate Template;
+
+	if (CollisionProfile->GetProfileTemplate(ProfileName, Template))
+	{
+		CollisionChannel = Template.ObjectType;
+		ResponseParams = FCollisionResponseParams(Template.ResponseToChannels);
+		return;
+	}
+
+	// Check for redirects
+	const FName* RedirectName = CollisionProfile->LookForProfileRedirect(ProfileName);
+	if (RedirectName && CollisionProfile->GetProfileTemplate(*RedirectName, Template))
+	{
+		CollisionChannel = Template.ObjectType;
+		ResponseParams = FCollisionResponseParams(Template.ResponseToChannels);
+		return;
+	}
+
+	// No profile found
+	UE_LOG(LogPhysics, Warning, TEXT("COLLISION PROFILE [%s] is not found"), *ProfileName.ToString());
+
+	CollisionChannel = ECC_WorldStatic;
+	ResponseParams = FCollisionResponseParams::DefaultResponseParam;
+}
+
+bool UWorld::LineTraceTestByProfile(const FVector& Start, const FVector& End, FName ProfileName, const struct FCollisionQueryParams& Params) const
+{
+	ECollisionChannel TraceChannel;
+	FCollisionResponseParams ResponseParam;
+	GetCollisionProfileChannelAndResponseParams(ProfileName, TraceChannel, ResponseParam);
+
+	return LineTraceTest(Start, End, TraceChannel, Params, ResponseParam);
+}
+
+bool UWorld::LineTraceSingleByProfile(struct FHitResult& OutHit, const FVector& Start, const FVector& End, FName ProfileName, const struct FCollisionQueryParams& Params) const
+{
+	ECollisionChannel TraceChannel;
+	FCollisionResponseParams ResponseParam;
+	GetCollisionProfileChannelAndResponseParams(ProfileName, TraceChannel, ResponseParam);
+
+	return LineTraceSingle(OutHit, Start, End, TraceChannel, Params, ResponseParam);
+}
+
+bool UWorld::LineTraceMultiByProfile(TArray<struct FHitResult>& OutHits, const FVector& Start, const FVector& End, FName ProfileName, const struct FCollisionQueryParams& Params) const
+{
+	ECollisionChannel TraceChannel;
+	FCollisionResponseParams ResponseParam;
+	GetCollisionProfileChannelAndResponseParams(ProfileName, TraceChannel, ResponseParam);
+
+	return LineTraceMulti(OutHits, Start, End, TraceChannel, Params, ResponseParam);
+}
+
+bool UWorld::SweepTestByProfile(const FVector& Start, const FVector& End, const FQuat& Rot, FName ProfileName, const struct FCollisionShape& CollisionShape, const struct FCollisionQueryParams& Params) const
+{
+	ECollisionChannel TraceChannel;
+	FCollisionResponseParams ResponseParam;
+	GetCollisionProfileChannelAndResponseParams(ProfileName, TraceChannel, ResponseParam);
+
+	return SweepTest(Start, End, Rot, TraceChannel, CollisionShape, Params, ResponseParam);
+}
+
+bool UWorld::SweepSingleByProfile(struct FHitResult& OutHit, const FVector& Start, const FVector& End, const FQuat& Rot, FName ProfileName, const struct FCollisionShape& CollisionShape, const struct FCollisionQueryParams& Params) const
+{
+	ECollisionChannel TraceChannel;
+	FCollisionResponseParams ResponseParam;
+	GetCollisionProfileChannelAndResponseParams(ProfileName, TraceChannel, ResponseParam);
+
+	return SweepSingle(OutHit, Start, End, Rot, TraceChannel, CollisionShape, Params, ResponseParam);
+}
+
+bool UWorld::SweepMultiByProfile(TArray<FHitResult>& OutHits, const FVector& Start, const FVector& End, const FQuat& Rot, FName ProfileName, const struct FCollisionShape& CollisionShape, const struct FCollisionQueryParams& Params) const
+{
+	ECollisionChannel TraceChannel;
+	FCollisionResponseParams ResponseParam;
+	GetCollisionProfileChannelAndResponseParams(ProfileName, TraceChannel, ResponseParam);
+
+	return SweepMulti(OutHits, Start, End, Rot, TraceChannel, CollisionShape, Params, ResponseParam);
+}
+
+bool UWorld::OverlapTestByProfile(const FVector& Pos, const FQuat& Rot, FName ProfileName, const struct FCollisionShape& CollisionShape, const struct FCollisionQueryParams& Params) const
+{
+	ECollisionChannel TraceChannel;
+	FCollisionResponseParams ResponseParam;
+	GetCollisionProfileChannelAndResponseParams(ProfileName, TraceChannel, ResponseParam);
+
+	return OverlapTest(Pos, Rot, TraceChannel, CollisionShape, Params, ResponseParam);
+}
+
+bool UWorld::OverlapSingleByProfile(struct FOverlapResult& OutOverlap, const FVector& Pos, const FQuat& Rot, FName ProfileName, const struct FCollisionShape& CollisionShape, const struct FCollisionQueryParams& Params) const
+{
+	ECollisionChannel TraceChannel;
+	FCollisionResponseParams ResponseParam;
+	GetCollisionProfileChannelAndResponseParams(ProfileName, TraceChannel, ResponseParam);
+
+	return OverlapSingle(OutOverlap, Pos, Rot, TraceChannel, CollisionShape, Params, ResponseParam);
+}
+
+bool UWorld::OverlapMultiByProfile(TArray<struct FOverlapResult>& OutOverlaps, const FVector& Pos, const FQuat& Rot, FName ProfileName, const struct FCollisionShape& CollisionShape, const struct FCollisionQueryParams& Params) const
+{
+	ECollisionChannel TraceChannel;
+	FCollisionResponseParams ResponseParam;
+	GetCollisionProfileChannelAndResponseParams(ProfileName, TraceChannel, ResponseParam);
+
+	return OverlapMulti(OutOverlaps, Pos, Rot, TraceChannel, CollisionShape, Params, ResponseParam);
+}
+
+
 bool UWorld::ComponentOverlapMulti(TArray<struct FOverlapResult>& OutOverlaps, const class UPrimitiveComponent* PrimComp, const FVector& Pos, const FRotator& Rot, const struct FComponentQueryParams& Params, const struct FCollisionObjectQueryParams& ObjectQueryParams) const
 {
 	if (PrimComp)
diff --git a/Engine/Source/Runtime/Engine/Private/KismetSystemLibrary.cpp b/Engine/Source/Runtime/Engine/Private/KismetSystemLibrary.cpp
--- a/Engine/Source/Runtime/Engine/Private/KismetSystemLibrary.cpp
+++ b/Engine/Source/Runtime/Engine/Private/KismetSystemLibrary.cpp
@@ -603,6 +603,12 @@ void UKismetSystemLibrary::SetTransformPropertyByName(UObject* Object, FName Pro
 	}
 }
 
+void UKismetSystemLibrary::SetCollisionProfileNameProperty(UObject* Object, FName PropertyName, const FCollisionProfileName& Value)
+{
+	// We should never hit these!  They're stubs to avoid NoExport on the class.
+	check(0);
+}
+
 void UKismetSystemLibrary::Generic_SetStructurePropertyByName(UObject* OwnerObject, FName StructPropertyName, const void* SrcStructAddr)
 {
 	if (OwnerObject != NULL)
