Fixed MacOS CIS issues
Fixed Verbosity information for VisualLogger and NewLogVisualizer
Fixed possible crash in NewLogVisualizer

diff --git a/Engine/Source/Developer/NewLogVisualizer/Private/STimelineBar.cpp b/Engine/Source/Developer/NewLogVisualizer/Private/STimelineBar.cpp
--- a/Engine/Source/Developer/NewLogVisualizer/Private/STimelineBar.cpp
+++ b/Engine/Source/Developer/NewLogVisualizer/Private/STimelineBar.cpp
@@ -76,11 +76,6 @@ uint32 STimelineBar::GetClosestItem(float Time) const
 			continue;
 		}
 		TArray<FVisualLoggerCategoryVerbosityPair> OutCategories;
-		//FVisualLoggerHelpers::GetCategories(CurrentEntryItem.Entry, OutCategories);
-		//if (VisualLoggerInterface->HasValidCategories(OutCategories) == false)
-		//{
-		//	continue;
-		//}
 		const float CurrentDist = FMath::Abs(CurrentEntryItem.Entry.TimeStamp - Time);
 		if (CurrentDist < BestDistance)
 		{
diff --git a/Engine/Source/Developer/NewLogVisualizer/Private/SVisualLoggerFilters.cpp b/Engine/Source/Developer/NewLogVisualizer/Private/SVisualLoggerFilters.cpp
--- a/Engine/Source/Developer/NewLogVisualizer/Private/SVisualLoggerFilters.cpp
+++ b/Engine/Source/Developer/NewLogVisualizer/Private/SVisualLoggerFilters.cpp
@@ -1,6 +1,11 @@
 #include "LogVisualizer.h"
 #include "SFilterWidget.h"
 #include "SSearchBox.h"
+#if WITH_EDITOR
+#include "Editor.h"
+#include "EditorViewportClient.h"
+#endif
+
 /////////////////////
 // SLogFilterList
 /////////////////////
diff --git a/Engine/Source/Developer/NewLogVisualizer/Private/VisualLoggerCanvasRenderer.cpp b/Engine/Source/Developer/NewLogVisualizer/Private/VisualLoggerCanvasRenderer.cpp
--- a/Engine/Source/Developer/NewLogVisualizer/Private/VisualLoggerCanvasRenderer.cpp
+++ b/Engine/Source/Developer/NewLogVisualizer/Private/VisualLoggerCanvasRenderer.cpp
@@ -54,7 +54,7 @@ namespace LogVisualizer
 
 void FVisualLoggerCanvasRenderer::OnItemSelectionChanged(const struct FVisualLogEntry& EntryItem)
 {
-	SelectedEntry = &EntryItem;
+	SelectedEntry = EntryItem;
 }
 
 void FVisualLoggerCanvasRenderer::ObjectSelectionChanged(TSharedPtr<class STimeline> TimeLine)
@@ -64,7 +64,7 @@ void FVisualLoggerCanvasRenderer::ObjectSelectionChanged(TSharedPtr<class STimel
 
 void FVisualLoggerCanvasRenderer::DrawOnCanvas(class UCanvas* Canvas, class APlayerController*)
 {
-	if (SelectedEntry == NULL || GEditor == NULL || CurrentTimeLine.IsValid() == false)
+	if (GEngine == NULL || CurrentTimeLine.IsValid() == false)
 	{
 		return;
 	}
@@ -78,10 +78,10 @@ void FVisualLoggerCanvasRenderer::DrawOnCanvas(class UCanvas* Canvas, class APla
 	UFont* Font = GEngine->GetSmallFont();
 	FCanvasTextItem TextItem(FVector2D::ZeroVector, FText::GetEmpty(), Font, FLinearColor::White);
 
-	const FString TimeStampString = FString::Printf(TEXT("%.2f"), SelectedEntry->TimeStamp);
-	LogVisualizer::DrawTextShadowed(Canvas, Font, TimeStampString, SelectedEntry->Location);
+	const FString TimeStampString = FString::Printf(TEXT("%.2f"), SelectedEntry.TimeStamp);
+	LogVisualizer::DrawTextShadowed(Canvas, Font, TimeStampString, SelectedEntry.Location);
 
-	for (const auto CurrentData : SelectedEntry->DataBlocks)
+	for (const auto CurrentData : SelectedEntry.DataBlocks)
 	{
 		const FName TagName = CurrentData.TagName;
 		const bool bIsValidByFilter = VisualLoggerInterface->IsValidCategory(CurrentData.Category.ToString(), ELogVerbosity::All) && VisualLoggerInterface->IsValidCategory(CurrentData.TagName.ToString(), ELogVerbosity::All);
@@ -93,11 +93,11 @@ void FVisualLoggerCanvasRenderer::DrawOnCanvas(class UCanvas* Canvas, class APla
 
 		if (!bIsValidByFilter)
 		{
-			Extension->DisableDrawingForData(VisualLoggerInterface->GetWorld(), Canvas, NULL, TagName, CurrentData, SelectedEntry->TimeStamp);
+			Extension->DisableDrawingForData(VisualLoggerInterface->GetWorld(), Canvas, NULL, TagName, CurrentData, SelectedEntry.TimeStamp);
 		}
 		else
 		{
-			Extension->DrawData(VisualLoggerInterface->GetWorld(), Canvas, NULL, TagName, CurrentData, SelectedEntry->TimeStamp);
+			Extension->DrawData(VisualLoggerInterface->GetWorld(), Canvas, NULL, TagName, CurrentData, SelectedEntry.TimeStamp);
 		}
 	}
 
@@ -134,7 +134,7 @@ void FVisualLoggerCanvasRenderer::DrawHistogramGraphs(class UCanvas* Canvas, cla
 	const float LocalViewRangeMax = LocalViewRange.GetUpperBoundValue();
 	const float LocalSequenceLength = LocalViewRangeMax - LocalViewRangeMin;
 	const float WindowHalfWidth = LocalSequenceLength * TimeSliderArgs.CursorSize.Get() * 0.5f;
-	const FVector2D TimeStampWindow(SelectedEntry->TimeStamp - WindowHalfWidth, SelectedEntry->TimeStamp + WindowHalfWidth);
+	const FVector2D TimeStampWindow(SelectedEntry.TimeStamp - WindowHalfWidth, SelectedEntry.TimeStamp + WindowHalfWidth);
 
 	const TArray<FVisualLogDevice::FVisualLogEntryItem> &ObjectItems = CurrentTimeLine.Pin()->GetEntries();
 	int32 ColorIndex = 0;
@@ -257,7 +257,7 @@ void FVisualLoggerCanvasRenderer::DrawHistogramGraphs(class UCanvas* Canvas, cla
 
 			HistogramGraph->DrawCursorOnGraph(true);
 			HistogramGraph->UseTinyFont(CollectedGraphs.Num() >= 5);
-			HistogramGraph->SetCursorLocation(SelectedEntry->TimeStamp);
+			HistogramGraph->SetCursorLocation(SelectedEntry.TimeStamp);
 			HistogramGraph->SetNumThresholds(0);
 			HistogramGraph->SetStyles(EGraphAxisStyle::Grid, EGraphDataStyle::Lines);
 			HistogramGraph->SetBackgroundColor(FColor(0, 0, 0, 200));
diff --git a/Engine/Source/Developer/NewLogVisualizer/Private/VisualLoggerCanvasRenderer.h b/Engine/Source/Developer/NewLogVisualizer/Private/VisualLoggerCanvasRenderer.h
--- a/Engine/Source/Developer/NewLogVisualizer/Private/VisualLoggerCanvasRenderer.h
+++ b/Engine/Source/Developer/NewLogVisualizer/Private/VisualLoggerCanvasRenderer.h
@@ -4,8 +4,7 @@
 struct FVisualLoggerCanvasRenderer
 {
 	FVisualLoggerCanvasRenderer(TSharedPtr<IVisualLoggerInterface> InVisualLoggerInterface) 
-		: SelectedEntry(NULL)
-		, VisualLoggerInterface(InVisualLoggerInterface) 
+		: VisualLoggerInterface(InVisualLoggerInterface) 
 	{}
 
 	void DrawOnCanvas(class UCanvas* Canvas, class APlayerController*);
@@ -16,7 +15,7 @@ protected:
 	void DrawHistogramGraphs(class UCanvas* Canvas, class APlayerController* );
 
 private:
-	const struct FVisualLogEntry* SelectedEntry;
+	FVisualLogEntry SelectedEntry;
 	TSharedPtr<IVisualLoggerInterface> VisualLoggerInterface;
 	TMap<FString, TArray<FString> > UsedGraphCategories;
 	TWeakPtr<class STimeline> CurrentTimeLine;
diff --git a/Engine/Source/Runtime/AIModule/Private/EnvironmentQuery/EnvQueryDebugHelpers.cpp b/Engine/Source/Runtime/AIModule/Private/EnvironmentQuery/EnvQueryDebugHelpers.cpp
--- a/Engine/Source/Runtime/AIModule/Private/EnvironmentQuery/EnvQueryDebugHelpers.cpp
+++ b/Engine/Source/Runtime/AIModule/Private/EnvironmentQuery/EnvQueryDebugHelpers.cpp
@@ -143,6 +143,6 @@ void UEnvQueryDebugHelpers::LogQueryInternal(FEnvQueryInstance& Query, const FLo
 	Line.UniqueId = UniqueId;
 
 	CurrentEntry->LogLines.Add(Line);
-	CurrentEntry->AddDataBlock(EVisLogTags::TAG_EQS, BlobArray, Category.GetCategoryName()).UniqueId = UniqueId;
+	CurrentEntry->AddDataBlock(EVisLogTags::TAG_EQS, BlobArray, Category.GetCategoryName(), Verbosity).UniqueId = UniqueId;
 }
 #endif //ENABLE_VISUAL_LOG && USE_EQS_DEBUGGER
diff --git a/Engine/Source/Runtime/Engine/Private/VisualLogger/VisualLoggerTypes.cpp b/Engine/Source/Runtime/Engine/Private/VisualLogger/VisualLoggerTypes.cpp
--- a/Engine/Source/Runtime/Engine/Private/VisualLogger/VisualLoggerTypes.cpp
+++ b/Engine/Source/Runtime/Engine/Private/VisualLogger/VisualLoggerTypes.cpp
@@ -106,48 +106,52 @@ int32 FVisualLogEntry::AddEvent(const FVisualLogEventBase& Event)
 	return Events.Add(Event);
 }
 
-void FVisualLogEntry::AddText(const FString& TextLine, const FName& CategoryName)
+void FVisualLogEntry::AddText(const FString& TextLine, const FName& CategoryName, ELogVerbosity::Type Verbosity)
 {
-	LogLines.Add(FVisualLogLine(CategoryName, ELogVerbosity::All, TextLine));
+	LogLines.Add(FVisualLogLine(CategoryName, Verbosity, TextLine));
 }
 
-void FVisualLogEntry::AddElement(const TArray<FVector>& Points, const FName& CategoryName, const FColor& Color, const FString& Description, uint16 Thickness)
+void FVisualLogEntry::AddElement(const TArray<FVector>& Points, const FName& CategoryName, ELogVerbosity::Type Verbosity, const FColor& Color, const FString& Description, uint16 Thickness)
 {
 	FVisualLogShapeElement Element(Description, Color, Thickness, CategoryName);
 	Element.Points = Points;
 	Element.Type = EVisualLoggerShapeElement::Path;
+	Element.Verbosity = Verbosity;
 	ElementsToDraw.Add(Element);
 }
 
-void FVisualLogEntry::AddElement(const FVector& Point, const FName& CategoryName, const FColor& Color, const FString& Description, uint16 Thickness)
+void FVisualLogEntry::AddElement(const FVector& Point, const FName& CategoryName, ELogVerbosity::Type Verbosity, const FColor& Color, const FString& Description, uint16 Thickness)
 {
 	FVisualLogShapeElement Element(Description, Color, Thickness, CategoryName);
 	Element.Points.Add(Point);
 	Element.Type = EVisualLoggerShapeElement::SinglePoint;
+	Element.Verbosity = Verbosity;
 	ElementsToDraw.Add(Element);
 }
 
-void FVisualLogEntry::AddElement(const FVector& Start, const FVector& End, const FName& CategoryName, const FColor& Color, const FString& Description, uint16 Thickness)
+void FVisualLogEntry::AddElement(const FVector& Start, const FVector& End, const FName& CategoryName, ELogVerbosity::Type Verbosity, const FColor& Color, const FString& Description, uint16 Thickness)
 {
 	FVisualLogShapeElement Element(Description, Color, Thickness, CategoryName);
 	Element.Points.Reserve(2);
 	Element.Points.Add(Start);
 	Element.Points.Add(End);
 	Element.Type = EVisualLoggerShapeElement::Segment;
+	Element.Verbosity = Verbosity;
 	ElementsToDraw.Add(Element);
 }
 
-void FVisualLogEntry::AddElement(const FBox& Box, const FName& CategoryName, const FColor& Color, const FString& Description, uint16 Thickness)
+void FVisualLogEntry::AddElement(const FBox& Box, const FName& CategoryName, ELogVerbosity::Type Verbosity, const FColor& Color, const FString& Description, uint16 Thickness)
 {
 	FVisualLogShapeElement Element(Description, Color, Thickness, CategoryName);
 	Element.Points.Reserve(2);
 	Element.Points.Add(Box.Min);
 	Element.Points.Add(Box.Max);
 	Element.Type = EVisualLoggerShapeElement::Box;
+	Element.Verbosity = Verbosity;
 	ElementsToDraw.Add(Element);
 }
 
-void FVisualLogEntry::AddElement(const FVector& Orgin, const FVector& Direction, float Length, float AngleWidth, float AngleHeight, const FName& CategoryName, const FColor& Color, const FString& Description, uint16 Thickness)
+void FVisualLogEntry::AddElement(const FVector& Orgin, const FVector& Direction, float Length, float AngleWidth, float AngleHeight, const FName& CategoryName, ELogVerbosity::Type Verbosity, const FColor& Color, const FString& Description, uint16 Thickness)
 {
 	FVisualLogShapeElement Element(Description, Color, Thickness, CategoryName);
 	Element.Points.Reserve(3);
@@ -155,10 +159,11 @@ void FVisualLogEntry::AddElement(const FVector& Orgin, const FVector& Direction,
 	Element.Points.Add(Direction);
 	Element.Points.Add(FVector(Length, AngleWidth, AngleHeight));
 	Element.Type = EVisualLoggerShapeElement::Cone;
+	Element.Verbosity = Verbosity;
 	ElementsToDraw.Add(Element);
 }
 
-void FVisualLogEntry::AddElement(const FVector& Start, const FVector& End, float Radius, const FName& CategoryName, const FColor& Color, const FString& Description, uint16 Thickness)
+void FVisualLogEntry::AddElement(const FVector& Start, const FVector& End, float Radius, const FName& CategoryName, ELogVerbosity::Type Verbosity, const FColor& Color, const FString& Description, uint16 Thickness)
 {
 	FVisualLogShapeElement Element(Description, Color, Thickness, CategoryName);
 	Element.Points.Reserve(3);
@@ -166,10 +171,11 @@ void FVisualLogEntry::AddElement(const FVector& Start, const FVector& End, float
 	Element.Points.Add(End);
 	Element.Points.Add(FVector(Radius, Thickness, 0));
 	Element.Type = EVisualLoggerShapeElement::Cylinder;
+	Element.Verbosity = Verbosity;
 	ElementsToDraw.Add(Element);
 }
 
-void FVisualLogEntry::AddElement(const FVector& Center, float HalfHeight, float Radius, const FQuat & Rotation, const FName& CategoryName, const FColor& Color, const FString& Description)
+void FVisualLogEntry::AddElement(const FVector& Center, float HalfHeight, float Radius, const FQuat & Rotation, const FName& CategoryName, ELogVerbosity::Type Verbosity, const FColor& Color, const FString& Description)
 {
 	FVisualLogShapeElement Element(Description, Color, 0, CategoryName);
 	Element.Points.Reserve(3);
@@ -177,26 +183,29 @@ void FVisualLogEntry::AddElement(const FVector& Center, float HalfHeight, float
 	Element.Points.Add(FVector(HalfHeight, Radius, Rotation.X));
 	Element.Points.Add(FVector(Rotation.Y, Rotation.Z, Rotation.W));
 	Element.Type = EVisualLoggerShapeElement::Capsule;
+	Element.Verbosity = Verbosity;
 	ElementsToDraw.Add(Element);
 }
 
-void FVisualLogEntry::AddHistogramData(const FVector2D& DataSample, const FName& CategoryName, const FName& GraphName, const FName& DataName)
+void FVisualLogEntry::AddHistogramData(const FVector2D& DataSample, const FName& CategoryName, ELogVerbosity::Type Verbosity, const FName& GraphName, const FName& DataName)
 {
 	FVisualLogHistogramSample Sample;
 	Sample.Category = CategoryName;
 	Sample.GraphName = GraphName;
 	Sample.DataName = DataName;
 	Sample.SampleValue = DataSample;
+	Sample.Verbosity = Verbosity;
 
 	HistogramSamples.Add(Sample);
 }
 
-FVisualLogDataBlock& FVisualLogEntry::AddDataBlock(const FString& TagName, const TArray<uint8>& BlobDataArray, const FName& CategoryName)
+FVisualLogDataBlock& FVisualLogEntry::AddDataBlock(const FString& TagName, const TArray<uint8>& BlobDataArray, const FName& CategoryName, ELogVerbosity::Type Verbosity)
 {
 	FVisualLogDataBlock DataBlock;
 	DataBlock.Category = CategoryName;
 	DataBlock.TagName = *TagName;
 	DataBlock.Data = BlobDataArray;
+	DataBlock.Verbosity = Verbosity;
 
 	const int32 Index = DataBlocks.Add(DataBlock);
 	return DataBlocks[Index];
diff --git a/Engine/Source/Runtime/Engine/Public/VisualLogger/VisualLogger.inl b/Engine/Source/Runtime/Engine/Public/VisualLogger/VisualLogger.inl
--- a/Engine/Source/Runtime/Engine/Public/VisualLogger/VisualLogger.inl
+++ b/Engine/Source/Runtime/Engine/Public/VisualLogger/VisualLogger.inl
@@ -159,7 +159,7 @@ VARARG_BODY(void, FVisualLogger::CategorizedLogf, const TCHAR*, VARARG_EXTRA(con
 	}
 
 	COLLAPSED_LOGF(
-		CurrentEntry->AddText(Buffer, Category.GetCategoryName());
+		CurrentEntry->AddText(Buffer, Category.GetCategoryName(), Verbosity);
 	);
 }
 
@@ -175,7 +175,7 @@ VARARG_BODY(void, FVisualLogger::GeometryShapeLogf, const TCHAR*, VARARG_EXTRA(c
 	}
 
 	COLLAPSED_LOGF(
-		CurrentEntry->AddElement(Start, End, Category.GetCategoryName(), Color, Buffer);
+		CurrentEntry->AddElement(Start, End, Category.GetCategoryName(), Verbosity, Color, Buffer);
 	);
 }
 
@@ -191,7 +191,7 @@ VARARG_BODY(void, FVisualLogger::GeometryShapeLogf, const TCHAR*, VARARG_EXTRA(c
 	}
 
 	COLLAPSED_LOGF(
-		CurrentEntry->AddElement(Location, Category.GetCategoryName(), Color, Buffer, Radius);
+		CurrentEntry->AddElement(Location, Category.GetCategoryName(), Verbosity, Color, Buffer, Radius);
 	);
 }
 
@@ -207,7 +207,7 @@ VARARG_BODY(void, FVisualLogger::GeometryShapeLogf, const TCHAR*, VARARG_EXTRA(c
 	}
 
 	COLLAPSED_LOGF(
-		CurrentEntry->AddElement(Box, Category.GetCategoryName(), Color, Buffer);
+		CurrentEntry->AddElement(Box, Category.GetCategoryName(), Verbosity, Color, Buffer);
 	);
 }
 
@@ -223,7 +223,7 @@ VARARG_BODY(void, FVisualLogger::GeometryShapeLogf, const TCHAR*, VARARG_EXTRA(c
 	}
 
 	COLLAPSED_LOGF(
-		CurrentEntry->AddElement(Orgin, Direction, Length, Angle, Angle, Category.GetCategoryName(), Color, Buffer);
+		CurrentEntry->AddElement(Orgin, Direction, Length, Angle, Angle, Category.GetCategoryName(), Verbosity, Color, Buffer);
 	);
 }
 
@@ -239,7 +239,7 @@ VARARG_BODY(void, FVisualLogger::GeometryShapeLogf, const TCHAR*, VARARG_EXTRA(c
 	}
 
 	COLLAPSED_LOGF(
-		CurrentEntry->AddElement(Start, End, Radius, Category.GetCategoryName(), Color, Buffer);
+		CurrentEntry->AddElement(Start, End, Radius, Category.GetCategoryName(), Verbosity, Color, Buffer);
 	);
 }
 
@@ -255,7 +255,7 @@ VARARG_BODY(void, FVisualLogger::GeometryShapeLogf, const TCHAR*, VARARG_EXTRA(c
 	}
 
 	COLLAPSED_LOGF(
-		CurrentEntry->AddElement(Center, HalfHeight, Radius, Rotation, Category.GetCategoryName(), Color, Buffer);
+		CurrentEntry->AddElement(Center, HalfHeight, Radius, Rotation, Category.GetCategoryName(), Verbosity, Color, Buffer);
 	);
 }
 
@@ -271,7 +271,7 @@ VARARG_BODY(void, FVisualLogger::HistogramDataLogf, const TCHAR*, VARARG_EXTRA(c
 	}
 
 	COLLAPSED_LOGF(
-		CurrentEntry->AddHistogramData(Data, Category.GetCategoryName(), GraphName, DataName);
+		CurrentEntry->AddHistogramData(Data, Category.GetCategoryName(), Verbosity, GraphName, DataName);
 	);
 }
 
diff --git a/Engine/Source/Runtime/Engine/Public/VisualLogger/VisualLoggerTypes.h b/Engine/Source/Runtime/Engine/Public/VisualLogger/VisualLoggerTypes.h
--- a/Engine/Source/Runtime/Engine/Public/VisualLogger/VisualLoggerTypes.h
+++ b/Engine/Source/Runtime/Engine/Public/VisualLogger/VisualLoggerTypes.h
@@ -168,25 +168,25 @@ struct ENGINE_API FVisualLogEntry
 
 	void Reset();
 
-	void AddText(const FString& TextLine, const FName& CategoryName);
+	void AddText(const FString& TextLine, const FName& CategoryName, ELogVerbosity::Type Verbosity);
 	// path
-	void AddElement(const TArray<FVector>& Points, const FName& CategoryName, const FColor& Color = FColor::White, const FString& Description = TEXT(""), uint16 Thickness = 0);
+	void AddElement(const TArray<FVector>& Points, const FName& CategoryName, ELogVerbosity::Type Verbosity, const FColor& Color = FColor::White, const FString& Description = TEXT(""), uint16 Thickness = 0);
 	// location
-	void AddElement(const FVector& Point, const FName& CategoryName, const FColor& Color = FColor::White, const FString& Description = TEXT(""), uint16 Thickness = 0);
+	void AddElement(const FVector& Point, const FName& CategoryName, ELogVerbosity::Type Verbosity, const FColor& Color = FColor::White, const FString& Description = TEXT(""), uint16 Thickness = 0);
 	// segment
-	void AddElement(const FVector& Start, const FVector& End, const FName& CategoryName, const FColor& Color = FColor::White, const FString& Description = TEXT(""), uint16 Thickness = 0);
+	void AddElement(const FVector& Start, const FVector& End, const FName& CategoryName, ELogVerbosity::Type Verbosity, const FColor& Color = FColor::White, const FString& Description = TEXT(""), uint16 Thickness = 0);
 	// box
-	void AddElement(const FBox& Box, const FName& CategoryName, const FColor& Color = FColor::White, const FString& Description = TEXT(""), uint16 Thickness = 0);
+	void AddElement(const FBox& Box, const FName& CategoryName, ELogVerbosity::Type Verbosity, const FColor& Color = FColor::White, const FString& Description = TEXT(""), uint16 Thickness = 0);
 	// Cone
-	void AddElement(const FVector& Orgin, const FVector& Direction, float Length, float AngleWidth, float AngleHeight, const FName& CategoryName, const FColor& Color = FColor::White, const FString& Description = TEXT(""), uint16 Thickness = 0);
+	void AddElement(const FVector& Orgin, const FVector& Direction, float Length, float AngleWidth, float AngleHeight, const FName& CategoryName, ELogVerbosity::Type Verbosity, const FColor& Color = FColor::White, const FString& Description = TEXT(""), uint16 Thickness = 0);
 	// Cylinder
-	void AddElement(const FVector& Start, const FVector& End, float Radius, const FName& CategoryName, const FColor& Color = FColor::White, const FString& Description = TEXT(""), uint16 Thickness = 0);
+	void AddElement(const FVector& Start, const FVector& End, float Radius, const FName& CategoryName, ELogVerbosity::Type Verbosity, const FColor& Color = FColor::White, const FString& Description = TEXT(""), uint16 Thickness = 0);
 	// capsule
-	void AddElement(const FVector& Center, float HalfHeight, float Radius, const FQuat & Rotation, const FName& CategoryName, const FColor& Color = FColor::White, const FString& Description = TEXT(""));
+	void AddElement(const FVector& Center, float HalfHeight, float Radius, const FQuat & Rotation, const FName& CategoryName, ELogVerbosity::Type Verbosity, const FColor& Color = FColor::White, const FString& Description = TEXT(""));
 	// histogram sample
-	void AddHistogramData(const FVector2D& DataSample, const FName& CategoryName, const FName& GraphName, const FName& DataName);
+	void AddHistogramData(const FVector2D& DataSample, const FName& CategoryName, ELogVerbosity::Type Verbosity, const FName& GraphName, const FName& DataName);
 	// Custom data block
-	FVisualLogDataBlock& AddDataBlock(const FString& TagName, const TArray<uint8>& BlobDataArray, const FName& CategoryName);
+	FVisualLogDataBlock& AddDataBlock(const FString& TagName, const TArray<uint8>& BlobDataArray, const FName& CategoryName, ELogVerbosity::Type Verbosity);
 	// Event
 	int32 AddEvent(const FVisualLogEventBase& Event);
 	// find index of status category
