Fixed non-root bone root body falling through when scale is less than 1

diff --git a/Engine/Source/Runtime/Engine/Classes/Components/SkeletalMeshComponent.h b/Engine/Source/Runtime/Engine/Classes/Components/SkeletalMeshComponent.h
--- a/Engine/Source/Runtime/Engine/Classes/Components/SkeletalMeshComponent.h
+++ b/Engine/Source/Runtime/Engine/Classes/Components/SkeletalMeshComponent.h
@@ -642,7 +642,7 @@ public:
 	struct 
 	{
 		int32 BodyIndex;
-		int32 BoneIndex;
+		FTransform TransformToRoot;
 	} RootBodyData;
 
 	/** Set Root Body Index */
diff --git a/Engine/Source/Runtime/Engine/Private/SkeletalMeshComponent.cpp b/Engine/Source/Runtime/Engine/Private/SkeletalMeshComponent.cpp
--- a/Engine/Source/Runtime/Engine/Private/SkeletalMeshComponent.cpp
+++ b/Engine/Source/Runtime/Engine/Private/SkeletalMeshComponent.cpp
@@ -1911,14 +1911,36 @@ void USkeletalMeshComponent::SetRootBodyIndex(int32 InBodyIndex)
 {
 	RootBodyData.BodyIndex = InBodyIndex;
 
-	if(Bodies.IsValidIndex(RootBodyData.BodyIndex) && 
+	if(Bodies.IsValidIndex(RootBodyData.BodyIndex) && SkeletalMesh && 
 		Bodies[RootBodyData.BodyIndex]->BodySetup.IsValid() && Bodies[RootBodyData.BodyIndex]->BodySetup.Get()->BoneName != NAME_None)
 	{
-		RootBodyData.BoneIndex = GetBoneIndex(Bodies[RootBodyData.BodyIndex]->BodySetup->BoneName);
+		int32 BoneIndex = GetBoneIndex(Bodies[RootBodyData.BodyIndex]->BodySetup->BoneName);
+		// if bone index is valid and not 0, it SHOULD have parnet index
+		if (ensure (BoneIndex != INDEX_NONE))
+		{
+			int32 ParentIndex = SkeletalMesh->RefSkeleton.GetParentIndex(BoneIndex);
+			if (BoneIndex != 0 && ensure (ParentIndex != INDEX_NONE))
+			{
+				const TArray<FTransform>& RefPose = SkeletalMesh->RefSkeleton.GetRefBonePose();
+
+				FTransform RelativeTransform = FTransform(SkeletalMesh->RefBasesInvMatrix[BoneIndex]) * RefPose[ParentIndex];
+				// now get offset 
+				RootBodyData.TransformToRoot = RelativeTransform;
+			}
+			else
+			{
+				RootBodyData.TransformToRoot = FTransform::Identity;
+			}
+		}
+		else
+		{
+			RootBodyData.TransformToRoot = FTransform::Identity;
+		}
 	}
 	else
 	{
-		// error
-		RootBodyData.BoneIndex = INDEX_NONE;
+		// error - this should not happen
+		ensure(false);
+		RootBodyData.TransformToRoot = FTransform::Identity;
 	}
 }
\ No newline at end of file
diff --git a/Engine/Source/Runtime/Engine/Private/SkeletalMeshComponentPhysics.cpp b/Engine/Source/Runtime/Engine/Private/SkeletalMeshComponentPhysics.cpp
--- a/Engine/Source/Runtime/Engine/Private/SkeletalMeshComponentPhysics.cpp
+++ b/Engine/Source/Runtime/Engine/Private/SkeletalMeshComponentPhysics.cpp
@@ -4924,12 +4924,6 @@ FTransform USkeletalMeshComponent::GetComponentTransformFromBodyInstance(FBodyIn
 {
 	// undo root transform so that it only moves according to what actor itself suppose to move
 	FTransform BodyTransform = UseBI->GetUnrealWorldTransform();
-	if (RootBodyData.BoneIndex != INDEX_NONE)
-	{
-		FTransform RootTransform(GetSpaceBases()[RootBodyData.BoneIndex]);
-		return RootTransform.GetRelativeTransformReverse(BodyTransform);
-	}
-
-	return BodyTransform;
+	return RootBodyData.TransformToRoot * BodyTransform;
 }
 #undef LOCTEXT_NAMESPACE
