Morph targets no longer limited to 0 - 1

diff --git a/Engine/Source/Editor/Persona/Private/SMorphTargetViewer.cpp b/Engine/Source/Editor/Persona/Private/SMorphTargetViewer.cpp
--- a/Engine/Source/Editor/Persona/Private/SMorphTargetViewer.cpp
+++ b/Engine/Source/Editor/Persona/Private/SMorphTargetViewer.cpp
@@ -15,6 +15,8 @@ static const FName ColumnId_MorphTargetNameLabel( "MorphTargetName" );
 static const FName ColumnID_MorphTargetWeightLabel( "Weight" );
 static const FName ColumnID_MorphTargetVertCountLabel( "NumberOfVerts" );
 
+const float MaxMorphWeight = 5.f;
+
 //////////////////////////////////////////////////////////////////////////
 // SMorphTargetListRow
 
@@ -118,8 +120,10 @@ TSharedRef< SWidget > SMorphTargetListRow::GenerateWidgetForColumn( const FName&
 			.VAlign( VAlign_Center )
 			[
 				SNew( SSpinBox<float> )
-				.MinValue(0.00f)
-				.MaxValue(1.0f)
+				.MinSliderValue(-1.f)
+				.MaxSliderValue(1.f)
+				.MinValue(-MaxMorphWeight)
+				.MaxValue(MaxMorphWeight)
 				.Value( this, &SMorphTargetListRow::GetWeight )
 				.OnValueChanged( this, &SMorphTargetListRow::OnMorphTargetWeightChanged )
 			];
@@ -169,7 +173,7 @@ void SMorphTargetListRow::OnMorphTargetWeightChanged( float NewWeight )
 
 		if ( RowItem != Item ) // Don't do "this" row again if it's selected
 		{
-			RowItem->Weight = FMath::Clamp( RowItem->Weight + Delta, 0.0f, 1.0f );
+			RowItem->Weight = FMath::Clamp(RowItem->Weight + Delta, -MaxMorphWeight, MaxMorphWeight);
 			MorphTargetViewer->AddMorphTargetOverride( RowItem->Name, RowItem->Weight );
 		}
 	}
diff --git a/Engine/Source/Runtime/Engine/Private/SkeletalRenderCPUSkin.cpp b/Engine/Source/Runtime/Engine/Private/SkeletalRenderCPUSkin.cpp
--- a/Engine/Source/Runtime/Engine/Private/SkeletalRenderCPUSkin.cpp
+++ b/Engine/Source/Runtime/Engine/Private/SkeletalRenderCPUSkin.cpp
@@ -488,9 +488,11 @@ uint32 InitEvalInfos(const TArray<FActiveVertexAnim>& ActiveVertexAnims, int32 L
 		FVertexAnimEvalInfo NewInfo;
 
 		const FActiveVertexAnim& ActiveAnim = ActiveVertexAnims[AnimIdx];
+		const float ActiveAnimAbsVertexWeight = FMath::Abs(ActiveAnim.Weight);
+
 		if( ActiveAnim.VertAnim != NULL &&
-			ActiveAnim.Weight >= MinVertexAnimBlendWeight && 
-			ActiveAnim.Weight <= MaxVertexAnimBlendWeight &&				
+			ActiveAnimAbsVertexWeight >= MinVertexAnimBlendWeight &&
+			ActiveAnimAbsVertexWeight <= MaxVertexAnimBlendWeight &&
 			ActiveAnim.VertAnim->HasDataForLOD(LODIndex) )
 		{
 			// start at the first vertex since they affect base mesh verts in ascending order
diff --git a/Engine/Source/Runtime/Engine/Private/SkeletalRenderGPUSkin.cpp b/Engine/Source/Runtime/Engine/Private/SkeletalRenderGPUSkin.cpp
--- a/Engine/Source/Runtime/Engine/Private/SkeletalRenderGPUSkin.cpp
+++ b/Engine/Source/Runtime/Engine/Private/SkeletalRenderGPUSkin.cpp
@@ -361,7 +361,9 @@ void FSkeletalMeshObjectGPUSkin::FSkeletalMeshObjectLOD::UpdateMorphVertexBuffer
 				const FActiveVertexAnim& VertAnim = ActiveVertexAnims[AnimIdx];
 				checkSlow(VertAnim.VertAnim != NULL);
 				checkSlow(VertAnim.VertAnim->HasDataForLOD(LODIndex));
-				checkSlow(VertAnim.Weight >= MinVertexAnimBlendWeight && VertAnim.Weight <= MaxVertexAnimBlendWeight);
+
+				const float VertAnimAbsWeight = FMath::Abs(VertAnim.Weight);
+				checkSlow(VertAnimAbsWeight >= MinVertexAnimBlendWeight && VertAnimAbsWeight <= MaxVertexAnimBlendWeight);
 
 				// Allocate temp state
 				FVertexAnimEvalStateBase* AnimState = VertAnim.VertAnim->InitEval();
@@ -384,7 +386,7 @@ void FSkeletalMeshObjectGPUSkin::FSkeletalMeshObjectLOD::UpdateMorphVertexBuffer
 						// add to accumulated tangent Z
 						DeltaTangentZAccumulationArray[MorphVertex.SourceIdx] += MorphVertex.TangentZDelta * VertAnim.Weight;
 						// accumulate the weight so we can normalized it later
-						AccumulatedWeightArray[MorphVertex.SourceIdx] += VertAnim.Weight;
+						AccumulatedWeightArray[MorphVertex.SourceIdx] += VertAnimAbsWeight;
 					}
 				} // for all vertices
 
@@ -998,9 +1000,11 @@ FDynamicSkelMeshObjectDataGPUSkin::FDynamicSkelMeshObjectDataGPUSkin(
 	for( int32 AnimIdx=ActiveVertexAnims.Num()-1; AnimIdx >= 0; AnimIdx-- )
 	{
 		const FActiveVertexAnim& Anim = ActiveVertexAnims[AnimIdx];
+		const float AnimAbsWeight = FMath::Abs(Anim.Weight);
+
 		if( Anim.VertAnim != NULL && 
-			Anim.Weight >= MinVertexAnimBlendWeight &&
-			Anim.Weight <= MaxVertexAnimBlendWeight &&
+			AnimAbsWeight >= MinVertexAnimBlendWeight &&
+			AnimAbsWeight <= MaxVertexAnimBlendWeight &&
 			Anim.VertAnim->HasDataForLOD(LODIndex) ) 
 		{
 			NumWeightedActiveVertexAnims++;
diff --git a/Engine/Source/Runtime/Engine/Private/SkinnedMeshComponent.cpp b/Engine/Source/Runtime/Engine/Private/SkinnedMeshComponent.cpp
--- a/Engine/Source/Runtime/Engine/Private/SkinnedMeshComponent.cpp
+++ b/Engine/Source/Runtime/Engine/Private/SkinnedMeshComponent.cpp
@@ -1689,8 +1689,10 @@ TArray<FActiveVertexAnim> USkinnedMeshComponent::UpdateActiveVertexAnims(const U
 	for(int32 AnimIdx=0; AnimIdx < ActiveAnims.Num(); AnimIdx++)
 	{
 		const FActiveVertexAnim& ActiveAnim = ActiveAnims[AnimIdx];
+		const float ActiveAnimAbsWeight = FMath::Abs(ActiveAnim.Weight);
+
 		// Check it has valid weight, and works on this SkeletalMesh
-		if(	ActiveAnim.Weight > MinVertexAnimBlendWeight &&
+		if (	ActiveAnimAbsWeight > MinVertexAnimBlendWeight &&
 			ActiveAnim.VertAnim != NULL &&
 			ActiveAnim.VertAnim->BaseSkelMesh == InSkeletalMesh)
 		{
@@ -1706,7 +1708,7 @@ TArray<FActiveVertexAnim> USkinnedMeshComponent::UpdateActiveVertexAnims(const U
 		const float& Weight	= (CurveIter).Value();
 
 		// If it has a valid weight
-		if(Weight > MinVertexAnimBlendWeight)
+		if(FMath::Abs(Weight) > MinVertexAnimBlendWeight)
 		{
 			// Find morph reference
 			UMorphTarget* Target = InSkeletalMesh ? InSkeletalMesh->FindMorphTarget(CurveName) : NULL;
