BranchingPoints refactor:
- BranchingPoints are now regular AnimNotifies, with their MontageTickType set to 'BranchingPoint' instead of 'Queued'.
- We now supported blueprints and AnimNotify_State as branching points!
- Old branching points are converted to custom AnimNotifies.
- Fixed AnimNotify tool tips not showing up.

diff --git a/Engine/Source/Editor/BlueprintGraph/Private/BlueprintActionDatabase.cpp b/Engine/Source/Editor/BlueprintGraph/Private/BlueprintActionDatabase.cpp
--- a/Engine/Source/Editor/BlueprintGraph/Private/BlueprintActionDatabase.cpp
+++ b/Engine/Source/Editor/BlueprintGraph/Private/BlueprintActionDatabase.cpp
@@ -694,31 +694,6 @@ static void BlueprintActionDatabaseImpl::AddSkeletonActions(const USkeleton& Ske
 		FString SignatureName = FString::Printf(TEXT("AnimNotify_%s"), *Label);
 		ActionListOut.Add(FBlueprintNodeSpawnerFactory::MakeAnimOwnedEventSpawner(FName(*SignatureName), FEditorCategoryUtils::GetCommonCategory(FCommonEditorCategory::AnimNotify)));
 	}
-
-	// @todo anim: fix this to be same as notifies, save same list in the 
-	// add montage menus
-	// find all montages that uses current target skeleton
-	TArray<FName> BranchingPointHandlers;
-
-	for (FObjectIterator Iter(UAnimMontage::StaticClass()); Iter; ++Iter)
-	{
-		UAnimMontage * Montage = CastChecked<UAnimMontage>(*Iter);
-		if (Montage && Montage->GetSkeleton() == &Skeleton)
-		{
-			// now add event handler if exists
-			for (int32 I = 0; I < Montage->BranchingPoints.Num(); ++I)
-			{
-				BranchingPointHandlers.AddUnique(Montage->BranchingPoints[I].EventName);
-			}
-		}
-	}
-
-	for( auto NotifyName: BranchingPointHandlers )
-	{
-		FString Label = NotifyName.ToString();
-		FString SignatureName = FString::Printf(TEXT("MontageBranchingPoint_%s"), *Label);
-		ActionListOut.Add(FBlueprintNodeSpawnerFactory::MakeAnimOwnedEventSpawner(FName(*SignatureName),  FEditorCategoryUtils::GetCommonCategory(FCommonEditorCategory::BranchPoint)));
-	}
 }
 
 //------------------------------------------------------------------------------
diff --git a/Engine/Source/Editor/BlueprintGraph/Private/K2ActionMenuBuilder.cpp b/Engine/Source/Editor/BlueprintGraph/Private/K2ActionMenuBuilder.cpp
--- a/Engine/Source/Editor/BlueprintGraph/Private/K2ActionMenuBuilder.cpp
+++ b/Engine/Source/Editor/BlueprintGraph/Private/K2ActionMenuBuilder.cpp
@@ -1699,39 +1699,6 @@ void FK2ActionMenuBuilder::GetAnimNotifyMenuItems(FBlueprintGraphActionListBuild
 				EventNode->CustomFunctionName = EventNode->EventSignatureName;
 				Action->NodeTemplate = EventNode;
 			}
-
-			// @todo anim: fix this to be same as notifies, save same list in the 
-			// add montage menus
-			// find all montages that uses current target skeleton
-			TArray<FName> BrancingPointHandlers;
-
-			for ( FObjectIterator Iter(UAnimMontage::StaticClass()); Iter; ++Iter )
-			{
-				UAnimMontage * Montage = CastChecked<UAnimMontage>(*Iter);
-				if ( Montage && Montage->GetSkeleton() == AnimBlueprint->TargetSkeleton )
-				{
-					// now add event handler if exists
-					for ( int32 I=0; I<Montage->BranchingPoints.Num(); ++I )
-					{
-						BrancingPointHandlers.AddUnique(Montage->BranchingPoints[I].EventName);
-					}
-				}
-			}
-
-			for ( int32 I=0; I<BrancingPointHandlers.Num(); ++I )
-			{
-				FName NotifyName = BrancingPointHandlers[I];
-				FString Label = NotifyName.ToString();
-
-				TSharedPtr<FEdGraphSchemaAction_K2NewNode> Action = FK2ActionMenuBuilder::AddNewNodeAction(ContextMenuBuilder, K2ActionCategories::BranchPointCategory, FText::FromString(Label), TEXT(""));
-
-				UK2Node_Event* EventNode = ContextMenuBuilder.CreateTemplateNode<UK2Node_Event>();
-				FString SignatureName = FString::Printf(TEXT("MontageBranchingPoint_%s"), *Label);
-				EventNode->EventSignatureName = FName(*SignatureName);
-				EventNode->EventSignatureClass = UAnimInstance::StaticClass();
-				EventNode->CustomFunctionName = EventNode->EventSignatureName;
-				Action->NodeTemplate = EventNode;
-			}
 		}
 
 		// State machine events
diff --git a/Engine/Source/Editor/EditorStyle/Private/SlateEditorStyle.cpp b/Engine/Source/Editor/EditorStyle/Private/SlateEditorStyle.cpp
--- a/Engine/Source/Editor/EditorStyle/Private/SlateEditorStyle.cpp
+++ b/Engine/Source/Editor/EditorStyle/Private/SlateEditorStyle.cpp
@@ -4482,6 +4482,9 @@ void FSlateEditorStyle::FStyle::SetupPersonaStyle()
 		Set("AnimSlotManager.AddGroup", new IMAGE_BRUSH("Persona/AnimSlotManager/icon_AddGroup_40x", Icon40x40));
 		Set("AnimSlotManager.AddSlot", new IMAGE_BRUSH("Persona/AnimSlotManager/icon_AddSlot_40x", Icon40x40));
 		Set("AnimSlotManager.Warning", new IMAGE_BRUSH("Persona/AnimSlotManager/icon_Warning_14x", Icon16x16));
+
+		// Anim Notify Editor
+		Set("AnimNotifyEditor.BranchingPoint", new IMAGE_BRUSH("Persona/NotifyEditor/BranchingPoints_24x", Icon24x24));
 	}
 
 	// Kismet 2
diff --git a/Engine/Source/Editor/Persona/Private/SAnimMontagePanel.cpp b/Engine/Source/Editor/Persona/Private/SAnimMontagePanel.cpp
--- a/Engine/Source/Editor/Persona/Private/SAnimMontagePanel.cpp
+++ b/Engine/Source/Editor/Persona/Private/SAnimMontagePanel.cpp
@@ -12,284 +12,8 @@
 #include "STextEntryPopup.h"
 #include "STextComboBox.h"
 
-
 #define LOCTEXT_NAMESPACE "AnimMontagePanel"
 
-
-//////////////////////////////////////////////////////////////////////////
-// SMontageBranchingPointNode
-
-const float MontageBranchingPointHandleWidth = 12.f;
-const float BranchingPointAlignmentMarkerWidth = 8.f;
-
-class FBranchingPointDragDropOp : public FTrackNodeDragDropOp
-{
-public:
-	DRAG_DROP_OPERATOR_TYPE(FBranchingPointDragDropOp, FTrackNodeDragDropOp)
-
-	static TSharedRef<FTrackNodeDragDropOp> New(TSharedRef<STrackNode> TrackNode, const FVector2D &CursorPosition, const FVector2D &ScreenPositionOfNode)
-	{
-		TSharedRef<FBranchingPointDragDropOp> Operation = MakeShareable(new FBranchingPointDragDropOp);
-		Operation->OriginalTrackNode = TrackNode;
-
-		Operation->Offset = ScreenPositionOfNode - CursorPosition;
-		Operation->StartingScreenPos = ScreenPositionOfNode;
-
-		Operation->Construct();
-
-		return Operation;
-	}
-
-	virtual void OnDragged( const class FDragDropEvent& DragDropEvent ) override;
-};
-
-class SMontageBranchingPointNode : public STrackNode
-{
-public:
-	SLATE_BEGIN_ARGS( SMontageBranchingPointNode )
-		: _Montage()
-		, _BranchingPoint()
-		, _AllowDrag(true)
-		{}
-		SLATE_ARGUMENT( class UAnimMontage*, Montage )
-		SLATE_ARGUMENT( FBranchingPoint*, BranchingPoint )
-		SLATE_ATTRIBUTE( float, ViewInputMin )
-		SLATE_ATTRIBUTE( float, ViewInputMax )
-		SLATE_ATTRIBUTE( float, DataLength )
-		SLATE_ATTRIBUTE( float, DataStartPos )
-		SLATE_ATTRIBUTE( FString, NodeName )
-		SLATE_ATTRIBUTE( FLinearColor, NodeColor )
-		SLATE_ATTRIBUTE( FLinearColor, SelectedNodeColor )
-		SLATE_ARGUMENT( STrackNodeSelectionSet *, NodeSelectionSet )
-		SLATE_ARGUMENT( bool, AllowDrag )
-		SLATE_EVENT( FOnTrackNodeDragged, OnTrackNodeDragged )
-		SLATE_EVENT( FOnTrackNodeDropped, OnTrackNodeDropped )
-		SLATE_EVENT( FOnNodeSelectionChanged, OnSelectionChanged )
-		SLATE_EVENT( FOnNodeRightClickContextMenu, OnNodeRightClickContextMenu )
-		SLATE_EVENT( FOnTrackNodeClicked, OnTrackNodeClicked )
-		SLATE_ARGUMENT( bool, CenterOnPosition )
-	SLATE_END_ARGS()
-
-	void Construct(const FArguments& InArgs);
-	virtual int32 OnPaint(const FPaintArgs& Args, const FGeometry& AllottedGeometry, const FSlateRect& MyClippingRect, FSlateWindowElementList& OutDrawElements, int32 LayerId, const FWidgetStyle& InWidgetStyle, bool bParentEnabled) const override;
-
-	virtual FVector2D GetDragDropScreenSpacePosition(const FGeometry& ParentAllottedGeometry, const FDragDropEvent& DragDropEvent) const override;
-
-	virtual FReply BeginDrag( const FGeometry& MyGeometry, const FPointerEvent& MouseEvent ) override;
-
-	bool SnapToDragBars() const {return true;}
-
-	virtual void OnSnapNodeDataPosition(float OriginalX, float SnappedX);
-
-	FBranchingPoint* GetBranchingPoint() const { return BranchingPoint; }
-
-private:
-	UAnimMontage* Montage;
-
-	FBranchingPoint* BranchingPoint;
-
-	FString GetBranchingPointText() const;
-	virtual FVector2D GetSize() const override;
-
-	// virtual draw related functions STrackNode
-	virtual FVector2D GetOffsetRelativeToParent(const FGeometry& ParentAllottedGeometry) const override;
-	virtual FVector2D GetSizeRelativeToParent(const FGeometry& ParentAllottedGeometry) const override;
-};
-
-void FBranchingPointDragDropOp::OnDragged( const class FDragDropEvent& DragDropEvent )
-{
-	TSharedPtr<STrackNode> Node = OriginalTrackNode.Pin();
-
-	FVector2D NodePos = (DragDropEvent.GetScreenSpacePosition() + Offset);
-
-	if (Node.IsValid())
-	{
-		Node->OnDragged(DragDropEvent);
-		const FGeometry& TrackGeometry = Node->GetTrackGeometry();
-		NodePos = Node->GetOffsetRelativeToParent(TrackGeometry) + TrackGeometry.AbsolutePosition;
-	}
-
-	NodePos.Y = StartingScreenPos.Y;
-
-	CursorDecoratorWindow->MoveWindowTo(NodePos);
-}
-
-void SMontageBranchingPointNode::Construct(const FArguments& InArgs)
-{
-	STrackNode::Construct( STrackNode::FArguments()
-		.ViewInputMin(InArgs._ViewInputMin)
-		.ViewInputMax(InArgs._ViewInputMax)
-		.DataLength(InArgs._DataLength)
-		.DataStartPos(InArgs._DataStartPos)
-		.NodeName(InArgs._NodeName)
-		.NodeColor(InArgs._NodeColor)
-		.SelectedNodeColor(InArgs._SelectedNodeColor)
-		.NodeSelectionSet(InArgs._NodeSelectionSet)
-		.AllowDrag(InArgs._AllowDrag)
-		.OnTrackNodeDragged(InArgs._OnTrackNodeDragged)
-		.OnTrackNodeDropped(InArgs._OnTrackNodeDropped)
-		.OnSelectionChanged(InArgs._OnSelectionChanged)
-		.OnNodeRightClickContextMenu(InArgs._OnNodeRightClickContextMenu)
-		.OnTrackNodeClicked(InArgs._OnTrackNodeClicked)
-		.CenterOnPosition(false) ); // ignored anyway
-
-	Montage = InArgs._Montage;
-	BranchingPoint = InArgs._BranchingPoint;
-}
-
-FString SMontageBranchingPointNode::GetBranchingPointText() const
-{
-	return BranchingPoint->EventName.ToString();
-}
-
-FVector2D SMontageBranchingPointNode::GetDragDropScreenSpacePosition(const FGeometry& ParentAllottedGeometry, const FDragDropEvent& DragDropEvent) const
-{
-	return STrackNode::GetDragDropScreenSpacePosition(ParentAllottedGeometry, DragDropEvent) + FVector2D((MontageBranchingPointHandleWidth * 0.5f + 2), 0.0f);
-}
-
-FVector2D SMontageBranchingPointNode::GetOffsetRelativeToParent(const FGeometry& AllottedGeometry) const
-{
-	FTrackScaleInfo ScaleInfo(ViewInputMin.Get(), ViewInputMax.Get(), 0, 0, AllottedGeometry.Size);
-
-	return FVector2D( ScaleInfo.InputToLocalX(DataStartPos.Get()) - (MontageBranchingPointHandleWidth * 0.5f + 2), 0);
-}
-
-FVector2D SMontageBranchingPointNode::GetSizeRelativeToParent(const FGeometry& AllottedGeometry) const
-{
-	return GetSize();
-}
-
-FVector2D SMontageBranchingPointNode::GetSize() const
-{
-	const TSharedRef< FSlateFontMeasure > FontMeasureService = FSlateApplication::Get().GetRenderer()->GetFontMeasureService();
-	FVector2D DrawSize = FontMeasureService->Measure(GetBranchingPointText(), Font)+FVector2D(5, 5)+FVector2D(MontageBranchingPointHandleWidth, 0);
-	return DrawSize;
-}
-
-FReply SMontageBranchingPointNode::BeginDrag( const FGeometry& MyGeometry, const FPointerEvent& MouseEvent )
-{
-	FVector2D ScreenCursorPos = MouseEvent.GetScreenSpacePosition();
-	FVector2D CursorPos = MyGeometry.AbsoluteToLocal(ScreenCursorPos);
-	FVector2D ScreenNodePosition = MyGeometry.AbsolutePosition;
-
-	bBeingDragged = true;
-
-	return FReply::Handled().BeginDragDrop(FBranchingPointDragDropOp::New(SharedThis(this), ScreenCursorPos, ScreenNodePosition));
-}
-
-void SMontageBranchingPointNode::OnSnapNodeDataPosition(float OriginalX, float SnappedX)
-{
-	EAnimEventTriggerOffsets::Type Offset = EAnimEventTriggerOffsets::NoOffset;
-	if(OriginalX != SnappedX)
-	{
-		Offset = OriginalX < SnappedX ? EAnimEventTriggerOffsets::OffsetBefore : EAnimEventTriggerOffsets::OffsetAfter;
-	}
-	BranchingPoint->TriggerTimeOffset = GetTriggerTimeOffsetForType(Offset);
-}
-
-int32 SMontageBranchingPointNode::OnPaint(const FPaintArgs& Args, const FGeometry& AllottedGeometry, const FSlateRect& MyClippingRect, FSlateWindowElementList& OutDrawElements, int32 LayerId, const FWidgetStyle& InWidgetStyle, bool bParentEnabled) const
-{
-	int32 TextLayerID = LayerId+1;
-
-	FString Text = GetBranchingPointText();
-
-	FVector2D DrawSize = AllottedGeometry.Size;//GetSize();
-	FVector2D TextPos(0, 0); //GetPosition(AllottedGeometry);
-
-	FPaintGeometry MyGeometry =	AllottedGeometry.ToPaintGeometry();
-	// draw border first
-	TArray<FVector2D> LinePoints;
-
-
-	LinePoints.Empty();
-	LinePoints.Add(TextPos);
-	LinePoints.Add(TextPos+FVector2D(DrawSize.X, 0.f));
-	LinePoints.Add(TextPos+FVector2D(DrawSize.X, DrawSize.Y));
-	LinePoints.Add(TextPos+FVector2D(0.f, DrawSize.Y));
-	LinePoints.Add(TextPos);
-
-	FSlateDrawElement::MakeLines( 
-		OutDrawElements,
-		TextLayerID,
-		AllottedGeometry.ToPaintGeometry(),
-		LinePoints,
-		MyClippingRect,
-		ESlateDrawEffect::None,
-		FLinearColor::Black, 
-		false
-		);
-
-	FLinearColor UseNodeColor = NodeColor.Get();
-	FLinearColor HandleColor = NodeColor.Get();
-	if (IsSelected())
-	{
-		UseNodeColor = SelectedNodeColor.Get();
-		HandleColor = FLinearColor::Red;
-	}
-
-	const FSlateBrush* StyleInfo = FEditorStyle::GetBrush( TEXT( "SpecialEditableTextImageNormal" ) );
-	FSlateDrawElement::MakeBox( 
-		OutDrawElements,
-		LayerId, 
-		MyGeometry, 
-		StyleInfo, 
-		MyClippingRect, 
-		ESlateDrawEffect::None,
-		UseNodeColor);
-
-	MyGeometry = AllottedGeometry.ToPaintGeometry( TextPos + FVector2D(2+MontageBranchingPointHandleWidth, 2), DrawSize );
-	FSlateDrawElement::MakeText( 
-		OutDrawElements,
-		TextLayerID,
-		MyGeometry,
-		Text,
-		Font,
-		MyClippingRect,
-		ESlateDrawEffect::None,
-		FLinearColor::Black
-		);
-
-	MyGeometry = AllottedGeometry.ToPaintGeometry( FVector2D(2+TextPos.X, TextPos.Y), FVector2D(MontageBranchingPointHandleWidth, 16) );
-	FSlateDrawElement::MakeBox(
-		OutDrawElements,
-		LayerId, 
-		MyGeometry, 
-		FEditorStyle::GetBrush( TEXT( "Sequencer.Timeline.ScrubHandleUp" ) ), 
-		MyClippingRect, 
-		ESlateDrawEffect::None,
-		HandleColor
-		);
-
-	if(BranchingPoint->TriggerTimeOffset != 0.f) //Do we have an offset to render?
-	{
-		float BranchingPointTime = BranchingPoint->GetTime();
-		if(BranchingPointTime != 0.f && BranchingPointTime != Montage->SequenceLength) //Don't render offset when we are at the start/end of the sequence, doesn't help the user
-		{
-			// ScrubHandle
-			FVector2D MarkerPosition = TextPos;
-			FVector2D MarkerSize = FVector2D(BranchingPointAlignmentMarkerWidth,16.f);
-
-			if(BranchingPoint->TriggerTimeOffset > 0.f)
-			{
-				MarkerPosition.X += 2.f*BranchingPointAlignmentMarkerWidth;
-				MarkerSize.X = -MarkerSize.X;
-			}
-
-			FSlateDrawElement::MakeBox( 
-				OutDrawElements,
-				LayerId + 1, 
-				AllottedGeometry.ToPaintGeometry(MarkerPosition, MarkerSize), 
-				FEditorStyle::GetBrush( TEXT( "Sequencer.Timeline.NotifyAlignmentMarker" ) ), 
-				MyClippingRect, 
-				ESlateDrawEffect::None,
-				FLinearColor(0.f, 0.f, 1.f)
-				);
-		}
-	}
-
-	return TextLayerID;
-}
-
 //////////////////////////////////////////////////////////////////////////
 // SAnimMontagePanel
 
@@ -349,10 +73,6 @@ void SAnimMontagePanel::Update()
 		FLinearColor NodeColor = FLinearColor(0.f, 0.5f, 0.0f, 0.5f);
 		//FLinearColor SelectedColor = FLinearColor(1.0f,0.65,0.0f);
 
-		FLinearColor BranchNodeColor = FLinearColor(0.f, 0.5f, 0.0f, 0.5f);
-		FLinearColor BranchSelectedColor = FLinearColor(1.0f,0.65,0.0f);
-		
-
 		TSharedPtr<SVerticalBox> MontageSlots;
 		PanelArea->SetContent(
 			SAssignNew( MontageSlots, SVerticalBox )
@@ -528,56 +248,6 @@ void SAnimMontagePanel::Update()
 				];
 			}
 		}
-		// ===================================
-		// Branch PointTrack
-		// ===================================
-		{
-			TSharedPtr<STrack> BranchTrack;
-			TSharedRef<S2ColumnWidget> BranchEdTrack = Create2ColumnWidget(MontageSlots.ToSharedRef());
-			BranchEdTrack->LeftColumn->ClearChildren();
-			BranchEdTrack->LeftColumn->AddSlot()
-				.AutoHeight()
-				.VAlign(VAlign_Center)
-				.Padding( FMargin(0.5f, 0.5f) )
-				[
-					SAssignNew(BranchTrack, STrack)
-					.ViewInputMin(ViewInputMin)
-					.ViewInputMax(ViewInputMax)
-					.TrackColor( ColourTracker->GetNextColor() )
-					.OnBarDrag(Editor, &SMontageEditor::OnEditSectionTime)
-					.OnBarDrop(Editor, &SMontageEditor::OnEditSectionTimeFinish)
-					.OnBarClicked(SharedThis(this), &SAnimMontagePanel::ShowSectionInDetailsView)
-					.DraggableBars(Editor, &SMontageEditor::GetSectionStartTimes)
-					.DraggableBarSnapPositions(Editor, &SMontageEditor::GetAnimSegmentStartTimes)
-					.TrackMaxValue(Montage->SequenceLength) // FIXME -> Make actual attribute!
-					.TrackNumDiscreteValues(Montage->GetNumberOfFrames())
-					.OnTrackRightClickContextMenu( this, &SAnimMontagePanel::SummonTrackContextMenu, static_cast<int>(INDEX_NONE))
-					.ScrubPosition( Editor, &SMontageEditor::GetScrubValue )
-				];
-
-			// Generate Branch Point Nodes and add them to the track
-			for (int32 BranchIdx=0; BranchIdx < Montage->BranchingPoints.Num(); BranchIdx++)
-			{
-				BranchTrack->AddTrackNode(
-					SNew(SMontageBranchingPointNode)
-					.Montage(Montage)
-					.BranchingPoint(&Montage->BranchingPoints[BranchIdx])
-					.ViewInputMax(this->ViewInputMax)
-					.ViewInputMin(this->ViewInputMin)
-					.NodeColor(BranchNodeColor)
-					.SelectedNodeColor(BranchSelectedColor)
-					.DataLength( 0.f )
-					.CenterOnPosition(true)
-					.DataStartPos(Editor, &SMontageEditor::GetBranchPointStartPos, BranchIdx)
-					.ToolTipText(Editor, &SMontageEditor::GetBranchPointName, BranchIdx)
-					.OnTrackNodeDragged( Editor, &SMontageEditor::SetBranchPointStartPos, BranchIdx)
-					.OnTrackNodeDropped( Editor, &SMontageEditor::OnMontageChange, static_cast<UObject*>(NULL), true)
-					.OnNodeRightClickContextMenu( this, &SAnimMontagePanel::SummonBranchNodeContextMenu, BranchIdx)
-					.OnTrackNodeClicked( this, &SAnimMontagePanel::ShowBranchPointInDetailsView, BranchIdx )
-					.NodeSelectionSet(&SelectionSet)
-					);
-			}
-		}
 	}
 
 	UpdateSlotGroupWarningVisibility();
@@ -605,13 +275,6 @@ void SAnimMontagePanel::SummonTrackContextMenu( FMenuBuilder& MenuBuilder, float
 {
 	FUIAction UIAction;
 
-	MenuBuilder.BeginSection("AnimMontageBranch", LOCTEXT("Branch Point", "Branch Point"));
-	{
-		UIAction.ExecuteAction.BindRaw(this, &SAnimMontagePanel::OnNewBranchClicked, static_cast<float>(DataPosX));
-		MenuBuilder.AddMenuEntry(LOCTEXT("NewBranchPoint", "New Branch Point"), LOCTEXT("NewBranchPointToolTip", "Adds a new Branch Point"), FSlateIcon(), UIAction);
-	}
-	MenuBuilder.EndSection();
-
 	// Sections
 	MenuBuilder.BeginSection("AnimMontageSections", LOCTEXT("Sections", "Sections"));
 	{
@@ -642,7 +305,7 @@ void SAnimMontagePanel::SummonTrackContextMenu( FMenuBuilder& MenuBuilder, float
 
 	MenuBuilder.BeginSection("AnimMontageElementBulkActions", LOCTEXT("BulkLinkActions", "Bulk Link Actions"));
 	{
-		MenuBuilder.AddSubMenu(LOCTEXT("SetElementLink_SubMenu", "Set Elements to..."), LOCTEXT("SetElementLink_TooTip", "Sets all montage elements (Branching Points, Sections, Notifies) to a chosen link type."), FNewMenuDelegate::CreateSP(this, &SAnimMontagePanel::FillElementSubMenuForTimes));
+		MenuBuilder.AddSubMenu(LOCTEXT("SetElementLink_SubMenu", "Set Elements to..."), LOCTEXT("SetElementLink_TooTip", "Sets all montage elements (Sections, Notifies) to a chosen link type."), FNewMenuDelegate::CreateSP(this, &SAnimMontagePanel::FillElementSubMenuForTimes));
 
 		if(Montage->SlotAnimTracks.Num() > 1)
 		{
@@ -654,20 +317,6 @@ void SAnimMontagePanel::SummonTrackContextMenu( FMenuBuilder& MenuBuilder, float
 	LastContextHeading.Empty();
 }
 
-void SAnimMontagePanel::SummonBranchNodeContextMenu( FMenuBuilder& MenuBuilder, int32 BranchIndex )
-{
-	FUIAction UIAction;
-
-	MenuBuilder.BeginSection("AnimMontagePanelBranchPoint", LOCTEXT("Branch Point", "Branch Point") );
-	{
-		UIAction.ExecuteAction.BindRaw(MontageEditor.Pin().Get(), &SMontageEditor::RemoveBranchPoint, BranchIndex);
-		MenuBuilder.AddMenuEntry(LOCTEXT("DeleteBranch", "Delete Branch"), LOCTEXT("DeleteBranchToolTip", "Deletes Branch"), FSlateIcon(), UIAction);
-	}
-	MenuBuilder.EndSection();
-
-	LastContextHeading = "Branch Point";
-}
-
 void SAnimMontagePanel::FillElementSubMenuForTimes(FMenuBuilder& MenuBuilder)
 {
 	MenuBuilder.AddMenuEntry(LOCTEXT("SubLinkAbs", "Absolute"), LOCTEXT("SubLinkAbs", "Set all elements to absolute link"), FSlateIcon(), FUIAction(FExecuteAction::CreateSP(this, &SAnimMontagePanel::OnSetElementsToLinkMode, EAnimLinkMethod::Absolute)));
@@ -730,66 +379,6 @@ void SAnimMontagePanel::CreateNewSection(const FText& NewSectionName, ETextCommi
 	FSlateApplication::Get().DismissAllMenus();
 }
 
-/** Branches */
-void SAnimMontagePanel::OnNewBranchClicked(float DataPosX)
-{
-	// Show dialog to enter new track name
-	TSharedRef<STextEntryPopup> TextEntry =
-		SNew(STextEntryPopup)
-		.Label( LOCTEXT("NewBranchNameLabel", "Branch Name") )
-		.OnTextCommitted( this, &SAnimMontagePanel::CreateNewBranch, DataPosX );
-
-
-	// Show dialog to enter new event name
-	FSlateApplication::Get().PushMenu(
-		AsShared(), // Menu being summoned from a menu that is closing: Parent widget should be k2 not the menu thats open or it will be closed when the menu is dismissed
-		TextEntry,
-		FSlateApplication::Get().GetCursorPos(),
-		FPopupTransitionEffect( FPopupTransitionEffect::TypeInPopup )
-		);
-
-	TextEntry->FocusDefaultWidget();
-}
-
-void SAnimMontagePanel::CreateNewBranch(const FText& NewBranchName, ETextCommit::Type CommitInfo, float StartTime)
-{
-	if (CommitInfo == ETextCommit::OnEnter)
-	{
-		MontageEditor.Pin()->AddBranchPoint(StartTime,NewBranchName.ToString());
-	}
-	FSlateApplication::Get().DismissAllMenus();
-}
-
-void SAnimMontagePanel::OnBranchPointClicked( int32 BranchPointIndex )
-{
-	// Show dialog to enter new track name
-	TSharedRef<STextEntryPopup> TextEntry =
-		SNew(STextEntryPopup)
-		.Label( LOCTEXT("BranchPointClickedLabel", "Branch Name") )
-		.DefaultText( FText::FromName(Montage->BranchingPoints[BranchPointIndex].EventName) )
-		.OnTextCommitted( this, &SAnimMontagePanel::RenameBranchPoint, BranchPointIndex );
-
-
-	// Show dialog to enter new event name
-	FSlateApplication::Get().PushMenu(
-		AsShared(), // Menu being summoned from a menu that is closing: Parent widget should be k2 not the menu thats open or it will be closed when the menu is dismissed
-		TextEntry,
-		FSlateApplication::Get().GetCursorPos(),
-		FPopupTransitionEffect( FPopupTransitionEffect::TypeInPopup )
-		);
-
-	TextEntry->FocusDefaultWidget();
-}
-
-void SAnimMontagePanel::RenameBranchPoint(const FText &NewName, ETextCommit::Type CommitInfo, int32 BranchPointIndex)
-{
-	if (CommitInfo == ETextCommit::OnEnter)
-	{
-		MontageEditor.Pin()->RenameBranchPoint(BranchPointIndex, NewName.ToString() );
-		FSlateApplication::Get().DismissAllMenus();
-	}
-}
-
 void SAnimMontagePanel::ShowSegmentInDetailsView(int32 AnimSegmentIndex, int32 AnimSlotIndex)
 {
 	UEditorAnimSegment *Obj = Cast<UEditorAnimSegment>(MontageEditor.Pin()->ShowInDetailsView(UEditorAnimSegment::StaticClass()));
@@ -799,15 +388,6 @@ void SAnimMontagePanel::ShowSegmentInDetailsView(int32 AnimSegmentIndex, int32 A
 	}
 }
 
-void SAnimMontagePanel::ShowBranchPointInDetailsView(int32 BranchPointIndex)
-{
-	UEditorBranchPoint *Obj = Cast<UEditorBranchPoint>(MontageEditor.Pin()->ShowInDetailsView(UEditorBranchPoint::StaticClass()));
-	if(Obj != NULL)
-	{
-		Obj->InitBranchPoint(BranchPointIndex);
-	}
-}
-
 void SAnimMontagePanel::ShowSectionInDetailsView(int32 SectionIndex)
 {
 	MontageEditor.Pin()->ShowSectionInDetailsView(SectionIndex);
@@ -1041,12 +621,6 @@ void SAnimMontagePanel::CollectLinkableElements(TArray<FAnimLinkableElement*> &E
 		Elements.Add(Element);
 	}
 
-	for(auto& BranchingPoint : Montage->BranchingPoints)
-	{
-		FAnimLinkableElement* Element = &BranchingPoint;
-		Elements.Add(Element);
-	}
-
 	for(auto& Notify : Montage->Notifies)
 	{
 		FAnimLinkableElement* Element = &Notify;
diff --git a/Engine/Source/Editor/Persona/Private/SAnimMontagePanel.h b/Engine/Source/Editor/Persona/Private/SAnimMontagePanel.h
--- a/Engine/Source/Editor/Persona/Private/SAnimMontagePanel.h
+++ b/Engine/Source/Editor/Persona/Private/SAnimMontagePanel.h
@@ -66,18 +66,13 @@ public:
 	// Functions used as delegates that build the context menus for the montage tracks
 	void SummonTrackContextMenu( FMenuBuilder& MenuBuilder, float DataPosX, int32 SectionIndex, int32 AnimSlotIndex );
 	void SummonSectionContextMenu( FMenuBuilder& MenuBuilder, int32 SectionIndex );
-	void SummonBranchNodeContextMenu( FMenuBuilder& MenuBuilder, int32 BranchIndex );
 	void FillElementSubMenuForTimes( FMenuBuilder& MenuBuilder );
 	void FillSlotSubMenu( FMenuBuilder& Menubuilder );
 
 	void OnNewSlotClicked();
 	void OnNewSectionClicked(float DataPosX);
-	void OnNewBranchClicked(float DataPosX);
-
-	void OnBranchPointClicked( int32 BranchPointIndex );
 
 	void ShowSegmentInDetailsView(int32 AnimSegmentIndex, int32 AnimSlotIndex);
-	void ShowBranchPointInDetailsView(int32 BranchPointIndex);
 	void ShowSectionInDetailsView(int32 SectionIndex);
 
 	void ClearSelected();
@@ -135,8 +130,6 @@ private:
 	
 	void CreateNewSlot(const FText& NewSlotName, ETextCommit::Type CommitInfo);
 	void CreateNewSection(const FText& NewSectionName, ETextCommit::Type CommitInfo, float StartTime);
-	void CreateNewBranch(const FText& NewBranchName, ETextCommit::Type CommitInfo, float StartTime);
-	void RenameBranchPoint(const FText &NewName, ETextCommit::Type CommitInfo, int32 BranchPointIndex);
 
 	/** Called when a segment is removed from a track, so we can adjust the indices in linkable elements
 	 *	@param SegmentIndex - Index of the removed segment
diff --git a/Engine/Source/Editor/Persona/Private/SAnimNotifyPanel.cpp b/Engine/Source/Editor/Persona/Private/SAnimNotifyPanel.cpp
--- a/Engine/Source/Editor/Persona/Private/SAnimNotifyPanel.cpp
+++ b/Engine/Source/Editor/Persona/Private/SAnimNotifyPanel.cpp
@@ -168,6 +168,7 @@ private:
 	
 	FVector2D					TextSize;
 	float						LabelWidth;
+	FVector2D					BranchingPointIconSize;
 
 	/** Last position the user clicked in the widget */
 	FVector2D					LastMouseDownPosition;
@@ -936,6 +937,8 @@ void SAnimNotifyNode::Construct(const FArguments& InArgs)
 	ViewInputMax = InArgs._ViewInputMax;
 
 	MarkerBars = InArgs._MarkerBars;
+
+	SetToolTipText(TAttribute<FText>(this, &SAnimNotifyNode::GetNodeTooltip));
 }
 
 FReply SAnimNotifyNode::OnDragDetected( const FGeometry& MyGeometry, const FPointerEvent& MouseEvent )
@@ -999,13 +1002,22 @@ FText SAnimNotifyNode::GetNodeTooltip() const
 	const FText Frame = FText::AsNumber( (int32)(Percentage * Sequence->GetNumberOfFrames()) );
 	const FText Seconds = FText::AsNumber( Time );
 
+	FText ToolTipText;
 	if (NotifyEvent->GetDuration() > 0.0f)
 	{
 		const FText Duration = FText::AsNumber( NotifyEvent->GetDuration() );
-		return FText::Format( LOCTEXT("NodeToolTipLong", "@ {0} sec (frame {1}) for {2} sec"), Seconds, Frame, Duration );
+		ToolTipText = FText::Format( LOCTEXT("NodeToolTipLong", "@ {0} sec (frame {1}) for {2} sec"), Seconds, Frame, Duration );
+	}
+	else
+	{
+		ToolTipText = FText::Format(LOCTEXT("NodeToolTipShort", "@ {0} sec (frame {1})"), Seconds, Frame);
 	}
 
-	return FText::Format( LOCTEXT("NodeToolTipShort", "@ {0} sec (frame {1})"), Seconds, Frame );
+	if (NotifyEvent->IsBranchingPoint())
+	{
+		ToolTipText = FText::Format(LOCTEXT("AnimNotify_ToolTipBranchingPoint", "{0} (BranchingPoint)"), ToolTipText);
+	}
+	return ToolTipText;
 }
 
 /** @return the Node's position within the graph */
@@ -1089,9 +1101,17 @@ void SAnimNotifyNode::UpdateSizeAndPosition(const FGeometry& AllottedGeometry)
 
 	const TSharedRef< FSlateFontMeasure > FontMeasureService = FSlateApplication::Get().GetRenderer()->GetFontMeasureService();
 	TextSize = FontMeasureService->Measure( GetNotifyText(), Font );
+	LabelWidth = TextSize.X + (TextBorderSize.X * 2.f) + (ScrubHandleSize.X / 2.f);
 
-	LabelWidth = TextSize.X + (TextBorderSize.X * 2.f) + (ScrubHandleSize.X/2.f);
+	bool bDrawBranchingPoint = NotifyEvent->IsBranchingPoint();
+	BranchingPointIconSize = FVector2D(TextSize.Y, TextSize.Y);
+	if (bDrawBranchingPoint)
+	{
+		LabelWidth += BranchingPointIconSize.X + TextBorderSize.X * 2.f;
+	}
 
+	//Calculate scrub handle box size (the notional box around the scrub handle and the alignment marker)
+	float NotifyHandleBoxWidth = FMath::Max(ScrubHandleSize.X, AlignmentMarkerSize.X * 2);
 
 	// Work out where we will have to draw the tool tip
 	FVector2D Size = GetSize();
@@ -1099,9 +1119,6 @@ void SAnimNotifyNode::UpdateSizeAndPosition(const FGeometry& AllottedGeometry)
 	float RightEdgeToNotify = AllottedGeometry.Size.X - NotifyTimePositionX;
 	bDrawTooltipToRight = (RightEdgeToNotify > LabelWidth) || (RightEdgeToNotify > LeftEdgeToNotify);
 
-	//Calculate scrub handle box size (the notional box around the scrub handle and the alignment marker)
-	float NotifyHandleBoxWidth = FMath::Max(ScrubHandleSize.X, AlignmentMarkerSize.X*2);
-
 	// Calculate widget width/position based on where we are drawing the tool tip
 	WidgetX = bDrawTooltipToRight ? (NotifyTimePositionX - (NotifyHandleBoxWidth / 2.f)) : (NotifyTimePositionX - LabelWidth);
 	WidgetSize = bDrawTooltipToRight ? FVector2D(FMath::Max(LabelWidth, NotifyDurationSizeX), NotifyHeight) : FVector2D((LabelWidth + NotifyDurationSizeX), NotifyHeight);
@@ -1137,6 +1154,7 @@ int32 SAnimNotifyNode::OnPaint(const FPaintArgs& Args, const FGeometry& Allotted
 	int32 MarkerLayer = LayerId + 1;
 	int32 ScrubHandleID = MarkerLayer + 1;
 	int32 TextLayerID = ScrubHandleID + 1;
+	int32 BranchPointLayerID = TextLayerID + 1;
 
 	const FSlateBrush* StyleInfo = FEditorStyle::GetBrush( TEXT("SpecialEditableTextImageNormal") );
 	FSlateDrawElement::MakeBox( 
@@ -1183,15 +1201,18 @@ int32 SAnimNotifyNode::OnPaint(const FPaintArgs& Args, const FGeometry& Allotted
 		}
 	}
 
+	// Branching point
+	bool bDrawBranchingPoint = NotifyEvent->IsBranchingPoint();
+
 	// Background
 	FVector2D LabelSize = TextSize + TextBorderSize * 2.f;
-	LabelSize.X += HalfScrubHandleWidth;
+	LabelSize.X += HalfScrubHandleWidth + (bDrawBranchingPoint ? (BranchingPointIconSize.X + TextBorderSize.X * 2.f) : 0.f);
 
 	float LabelX = bDrawTooltipToRight ? NotifyScrubHandleCentre : NotifyScrubHandleCentre - LabelSize.X;
 	float BoxHeight = (NotifyDurationSizeX > 0.f) ? (NotifyHeight - LabelSize.Y) : ((NotifyHeight - LabelSize.Y) / 2.f);
 
-	FVector2D LabelPosition(LabelX, BoxHeight); 
-	
+	FVector2D LabelPosition(LabelX, BoxHeight);
+
 	FLinearColor NodeColor = SAnimNotifyNode::GetNotifyColor();
 
 	FSlateDrawElement::MakeBox( 
@@ -1252,7 +1273,7 @@ int32 SAnimNotifyNode::OnPaint(const FPaintArgs& Args, const FGeometry& Allotted
 	if(NotifyEvent->TriggerTimeOffset != 0.f) //Do we have an offset to render?
 	{
 		float NotifyTime = NotifyEvent->GetTime();
-		if(NotifyTime != 0.f && NotifyTime != Sequence->SequenceLength) //Don't render offset when we are at the start/end of the sequence, doesnt help the user
+		if(NotifyTime != 0.f && NotifyTime != Sequence->SequenceLength) //Don't render offset when we are at the start/end of the sequence, doesn't help the user
 		{
 			float HandleCentre = NotifyScrubHandleCentre;
 			float &Offset = NotifyEvent->TriggerTimeOffset;
@@ -1260,6 +1281,22 @@ int32 SAnimNotifyNode::OnPaint(const FPaintArgs& Args, const FGeometry& Allotted
 			DrawHandleOffset(NotifyEvent->TriggerTimeOffset, NotifyScrubHandleCentre, OutDrawElements, MarkerLayer, AllottedGeometry, MyClippingRect);
 		}
 	}
+
+	// Draw Branching Point
+	if (bDrawBranchingPoint)
+	{
+		FVector2D BranchPointIconPos = LabelPosition + LabelSize - BranchingPointIconSize - FVector2D(bDrawTooltipToRight ? TextBorderSize.X * 2.f : TextBorderSize.X * 4.f, 0.f);
+		FSlateDrawElement::MakeBox(
+			OutDrawElements,
+			BranchPointLayerID,
+			AllottedGeometry.ToPaintGeometry(BranchPointIconPos, BranchingPointIconSize),
+			FEditorStyle::GetBrush(TEXT("AnimNotifyEditor.BranchingPoint")),
+			MyClippingRect,
+			ESlateDrawEffect::None,
+			FLinearColor::White
+			);
+	}	
+
 	return TextLayerID;
 }
 
diff --git a/Engine/Source/Editor/Persona/Private/SMontageEditor.cpp b/Engine/Source/Editor/Persona/Private/SMontageEditor.cpp
--- a/Engine/Source/Editor/Persona/Private/SMontageEditor.cpp
+++ b/Engine/Source/Editor/Persona/Private/SMontageEditor.cpp
@@ -161,11 +161,6 @@ bool SMontageEditor::ValidSection(int32 SectionIndex) const
 	return (MontageObj != NULL && MontageObj->CompositeSections.IsValidIndex(SectionIndex));
 }
 
-bool SMontageEditor::ValidBranch(int32 BranchIndex) const
-{
-	return (MontageObj != NULL && MontageObj->BranchingPoints.IsValidIndex(BranchIndex));
-}
-
 void SMontageEditor::RefreshNotifyTriggerOffsets()
 {
 	for(auto Iter = MontageObj->Notifies.CreateIterator(); Iter; ++Iter)
@@ -187,14 +182,6 @@ void SMontageEditor::RefreshNotifyTriggerOffsets()
 			Notify.EndTriggerTimeOffset = 0.0f;
 		}
 	}
-
-	for(auto Iter = MontageObj->BranchingPoints.CreateIterator(); Iter; ++Iter)
-	{
-		FBranchingPoint& BranchPoint = (*Iter);
-
-		EAnimEventTriggerOffsets::Type PredictedOffset = MontageObj->CalculateOffsetForBranchingPoint(BranchPoint.GetTime());
-		BranchPoint.RefreshTriggerOffset(PredictedOffset);
-	}
 }
 
 bool SMontageEditor::GetSectionTime( int32 SectionIndex, float &OutTime ) const
@@ -246,13 +233,6 @@ TArray<FTrackMarkerBar> SMontageEditor::GetMarkerBarInformation() const
 			Bar.DrawColour = FLinearColor(0.f,1.f,0.f);
 			MarkerBars.Add(Bar);
 		}
-		for( int32 I = 0; I < MontageObj->BranchingPoints.Num(); ++I )
-		{
-			FTrackMarkerBar Bar;
-			Bar.Time = MontageObj->BranchingPoints[I].GetTime();
-			Bar.DrawColour = FLinearColor(0.8f,0.f,0.8f);
-			MarkerBars.Add(Bar);
-		}
 	}
 	return MarkerBars;
 }
@@ -378,8 +358,6 @@ void SMontageEditor::SortAndUpdateMontage()
 
 	SortSections();
 
-	SortBranchPoints();
-
 	RefreshNotifyTriggerOffsets();
 
 	// Update view (this will recreate everything)
@@ -404,22 +382,6 @@ void SMontageEditor::SortAnimSegments()
 	}
 }
 
-/** Sort BranchPoints by time */
-void SMontageEditor::SortBranchPoints()
-{
-	struct FCompareBranchPoints
-	{
-		bool operator()( const FBranchingPoint &A, const FBranchingPoint &B ) const
-		{
-			return A.GetTriggerTime() < B.GetTriggerTime();
-		}
-	};
-	if (MontageObj != NULL)
-	{
-		MontageObj->BranchingPoints.Sort(FCompareBranchPoints());
-	}
-}
-
 /** Sort Composite Sections by Start TIme */
 void SMontageEditor::SortSections()
 {
@@ -468,7 +430,7 @@ void SMontageEditor::EnsureSlotNode()
 	}
 }
 
-/** Make sure all Sections, Branches, and Notifies are clamped to NewEndTime (called before NewEndTime is set to SequenceLength) */
+/** Make sure all Sections and Notifies are clamped to NewEndTime (called before NewEndTime is set to SequenceLength) */
 bool SMontageEditor::ClampToEndTime(float NewEndTime)
 {
 	bool bClampingNeeded = SAnimEditorBase::ClampToEndTime(NewEndTime);
@@ -485,17 +447,6 @@ bool SMontageEditor::ClampToEndTime(float NewEndTime)
 			}
 		}
 
-		for(int32 i=0; i < MontageObj->BranchingPoints.Num(); i++)
-		{
-			FBranchingPoint& BranchPoint = MontageObj->BranchingPoints[i];
-			float BranchingPointTime = BranchPoint.GetTime();
-			if(BranchingPointTime > NewEndTime)
-			{
-				BranchPoint.SetTime(BranchingPointTime * ratio);
-				BranchPoint.TriggerTimeOffset = GetTriggerTimeOffsetForType(MontageObj->CalculateOffsetForBranchingPoint(BranchPoint.GetTime()));
-			}
-		}
-
 		for(int32 i=0; i < MontageObj->Notifies.Num(); i++)
 		{
 			FAnimNotifyEvent& Notify = MontageObj->Notifies[i];
@@ -550,147 +501,6 @@ FString	SMontageEditor::GetSectionName(int32 SectionIndex) const
 	return FString();
 }
 
-float SMontageEditor::GetBranchPointStartPos(int32 BranchPointIndex) const
-{
-	if (ValidBranch(BranchPointIndex))
-	{
-		return MontageObj->BranchingPoints[BranchPointIndex].GetTime();
-	}
-	return 0.f;
-}
-
-FString	SMontageEditor::GetBranchPointName(int32 BranchPointIndex) const
-{
-	if (ValidBranch(BranchPointIndex))
-	{
-		return MontageObj->BranchingPoints[BranchPointIndex].EventName.ToString();
-	}
-	return FString();
-}
-
-const float SMontageEditor::MinimumBranchPointSeperation = 1.0f/30.0f;
-
-void SMontageEditor::SetBranchPointStartPos(float NewStartPos, int32 BranchPointIndex)
-{
-	if (ValidBranch(BranchPointIndex))
-	{
-		const FScopedTransaction Transaction( LOCTEXT("SetBranchPoint", "Set Branch Point Start Time") );
-		MontageObj->Modify();
-
-		FBranchingPoint& BranchPoint = MontageObj->BranchingPoints[BranchPointIndex];
-
-		float OriginalPosition = BranchPoint.GetTime();
-		float TempPosition = NewStartPos;
-		int32 MovementDirection = 1;
-
-		int32 BranchIndex = 0;
-		bool bHasSnapped = false;
-		float Prediction = NewStartPos;
-		for (; MontageObj->BranchingPoints.IsValidIndex(BranchIndex); BranchIndex += MovementDirection)
-		{
-			// Don't check self
-			if (BranchPointIndex == BranchIndex)
-			{
-				continue;
-			}
-
-			float BranchTime = MontageObj->BranchingPoints[BranchIndex].GetTime();
-			float LowerBound = BranchTime - MinimumBranchPointSeperation;
-			float UpperBound = BranchTime + MinimumBranchPointSeperation;
-
-			if (Prediction > BranchTime && Prediction < UpperBound)
-			{
-				if (!bHasSnapped)
-				{
-					bHasSnapped = true;
-					MovementDirection = 1;
-				}
-
-				Prediction = MovementDirection == 1 ? UpperBound : LowerBound;
-			}
-			else if (Prediction <= BranchTime && Prediction > LowerBound)
-			{
-				if (!bHasSnapped)
-				{
-					bHasSnapped = true;
-					MovementDirection = -1;
-				}
-
-				Prediction = MovementDirection == 1 ? UpperBound : LowerBound;
-			}
-			else if (bHasSnapped)
-			{
-				break;
-			}
-		}
-
-		if (bHasSnapped)
-		{
-			float ToOrig = FMath::Abs(OriginalPosition - NewStartPos);
-			float ToPrediction = FMath::Abs(Prediction - NewStartPos);
-			
-			NewStartPos = ToOrig < ToPrediction ? OriginalPosition : Prediction;
-		}
-
-		float CurrentLowerBound = NewStartPos - MinimumBranchPointSeperation;
-		float CurrentUpperBound = NewStartPos + MinimumBranchPointSeperation;
-
-		if (!(CurrentLowerBound >= 0.0f && CurrentUpperBound <= MontageObj->SequenceLength))
-		{
-			NewStartPos = OriginalPosition;
-		}
-
-		BranchPoint.LinkMontage(MontageObj, NewStartPos, BranchPoint.GetSlotIndex());
-		BranchPoint.RefreshTriggerOffset(MontageObj->CalculateOffsetForBranchingPoint(BranchPoint.GetTime()));
-
-		if(AnimMontagePanel.IsValid())
-		{
-			AnimMontagePanel->ShowBranchPointInDetailsView(BranchPointIndex);
-		}
-	}
-}
-
-void SMontageEditor::RemoveBranchPoint(int32 BranchPointIndex)
-{
-	if(ValidBranch(BranchPointIndex))
-	{
-		const FScopedTransaction Transaction( LOCTEXT("RemoveBranchPoint", "Remove Branch Point") );
-		MontageObj->Modify();
-
-		MontageObj->BranchingPoints.RemoveAt(BranchPointIndex);
-		MontageObj->MarkPackageDirty();
-		SortAndUpdateMontage();
-	}
-}
-void SMontageEditor::AddBranchPoint(float StartTime, FString EventName)
-{
-	if(MontageObj!=NULL)
-	{
-		const FScopedTransaction Transaction( LOCTEXT("AddBranchPoint", "Add Branch Point") );
-		MontageObj->Modify();
-
-		FBranchingPoint NewBranch;
-		NewBranch.LinkMontage(MontageObj, StartTime);
-		NewBranch.TriggerTimeOffset = GetTriggerTimeOffsetForType(MontageObj->CalculateOffsetForBranchingPoint(NewBranch.GetTime()));
-		NewBranch.EventName=FName(*EventName);
-		MontageObj->BranchingPoints.Add(NewBranch);
-		MontageObj->MarkPackageDirty();
-		SortAndUpdateMontage();
-	}
-}
-
-void SMontageEditor::RenameBranchPoint(int32 BranchIndex, FString NewEventName)
-{
-	if(ValidBranch(BranchIndex))
-	{
-		const FScopedTransaction Transaction( LOCTEXT("RenameBranchPoint", "Rename Branch Point") );
-		MontageObj->Modify();
-
-		MontageObj->BranchingPoints[BranchIndex].EventName = FName( *NewEventName );
-		MontageObj->MarkPackageDirty();
-	}
-}
-
 void SMontageEditor::RenameSlotNode(int32 SlotIndex, FString NewSlotName)
 {
 	if(MontageObj->SlotAnimTracks.IsValidIndex(SlotIndex))
diff --git a/Engine/Source/Editor/Persona/Private/SMontageEditor.h b/Engine/Source/Editor/Persona/Private/SMontageEditor.h
--- a/Engine/Source/Editor/Persona/Private/SMontageEditor.h
+++ b/Engine/Source/Editor/Persona/Private/SMontageEditor.h
@@ -60,9 +60,6 @@ public:
 	void RestartPreviewPlayAllSections();
 
 private:
-	/** Minimum distance to make sure branching points are seperated */
-	static const float MinimumBranchPointSeperation;
-
 	/** Pointer to the animation sequence being edited */
 	UAnimMontage* MontageObj;
 
@@ -81,7 +78,6 @@ private:
 	virtual float CalculateSequenceLengthOfEditorObject() const override;
 	void SortAndUpdateMontage();
 	void CollapseMontage();
-	void SortBranchPoints();
 	void SortAnimSegments();
 	void SortSections();
 	void EnsureStartingSection();
@@ -93,7 +89,6 @@ private:
 
 	bool ValidIndexes(int32 AnimSlotIndex, int32 AnimSegmentIndex) const;
 	bool ValidSection(int32 SectionIndex) const;
-	bool ValidBranch(int32 BranchIndex) const;	
 
 	/** Updates Notify trigger offsets to take into account current montage state */
 	void RefreshNotifyTriggerOffsets();
@@ -119,13 +114,6 @@ public:
 	void					RemoveSection(int32 SectionIndex);
 
 	FString					GetSectionName(int32 SectionIndex) const;
-
-	float					GetBranchPointStartPos(int32 BranchPointIndex) const;
-	FString					GetBranchPointName(int32 BranchPointIndex) const;
-	void					SetBranchPointStartPos(float NewStartPos, int32 BranchPointIndex);
-	void					RemoveBranchPoint(int32 BranchPointIndex);
-	void					AddBranchPoint(float StartTime, FString EventName);
-	void					RenameBranchPoint(int32 BranchIndex, FString NewEventName);
 	void					RenameSlotNode(int32 SlotIndex, FString NewSlotName);
 
 	void					AddNewMontageSlot(FString NewSlotName);
diff --git a/Engine/Source/Editor/UnrealEd/Classes/Animation/EditorBranchPoint.h b/Engine/Source/Editor/UnrealEd/Classes/Animation/EditorBranchPoint.h
deleted file mode 100644
diff --git a/Engine/Source/Editor/UnrealEd/Private/Animation/EditorBranchPoint.cpp b/Engine/Source/Editor/UnrealEd/Private/Animation/EditorBranchPoint.cpp
deleted file mode 100644
diff --git a/Engine/Source/Runtime/Core/Public/UObject/ObjectVersion.h b/Engine/Source/Runtime/Core/Public/UObject/ObjectVersion.h
--- a/Engine/Source/Runtime/Core/Public/UObject/ObjectVersion.h
+++ b/Engine/Source/Runtime/Core/Public/UObject/ObjectVersion.h
@@ -649,7 +649,8 @@ enum EUnrealEngineObjectUE4Version
 	VER_UE4_RENAME_WIDGET_VISIBILITY,
 	// add track curves for animation
 	VER_UE4_ANIMATION_ADD_TRACKCURVES,
-
+	// Removed BranchingPoints from AnimMontages and converted them to regular AnimNotifies.
+	VER_UE4_MONTAGE_BRANCHING_POINT_REMOVAL,
 	// -----<new versions can be added before this line>-------------------------------------------------
 	// - this needs to be the last line (see note below)
 	VER_UE4_AUTOMATIC_VERSION_PLUS_ONE,
diff --git a/Engine/Source/Runtime/Engine/Classes/Animation/AnimInstance.h b/Engine/Source/Runtime/Engine/Classes/Animation/AnimInstance.h
--- a/Engine/Source/Runtime/Engine/Classes/Animation/AnimInstance.h
+++ b/Engine/Source/Runtime/Engine/Classes/Animation/AnimInstance.h
@@ -632,6 +632,7 @@ public:
 
 	/** Trigger AnimNotifies **/
 	void TriggerAnimNotifies(float DeltaSeconds);
+	void TriggerSingleAnimNotify(const FAnimNotifyEvent* AnimNotifyEvent);
 
 	/** Add curve float data using a curve Uid, the name of the curve will be resolved from the skeleton **/
 	void AddCurveValue(const USkeleton::AnimCurveUID Uid, float Value, int32 CurveTypeFlags);
diff --git a/Engine/Source/Runtime/Engine/Classes/Animation/AnimLinkableElement.h b/Engine/Source/Runtime/Engine/Classes/Animation/AnimLinkableElement.h
--- a/Engine/Source/Runtime/Engine/Classes/Animation/AnimLinkableElement.h
+++ b/Engine/Source/Runtime/Engine/Classes/Animation/AnimLinkableElement.h
@@ -88,6 +88,9 @@ struct FAnimLinkableElement
 	/** Gets the sequence this element is linked to */
 	ENGINE_API const UAnimSequenceBase* GetLinkedSequence() {return LinkedSequence;}
 
+	/** Gets the Montage this element is linked to, if any */
+	ENGINE_API const UAnimMontage* GetLinkedMontage() const { return LinkedMontage; }
+
 	/** Changes the way this element is linked to its segment
 	 * @param NewLinkMethod The new linking method to use
 	 */
diff --git a/Engine/Source/Runtime/Engine/Classes/Animation/AnimMontage.h b/Engine/Source/Runtime/Engine/Classes/Animation/AnimMontage.h
--- a/Engine/Source/Runtime/Engine/Classes/Animation/AnimMontage.h
+++ b/Engine/Source/Runtime/Engine/Classes/Animation/AnimMontage.h
@@ -69,11 +69,7 @@ struct FSlotAnimationTrack
 };
 
 /** 
- * Special event handler for branching point
- * This is different from Notifies since they handle multiple calls in a frame
- * when you can modify positions when it gets called, and Tick will handle correctly
- * however that limits multi threading, so this is expensive
- * Use notifies if possible
+ * Remove FBranchingPoint when VER_UE4_MONTAGE_BRANCHING_POINT_REMOVAL is removed.
  */
 USTRUCT()
 struct FBranchingPoint : public FAnimLinkableElement
@@ -92,9 +88,50 @@ struct FBranchingPoint : public FAnimLinkableElement
 
 	/** Returns the time this branching point should be triggered */
 	float GetTriggerTime() const { return GetTime() + TriggerTimeOffset; }
+};
+
+/**  */
+UENUM()
+namespace EAnimNotifyEventType
+{
+	enum Type
+	{
+		/** */
+		Begin,
+		/** */
+		End,
+	};
+}
+
+/** AnimNotifies marked as BranchingPoints will create these markers on their Begin/End times.
+	They create stopping points when the Montage is being ticked to dispatch events. */
+USTRUCT()
+struct FBranchingPointMarker
+{
+	GENERATED_USTRUCT_BODY()
+
+	UPROPERTY()
+	int32 NotifyIndex;
+	
+	UPROPERTY()
+	float TriggerTime;
+
+	UPROPERTY()
+	TEnumAsByte<EAnimNotifyEventType::Type> NotifyEventType;
+
+	FBranchingPointMarker()
+		: NotifyIndex(INDEX_NONE)
+		, TriggerTime(0.f)
+		, NotifyEventType(EAnimNotifyEventType::Begin)
+	{
+	}
 
-	/** Updates trigger offset based on a combination of predicted offset and current offset */
-	ENGINE_API void RefreshTriggerOffset(EAnimEventTriggerOffsets::Type PredictedOffsetType);
+	FBranchingPointMarker(int32 InNotifyIndex, float InTriggerTime, EAnimNotifyEventType::Type InNotifyEventType)
+		: NotifyIndex(InNotifyIndex)
+		, TriggerTime(InTriggerTime)
+		, NotifyEventType(InNotifyEventType)
+	{
+	}
 };
 
 /**
@@ -160,6 +197,12 @@ public:
 	// transient PreviousWeight - Weight of previous tick
 	float PreviousWeight;
 
+	/** Currently Active AnimNotifyState, stored as a copy of the event as we need to
+		call NotifyEnd on the event after a deletion in the editor. After this the event
+		is removed correctly. */
+	UPROPERTY(Transient)
+	TArray<FAnimNotifyEvent> ActiveStateBranchingPoints;
+
 private:
 	UPROPERTY()
 	float Position;
@@ -258,8 +301,13 @@ private:
 
 	/** Delegate function handlers
 	 */
-	void HandleEvents(float PreviousTrackPos, float CurrentTrackPos, const FBranchingPoint* BranchingPoint);
-	void TriggerEventHandler(FName EventName);
+	void HandleEvents(float PreviousTrackPos, float CurrentTrackPos, const FBranchingPointMarker* BranchingPointMarker);
+	
+	/** Updates ActiveStateBranchingPoints array and triggers Begin/End notifications based on CurrentTrackPosition */
+	void UpdateActiveStateBranchingPoints(float CurrentTrackPosition);
+
+	/** Trigger associated events when Montage ticking reaches given FBranchingPointMarker */
+	void BranchingPointEventHandler(const FBranchingPointMarker* BranchingPointMarker);
 	void RefreshNextPrevSections();
 };
 
@@ -284,9 +332,9 @@ class UAnimMontage : public UAnimCompositeBase
 	UPROPERTY()
 	TArray<struct FSlotAnimationTrack> SlotAnimTracks;
 
-	// slot data, each slot contains anim track
+	// Remove this when VER_UE4_MONTAGE_BRANCHING_POINT_REMOVAL is removed.
 	UPROPERTY()
-	TArray<struct FBranchingPoint> BranchingPoints;
+	TArray<struct FBranchingPoint> BranchingPoints_DEPRECATED;
 
 	/** If this is on, it will allow extracting root motion translation. DEPRECATED in 4.5 root motion is controlled by anim sequences **/
 	UPROPERTY()
@@ -334,9 +382,6 @@ public:
 	virtual void ReplaceReferredAnimations(const TMap<UAnimSequence*, UAnimSequence*>& ReplacementMap) override;
 	// End UAnimationAsset interface
 
-	/** Calculates what (if any) offset should be applied to the trigger time of a branch point given its display time */
-	ENGINE_API EAnimEventTriggerOffsets::Type CalculateOffsetForBranchingPoint(float BranchingPointDisplayTime) const;
-
 	/** Update all linkable elements contained in the montage */
 	ENGINE_API void UpdateLinkableElements();
 
@@ -348,10 +393,6 @@ public:
 	ENGINE_API void UpdateLinkableElements(int32 SlotIdx, int32 SegmentIdx);
 #endif
 
-	/** Find First BranchingPoint between ]StartTrackPos,EndTrackPos] 
-	 * Has to be contiguous range, does not handle looping and wrapping over. **/
-	const FBranchingPoint* FindFirstBranchingPoint(float StartTrackPos, float EndTrackPos);
-
 	/** Get FCompositeSection with InSectionName */
 	FCompositeSection& GetAnimCompositeSection(int32 SectionIndex);
 	const FCompositeSection& GetAnimCompositeSection(int32 SectionIndex) const;
@@ -448,10 +489,29 @@ private:
 
 private:
 
-	/** Sort CompositeSections in the order of StartPos */
-	void SortAnimBranchingPointByTime();
+	/** Convert all branching points to AnimNotifies */
+	void ConvertBranchingPointsToAnimNotifies();
+	/** Recreate BranchingPoint markers from AnimNotifies marked 'BranchingPoints' */
+	void RefreshBranchingPointMarkers();
+	void AddBranchingPointMarker(FBranchingPointMarker Marker, TMap<float, FAnimNotifyEvent*>& TriggerTimes);
+	
+	/** Cached list of Branching Point markers */
+	UPROPERTY()
+	TArray<FBranchingPointMarker> BranchingPointMarkers;
 
+	UPROPERTY(Transient)
 	// @remove me: temporary variable to do sort while property window changed
 	// this should be fixed when we have tool to do so.
 	bool bAnimBranchingPointNeedsSort;
+
+public:
+
+	/** Keep track of which AnimNotify_State are marked as BranchingPoints, so we can update their state when the Montage is ticked */
+	UPROPERTY()
+	TArray<int32> BranchingPointStateNotifyIndices;
+
+	/** Find first branching point marker between track positions */
+	const FBranchingPointMarker* FindFirstBranchingPointMarker(float StartTrackPos, float EndTrackPos);
+	/** Filter out notifies from array that are marked as 'BranchingPoints' */
+	void FilterOutNotifyBranchingPoints(TArray<const FAnimNotifyEvent*>& InAnimNotifies);
 };
diff --git a/Engine/Source/Runtime/Engine/Classes/Animation/AnimSequenceBase.h b/Engine/Source/Runtime/Engine/Classes/Animation/AnimSequenceBase.h
--- a/Engine/Source/Runtime/Engine/Classes/Animation/AnimSequenceBase.h
+++ b/Engine/Source/Runtime/Engine/Classes/Animation/AnimSequenceBase.h
@@ -29,6 +29,19 @@ namespace EAnimEventTriggerOffsets
 
 ENGINE_API float GetTriggerTimeOffsetForType(EAnimEventTriggerOffsets::Type OffsetType);
 
+/** Ticking method for AnimNotifies in AnimMontages */
+UENUM()
+namespace EMontageNotifyTickType
+{
+	enum Type
+	{
+		/** Queue notifies, and trigger them at the end of the evaluation phase (faster). Not suitable for changing sections or montage position. */
+		Queued,
+		/** Trigger notifies as they are encountered (Slower). Suitable for changing sections or montage position. */
+		BranchingPoint,
+	};
+}
+
 /**
  * Triggers an animation notify.  Each AnimNotifyEvent contains an AnimNotify object
  * which has its Notify method called and passed to the animation.
@@ -69,6 +82,13 @@ struct FAnimNotifyEvent : public FAnimLinkableElement
 	UPROPERTY()
 	FAnimLinkableElement EndLink;
 
+	/** If TRUE, this notify has been converted from an old BranchingPoint. */
+	UPROPERTY()
+	bool bConvertedFromBranchingPoint;
+
+	UPROPERTY(EditAnywhere, BlueprintReadWrite, Category=AnimNotifyEvent)
+	TEnumAsByte<EMontageNotifyTickType::Type> MontageTickType;
+
 #if WITH_EDITORONLY_DATA
 	/** Color of Notify in editor */
 	UPROPERTY()
@@ -90,6 +110,8 @@ struct FAnimNotifyEvent : public FAnimLinkableElement
 #endif // WITH_EDITORONLY_DATA
 		, NotifyStateClass(NULL)
 		, Duration(0)
+		, bConvertedFromBranchingPoint(false)
+		, MontageTickType(EMontageNotifyTickType::Queued)
 #if WITH_EDITORONLY_DATA
 		, TrackIndex(0)
 #endif // WITH_EDITORONLY_DATA
@@ -112,6 +134,12 @@ struct FAnimNotifyEvent : public FAnimLinkableElement
 
 	ENGINE_API void SetDuration(float NewDuration);
 
+	/** Returns true is this AnimNotify is a BranchingPoint */
+	ENGINE_API bool IsBranchingPoint() const
+	{
+		return (MontageTickType == EMontageNotifyTickType::BranchingPoint) && GetLinkedMontage();
+	}
+
 	/** Returns true if this is blueprint derived notifies **/
 	bool IsBlueprintNotify() const
 	{
diff --git a/Engine/Source/Runtime/Engine/Private/Animation/AnimInstance.cpp b/Engine/Source/Runtime/Engine/Private/Animation/AnimInstance.cpp
--- a/Engine/Source/Runtime/Engine/Private/Animation/AnimInstance.cpp
+++ b/Engine/Source/Runtime/Engine/Private/Animation/AnimInstance.cpp
@@ -1113,27 +1113,60 @@ void UAnimInstance::TriggerAnimNotifies(float DeltaSeconds)
 			continue;
 		}
 
-		if(AnimNotifyEvent->Notify != NULL)
+		// Trigger non 'state' AnimNotifies
+		TriggerSingleAnimNotify(AnimNotifyEvent);
+	}
+
+	// Send end notification to AnimNotifyState not active anymore.
+	for(int32 Index=0; Index<ActiveAnimNotifyState.Num(); Index++)
+	{
+		const FAnimNotifyEvent& AnimNotifyEvent = ActiveAnimNotifyState[Index];
+		AnimNotifyEvent.NotifyStateClass->NotifyEnd(SkelMeshComp, Cast<UAnimSequenceBase>(AnimNotifyEvent.NotifyStateClass->GetOuter()));
+	}
+
+	// Call 'NotifyBegin' event on freshly added AnimNotifyState.
+	for (int32 Index = 0; Index < NotifyStateBeginEvent.Num(); Index++)
+	{
+		const FAnimNotifyEvent* AnimNotifyEvent = NotifyStateBeginEvent[Index];
+		AnimNotifyEvent->NotifyStateClass->NotifyBegin(SkelMeshComp, Cast<UAnimSequenceBase>(AnimNotifyEvent->NotifyStateClass->GetOuter()), AnimNotifyEvent->GetDuration());
+	}
+
+	// Switch our arrays.
+	ActiveAnimNotifyState = MoveTemp(NewActiveAnimNotifyState);
+
+	// Tick currently active AnimNotifyState
+	for(int32 Index=0; Index<ActiveAnimNotifyState.Num(); Index++)
+	{
+		const FAnimNotifyEvent& AnimNotifyEvent = ActiveAnimNotifyState[Index];
+		AnimNotifyEvent.NotifyStateClass->NotifyTick(SkelMeshComp, Cast<UAnimSequenceBase>(AnimNotifyEvent.NotifyStateClass->GetOuter()), DeltaSeconds);
+	}
+}
+
+void UAnimInstance::TriggerSingleAnimNotify(const FAnimNotifyEvent* AnimNotifyEvent)
+{
+	// This is for non 'state' anim notifies.
+	if (AnimNotifyEvent && (AnimNotifyEvent->NotifyStateClass == NULL))
+	{
+		if (AnimNotifyEvent->Notify != NULL)
 		{
 			// Implemented notify: just call Notify. UAnimNotify will forward this to the event which will do the work.
-			AnimNotifyEvent->Notify->Notify(SkelMeshComp, Cast<UAnimSequenceBase>(AnimNotifyEvent->Notify->GetOuter()));
+			AnimNotifyEvent->Notify->Notify(GetSkelMeshComponent(), Cast<UAnimSequenceBase>(AnimNotifyEvent->Notify->GetOuter()));
 		}
-		else if( AnimNotifyEvent->NotifyName != NAME_None )
+		else if (AnimNotifyEvent->NotifyName != NAME_None)
 		{
 			// Custom Event based notifies. These will call a AnimNotify_* function on the AnimInstance.
 			FString FuncName = FString::Printf(TEXT("AnimNotify_%s"), *AnimNotifyEvent->NotifyName.ToString());
 			FName FuncFName = FName(*FuncName);
 
 			UFunction* Function = FindFunction(FuncFName);
-			if( Function )
+			if (Function)
 			{
 				// if parameter is none, add event
-				if ( Function->NumParms == 0 )
+				if (Function->NumParms == 0)
 				{
-					ProcessEvent( Function, NULL );								
+					ProcessEvent(Function, NULL);
 				}
-				else if ( Function->NumParms == 1 &&  
-					Cast<UObjectProperty>(Function->PropertyLink) != NULL)
+				else if ((Function->NumParms == 1) && (Cast<UObjectProperty>(Function->PropertyLink) != NULL))
 				{
 					struct FAnimNotifierHandler_Parms
 					{
@@ -1142,7 +1175,7 @@ void UAnimInstance::TriggerAnimNotifies(float DeltaSeconds)
 
 					FAnimNotifierHandler_Parms Parms;
 					Parms.Notify = AnimNotifyEvent->Notify;
-					ProcessEvent( Function, &Parms );								
+					ProcessEvent(Function, &Parms);
 				}
 				else
 				{
@@ -1152,30 +1185,6 @@ void UAnimInstance::TriggerAnimNotifies(float DeltaSeconds)
 			}
 		}
 	}
-
-	// Send end notification to AnimNotifyState not active anymore.
-	for(int32 Index=0; Index<ActiveAnimNotifyState.Num(); Index++)
-	{
-		const FAnimNotifyEvent& AnimNotifyEvent = ActiveAnimNotifyState[Index];
-		AnimNotifyEvent.NotifyStateClass->NotifyEnd(SkelMeshComp, Cast<UAnimSequenceBase>(AnimNotifyEvent.NotifyStateClass->GetOuter()));
-	}
-
-	// Call 'NotifyBegin' event on freshly added AnimNotifyState.
-	for (int32 Index = 0; Index < NotifyStateBeginEvent.Num(); Index++)
-	{
-		const FAnimNotifyEvent* AnimNotifyEvent = NotifyStateBeginEvent[Index];
-		AnimNotifyEvent->NotifyStateClass->NotifyBegin(SkelMeshComp, Cast<UAnimSequenceBase>(AnimNotifyEvent->NotifyStateClass->GetOuter()), AnimNotifyEvent->GetDuration());
-	}
-
-	// Switch our arrays.
-	ActiveAnimNotifyState = MoveTemp(NewActiveAnimNotifyState);
-
-	// Tick currently active AnimNotifyState
-	for(int32 Index=0; Index<ActiveAnimNotifyState.Num(); Index++)
-	{
-		const FAnimNotifyEvent& AnimNotifyEvent = ActiveAnimNotifyState[Index];
-		AnimNotifyEvent.NotifyStateClass->NotifyTick(SkelMeshComp, Cast<UAnimSequenceBase>(AnimNotifyEvent.NotifyStateClass->GetOuter()), DeltaSeconds);
-	}
 }
 
 void UAnimInstance::AnimNotify_Sound(UAnimNotify* AnimNotify)
diff --git a/Engine/Source/Runtime/Engine/Private/Animation/AnimMontage.cpp b/Engine/Source/Runtime/Engine/Private/Animation/AnimMontage.cpp
--- a/Engine/Source/Runtime/Engine/Private/Animation/AnimMontage.cpp
+++ b/Engine/Source/Runtime/Engine/Private/Animation/AnimMontage.cpp
@@ -9,16 +9,6 @@
 #include "AnimationUtils.h"
 #include "AnimationRuntime.h"
 
-///////////////////////////////////////////////////////////////////////////
-//
-void FBranchingPoint::RefreshTriggerOffset(EAnimEventTriggerOffsets::Type PredictedOffsetType)
-{
-	if(PredictedOffsetType == EAnimEventTriggerOffsets::NoOffset || TriggerTimeOffset == 0.f)
-	{
-		TriggerTimeOffset = GetTriggerTimeOffsetForType(PredictedOffsetType);
-	}
-}
-
 ///////////////////////////////////////////////////////////////////////////
 //
 UAnimMontage::UAnimMontage(const FObjectInitializer& ObjectInitializer)
@@ -391,21 +381,6 @@ void UAnimMontage::PostLoad()
 		}
 	}
 
-	for(auto& BranchingPoint : BranchingPoints)
-	{
-		if(BranchingPoint.DisplayTime_DEPRECATED != 0.0f)
-		{
-			BranchingPoint.Clear();
-			BranchingPoint.LinkMontage(this, BranchingPoint.DisplayTime_DEPRECATED);
-		}
-		else
-		{
-			BranchingPoint.LinkMontage(this, BranchingPoint.GetTime());
-		}
-	}
-
-	SortAnimBranchingPointByTime();
-
 	bool bRootMotionEnabled = bEnableRootMotionTranslation || bEnableRootMotionRotation;
 
 	if (bRootMotionEnabled)
@@ -445,21 +420,21 @@ void UAnimMontage::PostLoad()
 
 	// verify if skeleton matches, otherwise clear it, this can happen if anim sequence has been modified when this hasn't been loaded. 
 	{
-	USkeleton* MySkeleton = GetSkeleton();
-	for (int32 I=0; I<SlotAnimTracks.Num(); ++I)
-	{
-		if ( SlotAnimTracks[I].AnimTrack.AnimSegments.Num() > 0 )
+		USkeleton* MySkeleton = GetSkeleton();
+		for (int32 I=0; I<SlotAnimTracks.Num(); ++I)
 		{
-			UAnimSequence * Sequence = Cast<UAnimSequence>(SlotAnimTracks[I].AnimTrack.AnimSegments[0].AnimReference);
-			if ( Sequence && Sequence->GetSkeleton() != MySkeleton )
+			if ( SlotAnimTracks[I].AnimTrack.AnimSegments.Num() > 0 )
 			{
-				SlotAnimTracks[I].AnimTrack.AnimSegments[0].AnimReference = 0;
-				MarkPackageDirty();
-				break;
+				UAnimSequence * Sequence = Cast<UAnimSequence>(SlotAnimTracks[I].AnimTrack.AnimSegments[0].AnimReference);
+				if ( Sequence && Sequence->GetSkeleton() != MySkeleton )
+				{
+					SlotAnimTracks[I].AnimTrack.AnimSegments[0].AnimReference = 0;
+					MarkPackageDirty();
+					break;
+				}
 			}
 		}
 	}
-	}
 #endif // WITH_EDITORONLY_DATA
 
 	// Register Slots w/ Skeleton
@@ -492,79 +467,210 @@ void UAnimMontage::PostLoad()
 			Notify.EndLink.LinkMontage(this, Notify.GetTime() + Notify.Duration);
 		}
 	}
+
+	// Convert BranchingPoints to AnimNotifies.
+	if (GetLinker() && (GetLinker()->UE4Ver() < VER_UE4_MONTAGE_BRANCHING_POINT_REMOVAL) )
+	{
+		ConvertBranchingPointsToAnimNotifies();
+	}
 }
 
-void UAnimMontage::SortAnimBranchingPointByTime()
+void UAnimMontage::ConvertBranchingPointsToAnimNotifies()
 {
-	struct FCompareFBranchingPointTime
+	if (BranchingPoints_DEPRECATED.Num() > 0)
 	{
-		FORCEINLINE bool operator()( const FBranchingPoint &A, const FBranchingPoint &B ) const
+		// Handle deprecated DisplayTime first
+		for (auto& BranchingPoint : BranchingPoints_DEPRECATED)
 		{
-			return A.GetTriggerTime() < B.GetTriggerTime();
+			if (BranchingPoint.DisplayTime_DEPRECATED != 0.0f)
+			{
+				BranchingPoint.Clear();
+				BranchingPoint.LinkMontage(this, BranchingPoint.DisplayTime_DEPRECATED);
+			}
+			else
+			{
+				BranchingPoint.LinkMontage(this, BranchingPoint.GetTime());
+			}
 		}
-	};
 
-	BranchingPoints.Sort( FCompareFBranchingPointTime() );
+		// Then convert to AnimNotifies
+		USkeleton * MySkeleton = GetSkeleton();
+
+		// Add a new AnimNotifyTrack, and place all branching points in there.
+		int32 TrackIndex = AnimNotifyTracks.Num();
+
+		FAnimNotifyTrack NewItem;
+		NewItem.TrackName = *FString::FromInt(TrackIndex + 1);
+		NewItem.TrackColor = FLinearColor::White;
+		AnimNotifyTracks.Add(NewItem);
+
+		for (auto BranchingPoint : BranchingPoints_DEPRECATED)
+		{
+			int32 NewNotifyIndex = Notifies.Add(FAnimNotifyEvent());
+			FAnimNotifyEvent& NewEvent = Notifies[NewNotifyIndex];
+			NewEvent.NotifyName = BranchingPoint.EventName;
+
+			float TriggerTime = BranchingPoint.GetTriggerTime();
+			NewEvent.LinkMontage(this, TriggerTime);
+			NewEvent.TriggerTimeOffset = GetTriggerTimeOffsetForType(CalculateOffsetForNotify(TriggerTime));
+
+#if WITH_EDITORONLY_DATA
+			NewEvent.TrackIndex = TrackIndex;
+#endif
+			NewEvent.Notify = NULL;
+			NewEvent.NotifyStateClass = NULL;
+			NewEvent.bConvertedFromBranchingPoint = true;
+			NewEvent.MontageTickType = EMontageNotifyTickType::BranchingPoint;
+
+			// Add as a custom AnimNotify event to Skeleton.
+			if (MySkeleton)
+			{
+				MySkeleton->AnimationNotifies.AddUnique(NewEvent.NotifyName);
+			}
+		}
+
+		BranchingPoints_DEPRECATED.Empty();
+		RefreshBranchingPointMarkers();
+	}
+}
+
+void UAnimMontage::RefreshBranchingPointMarkers()
+{
+	BranchingPointMarkers.Empty();
+	BranchingPointStateNotifyIndices.Empty();
+
+	// Verify that we have no overlapping trigger times, this is not supported, and markers would not be triggered then.
+	TMap<float, FAnimNotifyEvent*> TriggerTimes;
+
+	int32 NumNotifies = Notifies.Num();
+	for (int32 NotifyIndex = 0; NotifyIndex < NumNotifies; NotifyIndex++)
+	{
+		FAnimNotifyEvent& NotifyEvent = Notifies[NotifyIndex];
+
+		if (NotifyEvent.MontageTickType == EMontageNotifyTickType::BranchingPoint)
+		{
+			AddBranchingPointMarker(FBranchingPointMarker(NotifyIndex, NotifyEvent.GetTriggerTime(), EAnimNotifyEventType::Begin), TriggerTimes);
+
+			if (NotifyEvent.NotifyStateClass)
+			{
+				// Track end point of AnimNotifyStates.
+				AddBranchingPointMarker(FBranchingPointMarker(NotifyIndex, NotifyEvent.GetEndTriggerTime(), EAnimNotifyEventType::End), TriggerTimes);
 
+				// Also track AnimNotifyStates separately, so we can tick them between their Begin and End points.
+				BranchingPointStateNotifyIndices.Add(NotifyIndex);
+			}
+		}
+	}
+	
+	if (BranchingPointMarkers.Num() > 0)
+	{
+		// Sort markers
+		struct FCompareNotifyTickMarkersTime
+		{
+			FORCEINLINE bool operator()(const FBranchingPointMarker &A, const FBranchingPointMarker &B) const
+			{
+				return A.TriggerTime < B.TriggerTime;
+			}
+		};
+
+		BranchingPointMarkers.Sort(FCompareNotifyTickMarkersTime());
+	}
 	bAnimBranchingPointNeedsSort = false;
 }
 
-#if WITH_EDITOR
-void UAnimMontage::PostEditChangeProperty(FPropertyChangedEvent& PropertyChangedEvent)
+void UAnimMontage::AddBranchingPointMarker(FBranchingPointMarker TickMarker, TMap<float, FAnimNotifyEvent*>& TriggerTimes)
 {
-	Super::PostEditChangeProperty(PropertyChangedEvent);
+	// Add Marker
+	BranchingPointMarkers.Add(TickMarker);
 
-	const FName PropertyName = PropertyChangedEvent.Property ? PropertyChangedEvent.Property->GetFName() : NAME_None;
-	if( PropertyName == FName(TEXT("BranchingPoints")) || PropertyName == FName(TEXT("Time")) )
+	// Check that there is no overlapping marker, as we don't support this.
+	// This would mean one of them is not getting triggered!
+	FAnimNotifyEvent** FoundNotifyEventPtr = TriggerTimes.Find(TickMarker.TriggerTime);
+	if (FoundNotifyEventPtr)
 	{
-		bAnimBranchingPointNeedsSort = true;
+		UE_LOG(LogAnimation, Warning, TEXT("Branching Point '%s' overlaps with '%s' at time: %f. One of them will not get triggered!"),
+			*Notifies[TickMarker.NotifyIndex].NotifyName.ToString(), *(*FoundNotifyEventPtr)->NotifyName.ToString(), TickMarker.TriggerTime);
+	}
+	else
+	{
+		TriggerTimes.Add(TickMarker.TriggerTime, &Notifies[TickMarker.NotifyIndex]);
 	}
 }
-#endif // WITH_EDITOR
 
-const FBranchingPoint* UAnimMontage::FindFirstBranchingPoint(float StartTrackPos, float EndTrackPos)
+const FBranchingPointMarker* UAnimMontage::FindFirstBranchingPointMarker(float StartTrackPos, float EndTrackPos)
 {
-	if( bAnimBranchingPointNeedsSort )
+	// Auto refresh this in editor to catch changes being made to AnimNotifies.
+	// PostEditChangeProperty does not get triggered for some reason when editing notifies.
+	// we need to come up with something better.
+	if (bAnimBranchingPointNeedsSort || (GIsEditor && (!GetWorld() || !GetWorld()->IsPlayInEditor())))
 	{
-		SortAnimBranchingPointByTime();
+		RefreshBranchingPointMarkers();
 	}
 
-	const bool bPlayingBackwards = (EndTrackPos < StartTrackPos);
-	if( !bPlayingBackwards )
+	if (BranchingPointMarkers.Num() > 0)
 	{
-		for( int32 Index=0; Index<BranchingPoints.Num(); Index++)
+		const bool bSearchBackwards = (EndTrackPos < StartTrackPos);
+		if (!bSearchBackwards)
 		{
-			const FBranchingPoint& BranchingPoint = BranchingPoints[Index];
-			if( BranchingPoint.GetTriggerTime() <= StartTrackPos )
+			for (int32 Index = 0; Index < BranchingPointMarkers.Num(); Index++)
 			{
-				continue;
+				const FBranchingPointMarker& Marker = BranchingPointMarkers[Index];
+				if (Marker.TriggerTime <= StartTrackPos)
+				{
+					continue;
+				}
+				if (Marker.TriggerTime > EndTrackPos)
+				{
+					break;
+				}
+				return &Marker;
 			}
-			if( BranchingPoint.GetTriggerTime() > EndTrackPos )
+		}
+		else
+		{
+			for (int32 Index = BranchingPointMarkers.Num() - 1; Index >= 0; Index--)
 			{
-				break;
+				const FBranchingPointMarker& Marker = BranchingPointMarkers[Index];
+				if (Marker.TriggerTime >= StartTrackPos)
+				{
+					continue;
+				}
+				if (Marker.TriggerTime < EndTrackPos)
+				{
+					break;
+				}
+				return &Marker;
 			}
-			return &BranchingPoint;
 		}
 	}
-	else
+	return NULL;
+}
+
+void UAnimMontage::FilterOutNotifyBranchingPoints(TArray<const FAnimNotifyEvent*>& InAnimNotifies)
+{
+	for (int32 Index = InAnimNotifies.Num()-1; Index >= 0; Index--)
 	{
-		for( int32 Index=BranchingPoints.Num()-1; Index>=0; Index--)
+		if (InAnimNotifies[Index]->MontageTickType == EMontageNotifyTickType::BranchingPoint)
 		{
-			const FBranchingPoint& BranchingPoint = BranchingPoints[Index];
-			if( BranchingPoint.GetTriggerTime() >= StartTrackPos )
-			{
-				continue;
-			}
-			if( BranchingPoint.GetTriggerTime() < EndTrackPos )
-			{
-				break;
-			}
-			return &BranchingPoint;
+			InAnimNotifies.RemoveAt(Index, 1);
 		}
 	}
-	return NULL;
 }
 
+#if WITH_EDITOR
+void UAnimMontage::PostEditChangeProperty(FPropertyChangedEvent& PropertyChangedEvent)
+{
+	Super::PostEditChangeProperty(PropertyChangedEvent);
+
+	// fixme laurent - this doesn't get triggered at all when editing notifies :(
+	const FName PropertyName = PropertyChangedEvent.Property ? PropertyChangedEvent.Property->GetFName() : NAME_None;
+	if( PropertyName == FName(TEXT("Notifies")))
+	{
+		bAnimBranchingPointNeedsSort = true;
+	}
+}
+#endif // WITH_EDITOR
+
 bool UAnimMontage::IsValidAdditive() const
 {
 	// if first one is additive, this is additive
@@ -598,33 +704,12 @@ EAnimEventTriggerOffsets::Type UAnimMontage::CalculateOffsetFromSections(float T
 }
 
 #if WITH_EDITOR
-EAnimEventTriggerOffsets::Type UAnimMontage::CalculateOffsetForBranchingPoint(float BranchingPointDisplayTime) const
-{
-	EAnimEventTriggerOffsets::Type Offset = Super::CalculateOffsetForNotify(BranchingPointDisplayTime);
-	if(Offset == EAnimEventTriggerOffsets::NoOffset)
-	{
-		Offset = CalculateOffsetFromSections(BranchingPointDisplayTime);
-	}
-	return Offset;
-}
-
 EAnimEventTriggerOffsets::Type UAnimMontage::CalculateOffsetForNotify(float NotifyDisplayTime) const
 {
 	EAnimEventTriggerOffsets::Type Offset = Super::CalculateOffsetForNotify(NotifyDisplayTime);
 	if(Offset == EAnimEventTriggerOffsets::NoOffset)
 	{
 		Offset = CalculateOffsetFromSections(NotifyDisplayTime);
-		if(Offset == EAnimEventTriggerOffsets::NoOffset)
-		{
-			for(auto Iter = BranchingPoints.CreateConstIterator(); Iter; ++Iter)
-			{
-				float BranchTime = Iter->GetTime();
-				if(BranchTime == NotifyDisplayTime)
-				{
-					return EAnimEventTriggerOffsets::OffsetBefore;
-				}
-			}
-		}
 	}
 	return Offset;
 }
@@ -892,6 +977,14 @@ void FAnimMontageInstance::Terminate()
 	UAnimMontage* OldMontage = Montage;
 	Montage = NULL;
 
+	// End all active State BranchingPoints
+	for (int32 Index = ActiveStateBranchingPoints.Num() - 1; Index >= 0; Index--)
+	{
+		FAnimNotifyEvent& NotifyEvent = ActiveStateBranchingPoints[Index];
+		NotifyEvent.NotifyStateClass->NotifyEnd(AnimInstance->GetSkelMeshComponent(), Cast<UAnimSequenceBase>(NotifyEvent.NotifyStateClass->GetOuter()));
+	}
+	ActiveStateBranchingPoints.Empty();
+
 	// terminating, trigger end
 	if (UAnimInstance* Inst = AnimInstance.Get())
 	{
@@ -1128,7 +1221,6 @@ void FAnimMontageInstance::Advance(float DeltaTime, struct FRootMotionMovementPa
 			RefreshNextPrevSections();
 		}
 #endif
-
 		// if no weight, no reason to update, and if not playing, we don't need to advance
 		// this portion is to advance position
 		// If we just reached zero weight, still tick this frame to fire end of animation events.
@@ -1162,11 +1254,11 @@ void FAnimMontageInstance::Advance(float DeltaTime, struct FRootMotionMovementPa
 
 					// Also look for a branching point. If we have one, stop there first to handle it.
 					const float CurrentSectionEndPos = CurrentSection.GetTime() + CurrentSectionLength;
-					const FBranchingPoint* NextBranchingPoint = Montage->FindFirstBranchingPoint(Position, FMath::Clamp(Position + DesiredDeltaMove, CurrentSection.GetTime(), CurrentSectionEndPos));
-					if( NextBranchingPoint )
+					const FBranchingPointMarker* BranchingPointMarker = Montage->FindFirstBranchingPointMarker(Position, FMath::Clamp(Position + DesiredDeltaMove, CurrentSection.GetTime(), CurrentSectionEndPos));
+					if (BranchingPointMarker)
 					{
 						// get the first one to see if it's less than AnimEnd
-						StoppingPosInSection = NextBranchingPoint->GetTriggerTime() - CurrentSection.GetTime();
+						StoppingPosInSection = BranchingPointMarker->TriggerTime - CurrentSection.GetTime();
 					}
 
 					// Update position within current section.
@@ -1222,13 +1314,13 @@ void FAnimMontageInstance::Advance(float DeltaTime, struct FRootMotionMovementPa
 						const float PositionBeforeFiringEvents = Position;
 						if( !bInterrupted )
 						{
-							HandleEvents(PrevPosition, Position, NextBranchingPoint);
+							HandleEvents(PrevPosition, Position, BranchingPointMarker);
 						}
 
 						// if we reached end of section, and we were not processing a branching point, and no events has messed with out current position..
 						// .. Move to next section.
 						// (this also handles looping, the same as jumping to a different section).
-						if( (AdvanceType != ETAA_Default) && !NextBranchingPoint && (PositionBeforeFiringEvents == Position) )
+						if ((AdvanceType != ETAA_Default) && !BranchingPointMarker && (PositionBeforeFiringEvents == Position))
 						{
 							// Get recent NextSectionIndex in case it's been changed by previous events.
 							const int32 RecentNextSectionIndex = bPlayingForward ? NextSections[CurrentSectionIndex] : PrevSections[CurrentSectionIndex];
@@ -1272,63 +1364,148 @@ void FAnimMontageInstance::Advance(float DeltaTime, struct FRootMotionMovementPa
 	{
 		Montage->EvaluateCurveData(AnimInstance.Get(), Position, Weight);
 	}
-}
 
-void FAnimMontageInstance::TriggerEventHandler(FName EventName)
-{
-	if ( EventName!=NAME_None && AnimInstance.IsValid() )
+	if (!bInterrupted)
 	{
-		FString FuncName = FString::Printf(TEXT("MontageBranchingPoint_%s"), *EventName.ToString());
-		FName FuncFName = FName(*FuncName);
-
-		UFunction* Function = AnimInstance.Get()->FindFunction(FuncFName);
-		if( Function )
+		// Tick all active state branching points
+		for (int32 Index = 0; Index < ActiveStateBranchingPoints.Num(); Index++)
 		{
-			AnimInstance.Get()->ProcessEvent( Function, NULL );
+			FAnimNotifyEvent& NotifyEvent = ActiveStateBranchingPoints[Index];
+			NotifyEvent.NotifyStateClass->NotifyTick(AnimInstance->GetSkelMeshComponent(), Montage, DeltaTime);
 		}
 	}
 }
 
-void FAnimMontageInstance::HandleEvents(float PreviousTrackPos, float CurrentTrackPos, const FBranchingPoint* BranchingPoint)
+void FAnimMontageInstance::HandleEvents(float PreviousTrackPos, float CurrentTrackPos, const FBranchingPointMarker* BranchingPointMarker)
 {
 	// Skip notifies and branching points if montage has been interrupted.
-	if( bInterrupted )
+	if (bInterrupted)
 	{
 		return;
 	}
 
 	// now get active Notifies based on how it advanced
-	if ( AnimInstance.IsValid() )
+	if (AnimInstance.IsValid())
 	{
-		// first handle notifies
-		{
-			TArray<const FAnimNotifyEvent*> Notifies;
+		TArray<const FAnimNotifyEvent*> Notifies;
 
-			// We already break up AnimMontage update to handle looping, so we guarantee that PreviousPos and CurrentPos are contiguous.
-			Montage->GetAnimNotifiesFromDeltaPositions(PreviousTrackPos, CurrentTrackPos, Notifies);
+		// We already break up AnimMontage update to handle looping, so we guarantee that PreviousPos and CurrentPos are contiguous.
+		Montage->GetAnimNotifiesFromDeltaPositions(PreviousTrackPos, CurrentTrackPos, Notifies);
 
-			// now trigger notifies for all animations within montage
-			// we'll do this for all slots for now
-			for (auto SlotTrack = Montage->SlotAnimTracks.CreateIterator(); SlotTrack; ++SlotTrack)
+		// For Montage only, remove notifies marked as 'branching points'. They are not queued and are handled separately.
+		Montage->FilterOutNotifyBranchingPoints(Notifies);
+
+		// now trigger notifies for all animations within montage
+		// we'll do this for all slots for now
+		for (auto SlotTrack = Montage->SlotAnimTracks.CreateIterator(); SlotTrack; ++SlotTrack)
+		{
+			if (AnimInstance->IsActiveSlotNode(SlotTrack->SlotName))
 			{
-				if( AnimInstance->IsActiveSlotNode(SlotTrack->SlotName) )
+				for (auto AnimSegment = SlotTrack->AnimTrack.AnimSegments.CreateIterator(); AnimSegment; ++AnimSegment)
 				{
-					for (auto AnimSegment = SlotTrack->AnimTrack.AnimSegments.CreateIterator() ; AnimSegment; ++AnimSegment)
-					{
-						AnimSegment->GetAnimNotifiesFromTrackPositions(PreviousTrackPos, CurrentTrackPos, Notifies);
-					}
+					AnimSegment->GetAnimNotifiesFromTrackPositions(PreviousTrackPos, CurrentTrackPos, Notifies);
 				}
 			}
+		}
+
+		// Queue all these notifies.
+		AnimInstance->AddAnimNotifies(Notifies, Weight);
+	}
+
+	// Update active state branching points, before we handle the immediate tick marker.
+	// In case our position jumped on the timeline, we need to begin/end state branching points accordingly.
+	UpdateActiveStateBranchingPoints(CurrentTrackPos);
+
+	// Trigger ImmediateTickMarker event if we have one
+	if (BranchingPointMarker)
+	{
+		BranchingPointEventHandler(BranchingPointMarker);
+	}
+}
+
+void FAnimMontageInstance::UpdateActiveStateBranchingPoints(float CurrentTrackPosition)
+{
+	int32 NumStateBranchingPoints = Montage->BranchingPointStateNotifyIndices.Num();
+	if (NumStateBranchingPoints > 0)
+	{
+		// End no longer active events first. We want this to happen before we trigger NotifyBegin on newly active events.
+		for (int32 Index = ActiveStateBranchingPoints.Num() - 1; Index >= 0; Index--)
+		{
+			FAnimNotifyEvent& NotifyEvent = ActiveStateBranchingPoints[Index];
+
+			const float NotifyStartTime = NotifyEvent.GetTriggerTime();
+			const float NotifyEndTime = NotifyEvent.GetEndTriggerTime();
+			bool bNotifyIsActive = (CurrentTrackPosition > NotifyStartTime) && (CurrentTrackPosition <= NotifyEndTime);
+
+			if (!bNotifyIsActive)
+			{
+				NotifyEvent.NotifyStateClass->NotifyEnd(AnimInstance->GetSkelMeshComponent(), Cast<UAnimSequenceBase>(NotifyEvent.NotifyStateClass->GetOuter()));
+				ActiveStateBranchingPoints.RemoveAt(Index, 1);
+			}
+		}
 
-			// if Notifies 
-			AnimInstance->AddAnimNotifies(Notifies, Weight);	
+		// Then, begin newly active notifies
+		for (int32 Index = 0; Index < NumStateBranchingPoints; Index++)
+		{
+			const int32 NotifyIndex = Montage->BranchingPointStateNotifyIndices[Index];
+			FAnimNotifyEvent& NotifyEvent = Montage->Notifies[NotifyIndex];
+
+			const float NotifyStartTime = NotifyEvent.GetTriggerTime();
+			const float NotifyEndTime = NotifyEvent.GetEndTriggerTime();
+
+			bool bNotifyIsActive = (CurrentTrackPosition > NotifyStartTime) && (CurrentTrackPosition <= NotifyEndTime);
+			if (bNotifyIsActive && !ActiveStateBranchingPoints.Contains(NotifyEvent))
+			{
+				NotifyEvent.NotifyStateClass->NotifyBegin(AnimInstance->GetSkelMeshComponent(), Cast<UAnimSequenceBase>(NotifyEvent.NotifyStateClass->GetOuter()), NotifyEvent.GetDuration());
+				ActiveStateBranchingPoints.Add(NotifyEvent);
+			}
 		}
 	}
+}
 
-	// Trigger Branching Point event if we have one
-	if( BranchingPoint )
+void FAnimMontageInstance::BranchingPointEventHandler(const FBranchingPointMarker* BranchingPointMarker)
+{
+	if (AnimInstance.IsValid() && Montage && BranchingPointMarker)
 	{
-		TriggerEventHandler(BranchingPoint->EventName);
+		FAnimNotifyEvent* NotifyEvent = (BranchingPointMarker->NotifyIndex < Montage->Notifies.Num()) ? &Montage->Notifies[BranchingPointMarker->NotifyIndex] : NULL;
+		if (NotifyEvent)
+		{
+			// Handle backwards compatibility with older BranchingPoints.
+			if (NotifyEvent->bConvertedFromBranchingPoint && (NotifyEvent->NotifyName != NAME_None))
+			{
+				FString FuncName = FString::Printf(TEXT("MontageBranchingPoint_%s"), *NotifyEvent->NotifyName.ToString());
+				FName FuncFName = FName(*FuncName);
+
+				UFunction* Function = AnimInstance.Get()->FindFunction(FuncFName);
+				if (Function)
+				{
+					AnimInstance.Get()->ProcessEvent(Function, NULL);
+				}
+				// In case older BranchingPoint has been re-implemented as a new Custom Notify, this is if BranchingPoint function hasn't been found.
+				else
+				{
+					AnimInstance.Get()->TriggerSingleAnimNotify(NotifyEvent);
+				}
+			}
+			else if (NotifyEvent->NotifyStateClass != NULL)
+			{
+				if (BranchingPointMarker->NotifyEventType == EAnimNotifyEventType::Begin)
+				{
+					NotifyEvent->NotifyStateClass->NotifyBegin(AnimInstance->GetSkelMeshComponent(), Cast<UAnimSequenceBase>(NotifyEvent->NotifyStateClass->GetOuter()), NotifyEvent->GetDuration());
+					ActiveStateBranchingPoints.Add(*NotifyEvent);
+				}
+				else
+				{
+					NotifyEvent->NotifyStateClass->NotifyEnd(AnimInstance->GetSkelMeshComponent(), Cast<UAnimSequenceBase>(NotifyEvent->NotifyStateClass->GetOuter()));
+					ActiveStateBranchingPoints.RemoveSingleSwap(*NotifyEvent);
+				}
+			}
+			// Regular 'non state' anim notifies
+			else
+			{
+				AnimInstance.Get()->TriggerSingleAnimNotify(NotifyEvent);
+			}
+		}
 	}
 }
 
@@ -1374,14 +1551,6 @@ ENGINE_API void UAnimMontage::UpdateLinkableElements()
 		Section.Update();
 	}
 
-	for(FBranchingPoint& BranchPoint : BranchingPoints)
-	{
-		BranchPoint.Update();
-
-		// Refresh branching point trigger offsets
-		BranchPoint.RefreshTriggerOffset(CalculateOffsetForBranchingPoint(BranchPoint.GetTime()));
-	}
-
 	for(FAnimNotifyEvent& Notify : Notifies)
 	{
 		Notify.Update();
@@ -1405,17 +1574,6 @@ ENGINE_API void UAnimMontage::UpdateLinkableElements(int32 SlotIdx, int32 Segmen
 		}
 	}
 
-	for(FBranchingPoint& BranchPoint : BranchingPoints)
-	{
-		if(BranchPoint.GetSlotIndex() == SlotIdx && BranchPoint.GetSegmentIndex() == SegmentIdx)
-		{
-			BranchPoint.Update();
-
-			// Refresh branching point trigger offsets
-			BranchPoint.RefreshTriggerOffset(CalculateOffsetForBranchingPoint(BranchPoint.GetTime()));
-		}
-	}
-
 	for(FAnimNotifyEvent& Notify : Notifies)
 	{
 		if(Notify.GetSlotIndex() == SlotIdx && Notify.GetSegmentIndex() == SegmentIdx)
