- Remove Materials access but added GetMaterials function
- Also made sure StaticMeshComponent::GetNumMaterials does the same behavior as SkeletalMeshComponent

I tested it. The old BPs that are there will still use old variable, but they shouldn't be able to do that anymore and the new ones should work as expected. Let me know if you see any issue with this.

diff --git a/Engine/Source/Runtime/Engine/Classes/Components/MeshComponent.h b/Engine/Source/Runtime/Engine/Classes/Components/MeshComponent.h
--- a/Engine/Source/Runtime/Engine/Classes/Components/MeshComponent.h
+++ b/Engine/Source/Runtime/Engine/Classes/Components/MeshComponent.h
@@ -20,8 +20,11 @@ class ENGINE_API UMeshComponent : public UPrimitiveComponent
 	GENERATED_UCLASS_BODY()
 
 	/** Per-Component material overrides.  These must NOT be set directly or a race condition can occur between GC and the rendering thread. */
-	UPROPERTY(EditAnywhere, BlueprintReadOnly, Category=Rendering, Meta=(ToolTip="Material overrides."))
+	UPROPERTY(EditAnywhere, Category=Rendering, Meta=(ToolTip="Material overrides."))
 	TArray<class UMaterialInterface*> Materials;
+	
+	UFUNCTION(BlueprintCallable, Category="Components|Mesh")
+	virtual TArray<class UMaterialInterface*> GetMaterials() const;
 
 
 	// Begin UObject Interface
diff --git a/Engine/Source/Runtime/Engine/Classes/Components/SkinnedMeshComponent.h b/Engine/Source/Runtime/Engine/Classes/Components/SkinnedMeshComponent.h
--- a/Engine/Source/Runtime/Engine/Classes/Components/SkinnedMeshComponent.h
+++ b/Engine/Source/Runtime/Engine/Classes/Components/SkinnedMeshComponent.h
@@ -465,6 +465,10 @@ public:
 	virtual int32 GetNumMaterials() const override;
 	// End UPrimitiveComponent interface
 
+	// Begin MeshComponent interface
+	virtual TArray<class UMaterialInterface*> GetMaterials() const override;
+	// End MeshComponent interface
+
 	/**
 	 *	Sets the value of the bForceWireframe flag and reattaches the component as necessary.
 	 *
diff --git a/Engine/Source/Runtime/Engine/Classes/Components/StaticMeshComponent.h b/Engine/Source/Runtime/Engine/Classes/Components/StaticMeshComponent.h
--- a/Engine/Source/Runtime/Engine/Classes/Components/StaticMeshComponent.h
+++ b/Engine/Source/Runtime/Engine/Classes/Components/StaticMeshComponent.h
@@ -296,6 +296,9 @@ public:
 #endif
 	// End UPrimitiveComponent interface.
 
+	// Begin UMeshComponent interface
+	virtual TArray<class UMaterialInterface*> GetMaterials() const override;
+	// End UMeshComponent interface
 	/**
 	 *	Returns true if the component uses texture lightmaps
 	 *
diff --git a/Engine/Source/Runtime/Engine/Private/MeshComponent.cpp b/Engine/Source/Runtime/Engine/Private/MeshComponent.cpp
--- a/Engine/Source/Runtime/Engine/Private/MeshComponent.cpp
+++ b/Engine/Source/Runtime/Engine/Private/MeshComponent.cpp
@@ -128,3 +128,9 @@ void UMeshComponent::SetTextureForceResidentFlag( bool bForceMiplevelsToBeReside
 		}
 	}
 }
+
+
+TArray<class UMaterialInterface*> UMeshComponent::GetMaterials() const
+{
+	return Materials;
+}
diff --git a/Engine/Source/Runtime/Engine/Private/SkinnedMeshComponent.cpp b/Engine/Source/Runtime/Engine/Private/SkinnedMeshComponent.cpp
--- a/Engine/Source/Runtime/Engine/Private/SkinnedMeshComponent.cpp
+++ b/Engine/Source/Runtime/Engine/Private/SkinnedMeshComponent.cpp
@@ -314,6 +314,22 @@ void USkinnedMeshComponent::UpdateSlaveComponent()
 	MarkRenderDynamicDataDirty();
 }
 
+TArray<class UMaterialInterface*> USkinnedMeshComponent::GetMaterials() const
+{
+	TArray<class UMaterialInterface*> Materials = Super::GetMaterials();
+
+	// if no material is overriden, look for mesh material;
+	if(Materials.Num() == 0 && SkeletalMesh)
+	{
+		for (auto Material : SkeletalMesh->Materials)
+		{
+			Materials.Add(Material.MaterialInterface);
+		}
+	}
+
+	return Materials;
+}
+
 int32 USkinnedMeshComponent::GetNumMaterials() const
 {
 	if (Materials.Num() > 0)
diff --git a/Engine/Source/Runtime/Engine/Private/StaticMeshComponent.cpp b/Engine/Source/Runtime/Engine/Private/StaticMeshComponent.cpp
--- a/Engine/Source/Runtime/Engine/Private/StaticMeshComponent.cpp
+++ b/Engine/Source/Runtime/Engine/Private/StaticMeshComponent.cpp
@@ -1456,14 +1456,29 @@ bool UStaticMeshComponent::GetEstimatedLightAndShadowMapMemoryUsage(
 
 int32 UStaticMeshComponent::GetNumMaterials() const
 {
-	if(StaticMesh)
+	if(Materials.Num() > 0)
+	{
+		return Materials.Num();
+	}
+	else if(StaticMesh)
 	{
 		return StaticMesh->Materials.Num();
 	}
-	else
+
+	return 0;
+}
+
+TArray<class UMaterialInterface*> UStaticMeshComponent::GetMaterials() const
+{
+	TArray<class UMaterialInterface*> Materials = Super::GetMaterials();
+
+	// if no material is overriden, look for mesh material;
+	if(Materials.Num() == 0 && StaticMesh)
 	{
-		return 0;
+		Materials = StaticMesh->Materials;
 	}
+
+	return Materials;
 }
 
 
