Fix crash if spawning a decal to a scene component that isn't a primitive component

diff --git a/Engine/Source/Runtime/Engine/Private/GameplayStatics.cpp b/Engine/Source/Runtime/Engine/Private/GameplayStatics.cpp
--- a/Engine/Source/Runtime/Engine/Private/GameplayStatics.cpp
+++ b/Engine/Source/Runtime/Engine/Private/GameplayStatics.cpp
@@ -926,20 +926,20 @@ UDecalComponent* UGameplayStatics::SpawnDecalAtLocation(UObject* WorldContextObj
 
 UDecalComponent* UGameplayStatics::SpawnDecalAttached(class UMaterialInterface* DecalMaterial, FVector DecalSize, class USceneComponent* AttachToComponent, FName AttachPointName, FVector Location, FRotator Rotation, EAttachLocation::Type LocationType, float LifeSpan)
 {
-	UDecalComponent* DecalComp = NULL;
+	UDecalComponent* DecalComp = nullptr;
 
 	if (DecalMaterial)
 	{
-		if (AttachToComponent == NULL)
+		if (AttachToComponent == nullptr)
 		{
 			UE_LOG(LogScript, Warning, TEXT("UGameplayStatics::SpawnDecalAttached: NULL AttachComponent specified!"));
 		}
 		else
 		{
 			UPrimitiveComponent* AttachToPrimitive = Cast<UPrimitiveComponent>(AttachToComponent);
-			if (AttachToPrimitive == NULL || AttachToPrimitive->bReceivesDecals)
+			if (AttachToPrimitive == nullptr || AttachToPrimitive->bReceivesDecals)
 			{
-				const bool bOnBSPBrush = (Cast<AWorldSettings>(AttachToPrimitive->GetOwner()) != NULL);
+				const bool bOnBSPBrush = AttachToPrimitive && (Cast<AWorldSettings>(AttachToPrimitive->GetOwner()) != nullptr);
 				if (bOnBSPBrush)
 				{
 					// special case: don't attach to component when it's owned by invisible WorldSettings (decals on BSP brush)
