Add Low-level callback for applying custom physics forces during sub-stepping. PR 585. (UE-5217)

diff --git a/Engine/Source/Runtime/Engine/Classes/PhysicsEngine/BodyInstance.h b/Engine/Source/Runtime/Engine/Classes/PhysicsEngine/BodyInstance.h
--- a/Engine/Source/Runtime/Engine/Classes/PhysicsEngine/BodyInstance.h
+++ b/Engine/Source/Runtime/Engine/Classes/PhysicsEngine/BodyInstance.h
@@ -13,6 +13,13 @@ struct FCollisionShape;
 struct FConstraintInstance;
 class UPhysicsConstraintComponent;
 
+/** Delegate for applying custom physics forces upon the body. Can be passed to "AddCustomPhysics" so 
+  * custom forces and torques can be calculated induvidually for every physics substep.
+  * The function provides delta time for a physics step and pointer to body instance upon which forces must be added.
+  * 
+  * Do not expect this callback to be called from the main game thread! It may get called from a physics simulation thread. */
+DECLARE_DELEGATE_TwoParams(FCalculateCustomPhysics, float, FBodyInstance*);
+
 #if WITH_PHYSX
 namespace physx
 {
@@ -449,6 +456,8 @@ public:
 	void WakeInstance();
 	/** Force this body to sleep */
 	void PutInstanceToSleep();
+	/** Add custom forces and torques on the body. The callback will be called more than once, if substepping enabled, for every substep.  */
+	void AddCustomPhysics(FCalculateCustomPhysics& CalculateCustomPhysics);
 	/** Add a force to this body */
 	void AddForce(const FVector& Force, bool bAllowSubstepping = true);
 	/** Add a force at a particular world position to this body */
diff --git a/Engine/Source/Runtime/Engine/Private/PhysicsEngine/BodyInstance.cpp b/Engine/Source/Runtime/Engine/Private/PhysicsEngine/BodyInstance.cpp
--- a/Engine/Source/Runtime/Engine/Private/PhysicsEngine/BodyInstance.cpp
+++ b/Engine/Source/Runtime/Engine/Private/PhysicsEngine/BodyInstance.cpp
@@ -3057,6 +3057,28 @@ void FBodyInstance::SetMaxDepenetrationVelocity(float MaxVelocity)
 }
 
 
+void FBodyInstance::AddCustomPhysics(FCalculateCustomPhysics& CalculateCustomPhysics)
+{
+	
+#if WITH_PHYSX
+	PxRigidBody* PRigidBody = GetPxRigidBody();
+	if (IsRigidBodyNonKinematic(PRigidBody))
+	{
+		const PxScene* PScene = PRigidBody->getScene();
+		FPhysScene* PhysScene = FPhysxUserData::Get<FPhysScene>(PScene->userData);
+		PhysScene->AddCustomPhysics(this, CalculateCustomPhysics);
+	}
+#endif // WITH_PHYSX
+
+#if WITH_BOX2D
+	if (BodyInstancePtr)
+	{
+		// Since Box2D does not have any substepping, might as well apply custom forces now
+		CalculateCustomPhysics.ExecuteIfBound(0.0f, this);
+	}
+#endif
+}
+
 void FBodyInstance::AddForce(const FVector& Force, bool bAllowSubstepping)
 {
 #if WITH_PHYSX
diff --git a/Engine/Source/Runtime/Engine/Private/PhysicsEngine/PhysScene.cpp b/Engine/Source/Runtime/Engine/Private/PhysicsEngine/PhysScene.cpp
--- a/Engine/Source/Runtime/Engine/Private/PhysicsEngine/PhysScene.cpp
+++ b/Engine/Source/Runtime/Engine/Private/PhysicsEngine/PhysScene.cpp
@@ -229,6 +229,25 @@ void FPhysScene::SetKinematicTarget(FBodyInstance* BodyInstance, const FTransfor
 #endif
 }
 
+void FPhysScene::AddCustomPhysics(FBodyInstance* BodyInstance, FCalculateCustomPhysics& CalculateCustomPhysics)
+{
+#if WITH_PHYSX
+#if WITH_SUBSTEPPING
+	uint32 BodySceneType = SceneType(BodyInstance);
+	if (IsSubstepping(BodySceneType))
+	{
+		FPhysSubstepTask * PhysSubStepper = PhysSubSteppers[SceneType(BodyInstance)];
+		PhysSubStepper->AddCustomPhysics(BodyInstance, CalculateCustomPhysics);
+	}
+	else
+#endif
+	{
+		// Since physics frame is set up before "pre-physics" tick group is called, can just fetch delta time from there
+		CalculateCustomPhysics.ExecuteIfBound(this->DeltaSeconds, BodyInstance);
+	}
+#endif
+}
+
 void FPhysScene::AddForce(FBodyInstance* BodyInstance, const FVector& Force, bool bAllowSubstepping)
 {
 #if WITH_PHYSX
diff --git a/Engine/Source/Runtime/Engine/Private/PhysicsEngine/PhysSubstepTasks.cpp b/Engine/Source/Runtime/Engine/Private/PhysicsEngine/PhysSubstepTasks.cpp
--- a/Engine/Source/Runtime/Engine/Private/PhysicsEngine/PhysSubstepTasks.cpp
+++ b/Engine/Source/Runtime/Engine/Private/PhysicsEngine/PhysSubstepTasks.cpp
@@ -85,6 +85,23 @@ bool FPhysSubstepTask::GetKinematicTarget(const FBodyInstance* Body, FTransform&
 	return false;
 }
 
+void FPhysSubstepTask::AddCustomPhysics(FBodyInstance* Body, const FCalculateCustomPhysics& CalculateCustomPhysics)
+{
+#if WITH_PHYSX
+	check(Body);
+	PxRigidBody* PRigidBody = Body->GetPxRigidBody();
+	SCOPED_SCENE_READ_LOCK(PRigidBody->getScene());
+	//Limit custom physics to non kinematic actors
+	if (IsRigidBodyNonKinematic(PRigidBody))
+	{
+		FCustomTarget CustomTarget(CalculateCustomPhysics);
+
+		FPhysTarget & TargetState = PhysTargetBuffers[External].FindOrAdd(Body);
+		TargetState.CustomPhysics.Add(CustomTarget);
+	}
+#endif
+}
+
 void FPhysSubstepTask::AddForce(FBodyInstance* Body, const FVector& Force)
 {
 #if WITH_PHYSX
@@ -144,6 +161,19 @@ void FPhysSubstepTask::AddTorque(FBodyInstance* Body, const FVector& Torque)
 #endif
 }
 
+/** Applies custom physics - Assumes caller has obtained writer lock */
+void FPhysSubstepTask::ApplyCustomPhysics(const FPhysTarget& PhysTarget, FBodyInstance* BodyInstance, float DeltaTime)
+{
+#if WITH_PHYSX
+	for (int32 i = 0; i < PhysTarget.CustomPhysics.Num(); ++i)
+	{
+		const FCustomTarget& CustomTarget = PhysTarget.CustomPhysics[i];
+
+		CustomTarget.CalculateCustomPhysics->ExecuteIfBound(DeltaTime, BodyInstance);
+	}
+#endif
+}
+
 /** Applies forces - Assumes caller has obtained writer lock */
 void FPhysSubstepTask::ApplyForces(const FPhysTarget& PhysTarget, FBodyInstance* BodyInstance)
 {
@@ -211,7 +241,7 @@ void FPhysSubstepTask::InterpolateKinematicActor(const FPhysTarget& PhysTarget,
 #endif
 }
 
-void FPhysSubstepTask::SubstepInterpolation(float InAlpha)
+void FPhysSubstepTask::SubstepInterpolation(float InAlpha, float DeltaTime)
 {
 #if WITH_PHYSX
 #if WITH_APEX
@@ -239,6 +269,7 @@ void FPhysSubstepTask::SubstepInterpolation(float InAlpha)
 		//We should only be iterating over actors that belong to this scene
 		check(PRigidBody->getScene() == PScene);
 
+		ApplyCustomPhysics(PhysTarget, BodyInstance, DeltaTime);
 		ApplyForces(PhysTarget, BodyInstance);
 		ApplyTorques(PhysTarget, BodyInstance);
 		InterpolateKinematicActor(PhysTarget, BodyInstance, InAlpha);
@@ -335,7 +366,7 @@ void FPhysSubstepTask::SubstepSimulationStart()
 	}
 #endif
 
-	SubstepInterpolation(Interpolation);
+	SubstepInterpolation(Interpolation, DeltaTime);
 
 #if WITH_APEX
 	PAScene->simulate(DeltaTime, bLastSubstep, SubstepTask);
diff --git a/Engine/Source/Runtime/Engine/Private/PhysicsEngine/PhysSubstepTasks.h b/Engine/Source/Runtime/Engine/Private/PhysicsEngine/PhysSubstepTasks.h
--- a/Engine/Source/Runtime/Engine/Private/PhysicsEngine/PhysSubstepTasks.h
+++ b/Engine/Source/Runtime/Engine/Private/PhysicsEngine/PhysSubstepTasks.h
@@ -73,6 +73,13 @@ struct FTorqueTarget
 	FVector Torque;
 };
 
+struct FCustomTarget
+{
+	FCustomTarget(){ }
+	FCustomTarget(const FCalculateCustomPhysics& GivenCalculateCustomPhysics) : CalculateCustomPhysics(&GivenCalculateCustomPhysics){}
+	const FCalculateCustomPhysics* CalculateCustomPhysics;
+};
+
 /** Holds information on everything we need to fix for substepping of a single frame */
 struct FPhysTarget
 {
@@ -80,6 +87,7 @@ struct FPhysTarget
 
 	TArray<FForceTarget>  Forces;	//we can apply force at multiple places
 	TArray<FTorqueTarget> Torques;
+	TArray<FCustomTarget> CustomPhysics; //for calculating custom physics forces
 	FKinematicTarget KinematicTarget;
 
 	/** Tells us if the kinematic target has been set */
@@ -106,6 +114,7 @@ public:
 
 	void SetKinematicTarget(FBodyInstance* Body, const FTransform& TM);
 	bool GetKinematicTarget(const FBodyInstance* Body, FTransform& OutTM) const;
+	void AddCustomPhysics(FBodyInstance* Body, const FCalculateCustomPhysics& CalculateCustomPhysics);
 	void AddForce(FBodyInstance* Body, const FVector& Force);
 	void AddForceAtPosition(FBodyInstance* Body, const FVector& Force, const FVector& Position);
 	void AddTorque(FBodyInstance* Body, const FVector& Torque);
@@ -127,7 +136,8 @@ public:
 private:
 
 	/** Applies interpolation and forces on all needed actors*/
-	void SubstepInterpolation(float Scale);
+	void SubstepInterpolation(float Scale, float DeltaTime);
+	void ApplyCustomPhysics(const FPhysTarget& PhysTarget, FBodyInstance* BodyInstance, float DeltaTime);
 	void ApplyForces(const FPhysTarget& PhysTarget, FBodyInstance* BodyInstance);
 	void ApplyTorques(const FPhysTarget& PhysTarget, FBodyInstance* BodyInstance);
 	void InterpolateKinematicActor(const FPhysTarget& PhysTarget, FBodyInstance* BodyInstance, float Alpha);
diff --git a/Engine/Source/Runtime/Engine/Public/PhysicsPublic.h b/Engine/Source/Runtime/Engine/Public/PhysicsPublic.h
--- a/Engine/Source/Runtime/Engine/Public/PhysicsPublic.h
+++ b/Engine/Source/Runtime/Engine/Public/PhysicsPublic.h
@@ -396,6 +396,9 @@ public:
 	/** Lets the scene update anything related to this BodyInstance as it's now being terminated */
 	void TermBody(FBodyInstance* BodyInstance);
 
+	/** Add a custom callback for next step that will be called on every substep */
+	void AddCustomPhysics(FBodyInstance* BodyInstance, FCalculateCustomPhysics& CalculateCustomPhysics);
+
 	/** Adds a force to a body - We need to go through scene to support substepping */
 	void AddForce(FBodyInstance* BodyInstance, const FVector& Force, bool bAllowSubstepping);
 
