Fix for the cursor being hidden over the viewport when it dosn't have mouse capture.

diff --git a/Engine/Source/Runtime/Engine/Private/Slate/SceneViewport.cpp b/Engine/Source/Runtime/Engine/Private/Slate/SceneViewport.cpp
--- a/Engine/Source/Runtime/Engine/Private/Slate/SceneViewport.cpp
+++ b/Engine/Source/Runtime/Engine/Private/Slate/SceneViewport.cpp
@@ -317,22 +317,19 @@ bool FSceneViewport::IsForegroundWindow() const
 
 FCursorReply FSceneViewport::OnCursorQuery( const FGeometry& MyGeometry, const FPointerEvent& CursorEvent )
 {
-	if (bCursorHiddenDueToCapture)
+	if (this->HasMouseCapture())
 	{
-		return FCursorReply::Cursor(EMouseCursor::None);
-	}
-
-	EMouseCursor::Type MouseCursorToUse = EMouseCursor::Default;
-
-	// If the cursor should be hidden, use EMouseCursor::None,
-	// only when in the foreground, or we'll hide the mouse in the window/program above us.
-	if( ViewportClient && GetSizeXY() != FIntPoint::ZeroValue  )
-	{
-		MouseCursorToUse = ViewportClient->GetCursor( this, GetMouseX(), GetMouseY() );
+		if (bCursorHiddenDueToCapture)
+		{
+			return FCursorReply::Cursor(EMouseCursor::None);
+		}
+		if (ViewportClient && GetSizeXY() != FIntPoint::ZeroValue)
+		{
+			return FCursorReply::Cursor(ViewportClient->GetCursor(this, GetMouseX(), GetMouseY()));
+		}
 	}
 
-	// Use the default cursor if there is no viewport client or we dont have focus
-	return FCursorReply::Cursor(MouseCursorToUse);
+	return FCursorReply::Unhandled();
 }
 
 
