fixed crash if no SubsurfaceProfile was used and VisualizeSubsurfaceProfiles gets activated

diff --git a/Engine/Source/Runtime/Engine/Classes/Engine/SubsurfaceProfile.h b/Engine/Source/Runtime/Engine/Classes/Engine/SubsurfaceProfile.h
--- a/Engine/Source/Runtime/Engine/Classes/Engine/SubsurfaceProfile.h
+++ b/Engine/Source/Runtime/Engine/Classes/Engine/SubsurfaceProfile.h
@@ -113,6 +113,7 @@ public:
 	// @param InRendererModule first call needs to have it != 0
 	void UpdateProfile(int32 AllocationId, const FSubsurfaceProfileStruct Settings);
 
+	// @return can be 0 if there is no SubsurfaceProfile
 	const struct IPooledRenderTarget* GetTexture(FRHICommandListImmediate& RHICmdList);
 
 
diff --git a/Engine/Source/Runtime/Engine/Private/Rendering/SubsurfaceProfile.cpp b/Engine/Source/Runtime/Engine/Private/Rendering/SubsurfaceProfile.cpp
--- a/Engine/Source/Runtime/Engine/Private/Rendering/SubsurfaceProfile.cpp
+++ b/Engine/Source/Runtime/Engine/Private/Rendering/SubsurfaceProfile.cpp
@@ -103,8 +103,11 @@ void FSubsurfaceProfileTexture::UpdateProfile(int32 AllocationId, const FSubsurf
 
 const IPooledRenderTarget* FSubsurfaceProfileTexture::GetTexture(FRHICommandListImmediate& RHICmdList)
 {
-	// call SetRendererModule() is missing
-	check(RendererModule);
+	if(!RendererModule)
+	{
+		// call SetRendererModule() is missing, thiscan be if no SubsurfaceProfile was used yet but VisualizeSubsurface requests the texture
+		return 0;
+	}
 
 	if (!GSSProfiles)
 	{
diff --git a/Engine/Source/Runtime/Renderer/Private/PostProcess/PostProcessSubsurface.cpp b/Engine/Source/Runtime/Renderer/Private/PostProcess/PostProcessSubsurface.cpp
--- a/Engine/Source/Runtime/Renderer/Private/PostProcess/PostProcessSubsurface.cpp
+++ b/Engine/Source/Runtime/Renderer/Private/PostProcess/PostProcessSubsurface.cpp
@@ -88,7 +88,11 @@ public:
 		{
 			const IPooledRenderTarget* PooledRT = GetSubsufaceProfileTexture_RT(Context.RHICmdList);
 
-			check(PooledRT);
+			if(!PooledRT)
+			{
+				// no subsurface profile was used yet
+				PooledRT = GSystemTextures.BlackDummy;
+			}
 
 			const FSceneRenderTargetItem& Item = PooledRT->GetRenderTargetItem();
 
