Fix non-unity

diff --git a/Engine/Source/Runtime/Engine/Private/Animation/AnimMontage.cpp b/Engine/Source/Runtime/Engine/Private/Animation/AnimMontage.cpp
--- a/Engine/Source/Runtime/Engine/Private/Animation/AnimMontage.cpp
+++ b/Engine/Source/Runtime/Engine/Private/Animation/AnimMontage.cpp
@@ -6,6 +6,7 @@
 
 #include "EnginePrivate.h"
 #include "Animation/AnimMontage.h"
+#include "Animation/AnimNotifies/AnimNotifyState.h"
 #include "AnimationUtils.h"
 #include "AnimationRuntime.h"
 
