Cleanup network demo file format, start thinking about future options (checkpoints, streaming, etc)

diff --git a/Engine/Source/Runtime/Engine/Classes/Engine/DemoNetDriver.h b/Engine/Source/Runtime/Engine/Classes/Engine/DemoNetDriver.h
--- a/Engine/Source/Runtime/Engine/Classes/Engine/DemoNetDriver.h
+++ b/Engine/Source/Runtime/Engine/Classes/Engine/DemoNetDriver.h
@@ -30,6 +30,9 @@ class UDemoNetDriver : public UNetDriver
 	/** during playback, set to total number of frames recorded in the demo */
 	int32				PlaybackTotalFrames;
 
+	/** during playback, set to offset of where the stream ends (we don't want to continue reading into the meta data section) */
+	int32				EndOfStreamOffset;
+
 	/** True if we're in the middle of recording a frame */
 	bool				bIsRecordingDemoFrame;
 
diff --git a/Engine/Source/Runtime/Engine/Private/DemoNetDriver.cpp b/Engine/Source/Runtime/Engine/Private/DemoNetDriver.cpp
--- a/Engine/Source/Runtime/Engine/Private/DemoNetDriver.cpp
+++ b/Engine/Source/Runtime/Engine/Private/DemoNetDriver.cpp
@@ -23,7 +23,7 @@ static TAutoConsoleVariable<float> CVarDemoTimeDilation( TEXT( "demo.TimeDilatio
 
 static const int32 MAX_DEMO_READ_WRITE_BUFFER = 1024 * 2;
 
-#define DEMO_CHECKSUMS 0
+#define DEMO_CHECKSUMS 0		// When setting this to 1, this will invalidate all demos, you will need to re-record and playback
 
 /*-----------------------------------------------------------------------------
 	UDemoNetDriver.
@@ -43,6 +43,7 @@ bool UDemoNetDriver::InitBase( bool bInitAsClient, FNetworkNotify* InNotify, con
 		DemoFrameNum			= 0;
 		bIsRecordingDemoFrame	= false;
 		bDemoPlaybackDone		= false;
+		EndOfStreamOffset		= 0;
 
 		return true;
 	}
@@ -69,6 +70,41 @@ FString UDemoNetDriver::LowLevelGetNetworkNumber()
 	return FString( TEXT( "" ) );
 }
 
+#define NETWORK_DEMO_MAGIC			( 0x2CF5A13D )
+#define NETWORK_DEMO_VERSION		( 0 )
+
+struct FNetworkDemoHeader
+{
+	uint32	Magic;					// Magic to ensure we're opening the right file.
+	uint32	Version;				// Version number to detect version mismatches.
+	uint32	EngineNetVersion;		// Version of engine networking format
+	FString LevelName;				// Name of level loaded for demo
+	int32	NumFrames;				// Number of total frames in the demo
+	int32	MetaDataOffset;			// Offset into the file where the meta data is stored for extra information that wasn't known at start of demo
+	int32	NumStreamingLevels;		// Number of streaming levels
+
+	FNetworkDemoHeader() : 
+		Magic( NETWORK_DEMO_MAGIC ), 
+		Version( NETWORK_DEMO_VERSION ),
+		EngineNetVersion( GEngineNetVersion ),
+		NumFrames( 0 ),
+		MetaDataOffset( 0 ),
+		NumStreamingLevels( 0 )
+	{}
+
+	friend FArchive& operator << ( FArchive& Ar, FNetworkDemoHeader& Header )
+	{
+		Ar << Header.Magic;
+		Ar << Header.Version;
+		Ar << Header.LevelName;
+		Ar << Header.NumFrames;
+		Ar << Header.MetaDataOffset;
+		Ar << Header.NumStreamingLevels;
+
+		return Ar;
+	}
+};
+
 bool UDemoNetDriver::InitConnect( FNetworkNotify* InNotify, const FURL& ConnectURL, FString& Error )
 {
 	if ( GetWorld() == nullptr )
@@ -107,30 +143,44 @@ bool UDemoNetDriver::InitConnect( FNetworkNotify* InNotify, const FURL& ConnectU
 	ServerConnection = ConstructObject<UNetConnection>( UDemoNetConnection::StaticClass() );
 	ServerConnection->InitConnection( this, USOCK_Pending, ConnectURL, 1000000 );
 
-#if 1
-	// Create fake control channel
-	ServerConnection->CreateChannel( CHTYPE_Control, 1 );
-#endif
-
 	// use the same byte format regardless of platform so that the demos are cross platform
 	// DEMO_FIXME: This is messing up for some reason, investigate
 	//FileAr->SetByteSwapping( true );
 
-	int32 EngineVersion = 0;
-	(*FileAr) << EngineVersion;
-	(*FileAr) << PlaybackTotalFrames;
+	FNetworkDemoHeader DemoHeader;
 
-	UE_LOG( LogDemo, Log, TEXT( "Starting demo playback with demo. Filename: %s, Frames: %i, Version %i" ), *DemoFilename, PlaybackTotalFrames, EngineVersion );
+	(*FileAr) << DemoHeader;
 
-#if 1
-	// Bypass UDemoPendingNetLevel
-	FString LevelName;
-	(*FileAr) << LevelName;
+	// Check magic value
+	if ( DemoHeader.Magic != NETWORK_DEMO_MAGIC )
+	{
+		Error = FString( TEXT( "Demo file is corrupt" ) );
+		UE_LOG( LogDemo, Error, TEXT( "UDemoNetDriver::InitConnect: %s" ), *Error );
+		GameInstance->HandleDemoPlaybackFailure( EDemoPlayFailure::Corrupt, Error );
+		return false;
+	}
 
+	// Check version
+	if ( DemoHeader.Version != NETWORK_DEMO_VERSION )
+	{
+		Error = FString( TEXT( "Demo file version is incorrect" ) );
+		UE_LOG( LogDemo, Error, TEXT( "UDemoNetDriver::InitConnect: %s" ), *Error );
+		GameInstance->HandleDemoPlaybackFailure( EDemoPlayFailure::InvalidVersion, Error );
+		return false;
+	}
+
+	// Create fake control channel
+	ServerConnection->CreateChannel( CHTYPE_Control, 1 );
+
+	PlaybackTotalFrames = DemoHeader.NumFrames;
+
+	UE_LOG( LogDemo, Log, TEXT( "Starting demo playback with demo. Filename: %s, Frames: %i, Version %i" ), *DemoFilename, PlaybackTotalFrames, DemoHeader.Version );
+
+	// Bypass UDemoPendingNetLevel
 	FString LoadMapError;
 
 	FURL DemoURL;
-	DemoURL.Map = LevelName;
+	DemoURL.Map = DemoHeader.LevelName;
 
 	FWorldContext * WorldContext = GEngine->GetWorldContextFromWorld( GetWorld() );
 
@@ -162,13 +212,15 @@ bool UDemoNetDriver::InitConnect( FNetworkNotify* InNotify, const FURL& ConnectU
 	SetWorld( WorldContext->World() );
 	WorldContext->World()->DemoNetDriver = this;
 	WorldContext->PendingNetGame = NULL;
-#endif
 
-	int32 NumStreamingLevels = 0;
+	// Remember where we are
+	const int32 OldPos = FileAr->Tell();
 
-	(*FileAr) << NumStreamingLevels;
+	// Jump to meta data
+	FileAr->Seek( DemoHeader.MetaDataOffset );
 
-	for ( int32 i = 0; i < NumStreamingLevels; ++i )
+	// Read meta data
+	for ( int32 i = 0; i < DemoHeader.NumStreamingLevels; ++i )
 	{
 		ULevelStreamingKismet* StreamingLevel = static_cast<ULevelStreamingKismet*>(StaticConstructObject(ULevelStreamingKismet::StaticClass(), GetWorld(), NAME_None, RF_NoFlags, NULL ) );
 
@@ -193,6 +245,12 @@ bool UDemoNetDriver::InitConnect( FNetworkNotify* InNotify, const FURL& ConnectU
 		UE_LOG( LogDemo, Log, TEXT( "  Loading streamingLevel: %s, %s" ), *PackageName, *PackageNameToLoad );
 	}
 
+	// Jump back to start of stream
+	FileAr->Seek( OldPos );
+
+	// Remember where the meta data is, this is where we must stop reading the demo stream
+	EndOfStreamOffset = DemoHeader.MetaDataOffset;
+
 	DemoDeltaTime = 0;
 
 	return true;
@@ -233,66 +291,20 @@ bool UDemoNetDriver::InitListen( FNetworkNotify* InNotify, FURL& ListenURL, bool
 	//@note: swap on non console platforms as the console archives have byte swapping compiled out by default
 	//FileAr->SetByteSwapping(true);
 
-	// write engine version info
-	int32 EngineVersion = GEngineNetVersion;
-	(*FileAr) << EngineVersion;
-
-	// write placeholder for total frames - will be updated when the demo is stopped
-	PlaybackTotalFrames = 0;
-	(*FileAr) << PlaybackTotalFrames;
-
-#if 0
-	// Create the control channel.
-	Connection->CreateChannel( CHTYPE_Control, 1, 0 );
-
-	// Send initial message.
-	uint8 IsLittleEndian = uint8( PLATFORM_LITTLE_ENDIAN );
-	check( IsLittleEndian == !!IsLittleEndian ); // should only be one or zero
-	FNetControlMessage<NMT_Hello>::Send( Connection, IsLittleEndian, GEngineMinNetVersion, GEngineNetVersion, Cast<UGeneralProjectSettings>(UGeneralProjectSettings::StaticClass()->GetDefaultObject())->ProjectID );
-	Connection->FlushNet();
-
-	// WelcomePlayer will send the needed map name
-	World->WelcomePlayer( Connection );
-#else
-	// Bypass UDemoPendingNetLevel
-	FString LevelName = World->GetCurrentLevel()->GetOutermost()->GetName();
-	(*FileAr) << LevelName;
-#endif
-
-	// Save out any levels that are in the streamed level list
-	// This needs some work, but for now, to try and get games that use heavy streaming working
-	int32 NumStreamingLevels = 0;
-
-	for ( int32 i = 0; i < World->StreamingLevels.Num(); ++i )
-	{
-		if ( World->StreamingLevels[i] != NULL )
-		{
-			NumStreamingLevels++;
-		}
-	}
-
-	(*FileAr) << NumStreamingLevels;
-
-	for ( int32 i = 0; i < World->StreamingLevels.Num(); ++i )
-	{
-		if ( World->StreamingLevels[i] != NULL )
-		{
-			FString PackageName = World->StreamingLevels[i]->GetWorldAssetPackageName();
-			FString PackageNameToLoad = World->StreamingLevels[i]->PackageNameToLoad.ToString();
+	FNetworkDemoHeader DemoHeader;
 
-			UE_LOG( LogDemo, Log, TEXT( "  StreamingLevel: %s, %s" ), *PackageName, *PackageNameToLoad );
+	// NOTE - This must be the SAME string we write when we close the demo or it will throw the header off
+	DemoHeader.LevelName = World->GetCurrentLevel()->GetOutermost()->GetName();
 
-			(*FileAr) << PackageName;
-			(*FileAr) << PackageNameToLoad;
-			(*FileAr) << World->StreamingLevels[i]->LevelTransform;
-		}
-	}
+	// Write the initial header (a lot of the fields will be placeholder until we fill them in later)
+	(*FileAr) << DemoHeader;
 
 	// Spawn the demo recording spectator.
 	SpawnDemoRecSpectator( Connection );
 
-	DemoDeltaTime = 0;
-	LastRecordTime = 0;
+	DemoDeltaTime		= 0;
+	LastRecordTime		= 0;
+	PlaybackTotalFrames = 0;
 
 	return true;
 }
@@ -393,14 +405,59 @@ void UDemoNetDriver::StopDemo()
 
 	if ( !ServerConnection )
 	{
-		// write the total number of frames in the placeholder at the beginning of the file
+		// Finish writing the header and other information that goes at the end
 		if ( FileAr != NULL && World != NULL )
 		{
 			PlaybackTotalFrames = DemoFrameNum;
-			int32 OldPos = FileAr->Tell();
-			FileAr->Seek( sizeof( int32 ) );
-			(*FileAr) << PlaybackTotalFrames;
-			FileAr->Seek( OldPos );
+
+			// Get the number of streaming levels so we can update the header with the correct info
+			int32 NumStreamingLevels = 0;
+
+			for ( int32 i = 0; i < World->StreamingLevels.Num(); ++i )
+			{
+				if ( World->StreamingLevels[i] != NULL )
+				{
+					NumStreamingLevels++;
+				}
+			}
+
+			// Make a new header with updated info
+			FNetworkDemoHeader DemoHeader;
+
+			DemoHeader.LevelName			= World->GetCurrentLevel()->GetOutermost()->GetName();
+			DemoHeader.NumFrames			= PlaybackTotalFrames;
+			DemoHeader.MetaDataOffset		= FileAr->Tell();
+			DemoHeader.NumStreamingLevels	= NumStreamingLevels;
+
+			// Seek to beginning
+			FileAr->Seek( 0 );
+
+			// Re-write header with new info
+			(*FileAr) << DemoHeader;
+
+			// Restore file position to end of stream
+			FileAr->Seek( DemoHeader.MetaDataOffset );
+
+			//
+			// Write meta data
+			//
+
+			// Save out any levels that are in the streamed level list
+			// This needs some work, but for now, to try and get games that use heavy streaming working
+			for ( int32 i = 0; i < World->StreamingLevels.Num(); ++i )
+			{
+				if ( World->StreamingLevels[i] != NULL )
+				{
+					FString PackageName = World->StreamingLevels[i]->GetWorldAssetPackageName();
+					FString PackageNameToLoad = World->StreamingLevels[i]->PackageNameToLoad.ToString();
+
+					UE_LOG( LogDemo, Log, TEXT( "  StreamingLevel: %s, %s" ), *PackageName, *PackageNameToLoad );
+
+					(*FileAr) << PackageName;
+					(*FileAr) << PackageNameToLoad;
+					(*FileAr) << World->StreamingLevels[i]->LevelTransform;
+				}
+			}
 		}
 
 		// let GC cleanup the object
@@ -419,9 +476,6 @@ void UDemoNetDriver::StopDemo()
 		ServerConnection->Close();
 		ServerConnection->CleanUp(); // make sure DemoRecSpectator gets destroyed immediately
 		ServerConnection = NULL;
-		
-		//GEngine->SetClientTravel( World, TEXT( "?closed" ), TRAVEL_Absolute );
-		//SetWorld( NULL );
 	}
 
 	delete FileAr;
@@ -563,7 +617,7 @@ bool UDemoNetDriver::ReadDemoFrame()
 		return false;
 	}
 
-	if ( FileAr->AtEnd() )
+	if ( FileAr->AtEnd() || FileAr->Tell() >= EndOfStreamOffset )
 	{
 		bDemoPlaybackDone = true;
 
@@ -650,7 +704,14 @@ bool UDemoNetDriver::ReadDemoFrame()
 		if ( PacketBytes > sizeof( ReadBuffer ) )
 		{
 			UE_LOG( LogDemo, Error, TEXT( "UDemoNetDriver::ReadDemoFrame: PacketBytes > sizeof( ReadBuffer )" ) );
+
 			StopDemo();
+
+			if ( World != NULL && World->GetGameInstance() != NULL )
+			{
+				World->GetGameInstance()->HandleDemoPlaybackFailure( EDemoPlayFailure::Generic, FString( TEXT( "UDemoNetDriver::ReadDemoFrame: PacketBytes > sizeof( ReadBuffer )" ) ) );
+			}
+
 			return false;
 		}
 
