non unity build fixes.

diff --git a/Engine/Source/Runtime/Engine/Private/Animation/AnimMontage.cpp b/Engine/Source/Runtime/Engine/Private/Animation/AnimMontage.cpp
--- a/Engine/Source/Runtime/Engine/Private/Animation/AnimMontage.cpp
+++ b/Engine/Source/Runtime/Engine/Private/Animation/AnimMontage.cpp
@@ -6,9 +6,11 @@
 
 #include "EnginePrivate.h"
 #include "Animation/AnimMontage.h"
-#include "Animation/AnimNotifies/AnimNotifyState.h"
 #include "AnimationUtils.h"
 #include "AnimationRuntime.h"
+#include "Animation/AnimNotifies/AnimNotifyState.h"
+#include "Animation/AnimNotifies/AnimNotify.h"
+#include "Animation/AnimSequenceBase.h"
 
 ///////////////////////////////////////////////////////////////////////////
 //
