Fix for non editor builds.

diff --git a/Engine/Source/Runtime/Engine/Private/Animation/AnimMontage.cpp b/Engine/Source/Runtime/Engine/Private/Animation/AnimMontage.cpp
--- a/Engine/Source/Runtime/Engine/Private/Animation/AnimMontage.cpp
+++ b/Engine/Source/Runtime/Engine/Private/Animation/AnimMontage.cpp
@@ -499,6 +499,7 @@ void UAnimMontage::ConvertBranchingPointsToAnimNotifies()
 		// Then convert to AnimNotifies
 		USkeleton * MySkeleton = GetSkeleton();
 
+#if WITH_EDITORONLY_DATA
 		// Add a new AnimNotifyTrack, and place all branching points in there.
 		int32 TrackIndex = AnimNotifyTracks.Num();
 
@@ -506,6 +507,7 @@ void UAnimMontage::ConvertBranchingPointsToAnimNotifies()
 		NewItem.TrackName = *FString::FromInt(TrackIndex + 1);
 		NewItem.TrackColor = FLinearColor::White;
 		AnimNotifyTracks.Add(NewItem);
+#endif
 
 		for (auto BranchingPoint : BranchingPoints_DEPRECATED)
 		{
