More non editor build fixes.

diff --git a/Engine/Source/Runtime/Engine/Private/Animation/AnimMontage.cpp b/Engine/Source/Runtime/Engine/Private/Animation/AnimMontage.cpp
--- a/Engine/Source/Runtime/Engine/Private/Animation/AnimMontage.cpp
+++ b/Engine/Source/Runtime/Engine/Private/Animation/AnimMontage.cpp
@@ -517,8 +517,9 @@ void UAnimMontage::ConvertBranchingPointsToAnimNotifies()
 
 			float TriggerTime = BranchingPoint.GetTriggerTime();
 			NewEvent.LinkMontage(this, TriggerTime);
+#if WITH_EDITOR
 			NewEvent.TriggerTimeOffset = GetTriggerTimeOffsetForType(CalculateOffsetForNotify(TriggerTime));
-
+#endif
 #if WITH_EDITORONLY_DATA
 			NewEvent.TrackIndex = TrackIndex;
 #endif
@@ -527,11 +528,13 @@ void UAnimMontage::ConvertBranchingPointsToAnimNotifies()
 			NewEvent.bConvertedFromBranchingPoint = true;
 			NewEvent.MontageTickType = EMontageNotifyTickType::BranchingPoint;
 
+#if WITH_EDITORONLY_DATA
 			// Add as a custom AnimNotify event to Skeleton.
 			if (MySkeleton)
 			{
 				MySkeleton->AnimationNotifies.AddUnique(NewEvent.NotifyName);
 			}
+#endif
 		}
 
 		BranchingPoints_DEPRECATED.Empty();
