Real fix for cursor being hidden.

diff --git a/Engine/Source/Runtime/Engine/Private/GameViewportClient.cpp b/Engine/Source/Runtime/Engine/Private/GameViewportClient.cpp
--- a/Engine/Source/Runtime/Engine/Private/GameViewportClient.cpp
+++ b/Engine/Source/Runtime/Engine/Private/GameViewportClient.cpp
@@ -519,7 +519,7 @@ EMouseCursor::Type UGameViewportClient::GetCursor(FViewport* InViewport, int32 X
 
 #endif
 
-	if ( !InViewport->HasFocus() || (ViewportConsole && ViewportConsole->ConsoleActive() ) )
+	if (!InViewport->HasMouseCapture() || !InViewport->HasFocus() || (ViewportConsole && ViewportConsole->ConsoleActive()))
 	{
 		return EMouseCursor::Default;
 	}
diff --git a/Engine/Source/Runtime/Engine/Private/Slate/SceneViewport.cpp b/Engine/Source/Runtime/Engine/Private/Slate/SceneViewport.cpp
--- a/Engine/Source/Runtime/Engine/Private/Slate/SceneViewport.cpp
+++ b/Engine/Source/Runtime/Engine/Private/Slate/SceneViewport.cpp
@@ -317,22 +317,23 @@ bool FSceneViewport::IsForegroundWindow() const
 
 FCursorReply FSceneViewport::OnCursorQuery( const FGeometry& MyGeometry, const FPointerEvent& CursorEvent )
 {
-	if (this->HasMouseCapture())
+	if (bCursorHiddenDueToCapture)
 	{
-		if (bCursorHiddenDueToCapture)
-		{
-			return FCursorReply::Cursor(EMouseCursor::None);
-		}
-		if (ViewportClient && GetSizeXY() != FIntPoint::ZeroValue)
-		{
-			return FCursorReply::Cursor(ViewportClient->GetCursor(this, GetMouseX(), GetMouseY()));
-		}
+		return FCursorReply::Cursor(EMouseCursor::None);
 	}
 
-	return FCursorReply::Unhandled();
-}
+	EMouseCursor::Type MouseCursorToUse = EMouseCursor::Default;
 
+	// If the cursor should be hidden, use EMouseCursor::None,
+	// only when in the foreground, or we'll hide the mouse in the window/program above us.
+	if( ViewportClient && GetSizeXY() != FIntPoint::ZeroValue  )
+	{
+		MouseCursorToUse = ViewportClient->GetCursor( this, GetMouseX(), GetMouseY() );
+	}
 
+	// Use the default cursor if there is no viewport client or we dont have focus
+	return FCursorReply::Cursor(MouseCursorToUse);
+}
 
 FReply FSceneViewport::OnMouseButtonDown( const FGeometry& InGeometry, const FPointerEvent& InMouseEvent )
 {
@@ -763,7 +764,11 @@ FReply FSceneViewport::OnFocusReceived(const FFocusEvent& InFocusEvent)
 		{
 			if (IsForegroundWindow())
 			{
-				const bool bIsCursorForcedVisible = ViewportClient->GetCursor(this, GetMouseX(), GetMouseY()) != EMouseCursor::None;
+				bool bIsCursorForcedVisible = false;
+				if (ViewportClient->GetWorld() && ViewportClient->GetWorld()->GetFirstPlayerController())
+				{
+					bIsCursorForcedVisible = ViewportClient->GetWorld()->GetFirstPlayerController()->GetMouseCursor() != EMouseCursor::None;
+				}
 
 				const bool bPlayInEditorCapture = !bIsPlayInEditorViewport || InFocusEvent.GetCause() != EFocusCause::SetDirectly || bPlayInEditorGetsMouseControl;
 
