Add demo total time to head + add demo time to timeline in shootergame

diff --git a/Engine/Source/Runtime/Engine/Classes/Engine/DemoNetDriver.h b/Engine/Source/Runtime/Engine/Classes/Engine/DemoNetDriver.h
--- a/Engine/Source/Runtime/Engine/Classes/Engine/DemoNetDriver.h
+++ b/Engine/Source/Runtime/Engine/Classes/Engine/DemoNetDriver.h
@@ -18,7 +18,7 @@ class UDemoNetDriver : public UNetDriver
 	/** Handle to the archive that will read/write network packets */
 	FArchive*			FileAr;
 
-	/** @todo document */
+	/** Current record/playback frame number */
 	int32				DemoFrameNum;
 
 	/** Last time (in real seconds) that we recorded a frame */
@@ -27,8 +27,14 @@ class UDemoNetDriver : public UNetDriver
 	/** Time (in game seconds) that have elapsed between recorded frames */
 	float				DemoDeltaTime;
 
-	/** during playback, set to total number of frames recorded in the demo */
-	int32				PlaybackTotalFrames;
+	/** Total time of demo in seconds */
+	float				DemoTotalTime;
+
+	/** Current record/playback position in seconds */
+	float				DemoCurrentTime;
+
+	/** Total number of frames in the demo */
+	int32				DemoTotalFrames;
 
 	/** during playback, set to offset of where the stream ends (we don't want to continue reading into the meta data section) */
 	int32				EndOfStreamOffset;
@@ -74,6 +80,7 @@ class UDemoNetDriver : public UNetDriver
 	bool ReadDemoFrame();
 	void TickDemoPlayback( float DeltaSeconds );
 	void SpawnDemoRecSpectator( UNetConnection* Connection );
+	void ResetDemoState();
 
 	void StopDemo();
 };
diff --git a/Engine/Source/Runtime/Engine/Private/DemoNetDriver.cpp b/Engine/Source/Runtime/Engine/Private/DemoNetDriver.cpp
--- a/Engine/Source/Runtime/Engine/Private/DemoNetDriver.cpp
+++ b/Engine/Source/Runtime/Engine/Private/DemoNetDriver.cpp
@@ -40,11 +40,12 @@ bool UDemoNetDriver::InitBase( bool bInitAsClient, FNetworkNotify* InNotify, con
 	{
 		DemoFilename			= URL.Map;
 		Time					= 0;
-		DemoFrameNum			= 0;
 		bIsRecordingDemoFrame	= false;
 		bDemoPlaybackDone		= false;
 		EndOfStreamOffset		= 0;
 
+		ResetDemoState();
+
 		return true;
 	}
 
@@ -80,6 +81,7 @@ struct FNetworkDemoHeader
 	uint32	EngineNetVersion;		// Version of engine networking format
 	FString LevelName;				// Name of level loaded for demo
 	int32	NumFrames;				// Number of total frames in the demo
+	float	TotalTime;				// Number of total time in seconds in demo
 	int32	MetaDataOffset;			// Offset into the file where the meta data is stored for extra information that wasn't known at start of demo
 	int32	NumStreamingLevels;		// Number of streaming levels
 
@@ -88,6 +90,7 @@ struct FNetworkDemoHeader
 		Version( NETWORK_DEMO_VERSION ),
 		EngineNetVersion( GEngineNetVersion ),
 		NumFrames( 0 ),
+		TotalTime( 0 ),
 		MetaDataOffset( 0 ),
 		NumStreamingLevels( 0 )
 	{}
@@ -98,6 +101,7 @@ struct FNetworkDemoHeader
 		Ar << Header.Version;
 		Ar << Header.LevelName;
 		Ar << Header.NumFrames;
+		Ar << Header.TotalTime;
 		Ar << Header.MetaDataOffset;
 		Ar << Header.NumStreamingLevels;
 
@@ -105,6 +109,16 @@ struct FNetworkDemoHeader
 	}
 };
 
+void UDemoNetDriver::ResetDemoState()
+{
+	DemoFrameNum	= 0;
+	LastRecordTime	= 0;
+	DemoDeltaTime	= 0;
+	DemoTotalTime	= 0;
+	DemoCurrentTime	= 0;
+	DemoTotalFrames	= 0;
+}
+
 bool UDemoNetDriver::InitConnect( FNetworkNotify* InNotify, const FURL& ConnectURL, FString& Error )
 {
 	if ( GetWorld() == nullptr )
@@ -128,6 +142,8 @@ bool UDemoNetDriver::InitConnect( FNetworkNotify* InNotify, const FURL& ConnectU
 		return false;
 	}
 
+	ResetDemoState();
+
 	// open the pre-recorded demo file
 	FileAr = IFileManager::Get().CreateFileReader( *DemoFilename );
 
@@ -172,9 +188,10 @@ bool UDemoNetDriver::InitConnect( FNetworkNotify* InNotify, const FURL& ConnectU
 	// Create fake control channel
 	ServerConnection->CreateChannel( CHTYPE_Control, 1 );
 
-	PlaybackTotalFrames = DemoHeader.NumFrames;
+	DemoTotalFrames = DemoHeader.NumFrames;
+	DemoTotalTime	= DemoHeader.TotalTime;
 
-	UE_LOG( LogDemo, Log, TEXT( "Starting demo playback with demo. Filename: %s, Frames: %i, Version %i" ), *DemoFilename, PlaybackTotalFrames, DemoHeader.Version );
+	UE_LOG( LogDemo, Log, TEXT( "Starting demo playback with demo. Filename: %s, Frames: %i, Version %i" ), *DemoFilename, DemoTotalFrames, DemoHeader.Version );
 
 	// Bypass UDemoPendingNetLevel
 	FString LoadMapError;
@@ -251,8 +268,6 @@ bool UDemoNetDriver::InitConnect( FNetworkNotify* InNotify, const FURL& ConnectU
 	// Remember where the meta data is, this is where we must stop reading the demo stream
 	EndOfStreamOffset = DemoHeader.MetaDataOffset;
 
-	DemoDeltaTime = 0;
-
 	return true;
 }
 
@@ -302,10 +317,6 @@ bool UDemoNetDriver::InitListen( FNetworkNotify* InNotify, FURL& ListenURL, bool
 	// Spawn the demo recording spectator.
 	SpawnDemoRecSpectator( Connection );
 
-	DemoDeltaTime		= 0;
-	LastRecordTime		= 0;
-	PlaybackTotalFrames = 0;
-
 	return true;
 }
 
@@ -408,7 +419,8 @@ void UDemoNetDriver::StopDemo()
 		// Finish writing the header and other information that goes at the end
 		if ( FileAr != NULL && World != NULL )
 		{
-			PlaybackTotalFrames = DemoFrameNum;
+			DemoTotalFrames = DemoFrameNum;
+			DemoTotalTime	= DemoCurrentTime;
 
 			// Get the number of streaming levels so we can update the header with the correct info
 			int32 NumStreamingLevels = 0;
@@ -425,7 +437,8 @@ void UDemoNetDriver::StopDemo()
 			FNetworkDemoHeader DemoHeader;
 
 			DemoHeader.LevelName			= World->GetCurrentLevel()->GetOutermost()->GetName();
-			DemoHeader.NumFrames			= PlaybackTotalFrames;
+			DemoHeader.NumFrames			= DemoTotalFrames;
+			DemoHeader.TotalTime			= DemoTotalTime;
 			DemoHeader.MetaDataOffset		= FileAr->Tell();
 			DemoHeader.NumStreamingLevels	= NumStreamingLevels;
 
@@ -542,6 +555,7 @@ void UDemoNetDriver::TickDemoRecord( float DeltaSeconds )
 	}
 
 	DemoDeltaTime += DeltaSeconds;
+	DemoCurrentTime += DeltaSeconds;
 
 	const double CurrentSeconds = FPlatformTime::Seconds();
 
@@ -559,7 +573,7 @@ void UDemoNetDriver::TickDemoRecord( float DeltaSeconds )
 	DemoFrameNum++;
 	ReplicationFrame++;
 
-	// Save elapsed game time
+	// Save elapsed game time for this frame
 	*FileAr << DemoDeltaTime;
 
 #if DEMO_CHECKSUMS == 1
@@ -672,7 +686,7 @@ bool UDemoNetDriver::ReadDemoFrame()
 		return false;
 	}
 
-	if ( DemoDeltaTime - ServerDeltaTime < 0 )//&& ServerConnection->State != USOCK_Pending )
+	if ( DemoDeltaTime < ServerDeltaTime )//&& ServerConnection->State != USOCK_Pending )
 	{
 		// Not enough time has passed to read another frame
 		FileAr->Seek( OldFilePos );
@@ -765,6 +779,7 @@ void UDemoNetDriver::TickDemoPlayback( float DeltaSeconds )
 	}
 
 	DemoDeltaTime += DeltaSeconds;
+	DemoCurrentTime += DeltaSeconds;
 
 	while ( true )
 	{
