Animation Layer Editing : Fixed crash when scale is added

diff --git a/Engine/Source/Runtime/Engine/Private/Animation/AnimSequence.cpp b/Engine/Source/Runtime/Engine/Private/Animation/AnimSequence.cpp
--- a/Engine/Source/Runtime/Engine/Private/Animation/AnimSequence.cpp
+++ b/Engine/Source/Runtime/Engine/Private/Animation/AnimSequence.cpp
@@ -3115,6 +3115,18 @@ void UAnimSequence::BakeTrackCurvesToRawAnimation()
 			VerifyCurveNames(CurSkeleton, USkeleton::AnimTrackCurveMappingName, RawCurveData.TransformCurves);
 			NameMapping = CurSkeleton->SmartNames.GetContainer(USkeleton::AnimTrackCurveMappingName);
 		}
+		
+		// since now I'm about to modify Scale Keys. I should add all of them here at least one key. 
+		// if all turns out to be same, it will clear it up. 
+		for (auto & RawTrack: RawAnimationData)
+		{
+			if (RawTrack.ScaleKeys.Num() == 0)
+			{
+				// at least add one
+				static FVector ScaleConstantKey(1.f);
+				RawTrack.ScaleKeys.Add(ScaleConstantKey);
+			}
+		}
 
 		for(const auto& Curve : RawCurveData.TransformCurves)
 		{
@@ -3172,12 +3184,6 @@ void UAnimSequence::BakeTrackCurvesToRawAnimation()
 					FVector OneKey = RawTrack.ScaleKeys[0];
 					RawTrack.ScaleKeys.Init(OneKey, NumFrames);
 				}
-				// scale key can be 0
-				else if(RawTrack.ScaleKeys.Num() == 0)
-				{
-					FVector OneKey = FVector(1.f);
-					RawTrack.ScaleKeys.Init(OneKey, NumFrames);
-				}
 				else
 				{
 					ensure(RawTrack.ScaleKeys.Num() == NumFrames);
