Check SM4's reflection texture exists before attempting to read from it.
Fix possible memory leak.

Fix for UE-5456.

diff --git a/Engine/Source/Runtime/Engine/Private/ReflectionCaptureComponent.cpp b/Engine/Source/Runtime/Engine/Private/ReflectionCaptureComponent.cpp
--- a/Engine/Source/Runtime/Engine/Private/ReflectionCaptureComponent.cpp
+++ b/Engine/Source/Runtime/Engine/Private/ReflectionCaptureComponent.cpp
@@ -1266,7 +1266,10 @@ void UReflectionCaptureComponent::ReadbackFromGPUAndSaveDerivedData(UWorld* Worl
 
 		if (WorldToUpdate->FeatureLevel == ERHIFeatureLevel::SM4)
 		{
-			ReadbackFromSM4Cubemap(SM4FullHDRCubemapTexture, NewDerivedData);
+			if (SM4FullHDRCubemapTexture)
+			{
+				ReadbackFromSM4Cubemap(SM4FullHDRCubemapTexture, NewDerivedData);
+			}
 		}
 		else
 		{
@@ -1280,6 +1283,10 @@ void UReflectionCaptureComponent::ReadbackFromGPUAndSaveDerivedData(UWorld* Worl
 			// Update our copy in memory
 			UpdateDerivedData(NewDerivedData);
 		}
+		else
+		{
+			delete NewDerivedData;
+		}
 	}
 }
 
