Convert Engine, QAGame, and Sample Games to use TimerHandle

diff --git a/Engine/Source/Developer/GameplayDebugger/Classes/GameplayDebuggingControllerComponent.h b/Engine/Source/Developer/GameplayDebugger/Classes/GameplayDebuggingControllerComponent.h
--- a/Engine/Source/Developer/GameplayDebugger/Classes/GameplayDebuggingControllerComponent.h
+++ b/Engine/Source/Developer/GameplayDebugger/Classes/GameplayDebuggingControllerComponent.h
@@ -5,6 +5,7 @@
  */
 
 #pragma once
+#include "TimerManager.h"
 #include "GameplayDebuggingTypes.h"
 #include "GameplayDebugger.h"
 #include "GameplayDebuggingControllerComponent.generated.h"
@@ -41,6 +42,8 @@ class GAMEPLAYDEBUGGER_API UGameplayDebuggingControllerComponent : public UActor
 	/** periodic update of navmesh data */
 	void UpdateNavMeshTimer();
 
+	float GetUpdateNavMeshTimeRemaining() const;
+
 	FOnChangeEQSQuery OnNextEQSQuery;
 	FOnChangeEQSQuery OnPreviousEQSQuery;
 protected:
@@ -77,6 +80,9 @@ protected:
 
 	TWeakObjectPtr<APlayerController> PlayerOwner;
 
+	/** Handle for efficient management of UpdateNavMesh timer */
+	FTimerHandle TimerHandle_UpdateNavMeshTimer;
+
 	FInputChord ActivationKey;
 	const float KeyPressActivationTime;
 	float ControlKeyPressedTime;
diff --git a/Engine/Source/Developer/GameplayDebugger/Private/GameplayDebuggingControllerComponent.cpp b/Engine/Source/Developer/GameplayDebugger/Private/GameplayDebuggingControllerComponent.cpp
--- a/Engine/Source/Developer/GameplayDebugger/Private/GameplayDebuggingControllerComponent.cpp
+++ b/Engine/Source/Developer/GameplayDebugger/Private/GameplayDebuggingControllerComponent.cpp
@@ -275,7 +275,7 @@ void UGameplayDebuggingControllerComponent::ToggleAIDebugView_SetView0()
 
 			if (IsViewActive(EAIDebugDrawDataView::NavMesh))
 			{
-				GetWorld()->GetTimerManager().SetTimer(this, &UGameplayDebuggingControllerComponent::UpdateNavMeshTimer, 5.0f, true);
+				GetWorld()->GetTimerManager().SetTimer(TimerHandle_UpdateNavMeshTimer, this, &UGameplayDebuggingControllerComponent::UpdateNavMeshTimer, 5.0f, true);
 				UpdateNavMeshTimer();
 
 				GetDebuggingReplicator()->ServerReplicateMessage(Pawn, EDebugComponentMessage::ActivateDataView, EAIDebugDrawDataView::NavMesh);
@@ -283,7 +283,7 @@ void UGameplayDebuggingControllerComponent::ToggleAIDebugView_SetView0()
 			}
 			else
 			{
-				GetWorld()->GetTimerManager().ClearTimer(this, &UGameplayDebuggingControllerComponent::UpdateNavMeshTimer);
+				GetWorld()->GetTimerManager().ClearTimer(TimerHandle_UpdateNavMeshTimer);
 
 				GetDebuggingReplicator()->ServerReplicateMessage(Pawn, EDebugComponentMessage::DeactivateDataView, EAIDebugDrawDataView::NavMesh);
 				OwnerComp->ServerDiscardNavmeshData();
@@ -346,6 +346,11 @@ void UGameplayDebuggingControllerComponent::TickComponent(float DeltaTime, enum
 	}
 }
 
+float UGameplayDebuggingControllerComponent::GetUpdateNavMeshTimeRemaining() const
+{
+	return GetWorld()->GetTimerManager().GetTimerRemaining(TimerHandle_UpdateNavMeshTimer);
+}
+
 void UGameplayDebuggingControllerComponent::UpdateNavMeshTimer()
 {
 	const APawn* PlayerPawn = PlayerOwner.IsValid() ? PlayerOwner->GetPawn() : NULL;
diff --git a/Engine/Source/Developer/GameplayDebugger/Private/GameplayDebuggingHUDComponent.cpp b/Engine/Source/Developer/GameplayDebugger/Private/GameplayDebuggingHUDComponent.cpp
--- a/Engine/Source/Developer/GameplayDebugger/Private/GameplayDebuggingHUDComponent.cpp
+++ b/Engine/Source/Developer/GameplayDebugger/Private/GameplayDebuggingHUDComponent.cpp
@@ -634,19 +634,14 @@ void AGameplayDebuggingHUDComponent::DrawNavMeshSnapshot(APlayerController* PC,
 #if !(UE_BUILD_SHIPPING || UE_BUILD_TEST)
 	if (DebugComponent && DebugComponent->NavmeshRepData.Num())
 	{
-		float TimeLeft = 0.0f;
 		UGameplayDebuggingControllerComponent*  GDC = PC ? PC->FindComponentByClass<UGameplayDebuggingControllerComponent>() : NULL;
-		if (GDC)
-		{
-			TimeLeft = GetWorldTimerManager().GetTimerRemaining(GDC, &UGameplayDebuggingControllerComponent::UpdateNavMeshTimer);
-		}
-
 		FString NextUpdateDesc;
-		if (TimeLeft > 0.0f)
+		if (GDC)
 		{
+			const float TimeLeft = GDC->GetUpdateNavMeshTimeRemaining();
 			NextUpdateDesc = FString::Printf(TEXT(", next update: {yellow}%.1fs"), TimeLeft);
 		}
-	
+
 		PrintString(DefaultContext, FString::Printf(TEXT("\n\n{green}Showing NavMesh (%.1fkB)%s\n"),
 			DebugComponent->NavmeshRepData.Num() / 1024.0f, *NextUpdateDesc));
 	}
diff --git a/Engine/Source/Editor/PropertyEditor/Private/PropertyEditorToolkit.cpp b/Engine/Source/Editor/PropertyEditor/Private/PropertyEditorToolkit.cpp
--- a/Engine/Source/Editor/PropertyEditor/Private/PropertyEditorToolkit.cpp
+++ b/Engine/Source/Editor/PropertyEditor/Private/PropertyEditorToolkit.cpp
@@ -166,8 +166,7 @@ void FPropertyEditorToolkit::Initialize( const EToolkitMode::Type Mode, const TS
 	PropertyTree->SetObjectArray( AdjustedObjectsToEditWeak );
 
 	PinColor = FSlateColor( FLinearColor( 1, 1, 1, 0 ) );
-	TickPinColorDelegate.BindSP( this, &FPropertyEditorToolkit::TickPinColorAndOpacity );
-	GEditor->GetTimerManager()->SetTimer( TickPinColorDelegate, 0.1f, true );
+	GEditor->GetTimerManager()->SetTimer( TimerHandle_TickPinColor, FTimerDelegate::CreateSP(this, &FPropertyEditorToolkit::TickPinColorAndOpacity), 0.1f, true );
 }
 
 
@@ -398,7 +397,7 @@ bool FPropertyEditorToolkit::TableHasCustomColumns() const
 
 bool FPropertyEditorToolkit::CloseWindow()
 {
-	GEditor->GetTimerManager()->ClearTimer( TickPinColorDelegate );
+	GEditor->GetTimerManager()->ClearTimer( TimerHandle_TickPinColor );
 	return FAssetEditorToolkit::CloseWindow();
 }
 
diff --git a/Engine/Source/Editor/PropertyEditor/Private/PropertyEditorToolkit.h b/Engine/Source/Editor/PropertyEditor/Private/PropertyEditorToolkit.h
--- a/Engine/Source/Editor/PropertyEditor/Private/PropertyEditorToolkit.h
+++ b/Engine/Source/Editor/PropertyEditor/Private/PropertyEditorToolkit.h
@@ -90,7 +90,7 @@ private:
 
 	/** Animation sequence to pulse the pin image */
 	FCurveSequence PinSequence;
-	FTimerDelegate TickPinColorDelegate;
+	FTimerHandle TimerHandle_TickPinColor;
 	FSlateColor PinColor;
 	TArray< TWeakPtr<IPropertyTreeRow> > PinRows;
 
diff --git a/Engine/Source/Runtime/AIModule/Classes/Actions/PawnAction_Move.h b/Engine/Source/Runtime/AIModule/Classes/Actions/PawnAction_Move.h
--- a/Engine/Source/Runtime/AIModule/Classes/Actions/PawnAction_Move.h
+++ b/Engine/Source/Runtime/AIModule/Classes/Actions/PawnAction_Move.h
@@ -1,7 +1,7 @@
 // Copyright 1998-2014 Epic Games, Inc. All Rights Reserved.
 
 #pragma once
-
+#include "TimerManager.h"
 #include "PawnAction_Move.generated.h"
 
 class AAIController;
@@ -84,6 +84,12 @@ protected:
 
 	FNavigationPath::FPathObserverDelegate::FDelegate PathObserver;
 	
+	/** Handle for efficient management of DeferredPerformMoveAction timer */
+	FTimerHandle TimerHandle_DeferredPerformMoveAction;
+
+	/** Handle for efficient management of TryToRepath timer */
+	FTimerHandle TimerHandle_TryToRepath;
+
 	void ClearPath();
 	virtual bool Start() override;
 	virtual bool Pause(const UPawnAction* PausedBy) override;
diff --git a/Engine/Source/Runtime/AIModule/Classes/Perception/PawnSensingComponent.h b/Engine/Source/Runtime/AIModule/Classes/Perception/PawnSensingComponent.h
--- a/Engine/Source/Runtime/AIModule/Classes/Perception/PawnSensingComponent.h
+++ b/Engine/Source/Runtime/AIModule/Classes/Perception/PawnSensingComponent.h
@@ -1,8 +1,7 @@
 // Copyright 1998-2014 Epic Games, Inc. All Rights Reserved.
 
-
-
 #pragma once
+#include "TimerManager.h"
 #include "Components/ActorComponent.h"
 #include "PawnSensingComponent.generated.h"
 
@@ -133,10 +132,13 @@ protected:
 
 	/** See if there are interesting sounds and sights that we want to detect, and respond to them if so. */
 	virtual void SensePawn(APawn& Pawn);
-	
+
 	/** Update function called on timer intervals. */
 	virtual void OnTimer();
 
+	/** Handle for efficient management of OnTimer timer */
+	FTimerHandle TimerHandle_OnTimer;
+
 	/** Modify the timer to fire in TimeDelay seconds. A value <= 0 disables the timer. */
 	virtual void SetTimer(const float TimeDelay);
 
diff --git a/Engine/Source/Runtime/AIModule/Private/Actions/PawnAction_Move.cpp b/Engine/Source/Runtime/AIModule/Private/Actions/PawnAction_Move.cpp
--- a/Engine/Source/Runtime/AIModule/Private/Actions/PawnAction_Move.cpp
+++ b/Engine/Source/Runtime/AIModule/Private/Actions/PawnAction_Move.cpp
@@ -106,7 +106,7 @@ bool UPawnAction_Move::PerformMoveAction()
 	if (bUsePathfinding && MyController->ShouldPostponePathUpdates())
 	{
 		UE_VLOG(MyController, LogPawnAction, Log, TEXT("Can't path right now, waiting..."));
-		MyController->GetWorldTimerManager().SetTimer(this, &UPawnAction_Move::DeferredPerformMoveAction, 0.1f);
+		MyController->GetWorldTimerManager().SetTimer(TimerHandle_DeferredPerformMoveAction, this, &UPawnAction_Move::DeferredPerformMoveAction, 0.1f);
 		return true;
 	}
 
@@ -282,7 +282,7 @@ void UPawnAction_Move::TryToRepath()
 		}
 		else if (GetWorld())
 		{
-			GetWorld()->GetTimerManager().SetTimer(this, &UPawnAction_Move::TryToRepath, 0.25f);
+			GetWorld()->GetTimerManager().SetTimer(TimerHandle_TryToRepath, this, &UPawnAction_Move::TryToRepath, 0.25f);
 		}
 	}
 }
@@ -291,7 +291,7 @@ void UPawnAction_Move::ClearPendingRepath()
 {
 	if (GetWorld())
 	{
-		GetWorld()->GetTimerManager().ClearTimer(this, &UPawnAction_Move::TryToRepath);
+		GetWorld()->GetTimerManager().ClearTimer(TimerHandle_TryToRepath);
 	}
 }
 
diff --git a/Engine/Source/Runtime/AIModule/Private/Perception/PawnSensingComponent.cpp b/Engine/Source/Runtime/AIModule/Private/Perception/PawnSensingComponent.cpp
--- a/Engine/Source/Runtime/AIModule/Private/Perception/PawnSensingComponent.cpp
+++ b/Engine/Source/Runtime/AIModule/Private/Perception/PawnSensingComponent.cpp
@@ -76,7 +76,7 @@ void UPawnSensingComponent::SetTimer(const float TimeInterval)
 	AActor* const Owner = GetOwner();
 	if (IsValid(Owner) && GEngine->GetNetMode(GetWorld()) < NM_Client)
 	{
-		Owner->GetWorldTimerManager().SetTimer(this, &UPawnSensingComponent::OnTimer, TimeInterval, false);
+		Owner->GetWorldTimerManager().SetTimer(TimerHandle_OnTimer, this, &UPawnSensingComponent::OnTimer, TimeInterval, false);
 	}
 }
 
@@ -96,7 +96,7 @@ void UPawnSensingComponent::SetSensingInterval(const float NewSensingInterval)
 			}
 			else if (bEnableSensingUpdates)
 			{
-				float CurrentElapsed = Owner->GetWorldTimerManager().GetTimerElapsed(this, &UPawnSensingComponent::OnTimer);
+				float CurrentElapsed = Owner->GetWorldTimerManager().GetTimerElapsed(TimerHandle_OnTimer);
 				CurrentElapsed = FMath::Max(0.f, CurrentElapsed);
 
 				if (CurrentElapsed < SensingInterval)
diff --git a/Engine/Source/Runtime/Engine/Classes/AI/Navigation/AvoidanceManager.h b/Engine/Source/Runtime/Engine/Classes/AI/Navigation/AvoidanceManager.h
--- a/Engine/Source/Runtime/Engine/Classes/AI/Navigation/AvoidanceManager.h
+++ b/Engine/Source/Runtime/Engine/Classes/AI/Navigation/AvoidanceManager.h
@@ -1,6 +1,7 @@
 // Copyright 1998-2014 Epic Games, Inc. All Rights Reserved.
 
 #pragma once
+#include "TimerManager.h"
 #include "AI/RVOAvoidanceInterface.h"
 #include "AvoidanceManager.generated.h"
 
@@ -155,6 +156,9 @@ class ENGINE_API UAvoidanceManager : public UObject, public FSelfRegisteringExec
 
 private:
 
+	/** Handle for efficient management of RemoveOutdatedObjects timer */
+	FTimerHandle TimerHandle_RemoveOutdatedObjects;
+
 	/** Cleanup AvoidanceObjects, called by timer */
 	void RemoveOutdatedObjects();
 
diff --git a/Engine/Source/Runtime/Engine/Classes/AI/Navigation/NavLinkCustomComponent.h b/Engine/Source/Runtime/Engine/Classes/AI/Navigation/NavLinkCustomComponent.h
--- a/Engine/Source/Runtime/Engine/Classes/AI/Navigation/NavLinkCustomComponent.h
+++ b/Engine/Source/Runtime/Engine/Classes/AI/Navigation/NavLinkCustomComponent.h
@@ -194,6 +194,9 @@ protected:
 	/** delegate to call when link is reached */
 	FOnMoveReachedLink OnMoveReachedLink;
 
+	/** Handle for efficient management of BroadcastStateChange timer */
+	FTimerHandle TimerHandle_BroadcastStateChange;
+
 	/** notify nearby agents about link changing state */
 	void BroadcastStateChange();
 	
diff --git a/Engine/Source/Runtime/Engine/Classes/Components/DecalComponent.h b/Engine/Source/Runtime/Engine/Classes/Components/DecalComponent.h
--- a/Engine/Source/Runtime/Engine/Classes/Components/DecalComponent.h
+++ b/Engine/Source/Runtime/Engine/Classes/Components/DecalComponent.h
@@ -2,6 +2,7 @@
 
 
 #pragma once
+#include "TimerManager.h"
 #include "DecalComponent.generated.h"
 
 class FDeferredDecalProxy;
@@ -52,8 +53,14 @@ public:
 	 */
 	ENGINE_API void PushSelectionToProxy();
 
+protected:
+	/** Handle for efficient management of DestroyDecalComponent timer */
+	FTimerHandle TimerHandle_DestroyDecalComponent;
 
+public:
 	
+	void SetLifeSpan(const float LifeSpan);
+
 	/**
 	 * Retrieves the materials used in this component
 	 *
diff --git a/Engine/Source/Runtime/Engine/Classes/Engine/LevelBounds.h b/Engine/Source/Runtime/Engine/Classes/Engine/LevelBounds.h
--- a/Engine/Source/Runtime/Engine/Classes/Engine/LevelBounds.h
+++ b/Engine/Source/Runtime/Engine/Classes/Engine/LevelBounds.h
@@ -1,5 +1,6 @@
 // Copyright 1998-2014 Epic Games, Inc. All Rights Reserved.
 #pragma once
+#include "TimerManager.h"
 #include "LevelBounds.generated.h"
 
 /**
@@ -46,6 +47,10 @@ class ALevelBounds : public AActor
 	ENGINE_API void UpdateLevelBoundsImmediately();
 	
 private:
+	
+	/** Handle for efficient management of OnTimerTick timer */
+	FTimerHandle TimerHandle_OnTimerTick;
+
 	/** Updates this actor bounding box by summing all level actors bounding boxes  */
 	void UpdateLevelBounds();
 
diff --git a/Engine/Source/Runtime/Engine/Classes/GameFramework/Actor.h b/Engine/Source/Runtime/Engine/Classes/GameFramework/Actor.h
--- a/Engine/Source/Runtime/Engine/Classes/GameFramework/Actor.h
+++ b/Engine/Source/Runtime/Engine/Classes/GameFramework/Actor.h
@@ -8,6 +8,7 @@
 #include "Engine/EngineTypes.h"
 #include "InputCoreTypes.h"
 #include "RenderCommandFence.h"
+#include "TimerManager.h"
 
 struct FHitResult;
 class AActor;
@@ -368,6 +369,13 @@ protected:
 	UPROPERTY(EditAnywhere, BlueprintReadOnly, Category=Actor)
 	float InitialLifeSpan; 
 
+private:
+
+	/** Handle for efficient management of LifeSpanExpired timer */
+	FTimerHandle TimerHandle_LifeSpanExpired;
+
+protected:
+
 	/**
 	 * If false, the Blueprint ReceiveTick() event will be disabled on dedicated servers.
 	 * @see AllowReceiveTickEventOnDedicatedServer()
@@ -1501,7 +1509,7 @@ public:
 	 */
 	virtual void PostActorCreated();
 
-	/** Called when the lifespawn of an actor expires (if he has one). */
+	/** Called when the lifespan of an actor expires (if he has one). */
 	virtual void LifeSpanExpired();
 
 	// Always called immediately before properties are received from the remote.
diff --git a/Engine/Source/Runtime/Engine/Classes/GameFramework/GameMode.h b/Engine/Source/Runtime/Engine/Classes/GameFramework/GameMode.h
--- a/Engine/Source/Runtime/Engine/Classes/GameFramework/GameMode.h
+++ b/Engine/Source/Runtime/Engine/Classes/GameFramework/GameMode.h
@@ -1,6 +1,7 @@
 // Copyright 1998-2014 Epic Games, Inc. All Rights Reserved.
 
 #pragma once
+#include "TimerManager.h"
 #include "GameFramework/Info.h"
 #include "GameMode.generated.h"
 
@@ -593,6 +594,10 @@ public:
 	virtual void DefaultTimer();	
 
 protected:
+
+	/** Handle for efficient management of DefaultTimer timer */
+	FTimerHandle TimerHandle_DefaultTimer;
+
 	/**
 	 * Customize incoming player based on URL options
 	 *
diff --git a/Engine/Source/Runtime/Engine/Classes/GameFramework/GameNetworkManager.h b/Engine/Source/Runtime/Engine/Classes/GameFramework/GameNetworkManager.h
--- a/Engine/Source/Runtime/Engine/Classes/GameFramework/GameNetworkManager.h
+++ b/Engine/Source/Runtime/Engine/Classes/GameFramework/GameNetworkManager.h
@@ -1,6 +1,7 @@
 // Copyright 1998-2014 Epic Games, Inc. All Rights Reserved.
 
 #pragma once
+#include "TimerManager.h"
 #include "GameNetworkManager.generated.h"
 
 /** Describes which standby detection event occured so the game can take appropriate action */
@@ -152,6 +153,11 @@ class ENGINE_API AGameNetworkManager : public AInfo
 	/** If true, actor network relevancy is constrained by whether they are within their NetCullDistanceSquared from the client's view point. */
 	UPROPERTY(globalconfig)
 	bool	bUseDistanceBasedRelevancy;
+
+protected:
+
+	/** Handle for efficient management of UpdateNetSpeeds timer */
+	FTimerHandle TimerHandle_UpdateNetSpeedsTimer;	
 };
 
 
diff --git a/Engine/Source/Runtime/Engine/Classes/GameFramework/GameState.h b/Engine/Source/Runtime/Engine/Classes/GameFramework/GameState.h
--- a/Engine/Source/Runtime/Engine/Classes/GameFramework/GameState.h
+++ b/Engine/Source/Runtime/Engine/Classes/GameFramework/GameState.h
@@ -1,7 +1,7 @@
 // Copyright 1998-2014 Epic Games, Inc. All Rights Reserved.
 
 #pragma once
-
+#include "TimerManager.h"
 #include "GameState.generated.h"
 
 //=============================================================================
@@ -112,13 +112,17 @@ public:
 	/** Remove PlayerState from the PlayerArray. */
 	virtual void RemovePlayerState(class APlayerState* PlayerState);
 
-
 	/** Called periodically, overridden by subclasses */
 	virtual void DefaultTimer();
 
 	/** Should players show gore? */
 	virtual bool ShouldShowGore() const;
 
+protected:
+
+	/** Handle for efficient management of DefaultTimer timer */
+	FTimerHandle TimerHandle_DefaultTimer;
+
 private:
 	// Hidden functions that don't make sense to use on this class.
 	HIDE_ACTOR_TRANSFORM_FUNCTIONS();
diff --git a/Engine/Source/Runtime/Engine/Classes/GameFramework/PainCausingVolume.h b/Engine/Source/Runtime/Engine/Classes/GameFramework/PainCausingVolume.h
--- a/Engine/Source/Runtime/Engine/Classes/GameFramework/PainCausingVolume.h
+++ b/Engine/Source/Runtime/Engine/Classes/GameFramework/PainCausingVolume.h
@@ -6,6 +6,7 @@
 //=============================================================================
 
 #pragma once
+#include "TimerManager.h"
 #include "GameFramework/PhysicsVolume.h"
 #include "PainCausingVolume.generated.h"
 
@@ -68,6 +69,11 @@ class ENGINE_API APainCausingVolume : public APhysicsVolume
 
 	/** damage overlapping actors if pain causing. */
 	virtual void CausePainTo(class AActor* Other);
+
+protected:
+
+	/** Handle for efficient management of OnTimerTick timer */
+	FTimerHandle TimerHandle_PainTimer;
 };
 
 
diff --git a/Engine/Source/Runtime/Engine/Classes/GameFramework/PlayerController.h b/Engine/Source/Runtime/Engine/Classes/GameFramework/PlayerController.h
--- a/Engine/Source/Runtime/Engine/Classes/GameFramework/PlayerController.h
+++ b/Engine/Source/Runtime/Engine/Classes/GameFramework/PlayerController.h
@@ -1,7 +1,7 @@
 // Copyright 1998-2014 Epic Games, Inc. All Rights Reserved.
 
 #pragma once
-
+#include "TimerManager.h"
 #include "GameFramework/PlayerMuteList.h"
 #include "Camera/PlayerCameraManager.h"
 #include "Camera/CameraTypes.h"
@@ -1095,7 +1095,7 @@ private:
 	uint32 bInputEnabled:1;
 
 	/** Whether the PlayerController's voice chat is enabled. */
-	bool bSpeaking;
+	uint32 bSpeaking:1;
 
 protected:
 
@@ -1115,6 +1115,15 @@ protected:
 	UPROPERTY()
 	class UTouchInterface* CurrentTouchInterface;
 
+	/** Handle for efficient management of UnFreeze timer */
+	FTimerHandle TimerHandle_UnFreeze;
+
+private:
+	/** Handle for efficient management of DelayedPrepareMapChange timer */
+	FTimerHandle TimerHandle_DelayedPrepareMapChange;
+
+	/** Handle for efficient management of ClientCommitMapChange timer */
+	FTimerHandle TimerHandle_ClientCommitMapChange;
 
 public:
 	/** Adds an inputcomponent to the top of the input stack. */
diff --git a/Engine/Source/Runtime/Engine/Classes/Matinee/MatineeActor.h b/Engine/Source/Runtime/Engine/Classes/Matinee/MatineeActor.h
--- a/Engine/Source/Runtime/Engine/Classes/Matinee/MatineeActor.h
+++ b/Engine/Source/Runtime/Engine/Classes/Matinee/MatineeActor.h
@@ -5,7 +5,7 @@
 #pragma once
 
 #include "Engine/EngineBaseTypes.h"
-
+#include "TimerManager.h"
 #include "MatineeActor.generated.h"
 
 /** Signature of function to handle a matinee event track key */
@@ -267,6 +267,11 @@ public:
 	ENGINE_API void OnObjectsReplaced(const TMap<UObject*,UObject*>& ReplacementMap);
 #endif //WITH_EDITOR
 
+protected:
+
+	/** Handle for efficient management of CheckPriorityRefresh timer */
+	FTimerHandle TimerHandle_CheckPriorityRefresh;
+
 private:
 
 #if WITH_EDITORONLY_DATA
diff --git a/Engine/Source/Runtime/Engine/Private/AI/Navigation/AvoidanceManager.cpp b/Engine/Source/Runtime/Engine/Private/AI/Navigation/AvoidanceManager.cpp
--- a/Engine/Source/Runtime/Engine/Private/AI/Navigation/AvoidanceManager.cpp
+++ b/Engine/Source/Runtime/Engine/Private/AI/Navigation/AvoidanceManager.cpp
@@ -108,7 +108,7 @@ void UAvoidanceManager::RequestUpdateTimer()
 	if (!bRequestedUpdateTimer && MyWorld)
 	{
 		bRequestedUpdateTimer = true;
-		MyWorld->GetTimerManager().SetTimer(this, &UAvoidanceManager::RemoveOutdatedObjects, DefaultTimeToLive * 0.5f, false);
+		MyWorld->GetTimerManager().SetTimer(TimerHandle_RemoveOutdatedObjects, this, &UAvoidanceManager::RemoveOutdatedObjects, DefaultTimeToLive * 0.5f, false);
 	}
 }
 
diff --git a/Engine/Source/Runtime/Engine/Private/AI/Navigation/NavLinkCustomComponent.cpp b/Engine/Source/Runtime/Engine/Private/AI/Navigation/NavLinkCustomComponent.cpp
--- a/Engine/Source/Runtime/Engine/Private/AI/Navigation/NavLinkCustomComponent.cpp
+++ b/Engine/Source/Runtime/Engine/Private/AI/Navigation/NavLinkCustomComponent.cpp
@@ -210,7 +210,7 @@ void UNavLinkCustomComponent::SetEnabled(bool bNewEnabled)
 
 		if (GetWorld())
 		{
-			GetWorld()->GetTimerManager().ClearTimer(this, &UNavLinkCustomComponent::BroadcastStateChange);
+			GetWorld()->GetTimerManager().ClearTimer(TimerHandle_BroadcastStateChange);
 
 			if ((bLinkEnabled && bNotifyWhenEnabled) || (!bLinkEnabled && bNotifyWhenDisabled))
 			{
@@ -324,7 +324,7 @@ void UNavLinkCustomComponent::BroadcastStateChange()
 
 	if (BroadcastInterval > 0.0f)
 	{
-		GetWorld()->GetTimerManager().SetTimer(this, &UNavLinkCustomComponent::BroadcastStateChange, BroadcastInterval);
+		GetWorld()->GetTimerManager().SetTimer(TimerHandle_BroadcastStateChange, this, &UNavLinkCustomComponent::BroadcastStateChange, BroadcastInterval);
 	}
 }
 
diff --git a/Engine/Source/Runtime/Engine/Private/Actor.cpp b/Engine/Source/Runtime/Engine/Private/Actor.cpp
--- a/Engine/Source/Runtime/Engine/Private/Actor.cpp
+++ b/Engine/Source/Runtime/Engine/Private/Actor.cpp
@@ -3424,11 +3424,11 @@ void AActor::SetLifeSpan( float InLifespan )
 	{
 		if( InLifespan > 0.0f)
 		{
-			GetWorldTimerManager().SetTimer( this, &AActor::LifeSpanExpired, InLifespan );
+			GetWorldTimerManager().SetTimer( TimerHandle_LifeSpanExpired, this, &AActor::LifeSpanExpired, InLifespan );
 		}
 		else
 		{
-			GetWorldTimerManager().ClearTimer( this, &AActor::LifeSpanExpired );		
+			GetWorldTimerManager().ClearTimer( TimerHandle_LifeSpanExpired );		
 		}
 	}
 }
@@ -3436,7 +3436,7 @@ void AActor::SetLifeSpan( float InLifespan )
 float AActor::GetLifeSpan() const
 {
 	// Timer remaining returns -1.0f if there is no such timer - return this as ZERO
-	const float CurrentLifespan = GetWorldTimerManager().GetTimerRemaining( this, &AActor::LifeSpanExpired );
+	const float CurrentLifespan = GetWorldTimerManager().GetTimerRemaining(TimerHandle_LifeSpanExpired);
 	return ( CurrentLifespan != -1.0f ) ? CurrentLifespan : 0.0f;
 }
 
diff --git a/Engine/Source/Runtime/Engine/Private/DecalComponent.cpp b/Engine/Source/Runtime/Engine/Private/DecalComponent.cpp
--- a/Engine/Source/Runtime/Engine/Private/DecalComponent.cpp
+++ b/Engine/Source/Runtime/Engine/Private/DecalComponent.cpp
@@ -40,6 +40,18 @@ UDecalComponent::UDecalComponent(const FObjectInitializer& ObjectInitializer)
 {
 }
 
+void UDecalComponent::SetLifeSpan(const float LifeSpan)
+{
+	if (LifeSpan > 0.f)
+	{
+		GetWorld()->GetTimerManager().SetTimer(TimerHandle_DestroyDecalComponent, this, &UDecalComponent::DestroyComponent, LifeSpan, false);
+	}
+	else
+	{
+		GetWorld()->GetTimerManager().ClearTimer(TimerHandle_DestroyDecalComponent);
+	}
+}
+
 void UDecalComponent::SetSortOrder(int32 Value)
 {
 	SortOrder = Value;
diff --git a/Engine/Source/Runtime/Engine/Private/GameMode.cpp b/Engine/Source/Runtime/Engine/Private/GameMode.cpp
--- a/Engine/Source/Runtime/Engine/Private/GameMode.cpp
+++ b/Engine/Source/Runtime/Engine/Private/GameMode.cpp
@@ -395,7 +395,7 @@ void AGameMode::PreInitializeComponents()
 {
 	Super::PreInitializeComponents();
 
-	GetWorldTimerManager().SetTimer(this, &AGameMode::DefaultTimer, GetWorldSettings()->GetEffectiveTimeDilation(), true);
+	GetWorldTimerManager().SetTimer(TimerHandle_DefaultTimer, this, &AGameMode::DefaultTimer, GetWorldSettings()->GetEffectiveTimeDilation(), true);
 	FActorSpawnParameters SpawnInfo;
 	SpawnInfo.Instigator = Instigator;
 
diff --git a/Engine/Source/Runtime/Engine/Private/GameNetworkManager.cpp b/Engine/Source/Runtime/Engine/Private/GameNetworkManager.cpp
--- a/Engine/Source/Runtime/Engine/Private/GameNetworkManager.cpp
+++ b/Engine/Source/Runtime/Engine/Private/GameNetworkManager.cpp
@@ -73,7 +73,7 @@ void AGameNetworkManager::UpdateNetSpeeds(bool bIsLanMatch)
 
 	if ( GetWorld()->TimeSeconds - LastNetSpeedUpdateTime < 1.0f )
 	{
-		GetWorldTimerManager().SetTimer(this, &AGameNetworkManager::UpdateNetSpeedsTimer, 1.0f);
+		GetWorldTimerManager().SetTimer(TimerHandle_UpdateNetSpeedsTimer, this, &AGameNetworkManager::UpdateNetSpeedsTimer, 1.0f);
 		return;
 	}
 
diff --git a/Engine/Source/Runtime/Engine/Private/GameState.cpp b/Engine/Source/Runtime/Engine/Private/GameState.cpp
--- a/Engine/Source/Runtime/Engine/Private/GameState.cpp
+++ b/Engine/Source/Runtime/Engine/Private/GameState.cpp
@@ -48,14 +48,14 @@ void AGameState::DefaultTimer()
 		++ElapsedTime;
 	}
 
-	GetWorldTimerManager().SetTimer(this, &AGameState::DefaultTimer, GetWorldSettings()->GetEffectiveTimeDilation(), true);
+	GetWorldTimerManager().SetTimer(TimerHandle_DefaultTimer, this, &AGameState::DefaultTimer, GetWorldSettings()->GetEffectiveTimeDilation(), true);
 }
 
 void AGameState::PostInitializeComponents()
 {
 	Super::PostInitializeComponents();
 
-	GetWorldTimerManager().SetTimer(this, &AGameState::DefaultTimer, GetWorldSettings()->GetEffectiveTimeDilation(), true);
+	GetWorldTimerManager().SetTimer(TimerHandle_DefaultTimer, this, &AGameState::DefaultTimer, GetWorldSettings()->GetEffectiveTimeDilation(), true);
 
 	UWorld* World = GetWorld();
 	World->GameState = this;
diff --git a/Engine/Source/Runtime/Engine/Private/GameplayStatics.cpp b/Engine/Source/Runtime/Engine/Private/GameplayStatics.cpp
--- a/Engine/Source/Runtime/Engine/Private/GameplayStatics.cpp
+++ b/Engine/Source/Runtime/Engine/Private/GameplayStatics.cpp
@@ -894,7 +894,7 @@ void UGameplayStatics::DeactivateReverbEffect(FName TagName)
 
 UDecalComponent* CreateDecalComponent(class UMaterialInterface* DecalMaterial, FVector DecalSize, UWorld* World, AActor* Actor, float LifeSpan)
 {
-	const FMatrix DecalInternalTransform = FRotationMatrix(FRotator(0, 90.0f, -90.0f));
+	const FMatrix DecalInternalTransform = FRotationMatrix(FRotator(0.f, 90.0f, -90.0f));
 
 	UDecalComponent* DecalComp = ConstructObject<UDecalComponent>(UDecalComponent::StaticClass(), (Actor ? Actor : (UObject*)GetTransientPackage()));
 	DecalComp->DecalMaterial = DecalMaterial;
@@ -902,9 +902,9 @@ UDecalComponent* CreateDecalComponent(class UMaterialInterface* DecalMaterial, F
 	DecalComp->bAbsoluteScale = true;
 	DecalComp->RegisterComponentWithWorld(World);
 
-	if (LifeSpan > 0)
+	if (LifeSpan > 0.f)
 	{
-		World->GetTimerManager().SetTimer(DecalComp, &UDecalComponent::DestroyComponent, LifeSpan, false);
+		DecalComp->SetLifeSpan(LifeSpan);
 	}
 
 	return DecalComp;
diff --git a/Engine/Source/Runtime/Engine/Private/Interpolation.cpp b/Engine/Source/Runtime/Engine/Private/Interpolation.cpp
--- a/Engine/Source/Runtime/Engine/Private/Interpolation.cpp
+++ b/Engine/Source/Runtime/Engine/Private/Interpolation.cpp
@@ -1682,11 +1682,11 @@ void AMatineeActor::UpdateReplicatedData( bool bIsBeginningPlay )
 	
 	if (bIsPlaying || bIsBeginningPlay)
 	{
-		GetWorldTimerManager().SetTimer(this, &AMatineeActor::CheckPriorityRefresh, 1.0f, true);
+		GetWorldTimerManager().SetTimer(TimerHandle_CheckPriorityRefresh, this, &AMatineeActor::CheckPriorityRefresh, 1.0f, true);
 	}
 	else
 	{
-		GetWorldTimerManager().ClearTimer(this, &AMatineeActor::CheckPriorityRefresh);
+		GetWorldTimerManager().ClearTimer(TimerHandle_CheckPriorityRefresh);
 	}
 }
 
diff --git a/Engine/Source/Runtime/Engine/Private/LevelBounds.cpp b/Engine/Source/Runtime/Engine/Private/LevelBounds.cpp
--- a/Engine/Source/Runtime/Engine/Private/LevelBounds.cpp
+++ b/Engine/Source/Runtime/Engine/Private/LevelBounds.cpp
@@ -195,7 +195,7 @@ void ALevelBounds::SubscribeToUpdateEvents()
 {
 	if (bSubscribedToEvents == false && GIsEditor && !GetWorld()->IsPlayInEditor())
 	{
-		GetWorldTimerManager().SetTimer(this, &ALevelBounds::OnTimerTick, 1, true);
+		GetWorldTimerManager().SetTimer(TimerHandle_OnTimerTick, this, &ALevelBounds::OnTimerTick, 1.f, true);
 		GEngine->OnActorMoved().AddUObject(this, &ALevelBounds::OnLevelActorMoved);
 		GEngine->OnLevelActorDeleted().AddUObject(this, &ALevelBounds::OnLevelActorAddedRemoved);
 		GEngine->OnLevelActorAdded().AddUObject(this, &ALevelBounds::OnLevelActorAddedRemoved);
@@ -210,7 +210,7 @@ void ALevelBounds::UnsubscribeFromUpdateEvents()
 	{
 		if (GetWorld())
 		{
-			GetWorldTimerManager().ClearTimer(this, &ALevelBounds::OnTimerTick);
+			GetWorldTimerManager().ClearTimer(TimerHandle_OnTimerTick);
 		}
 		GEngine->OnActorMoved().RemoveUObject(this, &ALevelBounds::OnLevelActorMoved);
 		GEngine->OnLevelActorDeleted().RemoveUObject(this, &ALevelBounds::OnLevelActorAddedRemoved);
diff --git a/Engine/Source/Runtime/Engine/Private/PainCausingVolume.cpp b/Engine/Source/Runtime/Engine/Private/PainCausingVolume.cpp
--- a/Engine/Source/Runtime/Engine/Private/PainCausingVolume.cpp
+++ b/Engine/Source/Runtime/Engine/Private/PainCausingVolume.cpp
@@ -37,9 +37,9 @@ void APainCausingVolume::ActorEnteredVolume(AActor* Other)
 	}
 
 	// Start timer if none is active
-	if (!GetWorldTimerManager().IsTimerActive(this, &APainCausingVolume::PainTimer))
+	if (!GetWorldTimerManager().IsTimerActive(TimerHandle_PainTimer))
 	{
-		GetWorldTimerManager().SetTimer(this, &APainCausingVolume::PainTimer, PainInterval, true);
+		GetWorldTimerManager().SetTimer(TimerHandle_PainTimer, this, &APainCausingVolume::PainTimer, PainInterval, true);
 	}
 }
 
@@ -66,7 +66,7 @@ void APainCausingVolume::PainTimer()
 		// Stop timer if nothing is overlapping us
 		if (TouchingActors.Num() == 0)
 		{
-			GetWorldTimerManager().ClearTimer(this, &APainCausingVolume::PainTimer);
+			GetWorldTimerManager().ClearTimer(TimerHandle_PainTimer);
 		}
 	}
 }
diff --git a/Engine/Source/Runtime/Engine/Private/PlayerController.cpp b/Engine/Source/Runtime/Engine/Private/PlayerController.cpp
--- a/Engine/Source/Runtime/Engine/Private/PlayerController.cpp
+++ b/Engine/Source/Runtime/Engine/Private/PlayerController.cpp
@@ -1127,7 +1127,7 @@ void APlayerController::UnFreeze() {}
 
 bool APlayerController::IsFrozen()
 {
-	return GetWorldTimerManager().IsTimerActive(this, &APlayerController::UnFreeze);
+	return GetWorldTimerManager().IsTimerActive(TimerHandle_UnFreeze);
 }
 
 void APlayerController::ServerAcknowledgePossession_Implementation(APawn* P)
@@ -2895,7 +2895,7 @@ void APlayerController::ClientPrepareMapChange_Implementation(FName LevelName, b
 	if (bFirst)
 	{
 		PendingMapChangeLevelNames.Empty();
-		GetWorldTimerManager().ClearTimer( this, &APlayerController::DelayedPrepareMapChange );
+		GetWorldTimerManager().ClearTimer(TimerHandle_DelayedPrepareMapChange);
 	}
 	PendingMapChangeLevelNames.Add(LevelName);
 	if (bLast)
@@ -2909,7 +2909,7 @@ void APlayerController::DelayedPrepareMapChange()
 	if (GetWorld()->IsPreparingMapChange())
 	{
 		// we must wait for the previous one to complete
-		GetWorldTimerManager().SetTimer( this, &APlayerController::DelayedPrepareMapChange, 0.01f );
+		GetWorldTimerManager().SetTimer(TimerHandle_DelayedPrepareMapChange, this, &APlayerController::DelayedPrepareMapChange, 0.01f );
 	}
 	else
 	{
@@ -2920,9 +2920,9 @@ void APlayerController::DelayedPrepareMapChange()
 
 void APlayerController::ClientCommitMapChange_Implementation()
 {
-	if (GetWorldTimerManager().IsTimerActive(this, &APlayerController::DelayedPrepareMapChange))
+	if (GetWorldTimerManager().IsTimerActive(TimerHandle_DelayedPrepareMapChange))
 	{
-		GetWorldTimerManager().SetTimer(this, &APlayerController::ClientCommitMapChange, 0.01f);
+		GetWorldTimerManager().SetTimer(TimerHandle_ClientCommitMapChange, this, &APlayerController::ClientCommitMapChange, 0.01f);
 	}
 	else
 	{
@@ -4050,7 +4050,7 @@ void APlayerController::BeginInactiveState()
 	AGameState const* const GameState = GetWorld()->GameState;
 
 	float const MinRespawnDelay = ((GameState != NULL) && (GameState->GameModeClass != NULL)) ? GetDefault<AGameMode>(GameState->GameModeClass)->MinRespawnDelay : 1.0f;
-	GetWorldTimerManager().SetTimer(this, &APlayerController::UnFreeze, MinRespawnDelay);
+	GetWorldTimerManager().SetTimer(TimerHandle_UnFreeze, this, &APlayerController::UnFreeze, MinRespawnDelay);
 }
 
 void APlayerController::EndInactiveState()
diff --git a/Engine/Source/Runtime/GameplayAbilities/Private/Abilities/Tasks/AbilityTask_Repeat.cpp b/Engine/Source/Runtime/GameplayAbilities/Private/Abilities/Tasks/AbilityTask_Repeat.cpp
--- a/Engine/Source/Runtime/GameplayAbilities/Private/Abilities/Tasks/AbilityTask_Repeat.cpp
+++ b/Engine/Source/Runtime/GameplayAbilities/Private/Abilities/Tasks/AbilityTask_Repeat.cpp
@@ -38,7 +38,7 @@ void UAbilityTask_Repeat::Activate()
 		PerformAction();
 		if (ActionCounter < ActionPerformancesDesired)
 		{
-			GetWorld()->GetTimerManager().SetTimer(this, &UAbilityTask_Repeat::PerformAction, TimeBetweenActions, true);
+			GetWorld()->GetTimerManager().SetTimer(TimerHandle_PerformAction, this, &UAbilityTask_Repeat::PerformAction, TimeBetweenActions, true);
 		}
 	}
 	else
@@ -50,7 +50,7 @@ void UAbilityTask_Repeat::Activate()
 
 void UAbilityTask_Repeat::OnDestroy(bool AbilityIsEnding)
 {
-	GetWorld()->GetTimerManager().ClearTimer(this, &UAbilityTask_Repeat::PerformAction);
+	GetWorld()->GetTimerManager().ClearTimer(TimerHandle_PerformAction);
 
 	Super::OnDestroy(AbilityIsEnding);
 }
diff --git a/Engine/Source/Runtime/GameplayAbilities/Private/Abilities/Tasks/AbilityTask_VisualizeTargeting.cpp b/Engine/Source/Runtime/GameplayAbilities/Private/Abilities/Tasks/AbilityTask_VisualizeTargeting.cpp
--- a/Engine/Source/Runtime/GameplayAbilities/Private/Abilities/Tasks/AbilityTask_VisualizeTargeting.cpp
+++ b/Engine/Source/Runtime/GameplayAbilities/Private/Abilities/Tasks/AbilityTask_VisualizeTargeting.cpp
@@ -12,17 +12,22 @@ UAbilityTask_VisualizeTargeting::UAbilityTask_VisualizeTargeting(const FObjectIn
 
 UAbilityTask_VisualizeTargeting* UAbilityTask_VisualizeTargeting::VisualizeTargeting(UObject* WorldContextObject, TSubclassOf<AGameplayAbilityTargetActor> InTargetClass, FName TaskInstanceName, float Duration)
 {
-	auto MyObj = NewTask<UAbilityTask_VisualizeTargeting>(WorldContextObject, TaskInstanceName);		//Register for task list here, providing a given FName as a key
+	UAbilityTask_VisualizeTargeting* MyObj = NewTask<UAbilityTask_VisualizeTargeting>(WorldContextObject, TaskInstanceName);		//Register for task list here, providing a given FName as a key
 	MyObj->TargetClass = InTargetClass;
-	if (Duration > 0.0f)
-	{
-		MyObj->GetWorld()->GetTimerManager().SetTimer(MyObj, &UAbilityTask_VisualizeTargeting::OnTimeElapsed, Duration, false);
-	}
+	MyObj->SetDuration(Duration);
 	return MyObj;
 }
 
 // ---------------------------------------------------------------------------------------
 
+void UAbilityTask_VisualizeTargeting::SetDuration(const float Duration)
+{
+	if (Duration > 0.f)
+	{
+		GetWorld()->GetTimerManager().SetTimer(TimerHandle_OnTimeElapsed, this, &UAbilityTask_VisualizeTargeting::OnTimeElapsed, Duration, false);
+	}
+}
+
 bool UAbilityTask_VisualizeTargeting::BeginSpawningActor(UObject* WorldContextObject, TSubclassOf<AGameplayAbilityTargetActor> TargetClass, AGameplayAbilityTargetActor*& SpawnedActor)
 {
 	SpawnedActor = nullptr;
@@ -89,7 +94,7 @@ void UAbilityTask_VisualizeTargeting::OnDestroy(bool AbilityEnded)
 	{
 		MyTargetActor->Destroy();
 	}
-	GetWorld()->GetTimerManager().ClearTimer(this, &UAbilityTask_VisualizeTargeting::OnTimeElapsed);
+	GetWorld()->GetTimerManager().ClearTimer(TimerHandle_OnTimeElapsed);
 
 	Super::OnDestroy(AbilityEnded);
 }
diff --git a/Engine/Source/Runtime/GameplayAbilities/Private/Abilities/Tasks/AbilityTask_WaitDelay.cpp b/Engine/Source/Runtime/GameplayAbilities/Private/Abilities/Tasks/AbilityTask_WaitDelay.cpp
--- a/Engine/Source/Runtime/GameplayAbilities/Private/Abilities/Tasks/AbilityTask_WaitDelay.cpp
+++ b/Engine/Source/Runtime/GameplayAbilities/Private/Abilities/Tasks/AbilityTask_WaitDelay.cpp
@@ -20,7 +20,10 @@ void UAbilityTask_WaitDelay::Activate()
 {
 	UWorld* World = GetWorld();
 	TimeStarted = World->GetTimeSeconds();
-	World->GetTimerManager().SetTimer(this, &UAbilityTask_WaitDelay::OnTimeFinish, Time, false);
+
+	// Use a dummy timer handle as we don't need to store it for later but we don't need to look for something to clear
+	FTimerHandle TimerHandle;
+	World->GetTimerManager().SetTimer(TimerHandle, this, &UAbilityTask_WaitDelay::OnTimeFinish, Time, false);
 }
 
 void UAbilityTask_WaitDelay::OnDestroy(bool AbilityEnded)
diff --git a/Engine/Source/Runtime/GameplayAbilities/Public/Abilities/Tasks/AbilityTask_Repeat.h b/Engine/Source/Runtime/GameplayAbilities/Public/Abilities/Tasks/AbilityTask_Repeat.h
--- a/Engine/Source/Runtime/GameplayAbilities/Public/Abilities/Tasks/AbilityTask_Repeat.h
+++ b/Engine/Source/Runtime/GameplayAbilities/Public/Abilities/Tasks/AbilityTask_Repeat.h
@@ -36,5 +36,8 @@ protected:
 	int32 ActionCounter;
 	float TimeBetweenActions;
 
+	/** Handle for efficient management of PerformAction timer */
+	FTimerHandle TimerHandle_PerformAction;
+
 	virtual void OnDestroy(bool AbilityIsEnding) override;
 };
\ No newline at end of file
diff --git a/Engine/Source/Runtime/GameplayAbilities/Public/Abilities/Tasks/AbilityTask_VisualizeTargeting.h b/Engine/Source/Runtime/GameplayAbilities/Public/Abilities/Tasks/AbilityTask_VisualizeTargeting.h
--- a/Engine/Source/Runtime/GameplayAbilities/Public/Abilities/Tasks/AbilityTask_VisualizeTargeting.h
+++ b/Engine/Source/Runtime/GameplayAbilities/Public/Abilities/Tasks/AbilityTask_VisualizeTargeting.h
@@ -1,6 +1,6 @@
 // Copyright 1998-2014 Epic Games, Inc. All Rights Reserved.
 #pragma once
-
+#include "TimerManager.h"
 #include "AbilityTask.h"
 #include "AbilityTask_VisualizeTargeting.generated.h"
 
@@ -18,6 +18,8 @@ class UAbilityTask_VisualizeTargeting: public UAbilityTask
 
 	void OnTimeElapsed();
 
+	void SetDuration(const float Duration);
+
 	/** Spawns Targeting actor. */
 	UFUNCTION(BlueprintCallable, meta=(HidePin = "WorldContextObject", DefaultToSelf = "WorldContextObject", BlueprintInternalUseOnly = "true", HideSpawnParms="Instigator"), Category="Ability|Tasks")
 	static UAbilityTask_VisualizeTargeting* VisualizeTargeting(UObject* WorldContextObject, TSubclassOf<AGameplayAbilityTargetActor> Class, FName TaskInstanceName, float Duration = -1.0f);
@@ -36,6 +38,9 @@ protected:
 
 	/** The TargetActor that we spawned, or the class CDO if this is a static targeting task */
 	TWeakObjectPtr<AGameplayAbilityTargetActor>	MyTargetActor;
+
+	/** Handle for efficient management of OnTimeElapsed timer */
+	FTimerHandle TimerHandle_OnTimeElapsed;
 };
 
 
diff --git a/Engine/Source/Runtime/Online/OnlineSubsystemUtils/Classes/LeaderboardQueryCallbackProxy.cpp b/Engine/Source/Runtime/Online/OnlineSubsystemUtils/Classes/LeaderboardQueryCallbackProxy.cpp
--- a/Engine/Source/Runtime/Online/OnlineSubsystemUtils/Classes/LeaderboardQueryCallbackProxy.cpp
+++ b/Engine/Source/Runtime/Online/OnlineSubsystemUtils/Classes/LeaderboardQueryCallbackProxy.cpp
@@ -106,6 +106,8 @@ void ULeaderboardQueryCallbackProxy::OnStatsRead(bool bWasSuccessful)
 
 	if (UWorld* World = WorldPtr.Get())
 	{
+		// Use a dummy timer handle as we don't need to store it for later but we don't need to look for something to clear
+		FTimerHandle TimerHandle;
 		World->GetTimerManager().SetTimer(this, &ULeaderboardQueryCallbackProxy::OnStatsRead_Delayed, 0.001f, false);
 	}
 	ReadObject = NULL;
diff --git a/Engine/Source/Runtime/Online/OnlineSubsystemUtils/Classes/OnlineBeaconClient.h b/Engine/Source/Runtime/Online/OnlineSubsystemUtils/Classes/OnlineBeaconClient.h
--- a/Engine/Source/Runtime/Online/OnlineSubsystemUtils/Classes/OnlineBeaconClient.h
+++ b/Engine/Source/Runtime/Online/OnlineSubsystemUtils/Classes/OnlineBeaconClient.h
@@ -4,6 +4,7 @@
 // Basic beacon
 //
 #pragma once
+#include "TimerManager.h"
 #include "Runtime/Online/OnlineSubsystemUtils/Classes/OnlineBeacon.h"
 #include "OnlineBeaconClient.generated.h"
 
@@ -64,6 +65,9 @@ protected:
 	UPROPERTY()
 	class AOnlineBeaconHostObject* BeaconOwner;
 
+	/** Handle for efficient management of OnFailure timer */
+	FTimerHandle TimerHandle_OnFailure;
+
 private:
 
 	/**
diff --git a/Engine/Source/Runtime/Online/OnlineSubsystemUtils/Private/InAppPurchaseCallbackProxy.cpp b/Engine/Source/Runtime/Online/OnlineSubsystemUtils/Private/InAppPurchaseCallbackProxy.cpp
--- a/Engine/Source/Runtime/Online/OnlineSubsystemUtils/Private/InAppPurchaseCallbackProxy.cpp
+++ b/Engine/Source/Runtime/Online/OnlineSubsystemUtils/Private/InAppPurchaseCallbackProxy.cpp
@@ -65,6 +65,8 @@ void UInAppPurchaseCallbackProxy::OnInAppPurchaseComplete(EInAppPurchaseState::T
     
 	if (UWorld* World = WorldPtr.Get())
 	{
+		// Use a local timer handle as we don't need to store it for later but we don't need to look for something to clear
+		FTimerHandle TimerHandle;
 		World->GetTimerManager().SetTimer(this, &UInAppPurchaseCallbackProxy::OnInAppPurchaseComplete_Delayed, 0.001f, false);
     }
     else
diff --git a/Engine/Source/Runtime/Online/OnlineSubsystemUtils/Private/InAppPurchaseQueryCallbackProxy.cpp b/Engine/Source/Runtime/Online/OnlineSubsystemUtils/Private/InAppPurchaseQueryCallbackProxy.cpp
--- a/Engine/Source/Runtime/Online/OnlineSubsystemUtils/Private/InAppPurchaseQueryCallbackProxy.cpp
+++ b/Engine/Source/Runtime/Online/OnlineSubsystemUtils/Private/InAppPurchaseQueryCallbackProxy.cpp
@@ -70,6 +70,8 @@ void UInAppPurchaseQueryCallbackProxy::OnInAppPurchaseRead(bool bWasSuccessful)
 
 	if (UWorld* World = WorldPtr.Get())
 	{
+		// Use a local timer handle as we don't need to store it for later but we don't need to look for something to clear
+		FTimerHandle TimerHandle;
 		World->GetTimerManager().SetTimer(this, &UInAppPurchaseQueryCallbackProxy::OnInAppPurchaseRead_Delayed, 0.001f, false);
 	}
 
diff --git a/Engine/Source/Runtime/Online/OnlineSubsystemUtils/Private/OnlineBeaconClient.cpp b/Engine/Source/Runtime/Online/OnlineSubsystemUtils/Private/OnlineBeaconClient.cpp
--- a/Engine/Source/Runtime/Online/OnlineSubsystemUtils/Private/OnlineBeaconClient.cpp
+++ b/Engine/Source/Runtime/Online/OnlineSubsystemUtils/Private/OnlineBeaconClient.cpp
@@ -67,7 +67,7 @@ void AOnlineBeaconClient::ClientOnConnected_Implementation()
 	SetAutonomousProxy(true);
 
 	// Fail safe for connection to server but no client connection RPC
-	GetWorldTimerManager().ClearTimer(this, &AOnlineBeaconClient::OnFailure);
+	GetWorldTimerManager().ClearTimer(TimerHandle_OnFailure);
 
 	if (NetDriver)
 	{
@@ -83,7 +83,7 @@ void AOnlineBeaconClient::ClientOnConnected_Implementation()
 void AOnlineBeaconClient::DestroyBeacon()
 {
 	// Fail safe for connection to server but no client connection RPC
-	GetWorldTimerManager().ClearTimer(this, &AOnlineBeaconClient::OnFailure);
+	GetWorldTimerManager().ClearTimer(TimerHandle_OnFailure);
 	Super::DestroyBeacon();
 }
 
@@ -140,7 +140,7 @@ void AOnlineBeaconClient::NotifyControlMessage(UNetConnection* Connection, uint8
 					// Server will send ClientOnConnected() when it gets this control message
 
 					// Fail safe for connection to server but no client connection RPC
-					GetWorldTimerManager().SetTimer(this, &AOnlineBeaconClient::OnFailure, BEACON_RPC_TIMEOUT, false);
+					GetWorldTimerManager().SetTimer(TimerHandle_OnFailure, this, &AOnlineBeaconClient::OnFailure, BEACON_RPC_TIMEOUT, false);
 				}
 				else
 				{
