Move FTimerDynamicDelegate definition to EngineTypes.h so as to not need a proxy delegate and then to copy the elements between them in UKismetSystemLibrary

diff --git a/Engine/Source/Runtime/Engine/Classes/Engine/EngineTypes.h b/Engine/Source/Runtime/Engine/Classes/Engine/EngineTypes.h
--- a/Engine/Source/Runtime/Engine/Classes/Engine/EngineTypes.h
+++ b/Engine/Source/Runtime/Engine/Classes/Engine/EngineTypes.h
@@ -2197,6 +2197,8 @@ namespace EEndPlayReason
 
 }
 
+DECLARE_DYNAMIC_DELEGATE(FTimerDynamicDelegate);
+
 /** Replicated movement data of our RootComponent.
   * Struct used for efficient replication as velocity and location are generally replicated together (this saves a repindex) 
   * and velocity.Z is commonly zero (most position replications are for walking pawns). 
diff --git a/Engine/Source/Runtime/Engine/Classes/Kismet/KismetSystemLibrary.h b/Engine/Source/Runtime/Engine/Classes/Kismet/KismetSystemLibrary.h
--- a/Engine/Source/Runtime/Engine/Classes/Kismet/KismetSystemLibrary.h
+++ b/Engine/Source/Runtime/Engine/Classes/Kismet/KismetSystemLibrary.h
@@ -7,9 +7,6 @@
 #include "Engine/CollisionProfile.h"
 #include "KismetSystemLibrary.generated.h"
 
-/** It's just a proxy for FTimerDynamicDelegate */
-DECLARE_DYNAMIC_DELEGATE(FBlueprintTimerDynamicDelegate);
-
 UENUM(BlueprintType)
 namespace EDrawDebugTrace
 {
@@ -267,7 +264,7 @@ class UKismetSystemLibrary : public UBlueprintFunctionLibrary
 	 * @param bLooping		true to keep executing the delegate every Time seconds, false to execute delegate only once.
 	 */
 	UFUNCTION(BlueprintCallable, meta=(FriendlyName = "SetTimerDelegate"), Category="Utilities|Time")
-	static void K2_SetTimerDelegate(FBlueprintTimerDynamicDelegate Delegate, float Time, bool bLooping);
+	static void K2_SetTimerDelegate(FTimerDynamicDelegate Delegate, float Time, bool bLooping);
 
 	/**
 	 * Clears a set timer.
@@ -443,7 +440,7 @@ class UKismetSystemLibrary : public UBlueprintFunctionLibrary
 	UFUNCTION(BlueprintCallable, Category="Collision", meta=(WorldContext="WorldContextObject", AutoCreateRefTerm="ActorsToIgnore", FriendlyName = "SphereOverlapActors"))
 	static ENGINE_API bool SphereOverlapActors_NEW(UObject* WorldContextObject, const FVector SpherePos, float SphereRadius, const TArray<TEnumAsByte<EObjectTypeQuery> > & ObjectTypes, UClass* ActorClassFilter, const TArray<AActor*>& ActorsToIgnore, TArray<class AActor*>& OutActors);
 
-	UFUNCTION(BlueprintCallable, Category="Collision", meta=(DeprecatedFunction, DeprecationMessage = "Use new SphereOverlaActors", WorldContext="WorldContextObject", AutoCreateRefTerm="ActorsToIgnore"))
+	UFUNCTION(BlueprintCallable, Category="Collision", meta=(DeprecatedFunction, DeprecationMessage = "Use new SphereOverlapActors", WorldContext="WorldContextObject", AutoCreateRefTerm="ActorsToIgnore"))
 	static ENGINE_API bool SphereOverlapActors_DEPRECATED(UObject* WorldContextObject, const FVector SpherePos, float SphereRadius, EOverlapFilterOption Filter, UClass* ActorClassFilter, const TArray<AActor*>& ActorsToIgnore, TArray<class AActor*>& OutActors);
 
 	/**
diff --git a/Engine/Source/Runtime/Engine/Private/KismetSystemLibrary.cpp b/Engine/Source/Runtime/Engine/Private/KismetSystemLibrary.cpp
--- a/Engine/Source/Runtime/Engine/Private/KismetSystemLibrary.cpp
+++ b/Engine/Source/Runtime/Engine/Private/KismetSystemLibrary.cpp
@@ -276,26 +276,21 @@ void UKismetSystemLibrary::K2_SetTimer(UObject* Object, FString FunctionName, fl
 		}
 	}
 
-	FBlueprintTimerDynamicDelegate Delegate;
+	FTimerDynamicDelegate Delegate;
 	Delegate.BindUFunction(Object, FunctionFName);
 	K2_SetTimerDelegate(Delegate, Time, bLooping);
 }
 
-void UKismetSystemLibrary::K2_SetTimerDelegate(FBlueprintTimerDynamicDelegate Delegate, float Time, bool bLooping)
+void UKismetSystemLibrary::K2_SetTimerDelegate(FTimerDynamicDelegate Delegate, float Time, bool bLooping)
 {
-	UObject* const Object = Delegate.GetUObject();
-	FName const FunctionName = Delegate.GetFunctionName();
-
-	FTimerDynamicDelegate InnerDelegate;
-	InnerDelegate.BindUFunction(Object, FunctionName);
-	if (InnerDelegate.IsBound())
+	if (Delegate.IsBound())
 	{
-		const UWorld* const World = GEngine->GetWorldFromContextObject(Object);
-		World->GetTimerManager().SetTimer(InnerDelegate, Time, bLooping);
+		const UWorld* const World = GEngine->GetWorldFromContextObject(Delegate.GetUObject());
+		World->GetTimerManager().SetTimer(Delegate, Time, bLooping);
 	}
 	else
 	{
-		UE_LOG(LogBlueprintUserMessages, Warning, TEXT("SetTimer passed a bad function (%s) or object (%s)"), *FunctionName.ToString(), *GetNameSafe(Object));
+		UE_LOG(LogBlueprintUserMessages, Warning, TEXT("SetTimer passed a bad function (%s) or object (%s)"), *Delegate.GetFunctionName().ToString(), *GetNameSafe(Delegate.GetUObject()));
 	}
 }
 
diff --git a/Engine/Source/Runtime/Engine/Public/TimerManager.h b/Engine/Source/Runtime/Engine/Public/TimerManager.h
--- a/Engine/Source/Runtime/Engine/Public/TimerManager.h
+++ b/Engine/Source/Runtime/Engine/Public/TimerManager.h
@@ -6,16 +6,7 @@
 
 #pragma once
 
-
-// @hack to make wrapped dynamic delegate work in a non-UClass
-inline void FTimerDynamicDelegate_DelegateWrapper(const FScriptDelegate& FTimerDynamicDelegate)
-{
-	FTimerDynamicDelegate.ProcessDelegate<UObject>(NULL);
-}
-
 DECLARE_DELEGATE(FTimerDelegate);
-DECLARE_DYNAMIC_DELEGATE(FTimerDynamicDelegate);
-
 
 /** Simple interface to wrap a timer delegate that can be either native or dynamic. */
 struct FTimerUnifiedDelegate
