Fix for crash using check being tiiggered in MIDs.

diff --git a/Engine/Shaders/ShaderVersion.usf b/Engine/Shaders/ShaderVersion.usf
--- a/Engine/Shaders/ShaderVersion.usf
+++ b/Engine/Shaders/ShaderVersion.usf
@@ -2,4 +2,4 @@
 // Each time the console command is executed it generates a new GUID. As this file is included
 // in common.usf (which should be included in any shader) it allows to invalidate the shader DDC.
 // 
-// GUID = 48B5ECE04311A6B3C8BB15A2A11716CD
\ No newline at end of file
+// GUID = 9F1892334C4B8A2429020D883842B32E
\ No newline at end of file
diff --git a/Engine/Source/Runtime/Engine/Private/Materials/MaterialInstance.cpp b/Engine/Source/Runtime/Engine/Private/Materials/MaterialInstance.cpp
--- a/Engine/Source/Runtime/Engine/Private/Materials/MaterialInstance.cpp
+++ b/Engine/Source/Runtime/Engine/Private/Materials/MaterialInstance.cpp
@@ -2505,13 +2505,11 @@ bool UMaterialInstance::HasOverridenBaseProperties()const
 {
 	check(IsInGameThread());
 
-	const UMaterial* Mat = GetMaterial();
-	check(Mat);
-
-	if ((FMath::Abs(GetOpacityMaskClipValue(true) - Mat->GetOpacityMaskClipValue(true)) > SMALL_NUMBER) ||
-		(GetBlendMode(true) != Mat->GetBlendMode(true)) ||
-		(GetShadingModel(true) != Mat->GetShadingModel(true)) ||
-		(IsTwoSided(true) != Mat->IsTwoSided(true))
+	if (Parent && (
+		(FMath::Abs(GetOpacityMaskClipValue(true) - Parent->GetOpacityMaskClipValue(true)) > SMALL_NUMBER) ||
+		(GetBlendMode(true) != Parent->GetBlendMode(true)) ||
+		(GetShadingModel(true) != Parent->GetShadingModel(true)) ||
+		(IsTwoSided(true) != Parent->IsTwoSided(true)))
 		)
 	{
 		return true;
