From John Barrett: Added 'RuntimeServerActors', to allow on-the-fly adding/removing of serveractors, without requiring config changes.

diff --git a/Engine/Source/Runtime/Engine/Classes/Engine/Engine.h b/Engine/Source/Runtime/Engine/Classes/Engine/Engine.h
--- a/Engine/Source/Runtime/Engine/Classes/Engine/Engine.h
+++ b/Engine/Source/Runtime/Engine/Classes/Engine/Engine.h
@@ -2224,11 +2224,15 @@ public:
 	UPROPERTY(Config, transient)
 	TArray<FNetDriverDefinition> NetDriverDefinitions;
 
-	/** @todo document */
+	/** A configurable list of actors that are automatically spawned upon server startup (just prior to InitGame) */
 	UPROPERTY(config)
 	TArray<FString> ServerActors;
 
-	/** Spawns all of the registered server actors	 */
+	/** Runtime-modified list of server actors, allowing plugins to use serveractors, without permanently adding them to config files */
+	UPROPERTY()
+	TArray<FString> RuntimeServerActors;
+
+	/** Spawns all of the registered server actors */
 	virtual void SpawnServerActors(UWorld *World);
 
 	/**
diff --git a/Engine/Source/Runtime/Engine/Private/UnrealEngine.cpp b/Engine/Source/Runtime/Engine/Private/UnrealEngine.cpp
--- a/Engine/Source/Runtime/Engine/Private/UnrealEngine.cpp
+++ b/Engine/Source/Runtime/Engine/Private/UnrealEngine.cpp
@@ -7797,10 +7797,15 @@ void UEngine::HandleTravelFailure_NotifyGameInstance(UWorld *World, ETravelFailu
 
 void UEngine::SpawnServerActors(UWorld *World)
 {
-	for( int32 i=0; i < ServerActors.Num(); i++ )
+	TArray<FString> FullServerActors;
+
+	FullServerActors.Append(ServerActors);
+	FullServerActors.Append(RuntimeServerActors);
+
+	for( int32 i=0; i < FullServerActors.Num(); i++ )
 	{
-		TCHAR Str[240];
-		const TCHAR* Ptr = * ServerActors[i];
+		TCHAR Str[2048];
+		const TCHAR* Ptr = * FullServerActors[i];
 		if( FParse::Token( Ptr, Str, ARRAY_COUNT(Str), 1 ) )
 		{
 			UE_LOG(LogNet, Log, TEXT("Spawning: %s"), Str );
