Expose wheeled component substructs to ENGINE_API so games can properly override behavior

(UE-5632)

diff --git a/Engine/Source/Runtime/Engine/Classes/Vehicles/WheeledVehicleMovementComponent.h b/Engine/Source/Runtime/Engine/Classes/Vehicles/WheeledVehicleMovementComponent.h
--- a/Engine/Source/Runtime/Engine/Classes/Vehicles/WheeledVehicleMovementComponent.h
+++ b/Engine/Source/Runtime/Engine/Classes/Vehicles/WheeledVehicleMovementComponent.h
@@ -84,7 +84,7 @@ struct FTireShaderOutput
  * Vehicle-specific wheel setup
  */
 USTRUCT()
-struct FWheelSetup
+struct ENGINE_API FWheelSetup
 {
 	GENERATED_USTRUCT_BODY()
 
@@ -104,7 +104,7 @@ struct FWheelSetup
 };
 
 USTRUCT()
-struct FReplicatedVehicleState
+struct ENGINE_API FReplicatedVehicleState
 {
 	GENERATED_USTRUCT_BODY()
 
@@ -130,7 +130,7 @@ struct FReplicatedVehicleState
 };
 
 USTRUCT()
-struct FVehicleInputRate
+struct ENGINE_API FVehicleInputRate
 {
 	GENERATED_USTRUCT_BODY()
 
