Emissive static lighting is now disabled by default, have to enable bUseEmissiveForStaticLighting on components to get it

diff --git a/Engine/Source/Runtime/Engine/Classes/Engine/EngineTypes.h b/Engine/Source/Runtime/Engine/Classes/Engine/EngineTypes.h
--- a/Engine/Source/Runtime/Engine/Classes/Engine/EngineTypes.h
+++ b/Engine/Source/Runtime/Engine/Classes/Engine/EngineTypes.h
@@ -1124,7 +1124,7 @@ struct FLightmassPrimitiveSettings
 	{
 		bUseTwoSidedLighting = false;
 		bShadowIndirectOnly = false;
-		bUseEmissiveForStaticLighting = true;
+		bUseEmissiveForStaticLighting = false;
 		EmissiveLightFalloffExponent = 8.0f;
 		EmissiveLightExplicitInfluenceRadius = 0.0f;
 		EmissiveBoost = 1.0f;
