Fixed performance issue BeginCacheForCookedPlatformData would not use correct DDC key when looking up cooked platform data.

diff --git a/Engine/Source/Runtime/Engine/Private/TextureDerivedData.cpp b/Engine/Source/Runtime/Engine/Private/TextureDerivedData.cpp
--- a/Engine/Source/Runtime/Engine/Private/TextureDerivedData.cpp
+++ b/Engine/Source/Runtime/Engine/Private/TextureDerivedData.cpp
@@ -1455,7 +1455,7 @@ void UTexture::BeginCacheForCookedPlatformData( const ITargetPlatform *TargetPla
 		for (int32 SettingsIndex = 0; SettingsIndex < BuildSettingsToCache.Num(); ++SettingsIndex)
 		{
 			FString DerivedDataKey;
-			GetTextureDerivedDataKeySuffix(*this, BuildSettingsToCache[SettingsIndex], DerivedDataKey);
+			GetTextureDerivedDataKey(*this, BuildSettingsToCache[SettingsIndex], DerivedDataKey);
 
 			FTexturePlatformData *PlatformData = CookedPlatformData.FindRef( DerivedDataKey );
 
