Removing the extraneous SetTickEnabled function from Actor, and exposing the other to blueprints instead.

diff --git a/Engine/Config/BaseEngine.ini b/Engine/Config/BaseEngine.ini
--- a/Engine/Config/BaseEngine.ini
+++ b/Engine/Config/BaseEngine.ini
@@ -812,6 +812,7 @@ bUseStreamingPause=false
 +K2FieldRedirects=(OldFieldName="Actor.Destroy",NewFieldName="K2_DestroyActor")
 +K2FieldRedirects=(OldFieldName="Actor.GetControllerAnalogStickState",NewFieldName="InputComponent.GetControllerAnalogStickState")
 +K2FieldRedirects=(OldFieldName="Actor.HasTag", NewFieldName="Actor.ActorHasTag")
++K2FieldRedirects=(OldFieldName="Actor.SetTickEnabled", NewFieldName="Actor.SetActorTickEnabled")
 
 +K2FieldRedirects=(OldFieldName="MaterialInstanceConstant.SetVectorParameterValue",NewFieldName="MaterialInstanceDynamic.SetVectorParameterValue")
 +K2FieldRedirects=(OldFieldName="MaterialInstanceConstant.SetScalarParameterValue",NewFieldName="MaterialInstanceDynamic.SetScalarParameterValue")
diff --git a/Engine/Source/Runtime/Engine/Classes/GameFramework/Actor.h b/Engine/Source/Runtime/Engine/Classes/GameFramework/Actor.h
--- a/Engine/Source/Runtime/Engine/Classes/GameFramework/Actor.h
+++ b/Engine/Source/Runtime/Engine/Classes/GameFramework/Actor.h
@@ -946,10 +946,6 @@ public:
 	UFUNCTION(BlueprintCallable, Category="Utilities")
 	void SetTickableWhenPaused(bool bTickableWhenPaused);
 
-	/** Disables/Enables ticking for this actor (CanEverTick() has to be true, for tick enabling to work). */
-	UFUNCTION(BlueprintCallable, Category = "Utilities")
-	void SetTickEnabled(bool bEnable = true);
-
 	/** Allocate a MID for a given parent material. */
 	UFUNCTION(BlueprintCallable, meta=(DeprecatedFunction, DeprecationMessage="Use PrimitiveComponent.CreateAndSetMaterialInstanceDynamic instead.", BlueprintProtected = "true"), Category="Rendering|Material")
 	class UMaterialInstanceDynamic* MakeMIDForMaterial(class UMaterialInterface* Parent);
@@ -1489,6 +1485,7 @@ public:
 	 * This only modifies the tick function on actor itself
 	 * @param	bEnabled - Rather it should be enabled or not
 	 */
+	UFUNCTION(BlueprintCallable, Category="Utilities")
 	virtual void SetActorTickEnabled(bool bEnabled);
 
 	/**  Returns whether this actor has tick enabled or not	 */
diff --git a/Engine/Source/Runtime/Engine/Private/Actor.cpp b/Engine/Source/Runtime/Engine/Private/Actor.cpp
--- a/Engine/Source/Runtime/Engine/Private/Actor.cpp
+++ b/Engine/Source/Runtime/Engine/Private/Actor.cpp
@@ -388,18 +388,6 @@ void AActor::SetTickableWhenPaused(bool bTickableWhenPaused)
 	PrimaryActorTick.bTickEvenWhenPaused = bTickableWhenPaused;
 }
 
-void AActor::SetTickEnabled(bool bEnable)
-{
-	if (CanEverTick())
-	{
-		PrimaryActorTick.SetTickFunctionEnable(bEnable);
-	}
-	else if (bEnable)
-	{
-		UE_LOG(LogActor, Warning, TEXT("Attempting to enable ticking for '%s' when the actor is explicitly prevented from doing so (check CanEverTick)."), *GetName());
-	}
-}
-
 void AActor::AddControllingMatineeActor( AMatineeActor& InMatineeActor )
 {
 	if (RootComponent)
