Construct RPC function parameters when needed (fixes blank FText replication crash)

diff --git a/Engine/Source/Runtime/Engine/Private/RepLayout.cpp b/Engine/Source/Runtime/Engine/Private/RepLayout.cpp
--- a/Engine/Source/Runtime/Engine/Private/RepLayout.cpp
+++ b/Engine/Source/Runtime/Engine/Private/RepLayout.cpp
@@ -2450,7 +2450,13 @@ void FRepLayout::ReceivePropertiesForRPC( UObject* Object, UFunction * Function,
 
 	for ( int32 i = 0; i < Parents.Num(); i++ )
 	{
-		if ( Cast<UBoolProperty>( Parents[i].Property ) || Reader.ReadBit() ) 
+		if ( Parents[i].ArrayIndex == 0 && ( Parents[i].Property->PropertyFlags & CPF_ZeroConstructor ) == 0 )
+		{
+			// If this property needs to be constructed, make sure we do that
+			Parents[i].Property->InitializeValue( (uint8*)Data + Cmds[ Parents[i].CmdStart ].Offset );
+		}
+
+		if ( Cast<UBoolProperty>( Parents[i].Property ) || Reader.ReadBit() )
 		{
 			bool bHasUnmapped = false;
 
