Retain encoded (ES2) reflection data when editor is active.

Fixes UE-5450

diff --git a/Engine/Source/Runtime/Engine/Private/ReflectionCaptureComponent.cpp b/Engine/Source/Runtime/Engine/Private/ReflectionCaptureComponent.cpp
--- a/Engine/Source/Runtime/Engine/Private/ReflectionCaptureComponent.cpp
+++ b/Engine/Source/Runtime/Engine/Private/ReflectionCaptureComponent.cpp
@@ -984,7 +984,7 @@ void UReflectionCaptureComponent::PostLoad()
 	Super::PostLoad();
 
 	bool bRetainAllFeatureLevelData = GIsEditor && GMaxRHIFeatureLevel >= ERHIFeatureLevel::SM4;
-	bool bEncodedDataRequired = (GMaxRHIFeatureLevel == ERHIFeatureLevel::ES2 || GMaxRHIFeatureLevel == ERHIFeatureLevel::ES3_1);
+	bool bEncodedDataRequired = bRetainAllFeatureLevelData || (GMaxRHIFeatureLevel == ERHIFeatureLevel::ES2 || GMaxRHIFeatureLevel == ERHIFeatureLevel::ES3_1);
 	bool bFullDataRequired = GMaxRHIFeatureLevel >= ERHIFeatureLevel::SM4;
 
 	// If we're loading on a platform that doesn't require cooked data, attempt to load missing data from the DDC
@@ -1019,7 +1019,7 @@ void UReflectionCaptureComponent::PostLoad()
 		// If we have full HDR data but not encoded HDR data, generate the encoded data now
 		if (FullHDRDerivedData 
 			&& !EncodedHDRDerivedData 
-			&& (bEncodedDataRequired))
+			&& bEncodedDataRequired)
 		{
 			EncodedHDRDerivedData = FReflectionCaptureEncodedHDRDerivedData::GenerateEncodedHDRData(*FullHDRDerivedData, StateId, Brightness);
 		}
@@ -1056,9 +1056,12 @@ void UReflectionCaptureComponent::PostLoad()
 		EncodedHDRCubemapTexture->SetupParameters(GReflectionCaptureSize, FMath::CeilLogTwo(GReflectionCaptureSize) + 1, PF_B8G8R8A8, &EncodedHDRDerivedData->CapturedData);
 		BeginInitResource(EncodedHDRCubemapTexture);
 
-		// Don't need the full HDR data for rendering on this feature level
-		delete FullHDRDerivedData;
-		FullHDRDerivedData = NULL;
+		// free up the full hdr data if we no longer need it.
+		if (FullHDRDerivedData && !bFullDataRequired)
+		{
+			delete FullHDRDerivedData;
+			FullHDRDerivedData = NULL;
+		}
 	}
 }
 
