fixed missing navmesh after loading a map when auto rebuilds are disabled

diff --git a/Engine/Source/Runtime/Engine/Classes/AI/Navigation/NavigationData.h b/Engine/Source/Runtime/Engine/Classes/AI/Navigation/NavigationData.h
--- a/Engine/Source/Runtime/Engine/Classes/AI/Navigation/NavigationData.h
+++ b/Engine/Source/Runtime/Engine/Classes/AI/Navigation/NavigationData.h
@@ -435,7 +435,7 @@ public:
 	const FNavDataConfig& GetConfig() const { return NavDataConfig; }
 	virtual void SetConfig(const FNavDataConfig& Src) { NavDataConfig = Src; }
 
-	void SetSupportsDefaultAgent(bool bIsDefault) { bSupportsDefaultAgent = bIsDefault; bEnableDrawing = bIsDefault; } 
+	void SetSupportsDefaultAgent(bool bIsDefault) { bSupportsDefaultAgent = bIsDefault; SetNavRenderingEnabled(bIsDefault); }
 	bool IsSupportingDefaultAgent() const { return bSupportsDefaultAgent; }
 
 	virtual bool DoesSupportAgent(const FNavAgentProperties& AgentProps) const { return false; }
@@ -709,6 +709,9 @@ public:
 	//----------------------------------------------------------------------//
 	virtual UPrimitiveComponent* ConstructRenderingComponent() { return NULL; }
 
+	/** updates state of rendering component */
+	void SetNavRenderingEnabled(bool bEnable);
+
 protected:
 	void InstantiateAndRegisterRenderingComponent();
 
diff --git a/Engine/Source/Runtime/Engine/Classes/AI/Navigation/NavigationSystem.h b/Engine/Source/Runtime/Engine/Classes/AI/Navigation/NavigationSystem.h
--- a/Engine/Source/Runtime/Engine/Classes/AI/Navigation/NavigationSystem.h
+++ b/Engine/Source/Runtime/Engine/Classes/AI/Navigation/NavigationSystem.h
@@ -678,6 +678,7 @@ protected:
 	uint32 bInitialBuildingLockActive:1;
 	uint32 bInitialSetupHasBeenPerformed:1;
 	uint32 bInitialLevelsAdded:1;
+	uint32 bSkipDirtyAreasOnce:1;
 	uint32 bAsyncBuildPaused:1;
 
 	/** cached navigable world bounding box*/
diff --git a/Engine/Source/Runtime/Engine/Private/AI/Navigation/NavigationData.cpp b/Engine/Source/Runtime/Engine/Private/AI/Navigation/NavigationData.cpp
--- a/Engine/Source/Runtime/Engine/Private/AI/Navigation/NavigationData.cpp
+++ b/Engine/Source/Runtime/Engine/Private/AI/Navigation/NavigationData.cpp
@@ -588,6 +588,15 @@ int32 ANavigationData::GetAreaID(const UClass* AreaClass) const
 	return PtrId ? *PtrId : INDEX_NONE;
 }
 
+void ANavigationData::SetNavRenderingEnabled(bool bEnable)
+{
+	if (bEnableDrawing != bEnable)
+	{
+		bEnableDrawing = bEnable;
+		MarkComponentsRenderStateDirty();
+	}
+}
+
 void ANavigationData::UpdateCustomLink(const class INavLinkCustomInterface* CustomLink)
 {
 	// no implementation for abstract class
diff --git a/Engine/Source/Runtime/Engine/Private/AI/Navigation/NavigationSystem.cpp b/Engine/Source/Runtime/Engine/Private/AI/Navigation/NavigationSystem.cpp
--- a/Engine/Source/Runtime/Engine/Private/AI/Navigation/NavigationSystem.cpp
+++ b/Engine/Source/Runtime/Engine/Private/AI/Navigation/NavigationSystem.cpp
@@ -375,6 +375,9 @@ void UNavigationSystem::OnWorldInitDone(FNavigationSystem::EMode Mode)
 		{
 			// don't lock navigation building in editor
 			bInitialBuildingLockActive = false;
+
+			// don't mark dirty areas after loading a map when automatic rebuild is disabled
+			bSkipDirtyAreasOnce = !bNavigationAutoUpdateEnabled;
 		}
 		
 		if (bAutoCreateNavigationData == true)
@@ -503,6 +506,12 @@ void UNavigationSystem::Tick(float DeltaSeconds)
 
 		if (DirtyAreas.Num() > 0 && bCanRebuildNow)
 		{
+			if (bSkipDirtyAreasOnce)
+			{
+				bSkipDirtyAreasOnce = false;
+				DirtyAreas.Reset();
+			}
+
 			for (int32 NavDataIndex = 0; NavDataIndex < NavDataSet.Num(); ++NavDataIndex)
 			{
 				ANavigationData* NavData = NavDataSet[NavDataIndex];
@@ -3164,11 +3173,7 @@ void UNavigationSystem::CycleNavigationDataDrawn()
 		if (NavData != NULL)
 		{
 			const bool bNewEnabledDrawing = (CurrentlyDrawnNavDataIndex == INDEX_NONE) || (NavDataIndex == CurrentlyDrawnNavDataIndex);
-			if (bNewEnabledDrawing != NavData->bEnableDrawing)
-			{
-				NavData->bEnableDrawing = bNewEnabledDrawing;
-				NavData->MarkComponentsRenderStateDirty();
-			}
+			NavData->SetNavRenderingEnabled(bNewEnabledDrawing);
 		}
 	}
 }
diff --git a/Engine/Source/Runtime/Engine/Private/AI/Navigation/RecastNavMeshGenerator.cpp b/Engine/Source/Runtime/Engine/Private/AI/Navigation/RecastNavMeshGenerator.cpp
--- a/Engine/Source/Runtime/Engine/Private/AI/Navigation/RecastNavMeshGenerator.cpp
+++ b/Engine/Source/Runtime/Engine/Private/AI/Navigation/RecastNavMeshGenerator.cpp
@@ -2710,18 +2710,32 @@ static int32 CaclulateMaxTilesCount(const TNavStatArray<FBox>& NavigableAreas, f
 //----------------------------------------------------------------------//
 
 FRecastNavMeshGenerator::FRecastNavMeshGenerator(class ARecastNavMesh* InDestNavMesh)
-	: NumActiveTiles(0)
-	, MaxTileGeneratorTasks(1)
-	, AvgLayersPerTile(8.0f)
-	, DestNavMesh(InDestNavMesh)
-	, bInitialized(false)
-	, Version(0)
+	: NumActiveTiles(0),
+	  MaxTileGeneratorTasks(1),
+	  AvgLayersPerTile(8.0f),
+	  DestNavMesh(InDestNavMesh),
+	  bInitialized(false),
+	  Version(0)
 {
-	INC_DWORD_STAT_BY( STAT_NavigationMemory, sizeof(*this) );
+	INC_DWORD_STAT_BY(STAT_NavigationMemory, sizeof(*this));
 
 	check(InDestNavMesh);
 	Init();
-	ConstructTiledNavMesh();
+
+	const bool bUseAutoUpdates = UNavigationSystem::GetIsNavigationAutoUpdateEnabled();
+	if (!DestNavMesh->HasValidNavmesh() && !bUseAutoUpdates)
+	{
+		// recreate navmesh from scratch if no data was loaded
+		// use only whe auto rebuilds are disabled (should reuse existing data)
+		ConstructTiledNavMesh();
+	}
+	else
+	{
+		// otherwise just update generator params
+		int32 MaxTiles = 0;
+		CalcNavMeshProperties(MaxTiles, Config.MaxPolysPerTile);
+		NumActiveTiles = GetTilesCountHelper(DestNavMesh->GetRecastNavMeshImpl()->DetourNavMesh);
+	}
 }
 
 FRecastNavMeshGenerator::~FRecastNavMeshGenerator()
@@ -2853,38 +2867,9 @@ bool FRecastNavMeshGenerator::ConstructTiledNavMesh()
 		TiledMeshParameters.tileWidth = Config.tileSize * Config.cs;
 		TiledMeshParameters.tileHeight = Config.tileSize * Config.cs;
 
-		// limit max amount of tiles
-#if USE_64BIT_ADDRESS
-		const int32 MaxTileBits = 30;
-#else
-		const int32 MaxTileBits = 14;
-#endif//USE_64BIT_ADDRESS
-
-		const int32 MaxTiles = (1 << MaxTileBits);
-		int32 MaxRequestedTiles = 0;
-		if (DestNavMesh->IsResizable())
-		{
-			MaxRequestedTiles = CaclulateMaxTilesCount(InclusionBounds, Config.tileSize * Config.cs, AvgLayersPerTile);
-		}
-		else
-		{
-			MaxRequestedTiles = DestNavMesh->TilePoolSize;
-		}
-		
-		if (MaxRequestedTiles < 0 || MaxRequestedTiles > MaxTiles)
-		{
-			UE_LOG(LogNavigation, Error, TEXT("Navmesh bounds are too large! Limiting requested tiles count (%d) to: (%d)"), MaxRequestedTiles, MaxTiles);
-			MaxRequestedTiles = MaxTiles;
-		}
+		CalcNavMeshProperties(TiledMeshParameters.maxTiles, TiledMeshParameters.maxPolys);
+		Config.MaxPolysPerTile = TiledMeshParameters.maxPolys;
 
-		// Max tiles and max polys affect how the tile IDs are calculated.
-		// There are (sizeof(dtPolyRef)*8 - DT_MIN_SALT_BITS) bits available for 
-		// identifying a tile and a polygon.
-		Config.MaxPolysPerTile = 1 << ((sizeof(dtPolyRef)*8 - DT_MIN_SALT_BITS) - MaxTileBits);
-
-		TiledMeshParameters.maxTiles = MaxRequestedTiles;
-		TiledMeshParameters.maxPolys = Config.MaxPolysPerTile;
-	
 		const dtStatus status = DetourMesh->init(&TiledMeshParameters);
 
 		if (dtStatusFailed(status))
@@ -2908,6 +2893,39 @@ bool FRecastNavMeshGenerator::ConstructTiledNavMesh()
 	return bSuccess;
 }
 
+void FRecastNavMeshGenerator::CalcNavMeshProperties(int32& MaxTiles, int32& MaxPolys)
+{
+	// limit max amount of tiles
+#if USE_64BIT_ADDRESS
+	const int32 MaxTileBits = 30;
+#else
+	const int32 MaxTileBits = 14;
+#endif//USE_64BIT_ADDRESS
+
+	const int32 MaxTilesFromMask = (1 << MaxTileBits);
+	int32 MaxRequestedTiles = 0;
+	if (DestNavMesh->IsResizable())
+	{
+		MaxRequestedTiles = CaclulateMaxTilesCount(InclusionBounds, Config.tileSize * Config.cs, AvgLayersPerTile);
+	}
+	else
+	{
+		MaxRequestedTiles = DestNavMesh->TilePoolSize;
+	}
+
+	if (MaxRequestedTiles < 0 || MaxRequestedTiles > MaxTilesFromMask)
+	{
+		UE_LOG(LogNavigation, Error, TEXT("Navmesh bounds are too large! Limiting requested tiles count (%d) to: (%d)"), MaxRequestedTiles, MaxTilesFromMask);
+		MaxRequestedTiles = MaxTilesFromMask;
+	}
+
+	// Max tiles and max polys affect how the tile IDs are calculated.
+	// There are (sizeof(dtPolyRef)*8 - DT_MIN_SALT_BITS) bits available for 
+	// identifying a tile and a polygon.
+	MaxPolys = 1 << ((sizeof(dtPolyRef)* 8 - DT_MIN_SALT_BITS) - MaxTileBits);
+	MaxTiles = MaxRequestedTiles;
+}
+
 bool FRecastNavMeshGenerator::RebuildAll()
 {
 	DestNavMesh->UpdateNavVersion();
diff --git a/Engine/Source/Runtime/Engine/Public/AI/Navigation/RecastNavMeshGenerator.h b/Engine/Source/Runtime/Engine/Public/AI/Navigation/RecastNavMeshGenerator.h
--- a/Engine/Source/Runtime/Engine/Public/AI/Navigation/RecastNavMeshGenerator.h
+++ b/Engine/Source/Runtime/Engine/Public/AI/Navigation/RecastNavMeshGenerator.h
@@ -374,6 +374,9 @@ private:
 
 	/** Instantiates dtNavMesh and configures it for tiles generation. Returns false if failed */
 	bool ConstructTiledNavMesh();
+
+	/** Determine bit masks for poly address */
+	void CalcNavMeshProperties(int32& MaxTiles, int32& MaxPolys);
 	
 	/** Marks grid tiles affected by specified areas as dirty */
 	void MarkDirtyTiles(const TArray<FNavigationDirtyArea>& DirtyAreas);
