fixed messed up condition in navigation generator related to auto updates

diff --git a/Engine/Source/Runtime/Engine/Private/AI/Navigation/RecastNavMeshGenerator.cpp b/Engine/Source/Runtime/Engine/Private/AI/Navigation/RecastNavMeshGenerator.cpp
--- a/Engine/Source/Runtime/Engine/Private/AI/Navigation/RecastNavMeshGenerator.cpp
+++ b/Engine/Source/Runtime/Engine/Private/AI/Navigation/RecastNavMeshGenerator.cpp
@@ -2723,7 +2723,7 @@ FRecastNavMeshGenerator::FRecastNavMeshGenerator(class ARecastNavMesh* InDestNav
 	Init();
 
 	const bool bUseAutoUpdates = UNavigationSystem::GetIsNavigationAutoUpdateEnabled();
-	if (!DestNavMesh->HasValidNavmesh() && !bUseAutoUpdates)
+	if (!DestNavMesh->HasValidNavmesh() || bUseAutoUpdates)
 	{
 		// recreate navmesh from scratch if no data was loaded
 		// use only whe auto rebuilds are disabled (should reuse existing data)
