Added SplineMeshActor.

diff --git a/Engine/Source/Runtime/Engine/Classes/Engine/SplineMeshActor.h b/Engine/Source/Runtime/Engine/Classes/Engine/SplineMeshActor.h
new file mode 100644
--- /dev/null
+++ b/Engine/Source/Runtime/Engine/Classes/Engine/SplineMeshActor.h
@@ -0,0 +1,43 @@
+// Copyright 1998-2014 Epic Games, Inc. All Rights Reserved.
+
+#pragma once
+#include "SplineMeshActor.generated.h"
+
+/**
+ * SplineMeshActor is an actor with a SplineMeshComponent.
+ *
+ * @see USplineMeshComponent
+ */
+UCLASS(hidecategories = (Input), showcategories = ("Input|MouseInput", "Input|TouchInput"), ConversionRoot, meta = (ChildCanTick))
+class ENGINE_API ASplineMeshActor : public AActor
+{
+	GENERATED_UCLASS_BODY()
+
+private_subobject:
+	UPROPERTY(Category = SplineMeshActor, VisibleAnywhere, BlueprintReadOnly, meta = (ExposeFunctionCategories = "Mesh,Rendering,Physics,Components|StaticMesh,Components|SplineMesh", AllowPrivateAccess = "true"))
+	class USplineMeshComponent* SplineMeshComponent;
+
+public:
+
+	/** Function to change mobility type */
+	void SetMobility(EComponentMobility::Type InMobility);
+
+#if WITH_EDITOR
+	// Begin AActor Interface
+	virtual void CheckForErrors() override;
+	virtual bool GetReferencedContentObjects(TArray<UObject*>& Objects) const override;
+	// End AActor Interface
+#endif // WITH_EDITOR	
+
+protected:
+	// Begin UObject interface.
+	virtual FString GetDetailedInfoInternal() const override;
+	// End UObject interface.
+
+public:
+	/** Returns SplineMeshComponent subobject **/
+	class USplineMeshComponent* GetSplineMeshComponent() const;
+};
+
+
+
diff --git a/Engine/Source/Runtime/Engine/Private/SplineMeshActor.cpp b/Engine/Source/Runtime/Engine/Private/SplineMeshActor.cpp
new file mode 100644
--- /dev/null
+++ b/Engine/Source/Runtime/Engine/Private/SplineMeshActor.cpp
@@ -0,0 +1,74 @@
+// Copyright 1998-2014 Epic Games, Inc. All Rights Reserved.
+
+/*=============================================================================
+	SplineMeshActor.cpp: Spline mesh actor class implementation.
+=============================================================================*/
+
+#include "EnginePrivate.h"
+#include "MessageLog.h"
+#include "UObjectToken.h"
+#include "MapErrors.h"
+#include "Engine/SplineMeshActor.h"
+#include "Components/SplineMeshComponent.h"
+
+#define LOCTEXT_NAMESPACE "SplineMeshActor"
+
+ASplineMeshActor::ASplineMeshActor(const FObjectInitializer& ObjectInitializer)
+	: Super(ObjectInitializer)
+{
+	bCanBeDamaged = false;
+
+	SplineMeshComponent = ObjectInitializer.CreateDefaultSubobject<USplineMeshComponent>(this, TEXT("SplineMeshComponent0"));
+	SplineMeshComponent->SetCollisionProfileName(UCollisionProfile::BlockAll_ProfileName);
+	SplineMeshComponent->Mobility = EComponentMobility::Static;
+	SplineMeshComponent->bGenerateOverlapEvents = false;
+
+	RootComponent = SplineMeshComponent;
+}
+
+
+FString ASplineMeshActor::GetDetailedInfoInternal() const
+{
+	check(SplineMeshComponent != nullptr);
+	return SplineMeshComponent->GetDetailedInfoInternal();
+}
+
+void ASplineMeshActor::SetMobility(EComponentMobility::Type InMobility)
+{
+	check(SplineMeshComponent != nullptr);
+	SplineMeshComponent->SetMobility(InMobility);
+}
+
+#if WITH_EDITOR
+
+bool ASplineMeshActor::GetReferencedContentObjects(TArray<UObject*>& Objects) const
+{
+	check(SplineMeshComponent != nullptr);
+	if (SplineMeshComponent->StaticMesh != nullptr)
+	{
+		Objects.Add(SplineMeshComponent->StaticMesh);
+	}
+	return true;
+}
+
+void ASplineMeshActor::CheckForErrors()
+{
+	Super::CheckForErrors();
+
+	FMessageLog MapCheck("MapCheck");
+
+	if (SplineMeshComponent->StaticMesh == NULL)
+	{
+		MapCheck.Warning()
+			->AddToken(FUObjectToken::Create(this))
+			->AddToken(FTextToken::Create(LOCTEXT("MapCheck_Message_SplineMeshNull", "Spline mesh actor has NULL StaticMesh property")))
+			->AddToken(FMapErrorToken::Create(FMapErrors::StaticMeshNull));
+	}
+}
+
+#endif // WITH_EDITOR
+
+/** Returns SplineMeshComponent subobject **/
+USplineMeshComponent* ASplineMeshActor::GetSplineMeshComponent() const { return SplineMeshComponent; }
+
+#undef LOCTEXT_NAMESPACE
