Back out revision 144 from //depot/UE4/Engine/Source/Runtime/Engine/Private/StaticMeshComponent.cpp

- Reverting GetNumMaterials change as that crashes physX cooking. There might be a code that requires that to be exact static mesh

diff --git a/Engine/Source/Runtime/Engine/Private/StaticMeshComponent.cpp b/Engine/Source/Runtime/Engine/Private/StaticMeshComponent.cpp
--- a/Engine/Source/Runtime/Engine/Private/StaticMeshComponent.cpp
+++ b/Engine/Source/Runtime/Engine/Private/StaticMeshComponent.cpp
@@ -1456,16 +1456,14 @@ bool UStaticMeshComponent::GetEstimatedLightAndShadowMapMemoryUsage(
 
 int32 UStaticMeshComponent::GetNumMaterials() const
 {
-	if(Materials.Num() > 0)
+	if(StaticMesh)
 	{
-		return Materials.Num();
+		return StaticMesh->Materials.Num();
 	}
-	else if(StaticMesh)
+	else
 	{
-		return StaticMesh->Materials.Num();
+		return 0;
 	}
-
-	return 0;
 }
 
 TArray<class UMaterialInterface*> UStaticMeshComponent::GetMaterials() const
@@ -1481,7 +1479,6 @@ TArray<class UMaterialInterface*> UStaticMeshComponent::GetMaterials() const
 	return Materials;
 }
 
-
 UMaterialInterface* UStaticMeshComponent::GetMaterial(int32 MaterialIndex) const
 {
 	// If we have a base materials array, use that
