Disable constraint drawing outside interactive editor (limit materials are not loaded)
https://udn.unrealengine.com/questions/227553/show-constraints-crashes-on-xbox-one.html#answer-227611

diff --git a/Engine/Source/Runtime/Engine/Private/PhysicsEngine/PhysDrawing.cpp b/Engine/Source/Runtime/Engine/Private/PhysicsEngine/PhysDrawing.cpp
--- a/Engine/Source/Runtime/Engine/Private/PhysicsEngine/PhysDrawing.cpp
+++ b/Engine/Source/Runtime/Engine/Private/PhysicsEngine/PhysDrawing.cpp
@@ -925,6 +925,14 @@ void FConstraintInstance::DrawConstraint(	FPrimitiveDrawInterface* PDI,
 											float Scale, float LimitDrawScale, bool bDrawLimits, bool bDrawSelected,
 											const FTransform& Con1Frame, const FTransform& Con2Frame, bool bDrawAsPoint ) const
 {
+	// Do nothing unless we are in the interactive editor, otherwise limit materials are not loaded
+	if (!GIsEditor || IsRunningCommandlet())
+	{
+		return;
+	}
+
+	check((GEngine->ConstraintLimitMaterialX != nullptr) && (GEngine->ConstraintLimitMaterialY != nullptr) && (GEngine->ConstraintLimitMaterialZ != nullptr));
+
 	static UMaterialInterface * LimitMaterialX = GEngine->ConstraintLimitMaterialX;
 	static UMaterialInterface * LimitMaterialY = GEngine->ConstraintLimitMaterialY;
 	static UMaterialInterface * LimitMaterialZ = GEngine->ConstraintLimitMaterialZ;
diff --git a/Engine/Source/Runtime/Engine/Public/ShowFlagsValues.inl b/Engine/Source/Runtime/Engine/Public/ShowFlagsValues.inl
--- a/Engine/Source/Runtime/Engine/Public/ShowFlagsValues.inl
+++ b/Engine/Source/Runtime/Engine/Public/ShowFlagsValues.inl
@@ -152,8 +152,8 @@ SHOWFLAG_ALWAYS_ACCESSIBLE(StationaryLightOverlap, SFG_Hidden, LOCTEXT("Stationa
 SHOWFLAG_ALWAYS_ACCESSIBLE(LightMapDensity, SFG_Hidden, LOCTEXT("LightMapDensitySF", "Light Map Density"))
 /** Render streaming bounding volumes for the currently selected texture */
 SHOWFLAG_ALWAYS_ACCESSIBLE(StreamingBounds, SFG_Advanced, LOCTEXT("StreamingBoundsSF", "Streaming Bounds"))
-/**  */
-SHOWFLAG_ALWAYS_ACCESSIBLE(Constraints, SFG_Advanced, LOCTEXT("ConstraintsSF", "Constraints"))
+/** Render joint limits */
+SHOWFLAG_FIXED_IN_SHIPPING(Constraints, 0, SFG_Advanced, LOCTEXT("ConstraintsSF", "Constraints"))
 /** Draws camera frustums */
 SHOWFLAG_ALWAYS_ACCESSIBLE(CameraFrustums, SFG_Advanced, LOCTEXT("CameraFrustumsSF", "Camera Frustums"))
 /** Draw sound actor radii */
