UE-5667 Fix crash passing NULL class to CreateCanvasRenderTarget2D

diff --git a/Engine/Source/Runtime/Engine/Private/CanvasRenderTarget2D.cpp b/Engine/Source/Runtime/Engine/Private/CanvasRenderTarget2D.cpp
--- a/Engine/Source/Runtime/Engine/Private/CanvasRenderTarget2D.cpp
+++ b/Engine/Source/Runtime/Engine/Private/CanvasRenderTarget2D.cpp
@@ -80,7 +80,7 @@ void UCanvasRenderTarget2D::UpdateResource()
 
 UCanvasRenderTarget2D* UCanvasRenderTarget2D::CreateCanvasRenderTarget2D(TSubclassOf<UCanvasRenderTarget2D> CanvasRenderTarget2DClass, int32 Width, int32 Height)
 {
-	if (Width > 0 && Height > 0)
+	if ((Width > 0) && (Height > 0) && (CanvasRenderTarget2DClass != NULL))
 	{
 		UCanvasRenderTarget2D* NewCanvasRenderTarget = ConstructObject<UCanvasRenderTarget2D>(CanvasRenderTarget2DClass, GetTransientPackage());
 		if (NewCanvasRenderTarget)
