Editor/BP - We no longer re-register components on PIE instances of actors when their properties change.  This allows you to modify properties of actors during PIE (through the normal details panel in the level editor) and not lose all your state.  Previously components were re-registered, the events for begin play...etc are not re-run, so a lot of your game state in your blueprint gets broken.

diff --git a/Engine/Source/Runtime/Engine/Classes/GameFramework/Actor.h b/Engine/Source/Runtime/Engine/Classes/GameFramework/Actor.h
--- a/Engine/Source/Runtime/Engine/Classes/GameFramework/Actor.h
+++ b/Engine/Source/Runtime/Engine/Classes/GameFramework/Actor.h
@@ -1165,6 +1165,9 @@ public:
 
 	virtual TSharedPtr<ITransactionObjectAnnotation> GetTransactionAnnotation() const override;
 	virtual void PostEditUndo(TSharedPtr<ITransactionObjectAnnotation> TransactionAnnotation) override;
+
+	/** @return true if the component is allowed to re-register its components when modified.  False for CDOs or PIE instances. */
+	bool ReregisterComponentsWhenModified() const;
 #endif // WITH_EDITOR
 	// End UObject Interface
 
@@ -1626,7 +1629,7 @@ public:
 	/** Called after the actor is spawned in the world.  Responsible for setting up actor for play. */
 	void PostSpawnInitialize(FVector const& SpawnLocation, FRotator const& SpawnRotation, AActor* InOwner, APawn* InInstigator, bool bRemoteOwned, bool bNoFail, bool bDeferConstruction);
 
-    /** Called to finish the spawning process, generally in the case of deferred spawning */
+	/** Called to finish the spawning process, generally in the case of deferred spawning */
 	void FinishSpawning(const FTransform& Transform, bool bIsDefaultTransform = false);
 
 private:
diff --git a/Engine/Source/Runtime/Engine/Private/ActorEditor.cpp b/Engine/Source/Runtime/Engine/Private/ActorEditor.cpp
--- a/Engine/Source/Runtime/Engine/Private/ActorEditor.cpp
+++ b/Engine/Source/Runtime/Engine/Private/ActorEditor.cpp
@@ -19,7 +19,10 @@ void AActor::PreEditChange(UProperty* PropertyThatWillChange)
 {
 	Super::PreEditChange(PropertyThatWillChange);
 
-	UnregisterAllComponents();
+	if ( ReregisterComponentsWhenModified() )
+	{
+		UnregisterAllComponents();
+	}
 }
 
 static FName Name_RelativeLocation = GET_MEMBER_NAME_CHECKED(USceneComponent, RelativeLocation);
@@ -53,7 +56,7 @@ void AActor::PostEditChangeProperty(FPropertyChangedEvent& PropertyChangedEvent)
 		}
 	}
 
-	if(!IsTemplate())
+	if ( ReregisterComponentsWhenModified() )
 	{
 		ReregisterAllComponents();
 
@@ -77,7 +80,7 @@ void AActor::PostEditChangeProperty(FPropertyChangedEvent& PropertyChangedEvent)
 
 void AActor::PostEditMove(bool bFinished)
 {
-	if(!IsTemplate())
+	if ( ReregisterComponentsWhenModified() )
 	{
 		UBlueprint* Blueprint = Cast<UBlueprint>(GetClass()->ClassGeneratedBy);
 		if(Blueprint && (Blueprint->bRunConstructionScriptOnDrag || bFinished) && !FLevelUtils::IsMovingLevel() )
@@ -131,6 +134,10 @@ void AActor::PostEditMove(bool bFinished)
 	}
 }
 
+bool AActor::ReregisterComponentsWhenModified() const
+{
+	return !IsTemplate() && ( GetOutermost()->PackageFlags & PKG_PlayInEditor ) == 0;
+}
 
 void AActor::DebugShowComponentHierarchy(  const TCHAR* Info, bool bShowPosition )
 {	
