[UE-4329] fix deferred component registration during Actor construction

summary:
- registration for all native components occurs in PostSpawnInitialize() (unchanged). all native components are thus registered prior to SCS/UCS execution.
- registration for all SCS components is now deferred until after the SCS has finished executing, in order to ensure the SCS scene hierarchy is always registered before SCS non-scene components. thus, all non-native SCS components are now registered prior to UCS.
- registration for all non-native components added during UCS execution now occurs immediately upon add. thus, line traces and such should now work for all mobility types during UCS since existing components will always be fully registered.
- added UActorComponent::IsRunningUserConstructionScript() to extend the check for static ULightComponent (and derivatives) to allow changes to pass in relevant property set methods during UCS even though the component may have already been registered.
- switched AActor::bRunningUserConstructionScript from 'private' to 'public' in order to allow UActorComponent to query it.
- removed current usage of the FDeferRegisterComponents utility class during Actor construction; the implementation remains in code for now in case it's needed, but may be removed later.
- fixed up AActor::AddComponent() to ensure that root components cannot be more mobile than their children on add from UCS (new scene components by default are attached to the root component).
- added various test case assets for QAGame

diff --git a/Engine/Source/Runtime/Engine/Classes/Components/ActorComponent.h b/Engine/Source/Runtime/Engine/Classes/Components/ActorComponent.h
--- a/Engine/Source/Runtime/Engine/Classes/Components/ActorComponent.h
+++ b/Engine/Source/Runtime/Engine/Classes/Components/ActorComponent.h
@@ -466,6 +466,9 @@ public:
 	virtual UAssetUserData* GetAssetUserDataOfClass(TSubclassOf<UAssetUserData> InUserDataClass) override;
 	// End IInterface_AssetUserData Interface
 
+	/** See if the owning Actor is currently running the UCS */
+	bool IsRunningUserConstructionScript() const;
+
 	/** See if this component is currently registered */
 	FORCEINLINE bool IsRegistered() const
 	{
diff --git a/Engine/Source/Runtime/Engine/Classes/Engine/SCS_Node.h b/Engine/Source/Runtime/Engine/Classes/Engine/SCS_Node.h
--- a/Engine/Source/Runtime/Engine/Classes/Engine/SCS_Node.h
+++ b/Engine/Source/Runtime/Engine/Classes/Engine/SCS_Node.h
@@ -94,8 +94,9 @@ class USCS_Node : public UObject
 	 * @param ParentComponent		If non-NULL, the component to which the new component should be attached as a child. If NULL, the new component will not be attached to anything.
 	 * @param RootTransform			The transform to apply if this node turns out to be the root component of the actor instance.
 	 * @param bIsDefaultTransform	Indicates whether or not the given transform is a "default" transform, in which case it can be overridden by template defaults.
+	 * @return The new component instance that was created, or NULL on failure.
 	 */
-	void ExecuteNodeOnActor(AActor* Actor, USceneComponent* ParentComponent, const FTransform* RootTransform, bool bIsDefaultTransform);
+	UActorComponent* ExecuteNodeOnActor(AActor* Actor, USceneComponent* ParentComponent, const FTransform* RootTransform, bool bIsDefaultTransform);
 
 	/** Returns an array containing this node and all children below it */
 	TArray<USCS_Node*> GetAllNodes();
diff --git a/Engine/Source/Runtime/Engine/Private/ActorComponent.cpp b/Engine/Source/Runtime/Engine/Private/ActorComponent.cpp
--- a/Engine/Source/Runtime/Engine/Private/ActorComponent.cpp
+++ b/Engine/Source/Runtime/Engine/Private/ActorComponent.cpp
@@ -1105,6 +1105,12 @@ bool UActorComponent::IsActive() const
 	return bIsActive;
 }
 
+bool UActorComponent::IsRunningUserConstructionScript() const
+{
+	const AActor* Owner = GetOwner();
+	return Owner != nullptr && Owner->bRunningUserConstructionScript;
+}
+
 void UActorComponent::AddAssetUserData(UAssetUserData* InUserData)
 {
 	if (InUserData != NULL)
diff --git a/Engine/Source/Runtime/Engine/Private/ActorConstruction.cpp b/Engine/Source/Runtime/Engine/Private/ActorConstruction.cpp
--- a/Engine/Source/Runtime/Engine/Private/ActorConstruction.cpp
+++ b/Engine/Source/Runtime/Engine/Private/ActorConstruction.cpp
@@ -423,34 +423,26 @@ void AActor::ExecuteConstruction(const FTransform& Transform, const FComponentIn
 	OnConstruction(Transform);
 }
 
-
-void FDeferRegisterComponents::DeferComponentRegistration(AActor* Actor, USceneComponent* Component, EComponentMobility::Type OriginalMobility)
+void FDeferRegisterComponents::DeferComponentRegistration(AActor* Actor, UActorComponent* Component)
 {
-	TArray<FDeferredComponentInfo>* ActorComponentsToRegister = ComponentsToRegister.Find(Actor);
+	TArray<UActorComponent*>* ActorComponentsToRegister = ComponentsToRegister.Find(Actor);
 	if (!ActorComponentsToRegister)
 	{
-		ActorComponentsToRegister = &ComponentsToRegister.Add(Actor, TArray<FDeferredComponentInfo>());
+		ActorComponentsToRegister = &ComponentsToRegister.Add(Actor, TArray<UActorComponent*>());
 	}
-	ActorComponentsToRegister->Add(FDeferredComponentInfo(Component, OriginalMobility));
 	
-	// Components with Mobility set to EComponentMobility::Static or EComponentMobility::Stationary 
-	// can't be properly set up from a construction script (all changes will be rejected due to 
-	// EComponentMobility::Static flag) so we're going to temporarily change the flag and defer 
-	// the registration until the construction-script has finished.
-	Component->Mobility = EComponentMobility::Movable;
+	ActorComponentsToRegister->Add(Component);
 }
 
 void FDeferRegisterComponents::RegisterComponents(AActor* Actor)
 {
-	TArray<FDeferredComponentInfo>* ActorComponentsToRegister = ComponentsToRegister.Find(Actor);
+	TArray<UActorComponent*>* ActorComponentsToRegister = ComponentsToRegister.Find(Actor);
 	if (ActorComponentsToRegister)
 	{
 		for (int32 ComponentIndex = 0; ComponentIndex < ActorComponentsToRegister->Num(); ++ComponentIndex)
 		{
-			const FDeferredComponentInfo& SavedInfo = (*ActorComponentsToRegister)[ComponentIndex];
-			USceneComponent* Component = SavedInfo.Component;
-			// Restore saved mobility just before registering this component.
-			Component->Mobility = SavedInfo.SavedMobility;
+			UActorComponent* Component = (*ActorComponentsToRegister)[ComponentIndex];
+			
 			Component->RegisterComponent();
 		}
 		ComponentsToRegister.Remove(Actor);
@@ -460,13 +452,13 @@ void FDeferRegisterComponents::RegisterComponents(AActor* Actor)
 // FGCObject interface
 void FDeferRegisterComponents::AddReferencedObjects(FReferenceCollector& Collector)
 {
-	for (TMap<AActor*, TArray<FDeferredComponentInfo> >::TIterator It(ComponentsToRegister); It; ++It)
+	for (TMap<AActor*, TArray<UActorComponent*> >::TIterator It(ComponentsToRegister); It; ++It)
 	{
 		Collector.AddReferencedObject(It.Key());
-		TArray<FDeferredComponentInfo>& Components = It.Value();
+		TArray<UActorComponent*>& Components = It.Value();
 		for (int32 ComponentIndex = 0; ComponentIndex < Components.Num(); ++ComponentIndex)
 		{
-			Collector.AddReferencedObject(Components[ComponentIndex].Component);
+			Collector.AddReferencedObject(Components[ComponentIndex]);
 		}
 	}
 }
@@ -484,9 +476,6 @@ void AActor::ProcessUserConstructionScript()
 	bRunningUserConstructionScript = true;
 	UserConstructionScript();
 	bRunningUserConstructionScript = false;
-
-	// Register all deferred components for this actor.
-	FDeferRegisterComponents::Get().RegisterComponents(this);
 }
 
 void AActor::FinishAndRegisterComponent(UActorComponent* Component)
@@ -514,12 +503,12 @@ UActorComponent* AActor::CreateComponentFromTemplate(UActorComponent* Template,
 
 UActorComponent* AActor::AddComponent(FName TemplateName, bool bManualAttachment, const FTransform& RelativeTransform, const UObject* ComponentTemplateContext)
 {
-	UActorComponent* Template = NULL;
-	UBlueprintGeneratedClass* BlueprintGeneratedClass = Cast<UBlueprintGeneratedClass>((ComponentTemplateContext != NULL) ? ComponentTemplateContext->GetClass() : GetClass());
-	while(BlueprintGeneratedClass != NULL)
+	UActorComponent* Template = nullptr;
+	UBlueprintGeneratedClass* BlueprintGeneratedClass = Cast<UBlueprintGeneratedClass>((ComponentTemplateContext != nullptr) ? ComponentTemplateContext->GetClass() : GetClass());
+	while(BlueprintGeneratedClass != nullptr)
 	{
 		Template = BlueprintGeneratedClass->FindComponentTemplateByName(TemplateName);
-		if(NULL != Template)
+		if(nullptr != Template)
 		{
 			break;
 		}
@@ -527,33 +516,36 @@ UActorComponent* AActor::AddComponent(FName TemplateName, bool bManualAttachment
 	}
 
 	UActorComponent* NewActorComp = CreateComponentFromTemplate(Template);
-	if(NewActorComp != NULL)
+	if(NewActorComp != nullptr)
 	{
 		// The user has the option of doing attachment manually where they have complete control or via the automatic rule
 		// that the first component added becomes the root component, with subsequent components attached to the root.
 		USceneComponent* NewSceneComp = Cast<USceneComponent>(NewActorComp);
-		bool bDeferRegisterStaticComponent = false;
-		EComponentMobility::Type OriginalMobility = EComponentMobility::Movable;
-
-		if(NewSceneComp != NULL)
+		if(NewSceneComp != nullptr)
 		{
-			// Components with Mobility set to EComponentMobility::Static or EComponentMobility::Stationary can't be properly set up in UCS (all changes will be rejected
-			// due to EComponentMobility::Static flag) so we're going to temporarily change the flag and defer the registration until UCS has finished.
-			bDeferRegisterStaticComponent = bRunningUserConstructionScript && NewSceneComp->Mobility != EComponentMobility::Movable;
-			OriginalMobility = NewSceneComp->Mobility;
-			if (bDeferRegisterStaticComponent)
-			{
-				NewSceneComp->Mobility = EComponentMobility::Movable;
-			}
-
 			if (!bManualAttachment)
 			{
-				if (RootComponent == NULL)
+				if (RootComponent == nullptr)
 				{
 					RootComponent = NewSceneComp;
 				}
 				else
 				{
+					// The root component can't be more mobile than its children, so we check for that here and adjust as needed before attaching.
+					if(RootComponent->Mobility > NewSceneComp->Mobility)
+					{
+						if(NewSceneComp->IsA<UStaticMeshComponent>())
+						{
+							// SMCs can't be stationary, so always set them to be movable
+							NewSceneComp->SetMobility(EComponentMobility::Movable);
+						}
+						else
+						{
+							// Set the new component to be at least as mobile as the root
+							NewSceneComp->SetMobility(RootComponent->Mobility);
+						}
+					}
+
 					NewSceneComp->AttachTo(RootComponent);
 				}
 			}
@@ -564,16 +556,8 @@ UActorComponent* AActor::AddComponent(FName TemplateName, bool bManualAttachment
 		// Call function to notify component it has been created
 		NewActorComp->OnComponentCreated();
 
-		if (bDeferRegisterStaticComponent)
-		{
-			// Defer registration until after UCS has completed.
-			FDeferRegisterComponents::Get().DeferComponentRegistration(this, NewSceneComp, OriginalMobility);
-		}
-		else
-		{
-			// Register component, which will create physics/rendering state, now component is in correct position
-			NewActorComp->RegisterComponent();
-		}
+		// Register component, which will create physics/rendering state, now component is in correct position
+		NewActorComp->RegisterComponent();
 	}
 
 	return NewActorComp;
diff --git a/Engine/Source/Runtime/Engine/Private/DeferRegisterComponents.h b/Engine/Source/Runtime/Engine/Private/DeferRegisterComponents.h
--- a/Engine/Source/Runtime/Engine/Private/DeferRegisterComponents.h
+++ b/Engine/Source/Runtime/Engine/Private/DeferRegisterComponents.h
@@ -2,19 +2,6 @@
 
 #pragma once
 
-
-class FDeferredComponentInfo
-{
-public:
-	USceneComponent* Component;
-	EComponentMobility::Type SavedMobility;
-
-	FDeferredComponentInfo(USceneComponent* InComponent, EComponentMobility::Type OriginalMobility) :
-		Component(InComponent),
-		SavedMobility(OriginalMobility)
-	{}
-};
-
 /**
  * Helper class to store components which registration should be deferred until
  * the construction script has completed... generally, all components get stuck here
@@ -24,7 +11,7 @@ public:
 class FDeferRegisterComponents : public FGCObject
 {
 	/** Map of actors and their components to register. */
-	TMap<AActor*, TArray<FDeferredComponentInfo> > ComponentsToRegister;
+	TMap<AActor*, TArray<UActorComponent*> > ComponentsToRegister;
 
 public:
 
@@ -34,7 +21,7 @@ public:
 	 * @param Actor Actor this component belongs to.
 	 * @param Component Component to register.
 	 */
-	void DeferComponentRegistration(AActor* Actor, USceneComponent* Component, EComponentMobility::Type OriginalMobility);
+	void DeferComponentRegistration(AActor* Actor, UActorComponent* Component);
 
 	/**
 	 * Registers all deferred components that belong to the specified actor.
diff --git a/Engine/Source/Runtime/Engine/Private/DirectionalLightComponent.cpp b/Engine/Source/Runtime/Engine/Private/DirectionalLightComponent.cpp
--- a/Engine/Source/Runtime/Engine/Private/DirectionalLightComponent.cpp
+++ b/Engine/Source/Runtime/Engine/Private/DirectionalLightComponent.cpp
@@ -797,7 +797,7 @@ void UDirectionalLightComponent::SetShadowDistanceFadeoutFraction(float NewValue
 
 void UDirectionalLightComponent::SetEnableLightShaftOcclusion(bool bNewValue)
 {
-	if (!(IsRegistered() && Mobility == EComponentMobility::Static)
+	if (!(IsRegistered() && Mobility == EComponentMobility::Static && !IsRunningUserConstructionScript())
 		&& bEnableLightShaftOcclusion != bNewValue)
 	{
 		bEnableLightShaftOcclusion = bNewValue;
@@ -807,7 +807,7 @@ void UDirectionalLightComponent::SetEnableLightShaftOcclusion(bool bNewValue)
 
 void UDirectionalLightComponent::SetOcclusionMaskDarkness(float NewValue)
 {
-	if (!(IsRegistered() && Mobility == EComponentMobility::Static)
+	if (!(IsRegistered() && Mobility == EComponentMobility::Static && !IsRunningUserConstructionScript())
 		&& OcclusionMaskDarkness != NewValue)
 	{
 		OcclusionMaskDarkness = NewValue;
@@ -817,7 +817,7 @@ void UDirectionalLightComponent::SetOcclusionMaskDarkness(float NewValue)
 
 void UDirectionalLightComponent::SetLightShaftOverrideDirection(FVector NewValue)
 {
-	if (!(IsRegistered() && Mobility == EComponentMobility::Static)
+	if (!(IsRegistered() && Mobility == EComponentMobility::Static && !IsRunningUserConstructionScript())
 		&& LightShaftOverrideDirection != NewValue)
 	{
 		LightShaftOverrideDirection = NewValue;
diff --git a/Engine/Source/Runtime/Engine/Private/LightComponent.cpp b/Engine/Source/Runtime/Engine/Private/LightComponent.cpp
--- a/Engine/Source/Runtime/Engine/Private/LightComponent.cpp
+++ b/Engine/Source/Runtime/Engine/Private/LightComponent.cpp
@@ -70,7 +70,7 @@ FArchive& operator<<(FArchive& Ar, FStaticShadowDepthMap& ShadowMap)
 
 void ULightComponentBase::SetCastShadows(bool bNewValue)
 {
-	if (!(IsRegistered() && Mobility == EComponentMobility::Static)
+	if (!(IsRegistered() && Mobility == EComponentMobility::Static && !IsRunningUserConstructionScript())
 		&& CastShadows != bNewValue)
 	{
 		CastShadows = bNewValue;
@@ -669,7 +669,7 @@ void ULightComponent::DestroyRenderState_Concurrent()
 void ULightComponent::SetIntensity(float NewIntensity)
 {
 	// Can't set brightness on a static light
-	if (!(IsRegistered() && Mobility == EComponentMobility::Static)
+	if (!(IsRegistered() && Mobility == EComponentMobility::Static && !IsRunningUserConstructionScript())
 		&& Intensity != NewIntensity)
 	{
 		Intensity = NewIntensity;
@@ -689,7 +689,7 @@ void ULightComponent::SetLightColor(FLinearColor NewLightColor)
 	FColor NewColor(NewLightColor);
 
 	// Can't set color on a static light
-	if (!(IsRegistered() && Mobility == EComponentMobility::Static)
+	if (!(IsRegistered() && Mobility == EComponentMobility::Static && !IsRunningUserConstructionScript())
 		&& LightColor != NewColor)
 	{
 		LightColor	= NewColor;
@@ -706,7 +706,7 @@ void ULightComponent::SetLightColor(FLinearColor NewLightColor)
 void ULightComponent::SetLightFunctionMaterial(UMaterialInterface* NewLightFunctionMaterial)
 {
 	// Can't set light function on a static light
-	if (!(IsRegistered() && Mobility == EComponentMobility::Static)
+	if (!(IsRegistered() && Mobility == EComponentMobility::Static && !IsRunningUserConstructionScript())
 		&& NewLightFunctionMaterial != LightFunctionMaterial)
 	{
 		LightFunctionMaterial = NewLightFunctionMaterial;
@@ -716,7 +716,7 @@ void ULightComponent::SetLightFunctionMaterial(UMaterialInterface* NewLightFunct
 
 void ULightComponent::SetLightFunctionScale(FVector NewLightFunctionScale)
 {
-	if (!(IsRegistered() && Mobility == EComponentMobility::Static)
+	if (!(IsRegistered() && Mobility == EComponentMobility::Static && !IsRunningUserConstructionScript())
 		&& NewLightFunctionScale != LightFunctionScale)
 	{
 		LightFunctionScale = NewLightFunctionScale;
@@ -726,7 +726,7 @@ void ULightComponent::SetLightFunctionScale(FVector NewLightFunctionScale)
 
 void ULightComponent::SetLightFunctionFadeDistance(float NewLightFunctionFadeDistance)
 {
-	if (!(IsRegistered() && Mobility == EComponentMobility::Static)
+	if (!(IsRegistered() && Mobility == EComponentMobility::Static && !IsRunningUserConstructionScript())
 		&& NewLightFunctionFadeDistance != LightFunctionFadeDistance)
 	{
 		LightFunctionFadeDistance = NewLightFunctionFadeDistance;
@@ -736,7 +736,7 @@ void ULightComponent::SetLightFunctionFadeDistance(float NewLightFunctionFadeDis
 
 void ULightComponent::SetAffectDynamicIndirectLighting(bool bNewValue)
 {
-	if (!(IsRegistered() && Mobility == EComponentMobility::Static)
+	if (!(IsRegistered() && Mobility == EComponentMobility::Static && !IsRunningUserConstructionScript())
 		&& bAffectDynamicIndirectLighting != bNewValue)
 	{
 		bAffectDynamicIndirectLighting = bNewValue;
@@ -746,7 +746,7 @@ void ULightComponent::SetAffectDynamicIndirectLighting(bool bNewValue)
 
 void ULightComponent::SetAffectTranslucentLighting(bool bNewValue)
 {
-	if (!(IsRegistered() && Mobility == EComponentMobility::Static)
+	if (!(IsRegistered() && Mobility == EComponentMobility::Static && !IsRunningUserConstructionScript())
 		&& bAffectTranslucentLighting != bNewValue)
 	{
 		bAffectTranslucentLighting = bNewValue;
@@ -756,7 +756,7 @@ void ULightComponent::SetAffectTranslucentLighting(bool bNewValue)
 
 void ULightComponent::SetEnableLightShaftBloom(bool bNewValue)
 {
-	if (!(IsRegistered() && Mobility == EComponentMobility::Static)
+	if (!(IsRegistered() && Mobility == EComponentMobility::Static && !IsRunningUserConstructionScript())
 		&& bEnableLightShaftBloom != bNewValue)
 	{
 		bEnableLightShaftBloom = bNewValue;
@@ -766,7 +766,7 @@ void ULightComponent::SetEnableLightShaftBloom(bool bNewValue)
 
 void ULightComponent::SetBloomScale(float NewValue)
 {
-	if (!(IsRegistered() && Mobility == EComponentMobility::Static)
+	if (!(IsRegistered() && Mobility == EComponentMobility::Static && !IsRunningUserConstructionScript())
 		&& BloomScale != NewValue)
 	{
 		BloomScale = NewValue;
@@ -776,7 +776,7 @@ void ULightComponent::SetBloomScale(float NewValue)
 
 void ULightComponent::SetBloomThreshold(float NewValue)
 {
-	if (!(IsRegistered() && Mobility == EComponentMobility::Static)
+	if (!(IsRegistered() && Mobility == EComponentMobility::Static && !IsRunningUserConstructionScript())
 		&& BloomThreshold != NewValue)
 	{
 		BloomThreshold = NewValue;
@@ -786,7 +786,7 @@ void ULightComponent::SetBloomThreshold(float NewValue)
 
 void ULightComponent::SetBloomTint(FColor NewValue)
 {
-	if (!(IsRegistered() && Mobility == EComponentMobility::Static)
+	if (!(IsRegistered() && Mobility == EComponentMobility::Static && !IsRunningUserConstructionScript())
 		&& BloomTint != NewValue)
 	{
 		BloomTint = NewValue;
@@ -796,7 +796,7 @@ void ULightComponent::SetBloomTint(FColor NewValue)
 
 void ULightComponent::SetIESTexture(UTextureLightProfile* NewValue)
 {
-	if (!(IsRegistered() && Mobility == EComponentMobility::Static)
+	if (!(IsRegistered() && Mobility == EComponentMobility::Static && !IsRunningUserConstructionScript())
 		&& IESTexture != NewValue)
 	{
 		IESTexture = NewValue;
diff --git a/Engine/Source/Runtime/Engine/Private/PointLightComponent.cpp b/Engine/Source/Runtime/Engine/Private/PointLightComponent.cpp
--- a/Engine/Source/Runtime/Engine/Private/PointLightComponent.cpp
+++ b/Engine/Source/Runtime/Engine/Private/PointLightComponent.cpp
@@ -130,7 +130,7 @@ FLightSceneProxy* UPointLightComponent::CreateSceneProxy() const
 void UPointLightComponent::SetAttenuationRadius(float NewRadius)
 {
 	// Only movable lights can change their radius at runtime
-	if (!(IsRegistered() && (Mobility == EComponentMobility::Static || Mobility == EComponentMobility::Stationary))
+	if (!(IsRegistered() && (Mobility == EComponentMobility::Static || Mobility == EComponentMobility::Stationary) && !IsRunningUserConstructionScript())
 		&& NewRadius != AttenuationRadius)
 	{
 		AttenuationRadius = NewRadius;
@@ -140,7 +140,7 @@ void UPointLightComponent::SetAttenuationRadius(float NewRadius)
 
 void UPointLightComponent::SetLightFalloffExponent(float NewLightFalloffExponent)
 {
-	if (!(IsRegistered() && Mobility == EComponentMobility::Static)
+	if (!(IsRegistered() && Mobility == EComponentMobility::Static && !IsRunningUserConstructionScript())
 		&& NewLightFalloffExponent != LightFalloffExponent)
 	{
 		LightFalloffExponent = NewLightFalloffExponent;
@@ -150,7 +150,7 @@ void UPointLightComponent::SetLightFalloffExponent(float NewLightFalloffExponent
 
 void UPointLightComponent::SetSourceRadius(float bNewValue)
 {
-	if (!(IsRegistered() && Mobility == EComponentMobility::Static)
+	if (!(IsRegistered() && Mobility == EComponentMobility::Static && !IsRunningUserConstructionScript())
 		&& SourceRadius != bNewValue)
 	{
 		SourceRadius = bNewValue;
diff --git a/Engine/Source/Runtime/Engine/Private/SCS_Node.cpp b/Engine/Source/Runtime/Engine/Private/SCS_Node.cpp
--- a/Engine/Source/Runtime/Engine/Private/SCS_Node.cpp
+++ b/Engine/Source/Runtime/Engine/Private/SCS_Node.cpp
@@ -26,36 +26,24 @@ USCS_Node::USCS_Node(const FObjectInitializer& ObjectInitializer)
 #endif
 }
 
-void USCS_Node::ExecuteNodeOnActor(AActor* Actor, USceneComponent* ParentComponent, const FTransform* RootTransform, bool bIsDefaultTransform)
+UActorComponent* USCS_Node::ExecuteNodeOnActor(AActor* Actor, USceneComponent* ParentComponent, const FTransform* RootTransform, bool bIsDefaultTransform)
 {
-	check(Actor != NULL);
-	check((ParentComponent != NULL && !ParentComponent->IsPendingKill()) || (RootTransform != NULL)); // must specify either a parent component or a world transform
+	check(Actor != nullptr);
+	check((ParentComponent != nullptr && !ParentComponent->IsPendingKill()) || (RootTransform != nullptr)); // must specify either a parent component or a world transform
 
 	// Create a new component instance based on the template
 	UActorComponent* NewActorComp = Actor->CreateComponentFromTemplate(ComponentTemplate, VariableName.ToString());
-	if(NewActorComp != NULL)
+	if(NewActorComp != nullptr)
 	{
-		bool bDeferRegisterStaticComponent = false;
-		EComponentMobility::Type OriginalMobility = EComponentMobility::Movable;
-
 		// SCS created components are net addressable
 		NewActorComp->SetNetAddressable();
 
 		// Special handling for scene components
 		USceneComponent* NewSceneComp = Cast<USceneComponent>(NewActorComp);
-		if (NewSceneComp != NULL)
+		if (NewSceneComp != nullptr)
 		{
-			// Components with Mobility set to EComponentMobility::Static or EComponentMobility::Stationary can't be properly set up in SCS (all changes will be rejected
-			// due to EComponentMobility::Static flag) so we're going to temporarily change the flag and defer the registration until SCS has finished.
-			bDeferRegisterStaticComponent = NewSceneComp->Mobility != EComponentMobility::Movable;
-			OriginalMobility = NewSceneComp->Mobility;
-			if (bDeferRegisterStaticComponent)
-			{
-				NewSceneComp->Mobility = EComponentMobility::Movable;
-			}
-
 			// If NULL is passed in, we are the root, so set transform and assign as RootComponent on Actor
-			if (ParentComponent == NULL || (ParentComponent && ParentComponent->IsPendingKill()))
+			if (ParentComponent == nullptr || (ParentComponent && ParentComponent->IsPendingKill()))
 			{
 				FTransform WorldTransform = *RootTransform;
 				if(bIsDefaultTransform)
@@ -76,17 +64,6 @@ void USCS_Node::ExecuteNodeOnActor(AActor* Actor, USceneComponent* ParentCompone
 		// Call function to notify component it has been created
 		NewActorComp->OnComponentCreated();
 
-		if (bDeferRegisterStaticComponent)
-		{
-			// need to defer component registration until after the construction script
-			// is ran, since the construction script can mutate the object (for collision, etc.)
-			FDeferRegisterComponents::Get().DeferComponentRegistration(Actor, NewSceneComp, OriginalMobility);
-		}
-		else
-		{
-			NewActorComp->RegisterComponent();
-		}
-
 		if (NewActorComp->GetIsReplicated())
 		{
 			// Make sure this component is added to owning actor's replicated list.
@@ -108,7 +85,7 @@ void USCS_Node::ExecuteNodeOnActor(AActor* Actor, USceneComponent* ParentCompone
 #if WITH_EDITOR
 				// If we're constructing editable components in the SCS editor, set the component instance corresponding to this node for editing purposes
 				USimpleConstructionScript* SCS = GetSCS();
-				if(SCS != NULL && (SCS->IsConstructingEditorComponents() || SCS->GetComponentEditorActorInstance() == Actor))
+				if(SCS != nullptr && (SCS->IsConstructingEditorComponents() || SCS->GetComponentEditorActorInstance() == Actor))
 				{
 					EditorComponentInstance = NewSceneComp;
 				}
@@ -117,16 +94,18 @@ void USCS_Node::ExecuteNodeOnActor(AActor* Actor, USceneComponent* ParentCompone
 		}
 
 		// Determine the parent component for our children (it's still our parent if we're a non-scene component)
-		USceneComponent* ParentSceneComponentOfChildren = (NewSceneComp != NULL) ? NewSceneComp : ParentComponent;
+		USceneComponent* ParentSceneComponentOfChildren = (NewSceneComp != nullptr) ? NewSceneComp : ParentComponent;
 
 		// If we made a component, go ahead and process our children
 		for (int32 NodeIdx = 0; NodeIdx < ChildNodes.Num(); NodeIdx++)
 		{
 			USCS_Node* Node = ChildNodes[NodeIdx];
-			check(Node != NULL);
-			Node->ExecuteNodeOnActor(Actor, ParentSceneComponentOfChildren, NULL, false);
+			check(Node != nullptr);
+			Node->ExecuteNodeOnActor(Actor, ParentSceneComponentOfChildren, nullptr, false);
 		}
 	}
+
+	return NewActorComp;
 }
 
 TArray<USCS_Node*> USCS_Node::GetAllNodes()
diff --git a/Engine/Source/Runtime/Engine/Private/SimpleConstructionScript.cpp b/Engine/Source/Runtime/Engine/Private/SimpleConstructionScript.cpp
--- a/Engine/Source/Runtime/Engine/Private/SimpleConstructionScript.cpp
+++ b/Engine/Source/Runtime/Engine/Private/SimpleConstructionScript.cpp
@@ -294,6 +294,7 @@ void USimpleConstructionScript::ExecuteScriptOnActor(AActor* Actor, const FTrans
 {
 	if(RootNodes.Num() > 0)
 	{
+		TArray<UActorComponent*> InstancedComponents;
 		for(auto NodeIt = RootNodes.CreateIterator(); NodeIt; ++NodeIt)
 		{
 			USCS_Node* RootNode = *NodeIt;
@@ -335,7 +336,29 @@ void USimpleConstructionScript::ExecuteScriptOnActor(AActor* Actor, const FTrans
 					}
 				}
 
-				RootNode->ExecuteNodeOnActor(Actor, ParentComponent != NULL ? ParentComponent : Actor->GetRootComponent(), &RootTransform, bIsDefaultTransform);
+				// Create the new component instance and any child components it may have
+				UActorComponent* InstancedComponent = RootNode->ExecuteNodeOnActor(Actor, ParentComponent != NULL ? ParentComponent : Actor->GetRootComponent(), &RootTransform, bIsDefaultTransform);
+				if(InstancedComponent != nullptr)
+				{
+					InstancedComponents.Add(InstancedComponent);
+				}
+			}
+		}
+
+		// Register all instanced SCS components once SCS execution has finished; do the scene component hierarchy first, followed by the remaining actor components (in case they happen to depend on something in the scene hierarchy)
+		InstancedComponents.Sort([](const UActorComponent& A, const UActorComponent& B) { return A.IsA<USceneComponent>(); });
+		for(auto InstancedComponent : InstancedComponents)
+		{
+			InstancedComponent->RegisterComponent();
+
+			// If this is a scene component, register any child components as well
+			USceneComponent* InstancedSceneComponent = Cast<USceneComponent>(InstancedComponent);
+			if(InstancedSceneComponent != nullptr)
+			{
+				for(auto InstancedChildComponent : InstancedSceneComponent->AttachChildren)
+				{
+					InstancedChildComponent->RegisterComponent();
+				}
 			}
 		}
 	}
diff --git a/Engine/Source/Runtime/Engine/Private/SkyLightComponent.cpp b/Engine/Source/Runtime/Engine/Private/SkyLightComponent.cpp
--- a/Engine/Source/Runtime/Engine/Private/SkyLightComponent.cpp
+++ b/Engine/Source/Runtime/Engine/Private/SkyLightComponent.cpp
@@ -433,7 +433,7 @@ void USkyLightComponent::CaptureEmissiveIrradianceEnvironmentMap(FSHVectorRGB3&
 void USkyLightComponent::SetIntensity(float NewIntensity)
 {
 	// Can't set brightness on a static light
-	if (!(IsRegistered() && Mobility == EComponentMobility::Static)
+	if (!(IsRegistered() && Mobility == EComponentMobility::Static && !IsRunningUserConstructionScript())
 		&& Intensity != NewIntensity)
 	{
 		Intensity = NewIntensity;
@@ -447,7 +447,7 @@ void USkyLightComponent::SetLightColor(FLinearColor NewLightColor)
 	FColor NewColor(NewLightColor);
 
 	// Can't set color on a static light
-	if (!(IsRegistered() && Mobility == EComponentMobility::Static)
+	if (!(IsRegistered() && Mobility == EComponentMobility::Static && !IsRunningUserConstructionScript())
 		&& LightColor != NewColor)
 	{
 		LightColor = NewColor;
@@ -458,7 +458,7 @@ void USkyLightComponent::SetLightColor(FLinearColor NewLightColor)
 void USkyLightComponent::SetCubemap(UTextureCube* NewCubemap)
 {
 	// Can't set color on a static light
-	if (!(IsRegistered() && Mobility == EComponentMobility::Static)
+	if (!(IsRegistered() && Mobility == EComponentMobility::Static && !IsRunningUserConstructionScript())
 		&& Cubemap != NewCubemap)
 	{
 		Cubemap = NewCubemap;
@@ -470,7 +470,7 @@ void USkyLightComponent::SetCubemap(UTextureCube* NewCubemap)
 void USkyLightComponent::SetOcclusionTint(const FColor& InTint)
 {
 	// Can't set on a static light
-	if (!(IsRegistered() && Mobility == EComponentMobility::Static)
+	if (!(IsRegistered() && Mobility == EComponentMobility::Static && !IsRunningUserConstructionScript())
 		&& OcclusionTint != InTint)
 	{
 		OcclusionTint = InTint;
@@ -481,7 +481,7 @@ void USkyLightComponent::SetOcclusionTint(const FColor& InTint)
 void USkyLightComponent::SetMinOcclusion(float InMinOcclusion)
 {
 	// Can't set on a static light
-	if (!(IsRegistered() && Mobility == EComponentMobility::Static)
+	if (!(IsRegistered() && Mobility == EComponentMobility::Static && !IsRunningUserConstructionScript())
 		&& MinOcclusion != InMinOcclusion)
 	{
 		MinOcclusion = InMinOcclusion;
