Remove FAudioDevice::SetClassVolume and console command ModifySoundClass

diff --git a/Engine/Source/Runtime/Engine/Private/AudioDevice.cpp b/Engine/Source/Runtime/Engine/Private/AudioDevice.cpp
--- a/Engine/Source/Runtime/Engine/Private/AudioDevice.cpp
+++ b/Engine/Source/Runtime/Engine/Private/AudioDevice.cpp
@@ -679,43 +679,6 @@ bool FAudioDevice::HandleResetSoundStateCommand( const TCHAR* Cmd, FOutputDevice
 	DebugState = DEBUGSTATE_None;
 	return true;
 }
-
-bool FAudioDevice::HandleModifySoundClassCommand( const TCHAR* Cmd, FOutputDevice& Ar )
-{
-	const FString SoundClassName = FParse::Token( Cmd, 0 );
-	float NewVolume = -1.0f;
-	FParse::Value( Cmd, TEXT( "Vol=" ), NewVolume );
-
-	// Set the volume in the original sound class
-	TArray<USoundClass*> FoundSoundClasses;
-	for (TMap<USoundClass*, FSoundClassProperties>::TIterator It(SoundClasses); It; ++It)
-	{
-		USoundClass* SoundClass = It.Key();
-		if (SoundClass && SoundClass->GetName() == SoundClassName)
-		{
-			FoundSoundClasses.Add(It.Key());
-		}
-	}
-
-	if (FoundSoundClasses.Num() == 1)
-	{
-		SetClassVolume( FoundSoundClasses[0], NewVolume );
-	}
-	else if (FoundSoundClasses.Num() == 0)
-	{
-		UE_LOG(LogAudio, Log, TEXT( "Couldn't find specified sound class (%s)!" ), *SoundClassName );
-	}
-	else
-	{
-		// If we found multiples we will set them all for now, but perhaps we could report an error message
-		// and go ahead and require them to supply an index for the one they want or -1 for all?
-		for(int32 Index = 0; Index < FoundSoundClasses.Num(); ++Index)
-		{
-			SetClassVolume( FoundSoundClasses[Index], NewVolume );
-		}
-	}
-	return true;
-}
 #endif // !UE_BUILD_SHIPPING
 
 EDebugState FAudioDevice::GetMixDebugState( void )
@@ -806,11 +769,6 @@ bool FAudioDevice::Exec( UWorld* InWorld, const TCHAR* Cmd, FOutputDevice& Ar )
 	{
 		return HandleResetSoundStateCommand( Cmd, Ar );
 	}
-	// usage ModifySoundClass <soundclassname> vol=<new volume>
-	else if( FParse::Command( &Cmd, TEXT( "ModifySoundClass" ) ) )
-	{
-		return HandleModifySoundClassCommand( Cmd, Ar );
-	}
 #endif // !UE_BUILD_SHIPPING
 
 	return false;
@@ -830,8 +788,7 @@ void FAudioDevice::InitSoundClasses( void )
 	for( TObjectIterator<USoundClass> It; It; ++It )
 	{
 		USoundClass* SoundClass = *It;
-		FSoundClassProperties& Properties = SoundClasses.Add( SoundClass, FSoundClassProperties() );
-		Properties = SoundClass->Properties;
+		FSoundClassProperties& Properties = SoundClasses.Add( SoundClass, SoundClass->Properties );
 	}
 
 	// Propagate the properties down the hierarchy
@@ -2045,14 +2002,6 @@ void FAudioDevice::RemoveActiveSound(FActiveSound* ActiveSound)
 	}
 }
 
-void FAudioDevice::SetClassVolume( USoundClass* InSoundClass, const float Volume )
-{
-	if (InSoundClass)
-	{
-		InSoundClass->Properties.Volume = Volume;
-	}
-}
-
 bool FAudioDevice::LocationIsAudible( FVector Location, float MaxDistance )
 {
 	if( MaxDistance >= WORLD_MAX )
diff --git a/Engine/Source/Runtime/Engine/Public/AudioDevice.h b/Engine/Source/Runtime/Engine/Public/AudioDevice.h
--- a/Engine/Source/Runtime/Engine/Public/AudioDevice.h
+++ b/Engine/Source/Runtime/Engine/Public/AudioDevice.h
@@ -205,7 +205,6 @@ public:
 	bool HandleDisableRadioCommand( const TCHAR* Cmd, FOutputDevice& Ar );
 	bool HandleEnableRadioCommand( const TCHAR* Cmd, FOutputDevice& Ar );
 	bool HandleResetSoundStateCommand( const TCHAR* Cmd, FOutputDevice& Ar );
-	bool HandleModifySoundClassCommand( const TCHAR* Cmd, FOutputDevice& Ar );
 #endif
 
 	/**
@@ -393,11 +392,6 @@ public:
 	 */
 	FSoundClassProperties* GetSoundClassCurrentProperties( class USoundClass* InSoundClass );
 
-	/**
-	 * Updates sound class volumes
-	 */
-	void SetClassVolume( class USoundClass* InSoundClass, const float Volume );
-
 	/**
 	 * Checks to see if a coordinate is within a distance of any listener
 	 */
