#ue4 - Change FVector_NetQuantize100 to actually use decimal precision in the hundredths rather than thousandths. Also lowered bits used to 30 rather that 31, to prevent overflowing Max value when serializing if close to the max size. In theory this supports up to 100km^2 sizes, however float precision will break down before that.

diff --git a/Engine/Source/Runtime/Engine/Classes/Engine/NetSerialization.h b/Engine/Source/Runtime/Engine/Classes/Engine/NetSerialization.h
--- a/Engine/Source/Runtime/Engine/Classes/Engine/NetSerialization.h
+++ b/Engine/Source/Runtime/Engine/Classes/Engine/NetSerialization.h
@@ -902,7 +902,7 @@ bool SerializeFixedVector(FVector &Vector, FArchive& Ar)
  *
  *	0 decimal place of precision.
  *	Up to 20 bits per component.
- *	Valid range: 2^20 = +/- 1048576
+ *	Valid range: 2^20 = +/- 1,048,576
  *
  *	Note: this is the historical UE format for vector net serialization
  *
@@ -949,7 +949,7 @@ struct TStructOpsTypeTraits< FVector_NetQuantize > : public TStructOpsTypeTraits
  *
  *	1 decimal place of precision.
  *	Up to 24 bits per component.
- *	Valid range: 2^24 / 10 = +/- 1677721.6
+ *	Valid range: 2^24 / 10 = +/- 1,677,721.6
  *
  */
 USTRUCT()
@@ -992,9 +992,9 @@ struct TStructOpsTypeTraits< FVector_NetQuantize10 > : public TStructOpsTypeTrai
 /**
  *	FVector_NetQuantize100
  *
- *	3 decimal place of precision.
- *	Up to 31 bits per component.
- *	Valid range: 2^31 / 1000 = +/- 2147483.648
+ *	2 decimal place of precision.
+ *	Up to 30 bits per component.
+ *	Valid range: 2^30 / 100 = +/- 10,737,418.24
  *
  */
 USTRUCT()
@@ -1020,7 +1020,7 @@ struct FVector_NetQuantize100 : public FVector
 
 	bool NetSerialize(FArchive& Ar, class UPackageMap* Map, bool& bOutSuccess)
 	{
-		bOutSuccess = SerializePackedVector<1000, 31>(*this, Ar);
+		bOutSuccess = SerializePackedVector<100, 30>(*this, Ar);
 		return true;
 	}
 };
