Fixed navmesh generation crashing when RecastNavMesh instance gets removed to be replaced by one created for a SupportedAgent #UE4

- also, removed unused FPImplRecastNavMesh::bOwnsNavMeshData

diff --git a/Engine/Source/Runtime/Engine/Private/AI/Navigation/PImplRecastNavMesh.cpp b/Engine/Source/Runtime/Engine/Private/AI/Navigation/PImplRecastNavMesh.cpp
--- a/Engine/Source/Runtime/Engine/Private/AI/Navigation/PImplRecastNavMesh.cpp
+++ b/Engine/Source/Runtime/Engine/Private/AI/Navigation/PImplRecastNavMesh.cpp
@@ -238,13 +238,10 @@ bool FRecastSpeciaLinkFilter::isLinkAllowed(const int32 UserId) const
 
 FPImplRecastNavMesh::FPImplRecastNavMesh(ARecastNavMesh* Owner)
 	: NavMeshOwner(Owner)
-	, bOwnsNavMeshData(false)
 	, DetourNavMesh(NULL)
 {
 	check(Owner && "Owner must never be NULL");
 
-	DetourNavMesh = dtAllocNavMesh();
-
 	INC_DWORD_STAT_BY( STAT_NavigationMemory
 		, Owner->HasAnyFlags(RF_ClassDefaultObject) == false ? sizeof(*this) : 0 );
 };
@@ -286,15 +283,13 @@ void FPImplRecastNavMesh::Serialize( FArchive& Ar )
 	if (Ar.IsLoading())
 	{
 		// allocate the navmesh object
+		ReleaseDetourNavMesh();
 		DetourNavMesh = dtAllocNavMesh();
+
 		if (DetourNavMesh == NULL)
 		{
 			UE_VLOG(NavMeshOwner, LogNavigation, Error, TEXT("Failed to allocate Recast navmesh"));
 		}
-		else
-		{
-			bOwnsNavMeshData = true;
-		}
 	}
 
 	int32 NumTiles = 0;
@@ -346,8 +341,7 @@ void FPImplRecastNavMesh::Serialize( FArchive& Ar )
 		if (ActorsTileSize != Params.tileWidth)
 		{
 			// just move archive position
-			dtFreeNavMesh(DetourNavMesh);
-			DetourNavMesh = NULL;
+			ReleaseDetourNavMesh();
 
 			for (int i = 0; i < NumTiles; ++i)
 			{
@@ -630,21 +624,14 @@ void FPImplRecastNavMesh::SerializeRecastMeshTile(FArchive& Ar, unsigned char*&
 	}
 }
 
-void FPImplRecastNavMesh::SetRecastMesh(dtNavMesh* NavMesh, bool bOwnData)
+void FPImplRecastNavMesh::SetRecastMesh(dtNavMesh* NavMesh)
 {
 	if (NavMesh == DetourNavMesh)
 	{
 		return;
 	}
 
-	if (DetourNavMesh != NULL && !!bOwnsNavMeshData)
-	{
-		// if there's already some recast navmesh, and it's owned by this instance then release it
-		dtFreeNavMesh(DetourNavMesh);
-		DetourNavMesh = NULL;
-	}
-
-	bOwnsNavMeshData = bOwnData;
+	ReleaseDetourNavMesh();
 	DetourNavMesh = NavMesh;
 
 	if (NavMeshOwner)
diff --git a/Engine/Source/Runtime/Engine/Private/AI/Navigation/RecastNavMeshGenerator.cpp b/Engine/Source/Runtime/Engine/Private/AI/Navigation/RecastNavMeshGenerator.cpp
--- a/Engine/Source/Runtime/Engine/Private/AI/Navigation/RecastNavMeshGenerator.cpp
+++ b/Engine/Source/Runtime/Engine/Private/AI/Navigation/RecastNavMeshGenerator.cpp
@@ -2722,7 +2722,7 @@ FRecastNavMeshGenerator::FRecastNavMeshGenerator(class ARecastNavMesh* InDestNav
 	check(InDestNavMesh);
 	Init();
 
-	if (!DestNavMesh->HasValidNavmesh())
+	if (DestNavMesh->HasValidNavmesh() == false)
 	{
 		// recreate navmesh from scratch if no data was loaded
 		ConstructTiledNavMesh();
@@ -2808,6 +2808,7 @@ void FRecastNavMeshGenerator::Init()
 	/** setup maximum number of active tile generator*/
 	const int32 NumberOfWorkerThreads = FTaskGraphInterface::Get().GetNumWorkerThreads();
 	MaxTileGeneratorTasks = FMath::Max(NumberOfWorkerThreads*2, 1);
+	UE_LOG(LogNavigation, Log, TEXT("Using max of %d workers to build navigation."), MaxTileGeneratorTasks);
 	NumActiveTiles = 0;
 
 	// prepare voxel cache if needed
@@ -3087,7 +3088,9 @@ TArray<uint32> FRecastNavMeshGenerator::RemoveTileLayers(const int32 TileX, cons
 {
 	TArray<uint32> ResultTileIndices;
 	dtNavMesh* DetourMesh = DestNavMesh->GetRecastNavMeshImpl()->GetRecastMesh();
-	const int32 NumLayers = DetourMesh ? DetourMesh->getTileCountAt(TileX, TileY) : 0;
+	
+	check(DetourMesh == nullptr || DetourMesh->isEmpty() == false);
+	const int32 NumLayers = DetourMesh != nullptr ? DetourMesh->getTileCountAt(TileX, TileY) : 0;
 	
 	if (NumLayers > 0)
 	{
diff --git a/Engine/Source/Runtime/Engine/Public/AI/Navigation/PImplRecastNavMesh.h b/Engine/Source/Runtime/Engine/Public/AI/Navigation/PImplRecastNavMesh.h
--- a/Engine/Source/Runtime/Engine/Public/AI/Navigation/PImplRecastNavMesh.h
+++ b/Engine/Source/Runtime/Engine/Public/AI/Navigation/PImplRecastNavMesh.h
@@ -186,13 +186,11 @@ public:
 	class dtNavMesh* GetRecastMesh() { return DetourNavMesh; };
 	void ReleaseDetourNavMesh();
 
-	bool GetOwnsNavMeshData() const { return bOwnsNavMeshData; }
-
 	/** Assigns recast generated navmesh to this instance.
 	 *	@param bOwnData if true from now on this FPImplRecastNavMesh instance will be responsible for this piece 
 	 *		of memory
 	 */
-	void SetRecastMesh(class dtNavMesh* NavMesh, bool bOwnData = false);
+	void SetRecastMesh(class dtNavMesh* NavMesh);
 
 	float GetTotalDataSize() const;
 
@@ -205,10 +203,7 @@ public:
 	float CalcSegmentCostOnPoly(NavNodeRef PolyID, const dtQueryFilter* Filter, const FVector& StartLoc, const FVector& EndLoc) const;
 
 	class ARecastNavMesh* NavMeshOwner;
-
-	/** if true instance is responsible for deallocation of recast navmesh */
-	bool bOwnsNavMeshData;
-
+	
 	/** Recast's runtime navmesh data that we can query against */
 	class dtNavMesh* DetourNavMesh;
 
