revised LightComponent checks for UCS execution in setter methods (per feedback)

diff --git a/Engine/Source/Runtime/Engine/Private/DirectionalLightComponent.cpp b/Engine/Source/Runtime/Engine/Private/DirectionalLightComponent.cpp
--- a/Engine/Source/Runtime/Engine/Private/DirectionalLightComponent.cpp
+++ b/Engine/Source/Runtime/Engine/Private/DirectionalLightComponent.cpp
@@ -797,7 +797,7 @@ void UDirectionalLightComponent::SetShadowDistanceFadeoutFraction(float NewValue
 
 void UDirectionalLightComponent::SetEnableLightShaftOcclusion(bool bNewValue)
 {
-	if (!(IsRegistered() && Mobility == EComponentMobility::Static && !IsRunningUserConstructionScript())
+	if ((IsRunningUserConstructionScript() || !(IsRegistered() && Mobility == EComponentMobility::Static))
 		&& bEnableLightShaftOcclusion != bNewValue)
 	{
 		bEnableLightShaftOcclusion = bNewValue;
@@ -807,7 +807,7 @@ void UDirectionalLightComponent::SetEnableLightShaftOcclusion(bool bNewValue)
 
 void UDirectionalLightComponent::SetOcclusionMaskDarkness(float NewValue)
 {
-	if (!(IsRegistered() && Mobility == EComponentMobility::Static && !IsRunningUserConstructionScript())
+	if ((IsRunningUserConstructionScript() || !(IsRegistered() && Mobility == EComponentMobility::Static))
 		&& OcclusionMaskDarkness != NewValue)
 	{
 		OcclusionMaskDarkness = NewValue;
@@ -817,7 +817,7 @@ void UDirectionalLightComponent::SetOcclusionMaskDarkness(float NewValue)
 
 void UDirectionalLightComponent::SetLightShaftOverrideDirection(FVector NewValue)
 {
-	if (!(IsRegistered() && Mobility == EComponentMobility::Static && !IsRunningUserConstructionScript())
+	if ((IsRunningUserConstructionScript() || !(IsRegistered() && Mobility == EComponentMobility::Static))
 		&& LightShaftOverrideDirection != NewValue)
 	{
 		LightShaftOverrideDirection = NewValue;
diff --git a/Engine/Source/Runtime/Engine/Private/LightComponent.cpp b/Engine/Source/Runtime/Engine/Private/LightComponent.cpp
--- a/Engine/Source/Runtime/Engine/Private/LightComponent.cpp
+++ b/Engine/Source/Runtime/Engine/Private/LightComponent.cpp
@@ -70,7 +70,7 @@ FArchive& operator<<(FArchive& Ar, FStaticShadowDepthMap& ShadowMap)
 
 void ULightComponentBase::SetCastShadows(bool bNewValue)
 {
-	if (!(IsRegistered() && Mobility == EComponentMobility::Static && !IsRunningUserConstructionScript())
+	if ((IsRunningUserConstructionScript() || !(IsRegistered() && Mobility == EComponentMobility::Static))
 		&& CastShadows != bNewValue)
 	{
 		CastShadows = bNewValue;
@@ -669,7 +669,7 @@ void ULightComponent::DestroyRenderState_Concurrent()
 void ULightComponent::SetIntensity(float NewIntensity)
 {
 	// Can't set brightness on a static light
-	if (!(IsRegistered() && Mobility == EComponentMobility::Static && !IsRunningUserConstructionScript())
+	if ((IsRunningUserConstructionScript() || !(IsRegistered() && Mobility == EComponentMobility::Static))
 		&& Intensity != NewIntensity)
 	{
 		Intensity = NewIntensity;
@@ -689,7 +689,7 @@ void ULightComponent::SetLightColor(FLinearColor NewLightColor)
 	FColor NewColor(NewLightColor);
 
 	// Can't set color on a static light
-	if (!(IsRegistered() && Mobility == EComponentMobility::Static && !IsRunningUserConstructionScript())
+	if ((IsRunningUserConstructionScript() || !(IsRegistered() && Mobility == EComponentMobility::Static))
 		&& LightColor != NewColor)
 	{
 		LightColor	= NewColor;
@@ -706,7 +706,7 @@ void ULightComponent::SetLightColor(FLinearColor NewLightColor)
 void ULightComponent::SetLightFunctionMaterial(UMaterialInterface* NewLightFunctionMaterial)
 {
 	// Can't set light function on a static light
-	if (!(IsRegistered() && Mobility == EComponentMobility::Static && !IsRunningUserConstructionScript())
+	if ((IsRunningUserConstructionScript() || !(IsRegistered() && Mobility == EComponentMobility::Static))
 		&& NewLightFunctionMaterial != LightFunctionMaterial)
 	{
 		LightFunctionMaterial = NewLightFunctionMaterial;
@@ -716,7 +716,7 @@ void ULightComponent::SetLightFunctionMaterial(UMaterialInterface* NewLightFunct
 
 void ULightComponent::SetLightFunctionScale(FVector NewLightFunctionScale)
 {
-	if (!(IsRegistered() && Mobility == EComponentMobility::Static && !IsRunningUserConstructionScript())
+	if ((IsRunningUserConstructionScript() || !(IsRegistered() && Mobility == EComponentMobility::Static))
 		&& NewLightFunctionScale != LightFunctionScale)
 	{
 		LightFunctionScale = NewLightFunctionScale;
@@ -726,7 +726,7 @@ void ULightComponent::SetLightFunctionScale(FVector NewLightFunctionScale)
 
 void ULightComponent::SetLightFunctionFadeDistance(float NewLightFunctionFadeDistance)
 {
-	if (!(IsRegistered() && Mobility == EComponentMobility::Static && !IsRunningUserConstructionScript())
+	if ((IsRunningUserConstructionScript() || !(IsRegistered() && Mobility == EComponentMobility::Static))
 		&& NewLightFunctionFadeDistance != LightFunctionFadeDistance)
 	{
 		LightFunctionFadeDistance = NewLightFunctionFadeDistance;
@@ -736,7 +736,7 @@ void ULightComponent::SetLightFunctionFadeDistance(float NewLightFunctionFadeDis
 
 void ULightComponent::SetAffectDynamicIndirectLighting(bool bNewValue)
 {
-	if (!(IsRegistered() && Mobility == EComponentMobility::Static && !IsRunningUserConstructionScript())
+	if ((IsRunningUserConstructionScript() || !(IsRegistered() && Mobility == EComponentMobility::Static))
 		&& bAffectDynamicIndirectLighting != bNewValue)
 	{
 		bAffectDynamicIndirectLighting = bNewValue;
@@ -746,7 +746,7 @@ void ULightComponent::SetAffectDynamicIndirectLighting(bool bNewValue)
 
 void ULightComponent::SetAffectTranslucentLighting(bool bNewValue)
 {
-	if (!(IsRegistered() && Mobility == EComponentMobility::Static && !IsRunningUserConstructionScript())
+	if ((IsRunningUserConstructionScript() || !(IsRegistered() && Mobility == EComponentMobility::Static))
 		&& bAffectTranslucentLighting != bNewValue)
 	{
 		bAffectTranslucentLighting = bNewValue;
@@ -756,7 +756,7 @@ void ULightComponent::SetAffectTranslucentLighting(bool bNewValue)
 
 void ULightComponent::SetEnableLightShaftBloom(bool bNewValue)
 {
-	if (!(IsRegistered() && Mobility == EComponentMobility::Static && !IsRunningUserConstructionScript())
+	if ((IsRunningUserConstructionScript() || !(IsRegistered() && Mobility == EComponentMobility::Static))
 		&& bEnableLightShaftBloom != bNewValue)
 	{
 		bEnableLightShaftBloom = bNewValue;
@@ -766,7 +766,7 @@ void ULightComponent::SetEnableLightShaftBloom(bool bNewValue)
 
 void ULightComponent::SetBloomScale(float NewValue)
 {
-	if (!(IsRegistered() && Mobility == EComponentMobility::Static && !IsRunningUserConstructionScript())
+	if ((IsRunningUserConstructionScript() || !(IsRegistered() && Mobility == EComponentMobility::Static))
 		&& BloomScale != NewValue)
 	{
 		BloomScale = NewValue;
@@ -776,7 +776,7 @@ void ULightComponent::SetBloomScale(float NewValue)
 
 void ULightComponent::SetBloomThreshold(float NewValue)
 {
-	if (!(IsRegistered() && Mobility == EComponentMobility::Static && !IsRunningUserConstructionScript())
+	if ((IsRunningUserConstructionScript() || !(IsRegistered() && Mobility == EComponentMobility::Static))
 		&& BloomThreshold != NewValue)
 	{
 		BloomThreshold = NewValue;
@@ -786,7 +786,7 @@ void ULightComponent::SetBloomThreshold(float NewValue)
 
 void ULightComponent::SetBloomTint(FColor NewValue)
 {
-	if (!(IsRegistered() && Mobility == EComponentMobility::Static && !IsRunningUserConstructionScript())
+	if ((IsRunningUserConstructionScript() || !(IsRegistered() && Mobility == EComponentMobility::Static))
 		&& BloomTint != NewValue)
 	{
 		BloomTint = NewValue;
@@ -796,7 +796,7 @@ void ULightComponent::SetBloomTint(FColor NewValue)
 
 void ULightComponent::SetIESTexture(UTextureLightProfile* NewValue)
 {
-	if (!(IsRegistered() && Mobility == EComponentMobility::Static && !IsRunningUserConstructionScript())
+	if ((IsRunningUserConstructionScript() || !(IsRegistered() && Mobility == EComponentMobility::Static))
 		&& IESTexture != NewValue)
 	{
 		IESTexture = NewValue;
diff --git a/Engine/Source/Runtime/Engine/Private/PointLightComponent.cpp b/Engine/Source/Runtime/Engine/Private/PointLightComponent.cpp
--- a/Engine/Source/Runtime/Engine/Private/PointLightComponent.cpp
+++ b/Engine/Source/Runtime/Engine/Private/PointLightComponent.cpp
@@ -130,7 +130,7 @@ FLightSceneProxy* UPointLightComponent::CreateSceneProxy() const
 void UPointLightComponent::SetAttenuationRadius(float NewRadius)
 {
 	// Only movable lights can change their radius at runtime
-	if (!(IsRegistered() && (Mobility == EComponentMobility::Static || Mobility == EComponentMobility::Stationary) && !IsRunningUserConstructionScript())
+	if ((IsRunningUserConstructionScript() || !(IsRegistered() && (Mobility == EComponentMobility::Static || Mobility == EComponentMobility::Stationary)))
 		&& NewRadius != AttenuationRadius)
 	{
 		AttenuationRadius = NewRadius;
@@ -140,7 +140,7 @@ void UPointLightComponent::SetAttenuationRadius(float NewRadius)
 
 void UPointLightComponent::SetLightFalloffExponent(float NewLightFalloffExponent)
 {
-	if (!(IsRegistered() && Mobility == EComponentMobility::Static && !IsRunningUserConstructionScript())
+	if ((IsRunningUserConstructionScript() || !(IsRegistered() && Mobility == EComponentMobility::Static))
 		&& NewLightFalloffExponent != LightFalloffExponent)
 	{
 		LightFalloffExponent = NewLightFalloffExponent;
@@ -150,7 +150,7 @@ void UPointLightComponent::SetLightFalloffExponent(float NewLightFalloffExponent
 
 void UPointLightComponent::SetSourceRadius(float bNewValue)
 {
-	if (!(IsRegistered() && Mobility == EComponentMobility::Static && !IsRunningUserConstructionScript())
+	if ((IsRunningUserConstructionScript() || !(IsRegistered() && Mobility == EComponentMobility::Static))
 		&& SourceRadius != bNewValue)
 	{
 		SourceRadius = bNewValue;
diff --git a/Engine/Source/Runtime/Engine/Private/SkyLightComponent.cpp b/Engine/Source/Runtime/Engine/Private/SkyLightComponent.cpp
--- a/Engine/Source/Runtime/Engine/Private/SkyLightComponent.cpp
+++ b/Engine/Source/Runtime/Engine/Private/SkyLightComponent.cpp
@@ -433,7 +433,7 @@ void USkyLightComponent::CaptureEmissiveIrradianceEnvironmentMap(FSHVectorRGB3&
 void USkyLightComponent::SetIntensity(float NewIntensity)
 {
 	// Can't set brightness on a static light
-	if (!(IsRegistered() && Mobility == EComponentMobility::Static && !IsRunningUserConstructionScript())
+	if ((IsRunningUserConstructionScript() || !(IsRegistered() && Mobility == EComponentMobility::Static))
 		&& Intensity != NewIntensity)
 	{
 		Intensity = NewIntensity;
@@ -447,7 +447,7 @@ void USkyLightComponent::SetLightColor(FLinearColor NewLightColor)
 	FColor NewColor(NewLightColor);
 
 	// Can't set color on a static light
-	if (!(IsRegistered() && Mobility == EComponentMobility::Static && !IsRunningUserConstructionScript())
+	if ((IsRunningUserConstructionScript() || !(IsRegistered() && Mobility == EComponentMobility::Static))
 		&& LightColor != NewColor)
 	{
 		LightColor = NewColor;
@@ -458,7 +458,7 @@ void USkyLightComponent::SetLightColor(FLinearColor NewLightColor)
 void USkyLightComponent::SetCubemap(UTextureCube* NewCubemap)
 {
 	// Can't set color on a static light
-	if (!(IsRegistered() && Mobility == EComponentMobility::Static && !IsRunningUserConstructionScript())
+	if ((IsRunningUserConstructionScript() || !(IsRegistered() && Mobility == EComponentMobility::Static))
 		&& Cubemap != NewCubemap)
 	{
 		Cubemap = NewCubemap;
@@ -470,7 +470,7 @@ void USkyLightComponent::SetCubemap(UTextureCube* NewCubemap)
 void USkyLightComponent::SetOcclusionTint(const FColor& InTint)
 {
 	// Can't set on a static light
-	if (!(IsRegistered() && Mobility == EComponentMobility::Static && !IsRunningUserConstructionScript())
+	if ((IsRunningUserConstructionScript() || !(IsRegistered() && Mobility == EComponentMobility::Static))
 		&& OcclusionTint != InTint)
 	{
 		OcclusionTint = InTint;
@@ -481,7 +481,7 @@ void USkyLightComponent::SetOcclusionTint(const FColor& InTint)
 void USkyLightComponent::SetMinOcclusion(float InMinOcclusion)
 {
 	// Can't set on a static light
-	if (!(IsRegistered() && Mobility == EComponentMobility::Static && !IsRunningUserConstructionScript())
+	if ((IsRunningUserConstructionScript() || !(IsRegistered() && Mobility == EComponentMobility::Static))
 		&& MinOcclusion != InMinOcclusion)
 	{
 		MinOcclusion = InMinOcclusion;
