Various fixes for doxygen warnings: change function definitions to be exactly the same as their declaration, token for token.

diff --git a/Engine/Source/Editor/UnrealEd/Private/Commandlets/CookCommandlet.cpp b/Engine/Source/Editor/UnrealEd/Private/Commandlets/CookCommandlet.cpp
--- a/Engine/Source/Editor/UnrealEd/Private/Commandlets/CookCommandlet.cpp
+++ b/Engine/Source/Editor/UnrealEd/Private/Commandlets/CookCommandlet.cpp
@@ -51,7 +51,7 @@ static FString GetPackageFilename( UPackage* Package )
 /* UCookCommandlet structors
  *****************************************************************************/
 
-UCookCommandlet::UCookCommandlet( const class FObjectInitializer& ObjectInitializer )
+UCookCommandlet::UCookCommandlet( const FObjectInitializer& ObjectInitializer )
 	: Super(ObjectInitializer)
 {
 
diff --git a/Engine/Source/Editor/UnrealEd/Private/Editor.cpp b/Engine/Source/Editor/UnrealEd/Private/Editor.cpp
--- a/Engine/Source/Editor/UnrealEd/Private/Editor.cpp
+++ b/Engine/Source/Editor/UnrealEd/Private/Editor.cpp
@@ -186,7 +186,7 @@ static void PrivateDestroySelectedSets()
 #endif
 }
 
-UEditorEngine::UEditorEngine(const class FObjectInitializer& ObjectInitializer)
+UEditorEngine::UEditorEngine(const FObjectInitializer& ObjectInitializer)
 : Super(ObjectInitializer)
 {
 	if (!IsRunningCommandlet())
diff --git a/Engine/Source/Editor/UnrealEd/Private/Factories/PackFactory.cpp b/Engine/Source/Editor/UnrealEd/Private/Factories/PackFactory.cpp
--- a/Engine/Source/Editor/UnrealEd/Private/Factories/PackFactory.cpp
+++ b/Engine/Source/Editor/UnrealEd/Private/Factories/PackFactory.cpp
@@ -17,7 +17,7 @@
 
 DEFINE_LOG_CATEGORY_STATIC(LogPackFactory, Log, All);
 
-UPackFactory::UPackFactory(const class FObjectInitializer& PCIP)
+UPackFactory::UPackFactory(const FObjectInitializer& PCIP)
 	: Super(PCIP)
 {
 	// Since this factory can output multiple and any number of class it doesn't really have a 
diff --git a/Engine/Source/Runtime/AIModule/Private/Navigation/CrowdManager.cpp b/Engine/Source/Runtime/AIModule/Private/Navigation/CrowdManager.cpp
--- a/Engine/Source/Runtime/AIModule/Private/Navigation/CrowdManager.cpp
+++ b/Engine/Source/Runtime/AIModule/Private/Navigation/CrowdManager.cpp
@@ -313,13 +313,13 @@ void UCrowdManager::UnregisterAgent(const ICrowdAgentInterface* Agent)
 	ActiveAgents.Remove(Agent);
 }
 
-bool UCrowdManager::IsAgentValid(const UCrowdFollowingComponent* AgentComponent) const
+bool UCrowdManager::IsAgentValid(const class UCrowdFollowingComponent* AgentComponent) const
 {
 	const FCrowdAgentData* AgentData = ActiveAgents.Find(AgentComponent);
 	return AgentData && AgentData->IsValid();
 }
 
-bool UCrowdManager::IsAgentValid(const ICrowdAgentInterface* Agent) const
+bool UCrowdManager::IsAgentValid(const class ICrowdAgentInterface* Agent) const
 {
 	const FCrowdAgentData* AgentData = ActiveAgents.Find(Agent);
 	return AgentData && AgentData->IsValid();
diff --git a/Engine/Source/Runtime/Core/Private/Math/Color.cpp b/Engine/Source/Runtime/Core/Private/Math/Color.cpp
--- a/Engine/Source/Runtime/Core/Private/Math/Color.cpp
+++ b/Engine/Source/Runtime/Core/Private/Math/Color.cpp
@@ -31,7 +31,7 @@ const FColor FColor::Magenta(255,0,255);
 static const float OneOver255 = 1.0f / 255.0f;
 
 //	FColor->FLinearColor conversion.
-FLinearColor::FLinearColor(const class FColor& C)
+FLinearColor::FLinearColor(const FColor& C)
 {
 	R = PowOneOver255Table[C.R];
 	G = PowOneOver255Table[C.G];
@@ -39,14 +39,14 @@ FLinearColor::FLinearColor(const class FColor& C)
 	A =	float(C.A) * OneOver255;
 }
 
-FLinearColor::FLinearColor(const class FVector& Vector) :
+FLinearColor::FLinearColor(const FVector& Vector) :
 	R(Vector.X),
 	G(Vector.Y),
 	B(Vector.Z),
 	A(1.0f)
 {}
 
-FLinearColor::FLinearColor(const class FFloat16Color& C)
+FLinearColor::FLinearColor(const FFloat16Color& C)
 {
 	R = C.R.GetFloat();
 	G = C.G.GetFloat();
diff --git a/Engine/Source/Runtime/CoreUObject/Private/UObject/LinkerLoad.cpp b/Engine/Source/Runtime/CoreUObject/Private/UObject/LinkerLoad.cpp
--- a/Engine/Source/Runtime/CoreUObject/Private/UObject/LinkerLoad.cpp
+++ b/Engine/Source/Runtime/CoreUObject/Private/UObject/LinkerLoad.cpp
@@ -594,7 +594,7 @@ ULinkerLoad::ELinkerStatus ULinkerLoad::Tick( float InTimeLimit, bool bInUseTime
  * @param	Filename	Name of file on disk to load
  * @param	LoadFlags	Load flags determining behavior
  */
-ULinkerLoad::ULinkerLoad( const class FObjectInitializer& ObjectInitializer, UPackage* InParent, const TCHAR* InFilename, uint32 InLoadFlags )
+ULinkerLoad::ULinkerLoad( const FObjectInitializer& ObjectInitializer, UPackage* InParent, const TCHAR* InFilename, uint32 InLoadFlags )
 :	ULinker( ObjectInitializer, InParent, InFilename )
 ,	LoadFlags( InLoadFlags )
 ,	bHaveImportsBeenVerified( false )
diff --git a/Engine/Source/Runtime/Engine/Private/AI/Navigation/NavigationData.cpp b/Engine/Source/Runtime/Engine/Private/AI/Navigation/NavigationData.cpp
--- a/Engine/Source/Runtime/Engine/Private/AI/Navigation/NavigationData.cpp
+++ b/Engine/Source/Runtime/Engine/Private/AI/Navigation/NavigationData.cpp
@@ -58,7 +58,7 @@ FPathFindingQuery::FPathFindingQuery(FNavPathSharedRef PathToRecalculate, const
 //----------------------------------------------------------------------//
 uint32 FAsyncPathFindingQuery::LastPathFindingUniqueID = INVALID_NAVQUERYID;
 
-FAsyncPathFindingQuery::FAsyncPathFindingQuery(const UObject* InOwner, const class ANavigationData* InNavData, const FVector& Start, const FVector& End, const FNavPathQueryDelegate& Delegate, TSharedPtr<const FNavigationQueryFilter> SourceQueryFilter)
+FAsyncPathFindingQuery::FAsyncPathFindingQuery(const UObject* InOwner, const ANavigationData* InNavData, const FVector& Start, const FVector& End, const FNavPathQueryDelegate& Delegate, TSharedPtr<const FNavigationQueryFilter> SourceQueryFilter)
 : FPathFindingQuery(InOwner, InNavData, Start, End, SourceQueryFilter)
 , QueryID(GetUniqueID())
 , OnDoneDelegate(Delegate)
diff --git a/Engine/Source/Runtime/Engine/Private/AssetUserData.cpp b/Engine/Source/Runtime/Engine/Private/AssetUserData.cpp
--- a/Engine/Source/Runtime/Engine/Private/AssetUserData.cpp
+++ b/Engine/Source/Runtime/Engine/Private/AssetUserData.cpp
@@ -12,7 +12,7 @@ UAssetUserData::UAssetUserData(const FObjectInitializer& ObjectInitializer)
 
 //////////////////////////////////////////////////////////////////////////
 
-UInterface_AssetUserData::UInterface_AssetUserData(const class FObjectInitializer& ObjectInitializer)
+UInterface_AssetUserData::UInterface_AssetUserData(const FObjectInitializer& ObjectInitializer)
 	: Super(ObjectInitializer)
 {
 }
\ No newline at end of file
diff --git a/Engine/Source/Runtime/Engine/Private/Model.cpp b/Engine/Source/Runtime/Engine/Private/Model.cpp
--- a/Engine/Source/Runtime/Engine/Private/Model.cpp
+++ b/Engine/Source/Runtime/Engine/Private/Model.cpp
@@ -502,7 +502,7 @@ void UModel::EmptyModel( int32 EmptySurfInfo, int32 EmptyPolys )
 //
 // Create a new model and allocate all objects needed for it.
 //
-UModel::UModel( const class FObjectInitializer& ObjectInitializer,ABrush* Owner, bool InRootOutside )
+UModel::UModel( const FObjectInitializer& ObjectInitializer,ABrush* Owner, bool InRootOutside )
 :	UObject(ObjectInitializer)
 ,	Nodes		( this )
 ,	Verts		( this )
diff --git a/Engine/Source/Runtime/Engine/Private/StaticMesh.cpp b/Engine/Source/Runtime/Engine/Private/StaticMesh.cpp
--- a/Engine/Source/Runtime/Engine/Private/StaticMesh.cpp
+++ b/Engine/Source/Runtime/Engine/Private/StaticMesh.cpp
@@ -2570,7 +2570,7 @@ UStaticMeshSocket* UStaticMesh::FindSocket(FName InSocketName)
 UStaticMeshSocket
 -----------------------------------------------------------------------------*/
 
-UStaticMeshSocket::UStaticMeshSocket(const class FObjectInitializer& ObjectInitializer)
+UStaticMeshSocket::UStaticMeshSocket(const FObjectInitializer& ObjectInitializer)
 	: Super(ObjectInitializer)
 {
 	RelativeScale = FVector(1.0f, 1.0f, 1.0f);
diff --git a/Engine/Source/Runtime/Engine/Private/VisualLogger/VisualLoggerTypes.cpp b/Engine/Source/Runtime/Engine/Private/VisualLogger/VisualLoggerTypes.cpp
--- a/Engine/Source/Runtime/Engine/Private/VisualLogger/VisualLoggerTypes.cpp
+++ b/Engine/Source/Runtime/Engine/Private/VisualLogger/VisualLoggerTypes.cpp
@@ -38,7 +38,7 @@ FVisualLogEntry::FVisualLogEntry(const FVisualLogEntry& Entry)
 	DataBlocks = Entry.DataBlocks;
 }
 
-FVisualLogEntry::FVisualLogEntry(const AActor* InActor, TArray<TWeakObjectPtr<UObject> >* Children)
+FVisualLogEntry::FVisualLogEntry(const class AActor* InActor, TArray<TWeakObjectPtr<UObject> >* Children)
 {
 	if (InActor && InActor->IsPendingKill() == false)
 	{
diff --git a/Engine/Source/Runtime/UMG/Private/Components/WidgetComponent.cpp b/Engine/Source/Runtime/UMG/Private/Components/WidgetComponent.cpp
--- a/Engine/Source/Runtime/UMG/Private/Components/WidgetComponent.cpp
+++ b/Engine/Source/Runtime/UMG/Private/Components/WidgetComponent.cpp
@@ -295,7 +295,7 @@ private:
 	TWeakObjectPtr<UWorld> World;
 };
 
-UWidgetComponent::UWidgetComponent( const class FObjectInitializer& PCIP )
+UWidgetComponent::UWidgetComponent( const FObjectInitializer& PCIP )
 	: Super( PCIP )
 	, DrawSize( FIntPoint( 500, 500 ) )
 	, MaxInteractionDistance( 1000.f )
