Fixing various warnings when building documentation.

diff --git a/Engine/Source/Editor/Kismet/Private/SSCSEditor.cpp b/Engine/Source/Editor/Kismet/Private/SSCSEditor.cpp
--- a/Engine/Source/Editor/Kismet/Private/SSCSEditor.cpp
+++ b/Engine/Source/Editor/Kismet/Private/SSCSEditor.cpp
@@ -2504,7 +2504,7 @@ static FSCSEditorTreeNode* FindRecursive( FSCSEditorTreeNode* Node, FName Name )
 	return nullptr;
 }
 
-void SSCSEditor::HighlightTreeNode(FName TreeNodeName, const FPropertyPath& Property)
+void SSCSEditor::HighlightTreeNode(FName TreeNodeName, const class FPropertyPath& Property)
 {
 	for( const auto& Node : RootNodes )
 	{
diff --git a/Engine/Source/Runtime/AppFramework/Public/Widgets/Workflow/SWizard.h b/Engine/Source/Runtime/AppFramework/Public/Widgets/Workflow/SWizard.h
--- a/Engine/Source/Runtime/AppFramework/Public/Widgets/Workflow/SWizard.h
+++ b/Engine/Source/Runtime/AppFramework/Public/Widgets/Workflow/SWizard.h
@@ -2,6 +2,7 @@
 
 #pragma once
 
+#include "SlateBasics.h"
 
 /**
  * Delegate type for getting a page's enabled state.
diff --git a/Engine/Source/Runtime/Engine/Private/AI/Navigation/NavFilters/NavigationQueryFilter.cpp b/Engine/Source/Runtime/Engine/Private/AI/Navigation/NavFilters/NavigationQueryFilter.cpp
--- a/Engine/Source/Runtime/Engine/Private/AI/Navigation/NavFilters/NavigationQueryFilter.cpp
+++ b/Engine/Source/Runtime/Engine/Private/AI/Navigation/NavFilters/NavigationQueryFilter.cpp
@@ -116,7 +116,7 @@ UNavigationQueryFilter::UNavigationQueryFilter(const FObjectInitializer& ObjectI
 	ExcludeFlags.Packed = 0;
 }
 
-TSharedPtr<const FNavigationQueryFilter> UNavigationQueryFilter::GetQueryFilter(const ANavigationData* NavData) const
+TSharedPtr<const FNavigationQueryFilter> UNavigationQueryFilter::GetQueryFilter(const class ANavigationData* NavData) const
 {
 	TSharedPtr<const FNavigationQueryFilter> SharedFilter = NavData->GetQueryFilter(GetClass());
 	if (!SharedFilter.IsValid())
@@ -172,7 +172,7 @@ void UNavigationQueryFilter::InitializeFilter(const ANavigationData* NavData, FN
 	Filter->SetExcludeFlags(ExcludeFlags.Packed);
 }
 
-TSharedPtr<const struct FNavigationQueryFilter> UNavigationQueryFilter::GetQueryFilter(const ANavigationData* NavData, UClass* FilterClass)
+TSharedPtr<const struct FNavigationQueryFilter> UNavigationQueryFilter::GetQueryFilter(const class ANavigationData* NavData, UClass* FilterClass)
 {
 	UNavigationQueryFilter* DefFilterOb = FilterClass ? FilterClass->GetDefaultObject<UNavigationQueryFilter>() : NULL;
 	return NavData && DefFilterOb ? DefFilterOb->GetQueryFilter(NavData) : NULL;
diff --git a/Engine/Source/Runtime/Engine/Private/Vehicles/PhysXVehicleManager.h b/Engine/Source/Runtime/Engine/Private/Vehicles/PhysXVehicleManager.h
--- a/Engine/Source/Runtime/Engine/Private/Vehicles/PhysXVehicleManager.h
+++ b/Engine/Source/Runtime/Engine/Private/Vehicles/PhysXVehicleManager.h
@@ -2,6 +2,8 @@
 
 #pragma once
 
+#include "EngineDefines.h"
+
 #if WITH_VEHICLE
 
 #include "../PhysicsEngine/PhysXSupport.h"
diff --git a/Engine/Source/Runtime/Engine/Public/Tests/AutomationCommon.h b/Engine/Source/Runtime/Engine/Public/Tests/AutomationCommon.h
--- a/Engine/Source/Runtime/Engine/Public/Tests/AutomationCommon.h
+++ b/Engine/Source/Runtime/Engine/Public/Tests/AutomationCommon.h
@@ -3,6 +3,7 @@
 #pragma once
 
 #include "HardwareInfo.h"
+#include "AutomationTest.h"
 
 ///////////////////////////////////////////////////////////////////////
 // Common Latent commands which are used across test type. I.e. Engine, Network, etc...
diff --git a/Engine/Source/Runtime/OpenGLDrv/Private/OpenGLCommands.cpp b/Engine/Source/Runtime/OpenGLDrv/Private/OpenGLCommands.cpp
--- a/Engine/Source/Runtime/OpenGLDrv/Private/OpenGLCommands.cpp
+++ b/Engine/Source/Runtime/OpenGLDrv/Private/OpenGLCommands.cpp
@@ -246,6 +246,8 @@ static FORCEINLINE GLint ModifyFilterByMips(GLint Filter, bool bHasMips)
 	return Filter;
 }
 
+// Ignore functions from RHIMethods.h when parsing documentation; Doxygen's preprocessor can't parse the declaration, so spews warnings for the definitions.
+#if !UE_BUILD_DOCS
 
 void FOpenGLDynamicRHI::RHIGpuTimeBegin(uint32 Hash, bool bCompute)
 {
@@ -3353,3 +3355,5 @@ void FOpenGLDynamicRHI::RHIEnableDepthBoundsTest(bool bEnable,float MinDepth,flo
 		FOpenGL::DepthBounds(MinDepth,MaxDepth);
 	}
 }
+
+#endif
\ No newline at end of file
diff --git a/Engine/Source/Runtime/OpenGLDrv/Private/OpenGLDevice.cpp b/Engine/Source/Runtime/OpenGLDrv/Private/OpenGLDevice.cpp
--- a/Engine/Source/Runtime/OpenGLDrv/Private/OpenGLDevice.cpp
+++ b/Engine/Source/Runtime/OpenGLDrv/Private/OpenGLDevice.cpp
@@ -127,6 +127,9 @@ FOpenGLContextState& FOpenGLDynamicRHI::GetContextStateForCurrentContext()
 	}
 }
 
+// Ignore functions from RHIMethods.h when parsing documentation; Doxygen's preprocessor can't parse the declaration, so spews warnings for the definitions.
+#if !UE_BUILD_DOCS
+
 void FOpenGLDynamicRHI::RHIBeginFrame()
 {
 	RHIPrivateBeginFrame();
@@ -170,6 +173,8 @@ void FOpenGLDynamicRHI::RHIEndScene()
 	ResourceTableFrameCounter = INDEX_NONE;
 }
 
+#endif
+
 bool GDisableOpenGLDebugOutput = false;
 
 // workaround for HTML5. 
@@ -1177,6 +1182,9 @@ void FOpenGLDynamicRHI::UnregisterQuery( FOpenGLRenderQuery* Query )
 	Queries.RemoveSingleSwap(Query);
 }
 
+// Ignore functions from RHIMethods.h when parsing documentation; Doxygen's preprocessor can't parse the declaration, so spews warnings for the definitions.
+#if !UE_BUILD_DOCS
+
 void FOpenGLDynamicRHI::RHIAutomaticCacheFlushAfterComputeShader(bool bEnable)
 {
 	// Nothing to do here...
@@ -1187,6 +1195,8 @@ void FOpenGLDynamicRHI::RHIFlushComputeShaderCache()
 	// Nothing to do here...
 }
 
+#endif
+
 void* FOpenGLDynamicRHI::RHIGetNativeDevice()
 {
 	return nullptr;
