Fixed RecastNavMesh validity checking function #UE4

diff --git a/Engine/Source/Runtime/Engine/Private/AI/Navigation/RecastNavMesh.cpp b/Engine/Source/Runtime/Engine/Private/AI/Navigation/RecastNavMesh.cpp
--- a/Engine/Source/Runtime/Engine/Private/AI/Navigation/RecastNavMesh.cpp
+++ b/Engine/Source/Runtime/Engine/Private/AI/Navigation/RecastNavMesh.cpp
@@ -1774,13 +1774,10 @@ void ARecastNavMesh::UpdateNavObject()
 bool ARecastNavMesh::HasValidNavmesh() const
 {
 #if WITH_RECAST
-	if (RecastNavMeshImpl && RecastNavMeshImpl->DetourNavMesh)
-	{
-		return true;
-	}
-#endif // WITH_RECAST
-
+	return (RecastNavMeshImpl && RecastNavMeshImpl->DetourNavMesh && RecastNavMeshImpl->DetourNavMesh->isEmpty() == false);
+#else
 	return false;
+#endif // WITH_RECAST
 }
 
 #if WITH_RECAST
