Fix build error on Mac (no warning in VC, suprisingly) - using const keyword instead of class keyword is an empty declaration.

diff --git a/Engine/Source/Runtime/Engine/Classes/AI/Navigation/RecastNavMesh.h b/Engine/Source/Runtime/Engine/Classes/AI/Navigation/RecastNavMesh.h
--- a/Engine/Source/Runtime/Engine/Classes/AI/Navigation/RecastNavMesh.h
+++ b/Engine/Source/Runtime/Engine/Classes/AI/Navigation/RecastNavMesh.h
@@ -39,7 +39,7 @@ class UNavArea;
 class UNavigationSystem;
 class URecastNavMeshDataChunk;
 struct FAreaNavModifier;
-const FNavigationQueryFilter;
+class FNavigationQueryFilter;
 
 UENUM()
 namespace ERecastPartitioning
