CIS fix #UE4

diff --git a/Engine/Source/Runtime/Engine/Classes/AI/Navigation/RecastNavMesh.h b/Engine/Source/Runtime/Engine/Classes/AI/Navigation/RecastNavMesh.h
--- a/Engine/Source/Runtime/Engine/Classes/AI/Navigation/RecastNavMesh.h
+++ b/Engine/Source/Runtime/Engine/Classes/AI/Navigation/RecastNavMesh.h
@@ -39,7 +39,7 @@ class UNavArea;
 class UNavigationSystem;
 class URecastNavMeshDataChunk;
 struct FAreaNavModifier;
-class FNavigationQueryFilter;
+struct FNavigationQueryFilter;
 
 UENUM()
 namespace ERecastPartitioning
