discarding serialized navmesh when generation params are different (usually: tile grid origin)

diff --git a/Engine/Source/Runtime/Engine/Private/AI/Navigation/RecastNavMeshGenerator.cpp b/Engine/Source/Runtime/Engine/Private/AI/Navigation/RecastNavMeshGenerator.cpp
--- a/Engine/Source/Runtime/Engine/Private/AI/Navigation/RecastNavMeshGenerator.cpp
+++ b/Engine/Source/Runtime/Engine/Private/AI/Navigation/RecastNavMeshGenerator.cpp
@@ -2722,14 +2722,32 @@ FRecastNavMeshGenerator::FRecastNavMeshGenerator(ARecastNavMesh* InDestNavMesh)
 	check(InDestNavMesh);
 	Init();
 
-	if (DestNavMesh->HasValidNavmesh() == false)
+	// recreate navmesh if no data was loaded, or when loaded data doesn't match current grid layout
+	bool bRecreateNavmesh = true;
+	if (DestNavMesh->HasValidNavmesh())
+	{
+		const dtNavMeshParams* SavedNavParams = DestNavMesh->GetRecastNavMeshImpl()->DetourNavMesh->getParams();
+		if (SavedNavParams)
+		{
+			const float TileDim = Config.tileSize * Config.cs;
+			if (SavedNavParams->tileHeight == TileDim && SavedNavParams->tileWidth == TileDim)
+			{
+				const FVector Orig = Recast2UnrealPoint(SavedNavParams->orig);
+				const FVector OrigError(FMath::Fmod(Orig.X, TileDim), FMath::Fmod(Orig.X, TileDim), FMath::Fmod(Orig.X, TileDim));
+				if (OrigError.IsNearlyZero())
+				{
+					bRecreateNavmesh = false;
+				}
+			}
+		};
+	}
+
+	if (bRecreateNavmesh)
 	{
-		// recreate navmesh from scratch if no data was loaded
 		ConstructTiledNavMesh();
 	}
 	else
 	{
-		// otherwise just update generator params
 		int32 MaxTiles = 0;
 		CalcNavMeshProperties(MaxTiles, Config.MaxPolysPerTile);
 		NumActiveTiles = GetTilesCountHelper(DestNavMesh->GetRecastNavMeshImpl()->DetourNavMesh);
