Fix for cinematic mips not getting used
* The mips were getting thrown away at cook time and UpdateCachedLODBias was never getting called on non-editor platforms to handle the cinematic mips.

diff --git a/Engine/Source/Runtime/Engine/Classes/Engine/Texture.h b/Engine/Source/Runtime/Engine/Classes/Engine/Texture.h
--- a/Engine/Source/Runtime/Engine/Classes/Engine/Texture.h
+++ b/Engine/Source/Runtime/Engine/Classes/Engine/Texture.h
@@ -547,6 +547,11 @@ public:
 	 */
 	ENGINE_API int32 GetCachedLODBias() const;
 
+	/**
+	 * Cache the combined LOD bias based on texture LOD group and LOD bias.
+	 */
+	ENGINE_API void UpdateCachedLODBias( bool bIncTextureMips = true );
+
 	/**
 	 * @return The material value type of this texture.
 	 */
@@ -624,8 +629,6 @@ public:
 	 */
 	ENGINE_API void ForceRebuildPlatformData();
 
-	ENGINE_API void UpdateCachedLODBias( bool bIncTextureMips = true );
-
 	/**
 	 * Marks platform data as transient. This optionally removes persistent or cached data associated with the platform.
 	 */
diff --git a/Engine/Source/Runtime/Engine/Private/Texture.cpp b/Engine/Source/Runtime/Engine/Private/Texture.cpp
--- a/Engine/Source/Runtime/Engine/Private/Texture.cpp
+++ b/Engine/Source/Runtime/Engine/Private/Texture.cpp
@@ -313,10 +313,8 @@ void UTexture::PostLoad()
 
 	if( !IsTemplate() )
 	{
-#if WITH_EDITOR
 		// Update cached LOD bias.
 		UpdateCachedLODBias();
-#endif // #if WITH_EDITOR
 
 		// The texture will be cached by the cubemap it is contained within on consoles.
 		UTextureCube* CubeMap = Cast<UTextureCube>(GetOuter());
diff --git a/Engine/Source/Runtime/Engine/Private/TextureDerivedData.cpp b/Engine/Source/Runtime/Engine/Private/TextureDerivedData.cpp
--- a/Engine/Source/Runtime/Engine/Private/TextureDerivedData.cpp
+++ b/Engine/Source/Runtime/Engine/Private/TextureDerivedData.cpp
@@ -1217,7 +1217,7 @@ static void SerializePlatformData(
 			const int32 NumCinematicMipLevels = Texture->NumCinematicMipLevels;
 			const TextureMipGenSettings MipGenSetting = Texture->MipGenSettings;
 
-			FirstMipToSerialize = Ar.CookingTarget()->GetTextureLODSettings().CalculateLODBias(Width, Height, LODGroup, LODBias, NumCinematicMipLevels, MipGenSetting);
+			FirstMipToSerialize = Ar.CookingTarget()->GetTextureLODSettings().CalculateLODBias(Width, Height, LODGroup, LODBias, 0, MipGenSetting);
 			FirstMipToSerialize = FMath::Clamp(FirstMipToSerialize, 0, FMath::Max(NumMips-1,0));
 			NumMips -= FirstMipToSerialize;
 		}
@@ -1348,6 +1348,11 @@ void UTexture::CleanupCachedCookedPlatformData()
 	}
 }
 
+void UTexture::UpdateCachedLODBias( bool bIncTextureMips )
+{
+	CachedCombinedLODBias = GSystemSettings.TextureLODSettings.CalculateLODBias( this, bIncTextureMips );
+}
+
 #if WITH_EDITOR
 void UTexture::CachePlatformData(bool bAsyncCache)
 {
@@ -1387,11 +1392,6 @@ void UTexture::CachePlatformData(bool bAsyncCache)
 	}
 }
 
-void UTexture::UpdateCachedLODBias( bool bIncTextureMips )
-{
-	CachedCombinedLODBias = GSystemSettings.TextureLODSettings.CalculateLODBias(this, bIncTextureMips);
-}
-
 void UTexture::BeginCachePlatformData()
 {
 	CachePlatformData(true);
