include EngineTypes.h in TimerManager.h to ensure the moved function delegate is always found

diff --git a/Engine/Source/Runtime/Engine/Public/TimerManager.h b/Engine/Source/Runtime/Engine/Public/TimerManager.h
--- a/Engine/Source/Runtime/Engine/Public/TimerManager.h
+++ b/Engine/Source/Runtime/Engine/Public/TimerManager.h
@@ -6,6 +6,8 @@
 
 #pragma once
 
+#include "Engine/EngineTypes.h"
+
 DECLARE_DELEGATE(FTimerDelegate);
 
 /** Simple interface to wrap a timer delegate that can be either native or dynamic. */
