Fixed crash related to VisualLogger and ToggleAILogging cheat

diff --git a/Engine/Source/Runtime/Engine/Private/VisualLogger/VisualLoggerBinaryFileDevice.cpp b/Engine/Source/Runtime/Engine/Private/VisualLogger/VisualLoggerBinaryFileDevice.cpp
--- a/Engine/Source/Runtime/Engine/Private/VisualLogger/VisualLoggerBinaryFileDevice.cpp
+++ b/Engine/Source/Runtime/Engine/Private/VisualLogger/VisualLoggerBinaryFileDevice.cpp
@@ -48,7 +48,7 @@ void FVisualLoggerBinaryFileDevice::StopRecordingToFile(float TimeStamp)
 	const int32 NumEntries = FrameCache.Num();
 	if (NumEntries> 0)
 	{
-		*FileArchive << FrameCache;
+		FVisualLoggerHelpers::Serialize(*FileArchive, FrameCache);
 		FrameCache.Reset();
 	}
 
