Fix crash when using restart level in PIE with bUseSeamlessTravel true

diff --git a/Engine/Source/Runtime/Engine/Private/GameMode.cpp b/Engine/Source/Runtime/Engine/Private/GameMode.cpp
--- a/Engine/Source/Runtime/Engine/Private/GameMode.cpp
+++ b/Engine/Source/Runtime/Engine/Private/GameMode.cpp
@@ -908,7 +908,7 @@ void AGameMode::ProcessServerTravel(const FString& URL, bool bAbsolute)
 	FString NextMap;
 	if (URL.ToUpper().Contains(TEXT("?RESTART")))
 	{
-		NextMap = GetOutermost()->GetName();
+		NextMap = UWorld::RemovePIEPrefix(GetOutermost()->GetName());
 	}
 	else
 	{
diff --git a/Engine/Source/Runtime/Engine/Private/World.cpp b/Engine/Source/Runtime/Engine/Private/World.cpp
--- a/Engine/Source/Runtime/Engine/Private/World.cpp
+++ b/Engine/Source/Runtime/Engine/Private/World.cpp
@@ -4220,6 +4220,23 @@ void FSeamlessTravelHandler::SeamlessTravelLoadCallback(const FName& PackageName
 			}
 		}
 
+#if WITH_EDITOR
+		FWorldContext &WorldContext = GEngine->GetWorldContextFromWorldChecked(CurrentWorld);
+		if (WorldContext.WorldType == EWorldType::PIE)
+		{
+			// If we are a PIE world and the world we just "loaded" is already initialized, then we're probably travelling to the editor world and we
+			// need to create a PIE world by duplication instead
+			if ( World && World->bIsWorldInitialized)
+			{
+				FString PIEMapName;
+				World = GEngine->CreatePIEWorldByDuplication(WorldContext, World, PIEMapName);
+				World->ClearFlags(RF_Standalone);
+				// CreatePIEWorldByDuplication clears GIsPlayInEditorWorld so set it again
+				GIsPlayInEditorWorld = true;
+			}
+		}
+#endif
+
 		SetHandlerLoadedData(LevelPackage, World);
 	}
 }
