Fixed crash in "PhysXVehicleManager.cpp" due to array index out of bound (PR-627)

diff --git a/Engine/Source/Runtime/Engine/Private/Vehicles/PhysXVehicleManager.cpp b/Engine/Source/Runtime/Engine/Private/Vehicles/PhysXVehicleManager.cpp
--- a/Engine/Source/Runtime/Engine/Private/Vehicles/PhysXVehicleManager.cpp
+++ b/Engine/Source/Runtime/Engine/Private/Vehicles/PhysXVehicleManager.cpp
@@ -122,8 +122,7 @@ void FPhysXVehicleManager::UpdateTireFrictionTableInternal()
 	PxVehicleDrivableSurfaceType		DrivableSurfaceTypes[MAX_NUM_MATERIALS];
 
 	// Gather all the physical materials
-	GPhysXSDK->getMaterials(AllPhysicsMaterials, MAX_NUM_MATERIALS);
-	uint32 NumMaterials = GPhysXSDK->getNbMaterials();
+	uint32 NumMaterials = GPhysXSDK->getMaterials(AllPhysicsMaterials, MAX_NUM_MATERIALS);
 
 	uint32 NumTireTypes = UTireType::AllTireTypes.Num();
 
