Expose customized collision flag for static mesh editor

diff --git a/Engine/Source/Runtime/Engine/Classes/Engine/StaticMesh.h b/Engine/Source/Runtime/Engine/Classes/Engine/StaticMesh.h
--- a/Engine/Source/Runtime/Engine/Classes/Engine/StaticMesh.h
+++ b/Engine/Source/Runtime/Engine/Classes/Engine/StaticMesh.h
@@ -436,7 +436,7 @@ class UStaticMesh : public UObject, public IInterface_CollisionDataProvider, pub
 	FAssetEditorOrbitCameraPosition EditorCameraPosition;
 
 	/** If the user has modified collision in any way or has custom collision imported. Used for determining if to auto generate collision on import */
-	UPROPERTY()
+	UPROPERTY(EditAnywhere, Category = Collision)
 	bool bCustomizedCollision;
 
 #endif // WITH_EDITORONLY_DATA
