Expose GameMode properties bUseSeamlessTravel, bPauseable, bStartPlayersAsSpectator, bDelayedStart and OptionsString to defaults and blueprints

diff --git a/Engine/Source/Runtime/Engine/Classes/GameFramework/GameMode.h b/Engine/Source/Runtime/Engine/Classes/GameFramework/GameMode.h
--- a/Engine/Source/Runtime/Engine/Classes/GameFramework/GameMode.h
+++ b/Engine/Source/Runtime/Engine/Classes/GameFramework/GameMode.h
@@ -143,23 +143,23 @@ public:
 	 * perform map travels using SeamlessTravel() which loads in the background and doesn't disconnect clients
 	 * @see World::SeamlessTravel()
 	 */
-	UPROPERTY()
+	UPROPERTY(EditAnywhere, BlueprintReadWrite, Category="GameMode")
 	uint32 bUseSeamlessTravel : 1;
 
 	/** Whether the game is pauseable. */
-	UPROPERTY()
+	UPROPERTY(EditAnywhere, BlueprintReadWrite, Category="GameMode")
 	uint32 bPauseable:1;
 
 	/** Whether players should immediately spawn when logging in, or stay as spectators until they manually spawn */
-	UPROPERTY()
+	UPROPERTY(EditAnywhere, BlueprintReadOnly, Category="GameMode")
 	uint32 bStartPlayersAsSpectators : 1;
 
 	/** Whether the game should immediately start when the first player logs in. Affects the default behavior of ReadyToStartMatch */
-	UPROPERTY()
+	UPROPERTY(EditAnywhere, BlueprintReadOnly, Category="GameMode")
 	uint32 bDelayedStart : 1;
 
 	/** Save options string and parse it when needed */
-	UPROPERTY()
+	UPROPERTY(BlueprintReadOnly, Category="GameMode")
 	FString OptionsString;
 
 	/** The default pawn class used by players. */
