Expose current swing1,swing2,twist angles for engine and blueprints

diff --git a/Engine/Source/Runtime/Engine/Classes/PhysicsEngine/ConstraintInstance.h b/Engine/Source/Runtime/Engine/Classes/PhysicsEngine/ConstraintInstance.h
--- a/Engine/Source/Runtime/Engine/Classes/PhysicsEngine/ConstraintInstance.h
+++ b/Engine/Source/Runtime/Engine/Classes/PhysicsEngine/ConstraintInstance.h
@@ -419,6 +419,14 @@ public:
 	// Pass in reference orientation in (maintains reference position). If the constraint is currently active, this will set its active local pose. Otherwise the change will take affect in InitConstraint.
 	void SetRefOrientation(EConstraintFrame::Type Frame, const FVector& PriAxis, const FVector& SecAxis);
 	
+	// The current twist of the constraint
+	float GetCurrentTwist() const;
+
+	// The current swing1 of the constraint
+	float GetCurrentSwing1() const;
+
+	// The current swing2 of the constraint
+	float GetCurrentSwing2() const;
 
 	// Get component ref frame
 	FTransform GetRefFrame(EConstraintFrame::Type Frame) const;
diff --git a/Engine/Source/Runtime/Engine/Classes/PhysicsEngine/PhysicsConstraintComponent.h b/Engine/Source/Runtime/Engine/Classes/PhysicsEngine/PhysicsConstraintComponent.h
--- a/Engine/Source/Runtime/Engine/Classes/PhysicsEngine/PhysicsConstraintComponent.h
+++ b/Engine/Source/Runtime/Engine/Classes/PhysicsEngine/PhysicsConstraintComponent.h
@@ -218,6 +218,19 @@ public:
 	UFUNCTION(BlueprintCallable, Category = "Physics|Components|PhysicsConstraint")
 	ENGINE_API void	SetAngularTwistLimit(EAngularConstraintMotion ConstraintType, float TwistLimitAngle);
 
+
+	/** Gets the current Angular Twist of the constraint */
+	UFUNCTION(BlueprintCallable, Category = "Physics|Components|PhysicsConstraint")
+	ENGINE_API float GetCurrentTwist() const;
+
+	/** Gets the current Swing1 of the constraint */
+	UFUNCTION(BlueprintCallable, Category = "Physics|Components|PhysicsConstraint")
+	ENGINE_API float GetCurrentSwing1() const;
+
+	/** Gets the current Swing2 of the constraint */
+	UFUNCTION(BlueprintCallable, Category = "Physics|Components|PhysicsConstraint")
+	ENGINE_API float GetCurrentSwing2() const;
+
 	/** 
 	 *	Update the reference frames held inside the constraint that indicate the joint location in the reference frame 
 	 *	of the two connected bodies. You should call this whenever the constraint or either Component moves, or if you change
diff --git a/Engine/Source/Runtime/Engine/Private/PhysicsEngine/ConstraintInstance.cpp b/Engine/Source/Runtime/Engine/Private/PhysicsEngine/ConstraintInstance.cpp
--- a/Engine/Source/Runtime/Engine/Private/PhysicsEngine/ConstraintInstance.cpp
+++ b/Engine/Source/Runtime/Engine/Private/PhysicsEngine/ConstraintInstance.cpp
@@ -897,6 +897,48 @@ void FConstraintInstance::SetAngularOrientationTarget(const FQuat& InPosTarget)
 	AngularOrientationTarget = OrientationTargetRot;
 }
 
+float FConstraintInstance::GetCurrentSwing1() const
+{
+	float Swing1 = 0.f;
+#if WITH_PHYSX
+	PxD6Joint* Joint = (PxD6Joint*)ConstraintData;
+	if (Joint)
+	{
+		Swing1 = Joint->getSwingZAngle();
+	}
+#endif
+
+	return Swing1;
+}
+
+float FConstraintInstance::GetCurrentSwing2() const
+{
+	float Swing2 = 0.f;
+#if WITH_PHYSX
+	PxD6Joint* Joint = (PxD6Joint*)ConstraintData;
+	if (Joint)
+	{
+		Swing2 = Joint->getSwingYAngle();
+	}
+#endif
+
+	return Swing2;
+}
+
+float FConstraintInstance::GetCurrentTwist() const
+{
+	float Twist = 0.f;
+#if WITH_PHYSX
+	PxD6Joint* Joint = (PxD6Joint*)ConstraintData;
+	if (Joint)
+	{
+		Twist = Joint->getTwist();
+	}
+#endif
+
+	return Twist;
+}
+
 
 /** Function for setting target angular velocity. */
 void FConstraintInstance::SetAngularVelocityTarget(const FVector& InVelTarget)
diff --git a/Engine/Source/Runtime/Engine/Private/PhysicsEngine/PhysicsConstraintComponent.cpp b/Engine/Source/Runtime/Engine/Private/PhysicsEngine/PhysicsConstraintComponent.cpp
--- a/Engine/Source/Runtime/Engine/Private/PhysicsEngine/PhysicsConstraintComponent.cpp
+++ b/Engine/Source/Runtime/Engine/Private/PhysicsEngine/PhysicsConstraintComponent.cpp
@@ -538,4 +538,22 @@ void UPhysicsConstraintComponent::SetAngularTwistLimit(EAngularConstraintMotion
 	ConstraintInstance.SetAngularTwistLimit(Motion, TwistLimitAngle);
 }
 
+float UPhysicsConstraintComponent::GetCurrentTwist() const
+{
+	const float CurrentTwistRads = ConstraintInstance.GetCurrentTwist();
+	return FMath::RadiansToDegrees(CurrentTwistRads);
+}
+
+float UPhysicsConstraintComponent::GetCurrentSwing1() const
+{
+	const float CurrentSwing1Rads = ConstraintInstance.GetCurrentSwing1();
+	return FMath::RadiansToDegrees(CurrentSwing1Rads);
+}
+
+float UPhysicsConstraintComponent::GetCurrentSwing2() const
+{
+	const float CurrentSwing2Rads = ConstraintInstance.GetCurrentSwing2();
+	return FMath::RadiansToDegrees(CurrentSwing2Rads);
+}
+
 #undef LOCTEXT_NAMESPACE
