Don't follow joints when computing physx size because this means each body returns the size of all bodies it's connected to. FPhysxSharedData::Get().GetCollection() should already add all bodies

diff --git a/Engine/Source/Runtime/Engine/Private/PhysicsEngine/PhysXSupport.cpp b/Engine/Source/Runtime/Engine/Private/PhysicsEngine/PhysXSupport.cpp
--- a/Engine/Source/Runtime/Engine/Private/PhysicsEngine/PhysXSupport.cpp
+++ b/Engine/Source/Runtime/Engine/Private/PhysicsEngine/PhysXSupport.cpp
@@ -713,7 +713,7 @@ SIZE_T GetPhysxObjectSize(PxBase* Obj, const PxCollection* SharedCollection)
 	PxCollection* Collection = PxCreateCollection();
 
 	Collection->add(*Obj);
-	PxSerialization::complete(*Collection, *Sr, SharedCollection, true);	// chase all other stuff (shared shaps, materials, etc) needed to serialize this collection
+	PxSerialization::complete(*Collection, *Sr, SharedCollection);	// chase all other stuff (shared shaps, materials, etc) needed to serialize this collection
 
 	FPhysXCountMemoryStream Out;
 	PxSerialization::serializeCollectionToBinary(Out, *Collection, *Sr, SharedCollection);
