Added comment for GetNumMaterials function.

diff --git a/Engine/Source/Runtime/Engine/Private/StaticMeshComponent.cpp b/Engine/Source/Runtime/Engine/Private/StaticMeshComponent.cpp
--- a/Engine/Source/Runtime/Engine/Private/StaticMeshComponent.cpp
+++ b/Engine/Source/Runtime/Engine/Private/StaticMeshComponent.cpp
@@ -1456,6 +1456,8 @@ bool UStaticMeshComponent::GetEstimatedLightAndShadowMapMemoryUsage(
 
 int32 UStaticMeshComponent::GetNumMaterials() const
 {
+	// @note : you don't have to consider Materials.Num()
+	// that only counts if overriden and it can't be more than StaticMesh->Materials. 
 	if(StaticMesh)
 	{
 		return StaticMesh->Materials.Num();
