Fix Montage to still add reference if not transient - if map transitions out, you might still like to play.

- reported by licensee.

diff --git a/Engine/Source/Runtime/Engine/Private/Animation/AnimMontage.cpp b/Engine/Source/Runtime/Engine/Private/Animation/AnimMontage.cpp
--- a/Engine/Source/Runtime/Engine/Private/Animation/AnimMontage.cpp
+++ b/Engine/Source/Runtime/Engine/Private/Animation/AnimMontage.cpp
@@ -969,7 +969,7 @@ void FAnimMontageInstance::RefreshNextPrevSections()
 
 void FAnimMontageInstance::AddReferencedObjects( FReferenceCollector& Collector )
 {
-	if (Montage && Montage->GetOuter() == GetTransientPackage())
+	if (Montage)
 	{
 		Collector.AddReferencedObject(Montage);
 	}
