#ue4 - Fix CalcVelocity() not actually clamping acceleration result of ApplyRequestedMove(). Was only an issue if someone had overridden ApplyRequestedMove() and not clamped it internally, since our version still did.

diff --git a/Engine/Source/Runtime/Engine/Private/CharacterMovementComponent.cpp b/Engine/Source/Runtime/Engine/Private/CharacterMovementComponent.cpp
--- a/Engine/Source/Runtime/Engine/Private/CharacterMovementComponent.cpp
+++ b/Engine/Source/Runtime/Engine/Private/CharacterMovementComponent.cpp
@@ -2141,7 +2141,7 @@ void UCharacterMovementComponent::CalcVelocity(float DeltaTime, float Friction,
 	float RequestedSpeed = 0.0f;
 	if (ApplyRequestedMove(DeltaTime, MaxAccel, MaxSpeed, Friction, BrakingDeceleration, RequestedAcceleration, RequestedSpeed))
 	{
-		RequestedAcceleration.ClampMaxSize(MaxAccel);
+		RequestedAcceleration = RequestedAcceleration.ClampMaxSize(MaxAccel);
 		bZeroRequestedAcceleration = false;
 	}
 
