In-function type forward declarations removed from navigation headers #UE4

diff --git a/Engine/Source/Runtime/Engine/Classes/AI/Navigation/NavFilters/NavigationQueryFilter.h b/Engine/Source/Runtime/Engine/Classes/AI/Navigation/NavFilters/NavigationQueryFilter.h
--- a/Engine/Source/Runtime/Engine/Classes/AI/Navigation/NavFilters/NavigationQueryFilter.h
+++ b/Engine/Source/Runtime/Engine/Classes/AI/Navigation/NavFilters/NavigationQueryFilter.h
@@ -4,6 +4,7 @@
 #include "NavigationQueryFilter.generated.h"
 
 class UNavArea;
+class ANavigationData;
 struct FNavigationQueryFilter;
 
 USTRUCT()
@@ -109,7 +110,7 @@ public:
 	virtual void SetExcludeFlags(uint16 Flags) = 0;
 	virtual uint16 GetExcludeFlags() const = 0;
 
-	virtual class INavigationQueryFilterInterface* CreateCopy() const = 0;
+	virtual INavigationQueryFilterInterface* CreateCopy() const = 0;
 };
 
 struct ENGINE_API FNavigationQueryFilter : public TSharedFromThis<FNavigationQueryFilter>
@@ -169,7 +170,7 @@ public:
 		QueryFilterImpl = MakeShareable(new FilterType());
 	}
 
-	FORCEINLINE_DEBUGGABLE void SetFilterImplementation(const class INavigationQueryFilterInterface* InQueryFilterImpl)
+	FORCEINLINE_DEBUGGABLE void SetFilterImplementation(const INavigationQueryFilterInterface* InQueryFilterImpl)
 	{
 		QueryFilterImpl = MakeShareable(InQueryFilterImpl->CreateCopy());
 	}
@@ -215,13 +216,13 @@ class ENGINE_API UNavigationQueryFilter : public UObject
 	FNavigationFilterFlags ExcludeFlags;
 
 	/** get filter for given navigation data and initialize on first access */
-	TSharedPtr<const FNavigationQueryFilter> GetQueryFilter(const class ANavigationData* NavData) const;
+	TSharedPtr<const FNavigationQueryFilter> GetQueryFilter(const ANavigationData* NavData) const;
 	
 	/** helper functions for accessing filter */
-	static TSharedPtr<const FNavigationQueryFilter> GetQueryFilter(const class ANavigationData* NavData, UClass* FilterClass);
+	static TSharedPtr<const FNavigationQueryFilter> GetQueryFilter(const ANavigationData* NavData, UClass* FilterClass);
 
 	template<class T>
-	static TSharedPtr<const FNavigationQueryFilter> GetQueryFilter(const class ANavigationData* NavData, UClass* FilterClass = T::StaticClass())
+	static TSharedPtr<const FNavigationQueryFilter> GetQueryFilter(const ANavigationData* NavData, UClass* FilterClass = T::StaticClass())
 	{
 		return GetQueryFilter(NavData, FilterClass);
 	}
@@ -241,5 +242,5 @@ protected:
 	int32 FindAreaOverride(TSubclassOf<UNavArea> AreaClass) const;
 
 	/** setup filter for given navigation data, use to create custom filters */
-	virtual void InitializeFilter(const class ANavigationData* NavData, FNavigationQueryFilter* Filter) const;
+	virtual void InitializeFilter(const ANavigationData* NavData, FNavigationQueryFilter* Filter) const;
 };
diff --git a/Engine/Source/Runtime/Engine/Private/AI/Navigation/NavFilters/NavigationQueryFilter.cpp b/Engine/Source/Runtime/Engine/Private/AI/Navigation/NavFilters/NavigationQueryFilter.cpp
--- a/Engine/Source/Runtime/Engine/Private/AI/Navigation/NavFilters/NavigationQueryFilter.cpp
+++ b/Engine/Source/Runtime/Engine/Private/AI/Navigation/NavFilters/NavigationQueryFilter.cpp
@@ -116,7 +116,7 @@ UNavigationQueryFilter::UNavigationQueryFilter(const FObjectInitializer& ObjectI
 	ExcludeFlags.Packed = 0;
 }
 
-TSharedPtr<const FNavigationQueryFilter> UNavigationQueryFilter::GetQueryFilter(const class ANavigationData* NavData) const
+TSharedPtr<const FNavigationQueryFilter> UNavigationQueryFilter::GetQueryFilter(const ANavigationData* NavData) const
 {
 	TSharedPtr<const FNavigationQueryFilter> SharedFilter = NavData->GetQueryFilter(GetClass());
 	if (!SharedFilter.IsValid())
@@ -133,7 +133,7 @@ TSharedPtr<const FNavigationQueryFilter> UNavigationQueryFilter::GetQueryFilter(
 	return SharedFilter;
 }
 
-void UNavigationQueryFilter::InitializeFilter(const class ANavigationData* NavData, FNavigationQueryFilter* Filter) const
+void UNavigationQueryFilter::InitializeFilter(const ANavigationData* NavData, FNavigationQueryFilter* Filter) const
 {
 	// apply overrides
 	for (int32 i = 0; i < Areas.Num(); i++)
diff --git a/Engine/Source/Runtime/Engine/Private/AI/Navigation/RecastNavMeshGenerator.cpp b/Engine/Source/Runtime/Engine/Private/AI/Navigation/RecastNavMeshGenerator.cpp
--- a/Engine/Source/Runtime/Engine/Private/AI/Navigation/RecastNavMeshGenerator.cpp
+++ b/Engine/Source/Runtime/Engine/Private/AI/Navigation/RecastNavMeshGenerator.cpp
@@ -40,6 +40,8 @@
 #define RECAST_STAT(...) 
 #endif
 
+struct dtTileCacheAlloc;
+
 FORCEINLINE bool DoesBoxContainOrOverlapVector(const FBox& BigBox, const FVector& In)
 {
 	return (In.X >= BigBox.Min.X) && (In.X <= BigBox.Max.X) 
@@ -1554,7 +1556,7 @@ void FRecastTileGenerator::GatherGeometry(const FRecastNavMeshGenerator* ParentG
 	}
 }
 
-void FRecastTileGenerator::ApplyVoxelFilter(struct rcHeightfield* HF, float WalkableRadius)
+void FRecastTileGenerator::ApplyVoxelFilter(rcHeightfield* HF, float WalkableRadius)
 {
 	SCOPE_CYCLE_COUNTER(STAT_Navigation_TileVoxelFilteringAsync);
 
@@ -2238,7 +2240,7 @@ void FRecastTileGenerator::MarkStaticAreas(FNavMeshBuildContext& BuildContext, r
 
 struct FTileGenerationContext
 {
-	FTileGenerationContext(struct dtTileCacheAlloc* MyAllocator) :
+	FTileGenerationContext(dtTileCacheAlloc* MyAllocator) :
 		Allocator(MyAllocator), Layer(0), DistanceField(0), ContourSet(0), ClusterSet(0), PolyMesh(0), DetailMesh(0)
 	{
 	}
