Fixed crash due to FDynamicMeshBuilder::GetMesh using PDI->RegisterDynamicResource for resources referenced by the mesh collector

diff --git a/Engine/Source/Runtime/Engine/Private/DynamicMeshBuilder.cpp b/Engine/Source/Runtime/Engine/Private/DynamicMeshBuilder.cpp
--- a/Engine/Source/Runtime/Engine/Private/DynamicMeshBuilder.cpp
+++ b/Engine/Source/Runtime/Engine/Private/DynamicMeshBuilder.cpp
@@ -412,22 +412,54 @@ void FDynamicMeshBuilder::AddTriangles(const TArray<int32> &InIndices)
 	IndexBuffer->Indices.Append(InIndices);
 }
 
+class FMeshBuilderOneFrameResources : public FOneFrameResource
+{
+public:
+
+	FMeshBuilderOneFrameResources() :
+		VertexBuffer(NULL),
+		IndexBuffer(NULL),
+		VertexFactory(NULL),
+		PrimitiveUniformBuffer(NULL)
+	{}
+
+	FDynamicMeshVertexBuffer* VertexBuffer;
+	FDynamicMeshIndexBuffer* IndexBuffer;
+	FDynamicMeshVertexFactory* VertexFactory;
+	FDynamicMeshPrimitiveUniformBuffer* PrimitiveUniformBuffer;
+
+	virtual ~FMeshBuilderOneFrameResources()
+	{
+		VertexBuffer->ReleaseResource();
+		IndexBuffer->ReleaseResource();
+		VertexFactory->ReleaseResource();
+		PrimitiveUniformBuffer->ReleaseResource();
+
+		delete VertexBuffer;
+		delete IndexBuffer;
+		delete VertexFactory;
+		delete PrimitiveUniformBuffer;
+	}
+};
+
 void FDynamicMeshBuilder::GetMesh(const FMatrix& LocalToWorld,const FMaterialRenderProxy* MaterialRenderProxy,uint8 DepthPriorityGroup,bool bDisableBackfaceCulling, bool bReceivesDecals, int32 ViewIndex, FMeshElementCollector& Collector) 
 {
 	// Only draw non-empty meshes.
 	if(VertexBuffer->Vertices.Num() > 0 && IndexBuffer->Indices.Num() > 0)
 	{
-		FPrimitiveDrawInterface* PDI = Collector.GetPDI(ViewIndex);
-		// Register the dynamic resources with the PDI.
-		PDI->RegisterDynamicResource(VertexBuffer);
-		PDI->RegisterDynamicResource(IndexBuffer);
+		FMeshBuilderOneFrameResources& OneFrameResources = Collector.AllocateOneFrameResource<FMeshBuilderOneFrameResources>();
 
-		// Create the vertex factory.
-		FDynamicMeshVertexFactory* VertexFactory = new FDynamicMeshVertexFactory(VertexBuffer);
-		PDI->RegisterDynamicResource(VertexFactory);
+		OneFrameResources.VertexBuffer = VertexBuffer;
+		OneFrameResources.IndexBuffer = IndexBuffer;
+
+		OneFrameResources.VertexBuffer->InitResource();
+		OneFrameResources.IndexBuffer->InitResource();
+
+		OneFrameResources.VertexFactory = new FDynamicMeshVertexFactory(VertexBuffer);
+		OneFrameResources.VertexFactory->InitResource();
 
 		// Create the primitive uniform buffer.
-		FDynamicMeshPrimitiveUniformBuffer* PrimitiveUniformBuffer = new FDynamicMeshPrimitiveUniformBuffer();
+		OneFrameResources.PrimitiveUniformBuffer = new FDynamicMeshPrimitiveUniformBuffer();
 		FPrimitiveUniformShaderParameters PrimitiveParams = GetPrimitiveUniformShaderParameters(
 			LocalToWorld,
 			LocalToWorld.GetOrigin(),
@@ -442,21 +474,22 @@ void FDynamicMeshBuilder::GetMesh(const FMatrix& LocalToWorld,const FMaterialRen
 
 		if (IsInGameThread())
 		{
-			BeginSetUniformBufferContents(*PrimitiveUniformBuffer, PrimitiveParams);
+			BeginSetUniformBufferContents(*OneFrameResources.PrimitiveUniformBuffer, PrimitiveParams);
 		}
 		else
 		{
-			PrimitiveUniformBuffer->SetContents(PrimitiveParams);
+			OneFrameResources.PrimitiveUniformBuffer->SetContents(PrimitiveParams);
 		}
-		PDI->RegisterDynamicResource(PrimitiveUniformBuffer);
+
+		OneFrameResources.PrimitiveUniformBuffer->InitResource();
 
 		// Draw the mesh.
 		FMeshBatch& Mesh = Collector.AllocateMesh();
 		FMeshBatchElement& BatchElement = Mesh.Elements[0];
 		BatchElement.IndexBuffer = IndexBuffer;
-		Mesh.VertexFactory = VertexFactory;
+		Mesh.VertexFactory = OneFrameResources.VertexFactory;
 		Mesh.MaterialRenderProxy = MaterialRenderProxy;
-		BatchElement.PrimitiveUniformBufferResource = PrimitiveUniformBuffer;
+		BatchElement.PrimitiveUniformBufferResource = OneFrameResources.PrimitiveUniformBuffer;
 		// previous l2w not used so treat as static
 		BatchElement.FirstIndex = 0;
 		BatchElement.NumPrimitives = IndexBuffer->Indices.Num() / 3;
