[UE-5790] fix inconsistent component attachment issues

diff --git a/Engine/Source/Runtime/Engine/Private/SceneComponent.cpp b/Engine/Source/Runtime/Engine/Private/SceneComponent.cpp
--- a/Engine/Source/Runtime/Engine/Private/SceneComponent.cpp
+++ b/Engine/Source/Runtime/Engine/Private/SceneComponent.cpp
@@ -771,6 +771,10 @@ void USceneComponent::AttachTo(class USceneComponent* Parent, FName InSocketName
 		const bool bSavedDisableDetachmentUpdateOverlaps = bDisableDetachmentUpdateOverlaps;
 		bDisableDetachmentUpdateOverlaps = true;
 
+		// Find out if we're already attached, and save off our position in the array if we are
+		int32 LastAttachIndex = INDEX_NONE;
+		Parent->AttachChildren.Find(this, LastAttachIndex);
+
 		// Make sure we are detached
 		const bool bMaintainWorldPosition = (AttachType == EAttachLocation::KeepWorldPosition);
 		DetachFromParent(bMaintainWorldPosition);
@@ -811,7 +815,15 @@ void USceneComponent::AttachTo(class USceneComponent* Parent, FName InSocketName
 
 		OnAttachmentChanged();
 
-		Parent->AttachChildren.Add(this);
+		// Preserve order of previous attachment if valid (in case we're doing a reattach operation inside a loop that might assume the AttachChildren order won't change)
+		if(LastAttachIndex != INDEX_NONE)
+		{
+			Parent->AttachChildren.Insert(this, LastAttachIndex);
+		}
+		else
+		{
+			Parent->AttachChildren.Add(this);
+		}
 
 		switch ( AttachType )
 		{
diff --git a/Engine/Source/Runtime/Engine/Private/SimpleConstructionScript.cpp b/Engine/Source/Runtime/Engine/Private/SimpleConstructionScript.cpp
--- a/Engine/Source/Runtime/Engine/Private/SimpleConstructionScript.cpp
+++ b/Engine/Source/Runtime/Engine/Private/SimpleConstructionScript.cpp
@@ -17,6 +17,25 @@
 const FName USimpleConstructionScript::DefaultSceneRootVariableName = FName(TEXT("DefaultSceneRoot"));
 #endif
 
+namespace
+{
+	// Helper method to register instanced components post-construction
+	void RegisterInstancedComponent(UActorComponent* InstancedComponent)
+	{
+		InstancedComponent->RegisterComponent();
+
+		// If this is a scene component, recursively register any child components as well
+		USceneComponent* InstancedSceneComponent = Cast<USceneComponent>(InstancedComponent);
+		if(InstancedSceneComponent != nullptr)
+		{
+			for(auto InstancedChildComponent : InstancedSceneComponent->AttachChildren)
+			{
+				RegisterInstancedComponent(InstancedChildComponent);
+			}
+		}
+	}
+}
+
 USimpleConstructionScript::USimpleConstructionScript(const FObjectInitializer& ObjectInitializer)
 	: Super(ObjectInitializer)
 {
@@ -345,21 +364,11 @@ void USimpleConstructionScript::ExecuteScriptOnActor(AActor* Actor, const FTrans
 			}
 		}
 
-		// Register all instanced SCS components once SCS execution has finished; do the scene component hierarchy first, followed by the remaining actor components (in case they happen to depend on something in the scene hierarchy)
+		// Register all instanced SCS components once SCS execution has finished; sorted in order to register the scene component hierarchy first, followed by the remaining actor components (in case they happen to depend on something in the scene hierarchy)
 		InstancedComponents.Sort([](const UActorComponent& A, const UActorComponent& B) { return A.IsA<USceneComponent>(); });
 		for(auto InstancedComponent : InstancedComponents)
 		{
-			InstancedComponent->RegisterComponent();
-
-			// If this is a scene component, register any child components as well
-			USceneComponent* InstancedSceneComponent = Cast<USceneComponent>(InstancedComponent);
-			if(InstancedSceneComponent != nullptr)
-			{
-				for(auto InstancedChildComponent : InstancedSceneComponent->AttachChildren)
-				{
-					InstancedChildComponent->RegisterComponent();
-				}
-			}
+			RegisterInstancedComponent(InstancedComponent);
 		}
 	}
 	else if(Actor->GetRootComponent() == NULL) // Must have a root component at the end of SCS, so if we don't have one already (from base class), create a SceneComponent now
