fix UE-5780 Shot console command (F9) isn't working in game
 workaround to get "shot working" (save BMP instead of PNG)

diff --git a/Engine/Source/Runtime/Engine/Private/GameViewportClient.cpp b/Engine/Source/Runtime/Engine/Private/GameViewportClient.cpp
--- a/Engine/Source/Runtime/Engine/Private/GameViewportClient.cpp
+++ b/Engine/Source/Runtime/Engine/Private/GameViewportClient.cpp
@@ -2599,7 +2599,10 @@ bool UGameViewportClient::HandleScreenshotCommand( const TCHAR* Cmd, FOutputDevi
 	if(Viewport)
 	{
 		const bool bShowUI = FParse::Command(&Cmd, TEXT("SHOWUI"));
-		FScreenshotRequest::RequestScreenshot( bShowUI, TEXT("png") );
+
+//		FScreenshotRequest::RequestScreenshot( bShowUI, TEXT("png") );
+		// PNG is disabled for now as it breaks "shot" command in game, see UE-5780 
+		FScreenshotRequest::RequestScreenshot( bShowUI, TEXT("bmp") );
 
 		GScreenMessagesRestoreState = GAreScreenMessagesEnabled;
 		GAreScreenMessagesEnabled = false;
