#ue4 - Add substepping to camera SpringArmComponent to address variance at different framerates.

UDN: https://udn.unrealengine.com/questions/211646/fluctuating-frame-rate-and-jittery-vehicles.html
Jira: UE-3685

diff --git a/Engine/Source/Runtime/Engine/Classes/GameFramework/SpringArmComponent.h b/Engine/Source/Runtime/Engine/Classes/GameFramework/SpringArmComponent.h
--- a/Engine/Source/Runtime/Engine/Classes/GameFramework/SpringArmComponent.h
+++ b/Engine/Source/Runtime/Engine/Classes/GameFramework/SpringArmComponent.h
@@ -59,25 +59,42 @@ class ENGINE_API USpringArmComponent : public USceneComponent
 	UPROPERTY(EditAnywhere, BlueprintReadWrite, Category=CameraSettings)
 	uint32 bInheritRoll : 1;
 
-	/** If true, camera lags behind target position to smooth its movement */
+	/**
+	 * If true, camera lags behind target position to smooth its movement.
+	 * @see CameraLagSpeed
+	 */
 	UPROPERTY(EditAnywhere, BlueprintReadWrite, Category=Lag)
 	uint32 bEnableCameraLag : 1;
 
-	/** If true, camera lags behind target position to smooth its movement */
+	/**
+	 * If true, camera lags behind target position to smooth its movement.
+	 * @see CameraRotationLagSpeed
+	 */
 	UPROPERTY(EditAnywhere, BlueprintReadWrite, Category=Lag)
 	uint32 bEnableCameraRotationLag : 1;
 
-	/** If bEnableCameraLag is true, controls how quickly camera reaches target position */
-	UPROPERTY(EditAnywhere, BlueprintReadWrite, Category=Lag, meta=(editcondition="bEnableCameraLag"))
+	/** If bEnableCameraLag is true, controls how quickly camera reaches target position. Low values are slower (more lag), high values are faster (less lag), while zero is instant (no lag). */
+	UPROPERTY(EditAnywhere, BlueprintReadWrite, Category=Lag, meta=(editcondition="bEnableCameraLag", ClampMin="0.0", ClampMax="1000.0", UIMin = "0.0", UIMax = "1000.0"))
 	float CameraLagSpeed;
 
-	/** If bEnableCameraLag is true, controls how quickly camera reaches target position */
-	UPROPERTY(EditAnywhere, BlueprintReadWrite, Category=Lag, meta=(editcondition = "bEnableCameraRotationLag"))
+	/** If bEnableCameraRotationLag is true, controls how quickly camera reaches target position. Low values are slower (more lag), high values are faster (less lag), while zero is instant (no lag). */
+	UPROPERTY(EditAnywhere, BlueprintReadWrite, Category=Lag, meta=(editcondition = "bEnableCameraRotationLag", ClampMin="0.0", ClampMax="1000.0", UIMin = "0.0", UIMax = "1000.0"))
 	float CameraRotationLagSpeed;
 
-
-	/** Temporary variable when using camera lag, to record previous camera position */
+	/**
+	 * If bUseCameraLagSubstepping is true, sub-step camera damping so that it handles fluctuating frame rates well (though this comes at a cost).
+	 * @see CameraLagMaxTimeStep
+	 */
+	UPROPERTY(EditAnywhere, BlueprintReadWrite, Category = Lag, AdvancedDisplay)
+	bool bUseCameraLagSubstepping;
+	
+	/** Max time step used when sub-stepping camera lag. */
+	UPROPERTY(EditAnywhere, BlueprintReadWrite, Category = Lag, AdvancedDisplay, meta=(editcondition = "bUseCameraLagSubstepping", ClampMin="0.005", ClampMax="0.5", UIMin = "0.005", UIMax = "0.5"))
+	float CameraLagMaxTimeStep;
+
+	/** Temporary variables when using camera lag, to record previous camera position */
 	FVector PreviousDesiredLoc;
+	FVector PreviousArmOrigin;
 	/** Temporary variable for lagging camera rotation, for previous rotation */
 	FRotator PreviousDesiredRot;
 
diff --git a/Engine/Source/Runtime/Engine/Private/GameFramework/SpringArmComponent.cpp b/Engine/Source/Runtime/Engine/Private/GameFramework/SpringArmComponent.cpp
--- a/Engine/Source/Runtime/Engine/Private/GameFramework/SpringArmComponent.cpp
+++ b/Engine/Source/Runtime/Engine/Private/GameFramework/SpringArmComponent.cpp
@@ -9,8 +9,6 @@
 static void SetDeprecatedControllerViewRotation(USpringArmComponent& Component, bool bValue);
 const FName USpringArmComponent::SocketName(TEXT("SpringEndpoint"));
 
-extern float GAverageMS;
-
 USpringArmComponent::USpringArmComponent(const FObjectInitializer& ObjectInitializer)
 	: Super(ObjectInitializer)
 {
@@ -35,6 +33,9 @@ USpringArmComponent::USpringArmComponent(const FObjectInitializer& ObjectInitial
 
 	RelativeSocketRotation = FQuat::Identity;
 
+	bUseCameraLagSubstepping = true;
+	CameraLagMaxTimeStep = 1.f / 60.f;
+	
 	// Init deprecated var, for old code that may refer to it.
 	SetDeprecatedControllerViewRotation(*this, bUsePawnControlRotation);
 }
@@ -62,21 +63,60 @@ void USpringArmComponent::UpdateDesiredArmLocation(bool bDoTrace, bool bDoLocati
 		}
 	}
 
+	const float InverseCameraLagMaxTimeStep = (1.f / CameraLagMaxTimeStep);
+
 	// Apply 'lag' to rotation if desired
 	if(bDoRotationLag)
 	{
-		DesiredRot = FMath::RInterpTo(PreviousDesiredRot, DesiredRot, DeltaTime, CameraRotationLagSpeed);
+		if (bUseCameraLagSubstepping && DeltaTime > CameraLagMaxTimeStep && CameraRotationLagSpeed > 0.f)
+		{
+			const FRotator ArmRotStep = (DesiredRot - PreviousDesiredRot).GetNormalized() * (CameraLagMaxTimeStep / DeltaTime);
+			FRotator LerpTarget = PreviousDesiredRot;
+			float RemainingTime = DeltaTime;
+			while (RemainingTime > KINDA_SMALL_NUMBER)
+			{
+				const float LerpAmount = FMath::Min(CameraLagMaxTimeStep, RemainingTime);
+				LerpTarget += ArmRotStep * (LerpAmount * InverseCameraLagMaxTimeStep);
+				RemainingTime -= LerpAmount;
+
+				DesiredRot = FMath::RInterpTo(PreviousDesiredRot, LerpTarget, LerpAmount, CameraRotationLagSpeed);
+				PreviousDesiredRot = DesiredRot;
+			}
+		}
+		else
+		{
+			DesiredRot = FMath::RInterpTo(PreviousDesiredRot, DesiredRot, DeltaTime, CameraRotationLagSpeed);
+		}
 	}
 	PreviousDesiredRot = DesiredRot;
 
 	// Get the spring arm 'origin', the target we want to look at
 	FVector ArmOrigin = GetComponentLocation() + TargetOffset;
-	// We lag the target, not the actuall camera position, so rotating the camera around does not hav lag
+	// We lag the target, not the actual camera position, so rotating the camera around does not have lag
 	FVector DesiredLoc = ArmOrigin;
 	if (bDoLocationLag)
 	{
-		DesiredLoc = FMath::VInterpTo(PreviousDesiredLoc, DesiredLoc, DeltaTime, CameraLagSpeed);
+		if (bUseCameraLagSubstepping && DeltaTime > CameraLagMaxTimeStep && CameraLagSpeed > 0.f)
+		{
+			const FVector ArmMovementStep = (ArmOrigin - PreviousArmOrigin) * (CameraLagMaxTimeStep / DeltaTime);			
+			FVector LerpTarget = PreviousArmOrigin;
+			float RemainingTime = DeltaTime;
+			while (RemainingTime > KINDA_SMALL_NUMBER)
+			{
+				const float LerpAmount = FMath::Min(CameraLagMaxTimeStep, RemainingTime);
+				LerpTarget += ArmMovementStep * (LerpAmount * InverseCameraLagMaxTimeStep);
+				RemainingTime -= LerpAmount;
+
+				DesiredLoc = FMath::VInterpTo(PreviousDesiredLoc, LerpTarget, LerpAmount, CameraLagSpeed);
+				PreviousDesiredLoc = DesiredLoc;
+			}
+		}
+		else
+		{
+			DesiredLoc = FMath::VInterpTo(PreviousDesiredLoc, DesiredLoc, DeltaTime, CameraLagSpeed);
+		}
 	}
+	PreviousArmOrigin = ArmOrigin;
 	PreviousDesiredLoc = DesiredLoc;
 
 	// Now offset camera position back along our rotation
@@ -121,6 +161,12 @@ FVector USpringArmComponent::BlendLocations(const FVector& DesiredArmLocation, c
 void USpringArmComponent::OnRegister()
 {
 	Super::OnRegister();
+
+	// enforce reasonable limits to avoid potential div-by-zero
+	CameraLagMaxTimeStep = FMath::Max(CameraLagMaxTimeStep, 1.f / 200.f);
+	CameraLagSpeed = FMath::Max(CameraLagSpeed, 0.f);
+
+	// Set initial location (without lag).
 	UpdateDesiredArmLocation(false, false, false, 0.f);
 
 	// Init deprecated var, for old code that may refer to it.
