Add a GetNavAgentLocation override to ACharacter to return character's "feet" location #UE4

diff --git a/Engine/Source/Runtime/Engine/Classes/GameFramework/Character.h b/Engine/Source/Runtime/Engine/Classes/GameFramework/Character.h
--- a/Engine/Source/Runtime/Engine/Classes/GameFramework/Character.h
+++ b/Engine/Source/Runtime/Engine/Classes/GameFramework/Character.h
@@ -288,6 +288,10 @@ public:
 	/** @return Desired translation offset of mesh. */
 	const FVector& GetBaseTranslationOffset() const { return BaseTranslationOffset; }
 
+	// Begin INavAgentInterface Interface
+	virtual FVector GetNavAgentLocation() const override;
+	// End INavAgentInterface Interface
+
 	/** Default crouched eye height */
 	UPROPERTY(EditAnywhere, BlueprintReadWrite, Category=Camera)
 	float CrouchedEyeHeight;
diff --git a/Engine/Source/Runtime/Engine/Classes/GameFramework/NavMovementComponent.h b/Engine/Source/Runtime/Engine/Classes/GameFramework/NavMovementComponent.h
--- a/Engine/Source/Runtime/Engine/Classes/GameFramework/NavMovementComponent.h
+++ b/Engine/Source/Runtime/Engine/Classes/GameFramework/NavMovementComponent.h
@@ -62,11 +62,11 @@ public:
 	/** @returns location of controlled actor - meaning center of collision bounding box */
 	FORCEINLINE FVector GetActorLocation() const { return UpdatedComponent ? UpdatedComponent->GetComponentLocation() : FVector(FLT_MAX); }
 	/** @returns location of controlled actor's "feet" meaning center of bottom of collision bounding box */
-	FORCEINLINE FVector GetActorFeetLocation() const { return UpdatedComponent ? (UpdatedComponent->GetComponentLocation() - FVector(0,0,UpdatedComponent->Bounds.BoxExtent.Z)) : FVector::ZeroVector; }
+	FORCEINLINE FVector GetActorFeetLocation() const { return UpdatedComponent ? (UpdatedComponent->GetComponentLocation() - FVector(0,0,UpdatedComponent->Bounds.BoxExtent.Z)) : FNavigationSystem::InvalidLocation; }
 	/** @returns based location of controlled actor */
 	virtual FBasedPosition GetActorFeetLocationBased() const;
 	/** @returns navigation location of controlled actor */
-	FORCEINLINE FVector GetActorNavLocation() const { INavAgentInterface* MyOwner = Cast<INavAgentInterface>(GetOwner()); return MyOwner ? MyOwner->GetNavAgentLocation() : FVector::ZeroVector; }
+	FORCEINLINE FVector GetActorNavLocation() const { INavAgentInterface* MyOwner = Cast<INavAgentInterface>(GetOwner()); return MyOwner ? MyOwner->GetNavAgentLocation() : FNavigationSystem::InvalidLocation; }
 
 	/** request direct movement */
 	virtual void RequestDirectMove(const FVector& MoveVelocity, bool bForceMaxSpeed);
diff --git a/Engine/Source/Runtime/Engine/Private/Character.cpp b/Engine/Source/Runtime/Engine/Private/Character.cpp
--- a/Engine/Source/Runtime/Engine/Private/Character.cpp
+++ b/Engine/Source/Runtime/Engine/Private/Character.cpp
@@ -640,6 +640,22 @@ void ACharacter::SaveRelativeBasedMovement(const FVector& NewRelativeLocation, c
 	BasedMovement.bRelativeRotation = bRelativeRotation;
 }
 
+FVector ACharacter::GetNavAgentLocation() const
+{
+	FVector AgentLocation = FNavigationSystem::InvalidLocation;
+
+	if (GetCharacterMovement() != nullptr)
+	{
+		AgentLocation = GetCharacterMovement()->GetActorFeetLocation();
+	}
+
+	if (FNavigationSystem::IsValidLocation(AgentLocation) == false && CapsuleComponent != nullptr)
+	{
+		AgentLocation = GetActorLocation() - FVector(0, 0, CapsuleComponent->GetScaledCapsuleHalfHeight());
+	}
+
+	return AgentLocation;
+}
 
 void ACharacter::TurnOff()
 {
