Turn on Parallel Anim Evaluation by default

diff --git a/Engine/Source/Runtime/Engine/Private/SkeletalMeshComponent.cpp b/Engine/Source/Runtime/Engine/Private/SkeletalMeshComponent.cpp
--- a/Engine/Source/Runtime/Engine/Private/SkeletalMeshComponent.cpp
+++ b/Engine/Source/Runtime/Engine/Private/SkeletalMeshComponent.cpp
@@ -26,7 +26,7 @@
 #include "PhysicsEngine/PhysicsSettings.h"
 #include "PhysicsEngine/PhysicsAsset.h"
 
-TAutoConsoleVariable<int32> CVarUseParallelAnimationEvaluation(TEXT("a.ParallelAnimEvaluation"), 0, TEXT("If 1, animation evaluation will be run across the task graph system. If 0, evaluation will run purely on the game thread"));
+TAutoConsoleVariable<int32> CVarUseParallelAnimationEvaluation(TEXT("a.ParallelAnimEvaluation"), 1, TEXT("If 1, animation evaluation will be run across the task graph system. If 0, evaluation will run purely on the game thread"));
 
 class FParallelAnimationEvaluationTask
 {
