Fix UE-5580 Physics Assets flicker in Player Collision View mode - the physics collision drawing code wasn't applying scale when drawing the solid boxes, so it was drawing a box that was far bigger than it should have been & thus was intermittently failing the depth test.

diff --git a/Engine/Source/Runtime/Engine/Private/PhysicsEngine/PhysDrawing.cpp b/Engine/Source/Runtime/Engine/Private/PhysicsEngine/PhysDrawing.cpp
--- a/Engine/Source/Runtime/Engine/Private/PhysicsEngine/PhysDrawing.cpp
+++ b/Engine/Source/Runtime/Engine/Private/PhysicsEngine/PhysDrawing.cpp
@@ -95,12 +95,12 @@ void FKBoxElem::DrawElemWire(FPrimitiveDrawInterface* PDI, const FTransform& Ele
 
 void FKBoxElem::DrawElemSolid(class FPrimitiveDrawInterface* PDI, const FTransform& ElemTM, float Scale, const FMaterialRenderProxy* MaterialRenderProxy) const
 {
-	DrawBox(PDI, ElemTM.ToMatrixWithScale(), 0.5f * FVector(X, Y, Z), MaterialRenderProxy, SDPG_World );
+	DrawBox(PDI, ElemTM.ToMatrixWithScale(), Scale * 0.5f * FVector(X, Y, Z), MaterialRenderProxy, SDPG_World );
 }
 
 void FKBoxElem::GetElemSolid(const FTransform& ElemTM, float Scale, const FMaterialRenderProxy* MaterialRenderProxy, int32 ViewIndex, FMeshElementCollector& Collector) const
 {
-	GetBoxMesh(ElemTM.ToMatrixWithScale(), 0.5f * FVector(X, Y, Z), MaterialRenderProxy, SDPG_World, ViewIndex, Collector);
+	GetBoxMesh(ElemTM.ToMatrixWithScale(), Scale * 0.5f * FVector(X, Y, Z), MaterialRenderProxy, SDPG_World, ViewIndex, Collector);
 }
 
 /////////////////////////////////////////////////////////////////////////////////////
