#ue4 - Add option to Camera SpringArmComponent for CameraLagMaxDistance. If non-zero, the target my not lag any further than this distance.

- Add bDrawDebugLagMarkers flag to enable debug visualization of the lag target and current location. The visualization changes color when the distance is being clamped.

diff --git a/Engine/Source/Runtime/Engine/Classes/GameFramework/SpringArmComponent.h b/Engine/Source/Runtime/Engine/Classes/GameFramework/SpringArmComponent.h
--- a/Engine/Source/Runtime/Engine/Classes/GameFramework/SpringArmComponent.h
+++ b/Engine/Source/Runtime/Engine/Classes/GameFramework/SpringArmComponent.h
@@ -73,6 +73,20 @@ class ENGINE_API USpringArmComponent : public USceneComponent
 	UPROPERTY(EditAnywhere, BlueprintReadWrite, Category=Lag)
 	uint32 bEnableCameraRotationLag : 1;
 
+	/**
+	 * If bUseCameraLagSubstepping is true, sub-step camera damping so that it handles fluctuating frame rates well (though this comes at a cost).
+	 * @see CameraLagMaxTimeStep
+	 */
+	UPROPERTY(EditAnywhere, BlueprintReadWrite, Category = Lag, AdvancedDisplay)
+	uint32 bUseCameraLagSubstepping : 1;
+
+	/**
+	 * If true and camera location lag is enabled, draws markers at the camera target (in green) and the lagged position (in yellow).
+	 * A line is drawn between the two locations, in green normally but in red if the distance to the lag target has been clamped (by CameraLagMaxDistance).
+	 */
+	UPROPERTY(EditAnywhere, BlueprintReadWrite, Category=Lag)
+	uint32 bDrawDebugLagMarkers : 1;
+
 	/** If bEnableCameraLag is true, controls how quickly camera reaches target position. Low values are slower (more lag), high values are faster (less lag), while zero is instant (no lag). */
 	UPROPERTY(EditAnywhere, BlueprintReadWrite, Category=Lag, meta=(editcondition="bEnableCameraLag", ClampMin="0.0", ClampMax="1000.0", UIMin = "0.0", UIMax = "1000.0"))
 	float CameraLagSpeed;
@@ -80,18 +94,15 @@ class ENGINE_API USpringArmComponent : public USceneComponent
 	/** If bEnableCameraRotationLag is true, controls how quickly camera reaches target position. Low values are slower (more lag), high values are faster (less lag), while zero is instant (no lag). */
 	UPROPERTY(EditAnywhere, BlueprintReadWrite, Category=Lag, meta=(editcondition = "bEnableCameraRotationLag", ClampMin="0.0", ClampMax="1000.0", UIMin = "0.0", UIMax = "1000.0"))
 	float CameraRotationLagSpeed;
-
-	/**
-	 * If bUseCameraLagSubstepping is true, sub-step camera damping so that it handles fluctuating frame rates well (though this comes at a cost).
-	 * @see CameraLagMaxTimeStep
-	 */
-	UPROPERTY(EditAnywhere, BlueprintReadWrite, Category = Lag, AdvancedDisplay)
-	bool bUseCameraLagSubstepping;
 	
 	/** Max time step used when sub-stepping camera lag. */
 	UPROPERTY(EditAnywhere, BlueprintReadWrite, Category = Lag, AdvancedDisplay, meta=(editcondition = "bUseCameraLagSubstepping", ClampMin="0.005", ClampMax="0.5", UIMin = "0.005", UIMax = "0.5"))
 	float CameraLagMaxTimeStep;
 
+	/** Max distance the camera target may lag behind the current location. If set to zero, no max distance is enforced. */
+	UPROPERTY(EditAnywhere, BlueprintReadWrite, Category=Lag, meta=(editcondition="bEnableCameraLag", ClampMin="0.0", UIMin = "0.0"))
+	float CameraLagMaxDistance;
+
 	/** Temporary variables when using camera lag, to record previous camera position */
 	FVector PreviousDesiredLoc;
 	FVector PreviousArmOrigin;
diff --git a/Engine/Source/Runtime/Engine/Private/GameFramework/SpringArmComponent.cpp b/Engine/Source/Runtime/Engine/Private/GameFramework/SpringArmComponent.cpp
--- a/Engine/Source/Runtime/Engine/Private/GameFramework/SpringArmComponent.cpp
+++ b/Engine/Source/Runtime/Engine/Private/GameFramework/SpringArmComponent.cpp
@@ -28,13 +28,13 @@ USpringArmComponent::USpringArmComponent(const FObjectInitializer& ObjectInitial
 	ProbeSize = 12.0f;
 	ProbeChannel = ECC_Camera;
 
-	CameraLagSpeed = 10.f;
-	CameraRotationLagSpeed = 10.f;
-
 	RelativeSocketRotation = FQuat::Identity;
 
 	bUseCameraLagSubstepping = true;
+	CameraLagSpeed = 10.f;
+	CameraRotationLagSpeed = 10.f;
 	CameraLagMaxTimeStep = 1.f / 60.f;
+	CameraLagMaxDistance = 0.f;
 	
 	// Init deprecated var, for old code that may refer to it.
 	SetDeprecatedControllerViewRotation(*this, bUsePawnControlRotation);
@@ -115,7 +115,32 @@ void USpringArmComponent::UpdateDesiredArmLocation(bool bDoTrace, bool bDoLocati
 		{
 			DesiredLoc = FMath::VInterpTo(PreviousDesiredLoc, DesiredLoc, DeltaTime, CameraLagSpeed);
 		}
+
+		// Clamp distance if requested
+		bool bClampedDist = false;
+		if (CameraLagMaxDistance > 0.f)
+		{
+			const FVector FromOrigin = DesiredLoc - ArmOrigin;
+			if (FromOrigin.SizeSquared() > FMath::Square(CameraLagMaxDistance))
+			{
+				DesiredLoc = ArmOrigin + FromOrigin.ClampMaxSize(CameraLagMaxDistance);
+				bClampedDist = true;
+			}
+		}		
+
+#if !(UE_BUILD_SHIPPING || UE_BUILD_TEST)
+		if (bDrawDebugLagMarkers)
+		{
+			DrawDebugSphere(GetWorld(), ArmOrigin, 5.f, 8, FColor::Green);
+			DrawDebugSphere(GetWorld(), DesiredLoc, 5.f, 8, FColor::Yellow);
+
+			const FVector ToOrigin = ArmOrigin - DesiredLoc;
+			DrawDebugDirectionalArrow(GetWorld(), DesiredLoc, DesiredLoc + ToOrigin * 0.5f, 7.5f, bClampedDist ? FColor::Red : FColor::Green);
+			DrawDebugDirectionalArrow(GetWorld(), DesiredLoc + ToOrigin * 0.5f, ArmOrigin,  7.5f, bClampedDist ? FColor::Red : FColor::Green);
+		}
+#endif
 	}
+
 	PreviousArmOrigin = ArmOrigin;
 	PreviousDesiredLoc = DesiredLoc;
 
