[UE-4487] fix invalid scene hierarchy when Actor subclass has native scene components without an explicitly-set native root component

change summary:
- modified USimpleConstructionScript::ExecuteScriptOnActor() to use the first available native scene component as root if the serialized RootComponent property is not assigned. this matches what the SCS editor does, and will handle existing BPGCs that have an SCS with the DefaultSceneRoot incorrectly present in the hierarchy.
- modified USimpleConstructionScript::ValidateSceneRootNodes() to exclude the DefaultSceneRoot when native scene components are detected without a native root.

diff --git a/Engine/Source/Runtime/Engine/Private/SimpleConstructionScript.cpp b/Engine/Source/Runtime/Engine/Private/SimpleConstructionScript.cpp
--- a/Engine/Source/Runtime/Engine/Private/SimpleConstructionScript.cpp
+++ b/Engine/Source/Runtime/Engine/Private/SimpleConstructionScript.cpp
@@ -317,22 +317,30 @@ void USimpleConstructionScript::ExecuteScriptOnActor(AActor* Actor, const FTrans
 		for(auto NodeIt = RootNodes.CreateIterator(); NodeIt; ++NodeIt)
 		{
 			USCS_Node* RootNode = *NodeIt;
-			if(RootNode != NULL)
+			if(RootNode != nullptr)
 			{
+				// Get all native scene components
+				TArray<USceneComponent*> Components;
+				Actor->GetComponents(Components);
+
+				// Get the native root component; if it's not set, the first native scene component will be used as root. This matches what's done in the SCS editor.
+				USceneComponent* RootComponent = Actor->GetRootComponent();
+				if(RootComponent == nullptr && Components.Num() > 0)
+				{
+					RootComponent = Components[0];
+				}
+
 				// If the root node specifies that it has a parent
-				USceneComponent* ParentComponent = NULL;
+				USceneComponent* ParentComponent = nullptr;
 				if(RootNode->ParentComponentOrVariableName != NAME_None)
 				{
 					// Get the Actor class object
 					UClass* ActorClass = Actor->GetClass();
-					check(ActorClass != NULL);
+					check(ActorClass != nullptr);
 
 					// If the root node is parented to a "native" component (i.e. in the 'Components' array)
 					if(RootNode->bIsParentComponentNative)
 					{
-						TArray<USceneComponent*> Components;
-						Actor->GetComponents(Components);
-
 						for(int32 CompIndex = 0; CompIndex < Components.Num(); ++CompIndex)
 						{
 							// If we found a match, remember the index
@@ -347,7 +355,7 @@ void USimpleConstructionScript::ExecuteScriptOnActor(AActor* Actor, const FTrans
 					{
 						// In the non-native case, the SCS node's variable name property is used as the parent identifier
 						UObjectPropertyBase* Property = FindField<UObjectPropertyBase>(ActorClass, RootNode->ParentComponentOrVariableName);
-						if(Property != NULL)
+						if(Property != nullptr)
 						{
 							// If we found a matching property, grab its value and use that as the parent for this node
 							ParentComponent = Cast<USceneComponent>(Property->GetObjectPropertyValue_InContainer(Actor));
@@ -356,7 +364,7 @@ void USimpleConstructionScript::ExecuteScriptOnActor(AActor* Actor, const FTrans
 				}
 
 				// Create the new component instance and any child components it may have
-				UActorComponent* InstancedComponent = RootNode->ExecuteNodeOnActor(Actor, ParentComponent != NULL ? ParentComponent : Actor->GetRootComponent(), &RootTransform, bIsDefaultTransform);
+				UActorComponent* InstancedComponent = RootNode->ExecuteNodeOnActor(Actor, ParentComponent != nullptr ? ParentComponent : RootComponent, &RootTransform, bIsDefaultTransform);
 				if(InstancedComponent != nullptr)
 				{
 					InstancedComponents.Add(InstancedComponent);
@@ -579,10 +587,10 @@ void USimpleConstructionScript::ValidateSceneRootNodes()
 #if WITH_EDITOR
 	UBlueprint* Blueprint = GetBlueprint();
 
-	if(DefaultSceneRootNode == NULL)
+	if(DefaultSceneRootNode == nullptr)
 	{
 		// If applicable, create a default scene component node
-		if(Blueprint != NULL
+		if(Blueprint != nullptr
 			&& FBlueprintEditorUtils::IsActorBased(Blueprint)
 			&& Blueprint->BlueprintType != BPTYPE_MacroLibrary)
 		{
@@ -590,34 +598,47 @@ void USimpleConstructionScript::ValidateSceneRootNodes()
 		}
 	}
 
-	if(DefaultSceneRootNode != NULL)
+	if(DefaultSceneRootNode != nullptr)
 	{
 		UClass* GeneratedClass = GetOwnerClass();
 
 		// Get the Blueprint class default object
-		AActor* CDO = NULL;
-		if(GeneratedClass != NULL)
+		AActor* CDO = nullptr;
+		if(GeneratedClass != nullptr)
 		{
 			CDO = Cast<AActor>(GeneratedClass->GetDefaultObject(false));
 		}
 
 		// If the generated class does not yet have a CDO, defer to the parent class
-		if(CDO == NULL && Blueprint->ParentClass != NULL)
+		if(CDO == nullptr && Blueprint->ParentClass != nullptr)
 		{
 			CDO = Cast<AActor>(Blueprint->ParentClass->GetDefaultObject(false));
 		}
 
-		// Check the native root component property; don't add the default scene root if it's set
-		bool bHasSceneComponentRootNodes = (CDO != NULL && CDO->GetRootComponent() != NULL);
+		// Check to see if we already have a native root component
+		bool bHasSceneComponentRootNodes = false;
+		if(CDO != nullptr)
+		{
+			// If the root component property is not set, the first available scene component will be used as the root. This matches what's done in the SCS editor.
+			bHasSceneComponentRootNodes = CDO->GetRootComponent() != nullptr;
+			if(!bHasSceneComponentRootNodes)
+			{
+				TArray<USceneComponent*> SceneComponents;
+				CDO->GetComponents(SceneComponents);
+				bHasSceneComponentRootNodes = SceneComponents.Num() > 0;
+			}
+		}
+
+		// Don't add the default scene root if we already have a native scene root component
 		if(!bHasSceneComponentRootNodes)
 		{
 			// Get the Blueprint hierarchy
 			TArray<UBlueprint*> BPStack;
-			if(Blueprint->GeneratedClass != NULL)
+			if(Blueprint->GeneratedClass != nullptr)
 			{
 				UBlueprint::GetBlueprintHierarchyFromClass(Blueprint->GeneratedClass, BPStack);
 			}
-			else if(Blueprint->ParentClass != NULL)
+			else if(Blueprint->ParentClass != nullptr)
 			{
 				UBlueprint::GetBlueprintHierarchyFromClass(Blueprint->ParentClass, BPStack);
 			}
@@ -642,9 +663,9 @@ void USimpleConstructionScript::ValidateSceneRootNodes()
 				for(int32 RootNodeIndex = 0; RootNodeIndex < SCSRootNodes.Num() && !bHasSceneComponentRootNodes; ++RootNodeIndex)
 				{
 					USCS_Node* RootNode = SCSRootNodes[RootNodeIndex];
-					bHasSceneComponentRootNodes = RootNode != NULL
+					bHasSceneComponentRootNodes = RootNode != nullptr
 						&& RootNode != DefaultSceneRootNode
-						&& RootNode->ComponentTemplate != NULL
+						&& RootNode->ComponentTemplate != nullptr
 						&& RootNode->ComponentTemplate->IsA<USceneComponent>();
 				}
 			}
