Removed DrawDynamicElements, PreRenderView and dependencies.  These have been disabled since 11/04, and the GetDynamicMeshElements path is now the only path for dynamic mesh rendering.

diff --git a/Engine/Source/Runtime/Engine/Private/NavMeshRenderingComponent.cpp b/Engine/Source/Runtime/Engine/Private/NavMeshRenderingComponent.cpp
--- a/Engine/Source/Runtime/Engine/Private/NavMeshRenderingComponent.cpp
+++ b/Engine/Source/Runtime/Engine/Private/NavMeshRenderingComponent.cpp
@@ -100,9 +100,6 @@ void UNavMeshRenderingComponent::GatherData(struct FNavMeshSceneProxyData* Curre
 	CurrentData->bEnableDrawing = NavMesh->bEnableDrawing;
 	CurrentData->bNeedsNewData = false;
 
-	static const auto CVarUseGetMeshElements = IConsoleManager::Get().FindTConsoleVariableDataInt(TEXT("r.UseGetMeshElements"));
-	const bool bUseGetMeshElements = !GRHICommandList.Bypass() || CVarUseGetMeshElements->GetValueOnGameThread() != 0;
-
 	if (CurrentData && NavMesh && NavMesh->bEnableDrawing)
 	{
 		FHitProxyId HitProxyId = FHitProxyId();
@@ -158,18 +155,9 @@ void UNavMeshRenderingComponent::GatherData(struct FNavMeshSceneProxyData* Curre
 				const TArray<int32>& MeshIndices = CurrentData->NavMeshGeometry.AreaIndices[AreaIdx];
 				for (int32 Idx=0; Idx<MeshIndices.Num(); Idx += 3)
 				{
-					if (bUseGetMeshElements)
-					{
-						CurrentData->ThickLineItems.Add(FNavMeshSceneProxyData::FDebugThickLine(MeshVerts[MeshIndices[Idx + 0]] + CurrentData->NavMeshDrawOffset, MeshVerts[MeshIndices[Idx + 1]] + CurrentData->NavMeshDrawOffset, NavMeshRenderColor_Recast_TriangleEdges, DefaultEdges_LineThickness));
-						CurrentData->ThickLineItems.Add(FNavMeshSceneProxyData::FDebugThickLine(MeshVerts[MeshIndices[Idx + 1]] + CurrentData->NavMeshDrawOffset, MeshVerts[MeshIndices[Idx + 2]] + CurrentData->NavMeshDrawOffset, NavMeshRenderColor_Recast_TriangleEdges, DefaultEdges_LineThickness));
-						CurrentData->ThickLineItems.Add(FNavMeshSceneProxyData::FDebugThickLine(MeshVerts[MeshIndices[Idx + 2]] + CurrentData->NavMeshDrawOffset, MeshVerts[MeshIndices[Idx + 0]] + CurrentData->NavMeshDrawOffset, NavMeshRenderColor_Recast_TriangleEdges, DefaultEdges_LineThickness));
-					}
-					else
-					{
-						CurrentData->BatchedElements.AddLine(MeshVerts[MeshIndices[Idx + 0]] + CurrentData->NavMeshDrawOffset, MeshVerts[MeshIndices[Idx + 1]] + CurrentData->NavMeshDrawOffset, NavMeshRenderColor_Recast_TriangleEdges, HitProxyId, DefaultEdges_LineThickness, 0, true);
-						CurrentData->BatchedElements.AddLine(MeshVerts[MeshIndices[Idx + 1]] + CurrentData->NavMeshDrawOffset, MeshVerts[MeshIndices[Idx + 2]] + CurrentData->NavMeshDrawOffset, NavMeshRenderColor_Recast_TriangleEdges, HitProxyId, DefaultEdges_LineThickness, 0, true);
-						CurrentData->BatchedElements.AddLine(MeshVerts[MeshIndices[Idx + 2]] + CurrentData->NavMeshDrawOffset, MeshVerts[MeshIndices[Idx + 0]] + CurrentData->NavMeshDrawOffset, NavMeshRenderColor_Recast_TriangleEdges, HitProxyId, DefaultEdges_LineThickness, 0, true);
-					}
+					CurrentData->ThickLineItems.Add(FNavMeshSceneProxyData::FDebugThickLine(MeshVerts[MeshIndices[Idx + 0]] + CurrentData->NavMeshDrawOffset, MeshVerts[MeshIndices[Idx + 1]] + CurrentData->NavMeshDrawOffset, NavMeshRenderColor_Recast_TriangleEdges, DefaultEdges_LineThickness));
+					CurrentData->ThickLineItems.Add(FNavMeshSceneProxyData::FDebugThickLine(MeshVerts[MeshIndices[Idx + 1]] + CurrentData->NavMeshDrawOffset, MeshVerts[MeshIndices[Idx + 2]] + CurrentData->NavMeshDrawOffset, NavMeshRenderColor_Recast_TriangleEdges, DefaultEdges_LineThickness));
+					CurrentData->ThickLineItems.Add(FNavMeshSceneProxyData::FDebugThickLine(MeshVerts[MeshIndices[Idx + 2]] + CurrentData->NavMeshDrawOffset, MeshVerts[MeshIndices[Idx + 0]] + CurrentData->NavMeshDrawOffset, NavMeshRenderColor_Recast_TriangleEdges, DefaultEdges_LineThickness));
 				}
 			}
 		}
@@ -340,20 +328,10 @@ void UNavMeshRenderingComponent::GatherData(struct FNavMeshSceneProxyData* Curre
 						FVector UR(TileBox.Max.X, TileBox.Max.Y, DrawZ);
 						FVector UL(LL.X, UR.Y, DrawZ);
 						FVector LR(UR.X, LL.Y, DrawZ);
-						if (bUseGetMeshElements)
-						{
-							CurrentData->ThickLineItems.Add(FNavMeshSceneProxyData::FDebugThickLine(LL, UL, NavMeshRenderColor_TileBounds, DefaultEdges_LineThickness));
-							CurrentData->ThickLineItems.Add(FNavMeshSceneProxyData::FDebugThickLine(UL, UR, NavMeshRenderColor_TileBounds, DefaultEdges_LineThickness));
-							CurrentData->ThickLineItems.Add(FNavMeshSceneProxyData::FDebugThickLine(UR, LR, NavMeshRenderColor_TileBounds, DefaultEdges_LineThickness));
-							CurrentData->ThickLineItems.Add(FNavMeshSceneProxyData::FDebugThickLine(LR, LL, NavMeshRenderColor_TileBounds, DefaultEdges_LineThickness));
-						}
-						else
-						{
-							CurrentData->BatchedElements.AddLine(LL, UL, NavMeshRenderColor_TileBounds, HitProxyId, DefaultEdges_LineThickness, 0, true);
-							CurrentData->BatchedElements.AddLine(UL, UR, NavMeshRenderColor_TileBounds, HitProxyId, DefaultEdges_LineThickness, 0, true);
-							CurrentData->BatchedElements.AddLine(UR, LR, NavMeshRenderColor_TileBounds, HitProxyId, DefaultEdges_LineThickness, 0, true);
-							CurrentData->BatchedElements.AddLine(LR, LL, NavMeshRenderColor_TileBounds, HitProxyId, DefaultEdges_LineThickness, 0, true);
-						}
+						CurrentData->ThickLineItems.Add(FNavMeshSceneProxyData::FDebugThickLine(LL, UL, NavMeshRenderColor_TileBounds, DefaultEdges_LineThickness));
+						CurrentData->ThickLineItems.Add(FNavMeshSceneProxyData::FDebugThickLine(UL, UR, NavMeshRenderColor_TileBounds, DefaultEdges_LineThickness));
+						CurrentData->ThickLineItems.Add(FNavMeshSceneProxyData::FDebugThickLine(UR, LR, NavMeshRenderColor_TileBounds, DefaultEdges_LineThickness));
+						CurrentData->ThickLineItems.Add(FNavMeshSceneProxyData::FDebugThickLine(LR, LL, NavMeshRenderColor_TileBounds, DefaultEdges_LineThickness));
 					}
 				}
 			}
diff --git a/Engine/Source/Runtime/Engine/Private/NavMeshRenderingHelpers.h b/Engine/Source/Runtime/Engine/Private/NavMeshRenderingHelpers.h
--- a/Engine/Source/Runtime/Engine/Private/NavMeshRenderingHelpers.h
+++ b/Engine/Source/Runtime/Engine/Private/NavMeshRenderingHelpers.h
@@ -57,7 +57,6 @@ struct FNavMeshSceneProxyData : public TSharedFromThis<FNavMeshSceneProxyData, E
 		NavMeshGeometry.OffMeshSegments.Reset();
 		for (int32 Index=0; Index < RECAST_MAX_AREAS; Index++)
 			NavMeshGeometry.OffMeshSegmentAreas[Index].Reset();
-		BatchedElements.Clear();
 		TileEdgeLines.Reset();
 		NavMeshEdgeLines.Reset();
 		NavLinkLines.Reset();
@@ -78,7 +77,6 @@ struct FNavMeshSceneProxyData : public TSharedFromThis<FNavMeshSceneProxyData, E
 	}
 
 	FRecastDebugGeometry NavMeshGeometry;
-	FBatchedElements BatchedElements;
 
 	TArray<FDebugRenderSceneProxy::FDebugLine> TileEdgeLines;
 	TArray<FDebugRenderSceneProxy::FDebugLine> NavMeshEdgeLines;
@@ -381,142 +379,6 @@ public:
 		}
 	}
 
-	virtual void DrawDynamicElements(FPrimitiveDrawInterface* PDI,const FSceneView* View) override
-	{
-		QUICK_SCOPE_CYCLE_COUNTER( STAT_RecastRenderingSceneProxy_DrawDynamicElements );
-
-		const bool bVisible = !!View->Family->EngineShowFlags.Navigation || bForceRendering;
-		if (!ProxyData.bEnableDrawing || !bVisible) //check if we have any data to render
-		{
-			return;
-		}
-
-		ProxyData.bSkipDistanceCheck = GIsEditor && (GEngine->GetDebugLocalPlayer() == NULL);
-
-		FVector const PosX(1.f,0,0);
-		FVector const PosY(0,1.f,0);
-		FVector const PosZ(0,0,1.f);
-
-		const TArray<FVector>& MeshVerts = ProxyData.NavMeshGeometry.MeshVerts;
-
-		// Draw Mesh
-		for (int32 Index = 0; Index < ProxyData.MeshBuilders.Num(); ++Index)
-		{
-			const FColoredMaterialRenderProxy *MeshColorInstance = new(FMemStack::Get()) FColoredMaterialRenderProxy(GEngine->DebugMeshMaterial->GetRenderProxy(false), ProxyData.MeshBuilders[Index].ClusterColor);						
-			FDynamicMeshBuilder	MeshBuilder;
-			MeshBuilder.AddVertices( ProxyData.MeshBuilders[Index].Vertices );
-			MeshBuilder.AddTriangles( ProxyData.MeshBuilders[Index].Indices );
-			MeshBuilder.Draw(PDI, FMatrix::Identity, MeshColorInstance, GetDepthPriorityGroup(View));
-		}
-		
-		for (int32 Index = 0; Index < ProxyData.OctreeBounds.Num(); ++Index)
-		{
-			FBoxCenterAndExtent& Bounds = ProxyData.OctreeBounds[Index];
-			DrawDebugBox(PDI, Bounds.Center, Bounds.Extent, FColor::White);
-		}
-		
-		FHitProxyId HitProxyId;
-		FBatchedElements BatchedElements;
-		int32 Num = ProxyData.NavMeshEdgeLines.Num();
-		for (int32 Index = 0; Index < Num; ++Index)
-		{
-			const FDebugLine &Line = ProxyData.NavMeshEdgeLines[Index];
-			if( LineInView(Line.Start,Line.End,View,false) )
-			{
-				if (LineInCorrectDistance(Line.Start,Line.End,View))
-				{
-					BatchedElements.AddLine(Line.Start, Line.End, Line.Color,HitProxyId,NavMeshEdges_LineThickness, 0, true);
-				}
-				else if (GIsEditor)
-				{
-					BatchedElements.AddLine(Line.Start, Line.End, Line.Color,HitProxyId,DefaultEdges_LineThickness, 0, true);
-				}
-			}
-		}
-
-		Num = ProxyData.ClusterLinkLines.Num();
-		for (int32 Index = 0; Index < Num; ++Index)
-		{
-			const FDebugLine &Line = ProxyData.ClusterLinkLines[Index];
-			if( LineInView(Line.Start,Line.End,View,false) )
-			{
-				if (LineInCorrectDistance(Line.Start,Line.End,View))
-				{
-					BatchedElements.AddLine(Line.Start, Line.End, Line.Color,HitProxyId,ClusterLinkLines_LineThickness, 0, true);
-				}
-				else if (GIsEditor)
-				{
-					BatchedElements.AddLine(Line.Start, Line.End, Line.Color, HitProxyId, DefaultEdges_LineThickness, 0, true);
-				}
-			}
-		}
-
-		Num = ProxyData.TileEdgeLines.Num();
-		for (int32 Index = 0; Index < Num; ++Index)
-		{
-			const FDebugLine &Line = ProxyData.TileEdgeLines[Index];
-			if( LineInView(Line.Start,Line.End,View,false) )
-			{
-				if (LineInCorrectDistance(Line.Start,Line.End,View))
-				{
-					BatchedElements.AddLine(Line.Start, Line.End, Line.Color,HitProxyId,PolyEdges_LineThickness, 0, true);
-				}
-				else if (GIsEditor)
-				{
-					BatchedElements.AddLine(Line.Start, Line.End, Line.Color, HitProxyId, DefaultEdges_LineThickness, 0, true);
-				}
-			}
-		}
-
-		Num = ProxyData.NavLinkLines.Num();
-		for (int32 Index = 0; Index < Num; ++Index)
-		{
-			const FDebugLine &Line = ProxyData.NavLinkLines[Index];
-			if( LineInView(Line.Start,Line.End,View,false) )
-			{
-				if (LineInCorrectDistance(Line.Start,Line.End,View))
-				{
-					BatchedElements.AddLine(Line.Start, Line.End, Line.Color,HitProxyId,LinkLines_LineThickness, 0, true);
-				}
-				else if (GIsEditor)
-				{
-					BatchedElements.AddLine(Line.Start, Line.End, Line.Color, HitProxyId, DefaultEdges_LineThickness, 0, true);
-				}
-			}
-		}
-
-		const bool bNeedToSwitchVerticalAxis = RHINeedsToSwitchVerticalAxis(View->GetShaderPlatform());
-		const FTexture2DRHIRef DepthTexture;
-		const auto FeatureLevel = View->GetFeatureLevel();
-			
-		FRHICommandListImmediate& RHICmdList = FRHICommandListExecutor::GetImmediateCommandList();
-		BatchedElements.Draw(
-			RHICmdList,
-			FeatureLevel,
-			bNeedToSwitchVerticalAxis,
-			View->ViewProjectionMatrix,
-			View->ViewRect.Width(),
-			View->ViewRect.Height(),
-			View->Family->EngineShowFlags.HitProxies,
-			1.0f,
-			View,
-			DepthTexture
-			);
-
-		ProxyData.BatchedElements.Draw(
-			RHICmdList,
-			FeatureLevel,
-			bNeedToSwitchVerticalAxis,
-			View->ViewProjectionMatrix,
-			View->ViewRect.Width(),
-			View->ViewRect.Height(),
-			View->Family->EngineShowFlags.HitProxies,
-			1.0f,
-			View,
-			DepthTexture
-			);
-	}
-
 	virtual void GetDynamicMeshElements(const TArray<const FSceneView*>& Views, const FSceneViewFamily& ViewFamily, uint32 VisibilityMap, FMeshElementCollector& Collector) const override
 	{
 		QUICK_SCOPE_CYCLE_COUNTER(STAT_RecastRenderingSceneProxy_GetDynamicMeshElements);
@@ -716,7 +578,7 @@ public:
 			+ ProxyData.TileEdgeLines.GetAllocatedSize() + ProxyData.NavMeshEdgeLines.GetAllocatedSize() + ProxyData.NavLinkLines.GetAllocatedSize()
 			+ ProxyData.PathCollidingGeomIndices.GetAllocatedSize() + ProxyData.PathCollidingGeomVerts.GetAllocatedSize()
 			+ ProxyData.DebugLabels.GetAllocatedSize() + ProxyData.ClusterLinkLines.GetAllocatedSize() + ProxyData.MeshBuilders.GetAllocatedSize()
-			+ ProxyData.BatchedElements.GetAllocatedSize() + IndexBuffer.Indices.GetAllocatedSize() + VertexBuffer.Vertices.GetAllocatedSize() + MeshColors.GetAllocatedSize() + MeshBatchElements.GetAllocatedSize());
+			+ IndexBuffer.Indices.GetAllocatedSize() + VertexBuffer.Vertices.GetAllocatedSize() + MeshColors.GetAllocatedSize() + MeshBatchElements.GetAllocatedSize());
 	}
 
 private:
diff --git a/Engine/Source/Runtime/Engine/Private/NiagaraComponent.cpp b/Engine/Source/Runtime/Engine/Private/NiagaraComponent.cpp
--- a/Engine/Source/Runtime/Engine/Private/NiagaraComponent.cpp
+++ b/Engine/Source/Runtime/Engine/Private/NiagaraComponent.cpp
@@ -90,30 +90,6 @@ void FNiagaraSceneProxy::OnTransformChanged()
 	//WorldSpacePrimitiveUniformBuffer.ReleaseResource();
 }
 
-void FNiagaraSceneProxy::PreRenderView(const FSceneViewFamily* ViewFamily, const uint32 VisibilityMap, int32 FrameNumber)
-{
-	for (NiagaraEffectRenderer *Renderer : EffectRenderers)
-	{
-		if (Renderer)
-		{
-			Renderer->PreRenderView(ViewFamily, VisibilityMap, FrameNumber, this);
-		}
-	}
-	return;
-}
-		  
-void FNiagaraSceneProxy::DrawDynamicElements(FPrimitiveDrawInterface* PDI, const FSceneView* View) 
-{
-	for (NiagaraEffectRenderer *Renderer : EffectRenderers)
-	{
-		if (Renderer)
-		{
-			Renderer->DrawDynamicElements(PDI, View, this);
-		}
-	}
-	return;
-}
-
 FPrimitiveViewRelevance FNiagaraSceneProxy::GetViewRelevance(const FSceneView* View)
 {
 	FPrimitiveViewRelevance Relevance;
diff --git a/Engine/Source/Runtime/Engine/Private/NiagaraEffectRenderer.cpp b/Engine/Source/Runtime/Engine/Private/NiagaraEffectRenderer.cpp
--- a/Engine/Source/Runtime/Engine/Private/NiagaraEffectRenderer.cpp
+++ b/Engine/Source/Runtime/Engine/Private/NiagaraEffectRenderer.cpp
@@ -67,7 +67,6 @@ NiagaraEffectRendererSprites::NiagaraEffectRendererSprites(ERHIFeatureLevel::Typ
 void NiagaraEffectRendererSprites::ReleaseRenderThreadResources()
 {
 	VertexFactory->ReleaseResource();
-	PerViewUniformBuffers.Empty();
 	WorldSpacePrimitiveUniformBuffer.ReleaseResource();
 }
 
@@ -76,140 +75,6 @@ void NiagaraEffectRendererSprites::CreateRenderThreadResources()
 	VertexFactory->InitResource();
 }
 
-
-
-void NiagaraEffectRendererSprites::PreRenderView(const FSceneViewFamily* ViewFamily, const uint32 VisibilityMap, int32 FrameNumber, const FNiagaraSceneProxy *SceneProxy)
-{
-	SCOPE_CYCLE_COUNTER(STAT_NiagaraPreRenderView);
-
-	check(DynamicDataRender);
-	if (DynamicDataRender->VertexData.Num() == 0)
-		return;
-
-	int32 SizeInBytes = DynamicDataRender->VertexData.GetTypeSize() * DynamicDataRender->VertexData.Num();
-	DynamicVertexAllocation = FGlobalDynamicVertexBuffer::Get().Allocate(SizeInBytes);
-	if (DynamicVertexAllocation.IsValid())
-	{
-		// Copy the vertex data over.
-		FMemory::Memcpy(DynamicVertexAllocation.Buffer, DynamicDataRender->VertexData.GetData(), SizeInBytes);
-		VertexFactory->SetInstanceBuffer(
-			DynamicVertexAllocation.VertexBuffer,
-			DynamicVertexAllocation.VertexOffset,
-			sizeof(FParticleSpriteVertex),
-			true
-			);
-		VertexFactory->SetDynamicParameterBuffer(NULL, 0, 0, true);
-		// Compute the per-view uniform buffers.
-		const int32 NumViews = ViewFamily->Views.Num();
-		uint32 ViewBit = 1;
-		for (int32 ViewIndex = 0; ViewIndex < NumViews; ++ViewIndex, ViewBit <<= 1)
-		{
-			FParticleSpriteUniformBufferRef* SpriteViewUniformBufferPtr = new(PerViewUniformBuffers)FParticleSpriteUniformBufferRef();
-			if (VisibilityMap & ViewBit)
-			{
-				FParticleSpriteUniformParameters PerViewUniformParameters;// = UniformParameters;
-				PerViewUniformParameters.AxisLockRight = FVector4(0.0f, 0.0f, 0.0f, 0.0f);
-				PerViewUniformParameters.AxisLockUp = FVector4(0.0f, 0.0f, 0.0f, 0.0f);
-				PerViewUniformParameters.RotationScale = 1.0f;
-				PerViewUniformParameters.RotationBias = 0.0f;
-				PerViewUniformParameters.TangentSelector = FVector4(0.0f, 0.0f, 0.0f, 1.0f);
-				PerViewUniformParameters.InvDeltaSeconds = 0.0f;
-				PerViewUniformParameters.SubImageSize = FVector4(1.0f, 1.0f, 1.0f, 1.0f);
-				PerViewUniformParameters.NormalsType = 0;
-				PerViewUniformParameters.NormalsSphereCenter = FVector4(0.0f, 0.0f, 0.0f, 1.0f);
-				PerViewUniformParameters.NormalsCylinderUnitDirection = FVector4(0.0f, 0.0f, 1.0f, 0.0f);
-				PerViewUniformParameters.PivotOffset = FVector2D(0.0f, 0.0f);
-				PerViewUniformParameters.MacroUVParameters = FVector4(0.0f, 0.0f, 1.0f, 1.0f);
-				*SpriteViewUniformBufferPtr = FParticleSpriteUniformBufferRef::CreateUniformBufferImmediate(PerViewUniformParameters, UniformBuffer_SingleFrame);
-			}
-		}
-
-		// Update the primitive uniform buffer if needed.
-		if (!WorldSpacePrimitiveUniformBuffer.IsInitialized())
-		{
-			FPrimitiveUniformShaderParameters PrimitiveUniformShaderParameters = GetPrimitiveUniformShaderParameters(
-				FMatrix::Identity,
-				SceneProxy->GetActorPosition(),
-				SceneProxy->GetBounds(),
-				SceneProxy->GetLocalBounds(),
-				SceneProxy->ReceivesDecals(),
-				false,
-				false,
-				false
-				);
-			WorldSpacePrimitiveUniformBuffer.SetContents(PrimitiveUniformShaderParameters);
-			WorldSpacePrimitiveUniformBuffer.InitResource();
-		}
-	}
-}
-
-
-void NiagaraEffectRendererSprites::DrawDynamicElements(FPrimitiveDrawInterface* PDI, const FSceneView* View, const FNiagaraSceneProxy *SceneProxy)
-{
-	SCOPE_CYCLE_COUNTER(STAT_NiagaraRender);
-
-	check(DynamicDataRender);
-	if (DynamicDataRender->VertexData.Num() == 0 || Material==nullptr)
-		return;
-
-	const bool bIsWireframe = View->Family->EngineShowFlags.Wireframe;
-	FMaterialRenderProxy* MaterialRenderProxy = Material->GetRenderProxy(SceneProxy->IsSelected(), SceneProxy->IsHovered());
-
-	if (DynamicVertexAllocation.IsValid()
-		&& (bIsWireframe || !PDI->IsMaterialIgnored(MaterialRenderProxy, View->GetFeatureLevel())))
-	{
-		FMeshBatch MeshBatch;
-		MeshBatch.VertexFactory = VertexFactory;
-		MeshBatch.CastShadow = SceneProxy->CastsDynamicShadow();
-		MeshBatch.bUseAsOccluder = false;
-		MeshBatch.ReverseCulling = SceneProxy->IsLocalToWorldDeterminantNegative();
-		MeshBatch.Type = PT_TriangleList;
-		MeshBatch.DepthPriorityGroup = SceneProxy->GetDepthPriorityGroup(View);
-		MeshBatch.LCI = NULL;
-		if (bIsWireframe)
-		{
-			MeshBatch.MaterialRenderProxy = UMaterial::GetDefaultMaterial(MD_Surface)->GetRenderProxy(SceneProxy->IsSelected(), SceneProxy->IsHovered());
-		}
-		else
-		{
-			MeshBatch.MaterialRenderProxy = MaterialRenderProxy;
-		}
-
-		FMeshBatchElement& MeshElement = MeshBatch.Elements[0];
-		MeshElement.IndexBuffer = &GParticleIndexBuffer;
-		MeshElement.FirstIndex = 0;
-		MeshElement.NumPrimitives = 2;
-		MeshElement.NumInstances = DynamicDataRender->VertexData.Num();
-		MeshElement.MinVertexIndex = 0;
-		MeshElement.MaxVertexIndex = MeshElement.NumInstances * 4 - 1;
-		MeshElement.PrimitiveUniformBufferResource = &WorldSpacePrimitiveUniformBuffer;
-
-		int32 ViewIndex = View->Family->Views.Find(View);
-		check(PerViewUniformBuffers.IsValidIndex(ViewIndex) && PerViewUniformBuffers[ViewIndex]);
-		VertexFactory->SetSpriteUniformBuffer(PerViewUniformBuffers[ViewIndex]);
-
-		DrawRichMesh(
-			PDI,
-			MeshBatch,
-			FLinearColor(1.0f, 0.0f, 0.0f),	//WireframeColor,
-			FLinearColor(1.0f, 1.0f, 0.0f),	//LevelColor,
-			FLinearColor(1.0f, 1.0f, 1.0f),	//PropertyColor,		
-			SceneProxy,
-			GIsEditor && (View->Family->EngineShowFlags.Selection) ? SceneProxy->IsSelected() : false,
-			bIsWireframe
-			);
-	}
-
-	/*
-	for(int32 i=0; i<DynamicDataRender->VertexData.Num(); i++)
-	{
-	FParticleSpriteVertex& Vertex = DynamicDataRender->VertexData[i];
-	PDI->DrawPoint(Vertex.Position, Vertex.Color, 1.0f, SDPG_World);
-	}
-	*/
-}
-
-
 void NiagaraEffectRendererSprites::GetDynamicMeshElements(const TArray<const FSceneView*>& Views, const FSceneViewFamily& ViewFamily, uint32 VisibilityMap, FMeshElementCollector& Collector, const FNiagaraSceneProxy *SceneProxy) const
 {
 	SCOPE_CYCLE_COUNTER(STAT_NiagaraRender);
@@ -410,7 +275,6 @@ DynamicDataRender(NULL)
 void NiagaraEffectRendererRibbon::ReleaseRenderThreadResources()
 {
 	VertexFactory->ReleaseResource();
-	PerViewUniformBuffers.Empty();
 	WorldSpacePrimitiveUniformBuffer.ReleaseResource();
 }
 
@@ -520,123 +384,6 @@ void NiagaraEffectRendererRibbon::GetDynamicMeshElements(const TArray<const FSce
 	}
 }
 
-void NiagaraEffectRendererRibbon::PreRenderView(const FSceneViewFamily* ViewFamily, const uint32 VisibilityMap, int32 FrameNumber, const FNiagaraSceneProxy *SceneProxy)
-{
-	SCOPE_CYCLE_COUNTER(STAT_NiagaraPreRenderView);
-
-	check(DynamicDataRender && DynamicDataRender->VertexData.Num());
-
-	int32 SizeInBytes = DynamicDataRender->VertexData.GetTypeSize() * DynamicDataRender->VertexData.Num();
-	DynamicVertexAllocation = FGlobalDynamicVertexBuffer::Get().Allocate(SizeInBytes);
-	if (DynamicVertexAllocation.IsValid())
-	{
-		// Copy the vertex data over.
-		FMemory::Memcpy(DynamicVertexAllocation.Buffer, DynamicDataRender->VertexData.GetData(), SizeInBytes);
-		VertexFactory->SetVertexBuffer(DynamicVertexAllocation.VertexBuffer, DynamicVertexAllocation.VertexOffset, sizeof(FParticleBeamTrailVertex));
-		VertexFactory->SetDynamicParameterBuffer(NULL, 0, 0);
-		// Compute the per-view uniform buffers.
-		const int32 NumViews = ViewFamily->Views.Num();
-		uint32 ViewBit = 1;
-		for (int32 ViewIndex = 0; ViewIndex < NumViews; ++ViewIndex, ViewBit <<= 1)
-		{
-			FParticleBeamTrailUniformBufferRef* SpriteViewUniformBufferPtr = new(PerViewUniformBuffers)FParticleBeamTrailUniformBufferRef();
-			if (VisibilityMap & ViewBit)
-			{
-				FParticleBeamTrailUniformParameters PerViewUniformParameters;// = UniformParameters;
-				PerViewUniformParameters.CameraUp = FVector4(0.0f, 0.0f, 1.0f, 0.0f);
-				PerViewUniformParameters.CameraRight = FVector4(1.0f, 0.0f, 0.0f, 0.0f);
-				PerViewUniformParameters.ScreenAlignment = FVector4(0.0f, 0.0f, 0.0f, 0.0f);
-				*SpriteViewUniformBufferPtr = FParticleBeamTrailUniformBufferRef::CreateUniformBufferImmediate(PerViewUniformParameters, UniformBuffer_SingleFrame);
-			}
-		}
-
-		// Update the primitive uniform buffer if needed.
-		if (!WorldSpacePrimitiveUniformBuffer.IsInitialized())
-		{
-			FPrimitiveUniformShaderParameters PrimitiveUniformShaderParameters = GetPrimitiveUniformShaderParameters(
-				FMatrix::Identity,
-				SceneProxy->GetActorPosition(),
-				SceneProxy->GetBounds(),
-				SceneProxy->GetLocalBounds(),
-				SceneProxy->ReceivesDecals(),
-				false,
-				false,
-				false
-				);
-			WorldSpacePrimitiveUniformBuffer.SetContents(PrimitiveUniformShaderParameters);
-			WorldSpacePrimitiveUniformBuffer.InitResource();
-		}
-	}
-}
-
-
-void NiagaraEffectRendererRibbon::DrawDynamicElements(FPrimitiveDrawInterface* PDI, const FSceneView* View, const FNiagaraSceneProxy *SceneProxy)
-{
-	SCOPE_CYCLE_COUNTER(STAT_NiagaraRender);
-
-	check(DynamicDataRender && DynamicDataRender->VertexData.Num());
-
-	const bool bIsWireframe = View->Family->EngineShowFlags.Wireframe;
-	FMaterialRenderProxy* MaterialRenderProxy = Material->GetRenderProxy(SceneProxy->IsSelected(), SceneProxy->IsHovered());
-
-	if (DynamicVertexAllocation.IsValid()
-		&& (bIsWireframe || !PDI->IsMaterialIgnored(MaterialRenderProxy, View->GetFeatureLevel())))
-	{
-		FMeshBatch MeshBatch;
-		MeshBatch.VertexFactory = VertexFactory;
-		MeshBatch.CastShadow = SceneProxy->CastsDynamicShadow();
-		MeshBatch.bUseAsOccluder = false;
-		MeshBatch.bDisableBackfaceCulling = true;
-		MeshBatch.ReverseCulling = SceneProxy->IsLocalToWorldDeterminantNegative();
-		MeshBatch.Type = PT_TriangleStrip;
-		MeshBatch.DepthPriorityGroup = SceneProxy->GetDepthPriorityGroup(View);
-		MeshBatch.LCI = NULL;
-		if (bIsWireframe)
-		{
-			MeshBatch.MaterialRenderProxy = UMaterial::GetDefaultMaterial(MD_Surface)->GetRenderProxy(SceneProxy->IsSelected(), SceneProxy->IsHovered());
-		}
-		else
-		{
-			MeshBatch.MaterialRenderProxy = MaterialRenderProxy;
-		}
-
-		FMeshBatchElement& MeshElement = MeshBatch.Elements[0];
-		MeshElement.IndexBuffer = &GParticleIndexBuffer;
-		MeshElement.FirstIndex = 0;
-		MeshElement.NumPrimitives = DynamicDataRender->VertexData.Num()-2;
-		MeshElement.NumInstances = 1;
-		MeshElement.MinVertexIndex = 0;
-		MeshElement.MaxVertexIndex = DynamicDataRender->VertexData.Num() - 1;
-		MeshElement.PrimitiveUniformBufferResource = &WorldSpacePrimitiveUniformBuffer;
-
-		int32 ViewIndex = View->Family->Views.Find(View);
-		check(PerViewUniformBuffers.IsValidIndex(ViewIndex) && PerViewUniformBuffers[ViewIndex]);
-		VertexFactory->SetBeamTrailUniformBuffer(PerViewUniformBuffers[ViewIndex]);
-
-		DrawRichMesh(
-			PDI,
-			MeshBatch,
-			FLinearColor(1.0f, 0.0f, 0.0f),	//WireframeColor,
-			FLinearColor(1.0f, 1.0f, 0.0f),	//LevelColor,
-			FLinearColor(1.0f, 1.0f, 1.0f),	//PropertyColor,		
-			SceneProxy,
-			GIsEditor && (View->Family->EngineShowFlags.Selection) ? SceneProxy->IsSelected() : false,
-			bIsWireframe
-			);
-	}
-
-	/*
-	for(int32 i=0; i<DynamicDataRender->VertexData.Num(); i++)
-	{
-	FParticleSpriteVertex& Vertex = DynamicDataRender->VertexData[i];
-	PDI->DrawPoint(Vertex.Position, Vertex.Color, 1.0f, SDPG_World);
-	}
-	*/
-}
-
-
-
-
 void NiagaraEffectRendererRibbon::SetDynamicData_RenderThread(FNiagaraDynamicDataBase* NewDynamicData)
 {
 	check(IsInRenderingThread());
diff --git a/Engine/Source/Runtime/Engine/Private/Particles/ParticleGpuSimulation.cpp b/Engine/Source/Runtime/Engine/Private/Particles/ParticleGpuSimulation.cpp
--- a/Engine/Source/Runtime/Engine/Private/Particles/ParticleGpuSimulation.cpp
+++ b/Engine/Source/Runtime/Engine/Private/Particles/ParticleGpuSimulation.cpp
@@ -2255,8 +2255,6 @@ public:
 
 	/** The vertex buffer used to access particles in the simulation. */
 	FGPUParticleVertexBuffer VertexBuffer;
-	/** The vertex factories used to render these particles (one per view). */
-	TIndirectArray<FGPUSpriteVertexFactory, TInlineAllocator<2> > VertexFactories;
 	/** The vertex factory for visualizing the local vector field. */
 	FVectorFieldVisualizationVertexFactory* VectorFieldVisualizationVertexFactory;
 
@@ -2322,10 +2320,6 @@ public:
 			// Store simulation resources for this emitter.
 			Simulation->EmitterSimulationResources = InEmitterSimulationResources;
 
-			// Create the vertex factory with which this simulation can be rendered.
-			FGPUSpriteVertexFactory* VertexFactory = new(Simulation->VertexFactories) FGPUSpriteVertexFactory();
-			VertexFactory->InitResource();
-
 			// If a visualization vertex factory has been created, initialize it.
 			if (Simulation->VectorFieldVisualizationVertexFactory)
 			{
@@ -2398,11 +2392,6 @@ private:
 		check( IsInRenderingThread() );
 		VertexBuffer.ReleaseResource();
 		TileVertexBuffer.ReleaseResource();
-		for ( int32 VertexFactoryIndex = 0; VertexFactoryIndex < VertexFactories.Num(); ++VertexFactoryIndex )
-		{
-			VertexFactories[VertexFactoryIndex].ReleaseResource();
-		}
-		VertexFactories.Reset();
 		if ( VectorFieldVisualizationVertexFactory )
 		{
 			VectorFieldVisualizationVertexFactory->ReleaseResource();
@@ -2496,8 +2485,6 @@ struct FGPUSpriteDynamicEmitterData : FDynamicEmitterDataBase
 	FParticlePerFrameSimulationParameters PerFrameSimulationParameters;
 	/** Per-emitter parameters that may change*/
 	FGPUSpriteEmitterDynamicUniformParameters EmitterDynamicParameters;
-	/** Emitter uniform buffer for parameters that may change. Used for rendering. */
-	FGPUSpriteEmitterDynamicUniformBufferRef DynamicUniformBuffer;
 	/** How the particles should be sorted, if at all. */
 	EParticleSortMode SortMode;
 	/** Whether to render particles in local space or world space. */
@@ -2582,81 +2569,6 @@ struct FGPUSpriteDynamicEmitterData : FDynamicEmitterDataBase
 	{		
 	}
 
-	/**
-	 * Called once the emitter has been determined to be visible. The emitter
-	 * will generate per-view render data.
-	 */
-	virtual void PreRenderView(FParticleSystemSceneProxy* Proxy, const FSceneViewFamily* ViewFamily, const uint32 VisibilityMap, int32 FrameNumber) override
-	{
-		auto FeatureLevel = ViewFamily->Scene->GetFeatureLevel();
-
-		if (RHISupportsGPUParticles(FeatureLevel))
-		{
-			SCOPE_CYCLE_COUNTER(STAT_GPUSpritePreRenderTime);
-
-			check(Simulation);
-
-			// Do not render orphaned emitters. This can happen if the emitter
-			// instance has been destroyed but we are rendering before the
-			// scene proxy has received the update to clear dynamic data.
-			if (Simulation->SimulationIndex != INDEX_NONE
-				&& Simulation->VertexBuffer.ParticleCount > 0)
-			{
-				// Create view agnostic render data.  Do here rather than in CreateRenderThreadResources because in some cases Render can be called before CreateRenderThreadResources
-				{
-					// Create per-emitter uniform buffer for dynamic parameters
-					DynamicUniformBuffer = FGPUSpriteEmitterDynamicUniformBufferRef::CreateUniformBufferImmediate(EmitterDynamicParameters, UniformBuffer_SingleFrame);
-				}
-
-				const bool bTranslucent = RendersWithTranslucentMaterial();
-				const bool bAllowSorting = FXConsoleVariables::bAllowGPUSorting
-					&& FeatureLevel == ERHIFeatureLevel::SM5
-					&& bTranslucent;
-
-				// Create vertex factories for any new views.
-				// TODO: We shouldn't modify the Simulation object here. This is just an optimization to cache vertex factories.
-				//       We should just create a temporary one here that lives for the duration of the frame.
-				const int32 ViewCount = ViewFamily->Views.Num();
-				while (Simulation->VertexFactories.Num() < ViewCount)
-				{
-					FGPUSpriteVertexFactory* VertexFactory = new(Simulation->VertexFactories) FGPUSpriteVertexFactory();
-					VertexFactory->InitResource();
-				}
-		
-				// Iterate over views and assign parameters for each.
-				FParticleSimulationResources* SimulationResources = FXSystem->GetParticleSimulationResources();
-				int32 ViewBit = 1;
-				for (int32 ViewIndex = 0; ViewIndex < ViewCount; ++ViewIndex, ViewBit <<= 1)
-				{
-					if (VisibilityMap & ViewBit)
-					{
-						FGPUSpriteVertexFactory& VertexFactory = Simulation->VertexFactories[ViewIndex];
-						if (bAllowSorting && SortMode == PSORTMODE_DistanceToView)
-						{
-							// Extensibility TODO: This call to AddSortedGPUSimulation is very awkward. When rendering a frame we need to
-							// accumulate all GPU particle emitters that need to be sorted. That is so they can be sorted in one big radix
-							// sort for efficiency. Ideally that state is per-scene renderer but the renderer doesn't know anything about particles.
-							const FSceneView* View = ViewFamily->Views[ViewIndex];
-							const int32 SortedBufferOffset = FXSystem->AddSortedGPUSimulation(Simulation, View->ViewMatrices.ViewOrigin);
-							check(SimulationResources->SortedVertexBuffer.IsInitialized());
-							VertexFactory.SetVertexBuffer(&SimulationResources->SortedVertexBuffer, SortedBufferOffset);
-						}
-						else
-						{
-							check(Simulation->VertexBuffer.IsInitialized());
-							VertexFactory.SetVertexBuffer(&Simulation->VertexBuffer, 0);
-						}
-					}
-				}
-
-				if (bUseLocalSpace == false)
-				{
-					Proxy->UpdateWorldSpacePrimitiveUniformBuffer();
-				}
-			}
-		}
-	}
-
 	virtual void GetDynamicMeshElementsEmitter(const FParticleSystemSceneProxy* Proxy, const FSceneView* View, const FSceneViewFamily& ViewFamily, int32 ViewIndex, FMeshElementCollector& Collector) const override
 	{
 		auto FeatureLevel = ViewFamily.GetFeatureLevel();
@@ -2770,87 +2682,6 @@ struct FGPUSpriteDynamicEmitterData : FDynamicEmitterDataBase
 		}
 	}
 
-	/**
-	 * Called to render the sprites for this emitter.
-	 */
-	virtual int32 RenderEmitter(FParticleSystemSceneProxy* Proxy, FPrimitiveDrawInterface* PDI,const FSceneView* View)
-	{
-		if (RHISupportsGPUParticles(View->GetFeatureLevel()))
-		{
-			check(Simulation);
-
-			if (Simulation->SimulationIndex != INDEX_NONE)
-			{
-				int32 DrawCalls = 0;
-				const int32 ParticleCount = Simulation->VertexBuffer.ParticleCount;
-				const int32 ViewIndex = View->Family->Views.Find(View);
-				const bool bIsWireframe = View->Family->EngineShowFlags.Wireframe;
-				if (ParticleCount > 0 && Simulation->VertexFactories.IsValidIndex(ViewIndex))
-				{
-					SCOPE_CYCLE_COUNTER(STAT_GPUSpriteRenderingTime);
-
-					FParticleSimulationResources* ParticleSimulationResources = FXSystem->GetParticleSimulationResources();
-					FParticleStateTextures& StateTextures = ParticleSimulationResources->GetCurrentStateTextures();
-					FGPUSpriteVertexFactory& VertexFactory = Simulation->VertexFactories[ViewIndex];
-					VertexFactory.EmitterUniformBuffer = Resources->UniformBuffer;
-					VertexFactory.EmitterDynamicUniformBuffer = DynamicUniformBuffer;
-					VertexFactory.PositionTextureRHI = StateTextures.PositionTextureRHI;
-					VertexFactory.VelocityTextureRHI = StateTextures.VelocityTextureRHI;
-					VertexFactory.AttributesTextureRHI = ParticleSimulationResources->RenderAttributesTexture.TextureRHI;
-
-					FMeshBatch Mesh;
-					FMeshBatchElement& BatchElement = Mesh.Elements[0];
-					BatchElement.IndexBuffer = &GParticleIndexBuffer;
-					BatchElement.NumPrimitives = MAX_PARTICLES_PER_INSTANCE * 2;
-					BatchElement.NumInstances = ParticleCount / MAX_PARTICLES_PER_INSTANCE;
-					BatchElement.FirstIndex = 0;
-					Mesh.VertexFactory = &VertexFactory;
-					Mesh.LCI = NULL;
-					if ( bUseLocalSpace )
-					{
-						BatchElement.PrimitiveUniformBufferResource = &Proxy->GetUniformBuffer();
-					}
-					else
-					{
-						BatchElement.PrimitiveUniformBufferResource = &Proxy->GetWorldSpacePrimitiveUniformBuffer();
-					}
-					BatchElement.MinVertexIndex = 0;
-					BatchElement.MaxVertexIndex = 3;
-					Mesh.ReverseCulling = Proxy->IsLocalToWorldDeterminantNegative();
-					Mesh.CastShadow = Proxy->GetCastShadow();
-					Mesh.DepthPriorityGroup = (ESceneDepthPriorityGroup)Proxy->GetDepthPriorityGroup(View);
-					const bool bUseSelectedMaterial = GIsEditor && (View->Family->EngineShowFlags.Selection) ? bSelected : 0;
-					Mesh.MaterialRenderProxy = Material->GetRenderProxy(bUseSelectedMaterial);
-					Mesh.Type = PT_TriangleList;
-
-					DrawCalls += DrawRichMesh(
-						PDI, 
-						Mesh, 
-						FLinearColor(1.0f, 0.0f, 0.0f),	//WireframeColor,
-						FLinearColor(1.0f, 1.0f, 0.0f),	//LevelColor,
-						FLinearColor(1.0f, 1.0f, 1.0f),	//PropertyColor,		
-						Proxy,
-						GIsEditor && (View->Family->EngineShowFlags.Selection) ? Proxy->IsSelected() : false,
-						bIsWireframe
-						);
-				}
-
-				const bool bHaveLocalVectorField = Simulation && Simulation->LocalVectorField.Resource;
-				if (bHaveLocalVectorField && View->Family->EngineShowFlags.VectorFields)
-				{
-					// Create a vertex factory for visualization if needed.
-					Simulation->CreateVectorFieldVisualizationVertexFactory();
-					check(Simulation->VectorFieldVisualizationVertexFactory);
-					DrawVectorFieldBounds(PDI, View, &Simulation->LocalVectorField);
-					DrawVectorField(PDI, View, Simulation->VectorFieldVisualizationVertexFactory, &Simulation->LocalVectorField);
-				}
-
-				return DrawCalls;
-			}
-		}
-		return 0;
-	}
-
 	/**
 	 * Retrieves the material render proxy with which to render sprites.
 	 */
diff --git a/Engine/Source/Runtime/Engine/Private/Particles/ParticleSystemRender.cpp b/Engine/Source/Runtime/Engine/Private/Particles/ParticleSystemRender.cpp
--- a/Engine/Source/Runtime/Engine/Private/Particles/ParticleSystemRender.cpp
+++ b/Engine/Source/Runtime/Engine/Private/Particles/ParticleSystemRender.cpp
@@ -36,226 +36,6 @@ float GMinParticleDrawTimeToTrack = .0001f;
 /** Whether to do LOD calculation on GameThread in game */
 extern bool GbEnableGameThreadLODCalculation;
 
-/*-----------------------------------------------------------------------------
-	Stats.
------------------------------------------------------------------------------*/
-
-/** Stats gathered about each UParticleSystem that is rendered. */
-struct FParticleTemplateRenderStats
-{
-	float MaxRenderTime;
-	int32 NumDraws;
-	int32 NumEmitters;
-	int32 NumComponents;
-	int32 NumDrawDynamicElements;
-
-	FParticleTemplateRenderStats() :
-		NumComponents(0),
-		NumDrawDynamicElements(0)
-	{}
-};
-
-/** Sorts FParticleTemplateRenderStats from longest MaxRenderTime to shortest. */
-struct FCompareFParticleTemplateRenderStatsByRenderTime
-{
-	FORCEINLINE bool operator()( const FParticleTemplateRenderStats& A, const FParticleTemplateRenderStats& B ) const { return B.MaxRenderTime < A.MaxRenderTime; }
-};
-
-/** Stats gathered about each UParticleSystemComponent that is rendered. */
-struct FParticleComponentRenderStats
-{
-	float MaxRenderTime;
-	int32 NumDraws;
-};
-
-/** Sorts FParticleComponentRenderStats from longest MaxRenderTime to shortest. */
-struct FCompareFParticleComponentRenderStatsByRenderTime
-{
-	FORCEINLINE bool operator()( const FParticleComponentRenderStats& A, const FParticleComponentRenderStats& B ) const { return B.MaxRenderTime < A.MaxRenderTime; }
-};
-
-/** Global map from UParticleSystem path name to stats about that particle system, used when tracking particle render stats. */
-TMap<FString, FParticleTemplateRenderStats> GTemplateRenderStats;
-
-/** Global map from UParticleSystemComponent path name to stats about that component, used when tracking particle render stats. */
-TMap<FString, FParticleComponentRenderStats> GComponentRenderStats;
-
-/** 
- * Dumps particle render stats to the game's log directory, and resets the stats tracked so far. 
- * This is hooked up to the DUMPPARTICLERENDERINGSTATS console command.
- */
-void DumpParticleRenderingStats(FOutputDevice& Ar)
-{
-#if STATS
-	if (GTrackParticleRenderingStats)
-	{
-#if ALLOW_DEBUG_FILES
-		{
-			// Have to keep this filename short enough to be valid for Xbox 360
-			const FString TemplateFileName = FString::Printf(TEXT("%sPT-%s.csv"), *FPaths::ProfilingDir(), *FDateTime::Now().ToString());
-			FDiagnosticTableViewer ParticleTemplateViewer(*TemplateFileName);
-
-			// Write a row of headings for the table's columns.
-			ParticleTemplateViewer.AddColumn(TEXT("MaxRenderTime ms"));
-			ParticleTemplateViewer.AddColumn(TEXT("NumEmitters"));
-			ParticleTemplateViewer.AddColumn(TEXT("NumDraws"));
-			ParticleTemplateViewer.AddColumn(TEXT("Template Name"));
-			ParticleTemplateViewer.CycleRow();
-
-			// Sort from longest render time to shortest
-			GTemplateRenderStats.ValueSort( FCompareFParticleTemplateRenderStatsByRenderTime() );
-
-			for (TMap<FString, FParticleTemplateRenderStats>::TIterator It(GTemplateRenderStats); It; ++It)
-			{	
-				const FParticleTemplateRenderStats& Stats = It.Value();
-				ParticleTemplateViewer.AddColumn(TEXT("%.2f"), Stats.MaxRenderTime * 1000.0f);
-				ParticleTemplateViewer.AddColumn(TEXT("%u"), Stats.NumEmitters);
-				ParticleTemplateViewer.AddColumn(TEXT("%u"), Stats.NumDraws);
-				ParticleTemplateViewer.AddColumn(*It.Key());
-				ParticleTemplateViewer.CycleRow();
-			}
-			Ar.Logf(TEXT("Template stats saved to %s"), *TemplateFileName);
-		}
-		
-		{
-			const FString ComponentFileName = FString::Printf(TEXT("%sPC-%s.csv"), *FPaths::ProfilingDir(), *FDateTime::Now().ToString());
-			FDiagnosticTableViewer ParticleComponentViewer(*ComponentFileName);
-
-			// Write a row of headings for the table's columns.
-			ParticleComponentViewer.AddColumn(TEXT("MaxRenderTime ms"));
-			ParticleComponentViewer.AddColumn(TEXT("NumDraws"));
-			ParticleComponentViewer.AddColumn(TEXT("Actor Name"));
-			ParticleComponentViewer.CycleRow();
-
-			// Sort from longest render time to shortest
-			GComponentRenderStats.ValueSort( FCompareFParticleComponentRenderStatsByRenderTime() );
-
-			for (TMap<FString, FParticleComponentRenderStats>::TIterator It(GComponentRenderStats); It; ++It)
-			{	
-				const FParticleComponentRenderStats& Stats = It.Value();
-				ParticleComponentViewer.AddColumn(TEXT("%.2f"), Stats.MaxRenderTime * 1000.0f);
-				ParticleComponentViewer.AddColumn(TEXT("%u"), Stats.NumDraws);
-				ParticleComponentViewer.AddColumn(*It.Key());
-				ParticleComponentViewer.CycleRow();
-			}
-			Ar.Logf(TEXT("Component stats saved to %s"), *ComponentFileName);
-		}
-#endif
-		GTemplateRenderStats.Empty();
-		GComponentRenderStats.Empty();
-	}
-	else
-	{
-		Ar.Logf(TEXT("Need to start tracking with TRACKPARTICLERENDERINGSTATS first."));
-	}
-#endif
-}
-
-bool GWantsParticleStatsNextFrame = false;
-bool GTrackParticleRenderingStatsForOneFrame = false;
-
-#if STATS
-
-/** Kicks off a particle render stat frame capture, initiated by the DUMPPARTICLEFRAMERENDERINGSTATS console command. */
-void ENGINE_API BeginOneFrameParticleStats()
-{
-	if (GWantsParticleStatsNextFrame)
-	{
-		UE_LOG(LogParticles, Warning, TEXT("BeginOneFrameParticleStats"));
-		GWantsParticleStatsNextFrame = false;
-		// Block until the renderer processes all queued commands, because the rendering thread reads from GTrackParticleRenderingStatsForOneFrame
-		FlushRenderingCommands();
-		GTrackParticleRenderingStatsForOneFrame = true;
-	}
-}
-
-TMap<FString, FParticleTemplateRenderStats> GOneFrameTemplateRenderStats;
-
-/** 
- * Dumps particle render stats to the game's log directory along with a screenshot. 
- * This is hooked up to the DUMPPARTICLEFRAMERENDERINGSTATS console command.
- */
-void ENGINE_API FinishOneFrameParticleStats()
-{
-	if (GTrackParticleRenderingStatsForOneFrame)
-	{
-		UE_LOG(LogParticles, Warning, TEXT("FinishOneFrameParticleStats"));
-
-		// Wait until the rendering thread finishes processing the previous frame
-		FlushRenderingCommands();
-
-		GTrackParticleRenderingStatsForOneFrame = false;
-
-		const FString PathFromScreenshots = FString(TEXT("Particle")) + FDateTime::Now().ToString();
-		const FString WritePath = FPaths::ScreenShotDir() + PathFromScreenshots;
-
-		const bool bDirSuccess = IFileManager::Get().MakeDirectory(*WritePath);
-
-		UE_LOG(LogParticles, Warning, TEXT("IFileManager::MakeDirectory %s %u"), *WritePath, bDirSuccess);
-
-#if ALLOW_DEBUG_FILES
-		// Have to keep this filename short enough to be valid for Xbox 360
-		const FString TemplateFileName = WritePath / TEXT("ParticleTemplates.csv");
-		FDiagnosticTableViewer ParticleTemplateViewer(*TemplateFileName);
-
-		UE_LOG(LogParticles, Warning, TEXT("TemplateFileName %s %u"), *TemplateFileName, ParticleTemplateViewer.OutputStreamIsValid());
-
-		// Write a row of headings for the table's columns.
-		ParticleTemplateViewer.AddColumn(TEXT("RenderTime ms"));
-		ParticleTemplateViewer.AddColumn(TEXT("NumComponents"));
-		ParticleTemplateViewer.AddColumn(TEXT("NumPasses"));
-		ParticleTemplateViewer.AddColumn(TEXT("NumEmitters"));
-		ParticleTemplateViewer.AddColumn(TEXT("NumDraws"));
-		ParticleTemplateViewer.AddColumn(TEXT("Template Name"));
-		ParticleTemplateViewer.CycleRow();
-
-		float TotalRenderTime = 0;
-		int32 TotalNumComponents = 0;
-		int32 TotalNumDrawDynamicElements = 0;
-		int32 TotalNumEmitters = 0;
-		int32 TotalNumDraws = 0;
-
-		// Sort from longest render time to shortest
-		GOneFrameTemplateRenderStats.ValueSort( FCompareFParticleTemplateRenderStatsByRenderTime() );
-
-		for (TMap<FString, FParticleTemplateRenderStats>::TIterator It(GOneFrameTemplateRenderStats); It; ++It)
-		{	
-			const FParticleTemplateRenderStats& Stats = It.Value();
-			TotalRenderTime += Stats.MaxRenderTime;
-			TotalNumComponents += Stats.NumComponents;
-			TotalNumDrawDynamicElements += Stats.NumDrawDynamicElements;
-			TotalNumEmitters += Stats.NumEmitters;
-			TotalNumDraws += Stats.NumDraws;
-			ParticleTemplateViewer.AddColumn(TEXT("%.2f"), Stats.MaxRenderTime * 1000.0f);
-			ParticleTemplateViewer.AddColumn(TEXT("%u"), Stats.NumComponents);
-			ParticleTemplateViewer.AddColumn(TEXT("%.1f"), Stats.NumDrawDynamicElements / (float)Stats.NumComponents);
-			ParticleTemplateViewer.AddColumn(TEXT("%u"), Stats.NumEmitters);
-			ParticleTemplateViewer.AddColumn(TEXT("%u"), Stats.NumDraws);
-			ParticleTemplateViewer.AddColumn(*It.Key());
-			ParticleTemplateViewer.CycleRow();
-		}
-
-		ParticleTemplateViewer.AddColumn(TEXT("%.2f"), TotalRenderTime * 1000.0f);
-		ParticleTemplateViewer.AddColumn(TEXT("%u"), TotalNumComponents);
-		ParticleTemplateViewer.AddColumn(TEXT("%.1f"), TotalNumDrawDynamicElements / (float)TotalNumComponents);
-		ParticleTemplateViewer.AddColumn(TEXT("%u"), TotalNumEmitters);
-		ParticleTemplateViewer.AddColumn(TEXT("%u"), TotalNumDraws);
-		ParticleTemplateViewer.AddColumn(TEXT("Totals"));
-		ParticleTemplateViewer.CycleRow();
-
-		UE_LOG(LogParticles, Warning, TEXT("One frame stats saved to %s"), *TemplateFileName);
-#endif
-
-		// Request a screenshot to be saved next frame
-		FScreenshotRequest::RequestScreenshot( PathFromScreenshots / TEXT("Shot.bmp"), false );
-		
-		UE_LOG(LogParticles, Warning, TEXT("ScreenShotName %s"), *FScreenshotRequest::GetFilename() );
-
-		GOneFrameTemplateRenderStats.Empty();
-	}
-}
-#endif
-
 /** 
  *	Create a vertex factory for the given type.
  *
@@ -1025,104 +805,6 @@ bool FDynamicSpriteEmitterData::GetVertexAndIndexDataNonInstanced(void* VertexDa
 	return true;
 }
 
-
-
-/**
- *	Called during InitViews for view processing on scene proxies before rendering them
- *  Only called for primitives that are visible and have bDynamicRelevance
- *
- *	@param	Proxy			The 'owner' particle system scene proxy
- *	@param	ViewFamily		The ViewFamily to pre-render for
- *	@param	VisibilityMap	A BitArray that indicates whether the primitive was visible in that view (index)
- *	@param	FrameNumber		The frame number of this pre-render
- */
-void FDynamicSpriteEmitterData::PreRenderView(FParticleSystemSceneProxy* Proxy, const FSceneViewFamily* ViewFamily, const uint32 VisibilityMap, int32 FrameNumber)
-{
-	const auto FeatureLevel = ViewFamily->Scene->GetFeatureLevel();
-	const bool bInstanced = FeatureLevel >= ERHIFeatureLevel::SM4;
-
-	// Sort and generate particles for this view.
-	const FDynamicSpriteEmitterReplayDataBase* SourceData = GetSourceData();
-	if ( SourceData && SourceData->EmitterRenderMode == ERM_Normal )
-	{
-		// Determine how many vertices and indices are needed to render.
-		const int32 ParticleCount = SourceData->ActiveParticleCount;
-		const int32 VertexSize = GetDynamicVertexStride(FeatureLevel);
-		const int32 DynamicParameterVertexSize = sizeof(FParticleVertexDynamicParameter);
-		const int32 VertexPerParticle = bInstanced ? 1 : 4;
-		const int32 ViewCount = ViewFamily->Views.Num();
-		int32 ViewBit = 1;
-		InstanceDataAllocations.Empty();
-		DynamicParameterDataAllocations.Empty();
-		PerViewUniformBuffers.Empty();
-
-		if (SourceData->bUseLocalSpace == false)
-		{
-			Proxy->UpdateWorldSpacePrimitiveUniformBuffer();
-		}
-
-		// Iterate over views and generate vertices for each view.
-		for ( int32 ViewIndex = 0; ViewIndex < ViewCount; ++ViewIndex, ViewBit <<= 1 )
-		{
-			FGlobalDynamicVertexBuffer::FAllocation* Allocation = new(InstanceDataAllocations) FGlobalDynamicVertexBuffer::FAllocation();
-			FGlobalDynamicVertexBuffer::FAllocation* DynamicParameterAllocation = new(DynamicParameterDataAllocations) FGlobalDynamicVertexBuffer::FAllocation();
-			
-			FParticleSpriteUniformBufferRef* SpriteUniformBufferPtr = new(PerViewUniformBuffers) FParticleSpriteUniformBufferRef();
-			if ( VisibilityMap & ViewBit )
-			{
-				const FSceneView* View = ViewFamily->Views[ViewIndex];
-
-				// Allocate memory for render data.
-				*Allocation = FGlobalDynamicVertexBuffer::Get().Allocate( ParticleCount * VertexSize * VertexPerParticle );
-
-				if (bUsesDynamicParameter)
-				{
-					*DynamicParameterAllocation = FGlobalDynamicVertexBuffer::Get().Allocate( ParticleCount * DynamicParameterVertexSize * VertexPerParticle );
-				}
-				
-				if ( Allocation->IsValid() && (!bUsesDynamicParameter || DynamicParameterAllocation->IsValid()) )
-				{
-					// Sort the particles if needed.
-					FParticleOrder* ParticleOrder = NULL;
-					if (SourceData->SortMode != PSORTMODE_None)
-					{
-						// If material is using unlit translucency and the blend mode is translucent then we need to sort (back to front)
-						const FMaterial* Material = MaterialResource[bSelected]->GetMaterial(FeatureLevel);
-						if (Material && 
-							(Material->GetBlendMode() == BLEND_Translucent ||
-							((SourceData->SortMode == PSORTMODE_Age_OldestFirst) || (SourceData->SortMode == PSORTMODE_Age_NewestFirst)))
-							)
-						{
-							ParticleOrder = GParticleOrderPool.GetParticleOrderData(ParticleCount);
-							SortSpriteParticles(SourceData->SortMode, SourceData->bUseLocalSpace, SourceData->ActiveParticleCount, 
-								SourceData->ParticleData, SourceData->ParticleStride, SourceData->ParticleIndices,
-								View, Proxy->GetLocalToWorld(), ParticleOrder);
-						}
-					}
-
-					// Fill vertex buffers.
-					if(bInstanced)
-					{
-						GetVertexAndIndexData(Allocation->Buffer, DynamicParameterAllocation->Buffer, NULL, ParticleOrder, View->ViewMatrices.ViewOrigin, Proxy->GetLocalToWorld());
-					}
-					else
-					{
-						GetVertexAndIndexDataNonInstanced(Allocation->Buffer, DynamicParameterAllocation->Buffer, NULL, ParticleOrder, View->ViewMatrices.ViewOrigin, Proxy->GetLocalToWorld());
-					}
-				}
-
-				// Create per-view uniform buffer.
-				FParticleSpriteUniformParameters PerViewUniformParameters = UniformParameters;
-				FVector2D ObjectNDCPosition;
-				FVector2D ObjectMacroUVScales;
-				Proxy->GetObjectPositionAndScale(*View,ObjectNDCPosition, ObjectMacroUVScales);
-				PerViewUniformParameters.MacroUVParameters = FVector4(ObjectNDCPosition.X, ObjectNDCPosition.Y, ObjectMacroUVScales.X, ObjectMacroUVScales.Y);
-				*SpriteUniformBufferPtr = FParticleSpriteUniformBufferRef::CreateUniformBufferImmediate(PerViewUniformParameters, UniformBuffer_SingleFrame);
-			}
-		}
-	}
-}
-
 void GatherParticleLightData(const FDynamicSpriteEmitterReplayDataBase& Source, const FMatrix& InLocalToWorld, const FSceneViewFamily& InViewFamily, FSimpleLightArray& OutParticleLights)
 {
 	if (Source.LightDataOffset != 0)
@@ -1189,7 +871,6 @@ void GatherParticleLightData(const FDynamicSpriteEmitterReplayDataBase& Source,
 
 				ApplyOrbitToPosition(Particle, Source, InLocalToWorld, LocalOrbitPayload, Unused, Unused2, ParticlePosition, Unused3);
 
-
 				FVector LightPosition = Source.bUseLocalSpace ? FVector(InLocalToWorld.TransformPosition(ParticlePosition)) : ParticlePosition;
 				
 				// Calculate light positions per-view if we are using a camera offset.
@@ -1431,136 +1112,6 @@ void FDynamicSpriteEmitterData::GetDynamicMeshElementsEmitter(const FParticleSys
 	}
 }
 
-/**
- *	Render thread only draw call
- *
- *	@param	Proxy		The scene proxy for the particle system that owns this emitter
- *	@param	PDI			The primitive draw interface to render with
- *	@param	View		The scene view being rendered
- */
-int32 FDynamicSpriteEmitterData::RenderEmitter(FParticleSystemSceneProxy* Proxy, FPrimitiveDrawInterface* PDI,const FSceneView* View)
-{
-	SCOPE_CYCLE_COUNTER(STAT_SpriteRenderingTime);
-
-	const auto FeatureLevel = View->GetFeatureLevel();
-
-	const bool bInstanced = FeatureLevel >= ERHIFeatureLevel::SM4;
-
-	if (bValid == false)
-	{
-		return 0;
-	}
-
-	const FDynamicSpriteEmitterReplayDataBase* SourceData = GetSourceData();
-	check(SourceData);
-	FParticleSpriteVertexFactory* SpriteVertexFactory = (FParticleSpriteVertexFactory*)GetVertexFactory(Proxy);
-
-	int32 NumDraws = 0;
-	if (SourceData->EmitterRenderMode == ERM_Normal)
-	{
-		// Don't render if the material will be ignored
-		const bool bIsWireframe = View->Family->EngineShowFlags.Wireframe;
-		if (PDI->IsMaterialIgnored(MaterialResource[bSelected], FeatureLevel) && !bIsWireframe)
-		{
-			return 0;
-		}
-
-		// Find the vertex allocation for this view.
-		FGlobalDynamicVertexBuffer::FAllocation* Allocation = NULL;
-		FGlobalDynamicVertexBuffer::FAllocation* DynamicParameterAllocation = NULL;
-		
-		const int32 ViewIndex = View->Family->Views.Find( View );
-		if ( InstanceDataAllocations.IsValidIndex( ViewIndex ) )
-		{
-			Allocation = &InstanceDataAllocations[ ViewIndex ];
-		}
-
-		if ( DynamicParameterDataAllocations.IsValidIndex( ViewIndex ) )
-		{
-			DynamicParameterAllocation = &DynamicParameterDataAllocations[ ViewIndex ];
-		}
-
-		if ( Allocation && DynamicParameterAllocation && Allocation->IsValid() && (!bUsesDynamicParameter || DynamicParameterAllocation->IsValid()) && IsValidRef( PerViewUniformBuffers[ViewIndex] ) )
-		{
-			// Calculate the number of particles that must be drawn.
-			int32 ParticleCount = Source.ActiveParticleCount;
-			if ((Source.MaxDrawCount >= 0) && (ParticleCount > Source.MaxDrawCount))
-			{
-				ParticleCount = Source.MaxDrawCount;
-			}
-
-			// Set the sprite uniform buffer for this view.
-			check( SpriteVertexFactory );
-			SpriteVertexFactory->SetSpriteUniformBuffer( PerViewUniformBuffers[ViewIndex] );
-			SpriteVertexFactory->SetInstanceBuffer( Allocation->VertexBuffer, Allocation->VertexOffset, GetDynamicVertexStride(View->GetFeatureLevel()), bInstanced );
-			SpriteVertexFactory->SetDynamicParameterBuffer( DynamicParameterAllocation ? DynamicParameterAllocation->VertexBuffer : NULL, DynamicParameterAllocation->VertexOffset, GetDynamicParameterVertexStride(), bInstanced );
-
-			// Construct the mesh element to render.
-			FMeshBatch Mesh;
-			FMeshBatchElement& BatchElement = Mesh.Elements[0];
-			BatchElement.IndexBuffer = &GParticleIndexBuffer;
-			if(bInstanced)
-			{
-				BatchElement.NumPrimitives = 2;
-				BatchElement.NumInstances = ParticleCount;
-			}
-			else
-			{
-				BatchElement.NumPrimitives = 2 * ParticleCount;
-				BatchElement.NumInstances = 1;
-			}
-			BatchElement.FirstIndex = 0;
-			Mesh.VertexFactory = SpriteVertexFactory;
-			// if the particle rendering data is presupplied, use it directly
-			Mesh.LCI = NULL;
-			if (SourceData->bUseLocalSpace == true)
-			{
-				BatchElement.PrimitiveUniformBufferResource = &Proxy->GetUniformBuffer();
-			}
-			else
-			{
-				BatchElement.PrimitiveUniformBufferResource = &Proxy->GetWorldSpacePrimitiveUniformBuffer();
-			}
-			BatchElement.MinVertexIndex = 0;
-			BatchElement.MaxVertexIndex = (ParticleCount * 4) - 1;
-			Mesh.CastShadow = Proxy->GetCastShadow();
-			Mesh.DepthPriorityGroup = (ESceneDepthPriorityGroup)Proxy->GetDepthPriorityGroup(View);
-
-			if ( bIsWireframe )
-			{
-				Mesh.MaterialRenderProxy = UMaterial::GetDefaultMaterial( MD_Surface )->GetRenderProxy( View->Family->EngineShowFlags.Selection ? bSelected : false );
-			}
-			else
-			{
-				Mesh.MaterialRenderProxy = MaterialResource[GIsEditor && (View->Family->EngineShowFlags.Selection) ? bSelected : 0];
-			}
-			Mesh.Type = PT_TriangleList;
-
-			NumDraws += DrawRichMesh(
-				PDI, 
-				Mesh, 
-				FLinearColor(1.0f, 0.0f, 0.0f),	//WireframeColor,
-				FLinearColor(1.0f, 1.0f, 0.0f),	//LevelColor,
-				FLinearColor(1.0f, 1.0f, 1.0f),	//PropertyColor,		
-				Proxy,
-				GIsEditor && (View->Family->EngineShowFlags.Selection) ? Proxy->IsSelected() : false,
-				bIsWireframe
-				);
-		}
-	}
-	else if (SourceData->EmitterRenderMode == ERM_Point)
-	{
-		RenderDebug(Proxy, PDI, View, false);
-	}
-	else if (SourceData->EmitterRenderMode == ERM_Cross)
-	{
-		RenderDebug(Proxy, PDI, View, true);
-	}
-
-	return NumDraws;
-}
-
-
 FParticleVertexFactoryBase* FDynamicSpriteEmitterData::BuildVertexFactory(const FParticleSystemSceneProxy* InOwnerProxy)
 {
 	return GParticleVertexFactoryPool.GetParticleVertexFactory(PVFT_Sprite, InOwnerProxy->GetScene().GetFeatureLevel(), this);
@@ -1649,19 +1200,6 @@ void FDynamicSpriteEmitterData::UpdateRenderThreadResourcesEmitter(const FPartic
 	}
 }
 
-/**
- *	Release the render thread resources for this emitter data
- *
- *	@param	InOwnerProxy	The proxy that owns this dynamic emitter data
- *
- *	@return	bool			true if successful, false if failed
- */
-void FDynamicSpriteEmitterData::ReleaseRenderThreadResources(const FParticleSystemSceneProxy* InOwnerProxy)
-{
-	FDynamicSpriteEmitterDataBase::ReleaseRenderThreadResources( InOwnerProxy );
-	PerViewUniformBuffers.Empty();
-}
-
 ///////////////////////////////////////////////////////////////////////////////
 //	FDynamicMeshEmitterData
 ///////////////////////////////////////////////////////////////////////////////
@@ -2004,265 +1542,6 @@ void FDynamicMeshEmitterData::GetDynamicMeshElementsEmitter(const FParticleSyste
 	}
 }
 
-int32 FDynamicMeshEmitterData::RenderEmitter(FParticleSystemSceneProxy* Proxy, FPrimitiveDrawInterface* PDI,const FSceneView* View)
-{
-	SCOPE_CYCLE_COUNTER(STAT_MeshRenderingTime);
-
-	if (bValid == false)
-	{
-		return 0;
-	}
-
-	const bool bInstanced = View->GetFeatureLevel() >= ERHIFeatureLevel::SM4;
-
-	int32 NumDraws = 0;
-	if (Source.EmitterRenderMode == ERM_Normal)
-	{
-		const bool bIsWireframe = AllowDebugViewmodes() && View->Family->EngineShowFlags.Wireframe;
-
-		// Find the vertex allocation for this view.
-		const int32 ViewIndex = View->Family->Views.Find( View );
-		
-		// Only unsorted instance data will exist for unsorted emitters and must be used for all views being rendered. 
-		const int32 InstanceBufferIndex = ((Source.SortMode == PSORTMODE_None) ? 0 : ViewIndex);
-
-		FMeshParticleVertexFactory* MeshVertexFactory = (FMeshParticleVertexFactory*)GetVertexFactory(Proxy);
-
-		if(bInstanced)
-		{
-			check(InstanceDataAllocations.IsValidIndex( InstanceBufferIndex ) != false);
-			check(InstanceDataAllocations[InstanceBufferIndex].IsValid());
-			check(DynamicParameterDataAllocations.IsValidIndex( InstanceBufferIndex ) != false);
-			check(!bUsesDynamicParameter || DynamicParameterDataAllocations[InstanceBufferIndex].IsValid());
-
-			// Set the particle instance data. This buffer is generated in PreRenderView.
-			MeshVertexFactory->SetInstanceBuffer(InstanceDataAllocations[InstanceBufferIndex].VertexBuffer, InstanceDataAllocations[InstanceBufferIndex].VertexOffset, GetDynamicVertexStride(View->GetFeatureLevel()));
-			MeshVertexFactory->SetDynamicParameterBuffer(DynamicParameterDataAllocations[InstanceBufferIndex].VertexBuffer, DynamicParameterDataAllocations[InstanceBufferIndex].VertexOffset , GetDynamicParameterVertexStride());
-		}
-
-		check(InstanceBufferIndex < (FirstBatchForView.Num() - 1));
-			
-		int32 FirstBatch = FirstBatchForView[InstanceBufferIndex];
-		int32 LastBatchPlusOne = FirstBatchForView[InstanceBufferIndex + 1];
-		for (int32 BatchIndex = FirstBatch; BatchIndex < LastBatchPlusOne; ++BatchIndex)
-		{
-			FMeshBatch& Mesh = *MeshBatches[BatchIndex];
-			NumDraws += DrawRichMesh(
-				PDI, 
-				Mesh, 
-				FLinearColor(1.0f, 0.0f, 0.0f),	//WireframeColor,
-				FLinearColor(1.0f, 1.0f, 0.0f),	//LevelColor,
-				FLinearColor(1.0f, 1.0f, 1.0f),	//PropertyColor,		
-				Proxy,
-				GIsEditor && (View->Family->EngineShowFlags.Selection) ? bSelected : false,
-				bIsWireframe
-				);
-		}
-	}
-	else if (Source.EmitterRenderMode == ERM_Point)
-	{
-		RenderDebug(Proxy, PDI, View, false);
-	}
-	else if (Source.EmitterRenderMode == ERM_Cross)
-	{
-		RenderDebug(Proxy, PDI, View, true);
-	}
-
-	return NumDraws;
-}
-
-/**
- *	Called during InitViews for view processing on scene proxies before rendering them
- *  Only called for primitives that are visible and have bDynamicRelevance
- *
- *	@param	Proxy			The 'owner' particle system scene proxy
- *	@param	ViewFamily		The ViewFamily to pre-render for
- *	@param	VisibilityMap	A BitArray that indicates whether the primitive was visible in that view (index)
- *	@param	FrameNumber		The frame number of this pre-render
- */
-void FDynamicMeshEmitterData::PreRenderView(FParticleSystemSceneProxy* Proxy, const FSceneViewFamily* ViewFamily, const uint32 VisibilityMap, int32 FrameNumber)
-{
-	if (bValid == false)
-	{
-		return;
-	}
-
-	const auto FeatureLevel = ViewFamily->Scene->GetFeatureLevel();
-	const bool bInstanced = FeatureLevel >= ERHIFeatureLevel::SM4;
-
-	int32 ParticleCount = Source.ActiveParticleCount;
-	if ((Source.MaxDrawCount >= 0) && (ParticleCount > Source.MaxDrawCount))
-	{
-		ParticleCount = Source.MaxDrawCount;
-	}
-
-	const int32 InstanceVertexStride = GetDynamicVertexStride(FeatureLevel);
-	const int32 DynamicParameterVertexStride  = GetDynamicParameterVertexStride();
-	const int32 ViewCount = ViewFamily->Views.Num();
-	if(bInstanced)
-	{
-		InstanceDataAllocations.Empty();
-		DynamicParameterDataAllocations.Empty();
-	}
-	else
-	{
-		InstanceDataAllocationsCPU.Reset();
-		DynamicParameterDataAllocationsCPU.Reset();
-	}
-
-	// Iterate over views and generate instance data for each view.
-	for ( int32 ViewIndex = 0; ViewIndex < ViewCount; ++ViewIndex )
-	{	
-		if(bInstanced)
-		{
-			FGlobalDynamicVertexBuffer::FAllocation* Allocation = new(InstanceDataAllocations) FGlobalDynamicVertexBuffer::FAllocation();
-			FGlobalDynamicVertexBuffer::FAllocation* DynamicParameterAllocation = new(DynamicParameterDataAllocations) FGlobalDynamicVertexBuffer::FAllocation();
-			
-			*Allocation = FGlobalDynamicVertexBuffer::Get().Allocate( ParticleCount * InstanceVertexStride );
-
-			if (bUsesDynamicParameter)
-			{
-				*DynamicParameterAllocation = FGlobalDynamicVertexBuffer::Get().Allocate( ParticleCount * DynamicParameterVertexStride );
-			}
-
-			if(Allocation->IsValid() && (!bUsesDynamicParameter || DynamicParameterAllocation->IsValid()))
-			{
-				// Fill instance buffer.
-				GetInstanceData(Allocation->Buffer, DynamicParameterAllocation->Buffer, Proxy, ViewFamily->Views[ViewIndex]);
-			}
-		}
-		else
-		{
-			InstanceDataAllocationsCPU.Reset(ParticleCount);
-			InstanceDataAllocationsCPU.AddUninitialized(ParticleCount);
-			DynamicParameterDataAllocationsCPU.Reset(ParticleCount);
-			DynamicParameterDataAllocationsCPU.AddUninitialized(ParticleCount);
-
-			// Fill instance buffer.
-			GetInstanceData((void*) InstanceDataAllocationsCPU.GetData(), (void*) DynamicParameterDataAllocationsCPU.GetData(), Proxy, ViewFamily->Views[ViewIndex]);
-		}
-
-		// Instance data is not calculated per-view when the emitter is unsorted.
-		if(Source.SortMode == PSORTMODE_None)
-		{
-			break;
-		}
-	}
-
-	// Setup the vertex factory.
-	FMeshParticleVertexFactory* MeshVertexFactory = (FMeshParticleVertexFactory*)GetVertexFactory(Proxy);
-	MeshVertexFactory->SetUniformBuffer( UniformBuffer );
-
-	// Update the primitive uniform buffer
-	if (LastFramePreRendered != FrameNumber)
-	{
-		Proxy->UpdateWorldSpacePrimitiveUniformBuffer();
-		LastFramePreRendered = FrameNumber;
-	}
-
-	// Setup the mesh elements we will render with later on.
-	MeshBatches.Reset();
-	FirstBatchForView.Reset();
-	if (Source.EmitterRenderMode == ERM_Normal)
-	{
-		const FStaticMeshLODResources& LODModel = StaticMesh->RenderData->LODResources[0];
-
-		check(MeshBatches.Num() == 0);
-		FirstBatchForView.Add(0);
-		for ( int32 ViewIndex = 0; ViewIndex < ViewCount; ++ViewIndex )
-		{
-			const FSceneView* View = ViewFamily->Views[ViewIndex];
-			const bool bIsWireframe = AllowDebugViewmodes() && View->Family->EngineShowFlags.Wireframe;
-
-			//@todo. Handle LODs.
-			for (int32 LODIndex = 0; LODIndex < 1; LODIndex++)
-			{
-				for (int32 SectionIndex = 0; SectionIndex < LODModel.Sections.Num(); SectionIndex++)
-				{
-					FMeshBatch* MeshPtr = Proxy->GetPooledMeshBatch();
-					FMeshBatch& Mesh = *MeshPtr;
-					Mesh.Elements.Reset();
-					FMeshBatchElement& BatchElement = *new(Mesh.Elements) FMeshBatchElement();
-					Mesh.VertexFactory = GetVertexFactory(Proxy);
-					Mesh.DynamicVertexData = NULL;
-					Mesh.LCI = NULL;
-					Mesh.UseDynamicData = false;
-					Mesh.ReverseCulling = Proxy->IsLocalToWorldDeterminantNegative();
-					Mesh.CastShadow = Proxy->GetCastShadow();
-					Mesh.DepthPriorityGroup = (ESceneDepthPriorityGroup)Proxy->GetDepthPriorityGroup(View);
-
-					FMaterialRenderProxy* MaterialProxy = MeshMaterials[SectionIndex]->GetRenderProxy(bSelected);
-					const FStaticMeshSection& Section = LODModel.Sections[SectionIndex];
-					if ((Section.NumTriangles == 0) || (MaterialProxy == NULL))
-					{
-						//@todo. This should never occur, but it does occasionally.
-						continue;
-					}
-
-					BatchElement.PrimitiveUniformBufferResource = &Proxy->GetWorldSpacePrimitiveUniformBuffer();
-					BatchElement.FirstIndex = Section.FirstIndex;
-					BatchElement.MinVertexIndex = Section.MinVertexIndex;
-					BatchElement.MaxVertexIndex = Section.MaxVertexIndex;
-					BatchElement.NumInstances = bInstanced ? ParticleCount : 1;
-
-					if (bIsWireframe)
-					{
-						if( LODModel.WireframeIndexBuffer.IsInitialized()
-							&& !(RHISupportsTessellation(GShaderPlatformForFeatureLevel[FeatureLevel]) && Mesh.VertexFactory->GetType()->SupportsTessellationShaders())
-							)
-						{
-							Mesh.Type = PT_LineList;
-							Mesh.MaterialRenderProxy = Proxy->GetDeselectedWireframeMatInst();
-							BatchElement.FirstIndex = 0;
-							BatchElement.IndexBuffer = &LODModel.WireframeIndexBuffer;
-							BatchElement.NumPrimitives = LODModel.WireframeIndexBuffer.GetNumIndices() / 2;
-						
-						}
-						else
-						{
-							Mesh.Type = PT_TriangleList;
-							Mesh.MaterialRenderProxy = MeshMaterials[SectionIndex]->GetRenderProxy(bSelected);
-							Mesh.bWireframe = true;	
-							BatchElement.FirstIndex = 0;
-							BatchElement.IndexBuffer = &LODModel.IndexBuffer;
-							BatchElement.NumPrimitives = LODModel.IndexBuffer.GetNumIndices() / 3;
-						}
-					}
-					else
-					{
-						Mesh.Type = PT_TriangleList;
-						Mesh.MaterialRenderProxy = MeshMaterials[SectionIndex]->GetRenderProxy(bSelected);
-						Mesh.bWireframe = false;
-						BatchElement.IndexBuffer = &LODModel.IndexBuffer;
-						BatchElement.FirstIndex = Section.FirstIndex;
-						BatchElement.NumPrimitives = Section.NumTriangles;
-					}
-
-					if(!bInstanced)
-					{
-						FMeshParticleVertexFactory::FBatchParametersCPU* BatchParameters = new(MeshBatchParameters) FMeshParticleVertexFactory::FBatchParametersCPU();
-						BatchParameters->InstanceBuffer = InstanceDataAllocationsCPU.GetData();
-						BatchParameters->DynamicParameterBuffer = DynamicParameterDataAllocationsCPU.GetData();
-						BatchElement.UserData = BatchParameters;
-						BatchElement.UserIndex = 0;
-
-						Mesh.Elements.Reserve(ParticleCount);
-						for(int32 ParticleIndex = 1; ParticleIndex < ParticleCount; ++ParticleIndex)
-						{
-							FMeshBatchElement* NextElement = new(Mesh.Elements) FMeshBatchElement();
-							*NextElement = Mesh.Elements[0];
-							NextElement->UserIndex = ParticleIndex;
-						}
-					}
-
-					MeshBatches.Add(MeshPtr);
-				}
-			}
-			FirstBatchForView.Add(MeshBatches.Num());
-		}
-	}
-}
-
 void FDynamicMeshEmitterData::GatherSimpleLights(const FParticleSystemSceneProxy* Proxy, const FSceneViewFamily& ViewFamily, FSimpleLightArray& OutParticleLights) const
 {
 	GatherParticleLightData(Source, Proxy->GetLocalToWorld(), ViewFamily, OutParticleLights);
@@ -3019,202 +2298,31 @@ void FDynamicBeam2EmitterData::GetDynamicMeshElementsEmitter(const FParticleSyst
 
 		if (AllowDebugViewmodes() && bIsWireframe && !ViewFamily.EngineShowFlags.Materials)
 		{
-			Mesh.MaterialRenderProxy	= Proxy->GetDeselectedWireframeMatInst();
-		}
-		else
-		{
-			Mesh.MaterialRenderProxy	= MaterialResource[GIsEditor && (ViewFamily.EngineShowFlags.Selection) ? bSelected : 0];
-		}
-		Mesh.Type = PT_TriangleStrip;
-
-		Mesh.bCanApplyViewModeOverrides = true;
-		Mesh.bUseWireframeSelectionColoring = Proxy->IsSelected();
-
-		Collector.AddMesh(ViewIndex, Mesh);
-
-		INC_DWORD_STAT_BY(STAT_BeamParticlesTrianglesRendered, Mesh.GetNumPrimitives());
-
-		if (Source.bRenderDirectLine == true)
-		{
-			RenderDirectLine(Proxy, Collector.GetPDI(ViewIndex), View);
-		}
-
-		if ((Source.bRenderLines == true) ||
-			(Source.bRenderTessellation == true))
-		{
-			RenderLines(Proxy, Collector.GetPDI(ViewIndex), View);
-		}
-	}
-}
-
-int32 FDynamicBeam2EmitterData::RenderEmitter(FParticleSystemSceneProxy* Proxy, FPrimitiveDrawInterface* PDI,const FSceneView* View)
-{
-	SCOPE_CYCLE_COUNTER(STAT_BeamRenderingTime);
-	INC_DWORD_STAT(STAT_BeamParticlesRenderCalls);
-
-	if (bValid == false)
-	{
-		return 0;
-	}
-
-	if ((Source.VertexCount == 0) && (Source.IndexCount == 0))
-	{
-		return 0;
-	}
-
-	const bool bIsWireframe = View->Family->EngineShowFlags.Wireframe;
-	bool bMaterialIgnored = PDI->IsMaterialIgnored(MaterialResource[bSelected], View->GetFeatureLevel());
-	if (bMaterialIgnored && !bIsWireframe)
-	{
-		return 0;
-	}
-
-	const FAsyncBufferFillData& Data = EnsureFillCompletion(View);
-
-	int32 NumDraws = 0;
-	if (Data.OutTriangleCount > 0)
-	{
-		if (!bMaterialIgnored || bIsWireframe)
-		{
-			FGlobalDynamicVertexBuffer::FAllocation* Allocation = NULL;
-			FGlobalDynamicIndexBuffer::FAllocation* IndexAllocation = NULL;
-
-			const int32 ViewIndex = View->Family->Views.Find( View );
-			if ( InstanceDataAllocations.IsValidIndex( ViewIndex ) )
-			{
-				Allocation = &InstanceDataAllocations[ ViewIndex ];
-			}	
-
-			if ( IndexDataAllocations.IsValidIndex( ViewIndex ) )
-			{
-				IndexAllocation = &IndexDataAllocations[ ViewIndex ];
-			}
-
-			// Create and set the uniform buffer for this emitter.
-			FParticleBeamTrailVertexFactory* BeamTrailVertexFactory = (FParticleBeamTrailVertexFactory*)GetVertexFactory(Proxy);
-			BeamTrailVertexFactory->SetBeamTrailUniformBuffer( CreateBeamTrailUniformBuffer( Proxy, &Source, View ) );			
-
-			if( Allocation && IndexAllocation && Allocation->IsValid() && IndexAllocation->IsValid() )
-			{
-				BeamTrailVertexFactory->SetVertexBuffer( Allocation->VertexBuffer, Allocation->VertexOffset, GetDynamicVertexStride(View->GetFeatureLevel()) );
-				BeamTrailVertexFactory->SetDynamicParameterBuffer( NULL, 0, 0 );
-
-				FMeshBatch Mesh;
-				FMeshBatchElement& BatchElement = Mesh.Elements[0];
-				BatchElement.IndexBuffer	= IndexAllocation->IndexBuffer;
-				BatchElement.FirstIndex		= IndexAllocation->FirstIndex;
-				Mesh.VertexFactory			= BeamTrailVertexFactory;
-				Mesh.DynamicVertexData		= NULL;
-				Mesh.DynamicVertexStride	= 0;
-				BatchElement.DynamicIndexData		= NULL;
-				BatchElement.DynamicIndexStride		= 0;
-				Mesh.LCI					= NULL;
-				if (Source.bUseLocalSpace == true)
-				{
-					BatchElement.PrimitiveUniformBufferResource = &Proxy->GetUniformBuffer();
-				}
-				else
-				{
-					BatchElement.PrimitiveUniformBufferResource = &Proxy->GetWorldSpacePrimitiveUniformBuffer();
-				}
-				int32 TrianglesToRender = Data.OutTriangleCount;
-				if ((TrianglesToRender % 2) != 0)
-				{
-					TrianglesToRender--;
-				}
-				BatchElement.NumPrimitives			= TrianglesToRender;
-				BatchElement.MinVertexIndex			= 0;
-				BatchElement.MaxVertexIndex			= Source.VertexCount - 1;
-				Mesh.UseDynamicData			= false;
-				Mesh.ReverseCulling			= Proxy->IsLocalToWorldDeterminantNegative();
-				Mesh.CastShadow				= Proxy->GetCastShadow();
-				Mesh.DepthPriorityGroup		= (ESceneDepthPriorityGroup)Proxy->GetDepthPriorityGroup(View);
-
-				if (AllowDebugViewmodes() && bIsWireframe && !View->Family->EngineShowFlags.Materials)
-				{
-					Mesh.MaterialRenderProxy	= Proxy->GetDeselectedWireframeMatInst();
-				}
-				else
-				{
-					Mesh.MaterialRenderProxy	= MaterialResource[GIsEditor && (View->Family->EngineShowFlags.Selection) ? bSelected : 0];
-				}
-				Mesh.Type = PT_TriangleStrip;
-
-				NumDraws += DrawRichMesh(
-					PDI,
-					Mesh,
-					FLinearColor(1.0f, 0.0f, 0.0f),
-					FLinearColor(1.0f, 1.0f, 0.0f),
-					FLinearColor(1.0f, 1.0f, 1.0f),
-					Proxy,
-					GIsEditor && (View->Family->EngineShowFlags.Selection) ? Proxy->IsSelected() : false,
-					bIsWireframe
-					);
-
-				INC_DWORD_STAT_BY(STAT_BeamParticlesTrianglesRendered, Mesh.GetNumPrimitives());
-			}
-		}
-
-		if (Source.bRenderDirectLine == true)
-		{
-			RenderDirectLine(Proxy, PDI, View);
+			Mesh.MaterialRenderProxy	= Proxy->GetDeselectedWireframeMatInst();
 		}
-
-		if ((Source.bRenderLines == true) ||
-			(Source.bRenderTessellation == true))
+		else
 		{
-			RenderLines(Proxy, PDI, View);
+			Mesh.MaterialRenderProxy	= MaterialResource[GIsEditor && (ViewFamily.EngineShowFlags.Selection) ? bSelected : 0];
 		}
-	}
-	return NumDraws;
-}
-
-/**
- *	Called during InitViews for view processing on scene proxies before rendering them
- *  Only called for primitives that are visible and have bDynamicRelevance
- *
- *	@param	Proxy			The 'owner' particle system scene proxy
- *	@param	ViewFamily		The ViewFamily to pre-render for
- *	@param	VisibilityMap	A BitArray that indicates whether the primitive was visible in that view (index)
- *	@param	FrameNumber		The frame number of this pre-render
- */
-void FDynamicBeam2EmitterData::PreRenderView(FParticleSystemSceneProxy* Proxy, const FSceneViewFamily* ViewFamily, const uint32 VisibilityMap, int32 FrameNumber)
-{
-	if (bValid == false)
-	{
-		return;
-	}
+		Mesh.Type = PT_TriangleStrip;
 
-	// Only need to do this once per-view
-	if (LastFramePreRendered < FrameNumber)
-	{
-		InstanceDataAllocations.Empty();
-		IndexDataAllocations.Empty();
-		DynamicParameterDataAllocations.Empty();
+		Mesh.bCanApplyViewModeOverrides = true;
+		Mesh.bUseWireframeSelectionColoring = Proxy->IsSelected();
 
-		for (int32 ViewIndex = 0; ViewIndex < ViewFamily->Views.Num(); ViewIndex++)
-		{
-			if (VisibilityMap & (1 << ViewIndex))
-			{
-				const FSceneView* View = ViewFamily->Views[ViewIndex];
+		Collector.AddMesh(ViewIndex, Mesh);
 
-				FGlobalDynamicVertexBuffer::FAllocation* DynamicVertexAllocation = new(InstanceDataAllocations) FGlobalDynamicVertexBuffer::FAllocation();
-				FGlobalDynamicIndexBuffer::FAllocation* DynamicIndexAllocation = new(IndexDataAllocations) FGlobalDynamicIndexBuffer::FAllocation();
-				FGlobalDynamicVertexBuffer::FAllocation* DynamicParameterAllocation = bUsesDynamicParameter ? new(DynamicParameterDataAllocations) FGlobalDynamicVertexBuffer::FAllocation() : NULL;
-				FAsyncBufferFillData* Data = new (AsyncBufferFillTasks) FAsyncBufferFillData();
+		INC_DWORD_STAT_BY(STAT_BeamParticlesTrianglesRendered, Mesh.GetNumPrimitives());
 
-				BuildViewFillData(Proxy, View, Source.VertexCount, sizeof(FParticleBeamTrailVertex), 0, *DynamicVertexAllocation, *DynamicIndexAllocation, DynamicParameterAllocation, *Data);
-				DoBufferFill(*Data);
-			}
+		if (Source.bRenderDirectLine == true)
+		{
+			RenderDirectLine(Proxy, Collector.GetPDI(ViewIndex), View);
 		}
 
-		// Set the frame tracker
-		LastFramePreRendered = FrameNumber;
-	}
-
-	if (Source.bUseLocalSpace == false)
-	{
-		Proxy->UpdateWorldSpacePrimitiveUniformBuffer();
+		if ((Source.bRenderLines == true) ||
+			(Source.bRenderTessellation == true))
+		{
+			RenderLines(Proxy, Collector.GetPDI(ViewIndex), View);
+		}
 	}
 }
 
@@ -5935,200 +5043,6 @@ void FDynamicTrailsEmitterData::GetDynamicMeshElementsEmitter(const FParticleSys
 	INC_DWORD_STAT_BY(STAT_TrailParticlesTrianglesRendered, RenderedPrimitiveCount);
 }
 
-// Render thread only draw call
-int32 FDynamicTrailsEmitterData::RenderEmitter(FParticleSystemSceneProxy* Proxy, FPrimitiveDrawInterface* PDI, const FSceneView* View)
-{
-	SCOPE_CYCLE_COUNTER(STAT_TrailRenderingTime);
-	INC_DWORD_STAT(STAT_TrailParticlesRenderCalls);
-
-	if (bValid == false)
-	{
-		return 0;
-	}
-
-	check(PDI);
-	if ((SourcePointer->VertexCount <= 0) || (SourcePointer->ActiveParticleCount <= 0) || (SourcePointer->IndexCount < 3))
-	{
-		return 0;
-	}
-
-	bool const bIsWireframe = View->Family->EngineShowFlags.Wireframe;
-
-	// Don't render if the material will be ignored
-	if (PDI->IsMaterialIgnored(MaterialResource[bSelected], View->GetFeatureLevel()) && !bIsWireframe)
-	{
-		return 0;
-	}
-
-	int32 RenderedPrimitiveCount = 0;
-
-	const FAsyncBufferFillData& Data = EnsureFillCompletion(View);
-
-	if (Data.OutTriangleCount == 0)
-	{
-		return 0;
-	}
-
-	int32 NumDraws = 0;
-	if (bRenderGeometry == true)
-	{
-		FGlobalDynamicVertexBuffer::FAllocation* Allocation = NULL;
-		FGlobalDynamicIndexBuffer::FAllocation* IndexAllocation = NULL;
-		FGlobalDynamicVertexBuffer::FAllocation* DynamicParameterAllocation = NULL;
-
-		const int32 ViewIndex = View->Family->Views.Find( View );
-		if ( InstanceDataAllocations.IsValidIndex( ViewIndex ) )
-		{
-			Allocation = &InstanceDataAllocations[ ViewIndex ];
-		}	
-		
-		if ( IndexDataAllocations.IsValidIndex( ViewIndex ) )
-		{
-			IndexAllocation = &IndexDataAllocations[ ViewIndex ];
-		}
-
-		if ( DynamicParameterDataAllocations.IsValidIndex( ViewIndex ) )
-		{
-			DynamicParameterAllocation = &DynamicParameterDataAllocations[ ViewIndex ];
-		}
-
-		if ( Allocation && IndexAllocation && Allocation->IsValid() && IndexAllocation->IsValid() && (!bUsesDynamicParameter || (DynamicParameterAllocation && DynamicParameterAllocation->IsValid())) )
-		{
-			// Create and set the uniform buffer for this emitter.
-			FParticleBeamTrailVertexFactory* BeamTrailVertexFactory = (FParticleBeamTrailVertexFactory*)GetVertexFactory(Proxy);
-			BeamTrailVertexFactory->SetBeamTrailUniformBuffer( CreateBeamTrailUniformBuffer( Proxy, SourcePointer, View ) );
-			BeamTrailVertexFactory->SetVertexBuffer( Allocation->VertexBuffer, Allocation->VertexOffset, GetDynamicVertexStride(View->GetFeatureLevel()) );
-			BeamTrailVertexFactory->SetDynamicParameterBuffer( DynamicParameterAllocation ? DynamicParameterAllocation->VertexBuffer : NULL, DynamicParameterAllocation ? DynamicParameterAllocation->VertexOffset : 0, GetDynamicParameterVertexStride() );
-
-			FMeshBatch Mesh;
-			FMeshBatchElement& BatchElement = Mesh.Elements[0];
-			BatchElement.IndexBuffer	= IndexAllocation->IndexBuffer;
-			BatchElement.FirstIndex		= IndexAllocation->FirstIndex;
-			Mesh.VertexFactory			= BeamTrailVertexFactory;
-			Mesh.LCI					= NULL;
-
-			BatchElement.PrimitiveUniformBufferResource = &Proxy->GetWorldSpacePrimitiveUniformBuffer();
-			BatchElement.NumPrimitives			= Data.OutTriangleCount;
-			BatchElement.MinVertexIndex			= 0;
-			BatchElement.MaxVertexIndex			= SourcePointer->VertexCount - 1;
-			Mesh.UseDynamicData			= false;
-			Mesh.ReverseCulling			= Proxy->IsLocalToWorldDeterminantNegative();
-			Mesh.CastShadow				= Proxy->GetCastShadow();
-			Mesh.DepthPriorityGroup		= (ESceneDepthPriorityGroup)Proxy->GetDepthPriorityGroup(View);
-
-			if (AllowDebugViewmodes() && bIsWireframe && !View->Family->EngineShowFlags.Materials)
-			{
-				Mesh.MaterialRenderProxy = Proxy->GetDeselectedWireframeMatInst();
-			}
-			else
-			{
-	#if !(UE_BUILD_SHIPPING || UE_BUILD_TEST)
-				if (Data.OutTriangleCount != SourcePointer->PrimitiveCount)
-				{
-					UE_LOG(LogParticles, Log, TEXT("Data.OutTriangleCount = %4d vs. SourcePrimCount = %4d"), Data.OutTriangleCount, SourcePointer->PrimitiveCount);
-
-					int32 CheckTrailCount = 0;
-					int32 CheckTriangleCount = 0;
-					for (int32 ParticleIdx = 0; ParticleIdx < SourcePointer->ActiveParticleCount; ParticleIdx++)
-					{
-						int32 CurrentIndex = SourcePointer->ParticleIndices[ParticleIdx];
-						DECLARE_PARTICLE_PTR(CheckParticle, SourcePointer->ParticleData.GetData() + SourcePointer->ParticleStride * CurrentIndex);
-						FTrailsBaseTypeDataPayload* TrailPayload = (FTrailsBaseTypeDataPayload*)((uint8*)CheckParticle + SourcePointer->TrailDataOffset);
-						if (TRAIL_EMITTER_IS_HEAD(TrailPayload->Flags) == false)
-						{
-							continue;
-						}
-
-						UE_LOG(LogParticles, Log, TEXT("Trail %2d has %5d triangles"), TrailPayload->TrailIndex, TrailPayload->TriangleCount);
-						CheckTriangleCount += TrailPayload->TriangleCount;
-						CheckTrailCount++;
-					}
-					UE_LOG(LogParticles, Log, TEXT("Total 'live' trail count = %d"), CheckTrailCount);
-					UE_LOG(LogParticles, Log, TEXT("\t%5d triangles total (not counting degens)"), CheckTriangleCount);
-				}
-	#endif
-				checkf(Data.OutTriangleCount <= SourcePointer->PrimitiveCount, TEXT("Data.OutTriangleCount = %4d vs. SourcePrimCount = %4d"), Data.OutTriangleCount, SourcePointer->PrimitiveCount);
-				Mesh.MaterialRenderProxy = MaterialResource[GIsEditor && (View->Family->EngineShowFlags.Selection) ? bSelected : 0];
-			}
-			Mesh.Type = PT_TriangleStrip;
-
-			NumDraws += DrawRichMesh(
-				PDI,
-				Mesh,
-				FLinearColor(1.0f, 0.0f, 0.0f),
-				FLinearColor(1.0f, 1.0f, 0.0f),
-				FLinearColor(1.0f, 1.0f, 1.0f),
-				Proxy,
-				GIsEditor && (View->Family->EngineShowFlags.Selection) ? Proxy->IsSelected() : false,
-				bIsWireframe
-				);
-
-			RenderedPrimitiveCount = Mesh.GetNumPrimitives();
-		}
-	}
-
-	RenderDebug(Proxy, PDI, View, false);
-
-	INC_DWORD_STAT_BY(STAT_TrailParticlesTrianglesRendered, RenderedPrimitiveCount);
-	return NumDraws;
-}
-
-/**
- *	Called during InitViews for view processing on scene proxies before rendering them
- *  Only called for primitives that are visible and have bDynamicRelevance
- *
- *	@param	Proxy			The 'owner' particle system scene proxy
- *	@param	ViewFamily		The ViewFamily to pre-render for
- *	@param	VisibilityMap	A BitArray that indicates whether the primitive was visible in that view (index)
- *	@param	FrameNumber		The frame number of this pre-render
- */
-void FDynamicTrailsEmitterData::PreRenderView(FParticleSystemSceneProxy* Proxy, const FSceneViewFamily* ViewFamily, const uint32 VisibilityMap, int32 FrameNumber)
-{
-	if (bValid == false)
-	{
-		return;
-	}
-
-	// Only need to do this once per-view
-	if (LastFramePreRendered < FrameNumber)
-	{
-		int32 VertexStride = GetDynamicVertexStride(ViewFamily->Scene->GetFeatureLevel());
-		int32 DynamicParameterVertexStride = 0;
-		if (bUsesDynamicParameter == true)
-		{
-			DynamicParameterVertexStride = GetDynamicParameterVertexStride();
-		}
-
-		InstanceDataAllocations.Empty();
-		IndexDataAllocations.Empty();
-		DynamicParameterDataAllocations.Empty();
-
-		for (int32 ViewIndex = 0; ViewIndex < ViewFamily->Views.Num(); ViewIndex++)
-		{
-			if (VisibilityMap & (1 << ViewIndex))
-			{
-				const FSceneView* View = ViewFamily->Views[ViewIndex];
-
-				FGlobalDynamicVertexBuffer::FAllocation* DynamicVertexAllocation = new(InstanceDataAllocations) FGlobalDynamicVertexBuffer::FAllocation();
-				FGlobalDynamicIndexBuffer::FAllocation* DynamicIndexAllocation = new(IndexDataAllocations) FGlobalDynamicIndexBuffer::FAllocation();
-				FGlobalDynamicVertexBuffer::FAllocation* DynamicParameterAllocation = bUsesDynamicParameter ? new(DynamicParameterDataAllocations) FGlobalDynamicVertexBuffer::FAllocation() : NULL;
-				FAsyncBufferFillData* Data = new (AsyncBufferFillTasks) FAsyncBufferFillData();
-
-				BuildViewFillData(Proxy, View, SourcePointer->VertexCount, VertexStride, DynamicParameterVertexStride, *DynamicVertexAllocation, *DynamicIndexAllocation, DynamicParameterAllocation, *Data);
-				DoBufferFill(*Data);
-			}
-		}
-
-		// Set the frame tracker
-		LastFramePreRendered = FrameNumber;
-	}
-
-	if (SourcePointer->bUseLocalSpace == false)
-	{
-		Proxy->UpdateWorldSpacePrimitiveUniformBuffer();
-	}
-}
-
 void FDynamicTrailsEmitterData::RenderDebug(const FParticleSystemSceneProxy* Proxy, FPrimitiveDrawInterface* PDI, const FSceneView* View, bool bCrosses) const
 {
 	// Can't do anything in here...
@@ -6239,11 +5153,6 @@ void FDynamicRibbonEmitterData::Init(bool bInSelected)
 	bUsesDynamicParameter = GetSourceData()->DynamicParameterDataOffset > 0;
 }
 
-bool FDynamicRibbonEmitterData::ShouldUsePrerenderView()
-{
-	return (RenderAxisOption != Trails_CameraUp);
-}
-
 void FDynamicRibbonEmitterData::RenderDebug(const FParticleSystemSceneProxy* Proxy, FPrimitiveDrawInterface* PDI, const FSceneView* View, bool bCrosses) const
 {
 	if ((bRenderParticles == true) || (bRenderTangents == true))
@@ -7366,9 +6275,6 @@ FParticleSystemSceneProxy::FParticleSystemSceneProxy(const UParticleSystemCompon
 	LevelColor = FLinearColor(1.0f, 1.0f, 0.0f);
 	PropertyColor = FLinearColor(1.0f, 1.0f, 1.0f);
 
-#if STATS
-	LastStatCaptureTime = FApp::GetCurrentTime();
-#endif
 	LODMethod = Component->LODMethod;
 }
 
@@ -7418,7 +6324,7 @@ void FParticleSystemSceneProxy::GetDynamicMeshElements(const TArray<const FScene
 		SCOPE_CYCLE_COUNTER(STAT_ParticleRenderingTime);
 		FScopeCycleCounter Context(GetStatId());
 
-		const double StartTime = GTrackParticleRenderingStats || GTrackParticleRenderingStatsForOneFrame ? FPlatformTime::Seconds() : 0;
+		const double StartTime = GTrackParticleRenderingStats ? FPlatformTime::Seconds() : 0;
 		int32 NumDraws = 0;
 
 		if (DynamicData != NULL)
@@ -7464,169 +6370,6 @@ void FParticleSystemSceneProxy::GetDynamicMeshElements(const TArray<const FScene
 	}
 }
 
-/** 
-* Draw the scene proxy as a dynamic element
-*
-* @param	PDI - draw interface to render to
-* @param	View - current view
-*/
-void FParticleSystemSceneProxy::DrawDynamicElements(FPrimitiveDrawInterface* PDI,const FSceneView* View)
-{
-	QUICK_SCOPE_CYCLE_COUNTER( STAT_ParticleSystemSceneProxy_DrawDynamicElements );
-
-	if (View->Family->EngineShowFlags.Particles)
-	{
-		SCOPE_CYCLE_COUNTER(STAT_ParticleRenderingTime);
-		FScopeCycleCounter Context(GetStatId());
-
-		const double StartTime = GTrackParticleRenderingStats || GTrackParticleRenderingStatsForOneFrame ? FPlatformTime::Seconds() : 0;
-
-		int32 NumDraws = 0;
-		int32 NumEmitters = 0;
-
-			if (DynamicData != NULL)
-			{
-				for (int32 Index = 0; Index < DynamicData->DynamicEmitterDataArray.Num(); Index++)
-				{
-					FDynamicEmitterDataBase* Data =	DynamicData->DynamicEmitterDataArray[Index];
-					if ((Data == NULL) || (Data->bValid != true))
-					{
-						continue;
-					}
-					//hold on to the emitter index in case we need to access any of its properties
-					DynamicData->EmitterIndex = Index;
-
-					const int32 DrawCalls = Data->RenderEmitter(this, PDI, View);
-					NumDraws += DrawCalls;
-					if (DrawCalls > 0)
-					{
-						NumEmitters++;
-					}
-				}
-			}
-
-		INC_DWORD_STAT_BY(STAT_ParticleDrawCalls, NumDraws);
-
-		if (View->Family->EngineShowFlags.Particles)
-		{
-			RenderBounds(PDI, View->Family->EngineShowFlags, GetBounds(), IsSelected());
-			if (HasCustomOcclusionBounds())
-			{
-				RenderBounds(PDI, View->Family->EngineShowFlags, GetCustomOcclusionBounds(), IsSelected());
-			}
-		}
-
-#if STATS
-
-		// Capture one frame rendering stats if enabled and at least one draw call was submitted
-		if (GTrackParticleRenderingStatsForOneFrame
-			&& NumDraws > 0)
-		{
-			const double EndTime = FPlatformTime::Seconds();
-
-			const float DeltaTime = EndTime - StartTime;
-			const FString TemplateName = GetResourceName().ToString();
-			FParticleTemplateRenderStats* TemplateStats = GOneFrameTemplateRenderStats.Find(TemplateName);
-
-			if (TemplateStats)
-			{
-				// Update the existing record for this UParticleSystem
-				TemplateStats->NumDraws += NumDraws;
-				TemplateStats->MaxRenderTime += DeltaTime;
-				TemplateStats->NumEmitters += NumEmitters;
-				TemplateStats->NumDrawDynamicElements++;
-				if (!bCountedThisFrame)
-				{
-					TemplateStats->NumComponents++;
-				}
-			}
-			else
-			{
-				// Create a new record for this UParticleSystem
-				FParticleTemplateRenderStats NewStats;
-				NewStats.NumDraws = NumDraws;
-				NewStats.MaxRenderTime = DeltaTime;
-				NewStats.NumEmitters = NumEmitters;
-				NewStats.NumComponents = 1;
-				NewStats.NumDrawDynamicElements = 1;
-				GOneFrameTemplateRenderStats.Add(TemplateName, NewStats);
-			}
-			bCountedThisFrame = true;
-		}
-
-		if (!GTrackParticleRenderingStatsForOneFrame)
-		{
-			// Mark the proxy has not having been processed, this will be used on the next frame dump to count components rendered accurately
-			bCountedThisFrame = false;
-		}
-
-		// Capture render stats if enabled, and at least one draw call was submitted, and enough time has elapsed since the last capture
-		// This misses spikes but allows efficient stat gathering over large periods of time
-		if (GTrackParticleRenderingStats 
-			// The main purpose of particle rendering stats are to optimize splitscreen, so only capture when we are in medium detail mode
-			// This is needed to prevent capturing particle rendering time during cinematics where the engine switches back to high detail mode
-			&& GetCachedScalabilityCVars().DetailMode != DM_High
-			&& NumDraws > 0
-			&& FApp::GetCurrentTime() - LastStatCaptureTime > GTimeBetweenParticleRenderStatCaptures)
-		{
-			LastStatCaptureTime = FApp::GetCurrentTime();
-
-			const double EndTime = FPlatformTime::Seconds();
-
-			const float DeltaTime = EndTime - StartTime;
-			// Only capture stats if the rendering time was large enough to be considered
-			if (DeltaTime > GMinParticleDrawTimeToTrack)
-			{
-				const FString TemplateName = GetResourceName().ToString();
-				FParticleTemplateRenderStats* TemplateStats = GTemplateRenderStats.Find(TemplateName);
-				if (TemplateStats)
-				{
-					// Update the existing record for this UParticleSystem if the new stat had a longer render time
-					if (DeltaTime > TemplateStats->MaxRenderTime)
-					{
-						TemplateStats->NumDraws = NumDraws;
-						TemplateStats->MaxRenderTime = DeltaTime;
-						TemplateStats->NumEmitters = NumEmitters;
-					}
-				}
-				else
-				{
-					// Create a new record for this UParticleSystem
-					FParticleTemplateRenderStats NewStats;
-					NewStats.NumDraws = NumDraws;
-					NewStats.MaxRenderTime = DeltaTime;
-					NewStats.NumEmitters = NumEmitters;
-					GTemplateRenderStats.Add(TemplateName, NewStats);
-				}
-
-				const FString ComponentName = GetOwnerName().ToString();
-				if (ComponentName.Contains(TEXT("Emitter_")) )
-				{
-					FParticleComponentRenderStats* ComponentStats = GComponentRenderStats.Find(ComponentName);
-					if (ComponentStats)
-					{
-						// Update the existing record for this component if the new stat had a longer render time
-						if (DeltaTime > ComponentStats->MaxRenderTime)
-						{
-							ComponentStats->NumDraws = NumDraws;
-							ComponentStats->MaxRenderTime = DeltaTime;
-						}
-					}
-					else
-					{
-						// Create a new record for this component
-						FParticleComponentRenderStats NewStats;
-						NewStats.NumDraws = NumDraws;
-						NewStats.MaxRenderTime = DeltaTime;
-						GComponentRenderStats.Add(ComponentName, NewStats);
-					}
-				}
-			}
-		}
-#endif
-	}
-}
-
 void FParticleSystemSceneProxy::CreateRenderThreadResources()
 {
 	CreateRenderThreadResourcesForEmitterData();
@@ -7768,7 +6511,6 @@ FPrimitiveViewRelevance FParticleSystemSceneProxy::GetViewRelevance(const FScene
 	Result.bShadowRelevance = IsShadowCast(View);
 	Result.bRenderInMainPass = ShouldRenderInMainPass();
 	Result.bDynamicRelevance = true;
-	Result.bNeedsPreRenderView = Result.bDrawRelevance || Result.bShadowRelevance;
 	Result.bHasSimpleLights = true;
 	if (!View->Family->EngineShowFlags.Wireframe && View->Family->EngineShowFlags.Materials)
 	{
@@ -7819,52 +6561,6 @@ void FParticleSystemSceneProxy::UpdateWorldSpacePrimitiveUniformBuffer() const
 	}
 }
 
-/**
- *	Helper function for calculating the tessellation for a given view.
- *
- *	@param	View		The view of interest.
- *	@param	FrameNumber		The frame number being rendered.
- */
-void FParticleSystemSceneProxy::ProcessPreRenderView(const FSceneView* View, int32 FrameNumber)
-{
-	const FSceneView* LocalView = View;
-
-	if ((GIsEditor == true) || (GbEnableGameThreadLODCalculation == false))
-	{
-		DetermineLODDistance(LocalView, FrameNumber);
-	}
-}
-
-/**
- *	Called during InitViews for view processing on scene proxies before rendering them
- *  Only called for primitives that are visible and have bDynamicRelevance
- *
- *	@param	ViewFamily		The ViewFamily to pre-render for
- *	@param	VisibilityMap	A BitArray that indicates whether the primitive was visible in that view (index)
- *	@param	FrameNumber		The frame number of this pre-render
- */
-void FParticleSystemSceneProxy::PreRenderView(const FSceneViewFamily* ViewFamily, const uint32 VisibilityMap, int32 FrameNumber)
-{
-	FirstFreeMeshBatch = 0;
-	for (int32 ViewIndex = 0; ViewIndex < ViewFamily->Views.Num(); ViewIndex++)
-	{
-		ProcessPreRenderView(ViewFamily->Views[ViewIndex], FrameNumber);
-	}
-
-	if (DynamicData != NULL)
-	{
-		for (int32 EmitterIndex = 0; EmitterIndex < DynamicData->DynamicEmitterDataArray.Num(); EmitterIndex++)
-		{
-			FDynamicEmitterDataBase* DynamicEmitterData = DynamicData->DynamicEmitterDataArray[EmitterIndex];
-			if (DynamicEmitterData)
-			{
-				DynamicData->EmitterIndex = EmitterIndex;
-				DynamicEmitterData->PreRenderView(this, ViewFamily, VisibilityMap, FrameNumber);
-			}
-		}
-	}
-}
-
 void FParticleSystemSceneProxy::GatherSimpleLights(const FSceneViewFamily& ViewFamily, FSimpleLightArray& OutParticleLights) const
 {
 	if (DynamicData != NULL)
@@ -7899,41 +6595,6 @@ FParticleSystemOcclusionSceneProxy::~FParticleSystemOcclusionSceneProxy()
 {
 }
 
-/** 
- * Draw the scene proxy as a dynamic element
- *
- * @param	PDI - draw interface to render to
- * @param	View - current view
- */
-void FParticleSystemOcclusionSceneProxy::DrawDynamicElements(FPrimitiveDrawInterface* PDI,const FSceneView* View)
-{
-	QUICK_SCOPE_CYCLE_COUNTER( STAT_ParticleSystemOcclusionSceneProxy_DrawDynamicElements );
-
-	if (View->Family->EngineShowFlags.Particles)
-	{
-		if (DynamicData != NULL)
-		{
-			FBoxSphereBounds TempOcclusionBounds = OcclusionBounds;
-			if (bHasCustomOcclusionBounds == true)
-			{
-				OcclusionBounds = OcclusionBounds.TransformBy(GetLocalToWorld());
-			}
-
-			// Update the occlusion bounds for next frame...
-			if (bHasCustomOcclusionBounds == false)
-			{
-				OcclusionBounds = GetBounds();
-			}
-			else
-			{
-				OcclusionBounds = TempOcclusionBounds;
-			}
-
-			FParticleSystemSceneProxy::DrawDynamicElements(PDI, View);
-			}
-		}
-	}
-
 FPrimitiveSceneProxy* UParticleSystemComponent::CreateSceneProxy()
 {
 	FParticleSystemSceneProxy* NewProxy = NULL;
diff --git a/Engine/Source/Runtime/Engine/Private/PhysicsEngine/PhysDrawing.cpp b/Engine/Source/Runtime/Engine/Private/PhysicsEngine/PhysDrawing.cpp
--- a/Engine/Source/Runtime/Engine/Private/PhysicsEngine/PhysDrawing.cpp
+++ b/Engine/Source/Runtime/Engine/Private/PhysicsEngine/PhysDrawing.cpp
@@ -529,121 +529,6 @@ void FConvexCollisionVertexFactory::InitConvexVertexFactory(const FConvexCollisi
 	}
 }
 
-void FKAggregateGeom::DrawAggGeom(FPrimitiveDrawInterface* PDI, const FTransform& Transform, const FColor Color, const FMaterialRenderProxy* MatInst, bool bPerHullColor, bool bDrawSolid, bool bUseEditorDepthTest)
-{
-	const FVector Scale3D = Transform.GetScale3D();
-	FTransform ParentTM = Transform;
-	ParentTM.RemoveScaling();
-
-	if( Scale3D.GetAbs().IsUniform() )
-	{
-		for(int32 i=0; i<SphereElems.Num(); i++)
-		{
-			FTransform ElemTM = SphereElems[i].GetTransform();
-			ElemTM.ScaleTranslation(Scale3D);
-			ElemTM *= ParentTM;
-
-			if(bDrawSolid)
-				SphereElems[i].DrawElemSolid(PDI, ElemTM, Scale3D.X, MatInst);
-			else
-				SphereElems[i].DrawElemWire(PDI, ElemTM, Scale3D.X, Color);
-		}
-
-		for(int32 i=0; i<BoxElems.Num(); i++)
-		{
-			FTransform ElemTM = BoxElems[i].GetTransform();
-			ElemTM.ScaleTranslation(Scale3D);
-			ElemTM *= ParentTM;
-
-			if(bDrawSolid)
-				BoxElems[i].DrawElemSolid(PDI, ElemTM, Scale3D.X, MatInst);
-			else
-				BoxElems[i].DrawElemWire(PDI, ElemTM, Scale3D.X, Color);
-		}
-
-		for(int32 i=0; i<SphylElems.Num(); i++)
-		{
-			FTransform ElemTM = SphylElems[i].GetTransform();
-			ElemTM.ScaleTranslation(Scale3D);
-			ElemTM *= ParentTM;
-
-			if(bDrawSolid)
-				SphylElems[i].DrawElemSolid(PDI, ElemTM, Scale3D.X, MatInst);
-			else
-				SphylElems[i].DrawElemWire(PDI, ElemTM, Scale3D.X, Color);
-		}
-	}
-
-	if(ConvexElems.Num() > 0)
-	{
-		if(bDrawSolid)
-		{
-			// Cache collision vertex/index buffer
-			if(!RenderInfo)
-			{
-				RenderInfo = new FKConvexGeomRenderInfo();
-				RenderInfo->VertexBuffer = new FConvexCollisionVertexBuffer();
-				RenderInfo->IndexBuffer = new FConvexCollisionIndexBuffer();
-
-				for(int32 i=0; i<ConvexElems.Num(); i++)
-				{
-					// Get vertices/triangles from this hull.
-					ConvexElems[i].AddCachedSolidConvexGeom(RenderInfo->VertexBuffer->Vertices, RenderInfo->IndexBuffer->Indices, FColor(255,255,255));
-				}
-
-				// Only continue if we actuall got some valid geometry
-				// Will crash if we try to init buffers with no data
-				if(RenderInfo->HasValidGeometry())
-				{
-					RenderInfo->VertexBuffer->InitResource();
-					RenderInfo->IndexBuffer->InitResource();
-
-					RenderInfo->CollisionVertexFactory = new FConvexCollisionVertexFactory(RenderInfo->VertexBuffer);
-					RenderInfo->CollisionVertexFactory->InitResource();
-				}
-			}
-
-			// If we have geometry to draw, do so
-			if(RenderInfo->HasValidGeometry())
-			{
-				// Calculate transform
-				FTransform LocalToWorld = FTransform( FQuat::Identity, FVector::ZeroVector, Scale3D ) * ParentTM;
-
-				// Draw the mesh.
-				FMeshBatch Mesh;
-				FMeshBatchElement& BatchElement = Mesh.Elements[0];
-				BatchElement.IndexBuffer = RenderInfo->IndexBuffer;
-				Mesh.VertexFactory = RenderInfo->CollisionVertexFactory;
-				Mesh.MaterialRenderProxy = MatInst;
-				FBoxSphereBounds WorldBounds, LocalBounds;
-				CalcBoxSphereBounds(WorldBounds, LocalToWorld);
-				CalcBoxSphereBounds(LocalBounds, FTransform::Identity);
-				BatchElement.PrimitiveUniformBuffer = CreatePrimitiveUniformBufferImmediate(LocalToWorld.ToMatrixWithScale(), WorldBounds, LocalBounds, true, bUseEditorDepthTest);
-				// previous l2w not used so treat as static
-				BatchElement.FirstIndex = 0;
-				BatchElement.NumPrimitives = RenderInfo->IndexBuffer->Indices.Num() / 3;
-				BatchElement.MinVertexIndex = 0;
-				BatchElement.MaxVertexIndex = RenderInfo->VertexBuffer->Vertices.Num() - 1;
-				Mesh.ReverseCulling = LocalToWorld.GetDeterminant() < 0.0f ? true : false;
-				Mesh.Type = PT_TriangleList;
-				Mesh.DepthPriorityGroup = SDPG_World;
-				PDI->DrawMesh(Mesh);
-			}
-		}
-		else
-		{
-			for(int32 i=0; i<ConvexElems.Num(); i++)
-			{
-				FColor ConvexColor = bPerHullColor ? DebugUtilColor[i%NUM_DEBUG_UTIL_COLORS] : Color;
-				FTransform ElemTM = ConvexElems[i].GetTransform();
-				ElemTM *= Transform;
-				ConvexElems[i].DrawElemWire(PDI, ElemTM, ConvexColor);
-			}
-		}
-	}
-}
-
-
 void FKAggregateGeom::GetAggGeom(const FTransform& Transform, const FColor Color, const FMaterialRenderProxy* MatInst, bool bPerHullColor, bool bDrawSolid, bool bUseEditorDepthTest, int32 ViewIndex, FMeshElementCollector& Collector) const
 {
 	const FVector Scale3D = Transform.GetScale3D();
@@ -816,7 +701,7 @@ FTransform GetSkelBoneTransform(int32 BoneIndex, const TArray<FTransform>& Space
 	}
 }
 
-void UPhysicsAsset::DrawCollision(FPrimitiveDrawInterface* PDI, const USkeletalMesh* SkelMesh, const TArray<FTransform>& SpaceBases, const FTransform& LocalToWorld, float Scale)
+void UPhysicsAsset::GetCollisionMesh(int32 ViewIndex, FMeshElementCollector& Collector, const USkeletalMesh* SkelMesh, const TArray<FTransform>& SpaceBases, const FTransform& LocalToWorld, float Scale)
 {
 	for( int32 i=0; i<BodySetup.Num(); i++)
 	{
@@ -827,7 +712,7 @@ void UPhysicsAsset::DrawCollision(FPrimitiveDrawInterface* PDI, const USkeletalM
 		FTransform BoneTransform = GetSkelBoneTransform(BoneIndex, SpaceBases, LocalToWorld);
 		BoneTransform.SetScale3D(FVector(Scale));
 		BodySetup[i]->CreatePhysicsMeshes();
-		BodySetup[i]->AggGeom.DrawAggGeom( PDI, BoneTransform, *BoneColor, NULL, false, false, false );
+		BodySetup[i]->AggGeom.GetAggGeom(BoneTransform, *BoneColor, NULL, false, false, false, ViewIndex, Collector);
 	}
 }
 
