Removed DrawDynamicElements, PreRenderView and dependencies.  These have been disabled since 11/04, and the GetDynamicMeshElements path is now the only path for dynamic mesh rendering.

diff --git a/Engine/Source/Runtime/Engine/Private/PhysicsEngine2D/AggregateGeometry2D.cpp b/Engine/Source/Runtime/Engine/Private/PhysicsEngine2D/AggregateGeometry2D.cpp
--- a/Engine/Source/Runtime/Engine/Private/PhysicsEngine2D/AggregateGeometry2D.cpp
+++ b/Engine/Source/Runtime/Engine/Private/PhysicsEngine2D/AggregateGeometry2D.cpp
@@ -118,11 +118,6 @@ void FAggregateGeometry2D::EmptyElements()
 	ConvexElements.Empty();
 }
 
-void FAggregateGeometry2D::DrawAggGeom(FPrimitiveDrawInterface* PDI, const FTransform& Transform, const FColor Color, const FMaterialRenderProxy* MatInst, bool bPerHullColor, bool bDrawSolid, bool bUseEditorDepthTest) const
-{
-	//@TODO: PAPER2D: Implement this method
-}
-
 float FAggregateGeometry2D::GetArea(const FVector& Scale) const
 {
 	float SumArea = 0.0f;
diff --git a/Engine/Source/Runtime/Engine/Private/PrimitiveDrawingUtils.cpp b/Engine/Source/Runtime/Engine/Private/PrimitiveDrawingUtils.cpp
--- a/Engine/Source/Runtime/Engine/Private/PrimitiveDrawingUtils.cpp
+++ b/Engine/Source/Runtime/Engine/Private/PrimitiveDrawingUtils.cpp
@@ -1661,266 +1661,6 @@ void ApplyViewModeOverrides(
 #endif
 }
 
-/**
- * Draws a mesh, modifying the material which is used depending on the view's show flags.
- * Meshes with materials irrelevant to the pass which the mesh is being drawn for may be entirely ignored.
- *
- * @param PDI - The primitive draw interface to draw the mesh on.
- * @param Mesh - The mesh to draw.
- * @param WireframeColor - The color which is used when rendering the mesh with EngineShowFlags.Wireframe.
- * @param LevelColor - The color which is used when rendering the mesh with EngineShowFlags.LevelColoration.
- * @param PropertyColor - The color to use when rendering the mesh with EngineShowFlags.PropertyColoration.
- * @param PrimitiveInfo - The FScene information about the UPrimitiveComponent.
- * @param bSelected - True if the primitive is selected.
- * @param ExtraDrawFlags - optional flags to override the view family show flags when rendering
- * @return Number of passes rendered for the mesh
- */
-int32 DrawRichMesh(
-	FPrimitiveDrawInterface* PDI,
-	const FMeshBatch& Mesh,
-	const FLinearColor& WireframeColor,
-	const FLinearColor& LevelColor,
-	const FLinearColor& PropertyColor,
-	const FPrimitiveSceneProxy* PrimitiveSceneProxy,
-	bool bSelected,
-	bool bDrawInWireframe
-	)
-{
-	int32 PassesRendered=0;
-
-#if (UE_BUILD_SHIPPING || UE_BUILD_TEST)
-	// Draw the mesh unmodified.
-	PassesRendered = PDI->DrawMesh( Mesh );
-#else
-
-	// If debug viewmodes are not allowed, skip all of the debug viewmode handling
-	if (!AllowDebugViewmodes())
-	{
-		return PDI->DrawMesh( Mesh );
-	}
-
-	const auto FeatureLevel = PDI->View->GetFeatureLevel();
-	const FEngineShowFlags& EngineShowFlags = PDI->View->Family->EngineShowFlags;
-
-	if(bDrawInWireframe)
-	{
-		// In wireframe mode, draw the edges of the mesh with the specified wireframe color, or
-		// with the level or property color if level or property coloration is enabled.
-		FLinearColor BaseColor( WireframeColor );
-		if (EngineShowFlags.PropertyColoration)
-		{
-			BaseColor = PropertyColor;
-		}
-		else if (EngineShowFlags.LevelColoration)
-		{
-			BaseColor = LevelColor;
-		}
-
-		if (Mesh.MaterialRenderProxy->GetMaterial(FeatureLevel)->MaterialModifiesMeshPosition_RenderThread())
-		{
-			// If the material is mesh-modifying, we cannot rely on substitution
-			const FOverrideSelectionColorMaterialRenderProxy WireframeMaterialInstance(
-				Mesh.MaterialRenderProxy,
-				GetSelectionColor( BaseColor, bSelected, Mesh.MaterialRenderProxy->IsHovered(), /*bUseOverlayIntensity=*/false)
-				);
-
-			FMeshBatch ModifiedMesh = Mesh;
-			ModifiedMesh.bWireframe = true;
-			ModifiedMesh.MaterialRenderProxy = &WireframeMaterialInstance;
-			PassesRendered = PDI->DrawMesh(ModifiedMesh);
-		}
-		else
-		{
-			const FColoredMaterialRenderProxy WireframeMaterialInstance(
-				GEngine->WireframeMaterial->GetRenderProxy(Mesh.MaterialRenderProxy->IsSelected(), Mesh.MaterialRenderProxy->IsHovered()),
-				GetSelectionColor( BaseColor, bSelected, Mesh.MaterialRenderProxy->IsHovered(), /*bUseOverlayIntensity=*/false)
-				);
-			FMeshBatch ModifiedMesh = Mesh;
-			ModifiedMesh.bWireframe = true;
-			ModifiedMesh.MaterialRenderProxy = &WireframeMaterialInstance;
-			PassesRendered = PDI->DrawMesh(ModifiedMesh);
-		}
-	}
-	else if(PDI->View->Family->EngineShowFlags.LightComplexity)
-	{
-		// Don't render unlit translucency when in 'light complexity' viewmode.
-		if (!Mesh.IsTranslucent(FeatureLevel) || Mesh.MaterialRenderProxy->GetMaterial(FeatureLevel)->GetShadingModel() != MSM_Unlit)
-		{
-			// Count the number of lights interacting with this primitive.
-			int32 NumDynamicLights = GetRendererModule().GetNumDynamicLightsAffectingPrimitive(PrimitiveSceneProxy->GetPrimitiveSceneInfo(),Mesh.LCI);
-
-			// Get a colored material to represent the number of lights.
-			// Some component types (BSP) have multiple FLightCacheInterface's per component, so make sure the whole component represents the number of dominant lights affecting
-			NumDynamicLights = FMath::Min( NumDynamicLights, GEngine->LightComplexityColors.Num() - 1 );
-			const FColor Color = GEngine->LightComplexityColors[NumDynamicLights];
-
-			FColoredMaterialRenderProxy LightComplexityMaterialInstance(
-				GEngine->LevelColorationUnlitMaterial->GetRenderProxy(Mesh.MaterialRenderProxy->IsSelected(), Mesh.MaterialRenderProxy->IsHovered()),
-				Color
-				);
-
-			// Draw the mesh colored by light complexity.
-			FMeshBatch ModifiedMesh = Mesh;
-			ModifiedMesh.MaterialRenderProxy = &LightComplexityMaterialInstance;
-			PassesRendered = PDI->DrawMesh(ModifiedMesh);
-		}
-	}
-	else if(!EngineShowFlags.Materials && !PDI->View->bForceShowMaterials)
-	{
-		// Don't render unlit translucency when in 'lighting only' viewmode.
-		if (Mesh.MaterialRenderProxy->GetMaterial(FeatureLevel)->GetShadingModel() != MSM_Unlit
-			// Don't render translucency in 'lighting only', since the viewmode works by overriding with an opaque material
-			// This would cause a mismatch of the material's blend mode with the primitive's view relevance,
-			// And make faint particles block the view
-			&& !IsTranslucentBlendMode(Mesh.MaterialRenderProxy->GetMaterial(FeatureLevel)->GetBlendMode()))
-		{
-			// When materials aren't shown, apply the same basic material to all meshes.
-			FMeshBatch ModifiedMesh = Mesh;
-
-			bool bTextureMapped = false;
-			FVector2D LMResolution;
-
-			if (EngineShowFlags.LightMapDensity && Mesh.LCI)
-			{
-				auto Interaction = Mesh.LCI->GetLightMapInteraction(FeatureLevel);
-				auto Texture = Interaction.GetTexture(AllowHighQualityLightmaps(FeatureLevel));
-
-				if (Interaction.GetType() == LMIT_Texture && Texture)
-				{
-					LMResolution.X = Texture->GetSizeX();
-					LMResolution.Y = Texture->GetSizeY();
-					bTextureMapped = true;
-				}
-			}
-
-			if (bTextureMapped == false)
-			{
-				FMaterialRenderProxy* RenderProxy = GEngine->LevelColorationLitMaterial->GetRenderProxy(Mesh.MaterialRenderProxy->IsSelected(),Mesh.MaterialRenderProxy->IsHovered());
-				const FColoredMaterialRenderProxy LightingOnlyMaterialInstance(
-					RenderProxy,
-					GEngine->LightingOnlyBrightness
-					);
-
-				ModifiedMesh.MaterialRenderProxy = &LightingOnlyMaterialInstance;
-				PassesRendered = PDI->DrawMesh(ModifiedMesh);
-			}
-			else
-			{
-				FMaterialRenderProxy* RenderProxy = GEngine->LightingTexelDensityMaterial->GetRenderProxy(Mesh.MaterialRenderProxy->IsSelected(),Mesh.MaterialRenderProxy->IsHovered());
-				const FLightingDensityMaterialRenderProxy LightingDensityMaterialInstance(
-					RenderProxy,
-					GEngine->LightingOnlyBrightness,
-					LMResolution
-					);
-
-				ModifiedMesh.MaterialRenderProxy = &LightingDensityMaterialInstance;
-				PassesRendered = PDI->DrawMesh(ModifiedMesh);
-			}
-		}
-	}
-	else
-	{	
-		if(EngineShowFlags.PropertyColoration)
-		{
-			// In property coloration mode, override the mesh's material with a color that was chosen based on property value.
-			const UMaterial* PropertyColorationMaterial = EngineShowFlags.Lighting ? GEngine->LevelColorationLitMaterial : GEngine->LevelColorationUnlitMaterial;
-
-			const FColoredMaterialRenderProxy PropertyColorationMaterialInstance(
-				PropertyColorationMaterial->GetRenderProxy(Mesh.MaterialRenderProxy->IsSelected(), Mesh.MaterialRenderProxy->IsHovered()),
-				GetSelectionColor(PropertyColor,bSelected,Mesh.MaterialRenderProxy->IsHovered())
-				);
-			FMeshBatch ModifiedMesh = Mesh;
-			ModifiedMesh.MaterialRenderProxy = &PropertyColorationMaterialInstance;
-			PassesRendered = PDI->DrawMesh(ModifiedMesh);
-		}
-		else if (EngineShowFlags.LevelColoration)
-		{
-			const UMaterial* LevelColorationMaterial = EngineShowFlags.Lighting ? GEngine->LevelColorationLitMaterial : GEngine->LevelColorationUnlitMaterial;
-			// Draw the mesh with level coloration.
-			const FColoredMaterialRenderProxy LevelColorationMaterialInstance(
-				LevelColorationMaterial->GetRenderProxy(Mesh.MaterialRenderProxy->IsSelected(), Mesh.MaterialRenderProxy->IsHovered()),
-				GetSelectionColor(LevelColor,bSelected,Mesh.MaterialRenderProxy->IsHovered())
-				);
-			FMeshBatch ModifiedMesh = Mesh;
-			ModifiedMesh.MaterialRenderProxy = &LevelColorationMaterialInstance;
-			PassesRendered = PDI->DrawMesh(ModifiedMesh);
-		}
-		else if (EngineShowFlags.BSPSplit
-			&& PrimitiveSceneProxy->ShowInBSPSplitViewmode())
-		{
-			// Determine unique color for model component.
-			FLinearColor BSPSplitColor;
-			FRandomStream RandomStream(GetTypeHash(PrimitiveSceneProxy->GetPrimitiveComponentId().Value));
-			BSPSplitColor.R = RandomStream.GetFraction();
-			BSPSplitColor.G = RandomStream.GetFraction();
-			BSPSplitColor.B = RandomStream.GetFraction();
-			BSPSplitColor.A = 1.0f;
-
-			// Piggy back on the level coloration material.
-			const UMaterial* BSPSplitMaterial = EngineShowFlags.Lighting ? GEngine->LevelColorationLitMaterial : GEngine->LevelColorationUnlitMaterial;
-			
-			// Draw BSP mesh with unique color for each model component.
-			const FColoredMaterialRenderProxy BSPSplitMaterialInstance(
-				BSPSplitMaterial->GetRenderProxy(Mesh.MaterialRenderProxy->IsSelected(), Mesh.MaterialRenderProxy->IsHovered()),
-				GetSelectionColor(BSPSplitColor,bSelected,Mesh.MaterialRenderProxy->IsHovered())
-				);
-			FMeshBatch ModifiedMesh = Mesh;
-			ModifiedMesh.MaterialRenderProxy = &BSPSplitMaterialInstance;
-			PassesRendered = PDI->DrawMesh(ModifiedMesh);
-		}
-		else if (PrimitiveSceneProxy->HasStaticLighting() && !PrimitiveSceneProxy->HasValidSettingsForStaticLighting())
-		{
-			const FColoredMaterialRenderProxy InvalidSettingsMaterialInstance(
-				GEngine->InvalidLightmapSettingsMaterial->GetRenderProxy(bSelected),
-				GetSelectionColor(LevelColor,bSelected,Mesh.MaterialRenderProxy->IsHovered())
-				);
-			FMeshBatch ModifiedMesh = Mesh;
-			ModifiedMesh.MaterialRenderProxy = &InvalidSettingsMaterialInstance;
-			PassesRendered = PDI->DrawMesh(ModifiedMesh);
-		}
-		else 
-		{
-			// Draw the mesh unmodified.
-			PassesRendered = PDI->DrawMesh( Mesh );
-		}
-
-		//Draw a wireframe overlay last, if requested
-		if(EngineShowFlags.MeshEdges)
-		{
-			// Draw the mesh's edges in blue, on top of the base geometry.
-			if (Mesh.MaterialRenderProxy->GetMaterial(FeatureLevel)->MaterialModifiesMeshPosition_RenderThread())
-			{
-				const FSceneView* View = PDI->View;
-
-				// If the material is mesh-modifying, we cannot rely on substitution
-				const FOverrideSelectionColorMaterialRenderProxy WireframeMaterialInstance(
-					Mesh.MaterialRenderProxy,
-					WireframeColor
-					);
-
-				FMeshBatch ModifiedMesh = Mesh;
-				ModifiedMesh.bWireframe = true;
-				ModifiedMesh.MaterialRenderProxy = &WireframeMaterialInstance;
-				PassesRendered = PDI->DrawMesh(ModifiedMesh);
-			}
-			else
-			{
-				FColoredMaterialRenderProxy WireframeMaterialInstance(
-					GEngine->WireframeMaterial->GetRenderProxy(Mesh.MaterialRenderProxy->IsSelected(), Mesh.MaterialRenderProxy->IsHovered()),
-					WireframeColor
-					);
-				FMeshBatch ModifiedMesh = Mesh;
-				ModifiedMesh.bWireframe = true;
-				ModifiedMesh.MaterialRenderProxy = &WireframeMaterialInstance;
-				PassesRendered = PDI->DrawMesh(ModifiedMesh);
-			}
-		}
-	}
-#endif
-
-	return PassesRendered;
-}
-
 void ClampUVs(FVector2D* UVs, int32 NumUVs)
 {
 	const float FudgeFactor = 0.1f;
diff --git a/Engine/Source/Runtime/Engine/Private/SkeletalMesh.cpp b/Engine/Source/Runtime/Engine/Private/SkeletalMesh.cpp
--- a/Engine/Source/Runtime/Engine/Private/SkeletalMesh.cpp
+++ b/Engine/Source/Runtime/Engine/Private/SkeletalMesh.cpp
@@ -4113,30 +4113,6 @@ FSkeletalMeshSceneProxy::FSkeletalMeshSceneProxy(const USkinnedMeshComponent* Co
 
 // FPrimitiveSceneProxy interface.
 
-/**
- *	Called during InitViews for view processing on scene proxies before rendering them
- *  Only called for primitives that are visible and have bDynamicRelevance
- *
- *	@param	ViewFamily		The ViewFamily to pre-render for
- *	@param	VisibilityMap	A BitArray that indicates whether the primitive was visible in that view (index)
- *	@param	FrameNumber		The frame number of this pre-render
- */
-void FSkeletalMeshSceneProxy::PreRenderView(const FSceneViewFamily* ViewFamily, const uint32 VisibilityMap, int32 FrameNumber)
-{
-	/** Update the LOD level we want for this skeletal mesh (back on the game thread). */
-	if (MeshObject)
-	{
-		for (int32 ViewIndex = 0; ViewIndex < ViewFamily->Views.Num(); ViewIndex++)
-		{
-			if (VisibilityMap & (1<<ViewIndex))
-			{
-				const FSceneView* View = ViewFamily->Views[ViewIndex];
-				MeshObject->UpdateMinDesiredLODLevel(View, GetBounds(), FrameNumber);
-			}
-		}
-	}
-}
-
 /** 
  * Iterates over sections,chunks,elements based on current instance weight usage 
  */
@@ -4328,7 +4304,7 @@ void FSkeletalMeshSceneProxy::GetMeshElementsConditionallySelectable(const TArra
 
 			if( PhysicsAssetForDebug )
 			{
-				DebugDrawPhysicsAsset(Collector.GetPDI(ViewIndex),ViewFamily.EngineShowFlags);
+				DebugDrawPhysicsAsset(ViewIndex, Collector, ViewFamily.EngineShowFlags);
 			}
 
 			if (ViewFamily.EngineShowFlags.SkeletalMeshes)
@@ -4340,81 +4316,6 @@ void FSkeletalMeshSceneProxy::GetMeshElementsConditionallySelectable(const TArra
 #endif
 }
 
-/** 
-* Draw the scene proxy as a dynamic element
-*
-* @param	PDI - draw interface to render to
-* @param	View - current view
-*/
-void FSkeletalMeshSceneProxy::DrawDynamicElements(FPrimitiveDrawInterface* PDI, const FSceneView* View)
-{
-	QUICK_SCOPE_CYCLE_COUNTER( STAT_SkeletalMeshSceneProxy_DrawDynamicElements );
-
-	if( !MeshObject )
-	{
-		return;
-	}		  
-
-	const int32 LODIndex = MeshObject->GetLOD();
-	check(LODIndex < SkelMeshResource->LODModels.Num());
-	const FStaticLODModel& LODModel = SkelMeshResource->LODModels[LODIndex];
-
-	if( LODSections.Num() > 0 )
-	{
-		const FLODSectionElements& LODSection = LODSections[LODIndex];
-
-		check(LODSection.SectionElements.Num() == LODModel.Sections.Num());
-
-		for (FSkeletalMeshSectionIter Iter(LODIndex, *MeshObject, LODModel, LODSection); Iter; ++Iter)
-		{
-			FSkelMeshSection Section = Iter.GetSection();
-			const FSkelMeshChunk& Chunk = Iter.GetChunk();
-			const FSectionElementInfo& SectionElementInfo = Iter.GetSectionElementInfo();
-			const FTwoVectors& CustomLeftRightVectors = Iter.GetCustomLeftRightVectors();
-
-#if WITH_EDITORONLY_DATA
-			// TODO: This is not threadsafe! A render command should be used to propagate SelectedEditorSection to the scene proxy.
-			Section.bSelected = (SkeletalMeshForDebug->SelectedEditorSection == SectionElementInfo.UseMaterialIndex);
-#endif
-			// If hidden skip the draw
-			if (MeshObject->IsMaterialHidden(LODIndex, SectionElementInfo.UseMaterialIndex))
-			{
-				continue;
-			}
-
-			// If disabled, then skip the draw
-			if(Section.bDisabled)
-			{
-				continue;
-			}
-
-#if WITH_EDITOR
-			if ( GIsEditor && PDI->IsHitTesting() )
-			{
-				PDI->SetHitProxy(SectionElementInfo.HitProxy);
-			}
-#endif
-
-			DrawDynamicElementsSection(PDI, View, LODModel, LODIndex, Section, Chunk, SectionElementInfo, CustomLeftRightVectors );
-		}
-	}
-
-#if !(UE_BUILD_SHIPPING || UE_BUILD_TEST)
-	// debug drawing
-	TArray<FTransform>* BoneSpaceBases = MeshObject->GetSpaceBases();
-
-	if( PhysicsAssetForDebug )
-	{
-		DebugDrawPhysicsAsset(PDI,View->Family->EngineShowFlags);
-	}
-
-	if (View->Family->EngineShowFlags.SkeletalMeshes)
-	{
-		RenderBounds(PDI, View->Family->EngineShowFlags, GetBounds(), IsSelected());
-	}
-#endif
-}
-
 void FSkeletalMeshSceneProxy::GetDynamicElementsSection(const TArray<const FSceneView*>& Views, const FSceneViewFamily& ViewFamily, uint32 VisibilityMap, 
 	const FStaticLODModel& LODModel, const int32 LODIndex, const FSkelMeshSection& Section, const FSkelMeshChunk& Chunk, 
 	const FSectionElementInfo& SectionElementInfo, const FTwoVectors& CustomLeftRightVectors, bool bSelectable, FMeshElementCollector& Collector ) const
@@ -4550,149 +4451,6 @@ void FSkeletalMeshSceneProxy::GetDynamicElementsSection(const TArray<const FScen
 	}
 }
 
-void FSkeletalMeshSceneProxy::DrawDynamicElementsSection(FPrimitiveDrawInterface* PDI,const FSceneView* View,
-														const FStaticLODModel& LODModel, const int32 LODIndex, const FSkelMeshSection& Section, 
-														const FSkelMeshChunk& Chunk, const FSectionElementInfo& SectionElementInfo, const FTwoVectors& CustomLeftRightVectors )
-{
-#if !WITH_EDITOR
-	const bool bIsSelected = false;
-#else // #if !WITH_EDITOR
-	bool bIsSelected = IsSelected();
-
-	// if the mesh isn't selected but the mesh section is selected in the AnimSetViewer, find the mesh component and make sure that it can be highlighted (ie. are we rendering for the AnimSetViewer or not?)
-	if( !bIsSelected && Section.bSelected && bCanHighlightSelectedSections )
-	{
-		bIsSelected = true;
-	}
-#endif // #if WITH_EDITOR
-
-	const auto FeatureLevel = View->GetFeatureLevel();
-
-	const bool bIsWireframe = View->Family->EngineShowFlags.Wireframe;
-
-	if(Section.bDisabled)
-		return;
-
-	// If hidden skip the draw
-	if (MeshObject->IsMaterialHidden(LODIndex,SectionElementInfo.UseMaterialIndex))
-	{
-		return;
-	}
-
-	FMeshBatch Mesh;
-	FMeshBatchElement& BatchElement = Mesh.Elements[0];
-	Mesh.DynamicVertexData = NULL;
-	Mesh.UseDynamicData = false;
-	Mesh.LCI = NULL;
-	Mesh.bWireframe |= bForceWireframe;
-	Mesh.Type = PT_TriangleList;
-	Mesh.VertexFactory = MeshObject->GetVertexFactory(LODIndex,Section.ChunkIndex);
-	BatchElement.FirstIndex = Section.BaseIndex;
-
-	BatchElement.IndexBuffer = LODModel.MultiSizeIndexContainer.GetIndexBuffer();
-	BatchElement.MaxVertexIndex = LODModel.NumVertices - 1;
-
-	BatchElement.UserIndex = MeshObject->GPUSkinCacheKeys[Section.ChunkIndex];
-
-	const bool bRequiresAdjacencyInformation = RequiresAdjacencyInformation(SectionElementInfo.Material, Mesh.VertexFactory->GetType(), FeatureLevel);
-	if ( bRequiresAdjacencyInformation )
-	{
-		check( LODModel.AdjacencyMultiSizeIndexContainer.IsIndexBufferValid() );
-		BatchElement.IndexBuffer = LODModel.AdjacencyMultiSizeIndexContainer.GetIndexBuffer();
-		Mesh.Type = PT_12_ControlPointPatchList;
-		BatchElement.FirstIndex *= 4;
-	}
-
-#if WITH_EDITOR
-	const FOverrideSelectionColorMaterialRenderProxy SelectionOverrideProxy(
-		SectionElementInfo.Material->GetRenderProxy(bIsSelected, IsHovered()),
-		GetSelectionColor(GEngine->GetSelectedMaterialColor(), bIsSelected, IsHovered())
-		);
-	if (Section.bSelected)
-	{
-		Mesh.MaterialRenderProxy = &SelectionOverrideProxy;
-	}
-	else
-	{
-		Mesh.MaterialRenderProxy = SectionElementInfo.Material->GetRenderProxy(bIsSelected, IsHovered());
-	}
-#else
-	Mesh.MaterialRenderProxy = SectionElementInfo.Material->GetRenderProxy(bIsSelected, IsHovered());
-#endif
-
-	BatchElement.PrimitiveUniformBufferResource = &GetUniformBuffer();
-	
-	// Select which indices to use if TRISORT_CustomLeftRight
-	if( Section.TriangleSorting == TRISORT_CustomLeftRight )
-	{
-		switch( MeshObject->CustomSortAlternateIndexMode )
-		{
-		case CSAIM_Left:
-			// Left view - use second set of indices.
-			BatchElement.FirstIndex += Section.NumTriangles * 3;
-			break;
-		case  CSAIM_Right:
-			// Right view - use first set of indices.
-			break;
-		default:
-			// Calculate the viewing direction
-			FVector SortWorldOrigin = GetLocalToWorld().TransformPosition(CustomLeftRightVectors.v1);
-			FVector SortWorldDirection = GetLocalToWorld().TransformVector(CustomLeftRightVectors.v2);
-
-			if( (SortWorldDirection | (SortWorldOrigin - View->ViewMatrices.ViewOrigin)) < 0.f )
-			{
-				BatchElement.FirstIndex += Section.NumTriangles * 3;
-			}
-			break;
-		}
-	}
-
-	BatchElement.NumPrimitives = Section.NumTriangles;
-	if( GIsEditor && MeshObject->ProgressiveDrawingFraction != 1.f )
-	{
-		if (Mesh.MaterialRenderProxy->GetMaterial(FeatureLevel)->GetBlendMode() == BLEND_Translucent)
-		{
-			BatchElement.NumPrimitives = FMath::RoundToInt(((float)Section.NumTriangles)*FMath::Clamp<float>(MeshObject->ProgressiveDrawingFraction,0.f,1.f));
-			if( BatchElement.NumPrimitives == 0 )
-			{
-				return;
-			}
-		}
-	}
-	BatchElement.MinVertexIndex = Chunk.BaseVertexIndex;
-	Mesh.ReverseCulling = IsLocalToWorldDeterminantNegative();
-	Mesh.CastShadow = SectionElementInfo.bEnableShadowCasting;
-
-	// if not using dynamic index buffer nor tesselation, check that we're still within bounds
-	const int32 kIndicesPerPrimitive = 3;
-#if 0 //@todo UE4 merge, graphics,
-#if WITH_D3D11_TESSELLATION
-	check(DynamicIndexBuffer != NULL || 
-		PT_12_ControlPointPatchList == Mesh.Type ||
-		BatchElement.FirstIndex + kIndicesPerPrimitive * Mesh.NumPrimitives <= static_cast< uint32 >( LODModel.MultiSizeIndexContainer.GetIndexBuffer()->Num() ));
-#else
-	check(DynamicIndexBuffer != NULL || 
-		BatchElement.FirstIndex + kIndicesPerPrimitive * Mesh.NumPrimitives <= static_cast< uint32 >( LODModel.MultiSizeIndexContainer.GetIndexBuffer()->Num() ));
-#endif // WITH_D3D11_TESSELLATION
-#endif
-
-	const int32 NumPasses = DrawRichMesh(
-		PDI,
-		Mesh,
-		WireframeColor,
-		LevelColor,
-		PropertyColor,
-		this,
-		bIsSelected,
-		bIsWireframe
-		);
-
-	const int32 NumVertices = Chunk.NumRigidVertices + Chunk.NumSoftVertices;
-	INC_DWORD_STAT_BY(STAT_GPUSkinVertices,(uint32)(bIsCPUSkinned ? 0 : NumVertices * NumPasses));
-	INC_DWORD_STAT_BY(STAT_SkelMeshTriangles,Mesh.GetNumPrimitives() * NumPasses);
-	INC_DWORD_STAT(STAT_SkelMeshDrawCalls);
-}
-
 /**
  * Returns the world transform to use for drawing.
  * @param OutLocalToWorld - Will contain the local-to-world transform when the function returns.
@@ -4712,7 +4470,6 @@ FPrimitiveViewRelevance FSkeletalMeshSceneProxy::GetViewRelevance(const FSceneVi
 	FPrimitiveViewRelevance Result;
 	Result.bDrawRelevance = IsShown(View) && View->Family->EngineShowFlags.SkeletalMeshes;
 	Result.bShadowRelevance = IsShadowCast(View);
-	Result.bNeedsPreRenderView = Result.bDrawRelevance || Result.bShadowRelevance;
 	Result.bDynamicRelevance = true;
 	Result.bRenderCustomDepth = ShouldRenderCustomDepth();
 	Result.bRenderInMainPass = ShouldRenderInMainPass();
@@ -4742,7 +4499,7 @@ int32 FSkeletalMeshSceneProxy::GetCurrentLODIndex()
 /** 
  * Render physics asset for debug display
  */
-void FSkeletalMeshSceneProxy::DebugDrawPhysicsAsset(FPrimitiveDrawInterface* PDI, const FEngineShowFlags& EngineShowFlags) const
+void FSkeletalMeshSceneProxy::DebugDrawPhysicsAsset(int32 ViewIndex, FMeshElementCollector& Collector, const FEngineShowFlags& EngineShowFlags) const
 {
 	FMatrix LocalToWorld, WorldToLocal;
 	GetWorldMatrices( LocalToWorld, WorldToLocal );
@@ -4761,11 +4518,11 @@ void FSkeletalMeshSceneProxy::DebugDrawPhysicsAsset(FPrimitiveDrawInterface* PDI
 			check(PhysicsAssetForDebug);
 			if( EngineShowFlags.Collision && IsCollisionEnabled() )
 			{
-				PhysicsAssetForDebug->DrawCollision(PDI, SkeletalMeshForDebug, *BoneSpaceBases, LocalToWorldTransform, TotalScale.X);
+				PhysicsAssetForDebug->GetCollisionMesh(ViewIndex, Collector, SkeletalMeshForDebug, *BoneSpaceBases, LocalToWorldTransform, TotalScale.X);
 			}
 			if( EngineShowFlags.Constraints )
 			{
-				PhysicsAssetForDebug->DrawConstraints(PDI, SkeletalMeshForDebug, *BoneSpaceBases, LocalToWorldTransform, TotalScale.X);
+				PhysicsAssetForDebug->DrawConstraints(Collector.GetPDI(ViewIndex), SkeletalMeshForDebug, *BoneSpaceBases, LocalToWorldTransform, TotalScale.X);
 			}
 		}
 	}
diff --git a/Engine/Source/Runtime/Engine/Private/SphereComponent.cpp b/Engine/Source/Runtime/Engine/Private/SphereComponent.cpp
--- a/Engine/Source/Runtime/Engine/Private/SphereComponent.cpp
+++ b/Engine/Source/Runtime/Engine/Private/SphereComponent.cpp
@@ -130,41 +130,6 @@ FPrimitiveSceneProxy* USphereComponent::CreateSceneProxy()
 			}
 		}
 
-		/** 
-		* Draw the scene proxy as a dynamic element
-		*
-		* @param	PDI - draw interface to render to
-		* @param	View - current view
-		*/
-		virtual void DrawDynamicElements(FPrimitiveDrawInterface* PDI,const FSceneView* View) override
-		{
-			QUICK_SCOPE_CYCLE_COUNTER( STAT_SphereSceneProxy_DrawDynamicElements );
-
-			const FMatrix& LocalToWorld = GetLocalToWorld();
-			int32 SphereSides =  FMath::Clamp<int32>(SphereRadius/4.f, 16, 64);
-			if(ShapeMaterial && !View->Family->EngineShowFlags.Wireframe)
-			{
-				DrawSphere(PDI,LocalToWorld.GetOrigin(), FVector(SphereRadius), SphereSides, SphereSides/2, ShapeMaterial->GetRenderProxy(false), SDPG_World);
-			}
-			else
-			{
-				const FLinearColor DrawSphereColor = GetSelectionColor(SphereColor, IsSelected(), IsHovered(), /*bUseOverlayIntensity=*/false);
-				
-				float AbsScaleX = LocalToWorld.GetScaledAxis(EAxis::X).Size();
-				float AbsScaleY = LocalToWorld.GetScaledAxis(EAxis::Y).Size();
-				float AbsScaleZ = LocalToWorld.GetScaledAxis(EAxis::Z).Size();
-				float MinAbsScale = FMath::Min3(AbsScaleX, AbsScaleY, AbsScaleZ);
-
-				FVector ScaledX = LocalToWorld.GetUnitAxis(EAxis::X) * MinAbsScale;
-				FVector ScaledY = LocalToWorld.GetUnitAxis(EAxis::Y) * MinAbsScale;
-				FVector ScaledZ = LocalToWorld.GetUnitAxis(EAxis::Z) * MinAbsScale;
-
-				DrawCircle(PDI, LocalToWorld.GetOrigin(), ScaledX, ScaledY, DrawSphereColor, SphereRadius, SphereSides, SDPG_World);
-				DrawCircle(PDI, LocalToWorld.GetOrigin(), ScaledX, ScaledZ, DrawSphereColor, SphereRadius, SphereSides, SDPG_World);
-				DrawCircle(PDI, LocalToWorld.GetOrigin(), ScaledY, ScaledZ, DrawSphereColor, SphereRadius, SphereSides, SDPG_World);
-			}
-		}
-
 		virtual FPrimitiveViewRelevance GetViewRelevance(const FSceneView* View)  override
 		{
 			const bool bVisibleForSelection = !bDrawOnlyIfSelected || IsSelected();
diff --git a/Engine/Source/Runtime/Engine/Private/StaticMeshRender.cpp b/Engine/Source/Runtime/Engine/Private/StaticMeshRender.cpp
--- a/Engine/Source/Runtime/Engine/Private/StaticMeshRender.cpp
+++ b/Engine/Source/Runtime/Engine/Private/StaticMeshRender.cpp
@@ -837,291 +837,6 @@ void FStaticMeshSceneProxy::GetDynamicMeshElements(const TArray<const FSceneView
 	}
 }
 
-/** 
-* Draw the scene proxy as a dynamic element
-*
-* @param	PDI - draw interface to render to
-* @param	View - current view
-*/
-void FStaticMeshSceneProxy::DrawDynamicElements(FPrimitiveDrawInterface* PDI,const FSceneView* View, uint32 DrawDynamicFlags )
-{
-	QUICK_SCOPE_CYCLE_COUNTER( STAT_StaticMeshSceneProxy_DrawDynamicElements );
-
-	checkSlow(IsInParallelRenderingThread());
-
-	const bool bIsLightmapSettingError = HasStaticLighting() && !HasValidSettingsForStaticLighting();
-
-	bool bProxyIsSelected = IsSelected();
-
-	bool bDrawSimpleCollision = false, bDrawComplexCollision = false;
-	const bool bInCollisionView = IsCollisionView(View->Family->EngineShowFlags, bDrawSimpleCollision, bDrawComplexCollision);
-	// Draw simple collision as wireframe if 'show collision', collision is enabled, and we are not using the complex as the simple
-	const bool bDrawWireframeCollision = (View->Family->EngineShowFlags.Collision && IsCollisionEnabled() && CollisionTraceFlag != ECollisionTraceFlag::CTF_UseComplexAsSimple);
-
-	const bool bDrawMesh = (bInCollisionView) ? (bDrawComplexCollision) : 
-		(	IsRichView(*View->Family) || HasViewDependentDPG()
-			|| View->Family->EngineShowFlags.Collision
-			|| View->Family->EngineShowFlags.Bounds
-			|| bProxyIsSelected 
-			|| IsHovered()
-			|| bIsLightmapSettingError ) ;
-
-
-	// Draw polygon mesh if we are either not in a collision view, or are drawing it as collision.
-	if(View->Family->EngineShowFlags.StaticMeshes && bDrawMesh)
-	{
-		// how we should draw the collision for this mesh.
-		const bool bIsWireframeView = View->Family->EngineShowFlags.Wireframe;
-		const bool bLevelColorationEnabled = View->Family->EngineShowFlags.LevelColoration;
-		const bool bPropertyColorationEnabled = View->Family->EngineShowFlags.PropertyColoration;
-
-		// force lowest LOD (e.g. for Reflective Shadow Maps)
-		int32 LODIndex = 0;
-		if(DrawDynamicFlags & EDrawDynamicFlags::ForceLowestLOD)
-		{
-			LODIndex = RenderData->LODResources.Num() - 1;
-		} 
-		else
-		{
-			LODIndex = GetLOD(View);
-		}
-
-		if (RenderData->LODResources.IsValidIndex(LODIndex))
-		{
-			const FStaticMeshLODResources& LODModel = RenderData->LODResources[LODIndex];
-			const FLODInfo& ProxyLODInfo = LODs[LODIndex];
-
-			if (AllowDebugViewmodes() && bIsWireframeView && !View->Family->EngineShowFlags.Materials
-				// If any of the materials are mesh-modifying, we can't use the single merged mesh element of GetWireframeMeshElement()
-				&& !ProxyLODInfo.UsesMeshModifyingMaterials())
-			{
-				FLinearColor ViewWireframeColor( bLevelColorationEnabled ? LevelColor : WireframeColor );
-				if ( bPropertyColorationEnabled )
-				{
-					ViewWireframeColor = PropertyColor;
-				}
-
-				FColoredMaterialRenderProxy WireframeMaterialInstance(
-					GEngine->WireframeMaterial->GetRenderProxy(false),
-					GetSelectionColor(ViewWireframeColor,!(GIsEditor && View->Family->EngineShowFlags.Selection) || bProxyIsSelected, IsHovered(), /*bUseOverlayIntensity=*/false)
-					);
-
-				const int32 NumBatches = GetNumMeshBatches();
-
-				for (int32 BatchIndex = 0; BatchIndex < NumBatches; BatchIndex++)
-				{
-					FMeshBatch MeshBatch;
-
-					if (GetWireframeMeshElement(LODIndex, BatchIndex, &WireframeMaterialInstance, GetDepthPriorityGroup(View), MeshBatch))
-					{
-						const int32 NumPasses = PDI->DrawMesh(MeshBatch);
-						INC_DWORD_STAT_BY(STAT_StaticMeshTriangles,MeshBatch.GetNumPrimitives() * NumPasses);
-					}
-				}
-			}
-			else
-			{
-				const FLinearColor UtilColor( LevelColor );
-
-				// Draw the static mesh sections.
-				for(int32 SectionIndex = 0; SectionIndex < LODModel.Sections.Num(); SectionIndex++)
-				{
-					bool bSectionIsSelected = false;
-
-#if WITH_EDITOR
-					if( GIsEditor )
-					{
-						const FLODInfo::FSectionInfo& Section = LODs[LODIndex].Sections[SectionIndex];
-
-						bSectionIsSelected = Section.bSelected;
-						PDI->SetHitProxy(Section.HitProxy);
-					}
-#endif // WITH_EDITOR
-
-					const int32 NumBatches = GetNumMeshBatches();
-
-					for (int32 BatchIndex = 0; BatchIndex < NumBatches; BatchIndex++)
-					{
-						FMeshBatch MeshBatch;
-
-						if (GetMeshElement(LODIndex, BatchIndex, SectionIndex, GetDepthPriorityGroup(View), bSectionIsSelected, IsHovered(), MeshBatch))
-						{			
-#if !(UE_BUILD_SHIPPING || UE_BUILD_TEST)
-							if( View->Family->EngineShowFlags.VertexColors && AllowDebugViewmodes() )
-							{
-								// Override the mesh's material with our material that draws the vertex colors
-								UMaterial* VertexColorVisualizationMaterial = NULL;
-								switch( GVertexColorViewMode )
-								{
-								case EVertexColorViewMode::Color:
-									VertexColorVisualizationMaterial = GEngine->VertexColorViewModeMaterial_ColorOnly;
-									break;
-
-								case EVertexColorViewMode::Alpha:
-									VertexColorVisualizationMaterial = GEngine->VertexColorViewModeMaterial_AlphaAsColor;
-									break;
-
-								case EVertexColorViewMode::Red:
-									VertexColorVisualizationMaterial = GEngine->VertexColorViewModeMaterial_RedOnly;
-									break;
-
-								case EVertexColorViewMode::Green:
-									VertexColorVisualizationMaterial = GEngine->VertexColorViewModeMaterial_GreenOnly;
-									break;
-
-								case EVertexColorViewMode::Blue:
-									VertexColorVisualizationMaterial = GEngine->VertexColorViewModeMaterial_BlueOnly;
-									break;
-								}
-								check( VertexColorVisualizationMaterial != NULL );
-
-								const FColoredMaterialRenderProxy VertexColorVisualizationMaterialInstance(
-									VertexColorVisualizationMaterial->GetRenderProxy( MeshBatch.MaterialRenderProxy->IsSelected(),MeshBatch.MaterialRenderProxy->IsHovered() ),
-									GetSelectionColor( FLinearColor::White, bSectionIsSelected, IsHovered() )
-									);
-								FMeshBatch ModifiedMeshElement = MeshBatch;
-								ModifiedMeshElement.MaterialRenderProxy = &VertexColorVisualizationMaterialInstance;
-
-								const int32 NumPasses = DrawRichMesh(
-									PDI,
-									ModifiedMeshElement,
-									WireframeColor,
-									UtilColor,
-									PropertyColor,
-									this,
-									bSectionIsSelected,
-									bIsWireframeView
-									);
-								INC_DWORD_STAT_BY(STAT_StaticMeshTriangles,MeshBatch.GetNumPrimitives() * NumPasses);
-							}
-							else if( bInCollisionView && bDrawComplexCollision && AllowDebugViewmodes() )
-							{
-								if (LODModel.Sections[SectionIndex].bEnableCollision)
-								{
-									// Override the mesh's material with our material that draws the collision color
-									const FColoredMaterialRenderProxy CollisionMaterialInstance(
-										GEngine->ShadedLevelColorationUnlitMaterial->GetRenderProxy(bProxyIsSelected, IsHovered()),
-										WireframeColor
-										);
-
-									FMeshBatch ModifiedMeshElement = MeshBatch;
-									ModifiedMeshElement.MaterialRenderProxy = &CollisionMaterialInstance;
-
-									const int32 NumPasses = DrawRichMesh(
-										PDI,
-										ModifiedMeshElement,
-										WireframeColor,
-										UtilColor,
-										PropertyColor,
-										this,
-										bSectionIsSelected,
-										bIsWireframeView
-										);
-									INC_DWORD_STAT_BY(STAT_StaticMeshTriangles,MeshBatch.GetNumPrimitives() * NumPasses);
-								}
-							}
-							else
-#endif
-#if WITH_EDITOR
-							if (bSectionIsSelected)
-							{
-								// Override the mesh's material with our material that draws the collision color
-								const FOverrideSelectionColorMaterialRenderProxy SelectedMaterialInstance(
-									GEngine->ShadedLevelColorationUnlitMaterial->GetRenderProxy(bSectionIsSelected, IsHovered()),
-									GetSelectionColor(GEngine->GetSelectedMaterialColor(), bSectionIsSelected, IsHovered())
-									);
-
-								FMeshBatch ModifiedMeshElement(MeshBatch);
-								ModifiedMeshElement.MaterialRenderProxy = &SelectedMaterialInstance;
-
-								const int32 NumPasses = DrawRichMesh(
-									PDI,
-									ModifiedMeshElement,
-									WireframeColor,
-									UtilColor,
-									PropertyColor,
-									this,
-									bSectionIsSelected,
-									bIsWireframeView
-									);
-								INC_DWORD_STAT_BY(STAT_StaticMeshTriangles, MeshBatch.GetNumPrimitives() * NumPasses);
-							}
-							else
-#endif
-							{
-								const int32 NumPasses = DrawRichMesh(
-									PDI,
-									MeshBatch,
-									WireframeColor,
-									UtilColor,
-									PropertyColor,
-									this,
-									bSectionIsSelected,
-									bIsWireframeView
-									);
-								INC_DWORD_STAT_BY(STAT_StaticMeshTriangles,MeshBatch.GetNumPrimitives() * NumPasses);
-							}
-						}
-					}
-				}
-			}
-		}
-	}
-
-	if((bDrawSimpleCollision || bDrawWireframeCollision) && AllowDebugViewmodes())
-	{
-		if(BodySetup)
-		{
-			if(FMath::Abs(GetLocalToWorld().Determinant()) < SMALL_NUMBER)
-			{
-				// Catch this here or otherwise GeomTransform below will assert
-				// This spams so commented out
-				//UE_LOG(LogStaticMesh, Log, TEXT("Zero scaling not supported (%s)"), *StaticMesh->GetPathName());
-			}
-			else
-			{
-				const bool bDrawSolid = !bDrawWireframeCollision;
-
-				if(bDrawSolid)
-				{
-					// Make a material for drawing solid collision stuff
-					const FColoredMaterialRenderProxy SolidMaterialInstance(
-						GEngine->ShadedLevelColorationUnlitMaterial->GetRenderProxy(IsSelected(), IsHovered()),
-						WireframeColor
-						);
-
-					FTransform GeomTransform(GetLocalToWorld());
-					BodySetup->AggGeom.DrawAggGeom(PDI, GeomTransform, WireframeColor, /*Material=*/&SolidMaterialInstance, false, /*bSolid=*/true, UseEditorDepthTest());
-				}
-				// wireframe
-				else
-				{
-					FColor CollisionColor = FColor(157,149,223,255);
-					FTransform GeomTransform(GetLocalToWorld());
-					BodySetup->AggGeom.DrawAggGeom(PDI, GeomTransform, GetSelectionColor(CollisionColor, bProxyIsSelected, IsHovered()), /*Material=*/NULL, (Owner == NULL), /*bSolid=*/false, UseEditorDepthTest());
-				}
-
-
-				// The simple nav geometry is only used by dynamic obstacles for now
-				if (StaticMesh->NavCollision && StaticMesh->NavCollision->bIsDynamicObstacle)
-				{
-					// Draw the static mesh's body setup (simple collision)
-					FTransform GeomTransform(GetLocalToWorld());
-					FColor NavCollisionColor = FColor(118,84,255,255);
-					StaticMesh->NavCollision->DrawSimpleGeom(PDI, GeomTransform, GetSelectionColor(NavCollisionColor, bProxyIsSelected, IsHovered()));
-				}
-			}
-		}
-	}
-
-#if !(UE_BUILD_SHIPPING || UE_BUILD_TEST)
-	if (View->Family->EngineShowFlags.StaticMeshes)
-	{
-		RenderBounds(PDI, View->Family->EngineShowFlags, GetBounds(), !Owner || IsSelected());
-	}
-#endif
-}
-
 void FStaticMeshSceneProxy::OnTransformChanged()
 {
 	// Update the cached scaling.
diff --git a/Engine/Source/Runtime/Engine/Private/TextRenderComponent.cpp b/Engine/Source/Runtime/Engine/Private/TextRenderComponent.cpp
--- a/Engine/Source/Runtime/Engine/Private/TextRenderComponent.cpp
+++ b/Engine/Source/Runtime/Engine/Private/TextRenderComponent.cpp
@@ -357,7 +357,6 @@ public:
 	virtual ~FTextRenderSceneProxy();
 
 	// Begin FPrimitiveSceneProxy interface
-	virtual void DrawDynamicElements(FPrimitiveDrawInterface* PDI,const FSceneView* View) override;
 	virtual void GetDynamicMeshElements(const TArray<const FSceneView*>& Views, const FSceneViewFamily& ViewFamily, uint32 VisibilityMap, FMeshElementCollector& Collector) const override;
 	virtual void DrawStaticElements(FStaticPrimitiveDrawInterface* PDI) override;
 	virtual FPrimitiveViewRelevance GetViewRelevance(const FSceneView* View) override;
@@ -498,57 +497,6 @@ void FTextRenderSceneProxy::GetDynamicMeshElements(const TArray<const FSceneView
 	}
 }
 
-void FTextRenderSceneProxy::DrawDynamicElements(FPrimitiveDrawInterface* PDI,const FSceneView* View)
-{
-	QUICK_SCOPE_CYCLE_COUNTER( STAT_TextRenderSceneProxy_DrawDynamicElements );
-
-	// Vertex factory will not been initialized when the text string is empty or font is invalid.
-	if(VertexFactory.IsInitialized())
-	{
-		// Draw the mesh.
-		FMeshBatch Mesh;
-		FMeshBatchElement& BatchElement = Mesh.Elements[0];
-		BatchElement.IndexBuffer = &IndexBuffer;
-		Mesh.VertexFactory = &VertexFactory;
-		BatchElement.PrimitiveUniformBufferResource = &GetUniformBuffer();
-		BatchElement.FirstIndex = 0;
-		BatchElement.NumPrimitives = IndexBuffer.Indices.Num() / 3;
-		BatchElement.MinVertexIndex = 0;
-		BatchElement.MaxVertexIndex = VertexBuffer.Vertices.Num() - 1;
-		Mesh.ReverseCulling = IsLocalToWorldDeterminantNegative();
-		Mesh.bDisableBackfaceCulling = false;
-		Mesh.Type = PT_TriangleList;
-		Mesh.DepthPriorityGroup = SDPG_World;
-		const bool bUseSelectedMaterial = GIsEditor && (View->Family->EngineShowFlags.Selection) ? IsSelected() : false;
-		Mesh.MaterialRenderProxy = TextMaterial->GetRenderProxy(bUseSelectedMaterial);
-
-		const bool bIsWireframe = View->Family->EngineShowFlags.Wireframe;
-
-		if (bAlwaysRenderAsText)
-		{ 
-			// Render text unmodified
-			PDI->DrawMesh(Mesh);
-		}
-		else
-		{
-			uint32 NumPasses = DrawRichMesh(
-				PDI,
-				Mesh,
-				FLinearColor(1.0f, 0.0f, 0.0f),	//WireframeColor,
-				FLinearColor(1.0f, 1.0f, 0.0f),	//LevelColor,
-				FLinearColor(1.0f, 1.0f, 1.0f),	//PropertyColor,		
-				this,
-				bUseSelectedMaterial,
-				bIsWireframe
-				);
-		}
-
-#if !(UE_BUILD_SHIPPING || UE_BUILD_TEST)
-		RenderBounds(PDI, View->Family->EngineShowFlags, GetBounds(), IsSelected());
-#endif
-	}
-}
-
 void FTextRenderSceneProxy::DrawStaticElements(FStaticPrimitiveDrawInterface* PDI)
 {
 	// Vertex factory will not been initialized when the font is invalid or the text string is empty.
diff --git a/Engine/Source/Runtime/Engine/Private/UnrealEngine.cpp b/Engine/Source/Runtime/Engine/Private/UnrealEngine.cpp
--- a/Engine/Source/Runtime/Engine/Private/UnrealEngine.cpp
+++ b/Engine/Source/Runtime/Engine/Private/UnrealEngine.cpp
@@ -2421,14 +2421,6 @@ bool UEngine::Exec( UWorld* InWorld, const TCHAR* Cmd, FOutputDevice& Ar )
 	{
 		return HandleTrackParticleRenderingStatsCommand( Cmd, Ar );
 	}
-	else if( FParse::Command(&Cmd,TEXT("DUMPPARTICLERENDERINGSTATS")) )
-	{
-		return HandleDumpParticleRenderingStatsCommand( Cmd, Ar );
-	}
-	else if( FParse::Command(&Cmd,TEXT("DUMPPARTICLEFRAMERENDERINGSTATS")) )
-	{
-		return HandleDumpParticleFrameRenderingStatsCommand( Cmd, Ar );
-	}
 	else if ( FParse::Command(&Cmd,TEXT("DUMPALLOCS")) )
 	{
 		return HandleDumpAllocatorStats( Cmd, Ar );
@@ -5222,21 +5214,6 @@ bool UEngine::HandleTrackParticleRenderingStatsCommand( const TCHAR* Cmd, FOutpu
 	return 1;
 }
 
-bool UEngine::HandleDumpParticleRenderingStatsCommand( const TCHAR* Cmd, FOutputDevice& Ar )
-{
-	extern void DumpParticleRenderingStats(FOutputDevice& Ar);
-	DumpParticleRenderingStats(Ar);
-	return 1;
-}
-
-bool UEngine::HandleDumpParticleFrameRenderingStatsCommand( const TCHAR* Cmd, FOutputDevice& Ar )
-{
-	extern bool GWantsParticleStatsNextFrame;
-	GWantsParticleStatsNextFrame = true;
-	UE_LOG(LogEngine, Warning, TEXT("DUMPPARTICLEFRAMERENDERINGSTATS triggered"));
-	return 1;
-}
-
 bool UEngine::HandleDumpAllocatorStats( const TCHAR* Cmd, FOutputDevice& Ar )
 {
 	GMalloc->DumpAllocatorStats(Ar);
diff --git a/Engine/Source/Runtime/Engine/Private/VectorField.cpp b/Engine/Source/Runtime/Engine/Private/VectorField.cpp
--- a/Engine/Source/Runtime/Engine/Private/VectorField.cpp
+++ b/Engine/Source/Runtime/Engine/Private/VectorField.cpp
@@ -395,22 +395,6 @@ public:
 		}
 	}
 
-	/**
-	 * Draws the bounding box for the vector field.
-	 */
-	virtual void DrawDynamicElements(FPrimitiveDrawInterface* PDI,const FSceneView* View) override
-	{	
-		QUICK_SCOPE_CYCLE_COUNTER( STAT_VectorFieldSceneProxy_DrawDynamicElements );
-
-		DrawVectorFieldBounds(PDI, View, VectorFieldInstance);
-
-		// Draw a visualization of the vectors contained in the field when selected.
-		if (IsSelected() || View->Family->EngineShowFlags.VectorFields)
-		{
-			DrawVectorField(PDI, View, &VisualizationVertexFactory, VectorFieldInstance);
-		}
-	}
-
 	/**
 	 * Computes view relevance for this scene proxy.
 	 */
diff --git a/Engine/Source/Runtime/Engine/Private/VectorFieldVisualization.cpp b/Engine/Source/Runtime/Engine/Private/VectorFieldVisualization.cpp
--- a/Engine/Source/Runtime/Engine/Private/VectorFieldVisualization.cpp
+++ b/Engine/Source/Runtime/Engine/Private/VectorFieldVisualization.cpp
@@ -223,53 +223,6 @@ void DrawVectorFieldBounds(
 	}
 }
 
-/**
- * Draw the vector field for a vector field instance.
- * @param PDI - The primitive drawing interface with which to draw.
- * @param View - The view in which to draw.
- * @param VertexFactory - The vertex factory with which to draw.
- * @param VectorFieldInstance - The vector field instance to draw.
- */
-void DrawVectorField(
-	FPrimitiveDrawInterface* PDI,
-	const FSceneView* View,
-	FVectorFieldVisualizationVertexFactory* VertexFactory,
-	FVectorFieldInstance* VectorFieldInstance )
-{
-	FVectorFieldResource* Resource = VectorFieldInstance->Resource;
-
-	if (Resource && IsValidRef(Resource->VolumeTextureRHI))
-	{
-		// Set up parameters.
-		FVectorFieldVisualizationParameters UniformParameters;
-		UniformParameters.VolumeToWorld = VectorFieldInstance->VolumeToWorld;
-		UniformParameters.VolumeToWorldNoScale = VectorFieldInstance->VolumeToWorldNoScale;
-		UniformParameters.VoxelSize = FVector( 1.0f / Resource->SizeX, 1.0f / Resource->SizeY, 1.0f / Resource->SizeZ );
-		UniformParameters.Scale = VectorFieldInstance->Intensity * Resource->Intensity;
-		VertexFactory->SetParameters(UniformParameters, Resource->VolumeTextureRHI);
-
-		// Create a mesh batch for the visualization.
-		FMeshBatch MeshBatch;
-		MeshBatch.CastShadow = false;
-		MeshBatch.bUseAsOccluder = false;
-		MeshBatch.VertexFactory = VertexFactory;
-		MeshBatch.MaterialRenderProxy = GEngine->LevelColorationUnlitMaterial->GetRenderProxy(false);
-		MeshBatch.Type = PT_LineList;
-
-		// A single mesh element.
-		FMeshBatchElement& MeshElement = MeshBatch.Elements[0];
-		MeshElement.NumPrimitives = 1;
-		MeshElement.NumInstances = Resource->SizeX * Resource->SizeY * Resource->SizeZ;
-		MeshElement.FirstIndex = 0;
-		MeshElement.MinVertexIndex = 0;
-		MeshElement.MaxVertexIndex = 1;
-		MeshElement.PrimitiveUniformBufferResource = &GIdentityPrimitiveUniformBuffer;
-
-		// Draw!
-		PDI->DrawMesh(MeshBatch);
-	}
-}
-
 void GetVectorFieldMesh(
 	FVectorFieldVisualizationVertexFactory* VertexFactory,
 	FVectorFieldInstance* VectorFieldInstance,
diff --git a/Engine/Source/Runtime/Engine/Private/VectorFieldVisualization.h b/Engine/Source/Runtime/Engine/Private/VectorFieldVisualization.h
--- a/Engine/Source/Runtime/Engine/Private/VectorFieldVisualization.h
+++ b/Engine/Source/Runtime/Engine/Private/VectorFieldVisualization.h
@@ -103,19 +103,6 @@ void DrawVectorFieldBounds(
 	const FSceneView* View,
 	FVectorFieldInstance* VectorFieldInstance );
 
-/**
- * Draw the vector field for a vector field instance.
- * @param PDI - The primitive drawing interface with which to draw.
- * @param View - The view in which to draw.
- * @param VertexFactory - The vertex factory with which to draw.
- * @param VectorFieldInstance - The vector field instance to draw.
- */
-void DrawVectorField(
-	FPrimitiveDrawInterface* PDI,
-	const FSceneView* View,
-	FVectorFieldVisualizationVertexFactory* VertexFactory,
-	FVectorFieldInstance* VectorFieldInstance );
-
 void GetVectorFieldMesh(
 	FVectorFieldVisualizationVertexFactory* VertexFactory,
 	FVectorFieldInstance* VectorFieldInstance,
diff --git a/Engine/Source/Runtime/Engine/Public/AI/NavLinkRenderingProxy.h b/Engine/Source/Runtime/Engine/Public/AI/NavLinkRenderingProxy.h
--- a/Engine/Source/Runtime/Engine/Public/AI/NavLinkRenderingProxy.h
+++ b/Engine/Source/Runtime/Engine/Public/AI/NavLinkRenderingProxy.h
@@ -37,7 +37,6 @@ private:
 public:
 	/** Initialization constructor. */
 	FNavLinkRenderingProxy(const UPrimitiveComponent* InComponent);
-	virtual void DrawDynamicElements(FPrimitiveDrawInterface* PDI,const FSceneView* View) override;
 	virtual void GetDynamicMeshElements(const TArray<const FSceneView*>& Views, const FSceneViewFamily& ViewFamily, uint32 VisibilityMap, FMeshElementCollector& Collector) const override;
 	virtual FPrimitiveViewRelevance GetViewRelevance(const FSceneView* View) override;
 	virtual uint32 GetMemoryFootprint( void ) const override;
@@ -46,7 +45,4 @@ public:
 	void StoreSegmentLinks(const FTransform& LocalToWorld, const TArray<FNavigationSegmentLink>& LinksArray);
 
 	static void GetLinkMeshes(const TArray<FNavLinkDrawing>& OffMeshPointLinks, const TArray<FNavLinkSegmentDrawing>& OffMeshSegmentLinks, TArray<float>& StepHeights, FMaterialRenderProxy* const MeshColorInstance, int32 ViewIndex, FMeshElementCollector& Collector, uint32 AgentMask);
-
-	/** made static to allow consistent navlinks drawing even if something is drawing links without FNavLinkRenderingProxy */
-	static void DrawLinks(FPrimitiveDrawInterface* PDI, TArray<FNavLinkDrawing>& OffMeshPointLinks, TArray<FNavLinkSegmentDrawing>& OffMeshSegmentLinks, TArray<float>& StepHeights, FMaterialRenderProxy* const MeshColorInstance, uint32 AgentMask);
 };
diff --git a/Engine/Source/Runtime/Engine/Public/ParticleHelper.h b/Engine/Source/Runtime/Engine/Public/ParticleHelper.h
--- a/Engine/Source/Runtime/Engine/Public/ParticleHelper.h
+++ b/Engine/Source/Runtime/Engine/Public/ParticleHelper.h
@@ -1077,10 +1077,6 @@ struct FAsyncBufferFillData
 	}
 };
 
-// TAsyncBufferFillTasks - handy typedef for an inline array of buffer fill tasks
-typedef TArray<FAsyncBufferFillData, TInlineAllocator<2> > TAsyncBufferFillTasks;
-
-
 /*-----------------------------------------------------------------------------
 	Particle vertex factory pools
 -----------------------------------------------------------------------------*/
@@ -1320,15 +1316,6 @@ struct FDynamicEmitterDataBase
 		ReturnVertexFactory();
 	}
 
-	/**
-	 *	Render thread only draw call
-	 *
-	 *	@param	Proxy		The scene proxy for the particle system that owns this emitter
-	 *	@param	PDI			The primitive draw interface to render with
-	 *	@param	View		The scene view being rendered
-	 */
-	virtual int32 RenderEmitter(FParticleSystemSceneProxy* Proxy, FPrimitiveDrawInterface* PDI,const FSceneView* View) = 0;
-
 	virtual void GetDynamicMeshElementsEmitter(const FParticleSystemSceneProxy* Proxy, const FSceneView* View, const FSceneViewFamily& ViewFamily, int32 ViewIndex, FMeshElementCollector& Collector) const {}
 
 	/**
@@ -1340,17 +1327,6 @@ struct FDynamicEmitterDataBase
 	 */
 	virtual const FMaterialRenderProxy* GetMaterialRenderProxy(bool bSelected) = 0;
 
-	/**
-	 *	Called during InitViews for view processing on scene proxies before rendering them
-	 *  Only called for primitives that are visible and have bDynamicRelevance
-	 *
-	 *	@param	Proxy			The 'owner' particle system scene proxy
-	 *	@param	ViewFamily		The ViewFamily to pre-render for
-	 *	@param	VisibilityMap	A BitArray that indicates whether the primitive was visible in that view (index)
-	 *	@param	FrameNumber		The frame number of this pre-render
-	 */
-	virtual void PreRenderView(FParticleSystemSceneProxy* Proxy, const FSceneViewFamily* ViewFamily, const uint32 VisibilityMap, int32 FrameNumber) {}
-
 	/** Callback from the renderer to gather simple lights that this proxy wants renderered. */
 	virtual void GatherSimpleLights(const FParticleSystemSceneProxy* Proxy, const FSceneViewFamily& ViewFamily, FSimpleLightArray& OutParticleLights) const {}
 
@@ -1558,46 +1534,10 @@ struct FDynamicSpriteEmitterDataBase : public FDynamicEmitterDataBase
 		FGlobalDynamicVertexBuffer::FAllocation* DynamicParameterAllocation,
 		FAsyncBufferFillData& Data) const;
 
-	/**
-	 *	Called to verify that a buffer is ready to use, blocks to wait and can sometimes execute the buffer fill on the current thread
-	 *
-	 *	@param	InView			View to look up in the buffer table
-	 *  @return Completed buffers
-	 */
-	const FAsyncBufferFillData& EnsureFillCompletion(const FSceneView *InView) const
-	{
-		check(AsyncBufferFillTasks.Num());
-		// - 1 because we often fill only one, for _all_ views, if no match we always take the last one
-		int32 TaskIndex = 0;
-		for (; TaskIndex < AsyncBufferFillTasks.Num() - 1; TaskIndex++) 
-		{
-			if (AsyncBufferFillTasks[TaskIndex].View == InView)
-			{
-				break;
-			}
-		}
-		return AsyncBufferFillTasks[TaskIndex];
-	}
-
-	/** Async task is queued for execution */
-	bool									bAsyncTaskOutstanding;
-
-	/** Array of buffers for filling by async task */
-	TAsyncBufferFillTasks					AsyncBufferFillTasks;
-
 	/** The material render proxies for this emitter */
 	const FMaterialRenderProxy*	MaterialResource[2];
 	/** true if the particle emitter utilizes the DynamicParameter module */
 	uint32 bUsesDynamicParameter:1;
-
-	/** Particle instance data allocations. */
-	TArray<FGlobalDynamicVertexBuffer::FAllocation, TInlineAllocator<2> > InstanceDataAllocations;
-
-	/** Particle index data allocations. Used by some but not all child EmitterData variants.*/
-	TArray<FGlobalDynamicIndexBuffer::FAllocation, TInlineAllocator<2> > IndexDataAllocations;
-
-	/** Particle dynamic parameter data allocations. Used by some but not all child EmitterData variants.*/
-	TArray<FGlobalDynamicVertexBuffer::FAllocation, TInlineAllocator<2> > DynamicParameterDataAllocations;
 };
 
 /** Source data for Sprite emitters */
@@ -1689,31 +1629,11 @@ struct FDynamicSpriteEmitterData : public FDynamicSpriteEmitterDataBase
 	 */
 	bool GetVertexAndIndexDataNonInstanced(void* VertexData, void* DynamicParameterVertexData, void* FillIndexData, FParticleOrder* ParticleOrder, const FVector& InCameraPosition, const FMatrix& InLocalToWorld) const;
 
-	/**
-	 *	Called during InitViews for view processing on scene proxies before rendering them
-	 *  Only called for primitives that are visible and have bDynamicRelevance
-	 *
-	 *	@param	Proxy			The 'owner' particle system scene proxy
-	 *	@param	ViewFamily		The ViewFamily to pre-render for
-	 *	@param	VisibilityMap	A BitArray that indicates whether the primitive was visible in that view (index)
-	 *	@param	FrameNumber		The frame number of this pre-render
-	 */
-	virtual void PreRenderView(FParticleSystemSceneProxy* Proxy, const FSceneViewFamily* ViewFamily, const uint32 VisibilityMap, int32 FrameNumber) override;
-
 	/** Gathers simple lights for this emitter. */
 	virtual void GatherSimpleLights(const FParticleSystemSceneProxy* Proxy, const FSceneViewFamily& ViewFamily, FSimpleLightArray& OutParticleLights) const override;
 
 	virtual void GetDynamicMeshElementsEmitter(const FParticleSystemSceneProxy* Proxy, const FSceneView* View, const FSceneViewFamily& ViewFamily, int32 ViewIndex, FMeshElementCollector& Collector) const override;
 
-	/**
-	 *	Render thread only draw call
-	 *
-	 *	@param	Proxy		The scene proxy for the particle system that owns this emitter
-	 *	@param	PDI			The primitive draw interface to render with
-	 *	@param	View		The scene view being rendered
-	 */
-	virtual int32 RenderEmitter(FParticleSystemSceneProxy* Proxy, FPrimitiveDrawInterface* PDI,const FSceneView* View) override;
-
 	/**
 	 *	Create the render thread resources for this emitter data
 	 *
@@ -1725,15 +1645,6 @@ struct FDynamicSpriteEmitterData : public FDynamicSpriteEmitterDataBase
 
 	virtual FParticleVertexFactoryBase* BuildVertexFactory(const FParticleSystemSceneProxy* InOwnerProxy);
 
-	/**
-	 *	Release the render thread resources for this emitter data
-	 *
-	 *	@param	InOwnerProxy	The proxy that owns this dynamic emitter data
-	 *
-	 *	@return	bool			true if successful, false if failed
-	 */
-	virtual void ReleaseRenderThreadResources(const FParticleSystemSceneProxy* InOwnerProxy);
-
 	/** Returns the source data for this particle system */
 	virtual const FDynamicEmitterReplayDataBase& GetSource() const
 	{
@@ -1747,9 +1658,6 @@ struct FDynamicSpriteEmitterData : public FDynamicSpriteEmitterDataBase
 
 	/** Uniform parameters. Most fields are filled in when updates are sent to the rendering thread, some are per-view! */
 	FParticleSpriteUniformParameters UniformParameters;
-		
-	/** Per-view sprite uniform buffers. */
-	TArray<FParticleSpriteUniformBufferRef, TInlineAllocator<2> > PerViewUniformBuffers;
 };
 
 /** Source data for Mesh emitters */
@@ -1834,15 +1742,6 @@ struct FDynamicMeshEmitterData : public FDynamicSpriteEmitterDataBase
 
 	virtual void GetDynamicMeshElementsEmitter(const FParticleSystemSceneProxy* Proxy, const FSceneView* View, const FSceneViewFamily& ViewFamily, int32 ViewIndex, FMeshElementCollector& Collector) const override;
 
-	/**
-	 *	Render thread only draw call
-	 *
-	 *	@param	Proxy		The scene proxy for the particle system that owns this emitter
-	 *	@param	PDI			The primitive draw interface to render with
-	 *	@param	View		The scene view being rendered
-	 */
-	virtual int32 RenderEmitter(FParticleSystemSceneProxy* Proxy, FPrimitiveDrawInterface* PDI,const FSceneView* View);
-	
 	/**
 	 *	Retrieve the instance data required to render this emitter.
 	 *	Render-thread only
@@ -1869,17 +1768,6 @@ struct FDynamicMeshEmitterData : public FDynamicSpriteEmitterDataBase
 	void GetParticleTransform(const FBaseParticle& InParticle, const FVector& CameraPosition, const FVector& CameraFacingOpVector, 
 		const FQuat& PointToLockedAxis, const FParticleSystemSceneProxy* Proxy, const FSceneView* View, FMatrix& OutTransformMat) const;
 
-	/**
-	 *	Called during InitViews for view processing on scene proxies before rendering them
-	 *  Only called for primitives that are visible and have bDynamicRelevance
-	 *
-	 *	@param	Proxy			The 'owner' particle system scene proxy
-	 *	@param	ViewFamily		The ViewFamily to pre-render for
-	 *	@param	VisibilityMap	A BitArray that indicates whether the primitive was visible in that view (index)
-	 *	@param	FrameNumber		The frame number of this pre-render
-	 */
-	virtual void PreRenderView(FParticleSystemSceneProxy* Proxy, const FSceneViewFamily* ViewFamily, const uint32 VisibilityMap, int32 FrameNumber); 
-
 	/** Gathers simple lights for this emitter. */
 	virtual void GatherSimpleLights(const FParticleSystemSceneProxy* Proxy, const FSceneViewFamily& ViewFamily, FSimpleLightArray& OutParticleLights) const override;
 
@@ -1926,21 +1814,6 @@ struct FDynamicMeshEmitterData : public FDynamicSpriteEmitterDataBase
 	UStaticMesh*		StaticMesh;
 	TArray<UMaterialInterface*, TInlineAllocator<2> > MeshMaterials;
 
-	/** Mesh batches used for rendering, built in PreRenderView. */
-	TArray<FMeshBatch*, TInlineAllocator<4> > MeshBatches;
-
-	/** Mesh batch parameters used when instancing is not allowed. */
-	TArray<FMeshParticleVertexFactory::FBatchParametersCPU, TInlineAllocator<4> > MeshBatchParameters;
-
-	/** The first mesh batches to render for a given view. */
-	TArray<int32, TInlineAllocator<4> > FirstBatchForView;
-
-	/** Particle instance data allocations (ES2). */
-	TArray<FMeshParticleInstanceVertex> InstanceDataAllocationsCPU;
-
-	/** Particle dynamic parameter data allocations (ES2). */
-	TArray<FMeshParticleInstanceVertexDynamicParameter> DynamicParameterDataAllocationsCPU;
-
 	/** mesh particle uniform buffer reference. */
 	FMeshParticleUniformBufferRef UniformBuffer;
 
@@ -2133,26 +2006,6 @@ struct FDynamicBeam2EmitterData : public FDynamicSpriteEmitterDataBase
 	
 	virtual FParticleVertexFactoryBase* BuildVertexFactory(const FParticleSystemSceneProxy* InOwnerProxy);
 
-	/**
-	 *	Render thread only draw call
-	 *
-	 *	@param	Proxy		The scene proxy for the particle system that owns this emitter
-	 *	@param	PDI			The primitive draw interface to render with
-	 *	@param	View		The scene view being rendered
-	 */
-	virtual int32 RenderEmitter(FParticleSystemSceneProxy* Proxy, FPrimitiveDrawInterface* PDI,const FSceneView* View);
-
-	/**
-	 *	Called during InitViews for view processing on scene proxies before rendering them
-	 *  Only called for primitives that are visible and have bDynamicRelevance
-	 *
-	 *	@param	Proxy			The 'owner' particle system scene proxy
-	 *	@param	ViewFamily		The ViewFamily to pre-render for
-	 *	@param	VisibilityMap	A BitArray that indicates whether the primitive was visible in that view (index)
-	 *	@param	FrameNumber		The frame number of this pre-render
-	 */
-	virtual void PreRenderView(FParticleSystemSceneProxy* Proxy, const FSceneViewFamily* ViewFamily, const uint32 VisibilityMap, int32 FrameNumber);
-
 	// Debugging functions
 	virtual void RenderDirectLine(const FParticleSystemSceneProxy* Proxy, FPrimitiveDrawInterface* PDI,const FSceneView* View) const;
 	virtual void RenderLines(const FParticleSystemSceneProxy* Proxy, FPrimitiveDrawInterface* PDI,const FSceneView* View) const;
@@ -2297,31 +2150,6 @@ struct FDynamicTrailsEmitterData : public FDynamicSpriteEmitterDataBase
 
 	virtual void GetDynamicMeshElementsEmitter(const FParticleSystemSceneProxy* Proxy, const FSceneView* View, const FSceneViewFamily& ViewFamily, int32 ViewIndex, FMeshElementCollector& Collector) const override;
 
-	/**
-	 *	Render thread only draw call
-	 *
-	 *	@param	Proxy		The scene proxy for the particle system that owns this emitter
-	 *	@param	PDI			The primitive draw interface to render with
-	 *	@param	View		The scene view being rendered
-	 */
-	virtual int32 RenderEmitter(FParticleSystemSceneProxy* Proxy, FPrimitiveDrawInterface* PDI, const FSceneView* View);
-
-	virtual bool ShouldUsePrerenderView()
-	{
-		return true;
-	}
-
-	/**
-	 *	Called during InitViews for view processing on scene proxies before rendering them
-	 *  Only called for primitives that are visible and have bDynamicRelevance
-	 *
-	 *	@param	Proxy			The 'owner' particle system scene proxy
-	 *	@param	ViewFamily		The ViewFamily to pre-render for
-	 *	@param	VisibilityMap	A BitArray that indicates whether the primitive was visible in that view (index)
-	 *	@param	FrameNumber		The frame number of this pre-render
-	 */
-	virtual void PreRenderView(FParticleSystemSceneProxy* Proxy, const FSceneViewFamily* ViewFamily, const uint32 VisibilityMap, int32 FrameNumber);
-
 	virtual void RenderDebug(const FParticleSystemSceneProxy* Proxy, FPrimitiveDrawInterface* PDI, const FSceneView* View, bool bCrosses) const;
 
 	// Data fill functions
@@ -2404,7 +2232,6 @@ struct FDynamicRibbonEmitterData : public FDynamicTrailsEmitterData
 	/** Initialize this emitter's dynamic rendering data, called after source data has been filled in */
 	virtual void Init(bool bInSelected);
 
-	virtual bool ShouldUsePrerenderView();
 	virtual void RenderDebug(const FParticleSystemSceneProxy* Proxy, FPrimitiveDrawInterface* PDI, const FSceneView* View, bool bCrosses) const;
 
 	// Data fill functions
@@ -2517,30 +2344,11 @@ public:
 	{
 		return false;
 	}
-	virtual void DrawDynamicElements(FPrimitiveDrawInterface* PDI,const FSceneView* View) override;
 	virtual void GetDynamicMeshElements(const TArray<const FSceneView*>& Views, const FSceneViewFamily& ViewFamily, uint32 VisibilityMap, FMeshElementCollector& Collector) const override;
 	virtual FPrimitiveViewRelevance GetViewRelevance(const FSceneView* View) override;
 	virtual void OnActorPositionChanged() override;
 	virtual void OnTransformChanged() override;
 
-	/**
-	 *	Helper function for determining the LOD distance for a given view.
-	 *
-	 *	@param	View			The view of interest.
-	 *	@param	FrameNumber		The frame number being rendered.
-	 */
-	void ProcessPreRenderView(const FSceneView* View, int32 FrameNumber);
-
-	/**
-	 *	Called during InitViews for view processing on scene proxies before rendering them
-	 *  Only called for primitives that are visible and have bDynamicRelevance
-	 *
-	 *	@param	ViewFamily		The ViewFamily to pre-render for
-	 *	@param	VisibilityMap	A BitArray that indicates whether the primitive was visible in that view (index)
-	 *	@param	FrameNumber		The frame number of this pre-render
-	 */
-	virtual void PreRenderView(const FSceneViewFamily* ViewFamily, const uint32 VisibilityMap, int32 FrameNumber);
-
 	/** Gathers simple lights for this emitter. */
 	virtual void GatherSimpleLights(const FSceneViewFamily& ViewFamily, FSimpleLightArray& OutParticleLights) const override;
 
@@ -2666,14 +2474,6 @@ public:
 		return !MaterialRelevance.bDisableDepthTest;
 	}
 	
-	/** 
-	 * Draw the scene proxy as a dynamic element
-	 *
-	 * @param	PDI - draw interface to render to
-	 * @param	View - current view
-	 */
-	virtual void DrawDynamicElements(FPrimitiveDrawInterface* PDI,const FSceneView* View);
-
 	/**
 	 *	Returns whether the proxy utilizes custom occlusion bounds or not
 	 *
diff --git a/Engine/Source/Runtime/Engine/Public/PrimitiveSceneProxy.h b/Engine/Source/Runtime/Engine/Public/PrimitiveSceneProxy.h
--- a/Engine/Source/Runtime/Engine/Public/PrimitiveSceneProxy.h
+++ b/Engine/Source/Runtime/Engine/Public/PrimitiveSceneProxy.h
@@ -154,25 +154,6 @@ public:
 	 */
 	virtual void DrawStaticElements(FStaticPrimitiveDrawInterface* PDI) {}
 
-	/**
-	 * Draws the primitive's dynamic elements.  This is called from the rendering thread for each frame of each view.
-	 * The dynamic elements will only be rendered if GetViewRelevance declares dynamic relevance.
-	 * Called in the rendering thread.
-	 * @param PDI - The interface which receives the primitive elements.
-	 * @param View - The view which is being rendered.
-	 */
-	virtual void DrawDynamicElements(FPrimitiveDrawInterface* PDI, const FSceneView* View) {}
-
-	/**
-	 * Draws the primitive's dynamic elements.  This is called from the rendering thread for each frame of each view.
-	 * The dynamic elements will only be rendered if GetViewRelevance declares dynamic relevance.
-	 * Called in the rendering thread.
-	 * @param PDI - The interface which receives the primitive elements.
-	 * @param View - The view which is being rendered.
-	 * @param Flags - Optional flags
-	 */
-	virtual void DrawDynamicElements(FPrimitiveDrawInterface* PDI, const FSceneView* View, uint32 DrawDynamicFlags ) { DrawDynamicElements( PDI, View ); }
-
 	/** 
 	 * Gathers the primitive's dynamic mesh elements.  This will only be called if GetViewRelevance declares dynamic relevance.
 	 * This is called from the rendering thread for each set of views that might be rendered.  
@@ -195,16 +176,6 @@ public:
 	 */
 	ENGINE_API virtual FPrimitiveViewRelevance GetViewRelevance(const FSceneView* View);
 
-	/**
-	 *	Called during InitViews for view processing on scene proxies before rendering them
-	 *  Only called for primitives that are visible and have bDynamicRelevance
-	 *
-	 *	@param	ViewFamily		The ViewFamily to pre-render for
-	 *	@param	VisibilityMap	A BitArray that indicates whether the primitive was visible in that view (index)
-	 *	@param	FrameNumber		The frame number of this pre-render
-	 */
-	virtual void PreRenderView(const FSceneViewFamily* ViewFamily, const uint32 VisibilityMap, int32 FrameNumber) {}
-
 	/** Callback from the renderer to gather simple lights that this proxy wants renderered. */
 	virtual void GatherSimpleLights(const FSceneViewFamily& ViewFamily, FSimpleLightArray& OutParticleLights) const {}
 
diff --git a/Engine/Source/Runtime/Engine/Public/PrimitiveViewRelevance.h b/Engine/Source/Runtime/Engine/Public/PrimitiveViewRelevance.h
--- a/Engine/Source/Runtime/Engine/Public/PrimitiveViewRelevance.h
+++ b/Engine/Source/Runtime/Engine/Public/PrimitiveViewRelevance.h
@@ -24,8 +24,6 @@ struct FPrimitiveViewRelevance
 	uint32 bEditorPrimitiveRelevance : 1;
 	/** The primitive is drawn only in the editor and composited onto the scene after post processing using no depth testing */
 	uint32 bEditorNoDepthTestPrimitiveRelevance : 1;
-	/** The primitive needs PreRenderView to be called before rendering. */
-	uint32 bNeedsPreRenderView : 1;
 	/** The primitive should have GatherSimpleLights called on the proxy when gathering simple lights. */
 	uint32 bHasSimpleLights : 1;
 
@@ -64,7 +62,6 @@ struct FPrimitiveViewRelevance
 		bRenderInMainPass(true),
 		bEditorPrimitiveRelevance(false),
 		bEditorNoDepthTestPrimitiveRelevance(false),
-		bNeedsPreRenderView(false),
 		bHasSimpleLights(false),
 		bOpaqueRelevance(true),
 		bMaskedRelevance(false),
@@ -89,7 +86,6 @@ struct FPrimitiveViewRelevance
 		bRenderInMainPass |= B.bRenderInMainPass !=0;
 		bEditorPrimitiveRelevance |= B.bEditorPrimitiveRelevance !=0;
 		bEditorNoDepthTestPrimitiveRelevance |= B.bEditorNoDepthTestPrimitiveRelevance !=0;
-		bNeedsPreRenderView |= B.bNeedsPreRenderView != 0;
 		bHasSimpleLights |= B.bHasSimpleLights != 0;
 		bSeparateTranslucencyRelevance |= B.bSeparateTranslucencyRelevance != 0;
 		bNormalTranslucencyRelevance |= B.bNormalTranslucencyRelevance != 0;
diff --git a/Engine/Source/Runtime/Engine/Public/SceneManagement.h b/Engine/Source/Runtime/Engine/Public/SceneManagement.h
--- a/Engine/Source/Runtime/Engine/Public/SceneManagement.h
+++ b/Engine/Source/Runtime/Engine/Public/SceneManagement.h
@@ -1063,25 +1063,6 @@ public:
 		uint8 DepthPriorityGroup
 		) = 0;
 
-	/**
-	 * Determines whether a particular material will be ignored in this context.
-	 * @param MaterialRenderProxy - The render proxy of the material to check.
-	 * @param InFeatureLevel - The feature level we are currently rendering at
-	 * @return true if meshes using the material will be ignored in this context.
-	 */
-	virtual bool IsMaterialIgnored(const FMaterialRenderProxy* MaterialRenderProxy, ERHIFeatureLevel::Type InFeatureLevel) const
-	{
-		return false;
-	}
-
-	/**
-	 * @returns true if this PDI is rendering for the selection outline post process.
-	 */
-	virtual bool IsRenderingSelectionOutline() const
-	{
-		return false;
-	}
-
 	/**
 	 * Draw a mesh element.
 	 * This should only be called through the DrawMesh function.
@@ -1192,34 +1173,6 @@ public:
 		return 0;
 	}
 
-	// Legacy, should not be used
-	virtual bool IsMaterialIgnored(const FMaterialRenderProxy* MaterialRenderProxy, ERHIFeatureLevel::Type InFeatureLevel) const
-	{
-		static bool bTriggered = false;
-
-		if (!bTriggered)
-		{
-			bTriggered = true;
-			ensureMsg(false, TEXT("FSimpleElementCollector::IsMaterialIgnored called"));
-		}
-
-		return false;
-	}
-
-	// Legacy, should not be used
-	virtual bool IsRenderingSelectionOutline() const
-	{
-		static bool bTriggered = false;
-
-		if (!bTriggered)
-		{
-			bTriggered = true;
-			ensureMsg(false, TEXT("FSimpleElementCollector::IsRenderingSelectionOutline called"));
-		}
-
-		return false;
-	}
-
 	void DrawBatchedElements(FRHICommandList& RHICmdList, const FSceneView& View, FTexture2DRHIRef DepthTexture, EBlendModeFilter::Type Filter) const;
 
 	/** The batched simple elements. */
@@ -1744,8 +1697,7 @@ namespace EVertexColorViewMode
 extern ENGINE_API EVertexColorViewMode::Type GVertexColorViewMode;
 
 /**
- * Returns true if the given view is "rich".  Rich means that calling DrawRichMesh for the view will result in a modified draw call
- * being made.
+ * Returns true if the given view is "rich", and all primitives should be forced down the dynamic drawing path so that ApplyViewModeOverrides can implement the rich view feature.
  * A view is rich if is missing the EngineShowFlags.Materials showflag, or has any of the render mode affecting showflags.
  */
 extern ENGINE_API bool IsRichView(const FSceneViewFamily& ViewFamily);
@@ -1770,31 +1722,6 @@ FORCEINLINE void EmitMeshDrawEvents(FRHICommandList& RHICmdList, const class FPr
 #endif
 }
 
-/**
- * Draws a mesh, modifying the material which is used depending on the view's show flags.
- * Meshes with materials irrelevant to the pass which the mesh is being drawn for may be entirely ignored.
- *
- * @param PDI - The primitive draw interface to draw the mesh on.
- * @param Mesh - The mesh to draw.
- * @param WireframeColor - The color which is used when rendering the mesh with EngineShowFlags.Wireframe.
- * @param LevelColor - The color which is used when rendering the mesh with EngineShowFlags.LevelColoration.
- * @param PropertyColor - The color to use when rendering the mesh with EngineShowFlags.PropertyColoration.
- * @param PrimitiveInfo - The FScene information about the UPrimitiveComponent.
- * @param bSelected - True if the primitive is selected.
- * @param ExtraDrawFlags - optional flags to override the view family show flags when rendering
- * @return Number of passes rendered for the mesh
- */
-extern ENGINE_API int32 DrawRichMesh(
-	FPrimitiveDrawInterface* PDI,
-	const struct FMeshBatch& Mesh,
-	const FLinearColor& WireframeColor,
-	const FLinearColor& LevelColor,
-	const FLinearColor& PropertyColor,
-	const FPrimitiveSceneProxy* PrimitiveSceneProxy,
-	bool bSelected,
-	bool bDrawInWireframe = false
-	);
-
 extern ENGINE_API void ApplyViewModeOverrides(
 	int32 ViewIndex,
 	const FEngineShowFlags& EngineShowFlags,
diff --git a/Engine/Source/Runtime/Engine/Public/SkeletalMeshTypes.h b/Engine/Source/Runtime/Engine/Public/SkeletalMeshTypes.h
--- a/Engine/Source/Runtime/Engine/Public/SkeletalMeshTypes.h
+++ b/Engine/Source/Runtime/Engine/Public/SkeletalMeshTypes.h
@@ -1803,11 +1803,9 @@ public:
 #if WITH_EDITOR
 	virtual HHitProxy* CreateHitProxies(UPrimitiveComponent* Component, TArray<TRefCountPtr<HHitProxy> >& OutHitProxies) override;
 #endif
-	virtual void DrawDynamicElements(FPrimitiveDrawInterface* PDI, const FSceneView* View) override;
 	virtual void GetDynamicMeshElements(const TArray<const FSceneView*>& Views, const FSceneViewFamily& ViewFamily, uint32 VisibilityMap, FMeshElementCollector& Collector) const override;
 	virtual FPrimitiveViewRelevance GetViewRelevance(const FSceneView* View) override;
 	virtual bool CanBeOccluded() const override;
-	virtual void PreRenderView(const FSceneViewFamily* ViewFamily, const uint32 VisibilityMap, int32 FrameNumber) override;
 	
 	/**
 	 * Returns the world transform to use for drawing.
@@ -1829,7 +1827,7 @@ public:
 	/** 
 	 * Render physics asset for debug display
 	 */
-	void DebugDrawPhysicsAsset(FPrimitiveDrawInterface* PDI, const FEngineShowFlags& EngineShowFlags) const;
+	void DebugDrawPhysicsAsset(int32 ViewIndex, FMeshElementCollector& Collector, const FEngineShowFlags& EngineShowFlags) const;
 
 	virtual uint32 GetMemoryFootprint( void ) const { return( sizeof( *this ) + GetAllocatedSize() ); }
 	uint32 GetAllocatedSize( void ) const { return( FPrimitiveSceneProxy::GetAllocatedSize() + LODSections.GetAllocatedSize() ); }
@@ -1897,20 +1895,6 @@ protected:
 	TSet<UMaterialInterface*> MaterialsInUse_GameThread;
 	bool bMaterialsNeedMorphUsage_GameThread;
 	
-	/**
-	* Draw only the section of the scene proxy as a dynamic element
-	* 
-	* @param	PDI - draw interface to render to
-	* @param	View - current view
-	* @param	const FStaticLODModel& LODModel - LODModel 
-	* @param	const FSkelMeshSection& Section - Section
-	* @param	const FSkelMeshChunk& Chunk - Chunk
-	* @param	const FSectionElementInfo& SectionElementInfo - SectionElementInfo - material ID
-	*/
-	void DrawDynamicElementsSection(FPrimitiveDrawInterface* PDI,const FSceneView* View,
-		const FStaticLODModel& LODModel, const int32 LODIndex, const FSkelMeshSection& Section, 
-		const FSkelMeshChunk& Chunk, const FSectionElementInfo& SectionElementInfo, const FTwoVectors& CustomLeftRightVectors );
-
 	void GetDynamicElementsSection(const TArray<const FSceneView*>& Views, const FSceneViewFamily& ViewFamily, uint32 VisibilityMap, 
 		const FStaticLODModel& LODModel, const int32 LODIndex, const FSkelMeshSection& Section, const FSkelMeshChunk& Chunk, 
 		const FSectionElementInfo& SectionElementInfo, const FTwoVectors& CustomLeftRightVectors, bool bSelectable, FMeshElementCollector& Collector ) const;
diff --git a/Engine/Source/Runtime/Engine/Public/StaticMeshResources.h b/Engine/Source/Runtime/Engine/Public/StaticMeshResources.h
--- a/Engine/Source/Runtime/Engine/Public/StaticMeshResources.h
+++ b/Engine/Source/Runtime/Engine/Public/StaticMeshResources.h
@@ -770,8 +770,6 @@ public:
 	virtual HHitProxy* CreateHitProxies(UPrimitiveComponent* Component, TArray<TRefCountPtr<HHitProxy> >& OutHitProxies) override;
 #endif
 	virtual void DrawStaticElements(FStaticPrimitiveDrawInterface* PDI) override;
-	virtual void DrawDynamicElements(FPrimitiveDrawInterface* PDI,const FSceneView* View) override { DrawDynamicElements( PDI, View, 0 ); }
-	virtual void DrawDynamicElements(FPrimitiveDrawInterface* PDI,const FSceneView* View, uint32 DrawDynamicFlags ) override;
 	virtual void OnTransformChanged() override;
 	virtual int32 GetLOD(const FSceneView* View) const override;
 	virtual FPrimitiveViewRelevance GetViewRelevance(const FSceneView* View) override;
diff --git a/Engine/Source/Runtime/Landscape/Private/LandscapeGizmoActor.cpp b/Engine/Source/Runtime/Landscape/Private/LandscapeGizmoActor.cpp
--- a/Engine/Source/Runtime/Landscape/Private/LandscapeGizmoActor.cpp
+++ b/Engine/Source/Runtime/Landscape/Private/LandscapeGizmoActor.cpp
@@ -303,93 +303,6 @@ public:
 #endif
 	};
 
-	virtual void DrawDynamicElements(FPrimitiveDrawInterface* PDI,const FSceneView* View)
-	{
-		//FMemMark Mark(FMemStack::Get());
-#if WITH_EDITOR
-		if( GizmoRenderProxy &&  HeightmapRenderProxy )
-		{
-			// Axis
-			PDI->DrawLine( Origin, XAxis, FLinearColor(1, 0, 0), SDPG_World );
-			PDI->DrawLine( Origin, YAxis, FLinearColor(0, 1, 0), SDPG_World );
-
-			{
-				FDynamicMeshBuilder MeshBuilder;
-
-				MeshBuilder.AddVertex(FrustumVerts[0], FVector2D(0, 0), FVector(1,0,0), FVector(0,1,0), FVector(0,0,1), FColor(255,255,255));
-				MeshBuilder.AddVertex(FrustumVerts[1], FVector2D(1, 0), FVector(1,0,0), FVector(0,1,0), FVector(0,0,1), FColor(255,255,255));
-				MeshBuilder.AddVertex(FrustumVerts[2], FVector2D(1, 1), FVector(1,0,0), FVector(0,1,0), FVector(0,0,1), FColor(255,255,255));
-				MeshBuilder.AddVertex(FrustumVerts[3], FVector2D(0, 1), FVector(1,0,0), FVector(0,1,0), FVector(0,0,1), FColor(255,255,255));
-
-				MeshBuilder.AddVertex(FrustumVerts[4], FVector2D(0, 0), FVector(1,0,0), FVector(0,1,0), FVector(0,0,1), FColor(255,255,255));
-				MeshBuilder.AddVertex(FrustumVerts[5], FVector2D(1, 0), FVector(1,0,0), FVector(0,1,0), FVector(0,0,1), FColor(255,255,255));
-				MeshBuilder.AddVertex(FrustumVerts[6], FVector2D(1, 1), FVector(1,0,0), FVector(0,1,0), FVector(0,0,1), FColor(255,255,255));
-				MeshBuilder.AddVertex(FrustumVerts[7], FVector2D(0, 1), FVector(1,0,0), FVector(0,1,0), FVector(0,0,1), FColor(255,255,255));
-
-				MeshBuilder.AddVertex(FrustumVerts[1], FVector2D(0, 0), FVector(1,0,0), FVector(0,1,0), FVector(0,0,1), FColor(255,255,255));
-				MeshBuilder.AddVertex(FrustumVerts[0], FVector2D(1, 0), FVector(1,0,0), FVector(0,1,0), FVector(0,0,1), FColor(255,255,255));
-				MeshBuilder.AddVertex(FrustumVerts[4], FVector2D(1, 1), FVector(1,0,0), FVector(0,1,0), FVector(0,0,1), FColor(255,255,255));
-				MeshBuilder.AddVertex(FrustumVerts[5], FVector2D(0, 1), FVector(1,0,0), FVector(0,1,0), FVector(0,0,1), FColor(255,255,255));
-
-				MeshBuilder.AddVertex(FrustumVerts[3], FVector2D(0, 0), FVector(1,0,0), FVector(0,1,0), FVector(0,0,1), FColor(255,255,255));
-				MeshBuilder.AddVertex(FrustumVerts[2], FVector2D(1, 0), FVector(1,0,0), FVector(0,1,0), FVector(0,0,1), FColor(255,255,255));
-				MeshBuilder.AddVertex(FrustumVerts[6], FVector2D(1, 1), FVector(1,0,0), FVector(0,1,0), FVector(0,0,1), FColor(255,255,255));
-				MeshBuilder.AddVertex(FrustumVerts[7], FVector2D(0, 1), FVector(1,0,0), FVector(0,1,0), FVector(0,0,1), FColor(255,255,255));
-
-				MeshBuilder.AddVertex(FrustumVerts[2], FVector2D(0, 0), FVector(1,0,0), FVector(0,1,0), FVector(0,0,1), FColor(255,255,255));
-				MeshBuilder.AddVertex(FrustumVerts[1], FVector2D(1, 0), FVector(1,0,0), FVector(0,1,0), FVector(0,0,1), FColor(255,255,255));
-				MeshBuilder.AddVertex(FrustumVerts[5], FVector2D(1, 1), FVector(1,0,0), FVector(0,1,0), FVector(0,0,1), FColor(255,255,255));
-				MeshBuilder.AddVertex(FrustumVerts[6], FVector2D(0, 1), FVector(1,0,0), FVector(0,1,0), FVector(0,0,1), FColor(255,255,255));
-
-				MeshBuilder.AddVertex(FrustumVerts[0], FVector2D(0, 0), FVector(1,0,0), FVector(0,1,0), FVector(0,0,1), FColor(255,255,255));
-				MeshBuilder.AddVertex(FrustumVerts[3], FVector2D(1, 0), FVector(1,0,0), FVector(0,1,0), FVector(0,0,1), FColor(255,255,255));
-				MeshBuilder.AddVertex(FrustumVerts[7], FVector2D(1, 1), FVector(1,0,0), FVector(0,1,0), FVector(0,0,1), FColor(255,255,255));
-				MeshBuilder.AddVertex(FrustumVerts[4], FVector2D(0, 1), FVector(1,0,0), FVector(0,1,0), FVector(0,0,1), FColor(255,255,255));
-
-				for (int32 i = 0; i < 6; ++i)
-				{
-					int32 Idx = i*4;
-					MeshBuilder.AddTriangle( Idx, Idx+2, Idx+1 );
-					MeshBuilder.AddTriangle( Idx, Idx+3, Idx+2 );
-				}
-
-				MeshBuilder.Draw(PDI, FMatrix::Identity, GizmoRenderProxy , SDPG_World, true);
-			}
-
-			if (bHeightmapRendering)
-			{		  		
-				FDynamicMeshBuilder MeshBuilder;
-
-				for (int32 Y = 0; Y < SampleSizeY; ++Y)
-				{
-					for (int32 X = 0; X < SampleSizeX; ++X)
-					{
-						FVector SampledNormal = SampledNormals[X + Y * SampleSizeX];
-						FVector TangentX(SampledNormal.Z, 0, -SampledNormal.X);
-						TangentX = TangentX.SafeNormal();
-
-						MeshBuilder.AddVertex(SampledPositions[X + Y * SampleSizeX], FVector2D((float)X / (SampleSizeX), (float)Y / (SampleSizeY)), TangentX, SampledNormal^TangentX, SampledNormal, FColor(255, 255, 255) );
-					}
-				}
-
-				for (int32 Y = 0; Y < SampleSizeY; ++Y)
-				{
-					for (int32 X = 0; X < SampleSizeX; ++X)
-					{
-						if (X < SampleSizeX - 1 && Y < SampleSizeY - 1)
-						{
-							MeshBuilder.AddTriangle( (X+0) + (Y+0) * SampleSizeX, (X+1) + (Y+1) * SampleSizeX, (X+1) + (Y+0) * SampleSizeX );
-							MeshBuilder.AddTriangle( (X+0) + (Y+0) * SampleSizeX, (X+0) + (Y+1) * SampleSizeX, (X+1) + (Y+1) * SampleSizeX );
-						}
-					}
-				}
-
-				MeshBuilder.Draw(PDI, MeshRT, HeightmapRenderProxy, SDPG_World, false);
-			}
-		}
-#endif
-	};
-
 	virtual FPrimitiveViewRelevance GetViewRelevance(const FSceneView* View)
 	{
 		FPrimitiveViewRelevance Result;
diff --git a/Engine/Source/Runtime/Landscape/Private/LandscapeRender.cpp b/Engine/Source/Runtime/Landscape/Private/LandscapeRender.cpp
--- a/Engine/Source/Runtime/Landscape/Private/LandscapeRender.cpp
+++ b/Engine/Source/Runtime/Landscape/Private/LandscapeRender.cpp
@@ -662,40 +662,6 @@ FLandscapeComponentSceneProxy::FLandscapeComponentSceneProxy(ULandscapeComponent
 
 	bRequiresAdjacencyInformation = RequiresAdjacencyInformation(MaterialInterface, XYOffsetmapTexture == NULL ? &FLandscapeVertexFactory::StaticType : &FLandscapeXYOffsetVertexFactory::StaticType, InComponent->GetWorld()->FeatureLevel);
 	SharedBuffersKey = (SubsectionSizeQuads & 0xffff) | ((NumSubsections & 0xf) << 16) | (FeatureLevel <= ERHIFeatureLevel::ES3_1 ? 0 : 1<<20) | (XYOffsetmapTexture == NULL ? 0 : 1 << 31);
-
-	DynamicMesh.LCI = ComponentLightInfo; 
-	DynamicMesh.CastShadow = true;
-	DynamicMesh.Type = PT_TriangleList;
-	DynamicMeshBatchParamArray.Empty(NumSubsections*NumSubsections);
-
-#if WITH_EDITOR
-	DynamicMeshTools = DynamicMesh;
-	DynamicMeshTools.Type = PT_TriangleList;
-	DynamicMeshTools.CastShadow = false;
-#endif
-
-	for( int32 SubY=0;SubY<NumSubsections;SubY++ )
-	{
-		for( int32 SubX=0;SubX<NumSubsections;SubX++ )
-		{
-			FMeshBatchElement* BatchElement = (SubX==0 && SubY==0) ? DynamicMesh.Elements.GetData() : new(DynamicMesh.Elements) FMeshBatchElement;
-			BatchElement->PrimitiveUniformBufferResource = &GetUniformBuffer();
-			FLandscapeBatchElementParams* BatchElementParams = new(DynamicMeshBatchParamArray) FLandscapeBatchElementParams;
-			BatchElementParams->LocalToWorldNoScalingPtr = &LocalToWorldNoScaling;
-			BatchElement->UserData = BatchElementParams;
-
-			int32 SubSectionIdx = SubX + SubY*NumSubsections;
-
-			BatchElementParams->SceneProxy = this;
-			BatchElementParams->SubX = SubX;
-			BatchElementParams->SubY = SubY;
-
-#if WITH_EDITOR
-			FMeshBatchElement* ToolBatchElement = (SubX==0 && SubY==0) ? DynamicMeshTools.Elements.GetData() : new(DynamicMeshTools.Elements) FMeshBatchElement;
-			*ToolBatchElement = *BatchElement;
-#endif
-		}
-	}
 }
 
 void FLandscapeComponentSceneProxy::CreateRenderThreadResources()
@@ -760,20 +726,9 @@ void FLandscapeComponentSceneProxy::CreateRenderThreadResources()
 
 	// Assign vertex factory
 	VertexFactory = SharedBuffers->VertexFactory;
-	DynamicMesh.VertexFactory = VertexFactory;
-#if WITH_EDITOR
-	DynamicMeshTools.VertexFactory = VertexFactory;
-#endif
 	
 	// Assign LandscapeUniformShaderParameters
 	LandscapeUniformShaderParameters.InitResource();
-
-	for( int32 ElementIdx=0;ElementIdx<DynamicMesh.Elements.Num();ElementIdx++ )
-	{
-		FMeshBatchElement& BatchElement = DynamicMesh.Elements[ElementIdx];
-		FLandscapeBatchElementParams* BatchElementParams = (FLandscapeBatchElementParams*)BatchElement.UserData;
-		BatchElementParams->LandscapeUniformShaderParametersResource = &LandscapeUniformShaderParameters;
-	}
 }
 
 FLandscapeComponentSceneProxy::~FLandscapeComponentSceneProxy()
@@ -1043,8 +998,6 @@ void FLandscapeComponentSceneProxy::OnTransformChanged()
 		);
 
 	LandscapeUniformShaderParameters.SetContents(LandscapeParams);
-
-	DynamicMesh.ReverseCulling = IsLocalToWorldDeterminantNegative();
 }
 
 namespace
@@ -1689,231 +1642,6 @@ void FLandscapeComponentSceneProxy::GetDynamicMeshElements(const TArray<const FS
 	INC_DWORD_STAT_BY(STAT_LandscapeTriangles, NumTriangles * NumPasses);
 }
 
-void FLandscapeComponentSceneProxy::DrawDynamicElements(FPrimitiveDrawInterface* PDI,const FSceneView* View)
-{
-	SCOPE_CYCLE_COUNTER(STAT_LandscapeDynamicDrawTime);
-
-	int32 NumPasses=0;
-	int32 NumTriangles=0;
-	int32 NumDrawCalls=0;
-	const bool bIsWireframe = View->Family->EngineShowFlags.Wireframe;
-
-	FVector CameraLocalPos3D = WorldToLocal.TransformPosition(View->ViewMatrices.ViewOrigin); 
-	FVector2D CameraLocalPos(CameraLocalPos3D.X, CameraLocalPos3D.Y);
-
-	// Could be different from bRequiresAdjacencyInformation during shader compilation
-	FMaterialRenderProxy* RenderProxy = MaterialInterface->GetRenderProxy(false);
-	bool bCurrentRequiresAdjacencyInformation = RequiresAdjacencyInformation( RenderProxy, View->GetFeatureLevel() );
-	DynamicMesh.Type = bCurrentRequiresAdjacencyInformation ? PT_12_ControlPointPatchList : PT_TriangleList;
-
-	// Setup the LOD parameters in DynamicMesh.
-	for( int32 SubY=0;SubY<NumSubsections;SubY++ )
-	{
-		for( int32 SubX=0;SubX<NumSubsections;SubX++ )
-		{
-			int32 SubSectionIdx = SubX + SubY*NumSubsections;
-			int32 CurrentLOD = CalcLODForSubsection(*View, SubX, SubY, CameraLocalPos);
-
-			FMeshBatchElement& BatchElement = DynamicMesh.Elements[SubSectionIdx];
-			((FLandscapeBatchElementParams*)BatchElement.UserData)->CurrentLOD = CurrentLOD;
-
-			if (bCurrentRequiresAdjacencyInformation)
-			{
-				check(SharedBuffers->AdjacencyIndexBuffers);
-				BatchElement.IndexBuffer = SharedBuffers->AdjacencyIndexBuffers->IndexBuffers[CurrentLOD];
-			}
-			else
-			{
-				BatchElement.IndexBuffer = SharedBuffers->IndexBuffers[CurrentLOD];
-			}
-
-			int32 LodSubsectionSizeVerts = (SubsectionSizeVerts >> CurrentLOD);
-
-			int32 NumPrimitives = FMath::Square((LodSubsectionSizeVerts - 1)) * 2;
-			BatchElement.NumPrimitives = NumPrimitives;
-			BatchElement.FirstIndex = (SubX + SubY * NumSubsections) * NumPrimitives * 3;
-			BatchElement.MinVertexIndex = SharedBuffers->IndexRanges[CurrentLOD].MinIndex[SubX][SubY];
-			BatchElement.MaxVertexIndex = SharedBuffers->IndexRanges[CurrentLOD].MaxIndex[SubX][SubY];
-
-#if WITH_EDITOR
-			// Tools never use tessellation
-			FMeshBatchElement& BatchElementTools = DynamicMeshTools.Elements[SubSectionIdx];
-			((FLandscapeBatchElementParams*)BatchElementTools.UserData)->CurrentLOD = CurrentLOD;
-			BatchElementTools.IndexBuffer = SharedBuffers->IndexBuffers[CurrentLOD];
-			BatchElementTools.NumPrimitives = NumPrimitives;
-			BatchElementTools.FirstIndex = BatchElement.FirstIndex;
-			BatchElementTools.MinVertexIndex = SharedBuffers->IndexRanges[CurrentLOD].MinIndex[SubX][SubY];
-			BatchElementTools.MaxVertexIndex = SharedBuffers->IndexRanges[CurrentLOD].MaxIndex[SubX][SubY];
-#endif
-		}
-	}
-
-	// Render the landscape component
-#if WITH_EDITOR
-	switch(GLandscapeViewMode)
-	{
-	case ELandscapeViewMode::DebugLayer:
-		if( GLayerDebugColorMaterial && EditToolRenderData )
-		{
-			const FLandscapeDebugMaterialRenderProxy DebugColorMaterialInstance(GLayerDebugColorMaterial->GetRenderProxy(false), 
-				(EditToolRenderData->DebugChannelR >= 0 ? WeightmapTextures[EditToolRenderData->DebugChannelR/4] : NULL),
-				(EditToolRenderData->DebugChannelG >= 0 ? WeightmapTextures[EditToolRenderData->DebugChannelG/4] : NULL),
-				(EditToolRenderData->DebugChannelB >= 0 ? WeightmapTextures[EditToolRenderData->DebugChannelB/4] : NULL),	
-				(EditToolRenderData->DebugChannelR >= 0 ? DebugColorMask::Masks[EditToolRenderData->DebugChannelR%4] : DebugColorMask::Masks[4]),
-				(EditToolRenderData->DebugChannelG >= 0 ? DebugColorMask::Masks[EditToolRenderData->DebugChannelG%4] : DebugColorMask::Masks[4]),
-				(EditToolRenderData->DebugChannelB >= 0 ? DebugColorMask::Masks[EditToolRenderData->DebugChannelB%4] : DebugColorMask::Masks[4])
-				);
-
-			DynamicMeshTools.MaterialRenderProxy = &DebugColorMaterialInstance;
-			NumPasses += DrawRichMesh(PDI, DynamicMeshTools, FLinearColor(1,1,1,1), LevelColor, FLinearColor(1.0f,1.0f,1.0f), this, IsSelected(), bIsWireframe);
-			NumTriangles += DynamicMeshTools.GetNumPrimitives();
-			NumDrawCalls += DynamicMeshTools.Elements.Num();
-		}
-		break;
-
-	case ELandscapeViewMode::LayerDensity:
-		if( EditToolRenderData )
-		{
-			int32 ColorIndex = FMath::Min<int32>(NumWeightmapLayerAllocations, GEngine->ShaderComplexityColors.Num());
-			const FColoredMaterialRenderProxy LayerDensityMaterialInstance(GEngine->LevelColorationUnlitMaterial->GetRenderProxy(false), ColorIndex ? GEngine->ShaderComplexityColors[ColorIndex-1] : FLinearColor::Black);
-
-			DynamicMeshTools.MaterialRenderProxy = &LayerDensityMaterialInstance;
-
-			NumPasses += DrawRichMesh(PDI, DynamicMeshTools, FLinearColor(1,1,1,1), LevelColor, FLinearColor(1.0f,1.0f,1.0f), this, IsSelected(), bIsWireframe);
-			NumTriangles += DynamicMeshTools.GetNumPrimitives();
-			NumDrawCalls += DynamicMeshTools.Elements.Num();
-		}
-		break;
-
-	case ELandscapeViewMode::LOD:
-		{
-			FLinearColor WireColors[7];
-			WireColors[0] = FLinearColor(1,1,1,1);
-			WireColors[1] = FLinearColor(1,0,0,1);
-			WireColors[2] = FLinearColor(0,1,0,1);
-			WireColors[3] = FLinearColor(0,0,1,1);
-			WireColors[4] = FLinearColor(1,1,0,1);
-			WireColors[5] = FLinearColor(1,0,1,1);
-			WireColors[6] = FLinearColor(0,1,1,1);
-
-			for( int32 i=0;i<DynamicMeshTools.Elements.Num();i++ )
-			{
-				FMeshBatch TempMesh = DynamicMeshTools;
-				TempMesh.Elements.Empty(1);
-				TempMesh.Elements.Add(DynamicMeshTools.Elements[i]);
-				int32 ColorIndex = ((FLandscapeBatchElementParams*)DynamicMeshTools.Elements[i].UserData)->CurrentLOD;
-				FLinearColor Color = ForcedLOD >= 0 ? WireColors[ColorIndex] : WireColors[ColorIndex]*0.2f;
-				const FColoredMaterialRenderProxy LODMaterialInstance(GEngine->LevelColorationUnlitMaterial->GetRenderProxy(false), Color);
-				TempMesh.MaterialRenderProxy = &LODMaterialInstance;
-				NumPasses += DrawRichMesh(PDI, TempMesh, Color, LevelColor, FLinearColor(1.0f,1.0f,1.0f), this, IsSelected());
-				NumTriangles += DynamicMeshTools.Elements[i].NumPrimitives;
-				NumDrawCalls++;
-			}
-		}
-		break;
-
-	default:
-#else
-	{
-#endif // WITH_EDITOR
-		// Regular Landscape rendering. Only use the dynamic path if we're rendering a rich view or we've disabled the static path for debugging.
-		if( IsRichView(*View->Family) || 
-			GLandscapeDebugOptions.bDisableStatic || 
-			bIsWireframe ||
-#if WITH_EDITOR
-			(IsSelected() && !GLandscapeEditModeActive)
-#else
-			IsSelected()
-#endif
-			)
-		{
-			DynamicMesh.MaterialRenderProxy = MaterialInterface->GetRenderProxy(false);
-			NumPasses += DrawRichMesh(PDI, DynamicMesh, FLinearColor(1,1,1,1), LevelColor, FLinearColor(1.0f,1.0f,1.0f), this, IsSelected(), bIsWireframe);
-			NumTriangles += DynamicMesh.GetNumPrimitives();
-			NumDrawCalls += DynamicMesh.Elements.Num();
-		}
-	}
-
-#if WITH_EDITOR
-	// Extra render passes for landscape tools
-	if( GLandscapeEditModeActive )
-	{
-		// Region selection
-		if ( EditToolRenderData && EditToolRenderData->SelectedType )
-		{
-			if ((GLandscapeEditRenderMode & ELandscapeEditRenderMode::SelectRegion) && (EditToolRenderData->SelectedType & FLandscapeEditToolRenderData::ST_REGION)
-				&& !(GLandscapeEditRenderMode & ELandscapeEditRenderMode::Mask))
-			{
-				const FLandscapeSelectMaterialRenderProxy SelectMaterialInstance(GSelectionRegionMaterial->GetRenderProxy(false), EditToolRenderData->DataTexture ? EditToolRenderData->DataTexture : GLandscapeBlackTexture);
-				DynamicMeshTools.MaterialRenderProxy = &SelectMaterialInstance;
-				NumPasses += PDI->DrawMesh(DynamicMeshTools);
-				NumTriangles += DynamicMeshTools.GetNumPrimitives();
-				NumDrawCalls += DynamicMeshTools.Elements.Num();
-			}
-			if ((GLandscapeEditRenderMode & ELandscapeEditRenderMode::SelectComponent) && (EditToolRenderData->SelectedType & FLandscapeEditToolRenderData::ST_COMPONENT))
-			{
-				DynamicMeshTools.MaterialRenderProxy = GSelectionColorMaterial->GetRenderProxy(0);
-				NumPasses += PDI->DrawMesh(DynamicMeshTools);
-				NumTriangles += DynamicMeshTools.GetNumPrimitives();
-				NumDrawCalls += DynamicMeshTools.Elements.Num();
-			}
-		}
-
-		// Mask
-		if ((GLandscapeEditRenderMode & ELandscapeEditRenderMode::SelectRegion) && (GLandscapeEditRenderMode & ELandscapeEditRenderMode::Mask))
-		{
-			if (EditToolRenderData && (EditToolRenderData->SelectedType & FLandscapeEditToolRenderData::ST_REGION) )
-			{
-				const FLandscapeMaskMaterialRenderProxy MaskMaterialInstance(GMaskRegionMaterial->GetRenderProxy(false), EditToolRenderData->DataTexture ? EditToolRenderData->DataTexture : GLandscapeBlackTexture, !!(GLandscapeEditRenderMode & ELandscapeEditRenderMode::InvertedMask) );
-				DynamicMeshTools.MaterialRenderProxy = &MaskMaterialInstance;
-				NumPasses += PDI->DrawMesh(DynamicMeshTools);
-				NumTriangles += DynamicMeshTools.GetNumPrimitives();
-				NumDrawCalls += DynamicMeshTools.Elements.Num();
-			}
-			else if (!(GLandscapeEditRenderMode & ELandscapeEditRenderMode::InvertedMask))
-			{
-				const FLandscapeMaskMaterialRenderProxy MaskMaterialInstance(GMaskRegionMaterial->GetRenderProxy(false), GLandscapeBlackTexture, false );
-				DynamicMeshTools.MaterialRenderProxy = &MaskMaterialInstance;
-				NumPasses += PDI->DrawMesh(DynamicMeshTools);
-				NumTriangles += DynamicMeshTools.GetNumPrimitives();
-				NumDrawCalls += DynamicMeshTools.Elements.Num();
-			}
-		}
-
-		// Edit mode tools
-		if( EditToolRenderData )
-		{
-			if (EditToolRenderData->ToolMaterial)
-			{
-				DynamicMeshTools.MaterialRenderProxy = EditToolRenderData->ToolMaterial->GetRenderProxy(0);		
-				NumPasses += PDI->DrawMesh(DynamicMeshTools);
-				NumTriangles += DynamicMeshTools.GetNumPrimitives();
-				NumDrawCalls += DynamicMeshTools.Elements.Num();
-			}
-
-			if (EditToolRenderData->GizmoMaterial && GLandscapeEditRenderMode & ELandscapeEditRenderMode::Gizmo)
-			{
-				DynamicMeshTools.MaterialRenderProxy = EditToolRenderData->GizmoMaterial->GetRenderProxy(0);
-				NumPasses += PDI->DrawMesh(DynamicMeshTools);
-				NumTriangles += DynamicMeshTools.GetNumPrimitives();
-				NumDrawCalls += DynamicMeshTools.Elements.Num();
-			}
-		}
-	}
-#endif // WITH_EDITOR
-
-	if ( GLandscapeDebugOptions.bShowPatches )
-	{
-		DrawWireBox(PDI, GetBounds().GetBox(), FColor(255, 255, 0), SDPG_World);
-	}
-
-	RenderBounds(PDI, View->Family->EngineShowFlags, GetBounds(), IsSelected());
-
-	INC_DWORD_STAT_BY(STAT_LandscapeComponents, NumPasses);
-	INC_DWORD_STAT_BY(STAT_LandscapeDrawCalls, NumDrawCalls);
-	INC_DWORD_STAT_BY(STAT_LandscapeTriangles, NumTriangles * NumPasses);
-}
-
 //
 // FLandscapeVertexBuffer
 //
diff --git a/Engine/Source/Runtime/Landscape/Private/LandscapeRenderMobile.cpp b/Engine/Source/Runtime/Landscape/Private/LandscapeRenderMobile.cpp
--- a/Engine/Source/Runtime/Landscape/Private/LandscapeRenderMobile.cpp
+++ b/Engine/Source/Runtime/Landscape/Private/LandscapeRenderMobile.cpp
@@ -271,21 +271,10 @@ void FLandscapeComponentSceneProxyMobile::CreateRenderThreadResources()
 
 	LandscapeVertexFactory->InitResource();
 	VertexFactory = LandscapeVertexFactory;
-	DynamicMesh.VertexFactory = VertexFactory;
-#if WITH_EDITOR
-	DynamicMeshTools.VertexFactory = VertexFactory;
-#endif
 
 	// Assign LandscapeUniformShaderParameters
 	LandscapeUniformShaderParameters.InitResource();
 
-	for( int32 ElementIdx=0;ElementIdx<DynamicMesh.Elements.Num();ElementIdx++ )
-	{
-		FMeshBatchElement& BatchElement = DynamicMesh.Elements[ElementIdx];
-		FLandscapeBatchElementParams* BatchElementParams = (FLandscapeBatchElementParams*)BatchElement.UserData;
-		BatchElementParams->LandscapeUniformShaderParametersResource = &LandscapeUniformShaderParameters;
-	}
-
 	PlatformData.Empty();
 }
 
