Collision utility functions;

* Added UCollisionProfile::GetChannelAndResponseParams() which gets the collision channel and the response params out of a collision profile.

* Added FCollisionResponseContainer::ReplaceChannels() which will replace response channels matching a specific response with another response.

diff --git a/Engine/Source/Runtime/Engine/Classes/Engine/CollisionProfile.h b/Engine/Source/Runtime/Engine/Classes/Engine/CollisionProfile.h
--- a/Engine/Source/Runtime/Engine/Classes/Engine/CollisionProfile.h
+++ b/Engine/Source/Runtime/Engine/Classes/Engine/CollisionProfile.h
@@ -177,6 +177,9 @@ public:
 	/** Fill up the array with the profile names **/
 	ENGINE_API static void GetProfileNames(TArray<TSharedPtr<FName>>& OutNameList);
 
+	/** Get the channel and response params from the specified profile */
+	ENGINE_API static bool GetChannelAndResponseParams(FName ProfileName, ECollisionChannel &CollisionChannel, FCollisionResponseParams &ResponseParams);
+
 	/** Fill up the loaded config of the profile name to the BodyInstance **/
 	bool ReadConfig(FName ProfileName, struct FBodyInstance& BodyInstance) const;
 
diff --git a/Engine/Source/Runtime/Engine/Classes/Engine/EngineTypes.h b/Engine/Source/Runtime/Engine/Classes/Engine/EngineTypes.h
--- a/Engine/Source/Runtime/Engine/Classes/Engine/EngineTypes.h
+++ b/Engine/Source/Runtime/Engine/Classes/Engine/EngineTypes.h
@@ -688,6 +688,9 @@ struct ENGINE_API FCollisionResponseContainer
 	/** Set all channels to the specified response */
 	void SetAllChannels(ECollisionResponse NewResponse);
 
+	/** Replace the channels matching the old response with the new response */
+	void ReplaceChannels(ECollisionResponse OldResponse, ECollisionResponse NewResponse);
+
 	/** Returns the response set on the specified channel */
 	ECollisionResponse GetResponse(ECollisionChannel Channel) const;
 
diff --git a/Engine/Source/Runtime/Engine/Private/Collision/CollisionProfile.cpp b/Engine/Source/Runtime/Engine/Private/Collision/CollisionProfile.cpp
--- a/Engine/Source/Runtime/Engine/Private/Collision/CollisionProfile.cpp
+++ b/Engine/Source/Runtime/Engine/Private/Collision/CollisionProfile.cpp
@@ -62,7 +62,7 @@ UCollisionProfile* UCollisionProfile::Get()
 
 void UCollisionProfile::GetProfileNames(TArray<TSharedPtr<FName>>& OutNameList)
 {
-	UCollisionProfile* CollisionProfile = UCollisionProfile::Get();
+	const UCollisionProfile* CollisionProfile = UCollisionProfile::Get();
 	check(CollisionProfile);
 
 	int32 NumProfiles = CollisionProfile->GetNumOfProfiles();
@@ -78,6 +78,32 @@ void UCollisionProfile::GetProfileNames(TArray<TSharedPtr<FName>>& OutNameList)
 	}
 }
 
+bool UCollisionProfile::GetChannelAndResponseParams(FName ProfileName, ECollisionChannel &CollisionChannel, FCollisionResponseParams &ResponseParams)
+{
+	const UCollisionProfile* CollisionProfile = UCollisionProfile::Get();
+	check(CollisionProfile);
+
+	FCollisionResponseTemplate Template;
+
+	if (CollisionProfile->GetProfileTemplate(ProfileName, Template))
+	{
+		CollisionChannel = Template.ObjectType;
+		ResponseParams = FCollisionResponseParams(Template.ResponseToChannels);
+		return true;
+	}
+
+	// Check for redirects
+	const FName* RedirectName = CollisionProfile->LookForProfileRedirect(ProfileName);
+	if (RedirectName && CollisionProfile->GetProfileTemplate(*RedirectName, Template))
+	{
+		CollisionChannel = Template.ObjectType;
+		ResponseParams = FCollisionResponseParams(Template.ResponseToChannels);
+		return true;
+	}
+
+	return false;
+}
+
 bool UCollisionProfile::GetProfileTemplate(FName ProfileName, struct FCollisionResponseTemplate& ProfileData) const
 {
 	// verify if it is in redirect first
diff --git a/Engine/Source/Runtime/Engine/Private/Collision/WorldCollision.cpp b/Engine/Source/Runtime/Engine/Private/Collision/WorldCollision.cpp
--- a/Engine/Source/Runtime/Engine/Private/Collision/WorldCollision.cpp
+++ b/Engine/Source/Runtime/Engine/Private/Collision/WorldCollision.cpp
@@ -297,24 +297,8 @@ bool UWorld::OverlapSingle(struct FOverlapResult& OutOverlap, const FVector& Pos
 // profile interfaces
 static void GetCollisionProfileChannelAndResponseParams(FName ProfileName, ECollisionChannel &CollisionChannel, FCollisionResponseParams &ResponseParams)
 {
-	const UCollisionProfile* CollisionProfile = UCollisionProfile::Get();
-	check(CollisionProfile != NULL);
-
-	FCollisionResponseTemplate Template;
-
-	if (CollisionProfile->GetProfileTemplate(ProfileName, Template))
-	{
-		CollisionChannel = Template.ObjectType;
-		ResponseParams = FCollisionResponseParams(Template.ResponseToChannels);
-		return;
-	}
-
-	// Check for redirects
-	const FName* RedirectName = CollisionProfile->LookForProfileRedirect(ProfileName);
-	if (RedirectName && CollisionProfile->GetProfileTemplate(*RedirectName, Template))
+	if (UCollisionProfile::GetChannelAndResponseParams(ProfileName, CollisionChannel, ResponseParams))
 	{
-		CollisionChannel = Template.ObjectType;
-		ResponseParams = FCollisionResponseParams(Template.ResponseToChannels);
 		return;
 	}
 
diff --git a/Engine/Source/Runtime/Engine/Private/PhysicsEngine/PhysUtils.cpp b/Engine/Source/Runtime/Engine/Private/PhysicsEngine/PhysUtils.cpp
--- a/Engine/Source/Runtime/Engine/Private/PhysicsEngine/PhysUtils.cpp
+++ b/Engine/Source/Runtime/Engine/Private/PhysicsEngine/PhysUtils.cpp
@@ -216,6 +216,17 @@ void FCollisionResponseContainer::SetAllChannels(ECollisionResponse NewResponse)
 	}
 }
 
+void FCollisionResponseContainer::ReplaceChannels(ECollisionResponse OldResponse, ECollisionResponse NewResponse)
+{
+	for (int32 i = 0; i < ARRAY_COUNT(EnumArray); i++)
+	{
+		if(EnumArray[i] == OldResponse)
+		{
+			EnumArray[i] = NewResponse;
+		}
+	}
+}
+
 /** Returns if the specified channel is set */
 ECollisionResponse FCollisionResponseContainer::GetResponse(ECollisionChannel Channel) const
 {
