GitHub pull request #618 from user stephenwittle. Forced calls to the FRandomStream constructor when creating an FRandomStream in a blueprint. Does not work for FRandomStreams located on the CDO (see UE-5974)

diff --git a/Engine/Source/Runtime/Core/Public/Math/RandomStream.h b/Engine/Source/Runtime/Core/Public/Math/RandomStream.h
--- a/Engine/Source/Runtime/Core/Public/Math/RandomStream.h
+++ b/Engine/Source/Runtime/Core/Public/Math/RandomStream.h
@@ -64,6 +64,11 @@ public:
 		Seed = InitialSeed;
 	}
 
+	int32 GetInitialSeed() const
+	{
+		return InitialSeed;
+	}
+
 	/**
 	 * Generates a new random seed.
 	 */
diff --git a/Engine/Source/Runtime/CoreUObject/Classes/Object.h b/Engine/Source/Runtime/CoreUObject/Classes/Object.h
--- a/Engine/Source/Runtime/CoreUObject/Classes/Object.h
+++ b/Engine/Source/Runtime/CoreUObject/Classes/Object.h
@@ -722,13 +722,14 @@ struct FTransform
  * Thread-safe RNG.
  * The full C++ class is located here: Engine\Source\Runtime\Core\Public\Math\RandomStream.h
  */
-USTRUCT(noexport, BlueprintType)
+USTRUCT(noexport, BlueprintType, meta = (HasNativeMake = "Engine.KismetMathLibrary.MakeRandomStream", HasNativeBreak = "Engine.KismetMathLibrary.BreakRandomStream"))
 struct FRandomStream
 {
 public:
 	UPROPERTY(EditAnywhere, BlueprintReadWrite, Category=RandomStream, SaveGame)
 	int32 InitialSeed;
-
+	
+	UPROPERTY()
 	int32 Seed;
 };
 
diff --git a/Engine/Source/Runtime/Engine/Classes/Kismet/KismetMathLibrary.h b/Engine/Source/Runtime/Engine/Classes/Kismet/KismetMathLibrary.h
--- a/Engine/Source/Runtime/Engine/Classes/Kismet/KismetMathLibrary.h
+++ b/Engine/Source/Runtime/Engine/Classes/Kismet/KismetMathLibrary.h
@@ -1115,6 +1115,14 @@ class UKismetMathLibrary : public UBlueprintFunctionLibrary
 	UFUNCTION(BlueprintPure, Category="Math|Transform", meta=(NativeBreakFunc))
 	static void BreakTransform(const FTransform& InTransform, FVector& Location, FRotator& Rotation, FVector& Scale);
 
+	/** Makes a SRand-based random number generator */
+	UFUNCTION(BlueprintPure, meta = (Keywords = "construct build", NativeMakeFunc), Category = "Math|Random")
+		static FRandomStream MakeRandomStream(int32 InitialSeed);
+
+	/** Breaks apart a random number generator */
+	UFUNCTION(BlueprintPure, Category = "Math|Random", meta = (NativeBreakFunc))
+		static void BreakRandomStream(const FRandomStream& InRandomStream, int32& InitialSeed);
+
 	/** Make a color from individual color components (RGB space) */
 	UFUNCTION(BlueprintPure, Category="Math|Color", meta=(Keywords="construct build"))
 	static FLinearColor MakeColor(float R, float G, float B, float A = 1.0f);
diff --git a/Engine/Source/Runtime/Engine/Private/KismetMathLibrary.cpp b/Engine/Source/Runtime/Engine/Private/KismetMathLibrary.cpp
--- a/Engine/Source/Runtime/Engine/Private/KismetMathLibrary.cpp
+++ b/Engine/Source/Runtime/Engine/Private/KismetMathLibrary.cpp
@@ -1993,3 +1993,13 @@ bool UKismetMathLibrary::LinePlaneIntersection_OriginNormal(const FVector& LineS
 
 	return true;
 }
+
+void UKismetMathLibrary::BreakRandomStream(const FRandomStream& InRandomStream, int32& InitialSeed)
+{
+	InitialSeed = InRandomStream.GetInitialSeed();
+}
+
+FRandomStream UKismetMathLibrary::MakeRandomStream(int32 InitialSeed)
+{
+	return FRandomStream(InitialSeed);
+}
