BodyInstance utility function;

* Added FBodyInstance::ReplaceResponseToChannels() which will replace the body's response channels matching a specific response with another response.

diff --git a/Engine/Source/Runtime/Engine/Classes/PhysicsEngine/BodyInstance.h b/Engine/Source/Runtime/Engine/Classes/PhysicsEngine/BodyInstance.h
--- a/Engine/Source/Runtime/Engine/Classes/PhysicsEngine/BodyInstance.h
+++ b/Engine/Source/Runtime/Engine/Classes/PhysicsEngine/BodyInstance.h
@@ -70,6 +70,9 @@ struct ENGINE_API FCollisionResponse
 	/** Set all channels to the specified response */
 	void SetAllChannels(ECollisionResponse NewResponse);
 
+	/** Replace the channels matching the old response with the new response */
+	void ReplaceChannels(ECollisionResponse OldResponse, ECollisionResponse NewResponse);
+
 	/** Returns the response set on the specified channel */
 	ECollisionResponse GetResponse(ECollisionChannel Channel) const;
 	const FCollisionResponseContainer& GetResponseContainer() const { return ResponseToChannels; }
@@ -517,6 +520,9 @@ public:
 	/** Set the response of this body to all channels */
 	void SetResponseToAllChannels(ECollisionResponse NewResponse);
 
+	/** Replace the channels on this body matching the old response with the new response */
+	void ReplaceResponseToChannels(ECollisionResponse OldResponse, ECollisionResponse NewResponse);
+
 	/** Set the response of this body to the supplied settings */
 	void SetResponseToChannels(const FCollisionResponseContainer& NewReponses);
 
diff --git a/Engine/Source/Runtime/Engine/Private/PhysicsEngine/BodyInstance.cpp b/Engine/Source/Runtime/Engine/Private/PhysicsEngine/BodyInstance.cpp
--- a/Engine/Source/Runtime/Engine/Private/PhysicsEngine/BodyInstance.cpp
+++ b/Engine/Source/Runtime/Engine/Private/PhysicsEngine/BodyInstance.cpp
@@ -65,6 +65,14 @@ void FCollisionResponse::SetAllChannels(ECollisionResponse NewResponse)
 #endif
 }
 
+void FCollisionResponse::ReplaceChannels(ECollisionResponse OldResponse, ECollisionResponse NewResponse)
+{
+	ResponseToChannels.ReplaceChannels(OldResponse, NewResponse);
+#if 1// @hack until PostLoad is disabled for CDO of BP WITH_EDITOR
+	UpdateArrayFromResponseContainer();
+#endif
+}
+
 /** Returns the response set on the specified channel */
 ECollisionResponse FCollisionResponse::GetResponse(ECollisionChannel Channel) const
 {
@@ -474,6 +482,14 @@ void FBodyInstance::SetResponseToAllChannels(ECollisionResponse NewResponse)
 	UpdatePhysicsFilterData();
 }
 	
+void FBodyInstance::ReplaceResponseToChannels(ECollisionResponse OldResponse, ECollisionResponse NewResponse)
+{
+	InvalidateCollisionProfileName();
+	ResponseToChannels_DEPRECATED.ReplaceChannels(OldResponse, NewResponse);
+	CollisionResponses.ReplaceChannels(OldResponse, NewResponse);
+	UpdatePhysicsFilterData();
+}
+
 void FBodyInstance::SetResponseToChannels(const FCollisionResponseContainer& NewReponses)
 {
 	InvalidateCollisionProfileName();
