Restored FNavLinkRenderingProxy::DrawLinks which is still used

diff --git a/Engine/Source/Runtime/Engine/Private/AI/Navigation/NavLinkRenderingComponent.cpp b/Engine/Source/Runtime/Engine/Private/AI/Navigation/NavLinkRenderingComponent.cpp
--- a/Engine/Source/Runtime/Engine/Private/AI/Navigation/NavLinkRenderingComponent.cpp
+++ b/Engine/Source/Runtime/Engine/Private/AI/Navigation/NavLinkRenderingComponent.cpp
@@ -284,6 +284,97 @@ void FNavLinkRenderingProxy::GetLinkMeshes(const TArray<FNavLinkDrawing>& OffMes
 	}
 }
 
+void FNavLinkRenderingProxy::DrawLinks(FPrimitiveDrawInterface* PDI, TArray<FNavLinkDrawing>& OffMeshPointLinks, TArray<FNavLinkSegmentDrawing>& OffMeshSegmentLinks, TArray<float>& StepHeights, FMaterialRenderProxy* const MeshColorInstance, uint32 AgentMask)
+{
+	static const FColor LinkColor(0,0,166);
+	static const float LinkArcThickness = 3.f;
+	static const float LinkArcHeight = 0.4f;
+
+	if (StepHeights.Num() == 0)
+	{
+		StepHeights.Add(FNavigationSystem::FallbackAgentHeight / 2);
+	}
+
+	for (int32 LinkIndex = 0; LinkIndex < OffMeshPointLinks.Num(); ++LinkIndex)
+	{
+		const FNavLinkDrawing& Link = OffMeshPointLinks[LinkIndex];
+		if ((Link.SupportedAgentsBits & AgentMask) == 0)
+		{
+			continue;
+		}
+
+		const uint32 Segments = FPlatformMath::Max<uint32>(LinkArcHeight*(Link.Right-Link.Left).Size()/10, 8);
+		DrawArc(PDI, Link.Left, Link.Right, LinkArcHeight, Segments, Link.Color, SDPG_World, 3.5f);
+		const FVector VOffset(0,0,FVector::Dist(Link.Left, Link.Right)*1.333f);
+
+		switch (Link.Direction)
+		{
+		case ENavLinkDirection::LeftToRight:
+			DrawArrowHead(PDI, Link.Right, Link.Left+VOffset, 30.f, Link.Color, SDPG_World, 3.5f);
+			break;
+		case ENavLinkDirection::RightToLeft:
+			DrawArrowHead(PDI, Link.Left, Link.Right+VOffset, 30.f, Link.Color, SDPG_World, 3.5f);
+			break;
+		case ENavLinkDirection::BothWays:
+		default:
+			DrawArrowHead(PDI, Link.Right, Link.Left+VOffset, 30.f, Link.Color, SDPG_World, 3.5f);
+			DrawArrowHead(PDI, Link.Left, Link.Right+VOffset, 30.f, Link.Color, SDPG_World, 3.5f);
+			break;
+		}
+
+		// draw snap-spheres on both ends
+		for (int32 StepHeightIndex = 0; StepHeightIndex < StepHeights.Num(); ++StepHeightIndex)
+		{
+			DrawCylinder(PDI, Link.Right, FVector(1,0,0), FVector(0,1,0), FVector(0,0,1), Link.SnapRadius, StepHeights[StepHeightIndex], 10, MeshColorInstance, SDPG_World);
+			DrawCylinder(PDI, Link.Left, FVector(1,0,0), FVector(0,1,0), FVector(0,0,1), Link.SnapRadius, StepHeights[StepHeightIndex], 10, MeshColorInstance, SDPG_World);
+		}
+	}
+
+	static const float SegmentArcHeight = 0.25f;
+	for (int32 LinkIndex = 0; LinkIndex < OffMeshSegmentLinks.Num(); ++LinkIndex)
+	{
+		const FNavLinkSegmentDrawing& Link = OffMeshSegmentLinks[LinkIndex];
+		if ((Link.SupportedAgentsBits & AgentMask) == 0)
+		{
+			continue;
+		}
+
+		const uint32 SegmentsStart = FPlatformMath::Max<uint32>(SegmentArcHeight*(Link.RightStart-Link.LeftStart).Size()/10, 8);
+		const uint32 SegmentsEnd = FPlatformMath::Max<uint32>(SegmentArcHeight*(Link.RightEnd-Link.LeftEnd).Size()/10, 8);
+		DrawArc(PDI, Link.LeftStart, Link.RightStart, SegmentArcHeight, SegmentsStart, Link.Color, SDPG_World, 3.5f);
+		DrawArc(PDI, Link.LeftEnd, Link.RightEnd, SegmentArcHeight, SegmentsEnd, Link.Color, SDPG_World, 3.5f);
+		const FVector VOffset(0,0,FVector::Dist(Link.LeftStart, Link.RightStart)*1.333f);
+
+		switch (Link.Direction)
+		{
+		case ENavLinkDirection::LeftToRight:
+			DrawArrowHead(PDI, Link.RightStart, Link.LeftStart+VOffset, 30.f, Link.Color, SDPG_World, 3.5f);
+			DrawArrowHead(PDI, Link.RightEnd, Link.LeftEnd+VOffset, 30.f, Link.Color, SDPG_World, 3.5f);
+			break;
+		case ENavLinkDirection::RightToLeft:
+			DrawArrowHead(PDI, Link.LeftStart, Link.RightStart+VOffset, 30.f, Link.Color, SDPG_World, 3.5f);
+			DrawArrowHead(PDI, Link.LeftEnd, Link.RightEnd+VOffset, 30.f, Link.Color, SDPG_World, 3.5f);
+			break;
+		case ENavLinkDirection::BothWays:
+		default:
+			DrawArrowHead(PDI, Link.RightStart, Link.LeftStart+VOffset, 30.f, Link.Color, SDPG_World, 3.5f);
+			DrawArrowHead(PDI, Link.RightEnd, Link.LeftEnd+VOffset, 30.f, Link.Color, SDPG_World, 3.5f);
+			DrawArrowHead(PDI, Link.LeftStart, Link.RightStart+VOffset, 30.f, Link.Color, SDPG_World, 3.5f);
+			DrawArrowHead(PDI, Link.LeftEnd, Link.RightEnd+VOffset, 30.f, Link.Color, SDPG_World, 3.5f);
+			break;
+		}
+
+		// draw snap-spheres on both ends
+		for (int32 StepHeightIndex = 0; StepHeightIndex < StepHeights.Num(); ++StepHeightIndex)
+		{
+			DrawCylinder(PDI, Link.RightStart, FVector(1,0,0), FVector(0,1,0), FVector(0,0,1), Link.SnapRadius, StepHeights[StepHeightIndex], 10, MeshColorInstance, SDPG_World);
+			DrawCylinder(PDI, Link.RightEnd, FVector(1,0,0), FVector(0,1,0), FVector(0,0,1), Link.SnapRadius, StepHeights[StepHeightIndex], 10, MeshColorInstance, SDPG_World);
+			DrawCylinder(PDI, Link.LeftStart, FVector(1,0,0), FVector(0,1,0), FVector(0,0,1), Link.SnapRadius, StepHeights[StepHeightIndex], 10, MeshColorInstance, SDPG_World);
+			DrawCylinder(PDI, Link.LeftEnd, FVector(1,0,0), FVector(0,1,0), FVector(0,0,1), Link.SnapRadius, StepHeights[StepHeightIndex], 10, MeshColorInstance, SDPG_World);
+		}
+	}
+}
+
 FPrimitiveViewRelevance FNavLinkRenderingProxy::GetViewRelevance(const FSceneView* View)
 {
 	FPrimitiveViewRelevance Result;
diff --git a/Engine/Source/Runtime/Engine/Public/AI/NavLinkRenderingProxy.h b/Engine/Source/Runtime/Engine/Public/AI/NavLinkRenderingProxy.h
--- a/Engine/Source/Runtime/Engine/Public/AI/NavLinkRenderingProxy.h
+++ b/Engine/Source/Runtime/Engine/Public/AI/NavLinkRenderingProxy.h
@@ -45,4 +45,7 @@ public:
 	void StoreSegmentLinks(const FTransform& LocalToWorld, const TArray<FNavigationSegmentLink>& LinksArray);
 
 	static void GetLinkMeshes(const TArray<FNavLinkDrawing>& OffMeshPointLinks, const TArray<FNavLinkSegmentDrawing>& OffMeshSegmentLinks, TArray<float>& StepHeights, FMaterialRenderProxy* const MeshColorInstance, int32 ViewIndex, FMeshElementCollector& Collector, uint32 AgentMask);
+
+	/** made static to allow consistent navlinks drawing even if something is drawing links without FNavLinkRenderingProxy */
+	static void DrawLinks(FPrimitiveDrawInterface* PDI, TArray<FNavLinkDrawing>& OffMeshPointLinks, TArray<FNavLinkSegmentDrawing>& OffMeshSegmentLinks, TArray<float>& StepHeights, FMaterialRenderProxy* const MeshColorInstance, uint32 AgentMask);
 };
