PIE warning when using simple move to BP functions on pawn that can't support path following (e.g. rigid body)

diff --git a/Engine/Source/Runtime/AIModule/Classes/Navigation/PathFollowingComponent.h b/Engine/Source/Runtime/AIModule/Classes/Navigation/PathFollowingComponent.h
--- a/Engine/Source/Runtime/AIModule/Classes/Navigation/PathFollowingComponent.h
+++ b/Engine/Source/Runtime/AIModule/Classes/Navigation/PathFollowingComponent.h
@@ -251,6 +251,9 @@ class AIMODULE_API UPathFollowingComponent : public UActorComponent, public IAIR
 	/** Called when falling movement ends. */
 	virtual void OnLanded() {}
 
+	/** Check if path following can be activated */
+	virtual bool IsPathFollowingAllowed() const;
+
 	/** call when moving agent finishes using custom nav link, returns control back to path following */
 	virtual void FinishUsingCustomLink(INavLinkCustomInterface* CustomNavLink);
 
diff --git a/Engine/Source/Runtime/AIModule/Private/Navigation/PathFollowingComponent.cpp b/Engine/Source/Runtime/AIModule/Private/Navigation/PathFollowingComponent.cpp
--- a/Engine/Source/Runtime/AIModule/Private/Navigation/PathFollowingComponent.cpp
+++ b/Engine/Source/Runtime/AIModule/Private/Navigation/PathFollowingComponent.cpp
@@ -1405,6 +1405,11 @@ FString UPathFollowingComponent::GetDebugString() const
 	return Desc;
 }
 
+bool UPathFollowingComponent::IsPathFollowingAllowed() const
+{
+	return MovementComp && MovementComp->CanStartPathFollowing();
+}
+
 void UPathFollowingComponent::LockResource(EAIRequestPriority::Type LockSource)
 {
 	const static UEnum* SourceEnum = FindObject<UEnum>(ANY_PACKAGE, TEXT("EAILockSource"));
diff --git a/Engine/Source/Runtime/Engine/Classes/GameFramework/CharacterMovementComponent.h b/Engine/Source/Runtime/Engine/Classes/GameFramework/CharacterMovementComponent.h
--- a/Engine/Source/Runtime/Engine/Classes/GameFramework/CharacterMovementComponent.h
+++ b/Engine/Source/Runtime/Engine/Classes/GameFramework/CharacterMovementComponent.h
@@ -763,6 +763,7 @@ public:
 
 	//BEGIN UNavMovementComponent Interface
 	virtual void RequestDirectMove(const FVector& MoveVelocity, bool bForceMaxSpeed) override;
+	virtual bool CanStartPathFollowing() const override;
 	virtual bool CanStopPathFollowing() const override;
 	//END UNaVMovementComponent Interface
 
diff --git a/Engine/Source/Runtime/Engine/Classes/GameFramework/NavMovementComponent.h b/Engine/Source/Runtime/Engine/Classes/GameFramework/NavMovementComponent.h
--- a/Engine/Source/Runtime/Engine/Classes/GameFramework/NavMovementComponent.h
+++ b/Engine/Source/Runtime/Engine/Classes/GameFramework/NavMovementComponent.h
@@ -83,6 +83,8 @@ public:
 	/** Resets runtime movement state to character's movement capabilities */
 	void ResetMoveState() { MovementState = NavAgentProps; }
 
+	/** @return true if path following can start */
+	virtual bool CanStartPathFollowing() const { return true; }
 
 	/** @return true if component can crouch */
 	FORCEINLINE bool CanEverCrouch() const { return NavAgentProps.bCanCrouch; }
diff --git a/Engine/Source/Runtime/Engine/Private/AI/Navigation/NavigationSystem.cpp b/Engine/Source/Runtime/Engine/Private/AI/Navigation/NavigationSystem.cpp
--- a/Engine/Source/Runtime/Engine/Private/AI/Navigation/NavigationSystem.cpp
+++ b/Engine/Source/Runtime/Engine/Private/AI/Navigation/NavigationSystem.cpp
@@ -1,6 +1,7 @@
 // Copyright 1998-2014 Epic Games, Inc. All Rights Reserved.
 
 #include "EnginePrivate.h"
+#include "MessageLog.h"
 #include "NavDataGenerator.h"
 #include "NavigationOctree.h"
 #include "AI/Navigation/NavMeshBoundsVolume.h"
@@ -34,6 +35,8 @@ static const uint32 MAX_NAV_SEARCH_NODES = RECAST_MAX_SEARCH_NODES;
 static const uint32 MAX_NAV_SEARCH_NODES = 2048;
 #endif // WITH_RECAST
 
+#define LOCTEXT_NAMESPACE "Navigation"
+
 DEFINE_LOG_CATEGORY(LogNavigation);
 DEFINE_LOG_CATEGORY_STATIC(LogNavOctree, Warning, All);
 
@@ -898,6 +901,22 @@ void UNavigationSystem::SimpleMoveToLocation(AController* Controller, const FVec
 
 	Controller->InitNavigationControl(PFindComp, PFollowComp);
 
+	if (PFindComp == nullptr || PFollowComp == nullptr)
+	{
+		FMessageLog("PIE").Warning(FText::Format(
+			LOCTEXT("SimpleMoveErrorNoComp", "SimpleMove failed for {0}: missing components"),
+			FText::FromName(Controller->GetFName())
+			));
+	}
+
+	if (!PFollowComp->IsPathFollowingAllowed())
+	{
+		FMessageLog("PIE").Warning(FText::Format(
+			LOCTEXT("SimpleMoveErrorMovement", "SimpleMove failed for {0}: movement not allowed"),
+			FText::FromName(Controller->GetFName())
+			));
+	}
+
 	if (PFindComp && PFollowComp && !PFollowComp->HasReached(GoalLocation))
 	{
 		const bool bPathExists = PFindComp->FindPathToLocation(GoalLocation);
diff --git a/Engine/Source/Runtime/Engine/Private/CharacterMovementComponent.cpp b/Engine/Source/Runtime/Engine/Private/CharacterMovementComponent.cpp
--- a/Engine/Source/Runtime/Engine/Private/CharacterMovementComponent.cpp
+++ b/Engine/Source/Runtime/Engine/Private/CharacterMovementComponent.cpp
@@ -2279,6 +2279,28 @@ void UCharacterMovementComponent::RequestDirectMove(const FVector& MoveVelocity,
 	}
 }
 
+bool UCharacterMovementComponent::CanStartPathFollowing() const
+{
+	if (!HasValidData() || HasRootMotion())
+	{
+		return false;
+	}
+
+	if (CharacterOwner)
+	{
+		if (CharacterOwner->GetRootComponent() && CharacterOwner->GetRootComponent()->IsSimulatingPhysics())
+		{
+			return false;
+		}
+		else if (CharacterOwner->IsMatineeControlled())
+		{
+			return false;
+		}
+	}
+
+	return Super::CanStartPathFollowing();
+}
+
 bool UCharacterMovementComponent::CanStopPathFollowing() const
 {
 	return !IsFalling();
