PIE warning for invalid calls to SimpleMoveToActor

diff --git a/Engine/Source/Runtime/Engine/Private/AI/Navigation/NavigationSystem.cpp b/Engine/Source/Runtime/Engine/Private/AI/Navigation/NavigationSystem.cpp
--- a/Engine/Source/Runtime/Engine/Private/AI/Navigation/NavigationSystem.cpp
+++ b/Engine/Source/Runtime/Engine/Private/AI/Navigation/NavigationSystem.cpp
@@ -877,7 +877,25 @@ void UNavigationSystem::SimpleMoveToActor(AController* Controller, const AActor*
 	
 	Controller->InitNavigationControl(PFindComp, PFollowComp);
 
-	if (PFindComp && PFollowComp && !PFollowComp->HasReached(*Goal))
+	if (PFindComp == nullptr || PFollowComp == nullptr)
+	{
+		FMessageLog("PIE").Warning(FText::Format(
+			LOCTEXT("SimpleMoveErrorNoComp", "SimpleMove failed for {0}: missing components"),
+			FText::FromName(Controller->GetFName())
+			));
+		return;
+	}
+
+	if (!PFollowComp->IsPathFollowingAllowed())
+	{
+		FMessageLog("PIE").Warning(FText::Format(
+			LOCTEXT("SimpleMoveErrorMovement", "SimpleMove failed for {0}: movement not allowed"),
+			FText::FromName(Controller->GetFName())
+			));
+		return;
+	}
+
+	if (!PFollowComp->HasReached(*Goal))
 	{
 		const bool bPathExists = PFindComp->FindPathToActor(Goal);
 		if (bPathExists)
@@ -907,6 +925,7 @@ void UNavigationSystem::SimpleMoveToLocation(AController* Controller, const FVec
 			LOCTEXT("SimpleMoveErrorNoComp", "SimpleMove failed for {0}: missing components"),
 			FText::FromName(Controller->GetFName())
 			));
+		return;
 	}
 
 	if (!PFollowComp->IsPathFollowingAllowed())
@@ -915,9 +934,10 @@ void UNavigationSystem::SimpleMoveToLocation(AController* Controller, const FVec
 			LOCTEXT("SimpleMoveErrorMovement", "SimpleMove failed for {0}: movement not allowed"),
 			FText::FromName(Controller->GetFName())
 			));
+		return;
 	}
 
-	if (PFindComp && PFollowComp && !PFollowComp->HasReached(GoalLocation))
+	if (!PFollowComp->HasReached(GoalLocation))
 	{
 		const bool bPathExists = PFindComp->FindPathToLocation(GoalLocation);
 		if (bPathExists)
