Improved performance of material compiling by ensuring that child expressions aren't checked for loops repeatedly

Keep a list of all visited expressions and skip over them if encountered again.

diff --git a/Engine/Source/Runtime/Engine/Classes/Materials/MaterialExpression.h b/Engine/Source/Runtime/Engine/Classes/Materials/MaterialExpression.h
--- a/Engine/Source/Runtime/Engine/Classes/Materials/MaterialExpression.h
+++ b/Engine/Source/Runtime/Engine/Classes/Materials/MaterialExpression.h
@@ -365,9 +365,10 @@ protected:
 	 * Checks whether any inputs to this expression create a loop by recursively
 	 * calling itself and keeping a list of inputs as expression keys.
 	 *
-	 * @param ExpressionStack List of expression keys that have been checked already
+	 * @param ExpressionStack    List of expression keys that have been checked already in the current stack
+	 * @param VisitedExpressions List of all expression keys that have been visited
 	 */
-	bool ContainsInputLoopInternal(TArray<FMaterialExpressionKey>& ExpressionStack);
+	bool ContainsInputLoopInternal(TArray<FMaterialExpressionKey>& ExpressionStack, TSet<FMaterialExpressionKey>& VisitedExpressions);
 };
 
 
diff --git a/Engine/Source/Runtime/Engine/Private/Materials/MaterialExpressions.cpp b/Engine/Source/Runtime/Engine/Private/Materials/MaterialExpressions.cpp
--- a/Engine/Source/Runtime/Engine/Private/Materials/MaterialExpressions.cpp
+++ b/Engine/Source/Runtime/Engine/Private/Materials/MaterialExpressions.cpp
@@ -977,11 +977,11 @@ void UMaterialExpression::SetEditableName(const FString& NewName)
 bool UMaterialExpression::ContainsInputLoop()
 {
 	TArray<FMaterialExpressionKey> ExpressionStack;
-
-	return ContainsInputLoopInternal(ExpressionStack);
+	TSet<FMaterialExpressionKey> VisitedExpressions;
+	return ContainsInputLoopInternal(ExpressionStack, VisitedExpressions);
 }
 
-bool UMaterialExpression::ContainsInputLoopInternal(TArray<FMaterialExpressionKey>& ExpressionStack)
+bool UMaterialExpression::ContainsInputLoopInternal(TArray<FMaterialExpressionKey>& ExpressionStack, TSet<FMaterialExpressionKey>& VisitedExpressions)
 {
 	const TArray<FExpressionInput*> Inputs = GetInputs();
 	for (int32 Index = 0; Index < Inputs.Num(); ++Index)
@@ -994,10 +994,12 @@ bool UMaterialExpression::ContainsInputLoopInternal(TArray<FMaterialExpressionKe
 			{
 				return true;
 			}
-			else
+			// prevent recurring visits to expressions we've already checked
+			else if (!VisitedExpressions.Contains(InputExpressionKey))
 			{
+				VisitedExpressions.Add(InputExpressionKey);
 				ExpressionStack.Add(InputExpressionKey);
-				if (Input->Expression->ContainsInputLoopInternal(ExpressionStack))
+				if (Input->Expression->ContainsInputLoopInternal(ExpressionStack, VisitedExpressions))
 				{
 					return true;
 				}
