fixed missing navmesh on newly created map

diff --git a/Engine/Source/Runtime/Engine/Private/AI/Navigation/NavigationSystem.cpp b/Engine/Source/Runtime/Engine/Private/AI/Navigation/NavigationSystem.cpp
--- a/Engine/Source/Runtime/Engine/Private/AI/Navigation/NavigationSystem.cpp
+++ b/Engine/Source/Runtime/Engine/Private/AI/Navigation/NavigationSystem.cpp
@@ -1176,7 +1176,7 @@ ANavigationData* UNavigationSystem::GetMainNavData(FNavigationSystem::ECreateIfE
 		for (int32 NavDataIndex = 0; NavDataIndex < NavDataSet.Num(); ++NavDataIndex)
 		{
 			ANavigationData* NavData = NavDataSet[NavDataIndex];
-			if (NavData != NULL && NavData->IsPendingKill() == false)
+			if (NavData != NULL && NavData->IsPendingKill() == false && !NavData->IsA(AAbstractNavData::StaticClass()))
 			{
 				MainNavData = NavData;
 				break;
@@ -1470,7 +1470,10 @@ UNavigationSystem::ERegistrationResult UNavigationSystem::RegisterNavData(ANavig
 					bAgentSupported = true;
 
 					NavData->SetConfig(*Agent);
-					AgentToNavDataMap.Add(*Agent, NavData);
+					if (!NavData->IsA(AAbstractNavData::StaticClass()))
+					{
+						AgentToNavDataMap.Add(*Agent, NavData);
+					}
 
 					NavData->SetSupportsDefaultAgent(AgentIndex == 0);
 					NavData->ProcessNavAreas(NavAreaClasses, AgentIndex);
