fixed SimpleMoveToActor / Location ignoring failed paths or already standing at destination

diff --git a/Engine/Source/Runtime/Engine/Private/AI/Navigation/NavigationSystem.cpp b/Engine/Source/Runtime/Engine/Private/AI/Navigation/NavigationSystem.cpp
--- a/Engine/Source/Runtime/Engine/Private/AI/Navigation/NavigationSystem.cpp
+++ b/Engine/Source/Runtime/Engine/Private/AI/Navigation/NavigationSystem.cpp
@@ -884,15 +884,16 @@ bool UNavigationSystem::ProjectPointToNavigation(const FVector& Point, FNavLocat
 
 void UNavigationSystem::SimpleMoveToActor(AController* Controller, const AActor* Goal)
 {
-	if (Goal == NULL || Controller == NULL || Controller->GetPawn() == NULL)
+	UNavigationSystem* NavSys = Controller ? UNavigationSystem::GetCurrent(Controller->GetWorld()) : nullptr;
+	if (NavSys == nullptr || Goal == nullptr || Controller == nullptr || Controller->GetPawn() == nullptr)
 	{
-		UE_LOG(LogNavigation, Warning, TEXT("UNavigationSystem::SimpleMoveToActor called for Controller:%s controlling Pawn:%s with goal actor %s (if any of these is None then there's your problem")
-			, *GetNameSafe(Controller), Controller ? *GetNameSafe(Controller->GetPawn()) : TEXT("NULL"), *GetNameSafe(Goal));
+		UE_LOG(LogNavigation, Warning, TEXT("UNavigationSystem::SimpleMoveToActor called for NavSys:%s Controller:%s controlling Pawn:%s with goal actor %s (if any of these is None then there's your problem"),
+			*GetNameSafe(NavSys), *GetNameSafe(Controller), Controller ? *GetNameSafe(Controller->GetPawn()) : TEXT("NULL"), *GetNameSafe(Goal));
 		return;
 	}
 
-	UPathFollowingComponent* PFollowComp = NULL;
-	UNavigationSystem* NavSys = UNavigationSystem::GetCurrent(Controller->GetWorld());
+	UPathFollowingComponent* PFollowComp = nullptr;
+	Controller->InitNavigationControl(PFollowComp);
 
 	if (PFollowComp == nullptr)
 	{
@@ -912,7 +913,13 @@ void UNavigationSystem::SimpleMoveToActor(AController* Controller, const AActor*
 		return;
 	}
 
-	if (NavSys && !PFollowComp->HasReached(*Goal))
+	if (PFollowComp->HasReached(*Goal))
+	{
+		// make sure previous move request gets aborted
+		PFollowComp->AbortMove(TEXT("Aborting move due to new move request finishing with AlreadyAtGoal"), FAIRequestID::AnyRequest);
+		PFollowComp->SetLastMoveAtGoal(true);
+	}
+	else
 	{
 		const ANavigationData* NavData = NavSys->GetNavDataForProps(Controller->GetNavAgentProperties());
 		FPathFindingQuery Query(Controller, NavData, Controller->GetNavAgentLocation(), Goal->GetActorLocation());
@@ -923,21 +930,25 @@ void UNavigationSystem::SimpleMoveToActor(AController* Controller, const AActor*
 
 			PFollowComp->RequestMove(Result.Path, Goal);
 		}
+		else if (PFollowComp->GetStatus() != EPathFollowingStatus::Idle)
+		{
+			PFollowComp->AbortMove(TEXT("Aborting move due to new move request failing to generate a path"), FAIRequestID::AnyRequest);
+			PFollowComp->SetLastMoveAtGoal(false);
+		}
 	}
 }
 
 void UNavigationSystem::SimpleMoveToLocation(AController* Controller, const FVector& GoalLocation)
 {
-	if (Controller == NULL || Controller->GetPawn() == NULL)
+	UNavigationSystem* NavSys = Controller ? UNavigationSystem::GetCurrent(Controller->GetWorld()) : nullptr;
+	if (NavSys == nullptr || Controller == nullptr || Controller->GetPawn() == nullptr)
 	{
-		UE_LOG(LogNavigation, Warning, TEXT("UNavigationSystem::SimpleMoveToLocation called for Controller:%s controlling Pawn:%s (if any of these is None then there's your problem")
-			, *GetNameSafe(Controller), Controller ? *GetNameSafe(Controller->GetPawn()) : TEXT("NULL"));
+		UE_LOG(LogNavigation, Warning, TEXT("UNavigationSystem::SimpleMoveToActor called for NavSys:%s Controller:%s controlling Pawn:%s (if any of these is None then there's your problem"),
+			*GetNameSafe(NavSys), *GetNameSafe(Controller), Controller ? *GetNameSafe(Controller->GetPawn()) : TEXT("NULL"));
 		return;
 	}
 
-	UPathFollowingComponent* PFollowComp = NULL;
-	UNavigationSystem* NavSys = UNavigationSystem::GetCurrent(Controller->GetWorld());
-
+	UPathFollowingComponent* PFollowComp = nullptr;
 	Controller->InitNavigationControl(PFollowComp);
 
 	if (PFollowComp == nullptr)
@@ -958,7 +969,13 @@ void UNavigationSystem::SimpleMoveToLocation(AController* Controller, const FVec
 		return;
 	}
 
-	if (NavSys && !PFollowComp->HasReached(GoalLocation))
+	if (PFollowComp->HasReached(GoalLocation))
+	{
+		// make sure previous move request gets aborted
+		PFollowComp->AbortMove(TEXT("Aborting move due to new move request finishing with AlreadyAtGoal"), FAIRequestID::AnyRequest);
+		PFollowComp->SetLastMoveAtGoal(true);
+	}
+	else
 	{
 		const ANavigationData* NavData = NavSys->GetNavDataForProps(Controller->GetNavAgentProperties());
 		FPathFindingQuery Query(Controller, NavData, Controller->GetNavAgentLocation(), GoalLocation);
@@ -967,6 +984,11 @@ void UNavigationSystem::SimpleMoveToLocation(AController* Controller, const FVec
 		{
 			PFollowComp->RequestMove(Result.Path, NULL);
 		}
+		else if (PFollowComp->GetStatus() != EPathFollowingStatus::Idle)
+		{
+			PFollowComp->AbortMove(TEXT("Aborting move due to new move request failing to generate a path"), FAIRequestID::AnyRequest);
+			PFollowComp->SetLastMoveAtGoal(false);
+		}
 	}
 }
 
