Renaming FVector ClampSize and SafeNormal functions for clarity

Used DEPRECATED Macro to keep old names around in case I've missed anything

diff --git a/Engine/Plugins/2D/Paper2D/Source/Paper2D/Private/Terrain/PaperTerrainComponent.cpp b/Engine/Plugins/2D/Paper2D/Source/Paper2D/Private/Terrain/PaperTerrainComponent.cpp
--- a/Engine/Plugins/2D/Paper2D/Source/Paper2D/Private/Terrain/PaperTerrainComponent.cpp
+++ b/Engine/Plugins/2D/Paper2D/Source/Paper2D/Private/Terrain/PaperTerrainComponent.cpp
@@ -189,8 +189,8 @@ FTransform UPaperTerrainComponent::GetTransformAtDistance(float InDistance) cons
 	const float Param = AssociatedSpline->SplineReparamTable.Eval(InDistance, 0.0f);
 	const FVector Position3D = AssociatedSpline->SplineInfo.Eval(Param, FVector::ZeroVector);
 
-	const FVector Tangent = AssociatedSpline->SplineInfo.EvalDerivative(Param, FVector(1.0f, 0.0f, 0.0f)).SafeNormal();
-	const FVector NormalEst = AssociatedSpline->SplineInfo.EvalSecondDerivative(Param, FVector(0.0f, 1.0f, 0.0f)).SafeNormal();
+	const FVector Tangent = AssociatedSpline->SplineInfo.EvalDerivative(Param, FVector(1.0f, 0.0f, 0.0f)).GetSafeNormal();
+	const FVector NormalEst = AssociatedSpline->SplineInfo.EvalSecondDerivative(Param, FVector(0.0f, 1.0f, 0.0f)).GetSafeNormal();
 	const FVector Bitangent = FVector::CrossProduct(Tangent, NormalEst);
 	const FVector Normal = FVector::CrossProduct(Bitangent, Tangent);
 	const FVector Floop = FVector::CrossProduct(PaperAxisZ, Tangent);
diff --git a/Engine/Plugins/2D/Paper2D/Source/Paper2DEditor/Private/SpriteEditor/SpriteEditorViewportClient.cpp b/Engine/Plugins/2D/Paper2D/Source/Paper2DEditor/Private/SpriteEditor/SpriteEditorViewportClient.cpp
--- a/Engine/Plugins/2D/Paper2D/Source/Paper2DEditor/Private/SpriteEditor/SpriteEditorViewportClient.cpp
+++ b/Engine/Plugins/2D/Paper2D/Source/Paper2DEditor/Private/SpriteEditor/SpriteEditorViewportClient.cpp
@@ -391,7 +391,7 @@ void FSpriteEditorViewportClient::DrawGeometry(FViewport& InViewport, FSceneView
 			// Draw the normal tick
 			if (bShowNormals)
 			{
-				const FVector2D Direction = (NextScreenPos - ScreenPos).SafeNormal();
+				const FVector2D Direction = (NextScreenPos - ScreenPos).GetSafeNormal();
 				const FVector2D Normal = FVector2D(-Direction.Y, Direction.X);
 
 				const FVector2D Midpoint = (ScreenPos + NextScreenPos) * 0.5f;
diff --git a/Engine/Plugins/Runtime/CableComponent/Source/CableComponent/Private/CableComponent.cpp b/Engine/Plugins/Runtime/CableComponent/Source/CableComponent/Private/CableComponent.cpp
--- a/Engine/Plugins/Runtime/CableComponent/Source/CableComponent/Private/CableComponent.cpp
+++ b/Engine/Plugins/Runtime/CableComponent/Source/CableComponent/Private/CableComponent.cpp
@@ -164,14 +164,14 @@ public:
 			// Find direction of cable at this point, by averaging previous and next points
 			const int32 PrevIndex = FMath::Max(0, PointIdx-1);
 			const int32 NextIndex = FMath::Min(PointIdx+1, NumPoints-1);
-			const FVector ForwardDir = (InPoints[NextIndex] - InPoints[PrevIndex]).SafeNormal();
+			const FVector ForwardDir = (InPoints[NextIndex] - InPoints[PrevIndex]).GetSafeNormal();
 
 			// Find a side vector at this point
 			const FVector WorldUp(0,0,1);
-			const FVector RightDir = (WorldUp ^ ForwardDir).SafeNormal();
+			const FVector RightDir = (WorldUp ^ ForwardDir).GetSafeNormal();
 
 			// Find an up vector
-			const FVector UpDir = (ForwardDir ^ RightDir).SafeNormal();
+			const FVector UpDir = (ForwardDir ^ RightDir).GetSafeNormal();
 
 			// Generate a ring of verts
 			for(int32 VertIdx = 0; VertIdx<NumRingVerts; VertIdx++)
diff --git a/Engine/Plugins/Runtime/CustomMeshComponent/Source/CustomMeshComponent/Private/CustomMeshComponent.cpp b/Engine/Plugins/Runtime/CustomMeshComponent/Source/CustomMeshComponent/Private/CustomMeshComponent.cpp
--- a/Engine/Plugins/Runtime/CustomMeshComponent/Source/CustomMeshComponent/Private/CustomMeshComponent.cpp
+++ b/Engine/Plugins/Runtime/CustomMeshComponent/Source/CustomMeshComponent/Private/CustomMeshComponent.cpp
@@ -93,9 +93,9 @@ public:
 			const FVector Edge01 = (Tri.Vertex1 - Tri.Vertex0);
 			const FVector Edge02 = (Tri.Vertex2 - Tri.Vertex0);
 
-			const FVector TangentX = Edge01.SafeNormal();
-			const FVector TangentZ = (Edge02 ^ Edge01).SafeNormal();
-			const FVector TangentY = (TangentX ^ TangentZ).SafeNormal();
+			const FVector TangentX = Edge01.GetSafeNormal();
+			const FVector TangentZ = (Edge02 ^ Edge01).GetSafeNormal();
+			const FVector TangentY = (TangentX ^ TangentZ).GetSafeNormal();
 
 			FDynamicMeshVertex Vert0;
 			Vert0.Position = Tri.Vertex0;
diff --git a/Engine/Source/Developer/MeshSimplifier/Private/QuadricMeshReduction.cpp b/Engine/Source/Developer/MeshSimplifier/Private/QuadricMeshReduction.cpp
--- a/Engine/Source/Developer/MeshSimplifier/Private/QuadricMeshReduction.cpp
+++ b/Engine/Source/Developer/MeshSimplifier/Private/QuadricMeshReduction.cpp
@@ -556,7 +556,7 @@ private:
 				Normals->GetTuple(WedgeIndex, (float*)&RawMesh.WedgeTangentZ[WedgeIndex]);
 				if (!bReverseWinding)
 				{
-					RawMesh.WedgeTangentZ[WedgeIndex] = (-RawMesh.WedgeTangentZ[WedgeIndex]).SafeNormal();
+					RawMesh.WedgeTangentZ[WedgeIndex] = (-RawMesh.WedgeTangentZ[WedgeIndex]).GetSafeNormal();
 				}
 			}
 		}
diff --git a/Engine/Source/Developer/MeshUtilities/Private/MeshUtilities.cpp b/Engine/Source/Developer/MeshUtilities/Private/MeshUtilities.cpp
--- a/Engine/Source/Developer/MeshUtilities/Private/MeshUtilities.cpp
+++ b/Engine/Source/Developer/MeshUtilities/Private/MeshUtilities.cpp
@@ -382,7 +382,7 @@ void FMeshUtilities::GenerateSignedDistanceFieldVolumeData(
 				V2.Z = 0;
 			}
 
-			const FVector LocalNormal = ((V1 - V2) ^ (V0 - V2)).SafeNormal();
+			const FVector LocalNormal = ((V1 - V2) ^ (V0 - V2)).GetSafeNormal();
 
 			// No degenerates
 			if (LocalNormal.IsUnit())
@@ -1482,7 +1482,7 @@ static void ComputeTriangleTangents(
 			P[i] = GetPositionForWedge(RawMesh, TriangleIndex * 3 + i);
 		}
 
-		const FVector Normal = ((P[1] - P[2])^(P[0] - P[2])).SafeNormal(ComparisonThreshold);
+		const FVector Normal = ((P[1] - P[2])^(P[0] - P[2])).GetSafeNormal(ComparisonThreshold);
 		FMatrix	ParameterToLocal(
 			FPlane(P[1].X - P[0].X, P[1].Y - P[0].Y, P[1].Z - P[0].Z, 0),
 			FPlane(P[2].X - P[0].X, P[2].Y - P[0].Y, P[2].Z - P[0].Z, 0),
@@ -1503,8 +1503,8 @@ static void ComputeTriangleTangents(
 		// Use InverseSlow to catch singular matrices.  Inverse can miss this sometimes.
 		const FMatrix TextureToLocal = ParameterToTexture.Inverse() * ParameterToLocal;
 
-		TriangleTangentX.Add(TextureToLocal.TransformVector(FVector(1,0,0)).SafeNormal());
-		TriangleTangentY.Add(TextureToLocal.TransformVector(FVector(0,1,0)).SafeNormal());
+		TriangleTangentX.Add(TextureToLocal.TransformVector(FVector(1,0,0)).GetSafeNormal());
+		TriangleTangentY.Add(TextureToLocal.TransformVector(FVector(0,1,0)).GetSafeNormal());
 		TriangleTangentZ.Add(Normal);
 
 		FVector::CreateOrthonormalBasis(
@@ -2044,9 +2044,9 @@ static FStaticMeshBuildVertex BuildStaticMeshVertex(FRawMesh const& RawMesh, int
 	Vertex.Position = GetPositionForWedge(RawMesh, WedgeIndex) * BuildScale;
 
 	const FMatrix ScaleMatrix = FScaleMatrix( BuildScale ).Inverse().GetTransposed();	
-	Vertex.TangentX = ScaleMatrix.TransformVector(RawMesh.WedgeTangentX[WedgeIndex]).SafeNormal();
-	Vertex.TangentY = ScaleMatrix.TransformVector(RawMesh.WedgeTangentY[WedgeIndex]).SafeNormal();
-	Vertex.TangentZ = ScaleMatrix.TransformVector(RawMesh.WedgeTangentZ[WedgeIndex]).SafeNormal();
+	Vertex.TangentX = ScaleMatrix.TransformVector(RawMesh.WedgeTangentX[WedgeIndex]).GetSafeNormal();
+	Vertex.TangentY = ScaleMatrix.TransformVector(RawMesh.WedgeTangentY[WedgeIndex]).GetSafeNormal();
+	Vertex.TangentZ = ScaleMatrix.TransformVector(RawMesh.WedgeTangentZ[WedgeIndex]).GetSafeNormal();
 
 	if (RawMesh.WedgeColors.IsValidIndex(WedgeIndex))
 	{
@@ -2633,15 +2633,15 @@ bool FMeshUtilities::BuildSkeletalMesh( FStaticLODModel& LODModel, const FRefere
 				);
 
 			FMatrix	TextureToLocal = ParameterToTexture.Inverse() * ParameterToLocal;
-			FVector	TangentX = TextureToLocal.TransformVector(FVector(1,0,0)).SafeNormal(),
-				TangentY = TextureToLocal.TransformVector(FVector(0,1,0)).SafeNormal(),
+			FVector	TangentX = TextureToLocal.TransformVector(FVector(1,0,0)).GetSafeNormal(),
+				TangentY = TextureToLocal.TransformVector(FVector(0,1,0)).GetSafeNormal(),
 				TangentZ;
 
 			TangentX = TangentX - TriangleNormal * (TangentX | TriangleNormal);
 			TangentY = TangentY - TriangleNormal * (TangentY | TriangleNormal);
 
-			FaceTangentX[FaceIndex] = TangentX.SafeNormal();
-			FaceTangentY[FaceIndex] = TangentY.SafeNormal();
+			FaceTangentX[FaceIndex] = TangentX.GetSafeNormal();
+			FaceTangentY[FaceIndex] = TangentY.GetSafeNormal();
 		}
 	}
 
@@ -2871,12 +2871,12 @@ bool FMeshUtilities::BuildSkeletalMesh( FStaticLODModel& LODModel, const FRefere
 
 			if( bComputeNormals || bComputeTangents )
 			{
-				TangentX = VertexTangentX[VertexIndex].SafeNormal();
-				TangentY = VertexTangentY[VertexIndex].SafeNormal();
+				TangentX = VertexTangentX[VertexIndex].GetSafeNormal();
+				TangentY = VertexTangentY[VertexIndex].GetSafeNormal();
 
 				if( bComputeNormals )
 				{
-					TangentZ = VertexTangentZ[VertexIndex].SafeNormal();
+					TangentZ = VertexTangentZ[VertexIndex].GetSafeNormal();
 				}
 				else
 				{
diff --git a/Engine/Source/Developer/NewLogVisualizer/Private/VisualLoggerRenderingActor.cpp b/Engine/Source/Developer/NewLogVisualizer/Private/VisualLoggerRenderingActor.cpp
--- a/Engine/Source/Developer/NewLogVisualizer/Private/VisualLoggerRenderingActor.cpp
+++ b/Engine/Source/Developer/NewLogVisualizer/Private/VisualLoggerRenderingActor.cpp
@@ -186,7 +186,7 @@ void AVisualLoggerRenderingActor::OnItemSelectionChanged(const FVisualLogDevice:
 #endif
 
 	const float Length = 100;
-	const FVector DirectionNorm = FVector(0, 0, 1).SafeNormal();
+	const FVector DirectionNorm = FVector(0, 0, 1).GetSafeNormal();
 	FVector YAxis, ZAxis;
 	DirectionNorm.FindBestAxisVectors(YAxis, ZAxis);
 	Cones.Add(FDebugRenderSceneProxy::FCone(FScaleMatrix(FVector(Length)) * FMatrix(DirectionNorm, YAxis, ZAxis, Entry->Location), 5, 5, FColor::Red));
@@ -361,7 +361,7 @@ void AVisualLoggerRenderingActor::OnItemSelectionChanged(const FVisualLogDevice:
 			for (int32 Index = 0; Index + 2 < ElementToDraw->Points.Num(); Index += 3)
 			{
 				const FVector Orgin = ElementToDraw->Points[Index];
-				const FVector Direction = ElementToDraw->Points[Index + 1].SafeNormal();
+				const FVector Direction = ElementToDraw->Points[Index + 1].GetSafeNormal();
 				const FVector Angles = ElementToDraw->Points[Index + 2];
 				const float Length = Angles.X;
 
diff --git a/Engine/Source/Developer/TextureCompressor/Private/TextureCompressorModule.cpp b/Engine/Source/Developer/TextureCompressor/Private/TextureCompressorModule.cpp
--- a/Engine/Source/Developer/TextureCompressor/Private/TextureCompressorModule.cpp
+++ b/Engine/Source/Developer/TextureCompressor/Private/TextureCompressorModule.cpp
@@ -1589,7 +1589,7 @@ static void NormalizeMip(FImage& InOutMip)
 
 		FVector Normal = FVector(Color.R * 2.0f - 1.0f, Color.G * 2.0f - 1.0f, Color.B * 2.0f - 1.0f);
 
-		Normal = Normal.SafeNormal();
+		Normal = Normal.GetSafeNormal();
 
 		Color = FLinearColor(Normal.X * 0.5f + 0.5f, Normal.Y * 0.5f + 0.5f, Normal.Z * 0.5f + 0.5f, Color.A);
 	}
diff --git a/Engine/Source/Editor/BehaviorTreeEditor/Private/BehaviorTreeConnectionDrawingPolicy.cpp b/Engine/Source/Editor/BehaviorTreeEditor/Private/BehaviorTreeConnectionDrawingPolicy.cpp
--- a/Engine/Source/Editor/BehaviorTreeEditor/Private/BehaviorTreeConnectionDrawingPolicy.cpp
+++ b/Engine/Source/Editor/BehaviorTreeEditor/Private/BehaviorTreeConnectionDrawingPolicy.cpp
@@ -116,8 +116,8 @@ void FBehaviorTreeConnectionDrawingPolicy::Internal_DrawLineWithArrow(const FVec
 	const float LineSeparationAmount = 4.5f;
 
 	const FVector2D DeltaPos = EndAnchorPoint - StartAnchorPoint;
-	const FVector2D UnitDelta = DeltaPos.SafeNormal();
-	const FVector2D Normal = FVector2D(DeltaPos.Y, -DeltaPos.X).SafeNormal();
+	const FVector2D UnitDelta = DeltaPos.GetSafeNormal();
+	const FVector2D Normal = FVector2D(DeltaPos.Y, -DeltaPos.X).GetSafeNormal();
 
 	// Come up with the final start/end points
 	const FVector2D DirectionBias = Normal * LineSeparationAmount;
@@ -166,7 +166,7 @@ void FBehaviorTreeConnectionDrawingPolicy::DrawConnection(int32 LayerId, const F
 	const FVector2D& P1 = End;
 
 	const FVector2D Delta = End-Start;
-	const FVector2D NormDelta = Delta.SafeNormal();
+	const FVector2D NormDelta = Delta.GetSafeNormal();
 
 	const FVector2D P0Tangent = NormDelta;
 	const FVector2D P1Tangent = NormDelta;
diff --git a/Engine/Source/Editor/ComponentVisualizers/Private/SplineComponentVisualizer.cpp b/Engine/Source/Editor/ComponentVisualizers/Private/SplineComponentVisualizer.cpp
--- a/Engine/Source/Editor/ComponentVisualizers/Private/SplineComponentVisualizer.cpp
+++ b/Engine/Source/Editor/ComponentVisualizers/Private/SplineComponentVisualizer.cpp
@@ -351,8 +351,8 @@ bool FSplineComponentVisualizer::GetCustomInputCoordinateSystem(const FEditorVie
 			if (SelectedKeyIndex < SplineComp->SplineInfo.Points.Num())
 			{
 				const auto& Point = SplineComp->SplineInfo.Points[SelectedKeyIndex];
-				const FVector Tangent = Point.ArriveTangent.IsNearlyZero() ? FVector(1.0f, 0.0f, 0.0f) : Point.ArriveTangent.SafeNormal();
-				const FVector Bitangent = (Tangent.Z == 1.0f) ? FVector(1.0f, 0.0f, 0.0f) : FVector(-Tangent.Y, Tangent.X, 0.0f).SafeNormal();
+				const FVector Tangent = Point.ArriveTangent.IsNearlyZero() ? FVector(1.0f, 0.0f, 0.0f) : Point.ArriveTangent.GetSafeNormal();
+				const FVector Bitangent = (Tangent.Z == 1.0f) ? FVector(1.0f, 0.0f, 0.0f) : FVector(-Tangent.Y, Tangent.X, 0.0f).GetSafeNormal();
 				const FVector Normal = FVector::CrossProduct(Tangent, Bitangent);
 
 				OutMatrix = FMatrix(Tangent, Bitangent, Normal, FVector::ZeroVector) * FQuatRotationTranslationMatrix(SplineComp->ComponentToWorld.GetRotation(), FVector::ZeroVector);
diff --git a/Engine/Source/Editor/EnvironmentQueryEditor/Private/EnvironmentQueryConnectionDrawingPolicy.cpp b/Engine/Source/Editor/EnvironmentQueryEditor/Private/EnvironmentQueryConnectionDrawingPolicy.cpp
--- a/Engine/Source/Editor/EnvironmentQueryEditor/Private/EnvironmentQueryConnectionDrawingPolicy.cpp
+++ b/Engine/Source/Editor/EnvironmentQueryEditor/Private/EnvironmentQueryConnectionDrawingPolicy.cpp
@@ -72,8 +72,8 @@ void FEnvironmentQueryConnectionDrawingPolicy::Internal_DrawLineWithArrow(const
 	const float LineSeparationAmount = 4.5f;
 
 	const FVector2D DeltaPos = EndAnchorPoint - StartAnchorPoint;
-	const FVector2D UnitDelta = DeltaPos.SafeNormal();
-	const FVector2D Normal = FVector2D(DeltaPos.Y, -DeltaPos.X).SafeNormal();
+	const FVector2D UnitDelta = DeltaPos.GetSafeNormal();
+	const FVector2D Normal = FVector2D(DeltaPos.Y, -DeltaPos.X).GetSafeNormal();
 
 	// Come up with the final start/end points
 	const FVector2D DirectionBias = Normal * LineSeparationAmount;
@@ -122,7 +122,7 @@ void FEnvironmentQueryConnectionDrawingPolicy::DrawConnection(int32 LayerId, con
 	const FVector2D& P1 = End;
 
 	const FVector2D Delta = End-Start;
-	const FVector2D NormDelta = Delta.SafeNormal();
+	const FVector2D NormDelta = Delta.GetSafeNormal();
 
 	const FVector2D P0Tangent = NormDelta;
 	const FVector2D P1Tangent = NormDelta;
diff --git a/Engine/Source/Editor/GeometryMode/Private/GeometryModifiers.cpp b/Engine/Source/Editor/GeometryMode/Private/GeometryModifiers.cpp
--- a/Engine/Source/Editor/GeometryMode/Private/GeometryModifiers.cpp
+++ b/Engine/Source/Editor/GeometryMode/Private/GeometryModifiers.cpp
@@ -1747,7 +1747,7 @@ static ABrush* ClipBrushAgainstPlane( const FPlane& InPlane, ABrush* InBrush)
 
 		FPoly front, back;
 
-		int32 res = Poly.SplitWithPlane( PlaneBase, InPlane.SafeNormal(), &front, &back, true );
+		int32 res = Poly.SplitWithPlane( PlaneBase, InPlane.GetSafeNormal(), &front, &back, true );
 
 		switch( res )
 		{
@@ -2492,7 +2492,7 @@ bool UGeomModifier_Split::OnApply()
 
 		// Generate a base and a normal for the cutting plane
 
-		const FVector PlaneNormal( (Vtx1 - Vtx0).SafeNormal() );
+		const FVector PlaneNormal( (Vtx1 - Vtx0).GetSafeNormal() );
 		const FVector PlaneBase = 0.5f*(Vtx1 + Vtx0);
 
 		// Clip the selected polygon against the cutting plane
@@ -2570,7 +2570,7 @@ bool UGeomModifier_Split::OnApply()
 
 		const FVector v0 = *Vertex0->GetActualVertex( Vertex0->ActualVertexIndices[0] );
 		const FVector v1 = *Vertex1->GetActualVertex( Vertex1->ActualVertexIndices[0] );
-		const FVector PlaneNormal( (v1 - v0).SafeNormal() );
+		const FVector PlaneNormal( (v1 - v0).GetSafeNormal() );
 		const FVector PlaneBase = 0.5f*(v1 + v0);
 
 		ABrush* Brush = GeomObject->GetActualBrush();
diff --git a/Engine/Source/Editor/GraphEditor/Private/AnimationStateNodes/SGraphNodeAnimTransition.cpp b/Engine/Source/Editor/GraphEditor/Private/AnimationStateNodes/SGraphNodeAnimTransition.cpp
--- a/Engine/Source/Editor/GraphEditor/Private/AnimationStateNodes/SGraphNodeAnimTransition.cpp
+++ b/Engine/Source/Editor/GraphEditor/Private/AnimationStateNodes/SGraphNodeAnimTransition.cpp
@@ -355,11 +355,11 @@ void SGraphNodeAnimTransition::PositionBetweenTwoNodesWithOffset(const FGeometry
 		DeltaPos = FVector2D(10.0f, 0.0f);
 	}
 
-	const FVector2D Normal = FVector2D(DeltaPos.Y, -DeltaPos.X).SafeNormal();
+	const FVector2D Normal = FVector2D(DeltaPos.Y, -DeltaPos.X).GetSafeNormal();
 
 	const FVector2D NewCenter = StartAnchorPoint + (0.5f * DeltaPos) + (Height * Normal);
 
-	FVector2D DeltaNormal = DeltaPos.SafeNormal();
+	FVector2D DeltaNormal = DeltaPos.GetSafeNormal();
 	
 	// Calculate node offset in the case of multiple transitions between the same two nodes
 	// MultiNodeOffset: the offset where 0 is the centre of the transition, -1 is 1 <size of node>
diff --git a/Engine/Source/Editor/GraphEditor/Private/SGraphNode.cpp b/Engine/Source/Editor/GraphEditor/Private/SGraphNode.cpp
--- a/Engine/Source/Editor/GraphEditor/Private/SGraphNode.cpp
+++ b/Engine/Source/Editor/GraphEditor/Private/SGraphNode.cpp
@@ -1167,7 +1167,7 @@ void SGraphNode::PositionThisNodeBetweenOtherNodes(const FVector2D& PrevPos, con
 		DeltaPos = FVector2D(10.0f, 0.0f);
 	}
 
-	const FVector2D Normal = FVector2D(DeltaPos.Y, -DeltaPos.X).SafeNormal();
+	const FVector2D Normal = FVector2D(DeltaPos.Y, -DeltaPos.X).GetSafeNormal();
 
 	const FVector2D SlidingCapsuleBias = FVector2D::ZeroVector;//(0.5f * FMath::Sin(Normal.X * (float)HALF_PI) * DesiredSize.X, 0.0f);
 
diff --git a/Engine/Source/Editor/GraphEditor/Private/StateMachineConnectionDrawingPolicy.cpp b/Engine/Source/Editor/GraphEditor/Private/StateMachineConnectionDrawingPolicy.cpp
--- a/Engine/Source/Editor/GraphEditor/Private/StateMachineConnectionDrawingPolicy.cpp
+++ b/Engine/Source/Editor/GraphEditor/Private/StateMachineConnectionDrawingPolicy.cpp
@@ -129,8 +129,8 @@ void FStateMachineConnectionDrawingPolicy::Internal_DrawLineWithArrow(const FVec
 	const float LineSeparationAmount = 4.5f;
 
 	const FVector2D DeltaPos = EndAnchorPoint - StartAnchorPoint;
-	const FVector2D UnitDelta = DeltaPos.SafeNormal();
-	const FVector2D Normal = FVector2D(DeltaPos.Y, -DeltaPos.X).SafeNormal();
+	const FVector2D UnitDelta = DeltaPos.GetSafeNormal();
+	const FVector2D Normal = FVector2D(DeltaPos.Y, -DeltaPos.X).GetSafeNormal();
 
 	// Come up with the final start/end points
 	const FVector2D DirectionBias = Normal * LineSeparationAmount;
@@ -179,7 +179,7 @@ void FStateMachineConnectionDrawingPolicy::DrawConnection(int32 LayerId, const F
 	const FVector2D& P1 = End;
 
 	const FVector2D Delta = End-Start;
-	const FVector2D NormDelta = Delta.SafeNormal();
+	const FVector2D NormDelta = Delta.GetSafeNormal();
 	
 	const FVector2D P0Tangent = NormDelta;
 	const FVector2D P1Tangent = NormDelta;
diff --git a/Engine/Source/Editor/LandscapeEditor/Private/LandscapeEdModeBrushes.cpp b/Engine/Source/Editor/LandscapeEditor/Private/LandscapeEdModeBrushes.cpp
--- a/Engine/Source/Editor/LandscapeEditor/Private/LandscapeEdModeBrushes.cpp
+++ b/Engine/Source/Editor/LandscapeEditor/Private/LandscapeEdModeBrushes.cpp
@@ -1164,7 +1164,7 @@ public:
 			{
 				double SampleTime = FPlatformTime::Seconds();
 				float DeltaTime = (float)(SampleTime - LastMouseSampleTime);
-				FVector2D MouseDirection = MouseDelta.SafeNormal();
+				FVector2D MouseDirection = MouseDelta.GetSafeNormal();
 				float MouseAngle = FMath::Lerp(LastMouseAngle, FMath::Atan2(-MouseDirection.Y, MouseDirection.X), FMath::Min<float>(10.f * DeltaTime, 1.f));		// lerp over 100ms
 				LastMouseAngle = MouseAngle;
 				LastMouseSampleTime = SampleTime;
diff --git a/Engine/Source/Editor/LandscapeEditor/Private/LandscapeEdModePaintTools.cpp b/Engine/Source/Editor/LandscapeEditor/Private/LandscapeEdModePaintTools.cpp
--- a/Engine/Source/Editor/LandscapeEditor/Private/LandscapeEdModePaintTools.cpp
+++ b/Engine/Source/Editor/LandscapeEditor/Private/LandscapeEdModePaintTools.cpp
@@ -174,8 +174,8 @@ public:
 					FVector Vert10 = ToWorld.TransformPosition(FVector((float)X + 1.f, (float)Y + 0.f, DATA_AT(*SourceDataArrayPtr, X + 1, Y + 0)));
 					FVector Vert11 = ToWorld.TransformPosition(FVector((float)X + 1.f, (float)Y + 1.f, DATA_AT(*SourceDataArrayPtr, X + 1, Y + 1)));
 
-					FVector FaceNormal1 = ((Vert00 - Vert10) ^ (Vert10 - Vert11)).SafeNormal();
-					FVector FaceNormal2 = ((Vert11 - Vert01) ^ (Vert01 - Vert00)).SafeNormal();
+					FVector FaceNormal1 = ((Vert00 - Vert10) ^ (Vert10 - Vert11)).GetSafeNormal();
+					FVector FaceNormal2 = ((Vert11 - Vert01) ^ (Vert01 - Vert00)).GetSafeNormal();
 
 					// contribute to the vertex normals.
 					DATA_AT(Normals, X + 1, Y + 0) += FaceNormal1;
@@ -188,7 +188,7 @@ public:
 			{
 				for (int32 X = X1; X <= X2; X++)
 				{
-					DATA_AT(Normals, X, Y) = DATA_AT(Normals, X, Y).SafeNormal();
+					DATA_AT(Normals, X, Y) = DATA_AT(Normals, X, Y).GetSafeNormal();
 				}
 			}
 
@@ -213,7 +213,7 @@ public:
 			if (TotalWeight > 0.f)
 			{
 				AveragePoint /= TotalWeight;
-				AverageNormal = AverageNormal.SafeNormal();
+				AverageNormal = AverageNormal.GetSafeNormal();
 			}
 
 			// Convert to world space
diff --git a/Engine/Source/Editor/LandscapeEditor/Private/LandscapeEdModeRampTool.cpp b/Engine/Source/Editor/LandscapeEditor/Private/LandscapeEdModeRampTool.cpp
--- a/Engine/Source/Editor/LandscapeEditor/Private/LandscapeEdModeRampTool.cpp
+++ b/Engine/Source/Editor/LandscapeEditor/Private/LandscapeEdModeRampTool.cpp
@@ -344,7 +344,7 @@ public:
 
 			if (NumPoints == 2)
 			{
-				const FVector Side = FVector::CrossProduct(Points[1] - Points[0], FVector(0, 0, 1)).SafeNormal2D();
+				const FVector Side = FVector::CrossProduct(Points[1] - Points[0], FVector(0, 0, 1)).GetSafeNormal2D();
 				const FVector InnerSide = Side * (EdMode->UISettings->RampWidth * 0.5f * (1 - EdMode->UISettings->RampSideFalloff));
 				const FVector OuterSide = Side * (EdMode->UISettings->RampWidth * 0.5f);
 				FVector InnerVerts[2][2];
@@ -455,7 +455,7 @@ public:
 		const ALandscapeProxy* LandscapeProxy = LandscapeInfo->GetLandscapeProxy();
 		const FTransform LandscapeToWorld = LandscapeProxy->LandscapeActorToWorld();
 
-		const FVector2D Side = FVector2D(FVector::CrossProduct(Points[1] - Points[0], FVector(0,0,1))).SafeNormal();
+		const FVector2D Side = FVector2D(FVector::CrossProduct(Points[1] - Points[0], FVector(0,0,1))).GetSafeNormal();
 		const FVector2D InnerSide = Side * (EdMode->UISettings->RampWidth * 0.5f * (1 - EdMode->UISettings->RampSideFalloff)) / LandscapeToWorld.GetScale3D().X;
 		const FVector2D OuterSide = Side * (EdMode->UISettings->RampWidth * 0.5f) / LandscapeToWorld.GetScale3D().X;
 
diff --git a/Engine/Source/Editor/LandscapeEditor/Private/LandscapeEdModeTools.h b/Engine/Source/Editor/LandscapeEditor/Private/LandscapeEdModeTools.h
--- a/Engine/Source/Editor/LandscapeEditor/Private/LandscapeEdModeTools.h
+++ b/Engine/Source/Editor/LandscapeEditor/Private/LandscapeEdModeTools.h
@@ -419,9 +419,9 @@ struct TLandscapeEditCache
 		FVector Vert10 = FVector(1.f, 0.f, V10);
 		FVector Vert11 = FVector(1.f, 1.f, V11);
 
-		FVector FaceNormal1 = ((Vert00 - Vert10) ^ (Vert10 - Vert11)).SafeNormal();
-		FVector FaceNormal2 = ((Vert11 - Vert01) ^ (Vert01 - Vert00)).SafeNormal();
-		return (FaceNormal1 + FaceNormal2).SafeNormal();
+		FVector FaceNormal1 = ((Vert00 - Vert10) ^ (Vert10 - Vert11)).GetSafeNormal();
+		FVector FaceNormal2 = ((Vert11 - Vert01) ^ (Vert01 - Vert00)).GetSafeNormal();
+		return (FaceNormal1 + FaceNormal2).GetSafeNormal();
 	}
 
 	void SetValue(int32 LandscapeX, int32 LandscapeY, AccessorType Value)
diff --git a/Engine/Source/Editor/MeshPaint/Private/MeshPaintEdMode.cpp b/Engine/Source/Editor/MeshPaint/Private/MeshPaintEdMode.cpp
--- a/Engine/Source/Editor/MeshPaint/Private/MeshPaintEdMode.cpp
+++ b/Engine/Source/Editor/MeshPaint/Private/MeshPaintEdMode.cpp
@@ -1595,7 +1595,7 @@ void FEdModeMeshPaint::PaintMeshVertices(
 					}
 
 					// Check to see if the triangle is front facing
-					FVector TriangleNormal = ( TriVertices[ 1 ] - TriVertices[ 0 ] ^ TriVertices[ 2 ] - TriVertices[ 0 ] ).SafeNormal();
+					FVector TriangleNormal = ( TriVertices[ 1 ] - TriVertices[ 0 ] ^ TriVertices[ 2 ] - TriVertices[ 0 ] ).GetSafeNormal();
 					const float SignedPlaneDist = FVector::PointPlaneDist( ComponentSpaceCameraPosition, TriVertices[ 0 ], TriangleNormal );
 					if( !bOnlyFrontFacing || SignedPlaneDist < 0.0f )
 					{
@@ -1831,7 +1831,7 @@ void FEdModeMeshPaint::PaintMeshTexture( UStaticMeshComponent* StaticMeshCompone
 		{
 			// Check to see if the triangle is front facing
 			FMeshTriangle const& CurrentTri = TriIt.GetCurrentElement();
-			FVector TriangleNormal = ( CurrentTri.Vertices[ 1 ] - CurrentTri.Vertices[ 0 ] ^ CurrentTri.Vertices[ 2 ] - CurrentTri.Vertices[ 0 ] ).SafeNormal();
+			FVector TriangleNormal = ( CurrentTri.Vertices[ 1 ] - CurrentTri.Vertices[ 0 ] ^ CurrentTri.Vertices[ 2 ] - CurrentTri.Vertices[ 0 ] ).GetSafeNormal();
 			const float SignedPlaneDist = FVector::PointPlaneDist( ComponentSpaceCameraPosition, CurrentTri.Vertices[ 0 ], TriangleNormal );
 			if( !bOnlyFrontFacing || SignedPlaneDist < 0.0f )
 			{
diff --git a/Engine/Source/Editor/PhAT/Private/PhATPreviewViewportClient.cpp b/Engine/Source/Editor/PhAT/Private/PhATPreviewViewportClient.cpp
--- a/Engine/Source/Editor/PhAT/Private/PhATPreviewViewportClient.cpp
+++ b/Engine/Source/Editor/PhAT/Private/PhATPreviewViewportClient.cpp
@@ -479,7 +479,7 @@ bool FPhATEdPreviewViewportClient::InputWidgetDelta( FViewport* Viewport, EAxisL
 
 			if ( SharedData->WidgetMode == FWidget::WM_Translate )
 			{
-				FVector Dir = SelectedObject.WidgetTM.InverseTransformVector( Drag.SafeNormal() );
+				FVector Dir = SelectedObject.WidgetTM.InverseTransformVector( Drag.GetSafeNormal() );
 				FVector DragVec = Dir * Drag.Size();
 				SelectedObject.ManipulateTM.AddToTranslation( DragVec );
 			}
diff --git a/Engine/Source/Editor/StaticMeshEditor/Private/StaticMeshEditorViewportClient.cpp b/Engine/Source/Editor/StaticMeshEditor/Private/StaticMeshEditorViewportClient.cpp
--- a/Engine/Source/Editor/StaticMeshEditor/Private/StaticMeshEditorViewportClient.cpp
+++ b/Engine/Source/Editor/StaticMeshEditor/Private/StaticMeshEditorViewportClient.cpp
@@ -527,17 +527,17 @@ void FStaticMeshEditorViewportClient::Draw(const FSceneView* View,FPrimitiveDraw
 
 			if( bDrawNormals )
 			{
-				PDI->DrawLine( WorldPos, WorldPos+LocalToWorldInverseTranspose.TransformVector( Normal ).SafeNormal() * Len, FLinearColor( 0.0f, 1.0f, 0.0f), SDPG_World );
+				PDI->DrawLine( WorldPos, WorldPos+LocalToWorldInverseTranspose.TransformVector( Normal ).GetSafeNormal() * Len, FLinearColor( 0.0f, 1.0f, 0.0f), SDPG_World );
 			}
 
 			if( bDrawTangents )
 			{
-				PDI->DrawLine( WorldPos, WorldPos+LocalToWorldInverseTranspose.TransformVector( Tangent ).SafeNormal() * Len, FLinearColor( 1.0f, 0.0f, 0.0f), SDPG_World );
+				PDI->DrawLine( WorldPos, WorldPos+LocalToWorldInverseTranspose.TransformVector( Tangent ).GetSafeNormal() * Len, FLinearColor( 1.0f, 0.0f, 0.0f), SDPG_World );
 			}
 
 			if( bDrawBinormals )
 			{
-				PDI->DrawLine( WorldPos, WorldPos+LocalToWorldInverseTranspose.TransformVector( Binormal ).SafeNormal() * Len, FLinearColor( 0.0f, 0.0f, 1.0f), SDPG_World );
+				PDI->DrawLine( WorldPos, WorldPos+LocalToWorldInverseTranspose.TransformVector( Binormal ).GetSafeNormal() * Len, FLinearColor( 0.0f, 0.0f, 1.0f), SDPG_World );
 			}
 		}	
 	}
@@ -915,7 +915,7 @@ void FStaticMeshEditorViewportClient::ProcessClick(class FSceneView& InView, cla
 						// Compute the per-triangle normal
 						const FVector BA = A - B;
 						const FVector CA = A - C;
-						const FVector TriangleNormal = (CA ^ BA).SafeNormal();
+						const FVector TriangleNormal = (CA ^ BA).GetSafeNormal();
 
 						// Transform the view position from world to component space
 						const FVector ComponentSpaceViewOrigin = StaticMeshComponent->ComponentToWorld.InverseTransformPosition( View->ViewMatrices.ViewOrigin);
diff --git a/Engine/Source/Editor/UnrealEd/Private/CameraController.cpp b/Engine/Source/Editor/UnrealEd/Private/CameraController.cpp
--- a/Engine/Source/Editor/UnrealEd/Private/CameraController.cpp
+++ b/Engine/Source/Editor/UnrealEd/Private/CameraController.cpp
@@ -247,7 +247,7 @@ void FEditorCameraController::UpdatePosition( const FCameraControllerUserImpulse
 	// Constrain maximum movement speed
 	if( MovementVelocity.Size() > Config.MaximumMovementSpeed * MovementSpeedScale )
 	{
-		MovementVelocity = MovementVelocity.UnsafeNormal() * Config.MaximumMovementSpeed * MovementSpeedScale;
+		MovementVelocity = MovementVelocity.GetUnsafeNormal() * Config.MaximumMovementSpeed * MovementSpeedScale;
 	}
 
 
diff --git a/Engine/Source/Editor/UnrealEd/Private/Editor/ActorPositioning.cpp b/Engine/Source/Editor/UnrealEd/Private/Editor/ActorPositioning.cpp
--- a/Engine/Source/Editor/UnrealEd/Private/Editor/ActorPositioning.cpp
+++ b/Engine/Source/Editor/UnrealEd/Private/Editor/ActorPositioning.cpp
@@ -119,7 +119,7 @@ FActorPositionTraceResult FActorPositioning::TraceWorldForPosition(const UWorld&
 			{
 				ClosestHitDistance = DistanceToHit;
 				Results.Location = Hit.Location;
-				Results.SurfaceNormal = Hit.Normal.SafeNormal();
+				Results.SurfaceNormal = Hit.Normal.GetSafeNormal();
 				Results.State = FActorPositionTraceResult::HitSuccess;
 				Results.HitActor = Hit.Actor;
 			}
diff --git a/Engine/Source/Editor/UnrealEd/Private/EditorBrushBuilder.cpp b/Engine/Source/Editor/UnrealEd/Private/EditorBrushBuilder.cpp
--- a/Engine/Source/Editor/UnrealEd/Private/EditorBrushBuilder.cpp
+++ b/Engine/Source/Editor/UnrealEd/Private/EditorBrushBuilder.cpp
@@ -943,9 +943,9 @@ void UTetrahedronBuilder::Extrapolate( int32 a, int32 b, int32 c, int32 Count, f
 {
 	if( Count>1 )
 	{
-		int32 ab=Vertexv( InRadius*(GetVertex(a)+GetVertex(b)).SafeNormal() );
-		int32 bc=Vertexv( InRadius*(GetVertex(b)+GetVertex(c)).SafeNormal() );
-		int32 ca=Vertexv( InRadius*(GetVertex(c)+GetVertex(a)).SafeNormal() );
+		int32 ab=Vertexv( InRadius*(GetVertex(a)+GetVertex(b)).GetSafeNormal() );
+		int32 bc=Vertexv( InRadius*(GetVertex(b)+GetVertex(c)).GetSafeNormal() );
+		int32 ca=Vertexv( InRadius*(GetVertex(c)+GetVertex(a)).GetSafeNormal() );
 		Extrapolate(a,ab,ca,Count-1,InRadius);
 		Extrapolate(b,bc,ab,Count-1,InRadius);
 		Extrapolate(c,ca,bc,Count-1,InRadius);
diff --git a/Engine/Source/Editor/UnrealEd/Private/Factories/ActorFactory.cpp b/Engine/Source/Editor/UnrealEd/Private/Factories/ActorFactory.cpp
--- a/Engine/Source/Editor/UnrealEd/Private/Factories/ActorFactory.cpp
+++ b/Engine/Source/Editor/UnrealEd/Private/Factories/ActorFactory.cpp
@@ -107,7 +107,7 @@ FQuat FindActorAlignmentRotation(const FQuat& InActorRotation, const FVector& In
 		const FVector OrthoPitchAxis = FVector::CrossProduct(PitchReferenceAxis, InModelAxis);
 		const float Pitch = FMath::Acos(PitchReferenceAxis | DestNormalModelSpace) - FMath::Acos(PitchReferenceAxis | InModelAxis);//FMath::Asin(OrthoPitchAxis.Size());
 
-		DeltaRotation = FQuat(OrthoPitchAxis.SafeNormal(), Pitch);
+		DeltaRotation = FQuat(OrthoPitchAxis.GetSafeNormal(), Pitch);
 		DeltaRotation.Normalize();
 
 		// Transform the model axis with this new pitch rotation to see if there is any need for yaw
diff --git a/Engine/Source/Editor/UnrealEd/Private/Fbx/FbxStaticMeshImport.cpp b/Engine/Source/Editor/UnrealEd/Private/Fbx/FbxStaticMeshImport.cpp
--- a/Engine/Source/Editor/UnrealEd/Private/Fbx/FbxStaticMeshImport.cpp
+++ b/Engine/Source/Editor/UnrealEd/Private/Fbx/FbxStaticMeshImport.cpp
@@ -550,7 +550,7 @@ bool UnFbx::FFbxImporter::BuildStaticMeshFromGeometry(FbxMesh* Mesh, UStaticMesh
 					FbxVector4 TempValue = LayerElementTangent->GetDirectArray().GetAt(TangentValueIndex);
 					TempValue = TotalMatrixForNormal.MultT(TempValue);
 					FVector TangentX = Converter.ConvertDir(TempValue);
-					RawMesh.WedgeTangentX[WedgeIndex] = TangentX.SafeNormal();
+					RawMesh.WedgeTangentX[WedgeIndex] = TangentX.GetSafeNormal();
 
 					int BinormalMapIndex = (BinormalMappingMode == FbxLayerElement::eByControlPoint) ? 
 						ControlPointIndex : TriangleCornerIndex;
@@ -560,13 +560,13 @@ bool UnFbx::FFbxImporter::BuildStaticMeshFromGeometry(FbxMesh* Mesh, UStaticMesh
 					TempValue = LayerElementBinormal->GetDirectArray().GetAt(BinormalValueIndex);
 					TempValue = TotalMatrixForNormal.MultT(TempValue);
 					FVector TangentY = -Converter.ConvertDir(TempValue);
-					RawMesh.WedgeTangentY[WedgeIndex] = TangentY.SafeNormal();
+					RawMesh.WedgeTangentY[WedgeIndex] = TangentY.GetSafeNormal();
 				}
 
 				FbxVector4 TempValue = LayerElementNormal->GetDirectArray().GetAt(NormalValueIndex);
 				TempValue = TotalMatrixForNormal.MultT(TempValue);
 				FVector TangentZ = Converter.ConvertDir(TempValue);
-				RawMesh.WedgeTangentZ[WedgeIndex] = TangentZ.SafeNormal();
+				RawMesh.WedgeTangentZ[WedgeIndex] = TangentZ.GetSafeNormal();
 			}
 
 			//
diff --git a/Engine/Source/Editor/UnrealEd/Private/GeomTools.cpp b/Engine/Source/Editor/UnrealEd/Private/GeomTools.cpp
--- a/Engine/Source/Editor/UnrealEd/Private/GeomTools.cpp
+++ b/Engine/Source/Editor/UnrealEd/Private/GeomTools.cpp
@@ -316,10 +316,10 @@ static void FixPolyWinding(FUtilPoly2D& Poly)
 		int32 CIndex = (i+1)%Poly.Verts.Num();
 
 		float ABDist = (Poly.Verts[BIndex].Pos - Poly.Verts[AIndex].Pos).Size();
-		FVector2D ABEdge = (Poly.Verts[BIndex].Pos - Poly.Verts[AIndex].Pos).SafeNormal();
+		FVector2D ABEdge = (Poly.Verts[BIndex].Pos - Poly.Verts[AIndex].Pos).GetSafeNormal();
 
 		float BCDist = (Poly.Verts[CIndex].Pos - Poly.Verts[BIndex].Pos).Size();
-		FVector2D BCEdge = (Poly.Verts[CIndex].Pos - Poly.Verts[BIndex].Pos).SafeNormal();
+		FVector2D BCEdge = (Poly.Verts[CIndex].Pos - Poly.Verts[BIndex].Pos).GetSafeNormal();
 
 		// See if points are co-incident or edges are co-linear - if so, remove.
 		if(ABDist < 0.01f || BCDist < 0.01f || ABEdge.Equals(BCEdge, 0.01f))
@@ -595,7 +595,7 @@ FClipSMPolygon Transform2DPolygonToSMPolygon(const FUtilPoly2D& InPoly, const FM
 	}
 
 	// Assume that the matrix defines the polygon's normal.
-	Result.FaceNormal = InMatrix.TransformVector(FVector(0,0,-1)).SafeNormal();
+	Result.FaceNormal = InMatrix.TransformVector(FVector(0,0,-1)).GetSafeNormal();
 
 	return Result;
 }
@@ -636,7 +636,7 @@ static FMatrix ComputeTriangleParameterToAttribute(
 {
 	const FVector AttributeOverS = AttributeV1 - AttributeV0;
 	const FVector AttributeOverT = AttributeV2 - AttributeV0;
-	const FVector AttributeOverNormal = (AttributeOverS ^ AttributeOverT).SafeNormal();
+	const FVector AttributeOverNormal = (AttributeOverS ^ AttributeOverT).GetSafeNormal();
 	return FMatrix(
 		FPlane(	AttributeOverS.X,		AttributeOverS.Y,		AttributeOverS.Z,		0	),
 		FPlane(	AttributeOverT.X,		AttributeOverT.Y,		AttributeOverT.Z,		0	),
@@ -876,7 +876,7 @@ void Split2DPolysWithPlane(FUtilPoly2DSet& PolySet, const FPlane& Plane, const F
 	check(!Intersect.ContainsNaN());
 
 	// Make 2D line.
-	FVector2D Normal2D = FVector2D(LNormal.X, LNormal.Y).SafeNormal();
+	FVector2D Normal2D = FVector2D(LNormal.X, LNormal.Y).GetSafeNormal();
 	FVector2D Base2D = FVector2D(Intersect.X, Intersect.Y);
 	FSplitLine2D Plane2D(Base2D, Normal2D);
 
diff --git a/Engine/Source/Editor/UnrealEd/Private/LevelEditorViewport.cpp b/Engine/Source/Editor/UnrealEd/Private/LevelEditorViewport.cpp
--- a/Engine/Source/Editor/UnrealEd/Private/LevelEditorViewport.cpp
+++ b/Engine/Source/Editor/UnrealEd/Private/LevelEditorViewport.cpp
@@ -86,12 +86,12 @@ FViewportCursorLocation::FViewportCursorLocation( const FSceneView* View, FEdito
 	if ( ViewportClient->IsPerspective() )
 	{
 		Origin = View->ViewMatrices.ViewOrigin;
-		Direction = InvViewMatrix.TransformVector(FVector(InvProjMatrix.TransformFVector4(FVector4(ScreenX * GNearClippingPlane,ScreenY * GNearClippingPlane,0.0f,GNearClippingPlane)))).SafeNormal();
+		Direction = InvViewMatrix.TransformVector(FVector(InvProjMatrix.TransformFVector4(FVector4(ScreenX * GNearClippingPlane,ScreenY * GNearClippingPlane,0.0f,GNearClippingPlane)))).GetSafeNormal();
 	}
 	else
 	{
 		Origin = InvViewMatrix.TransformFVector4(InvProjMatrix.TransformFVector4(FVector4(ScreenX,ScreenY,0.5f,1.0f)));
-		Direction = InvViewMatrix.TransformVector(FVector(0,0,1)).SafeNormal();
+		Direction = InvViewMatrix.TransformVector(FVector(0,0,1)).GetSafeNormal();
 	}
 }
 
diff --git a/Engine/Source/Editor/UnrealEd/Private/PhysicsAssetUtils.cpp b/Engine/Source/Editor/UnrealEd/Private/PhysicsAssetUtils.cpp
--- a/Engine/Source/Editor/UnrealEd/Private/PhysicsAssetUtils.cpp
+++ b/Engine/Source/Editor/UnrealEd/Private/PhysicsAssetUtils.cpp
@@ -309,7 +309,7 @@ FVector ComputeEigenVector(const FMatrix& A)
 		}
 	}
 
-	return Bk.SafeNormal();
+	return Bk.GetSafeNormal();
 }
 
 
@@ -340,7 +340,7 @@ bool CreateCollisionFromBone( UBodySetup* bs, USkeletalMesh* skelMesh, int32 Bon
 			if(ChildPos.Size() > KINDA_SMALL_NUMBER)
 			{
 				// ZAxis for collision geometry lies down axis to child bone.
-				FVector ZAxis = ChildPos.SafeNormal();
+				FVector ZAxis = ChildPos.GetSafeNormal();
 
 				// Then we pick X and Y randomly. 
 				// JTODO: Should project all the vertices onto ZAxis plane and fit a bounding box using calipers or something...
diff --git a/Engine/Source/Editor/UnrealEd/Private/StaticLightingSystem/StaticLightingSystem.cpp b/Engine/Source/Editor/UnrealEd/Private/StaticLightingSystem/StaticLightingSystem.cpp
--- a/Engine/Source/Editor/UnrealEd/Private/StaticLightingSystem/StaticLightingSystem.cpp
+++ b/Engine/Source/Editor/UnrealEd/Private/StaticLightingSystem/StaticLightingSystem.cpp
@@ -1381,9 +1381,9 @@ void FStaticLightingSystem::AddBSPStaticLightingInfo(ULevel* Level, bool bBuildL
 				const FVector& TextureY = Model->Vectors[NodeSurf.vTextureV];
 				const int32 BaseVertexIndex = NodeGroup->Vertices.Num();
 				// Compute the surface's tangent basis.
-				FVector NodeTangentX = Model->Vectors[NodeSurf.vTextureU].SafeNormal();
-				FVector NodeTangentY = Model->Vectors[NodeSurf.vTextureV].SafeNormal();
-				FVector NodeTangentZ = Model->Vectors[NodeSurf.vNormal].SafeNormal();
+				FVector NodeTangentX = Model->Vectors[NodeSurf.vTextureU].GetSafeNormal();
+				FVector NodeTangentY = Model->Vectors[NodeSurf.vTextureV].GetSafeNormal();
+				FVector NodeTangentZ = Model->Vectors[NodeSurf.vNormal].GetSafeNormal();
 
 				// Generate the node's vertices.
 				for(uint32 VertexIndex = 0;VertexIndex < Node.NumVertices;VertexIndex++)
@@ -1536,9 +1536,9 @@ void FStaticLightingSystem::AddBSPStaticLightingInfo(ULevel* Level, TArray<FNode
 				const FVector& TextureY = Model->Vectors[NodeSurf.vTextureV];
 				const int32 BaseVertexIndex = NodeGroup->Vertices.Num();
 				// Compute the surface's tangent basis.
-				FVector NodeTangentX = Model->Vectors[NodeSurf.vTextureU].SafeNormal();
-				FVector NodeTangentY = Model->Vectors[NodeSurf.vTextureV].SafeNormal();
-				FVector NodeTangentZ = Model->Vectors[NodeSurf.vNormal].SafeNormal();
+				FVector NodeTangentX = Model->Vectors[NodeSurf.vTextureU].GetSafeNormal();
+				FVector NodeTangentY = Model->Vectors[NodeSurf.vTextureV].GetSafeNormal();
+				FVector NodeTangentZ = Model->Vectors[NodeSurf.vNormal].GetSafeNormal();
 
 				// Generate the node's vertices.
 				for(uint32 VertexIndex = 0;VertexIndex < Node.NumVertices;VertexIndex++)
diff --git a/Engine/Source/Editor/UnrealEd/Private/StaticMeshEdit.cpp b/Engine/Source/Editor/UnrealEd/Private/StaticMeshEdit.cpp
--- a/Engine/Source/Editor/UnrealEd/Private/StaticMeshEdit.cpp
+++ b/Engine/Source/Editor/UnrealEd/Private/StaticMeshEdit.cpp
@@ -704,7 +704,7 @@ static inline void FTexCoordsToVectors(const FVector& V0, const FVector& UV0,
 {
 	// Create polygon normal.
 	FVector PN = FVector((V0-V1) ^ (V2-V0));
-	PN = PN.SafeNormal();
+	PN = PN.GetSafeNormal();
 
 	FVector UV1( InUV1 );
 	FVector UV2( InUV2 );
diff --git a/Engine/Source/Editor/UnrealEd/Private/TexAlignTools.cpp b/Engine/Source/Editor/UnrealEd/Private/TexAlignTools.cpp
--- a/Engine/Source/Editor/UnrealEd/Private/TexAlignTools.cpp
+++ b/Engine/Source/Editor/UnrealEd/Private/TexAlignTools.cpp
@@ -289,7 +289,7 @@ void UTexAlignerFit::AlignSurf( ETexAlign InTexAlignType, UModel* InModel, FBspS
 	{
 		const FVector& VertexA = FirstPolyVertex;
 		const FVector& VertexB = WorldSpacePolyVertices[ 1 ];
-		FVector UpVec = ( VertexB - VertexA ).SafeNormal();
+		FVector UpVec = ( VertexB - VertexA ).GetSafeNormal();
 		FVector RightVec = InPoly->Normal ^ UpVec;
 		WorldToPolyRotationMatrix.SetIdentity();
 		WorldToPolyRotationMatrix.SetAxes( &RightVec, &UpVec, &InPoly->Normal );
@@ -313,7 +313,7 @@ void UTexAlignerFit::AlignSurf( ETexAlign InTexAlignType, UModel* InModel, FBspS
 		const FVector& NextVertex = WorldSpacePolyVertices[ NextWindingVertexIndex ];
 
 		// Compute the corner angle
-		float AbsDotProduct = FMath::Abs( ( PrevVertex - CurVertex ).SafeNormal() | ( NextVertex - CurVertex ).SafeNormal() );
+		float AbsDotProduct = FMath::Abs( ( PrevVertex - CurVertex ).GetSafeNormal() | ( NextVertex - CurVertex ).GetSafeNormal() );
 
 		// Compute how 'positive' this vertex is relative to the bottom left position in the polygon's plane
 		FVector PolySpaceVertex = WorldToPolyRotationMatrix.InverseTransformVector( CurVertex - FirstPolyVertex );
@@ -353,7 +353,7 @@ void UTexAlignerFit::AlignSurf( ETexAlign InTexAlignType, UModel* InModel, FBspS
 	const int32 NextWindingVertexIndex = ( BestVertexIndex < WorldSpacePolyVertices.Num() - 1 ) ? ( BestVertexIndex + 1 ) : 0;
 	const FVector& NextVertex = WorldSpacePolyVertices[ NextWindingVertexIndex ];
 
-	FVector TextureUpVec = ( NextVertex - BestVertex ).SafeNormal();
+	FVector TextureUpVec = ( NextVertex - BestVertex ).GetSafeNormal();
 	FVector TextureRightVec = InPoly->Normal ^ TextureUpVec;
 
 	FMatrix WorldToTextureRotationMatrix;
diff --git a/Engine/Source/Editor/UnrealEd/Private/UnrealEdSrv.cpp b/Engine/Source/Editor/UnrealEd/Private/UnrealEdSrv.cpp
--- a/Engine/Source/Editor/UnrealEd/Private/UnrealEdSrv.cpp
+++ b/Engine/Source/Editor/UnrealEd/Private/UnrealEdSrv.cpp
@@ -1914,7 +1914,7 @@ FPoly* CreateHugeTrianglePolygonOnPlane( const FPlane* InPlane )
 	// Using the plane normal, get 2 good axis vectors
 
 	FVector A, B;
-	InPlane->SafeNormal().FindBestAxisVectors( A, B );
+	InPlane->GetSafeNormal().FindBestAxisVectors( A, B );
 
 	// Create 4 vertices from the plane origin and the 2 axis generated above
 
@@ -2076,7 +2076,7 @@ bool UUnrealEdEngine::Exec_Actor( UWorld* InWorld, const TCHAR* Str, FOutputDevi
 				{
 					FPlane* plane = SplitterPlanes[x];
 
-					if( plane->SafeNormal().Equals( SplittingPlane->SafeNormal(), NormalTolerance ) )
+					if( plane->GetSafeNormal().Equals( SplittingPlane->GetSafeNormal(), NormalTolerance ) )
 					{
 						bAddPlaneToList = false;
 						break;
@@ -2085,15 +2085,15 @@ bool UUnrealEdEngine::Exec_Actor( UWorld* InWorld, const TCHAR* Str, FOutputDevi
 
 				// As a final test, make sure that this planes normal falls within the normal limits that were defined
 
-				if( FMath::Abs( SplittingPlane->SafeNormal().X ) > NormalLimits.X )
+				if( FMath::Abs( SplittingPlane->GetSafeNormal().X ) > NormalLimits.X )
 				{
 					bAddPlaneToList = false;
 				}
-				if( FMath::Abs( SplittingPlane->SafeNormal().Y ) > NormalLimits.Y )
+				if( FMath::Abs( SplittingPlane->GetSafeNormal().Y ) > NormalLimits.Y )
 				{
 					bAddPlaneToList = false;
 				}
-				if( FMath::Abs( SplittingPlane->SafeNormal().Z ) > NormalLimits.Z )
+				if( FMath::Abs( SplittingPlane->GetSafeNormal().Z ) > NormalLimits.Z )
 				{
 					bAddPlaneToList = false;
 				}
@@ -2132,7 +2132,7 @@ bool UUnrealEdEngine::Exec_Actor( UWorld* InWorld, const TCHAR* Str, FOutputDevi
 				poly = &(*BuilderBrushPolys)[bp];
 
 				FPoly Front, Back;
-				int res = poly->SplitWithPlane( FVector( plane->X, plane->Y, plane->Z ) * plane->W, plane->SafeNormal(), &Front, &Back, true );
+				int res = poly->SplitWithPlane( FVector( plane->X, plane->Y, plane->Z ) * plane->W, plane->GetSafeNormal(), &Front, &Back, true );
 				switch( res )
 				{
 					// Ignore these results.  We don't want them.
@@ -2183,7 +2183,7 @@ bool UUnrealEdEngine::Exec_Actor( UWorld* InWorld, const TCHAR* Str, FOutputDevi
 						plane = new FPlane( poly->Vertices[0], poly->Vertices[1], poly->Vertices[2] );
 
 						FPoly Front, Back;
-						int res = CappingPoly->SplitWithPlane( FVector( plane->X, plane->Y, plane->Z ) * plane->W, plane->SafeNormal(), &Front, &Back, true );
+						int res = CappingPoly->SplitWithPlane( FVector( plane->X, plane->Y, plane->Z ) * plane->W, plane->GetSafeNormal(), &Front, &Back, true );
 						switch( res )
 						{
 							case SP_Split:
diff --git a/Engine/Source/Editor/UnrealEd/Private/VertexSnapping.cpp b/Engine/Source/Editor/UnrealEd/Private/VertexSnapping.cpp
--- a/Engine/Source/Editor/UnrealEd/Private/VertexSnapping.cpp
+++ b/Engine/Source/Editor/UnrealEd/Private/VertexSnapping.cpp
@@ -708,7 +708,7 @@ bool FVertexSnappingImpl::SnapDraggedActorsToNearestVertex( FVector& DragDelta,
 		// Are there selected actors?
 		USelection* Selection = GEditor->GetSelectedActors();
 
-		FVector Direction = DragDelta.SafeNormal();
+		FVector Direction = DragDelta.GetSafeNormal();
 
 		FSceneViewFamilyContext ViewFamily(FSceneViewFamily::ConstructionValues(
 			ViewportClient->Viewport, 
@@ -860,7 +860,7 @@ bool FVertexSnappingImpl::SnapDragLocationToNearestVertex( const FVector& BaseLo
 	bool bSnapped = false;
 	if( !DragDelta.IsNearlyZero() )
 	{
-		FVector Direction = DragDelta.SafeNormal();
+		FVector Direction = DragDelta.GetSafeNormal();
 
 		FSceneViewFamilyContext ViewFamily(FSceneViewFamily::ConstructionValues(
 			ViewportClient->Viewport, 
diff --git a/Engine/Source/Programs/UnrealLightmass/Private/ImportExport/LightmassScene.cpp b/Engine/Source/Programs/UnrealLightmass/Private/ImportExport/LightmassScene.cpp
--- a/Engine/Source/Programs/UnrealLightmass/Private/ImportExport/LightmassScene.cpp
+++ b/Engine/Source/Programs/UnrealLightmass/Private/ImportExport/LightmassScene.cpp
@@ -466,7 +466,7 @@ FLinearColor FLight::GetDirectIntensity(const FVector4& Point, bool bCalculateFo
 	// light profile (IES)
 	float LightProfileAttenuation;
 	{
-		FVector4 NegLightVector = (Position - Point).SafeNormal();
+		FVector4 NegLightVector = (Position - Point).GetSafeNormal();
 
 		LightProfileAttenuation = ComputeLightProfileMultiplier(Dot3(NegLightVector, Direction));
 	}
@@ -970,7 +970,7 @@ float FPointLight::CustomAttenuation(const FVector4& Point, FLMRandomStream& Ran
 
 	// light profile (IES)
 	{
-		FVector4 NegLightVector = (Position - Point).SafeNormal();
+		FVector4 NegLightVector = (Position - Point).GetSafeNormal();
 
 		UnrealAttenuation *= ComputeLightProfileMultiplier(Dot3(NegLightVector, Direction));
 	}
@@ -1020,7 +1020,7 @@ void FPointLight::SampleDirection(FLMRandomStream& RandomStream, FLightRay& Samp
 		this
 		);
 
-	LightSourceNormal = (SurfaceSample.Position - Position).SafeNormal();
+	LightSourceNormal = (SurfaceSample.Position - Position).GetSafeNormal();
 
 	// Approximate the probability of generating this direction as uniform over all the solid angles
 	// This diverges from the actual probability for positions inside the light source radius
@@ -1321,7 +1321,7 @@ FLinearColor FSpotLight::GetDirectIntensity(const FVector4& Point, bool bCalcula
 			OuterCone = FMath::Cos(ClampedOuterConeAngle),
 			InnerCone = FMath::Cos(ClampedInnerConeAngle);
 
-	FVector4 LightVector = (Point - Position).SafeNormal();
+	FVector4 LightVector = (Point - Position).GetSafeNormal();
 	float SpotAttenuation = FMath::Square(FMath::Clamp<float>((Dot3(LightVector, Direction) - OuterCone) / (InnerCone - OuterCone),0.0f,1.0f));
 
 	if( LightFlags & GI_LIGHT_INVERSE_SQUARED )
@@ -1598,7 +1598,7 @@ FLinearColor FMeshAreaLight::GetDirectIntensity(const FVector4& Point, bool bCal
 		{
 			PrimitiveCenter += CurrentPrimitive.Corners[CornerIndex].WorldPosition / 4.0f;
 		}
-		const FVector4 LightVector = (Point - PrimitiveCenter).SafeNormal();
+		const FVector4 LightVector = (Point - PrimitiveCenter).GetSafeNormal();
 		const float NDotL = Dot3(LightVector, CurrentPrimitive.SurfaceNormal);
 		if (NDotL >= 0)
 		{
diff --git a/Engine/Source/Programs/UnrealLightmass/Private/ImportExport/Material.h b/Engine/Source/Programs/UnrealLightmass/Private/ImportExport/Material.h
--- a/Engine/Source/Programs/UnrealLightmass/Private/ImportExport/Material.h
+++ b/Engine/Source/Programs/UnrealLightmass/Private/ImportExport/Material.h
@@ -47,7 +47,7 @@ namespace Lightmass
 		{
 			Normal = MaterialNormal.SampleNormal(UV);
 			Normal.W = 0.0f;
-			Normal = Normal.SafeNormal();
+			Normal = Normal.GetSafeNormal();
 			if( Normal.SizeSquared3() < KINDA_SMALL_NUMBER )
 			{
 				Normal.Set( 0.0f, 0.0f, 1.0f, 0.0f );
diff --git a/Engine/Source/Programs/UnrealLightmass/Private/Lighting/Cache.h b/Engine/Source/Programs/UnrealLightmass/Private/Lighting/Cache.h
--- a/Engine/Source/Programs/UnrealLightmass/Private/Lighting/Cache.h
+++ b/Engine/Source/Programs/UnrealLightmass/Private/Lighting/Cache.h
@@ -268,8 +268,8 @@ bool TLightingCache<SampleType>::InterpolateLighting(
 		// Query points behind the lighting record may have nearby occluders that the lighting record does not see.
 		const FVector4 RecordToVertexVector = Vertex.WorldPosition - LightingRecord.Vertex.WorldPosition;
 		// Use the average normal to handle surfaces with constant concavity
-		const FVector4 AverageNormal = (LightingRecord.Vertex.TriangleNormal + Vertex.TriangleNormal).SafeNormal();
-		const float PlaneDistance = Dot3(AverageNormal, RecordToVertexVector.SafeNormal());
+		const FVector4 AverageNormal = (LightingRecord.Vertex.TriangleNormal + Vertex.TriangleNormal).GetSafeNormal();
+		const float PlaneDistance = Dot3(AverageNormal, RecordToVertexVector.GetSafeNormal());
 		// Setup an error metric that goes from 0 if the points are coplanar, to 1 if the point being shaded is at the angle corresponding to MinCosPointBehindPlane behind the plane
 		const float PointBehindPlaneError = FMath::Max(PlaneDistance / MinCosPointBehindPlane, 0.0f);
 
diff --git a/Engine/Source/Programs/UnrealLightmass/Private/Lighting/FluidSurface.cpp b/Engine/Source/Programs/UnrealLightmass/Private/Lighting/FluidSurface.cpp
--- a/Engine/Source/Programs/UnrealLightmass/Private/Lighting/FluidSurface.cpp
+++ b/Engine/Source/Programs/UnrealLightmass/Private/Lighting/FluidSurface.cpp
@@ -18,9 +18,9 @@ namespace Lightmass
 		)
 	{
 		OutVertex.WorldPosition = LocalToWorld.TransformPosition(QuadCorners[VertexIndex]);
-		OutVertex.WorldTangentX = LocalToWorld.TransformVector(FVector4(1, 0, 0, 1)).SafeNormal();
-		OutVertex.WorldTangentY = LocalToWorld.TransformVector(FVector4(0, 1, 0, 1)).SafeNormal();
-		OutVertex.WorldTangentZ = LocalToWorldInverseTranspose.TransformVector(FVector4(0, 0, 1, 1)).SafeNormal();
+		OutVertex.WorldTangentX = LocalToWorld.TransformVector(FVector4(1, 0, 0, 1)).GetSafeNormal();
+		OutVertex.WorldTangentY = LocalToWorld.TransformVector(FVector4(0, 1, 0, 1)).GetSafeNormal();
+		OutVertex.WorldTangentZ = LocalToWorldInverseTranspose.TransformVector(FVector4(0, 0, 1, 1)).GetSafeNormal();
 
 		for(uint32 UVIndex = 0; UVIndex < 1; UVIndex++)
 		{
diff --git a/Engine/Source/Programs/UnrealLightmass/Private/Lighting/LightingSystem.cpp b/Engine/Source/Programs/UnrealLightmass/Private/Lighting/LightingSystem.cpp
--- a/Engine/Source/Programs/UnrealLightmass/Private/Lighting/LightingSystem.cpp
+++ b/Engine/Source/Programs/UnrealLightmass/Private/Lighting/LightingSystem.cpp
@@ -1893,7 +1893,7 @@ void FStaticLightingSystem::CalculateApproximateDirectLighting(
 
 					FVector4 NormalForOffset = Vertex.WorldTangentZ;
 			
-					const FVector4 StartOffset = LightVector.SafeNormal() * SceneConstants.VisibilityRayOffsetDistance 
+					const FVector4 StartOffset = LightVector.GetSafeNormal() * SceneConstants.VisibilityRayOffsetDistance
 						+ NormalForOffset * SampleRadius * SceneConstants.VisibilityNormalOffsetSampleRadiusScale;
 
 					const FLightRay LightRay(
@@ -1956,13 +1956,13 @@ void FStaticLightingSystem::CalculateApproximateDirectLighting(
 					Dot3(WorldLightVector, Vertex.WorldTangentY),
 					Dot3(WorldLightVector, Vertex.WorldTangentZ),
 					0
-					).SafeNormal();
+					).GetSafeNormal();
 
 				// Compute the incident lighting of the light on the vertex.
 				const FLinearColor FinalIntensity = LightIntensity * Transmission;
 
 				// Compute the light-map sample for the front-face of the vertex.
-				FGatheredLightSample Lighting = FGatheredLightSample::PointLightWorldSpace(FinalIntensity, TangentLightVector, WorldLightVector.SafeNormal());
+				FGatheredLightSample Lighting = FGatheredLightSample::PointLightWorldSpace(FinalIntensity, TangentLightVector, WorldLightVector.GetSafeNormal());
 
 				if (Light->UseStaticLighting(false) || bCompositeAllLights)
 				{
@@ -2361,7 +2361,7 @@ bool FStaticLightingSystem::CalculatePointShadowing(
 			const FVector4 LightPosition = FVector4(Light->Position.X, Light->Position.Y, Light->Position.Z, 0);
 			const FVector4 LightVector = LightPosition - WorldSurfacePoint * Light->Position.W;
 			const FLightRay LightRay(
-				WorldSurfacePoint + LightVector.SafeNormal() * SceneConstants.VisibilityRayOffsetDistance,
+				WorldSurfacePoint + LightVector.GetSafeNormal() * SceneConstants.VisibilityRayOffsetDistance,
 				WorldSurfacePoint + LightVector,
 				Mapping,
 				Light
@@ -2456,7 +2456,7 @@ int32 FStaticLightingSystem::CalculatePointAreaShadowing(
 			const FLightRay LightRay(
 				// Offset the start of the ray by some fraction along the direction of the ray and some fraction along the vertex normal.
 				Vertex.WorldPosition 
-					+ LightVector.SafeNormal() * SceneConstants.VisibilityRayOffsetDistance 
+					+ LightVector.GetSafeNormal() * SceneConstants.VisibilityRayOffsetDistance 
 					+ NormalForOffset * SampleRadius * SceneConstants.VisibilityNormalOffsetSampleRadiusScale 
 					+ SampleOffset,
 				Vertex.WorldPosition + LightVector,
@@ -2517,13 +2517,13 @@ FGatheredLightSample FStaticLightingSystem::CalculatePointLighting(
 			    Dot3(WorldLightVector, Vertex.WorldTangentY),
 			    Dot3(WorldLightVector, Vertex.WorldTangentZ),
 				0
-			    ).SafeNormal();
+				).GetSafeNormal();
 
 		// Compute the incident lighting of the light on the vertex.
 		const FLinearColor LightIntensity = InLightIntensity * InTransmission;
 
 		// Compute the light-map sample for the front-face of the vertex.
-		FGatheredLightSample FrontFaceSample = FGatheredLightSample::PointLightWorldSpace(LightIntensity, TangentLightVector, WorldLightVector.SafeNormal());
+		FGatheredLightSample FrontFaceSample = FGatheredLightSample::PointLightWorldSpace(LightIntensity, TangentLightVector, WorldLightVector.GetSafeNormal());
 
 		if (Mapping->Mesh->UsesTwoSidedLighting(ElementIndex))
 		{
@@ -2533,8 +2533,8 @@ FGatheredLightSample FStaticLightingSystem::CalculatePointLighting(
 					Dot3(WorldLightVector, -Vertex.WorldTangentY),
 					Dot3(WorldLightVector, -Vertex.WorldTangentZ),
 					0
-					).SafeNormal();
-			const FGatheredLightSample BackFaceSample = FGatheredLightSample::PointLightWorldSpace(LightIntensity, BackFaceTangentLightVector, -WorldLightVector.SafeNormal());
+					).GetSafeNormal();
+			const FGatheredLightSample BackFaceSample = FGatheredLightSample::PointLightWorldSpace(LightIntensity, BackFaceTangentLightVector, -WorldLightVector.GetSafeNormal());
 			// Average front and back face lighting
 			return (FrontFaceSample + BackFaceSample) * .5f;
 		}
diff --git a/Engine/Source/Programs/UnrealLightmass/Private/Lighting/LightmapData.cpp b/Engine/Source/Programs/UnrealLightmass/Private/Lighting/LightmapData.cpp
--- a/Engine/Source/Programs/UnrealLightmass/Private/Lighting/LightmapData.cpp
+++ b/Engine/Source/Programs/UnrealLightmass/Private/Lighting/LightmapData.cpp
@@ -268,7 +268,7 @@ namespace Lightmass
 
 			const FVector BentNormal(SourceSample.SkyOcclusion[0], SourceSample.SkyOcclusion[1], SourceSample.SkyOcclusion[2]);
 			const float BentNormalLength = BentNormal.Size();
-			const FVector NormalizedBentNormal = BentNormal.SafeNormal() * FVector(.5f) + FVector(.5f);
+			const FVector NormalizedBentNormal = BentNormal.GetSafeNormal() * FVector(.5f) + FVector(.5f);
 
 			DestCoefficients.SkyOcclusion[0] = (uint8)FMath::Clamp<int32>( FMath::RoundToInt( NormalizedBentNormal[0] * 255.0f ), 0, 255 );
 			DestCoefficients.SkyOcclusion[1] = (uint8)FMath::Clamp<int32>( FMath::RoundToInt( NormalizedBentNormal[1] * 255.0f ), 0, 255 );
diff --git a/Engine/Source/Programs/UnrealLightmass/Private/Lighting/PhotonMapping.cpp b/Engine/Source/Programs/UnrealLightmass/Private/Lighting/PhotonMapping.cpp
--- a/Engine/Source/Programs/UnrealLightmass/Private/Lighting/PhotonMapping.cpp
+++ b/Engine/Source/Programs/UnrealLightmass/Private/Lighting/PhotonMapping.cpp
@@ -2043,7 +2043,7 @@ FIrradiancePhoton* FStaticLightingSystem::FindNearestIrradiancePhoton(
 				{
 					// Only accept irradiance photons within an angle of the plane defined by the vertex normal
 					// This avoids expensive visibility traces to photons that are probably not on the same surface
-					const float DirectionDotNormal = Dot3(CurrentPhoton.GetSurfaceNormal(), PhotonToVertexVector.SafeNormal());
+					const float DirectionDotNormal = Dot3(CurrentPhoton.GetSurfaceNormal(), PhotonToVertexVector.GetSafeNormal());
 					if (FMath::Abs(DirectionDotNormal) < PhotonMappingSettings.MinCosIrradiancePhotonSearchCone)
 					{
 						if (bVisibleOnly)
@@ -2075,7 +2075,7 @@ FIrradiancePhoton* FStaticLightingSystem::FindNearestIrradiancePhoton(
 			FIrradiancePhoton* CurrentPhoton = TempIrradiancePhotons[PhotonIndex];
 			const FVector4 VertexToPhoton = CurrentPhoton->GetPosition() - Vertex.WorldPosition;
 			const FLightRay VertexToPhotonRay(
-				Vertex.WorldPosition + VertexToPhoton.SafeNormal() * SceneConstants.VisibilityRayOffsetDistance + Vertex.WorldTangentZ * SceneConstants.VisibilityNormalOffsetDistance,
+				Vertex.WorldPosition + VertexToPhoton.GetSafeNormal() * SceneConstants.VisibilityRayOffsetDistance + Vertex.WorldTangentZ * SceneConstants.VisibilityNormalOffsetDistance,
 				CurrentPhoton->GetPosition() + CurrentPhoton->GetSurfaceNormal() * SceneConstants.VisibilityNormalOffsetDistance,
 				NULL,
 				NULL
diff --git a/Engine/Source/Programs/UnrealLightmass/Private/Lighting/PrecomputedVisibility.cpp b/Engine/Source/Programs/UnrealLightmass/Private/Lighting/PrecomputedVisibility.cpp
--- a/Engine/Source/Programs/UnrealLightmass/Private/Lighting/PrecomputedVisibility.cpp
+++ b/Engine/Source/Programs/UnrealLightmass/Private/Lighting/PrecomputedVisibility.cpp
@@ -914,7 +914,7 @@ bool ComputeBoxVisibility(
 				const int32 SampleIndex = FMath::TruncToInt(RandomStream.GetFraction() * FurthestSamples.Num());
 				const FVisibilityQuerySample& CurrentSample = *FurthestSamples[SampleIndex];
 				const float VectorLength = (CurrentSample.CellPosition - CurrentSample.MeshPosition).Size3();
-				const FVector4 CurrentDirection = (CurrentSample.MeshPosition - CurrentSample.CellPosition).SafeNormal();
+				const FVector4 CurrentDirection = (CurrentSample.MeshPosition - CurrentSample.CellPosition).GetSafeNormal();
 
 				FVector4 XAxis;
 				FVector4 YAxis;
diff --git a/Engine/Source/Programs/UnrealLightmass/Private/Lighting/StaticMesh.cpp b/Engine/Source/Programs/UnrealLightmass/Private/Lighting/StaticMesh.cpp
--- a/Engine/Source/Programs/UnrealLightmass/Private/Lighting/StaticMesh.cpp
+++ b/Engine/Source/Programs/UnrealLightmass/Private/Lighting/StaticMesh.cpp
@@ -41,7 +41,7 @@ static FVector SplineEvalDir(const FVector& StartPos, const FVector& StartTangen
 
 	const float A2 = A  * A;
 
-	return ((C * A2) + (D * A) + E).SafeNormal();
+	return ((C * A2) + (D * A) + E).GetSafeNormal();
 }
 
 /** Calculate full transform that defines frame along spline, given the Z of a vertex. */
@@ -59,8 +59,8 @@ static FMatrix CalcSliceTransform(float ZPos, const FSplineMeshParams& SplinePar
 	const FVector SplineDir = SplineEvalDir( SplineParams.StartPos, SplineParams.StartTangent, SplineParams.EndPos, SplineParams.EndTangent, Alpha );
 
 	// Find base frenet frame
-	const FVector BaseXVec = (SplineParams.SplineUpDir ^ SplineDir).SafeNormal();
-	const FVector BaseYVec = (SplineDir ^ BaseXVec).SafeNormal();
+	const FVector BaseXVec = (SplineParams.SplineUpDir ^ SplineDir).GetSafeNormal();
+	const FVector BaseYVec = (SplineDir ^ BaseXVec).GetSafeNormal();
 
 	// Offset the spline by the desired amount
 	const FVector2D SliceOffset = FMath::Lerp<FVector2D>(SplineParams.StartOffset, SplineParams.EndOffset, HermiteAlpha);
@@ -111,8 +111,8 @@ static FMatrix CalcSliceRot(float ZPos, const FSplineMeshParams& SplineParams)
 	const FVector SplineDir = SplineEvalDir( SplineParams.StartPos, SplineParams.StartTangent, SplineParams.EndPos, SplineParams.EndTangent, Alpha );
 
 	// Find base frenet frame
-	const FVector BaseXVec = (SplineParams.SplineUpDir ^ SplineDir).SafeNormal();
-	const FVector BaseYVec = (SplineDir ^ BaseXVec).SafeNormal();
+	const FVector BaseXVec = (SplineParams.SplineUpDir ^ SplineDir).GetSafeNormal();
+	const FVector BaseYVec = (SplineDir ^ BaseXVec).GetSafeNormal();
 
 	// Apply roll to frame around spline
 	const float UseRoll = FMath::Lerp(SplineParams.StartRoll, SplineParams.EndRoll, HermiteAlpha);
@@ -178,22 +178,22 @@ static void GetStaticLightingVertex(
 		const FVector4 LocalSpaceTangentY = SliceRot.TransformVector(InVertex.TangentY);
 		const FVector4 LocalSpaceTangentZ = SliceRot.TransformVector(InVertex.TangentZ);
 
-		OutVertex.WorldTangentX = LocalToWorld.TransformVector(LocalSpaceTangentX).SafeNormal();
-		OutVertex.WorldTangentY = LocalToWorld.TransformVector(LocalSpaceTangentY).SafeNormal();
-		OutVertex.WorldTangentZ = LocalToWorldInverseTranspose.TransformVector(LocalSpaceTangentZ).SafeNormal();
+		OutVertex.WorldTangentX = LocalToWorld.TransformVector(LocalSpaceTangentX).GetSafeNormal();
+		OutVertex.WorldTangentY = LocalToWorld.TransformVector(LocalSpaceTangentY).GetSafeNormal();
+		OutVertex.WorldTangentZ = LocalToWorldInverseTranspose.TransformVector(LocalSpaceTangentZ).GetSafeNormal();
 	}
 	else
 	{
 		OutVertex.WorldPosition = LocalToWorld.TransformPosition(InVertex.Position);
-		OutVertex.WorldTangentX = LocalToWorld.TransformVector(InVertex.TangentX).SafeNormal();
-		OutVertex.WorldTangentY = LocalToWorld.TransformVector(InVertex.TangentY).SafeNormal();
-		OutVertex.WorldTangentZ = LocalToWorldInverseTranspose.TransformVector(InVertex.TangentZ).SafeNormal();
+		OutVertex.WorldTangentX = LocalToWorld.TransformVector(InVertex.TangentX).GetSafeNormal();
+		OutVertex.WorldTangentY = LocalToWorld.TransformVector(InVertex.TangentY).GetSafeNormal();
+		OutVertex.WorldTangentZ = LocalToWorldInverseTranspose.TransformVector(InVertex.TangentZ).GetSafeNormal();
 	}
 
 	// WorldTangentZ can end up a 0 vector if it was small to begin with and LocalToWorld contains large scale factors.
 	if (!OutVertex.WorldTangentZ.IsUnit3())
 	{
-		OutVertex.WorldTangentZ = (OutVertex.WorldTangentX ^ OutVertex.WorldTangentY).SafeNormal();
+		OutVertex.WorldTangentZ = (OutVertex.WorldTangentX ^ OutVertex.WorldTangentY).GetSafeNormal();
 	}
 
 	for(uint32 LightmapTextureCoordinateIndex = 0; LightmapTextureCoordinateIndex < ARRAY_COUNT(InVertex.UVs); LightmapTextureCoordinateIndex++)
diff --git a/Engine/Source/Programs/UnrealLightmass/Private/Lighting/TextureMapping.cpp b/Engine/Source/Programs/UnrealLightmass/Private/Lighting/TextureMapping.cpp
--- a/Engine/Source/Programs/UnrealLightmass/Private/Lighting/TextureMapping.cpp
+++ b/Engine/Source/Programs/UnrealLightmass/Private/Lighting/TextureMapping.cpp
@@ -293,7 +293,7 @@ void FStaticLightingSystem::CacheIrradiancePhotonsTextureMapping(FStaticLighting
 		TextureMapping->Mesh->GetTriangle(TriangleIndex, V0.Vertex, V1.Vertex, V2.Vertex, Element);
 		V0.ElementIndex = V1.ElementIndex = V2.ElementIndex = Element;
 
-		const FVector4 TriangleNormal = ((V2.Vertex.WorldPosition - V0.Vertex.WorldPosition) ^ (V1.Vertex.WorldPosition - V0.Vertex.WorldPosition)).SafeNormal();
+		const FVector4 TriangleNormal = ((V2.Vertex.WorldPosition - V0.Vertex.WorldPosition) ^ (V1.Vertex.WorldPosition - V0.Vertex.WorldPosition)).GetSafeNormal();
 
 		// Don't rasterize degenerates 
 		if (!TriangleNormal.IsNearlyZero3())
@@ -806,7 +806,7 @@ void FStaticLightingSystem::SetupTextureMapping(
 			TextureMapping->Mesh->GetTriangle(TriangleIndex,V0.Vertex,V1.Vertex,V2.Vertex,Element);
 			V0.ElementIndex = V1.ElementIndex = V2.ElementIndex = Element;
 
-			const FVector4 TriangleNormal = ((V2.Vertex.WorldPosition - V0.Vertex.WorldPosition) ^ (V1.Vertex.WorldPosition - V0.Vertex.WorldPosition)).SafeNormal();
+			const FVector4 TriangleNormal = ((V2.Vertex.WorldPosition - V0.Vertex.WorldPosition) ^ (V1.Vertex.WorldPosition - V0.Vertex.WorldPosition)).GetSafeNormal();
 
 			// Don't rasterize degenerates 
 			if (!TriangleNormal.IsNearlyZero3())
@@ -1902,7 +1902,7 @@ void FStaticLightingSystem::CalculateDirectSignedDistanceFieldLightingTextureMap
 					CurrentSample.SetMapped(true);
 
 					const FVector4 LightPosition = Light->LightCenterPosition(TexelToVertex.WorldPosition, TexelToVertex.WorldTangentZ);
-					const FVector4 LightVector = (LightPosition - TexelToVertex.WorldPosition).SafeNormal();
+					const FVector4 LightVector = (LightPosition - TexelToVertex.WorldPosition).GetSafeNormal();
 
 					FVector4 NormalForOffset = CurrentSample.GetNormal();
 					// Flip the normal used for offsetting the start of the ray for two sided materials if a flipped normal would be closer to the light.
@@ -2133,7 +2133,7 @@ void FStaticLightingSystem::CalculateDirectSignedDistanceFieldLightingTextureMap
 								&& Light->AffectsBounds(FBoxSphereBounds(HighResSample.GetPosition(), FVector4(0,0,0),0)))
 							{
 								const FVector4 LightPosition = Light->LightCenterPosition(HighResSample.GetPosition(), HighResSample.GetNormal());
-								const FVector4 LightVector = (LightPosition - HighResSample.GetPosition()).SafeNormal();
+								const FVector4 LightVector = (LightPosition - HighResSample.GetPosition()).GetSafeNormal();
 
 								FVector4 NormalForOffset = HighResSample.GetNormal();
 								// Flip the normal used for offsetting the start of the ray for two sided materials if a flipped normal would be closer to the light.
diff --git a/Engine/Source/Programs/UnrealLightmass/Private/LightmassCore/Math/LMkDOP.h b/Engine/Source/Programs/UnrealLightmass/Private/LightmassCore/Math/LMkDOP.h
--- a/Engine/Source/Programs/UnrealLightmass/Private/LightmassCore/Math/LMkDOP.h
+++ b/Engine/Source/Programs/UnrealLightmass/Private/LightmassCore/Math/LMkDOP.h
@@ -402,7 +402,7 @@ public:
 	}
 	inline FVector4 GetLocalNormal() const
 	{
-		FVector4 LocalNormal = ((V1 - V2) ^ (V0 - V2)).SafeNormal();
+		FVector4 LocalNormal = ((V1 - V2) ^ (V0 - V2)).GetSafeNormal();
 		LocalNormal.W = Dot3(V0, LocalNormal);
 		return LocalNormal;
 	}
@@ -430,7 +430,7 @@ public:
 		// Calculate the centroid for the triangle
 		Centroid = (V0 + V1 + V2) / 3.f;
 		// Calculate the local normal used by line check code.
-		LocalNormal = ((V1 - V2) ^ (V0 - V2)).SafeNormal();
+		LocalNormal = ((V1 - V2) ^ (V0 - V2)).GetSafeNormal();
 		LocalNormal.W = Dot3(V0, LocalNormal);
 #endif
 	}
@@ -1285,7 +1285,7 @@ template <typename COLL_DATA_PROVIDER, typename KDOP_IDX_TYPE> struct TkDOPLineC
 	FORCEINLINE FVector4 GetHitNormal(void)
 	{
 		// Transform the hit back into world space using the transpose adjoint
-		FVector4 Normal = TkDOPCollisionCheck<COLL_DATA_PROVIDER,KDOP_IDX_TYPE>::CollDataProvider.GetLocalToWorldTransposeAdjoint().TransformVector(LocalHitNormal).SafeNormal();
+		FVector4 Normal = TkDOPCollisionCheck<COLL_DATA_PROVIDER,KDOP_IDX_TYPE>::CollDataProvider.GetLocalToWorldTransposeAdjoint().TransformVector(LocalHitNormal).GetSafeNormal();
 		// Flip the normal if the triangle is inverted
 		if (TkDOPCollisionCheck<COLL_DATA_PROVIDER,KDOP_IDX_TYPE>::CollDataProvider.GetDeterminant() < 0.f)
 		{
diff --git a/Engine/Source/Runtime/AIModule/Private/BehaviorTree/Decorators/BTDecorator_ConeCheck.cpp b/Engine/Source/Runtime/AIModule/Private/BehaviorTree/Decorators/BTDecorator_ConeCheck.cpp
--- a/Engine/Source/Runtime/AIModule/Private/BehaviorTree/Decorators/BTDecorator_ConeCheck.cpp
+++ b/Engine/Source/Runtime/AIModule/Private/BehaviorTree/Decorators/BTDecorator_ConeCheck.cpp
@@ -56,7 +56,7 @@ bool UBTDecorator_ConeCheck::CalculateDirection(const UBlackboardComponent* Blac
 		}
 		else if (BlackboardComp->GetLocationFromEntry(Origin.GetSelectedKeyID(), PointA) && BlackboardComp->GetLocationFromEntry(End.GetSelectedKeyID(), PointB))
 		{
-			Direction = (PointB - PointA).SafeNormal();
+			Direction = (PointB - PointA).GetSafeNormal();
 			return true;
 		}
 	}
diff --git a/Engine/Source/Runtime/AIModule/Private/BehaviorTree/Decorators/BTDecorator_KeepInCone.cpp b/Engine/Source/Runtime/AIModule/Private/BehaviorTree/Decorators/BTDecorator_KeepInCone.cpp
--- a/Engine/Source/Runtime/AIModule/Private/BehaviorTree/Decorators/BTDecorator_KeepInCone.cpp
+++ b/Engine/Source/Runtime/AIModule/Private/BehaviorTree/Decorators/BTDecorator_KeepInCone.cpp
@@ -64,7 +64,7 @@ bool UBTDecorator_KeepInCone::CalculateCurrentDirection(const UBehaviorTreeCompo
 
 	if (bHasPointA && bHasPointB)
 	{
-		Direction = (PointB - PointA).SafeNormal();
+		Direction = (PointB - PointA).GetSafeNormal();
 		return true;
 	}
 
diff --git a/Engine/Source/Runtime/AIModule/Private/BehaviorTree/Tasks/BTTask_RotateToFaceBBEntry.cpp b/Engine/Source/Runtime/AIModule/Private/BehaviorTree/Tasks/BTTask_RotateToFaceBBEntry.cpp
--- a/Engine/Source/Runtime/AIModule/Private/BehaviorTree/Tasks/BTTask_RotateToFaceBBEntry.cpp
+++ b/Engine/Source/Runtime/AIModule/Private/BehaviorTree/Tasks/BTTask_RotateToFaceBBEntry.cpp
@@ -67,7 +67,7 @@ EBTNodeResult::Type UBTTask_RotateToFaceBBEntry::ExecuteTask(UBehaviorTreeCompon
 		if (ActorValue != NULL)
 		{
 			const float AngleDifference = CalculateAngleDifferenceDot(Pawn->GetActorRotation().Vector()
-				, (ActorValue->GetActorLocation() - PawnLocation).SafeNormal2D());
+				, (ActorValue->GetActorLocation() - PawnLocation).GetSafeNormal2D());
 			
 			if (AngleDifference >= PrecisionDot)
 			{
@@ -89,7 +89,7 @@ EBTNodeResult::Type UBTTask_RotateToFaceBBEntry::ExecuteTask(UBehaviorTreeCompon
 		if (FAISystem::IsValidLocation(KeyValue))
 		{
 			const float AngleDifference = CalculateAngleDifferenceDot(Pawn->GetActorRotation().Vector()
-				, (KeyValue - PawnLocation).SafeNormal2D());
+				, (KeyValue - PawnLocation).GetSafeNormal2D());
 
 			if (AngleDifference >= PrecisionDot)
 			{
@@ -144,7 +144,7 @@ void UBTTask_RotateToFaceBBEntry::TickTask(UBehaviorTreeComponent* OwnerComp, ui
 		const FVector PawnDirection = AIController->GetPawn()->GetActorRotation().Vector();				
 		const FVector FocalPoint = AIController->GetFocalPoint(EAIFocusPriority::Gameplay);
 
-		if (CalculateAngleDifferenceDot(PawnDirection, (FocalPoint - AIController->GetPawn()->GetActorLocation()).SafeNormal2D()) >= PrecisionDot)
+		if (CalculateAngleDifferenceDot(PawnDirection, (FocalPoint - AIController->GetPawn()->GetActorLocation()).GetSafeNormal2D()) >= PrecisionDot)
 		{
 			CleanUp(*AIController, NodeMemory);
 			FinishLatentTask(OwnerComp, EBTNodeResult::Succeeded);
@@ -197,7 +197,7 @@ void UBTTask_RotateToFaceBBEntry::DescribeRuntimeValues(const UBehaviorTreeCompo
 		const FVector PawnDirection = AIController->GetPawn()->GetActorRotation().Vector();
 		const FVector FocalPoint = AIController->GetFocusItem(EAIFocusPriority::Gameplay).GetLocation();
 
-		const float CurrentAngleRadians = CalculateAngleDifferenceDot(PawnDirection, (FocalPoint - AIController->GetPawn()->GetActorLocation()).SafeNormal2D());
+		const float CurrentAngleRadians = CalculateAngleDifferenceDot(PawnDirection, (FocalPoint - AIController->GetPawn()->GetActorLocation()).GetSafeNormal2D());
 		Values.Add(FString::Printf(TEXT("Current angle: %.2f"), FMath::RadiansToDegrees(FMath::Acos(CurrentAngleRadians))));
 	}
 	else
diff --git a/Engine/Source/Runtime/AIModule/Private/EnvironmentQuery/Generators/EnvQueryGenerator_OnCircle.cpp b/Engine/Source/Runtime/AIModule/Private/EnvironmentQuery/Generators/EnvQueryGenerator_OnCircle.cpp
--- a/Engine/Source/Runtime/AIModule/Private/EnvironmentQuery/Generators/EnvQueryGenerator_OnCircle.cpp
+++ b/Engine/Source/Runtime/AIModule/Private/EnvironmentQuery/Generators/EnvQueryGenerator_OnCircle.cpp
@@ -145,7 +145,7 @@ FVector UEnvQueryGenerator_OnCircle::CalcDirection(FEnvQueryInstance& QueryInsta
 
 			if (Start.Num() > 0 && End.Num() > 0)
 			{
-				Direction = (End[0] - Start[0]).SafeNormal();
+				Direction = (End[0] - Start[0]).GetSafeNormal();
 			}
 			else
 			{
diff --git a/Engine/Source/Runtime/AIModule/Private/EnvironmentQuery/Tests/EnvQueryTest_Dot.cpp b/Engine/Source/Runtime/AIModule/Private/EnvironmentQuery/Tests/EnvQueryTest_Dot.cpp
--- a/Engine/Source/Runtime/AIModule/Private/EnvironmentQuery/Tests/EnvQueryTest_Dot.cpp
+++ b/Engine/Source/Runtime/AIModule/Private/EnvironmentQuery/Tests/EnvQueryTest_Dot.cpp
@@ -137,7 +137,7 @@ void UEnvQueryTest_Dot::GatherLineDirections(TArray<FVector>& Directions, FEnvQu
 		
 		for (int32 ToIndex = 0; ToIndex < ContextLocationTo.Num(); ToIndex++)
 		{
-			const FVector Dir = (ContextLocationTo[ToIndex] - ContextLocationFrom[FromIndex]).SafeNormal();
+			const FVector Dir = (ContextLocationTo[ToIndex] - ContextLocationFrom[FromIndex]).GetSafeNormal();
 			Directions.Add(Dir);
 		}
 	}
diff --git a/Engine/Source/Runtime/AIModule/Private/Navigation/CrowdFollowingComponent.cpp b/Engine/Source/Runtime/AIModule/Private/Navigation/CrowdFollowingComponent.cpp
--- a/Engine/Source/Runtime/AIModule/Private/Navigation/CrowdFollowingComponent.cpp
+++ b/Engine/Source/Runtime/AIModule/Private/Navigation/CrowdFollowingComponent.cpp
@@ -216,13 +216,13 @@ void UCrowdFollowingComponent::UpdateCachedDirections(const FVector& NewVelocity
 	const FVector ToCorner = NextPathCorner - AgentLoc;
 	if (ToCorner.SizeSquared() > FMath::Square(10.0f))
 	{
-		MoveSegmentDirection = ToCorner.SafeNormal();
+		MoveSegmentDirection = ToCorner.GetSafeNormal();
 	}
 
 	// CrowdAgentMoveDirection either direction on path or aligned with current velocity
 	if (!bTraversingLink)
 	{
-		CrowdAgentMoveDirection = bRotateToVelocity && (NewVelocity.SizeSquared() > KINDA_SMALL_NUMBER) ? NewVelocity.SafeNormal() : MoveSegmentDirection;
+		CrowdAgentMoveDirection = bRotateToVelocity && (NewVelocity.SizeSquared() > KINDA_SMALL_NUMBER) ? NewVelocity.GetSafeNormal() : MoveSegmentDirection;
 	}
 }
 
@@ -596,7 +596,7 @@ void UCrowdFollowingComponent::SetMoveSegment(int32 SegmentStartIndex)
 		bCheckMovementAngle = true;
 		bUpdateDirectMoveVelocity = DestinationActor.IsValid();
 		CurrentDestination.Set(Path->GetBaseActor(), CurrentTargetPt);
-		CrowdAgentMoveDirection = (CurrentTargetPt - AgentLoc).SafeNormal();
+		CrowdAgentMoveDirection = (CurrentTargetPt - AgentLoc).GetSafeNormal();
 		MoveSegmentDirection = CrowdAgentMoveDirection;
 		SuspendCrowdSteering(true);
 
@@ -706,7 +706,7 @@ void UCrowdFollowingComponent::FollowPathSegment(float DeltaTime)
 			const FVector AgentLoc = GetCrowdAgentLocation();
 
 			CurrentDestination.Set(Path->GetBaseActor(), CurrentTargetPt);
-			CrowdAgentMoveDirection = (CurrentTargetPt - AgentLoc).SafeNormal();
+			CrowdAgentMoveDirection = (CurrentTargetPt - AgentLoc).GetSafeNormal();
 			MoveSegmentDirection = CrowdAgentMoveDirection;
 
 			Manager->SetAgentMoveDirection(this, MoveSegmentDirection);
diff --git a/Engine/Source/Runtime/AIModule/Private/Navigation/PathFollowingComponent.cpp b/Engine/Source/Runtime/AIModule/Private/Navigation/PathFollowingComponent.cpp
--- a/Engine/Source/Runtime/AIModule/Private/Navigation/PathFollowingComponent.cpp
+++ b/Engine/Source/Runtime/AIModule/Private/Navigation/PathFollowingComponent.cpp
@@ -85,7 +85,7 @@ void LogBlockHelper(AActor* LogOwner, UNavMovementComponent* MoveComp, float Rad
 	{
 		const FVector AgentLocation = MoveComp->GetActorFeetLocation();
 		const FVector ToTarget = (SegmentEnd - AgentLocation);
-		const float SegmentDot = FVector::DotProduct(ToTarget.SafeNormal(), (SegmentEnd - SegmentStart).SafeNormal());
+		const float SegmentDot = FVector::DotProduct(ToTarget.GetSafeNormal(), (SegmentEnd - SegmentStart).GetSafeNormal());
 		UE_VLOG(LogOwner, LogPathFollowing, Verbose, TEXT("[agent to segment end] dot [segment dir]: %f"), SegmentDot);
 		
 		float AgentRadius = 0.0f;
@@ -557,12 +557,12 @@ void UPathFollowingComponent::SetMoveSegment(int32 SegmentStartIndex)
 		CurrentDestination.Set(Path->GetBaseActor(), SegmentEnd);
 		CurrentAcceptanceRadius = (Path->GetPathPoints().Num() == (EndSegmentIndex + 1)) ? AcceptanceRadius : 0.0f;
 
-		MoveSegmentDirection = (SegmentEnd - SegmentStart).SafeNormal();
+		MoveSegmentDirection = (SegmentEnd - SegmentStart).GetSafeNormal();
 		
 		// make sure we have a non-zero direction if still following a valid path
 		if (MoveSegmentDirection.IsZero() && int32(MoveSegmentEndIndex + 1) < Path->GetPathPoints().Num())
 		{
-			MoveSegmentDirection = (Path->GetPathPoints()[MoveSegmentEndIndex + 1].Location - SegmentStart).SafeNormal();
+			MoveSegmentDirection = (Path->GetPathPoints()[MoveSegmentEndIndex + 1].Location - SegmentStart).GetSafeNormal();
 		}
 
 		// handle moving through custom nav links
@@ -962,7 +962,7 @@ void UPathFollowingComponent::DebugReachTest(float& CurrentDot, float& CurrentDi
 	}
 
 	const FVector ToGoal = (GoalLocation - AgentLocation);
-	CurrentDot = FVector::DotProduct(ToGoal.SafeNormal(), MoveSegmentDirection);
+	CurrentDot = FVector::DotProduct(ToGoal.GetSafeNormal(), MoveSegmentDirection);
 	bDotFailed = (CurrentDot < 0.0f) ? 1 : 0;
 
 	// get cylinder of moving agent
diff --git a/Engine/Source/Runtime/AIModule/Private/Perception/AISense_Sight.cpp b/Engine/Source/Runtime/AIModule/Private/Perception/AISense_Sight.cpp
--- a/Engine/Source/Runtime/AIModule/Private/Perception/AISense_Sight.cpp
+++ b/Engine/Source/Runtime/AIModule/Private/Perception/AISense_Sight.cpp
@@ -17,7 +17,7 @@ FORCEINLINE_DEBUGGABLE bool CheckIsTargetInSightPie(const FPerceptionListener& L
 {
 	if (FVector::DistSquared(Listener.CachedLocation, TargetLocation) <= SightRadiusSq) 
 	{
-		const FVector DirectionToTarget = (TargetLocation - Listener.CachedLocation).UnsafeNormal();
+		const FVector DirectionToTarget = (TargetLocation - Listener.CachedLocation).GetUnsafeNormal();
 		return FVector::DotProduct(DirectionToTarget, Listener.CachedDirection) > DigestedProps.PeripheralVisionAngleCos;
 	}
 
diff --git a/Engine/Source/Runtime/AIModule/Private/Perception/PawnSensingComponent.cpp b/Engine/Source/Runtime/AIModule/Private/Perception/PawnSensingComponent.cpp
--- a/Engine/Source/Runtime/AIModule/Private/Perception/PawnSensingComponent.cpp
+++ b/Engine/Source/Runtime/AIModule/Private/Perception/PawnSensingComponent.cpp
@@ -334,7 +334,7 @@ bool UPawnSensingComponent::CouldSeePawn(const APawn *Other, bool bMaySkipChecks
 // 	UE_LOG(LogPath, Warning, TEXT("DistanceToOtherSquared = %f, SightRadiusSquared: %f"), SelfToOtherDistSquared, FMath::Square(SightRadius));
 
 	// check field of view
-	FVector const SelfToOtherDir = SelfToOther.SafeNormal();
+	FVector const SelfToOtherDir = SelfToOther.GetSafeNormal();
 	FVector const MyFacingDir = Owner->GetActorRotation().Vector();
 
 // 	UE_LOG(LogPath, Warning, TEXT("DotProductFacing: %f, PeripheralVisionCosine: %f"), SelfToOtherDir | MyFacingDir, PeripheralVisionCosine);
diff --git a/Engine/Source/Runtime/AppFramework/Private/Widgets/Testing/STestSuite.cpp b/Engine/Source/Runtime/AppFramework/Private/Widgets/Testing/STestSuite.cpp
--- a/Engine/Source/Runtime/AppFramework/Private/Widgets/Testing/STestSuite.cpp
+++ b/Engine/Source/Runtime/AppFramework/Private/Widgets/Testing/STestSuite.cpp
@@ -37,7 +37,7 @@ namespace
 	{
 		float uniScale = 5.8f;
 		FVector trans(5, 6, 7);
-		FQuat quat(FVector(1, 2, 3).SafeNormal(), 33.5f);
+		FQuat quat(FVector(1, 2, 3).GetSafeNormal(), 33.5f);
 		FRotator rot(7, 8, 9);
 		FMatrix mat = FRotationMatrix::Make(rot);
 
diff --git a/Engine/Source/Runtime/Core/Private/IOS/IOSInputInterface.cpp b/Engine/Source/Runtime/Core/Private/IOS/IOSInputInterface.cpp
--- a/Engine/Source/Runtime/Core/Private/IOS/IOSInputInterface.cpp
+++ b/Engine/Source/Runtime/Core/Private/IOS/IOSInputInterface.cpp
@@ -207,7 +207,7 @@ void FIOSInputInterface::GetMovementData(FVector& Attitude, FVector& RotationRat
 		FilteredAccelerometer = FilteredAccelerometer * VectorFilter + (1.0f - VectorFilter) * NewAcceleration;
 
 		// create an normalized acceleration vector
-		FVector FinalAcceleration = -FilteredAccelerometer.SafeNormal();
+		FVector FinalAcceleration = -FilteredAccelerometer.GetSafeNormal();
 
 		// calculate Roll/Pitch
 		float CurrentPitch = FMath::Atan2(FinalAcceleration.Y, FinalAcceleration.Z);
diff --git a/Engine/Source/Runtime/Core/Private/Math/UnrealMath.cpp b/Engine/Source/Runtime/Core/Private/Math/UnrealMath.cpp
--- a/Engine/Source/Runtime/Core/Private/Math/UnrealMath.cpp
+++ b/Engine/Source/Runtime/Core/Private/Math/UnrealMath.cpp
@@ -196,7 +196,7 @@ void FVector::FindBestAxisVectors( FVector& Axis1, FVector& Axis2 ) const
 	if( NZ>NX && NZ>NY )	Axis1 = FVector(1,0,0);
 	else					Axis1 = FVector(0,0,1);
 
-	Axis1 = (Axis1 - *this * (Axis1 | *this)).SafeNormal();
+	Axis1 = (Axis1 - *this * (Axis1 | *this)).GetSafeNormal();
 	Axis2 = Axis1 ^ *this;
 }
 
@@ -395,12 +395,12 @@ FMatrix FRotationMatrix::Make(FQuat const& Rot)
 
 FMatrix FRotationMatrix::MakeFromX(FVector const& XAxis)
 {
-	FVector const NewX = XAxis.SafeNormal();
+	FVector const NewX = XAxis.GetSafeNormal();
 
 	// try to use up if possible
 	FVector const UpVector = ( FMath::Abs(NewX.Z) < (1.f - KINDA_SMALL_NUMBER) ) ? FVector(0,0,1.f) : FVector(1.f,0,0);
 
-	const FVector NewY = (UpVector ^ NewX).SafeNormal();
+	const FVector NewY = (UpVector ^ NewX).GetSafeNormal();
 	const FVector NewZ = NewX ^ NewY;
 
 	return FMatrix(NewX, NewY, NewZ, FVector::ZeroVector);
@@ -408,12 +408,12 @@ FMatrix FRotationMatrix::MakeFromX(FVector const& XAxis)
 
 FMatrix FRotationMatrix::MakeFromY(FVector const& YAxis)
 {
-	FVector const NewY = YAxis.SafeNormal();
+	FVector const NewY = YAxis.GetSafeNormal();
 
 	// try to use up if possible
 	FVector const UpVector = ( FMath::Abs(NewY.Z) < (1.f - KINDA_SMALL_NUMBER) ) ? FVector(0,0,1.f) : FVector(1.f,0,0);
 
-	const FVector NewZ = (UpVector ^ NewY).SafeNormal();
+	const FVector NewZ = (UpVector ^ NewY).GetSafeNormal();
 	const FVector NewX = NewY ^ NewZ;
 
 	return FMatrix(NewX, NewY, NewZ, FVector::ZeroVector);
@@ -421,12 +421,12 @@ FMatrix FRotationMatrix::MakeFromY(FVector const& YAxis)
 
 FMatrix FRotationMatrix::MakeFromZ(FVector const& ZAxis)
 {
-	FVector const NewZ = ZAxis.SafeNormal();
+	FVector const NewZ = ZAxis.GetSafeNormal();
 
 	// try to use up if possible
 	FVector const UpVector = ( FMath::Abs(NewZ.Z) < (1.f - KINDA_SMALL_NUMBER) ) ? FVector(0,0,1.f) : FVector(1.f,0,0);
 
-	const FVector NewX = (UpVector ^ NewZ).SafeNormal();
+	const FVector NewX = (UpVector ^ NewZ).GetSafeNormal();
 	const FVector NewY = NewZ ^ NewX;
 
 	return FMatrix(NewX, NewY, NewZ, FVector::ZeroVector);
@@ -434,8 +434,8 @@ FMatrix FRotationMatrix::MakeFromZ(FVector const& ZAxis)
 
 FMatrix FRotationMatrix::MakeFromXY(FVector const& XAxis, FVector const& YAxis)
 {
-	FVector NewX = XAxis.SafeNormal();
-	FVector Norm = YAxis.SafeNormal();
+	FVector NewX = XAxis.GetSafeNormal();
+	FVector Norm = YAxis.GetSafeNormal();
 
 	// if they're almost same, we need to find arbitrary vector
 	if ( FMath::IsNearlyEqual(FMath::Abs(NewX | Norm), 1.f) )
@@ -444,7 +444,7 @@ FMatrix FRotationMatrix::MakeFromXY(FVector const& XAxis, FVector const& YAxis)
 		Norm = ( FMath::Abs(NewX.Z) < (1.f - KINDA_SMALL_NUMBER) ) ? FVector(0,0,1.f) : FVector(1.f,0,0);
 	}
 
-	const FVector NewZ = (NewX ^ Norm).SafeNormal();
+	const FVector NewZ = (NewX ^ Norm).GetSafeNormal();
 	const FVector NewY = NewZ ^ NewX;
 
 	return FMatrix(NewX, NewY, NewZ, FVector::ZeroVector);
@@ -452,8 +452,8 @@ FMatrix FRotationMatrix::MakeFromXY(FVector const& XAxis, FVector const& YAxis)
 
 FMatrix FRotationMatrix::MakeFromXZ(FVector const& XAxis, FVector const& ZAxis)
 {
-	FVector const NewX = XAxis.SafeNormal();
-	FVector Norm = ZAxis.SafeNormal();
+	FVector const NewX = XAxis.GetSafeNormal();
+	FVector Norm = ZAxis.GetSafeNormal();
 
 	// if they're almost same, we need to find arbitrary vector
 	if ( FMath::IsNearlyEqual(FMath::Abs(NewX | Norm), 1.f) )
@@ -462,7 +462,7 @@ FMatrix FRotationMatrix::MakeFromXZ(FVector const& XAxis, FVector const& ZAxis)
 		Norm = ( FMath::Abs(NewX.Z) < (1.f - KINDA_SMALL_NUMBER) ) ? FVector(0,0,1.f) : FVector(1.f,0,0);
 	}
 
-	const FVector NewY = (Norm ^ NewX).SafeNormal();
+	const FVector NewY = (Norm ^ NewX).GetSafeNormal();
 	const FVector NewZ = NewX ^ NewY;
 
 	return FMatrix(NewX, NewY, NewZ, FVector::ZeroVector);
@@ -470,8 +470,8 @@ FMatrix FRotationMatrix::MakeFromXZ(FVector const& XAxis, FVector const& ZAxis)
 
 FMatrix FRotationMatrix::MakeFromYX(FVector const& YAxis, FVector const& XAxis)
 {
-	FVector const NewY = YAxis.SafeNormal();
-	FVector Norm = XAxis.SafeNormal();
+	FVector const NewY = YAxis.GetSafeNormal();
+	FVector Norm = XAxis.GetSafeNormal();
 
 	// if they're almost same, we need to find arbitrary vector
 	if ( FMath::IsNearlyEqual(FMath::Abs(NewY | Norm), 1.f) )
@@ -480,7 +480,7 @@ FMatrix FRotationMatrix::MakeFromYX(FVector const& YAxis, FVector const& XAxis)
 		Norm = ( FMath::Abs(NewY.Z) < (1.f - KINDA_SMALL_NUMBER) ) ? FVector(0,0,1.f) : FVector(1.f,0,0);
 	}
 	
-	const FVector NewZ = (Norm ^ NewY).SafeNormal();
+	const FVector NewZ = (Norm ^ NewY).GetSafeNormal();
 	const FVector NewX = NewY ^ NewZ;
 
 	return FMatrix(NewX, NewY, NewZ, FVector::ZeroVector);
@@ -488,8 +488,8 @@ FMatrix FRotationMatrix::MakeFromYX(FVector const& YAxis, FVector const& XAxis)
 
 FMatrix FRotationMatrix::MakeFromYZ(FVector const& YAxis, FVector const& ZAxis)
 {
-	FVector const NewY = YAxis.SafeNormal();
-	FVector Norm = ZAxis.SafeNormal();
+	FVector const NewY = YAxis.GetSafeNormal();
+	FVector Norm = ZAxis.GetSafeNormal();
 
 	// if they're almost same, we need to find arbitrary vector
 	if ( FMath::IsNearlyEqual(FMath::Abs(NewY | Norm), 1.f) )
@@ -498,7 +498,7 @@ FMatrix FRotationMatrix::MakeFromYZ(FVector const& YAxis, FVector const& ZAxis)
 		Norm = ( FMath::Abs(NewY.Z) < (1.f - KINDA_SMALL_NUMBER) ) ? FVector(0,0,1.f) : FVector(1.f,0,0);
 	}
 
-	const FVector NewX = (NewY ^ Norm).SafeNormal();
+	const FVector NewX = (NewY ^ Norm).GetSafeNormal();
 	const FVector NewZ = NewX ^ NewY;
 
 	return FMatrix(NewX, NewY, NewZ, FVector::ZeroVector);
@@ -506,8 +506,8 @@ FMatrix FRotationMatrix::MakeFromYZ(FVector const& YAxis, FVector const& ZAxis)
 
 FMatrix FRotationMatrix::MakeFromZX(FVector const& ZAxis, FVector const& XAxis)
 {
-	FVector const NewZ = ZAxis.SafeNormal();
-	FVector Norm = XAxis.SafeNormal();
+	FVector const NewZ = ZAxis.GetSafeNormal();
+	FVector Norm = XAxis.GetSafeNormal();
 
 	// if they're almost same, we need to find arbitrary vector
 	if ( FMath::IsNearlyEqual(FMath::Abs(NewZ | Norm), 1.f) )
@@ -516,7 +516,7 @@ FMatrix FRotationMatrix::MakeFromZX(FVector const& ZAxis, FVector const& XAxis)
 		Norm = ( FMath::Abs(NewZ.Z) < (1.f - KINDA_SMALL_NUMBER) ) ? FVector(0,0,1.f) : FVector(1.f,0,0);
 	}
 
-	const FVector NewY = (NewZ ^ Norm).SafeNormal();
+	const FVector NewY = (NewZ ^ Norm).GetSafeNormal();
 	const FVector NewX = NewY ^ NewZ;
 
 	return FMatrix(NewX, NewY, NewZ, FVector::ZeroVector);
@@ -524,8 +524,8 @@ FMatrix FRotationMatrix::MakeFromZX(FVector const& ZAxis, FVector const& XAxis)
 
 FMatrix FRotationMatrix::MakeFromZY(FVector const& ZAxis, FVector const& YAxis)
 {
-	FVector const NewZ = ZAxis.SafeNormal();
-	FVector Norm = YAxis.SafeNormal();
+	FVector const NewZ = ZAxis.GetSafeNormal();
+	FVector Norm = YAxis.GetSafeNormal();
 
 	// if they're almost same, we need to find arbitrary vector
 	if ( FMath::IsNearlyEqual(FMath::Abs(NewZ | Norm), 1.f) )
@@ -534,7 +534,7 @@ FMatrix FRotationMatrix::MakeFromZY(FVector const& ZAxis, FVector const& YAxis)
 		Norm = ( FMath::Abs(NewZ.Z) < (1.f - KINDA_SMALL_NUMBER) ) ? FVector(0,0,1.f) : FVector(1.f,0,0);
 	}
 
-	const FVector NewX = (Norm ^ NewZ).SafeNormal();
+	const FVector NewX = (Norm ^ NewZ).GetSafeNormal();
 	const FVector NewY = NewZ ^ NewX;
 
 	return FMatrix(NewX, NewY, NewZ, FVector::ZeroVector);
@@ -1186,14 +1186,14 @@ void CORE_API CurveLinearColorFindIntervalBounds( const FInterpCurvePoint<FLinea
 
 CORE_API float FMath::PointDistToLine(const FVector &Point, const FVector &Direction, const FVector &Origin, FVector &OutClosestPoint)
 {
-	const FVector SafeDir = Direction.SafeNormal();
+	const FVector SafeDir = Direction.GetSafeNormal();
 	OutClosestPoint = Origin + (SafeDir * ((Point-Origin) | SafeDir));
 	return (OutClosestPoint-Point).Size();
 }
 
 CORE_API float FMath::PointDistToLine(const FVector &Point, const FVector &Direction, const FVector &Origin)
 {
-	const FVector SafeDir = Direction.SafeNormal();
+	const FVector SafeDir = Direction.GetSafeNormal();
 	const FVector OutClosestPoint = Origin + (SafeDir * ((Point-Origin) | SafeDir));
 	return (OutClosestPoint-Point).Size();
 }
@@ -1241,8 +1241,8 @@ void FMath::SegmentDistToSegmentSafe(FVector A1, FVector B1, FVector A2, FVector
 	const	FVector	S3		= A1 - A2;
 
 
-	const	FVector	S1_norm		= S1.SafeNormal();
-	const	FVector	S2_norm		= S2.SafeNormal();
+	const	FVector	S1_norm		= S1.GetSafeNormal();
+	const	FVector	S2_norm		= S2.GetSafeNormal();
 
 	const	float	Dot11	= S1 | S1;	// always >= 0
 	const	float	Dot22	= S2 | S2;	// always >= 0
@@ -1596,7 +1596,7 @@ FVector FMath::ComputeBaryCentric2D(const FVector& Point, const FVector& A, cons
 	//check collinearity of A,B,C
 	check(TriNorm.SizeSquared() > SMALL_NUMBER && "Collinear points in FMath::ComputeBaryCentric2D()");
 
-	const FVector N = TriNorm.SafeNormal();
+	const FVector N = TriNorm.GetSafeNormal();
 
 	// Compute twice area of triangle ABC
 	const float AreaABCInv = 1.0f / (N | TriNorm);
@@ -1738,7 +1738,7 @@ void FMath::SphereDistToLine(FVector SphereOrigin, float SphereRadius, FVector L
 	{
 		// line is not intersecting sphere (or is tangent at one point if D == 0 )
 		const FVector PointOnLine = LineOrigin + ( -B / 2.f * A ) * LineDir;
-		OutClosestPoint = SphereOrigin + (PointOnLine - SphereOrigin).SafeNormal() * SphereRadius;
+		OutClosestPoint = SphereOrigin + (PointOnLine - SphereOrigin).GetSafeNormal() * SphereRadius;
 	}
 	else
 	{
@@ -1797,7 +1797,7 @@ bool FMath::PointsAreCoplanar(const TArray<FVector>& Points, const float Toleran
 	}
 
 	//Get the Normal for plane determined by first 3 points
-	const FVector Normal = FVector::CrossProduct(Points[2] - Points[0], Points[1] - Points[0]).SafeNormal();
+	const FVector Normal = FVector::CrossProduct(Points[2] - Points[0], Points[1] - Points[0]).GetSafeNormal();
 
 	const int32 Total = Points.Num();
 	for (int32 v = 3; v < Total; v++)
@@ -1821,10 +1821,10 @@ bool FMath::GetDotDistance
 	const	FVector		&AxisZ 	
 )
 {
-	const FVector NormalDir = Direction.SafeNormal();
+	const FVector NormalDir = Direction.GetSafeNormal();
 
 	// Find projected point (on AxisX and AxisY, remove AxisZ component)
-	const FVector NoZProjDir = ( NormalDir - (NormalDir | AxisZ) * AxisZ ).SafeNormal();
+	const FVector NoZProjDir = (NormalDir - (NormalDir | AxisZ) * AxisZ).GetSafeNormal();
 	
 	// Figure out if projection is on right or left.
 	const float AzimuthSign = ( (NoZProjDir | AxisY) < 0.f ) ? -1.f : 1.f;
@@ -1844,9 +1844,9 @@ FVector2D FMath::GetAzimuthAndElevation
 	const FVector &AxisZ 	
 )
 {
-	const FVector NormalDir = Direction.SafeNormal();
+	const FVector NormalDir = Direction.GetSafeNormal();
 	// Find projected point (on AxisX and AxisY, remove AxisZ component)
-	const FVector NoZProjDir = (NormalDir - (NormalDir | AxisZ) * AxisZ).SafeNormal();
+	const FVector NoZProjDir = (NormalDir - (NormalDir | AxisZ) * AxisZ).GetSafeNormal();
 	// Figure out if projection is on right or left.
 	const float AzimuthSign = ((NoZProjDir | AxisY) < 0.f) ? -1.f : 1.f;
 	const float ElevationSin = NormalDir | AxisZ;
@@ -2188,13 +2188,13 @@ FVector FMath::VRandCone(FVector const& Dir, float ConeHalfAngleRad)
 		Result = Result.RotateAngleAxis(Theta * 180.f / PI, DirZ);
 
 		// ensure it's a unit vector (might not have been passed in that way)
-		Result = Result.SafeNormal();
+		Result = Result.GetSafeNormal();
 		
 		return Result;
 	}
 	else
 	{
-		return Dir.SafeNormal();
+		return Dir.GetSafeNormal();
 	}
 }
 
@@ -2230,13 +2230,13 @@ FVector FMath::VRandCone(FVector const& Dir, float HorizontalConeHalfAngleRad, f
 		Result = Result.RotateAngleAxis(Theta * 180.f / PI, DirZ);
 
 		// ensure it's a unit vector (might not have been passed in that way)
-		Result = Result.SafeNormal();
+		Result = Result.GetSafeNormal();
 
 		return Result;
 	}
 	else
 	{
-		return Dir.SafeNormal();
+		return Dir.GetSafeNormal();
 	}
 }
 
@@ -2249,7 +2249,7 @@ FVector FMath::RandPointInBox(const FBox& Box)
 
 FVector FMath::GetReflectionVector(const FVector& Direction, const FVector& SurfaceNormal)
 {
-	return Direction - 2 * (Direction | SurfaceNormal.SafeNormal()) * SurfaceNormal.SafeNormal();
+	return Direction - 2 * (Direction | SurfaceNormal.GetSafeNormal()) * SurfaceNormal.GetSafeNormal();
 }
 
 struct FClusterMovedHereToMakeCompile
diff --git a/Engine/Source/Runtime/Core/Public/Math/Matrix.h b/Engine/Source/Runtime/Core/Public/Math/Matrix.h
--- a/Engine/Source/Runtime/Core/Public/Math/Matrix.h
+++ b/Engine/Source/Runtime/Core/Public/Math/Matrix.h
@@ -425,8 +425,8 @@ struct FLookAtMatrix : FMatrix
 
 FORCEINLINE FLookAtMatrix::FLookAtMatrix(const FVector& EyePosition, const FVector& LookAtPosition, const FVector& UpVector)
 {
-	const FVector ZAxis = (LookAtPosition - EyePosition).SafeNormal();
-	const FVector XAxis = (UpVector ^ ZAxis).SafeNormal();
+	const FVector ZAxis = (LookAtPosition - EyePosition).GetSafeNormal();
+	const FVector XAxis = (UpVector ^ ZAxis).GetSafeNormal();
 	const FVector YAxis = ZAxis ^ XAxis;
 
 	for (uint32 RowIndex = 0; RowIndex < 3; RowIndex++)
diff --git a/Engine/Source/Runtime/Core/Public/Math/Plane.h b/Engine/Source/Runtime/Core/Public/Math/Plane.h
--- a/Engine/Source/Runtime/Core/Public/Math/Plane.h
+++ b/Engine/Source/Runtime/Core/Public/Math/Plane.h
@@ -312,7 +312,7 @@ FORCEINLINE FPlane::FPlane( FVector InBase, const FVector &InNormal )
 
 
 FORCEINLINE FPlane::FPlane( FVector A, FVector B, FVector C )
-	:	FVector( ((B-A)^(C-A)).SafeNormal() )
+	:	FVector( ((B-A)^(C-A)).GetSafeNormal() )
 {
 	W = A | (FVector)(*this);
 }
diff --git a/Engine/Source/Runtime/Core/Public/Math/RandomStream.h b/Engine/Source/Runtime/Core/Public/Math/RandomStream.h
--- a/Engine/Source/Runtime/Core/Public/Math/RandomStream.h
+++ b/Engine/Source/Runtime/Core/Public/Math/RandomStream.h
@@ -126,7 +126,7 @@ public:
 		}
 		while(L > 1.f || L < KINDA_SMALL_NUMBER);
 
-		return Result.UnsafeNormal();
+		return Result.GetUnsafeNormal();
 	}
 
 	/**
@@ -226,13 +226,13 @@ public:
 			Result = Result.RotateAngleAxis(Theta * 180.f / PI, DirZ);
 
 			// ensure it's a unit vector (might not have been passed in that way)
-			Result = Result.SafeNormal();
+			Result = Result.GetSafeNormal();
 		
 			return Result;
 		}
 		else
 		{
-			return Dir.SafeNormal();
+			return Dir.GetSafeNormal();
 		}
 	}
 
@@ -276,13 +276,13 @@ public:
 			Result = Result.RotateAngleAxis(Theta * 180.f / PI, DirZ);
 
 			// ensure it's a unit vector (might not have been passed in that way)
-			Result = Result.SafeNormal();
+			Result = Result.GetSafeNormal();
 
 			return Result;
 		}
 		else
 		{
-			return Dir.SafeNormal();
+			return Dir.GetSafeNormal();
 		}
 	}
 
diff --git a/Engine/Source/Runtime/Core/Public/Math/SHMath.h b/Engine/Source/Runtime/Core/Public/Math/SHMath.h
--- a/Engine/Source/Runtime/Core/Public/Math/SHMath.h
+++ b/Engine/Source/Runtime/Core/Public/Math/SHMath.h
@@ -249,7 +249,7 @@ public:
 	FVector GetMaximumDirection() const
 	{
 		// This is an approximation which only takes into account first and second order spherical harmonics.
-		return FVector(-V[3], -V[1], V[2]).SafeNormal();
+		return FVector(-V[3], -V[1], V[2]).GetSafeNormal();
 	}
 
 	static TSHVector CalcDiffuseTransfer(const FVector& Normal)
diff --git a/Engine/Source/Runtime/Core/Public/Math/UnrealMath.h b/Engine/Source/Runtime/Core/Public/Math/UnrealMath.h
--- a/Engine/Source/Runtime/Core/Public/Math/UnrealMath.h
+++ b/Engine/Source/Runtime/Core/Public/Math/UnrealMath.h
@@ -102,7 +102,7 @@ inline FPlane FPlane::TransformBy( const FMatrix& M ) const
 
 inline FPlane FPlane::TransformByUsingAdjointT( const FMatrix& M, float DetM, const FMatrix& TA ) const
 {
-	FVector newNorm = TA.TransformVector(*this).SafeNormal();
+	FVector newNorm = TA.TransformVector(*this).GetSafeNormal();
 
 	if(DetM < 0.f)
 	{
diff --git a/Engine/Source/Runtime/Core/Public/Math/UnrealMatrix.h b/Engine/Source/Runtime/Core/Public/Math/UnrealMatrix.h
--- a/Engine/Source/Runtime/Core/Public/Math/UnrealMatrix.h
+++ b/Engine/Source/Runtime/Core/Public/Math/UnrealMatrix.h
@@ -546,7 +546,7 @@ inline void FMatrix::GetScaledAxes(FVector &X, FVector &Y, FVector &Z) const
 
 inline FVector FMatrix::GetUnitAxis( EAxis::Type InAxis ) const
 {
-	return GetScaledAxis( InAxis ).SafeNormal();
+	return GetScaledAxis( InAxis ).GetSafeNormal();
 }
 
 inline void FMatrix::GetUnitAxes(FVector &X, FVector &Y, FVector &Z) const
diff --git a/Engine/Source/Runtime/Core/Public/Math/Vector.h b/Engine/Source/Runtime/Core/Public/Math/Vector.h
--- a/Engine/Source/Runtime/Core/Public/Math/Vector.h
+++ b/Engine/Source/Runtime/Core/Public/Math/Vector.h
@@ -480,8 +480,11 @@ public:
 	 * Calculates normalized version of vector without checking for zero length.
 	 *
 	 * @return Normalized version of vector.
-	 * @see SafeNormal()
+	 * @see GetSafeNormal()
 	 */
+	FORCEINLINE FVector GetUnsafeNormal() const;
+
+	DEPRECATED(4.7, "Deprecated due to unclear name, use GetUnsafeNormal instead.")
 	FORCEINLINE FVector UnsafeNormal() const;
 
 	/**
@@ -502,15 +505,27 @@ public:
 	FVector BoundToCube( float Radius ) const;
 
 	/** Create a copy of this vector, with its magnitude clamped between Min and Max. */
+	FVector GetClampedToSize(float Min, float Max) const;
+
+	DEPRECATED(4.7, "Deprecated due to unclear name, use GetClampedToSize instead.")
 	FVector ClampSize(float Min, float Max) const;
 
 	/** Create a copy of this vector, with the 2D magnitude clamped between Min and Max. Z is unchanged. */
+	FVector GetClampedToSize2D(float Min, float Max) const;
+
+	DEPRECATED(4.7, "Deprecated due to unclear name, use GetClampedToSize2D instead.")
 	FVector ClampSize2D(float Min, float Max) const;
 
 	/** Create a copy of this vector, with its maximum magnitude clamped to MaxSize. */
+	FVector GetClampedToMaxSize(float MaxSize) const;
+
+	DEPRECATED(4.7, "Deprecated due to unclear name, use GetClampedToMaxSize instead.")
 	FVector ClampMaxSize(float MaxSize) const;
 
 	/** Create a copy of this vector, with the maximum 2D magnitude clamped to MaxSize. Z is unchanged. */
+	FVector GetClampedToMaxSize2D(float MaxSize) const;
+
+	DEPRECATED(4.7, "Deprecated due to unclear name, use GetClampedToMaxSize2D instead.")
 	FVector ClampMaxSize2D(float MaxSize) const;
 
 	/**
@@ -569,7 +584,10 @@ public:
 	 * @param Tolerance Minimum squared vector length.
 	 * @return A normalized copy if safe, (0,0,0) otherwise.
 	 */
-	FVector SafeNormal(float Tolerance=SMALL_NUMBER) const;
+	FVector GetSafeNormal(float Tolerance=SMALL_NUMBER) const;
+
+	DEPRECATED(4.7, "Deprecated due to unclear name, use GetSafeNormal instead.")
+	FVector SafeNormal(float Tolerance = SMALL_NUMBER) const;
 
 	/**
 	 * Gets a normalized copy of the 2D components of the vector, checking it is safe to do so. Z is set to zero. 
@@ -578,7 +596,10 @@ public:
 	 * @param Tolerance Minimum squared vector length.
 	 * @return Normalized copy if safe, otherwise returns zero vector.
 	 */
-	FVector SafeNormal2D(float Tolerance=SMALL_NUMBER) const;
+	FVector GetSafeNormal2D(float Tolerance=SMALL_NUMBER) const;
+
+	DEPRECATED(4.7, "Deprecated due to unclear name, use GetSafeNormal2D instead.")
+	FVector SafeNormal2D(float Tolerance = SMALL_NUMBER) const;
 
 	/**
 	 * Returns the cosine of the angle between this vector and another projected onto the XY plane (no Z).
@@ -1474,12 +1495,17 @@ FORCEINLINE FVector FVector::Projection() const
 	return FVector( X*RZ, Y*RZ, 1 );
 }
 
-FORCEINLINE FVector FVector::UnsafeNormal() const
+FORCEINLINE FVector FVector::GetUnsafeNormal() const
 {
 	const float Scale = FMath::InvSqrt(X*X+Y*Y+Z*Z);
 	return FVector( X*Scale, Y*Scale, Z*Scale );
 }
 
+FORCEINLINE FVector FVector::UnsafeNormal() const
+{
+	return GetUnsafeNormal();
+}
+
 FORCEINLINE FVector FVector::GridSnap( const float& GridSz ) const
 {
 	return FVector( FMath::GridSnap(X, GridSz),FMath::GridSnap(Y, GridSz),FMath::GridSnap(Z, GridSz) );
@@ -1495,7 +1521,7 @@ FORCEINLINE FVector FVector::BoundToCube( float Radius ) const
 		);
 }
 
-FORCEINLINE FVector FVector::ClampSize(float Min, float Max) const
+FORCEINLINE FVector FVector::GetClampedToSize(float Min, float Max) const
 {
 	float VecSize = Size();
 	const FVector VecDir = (VecSize > SMALL_NUMBER) ? (*this/VecSize) : FVector::ZeroVector;
@@ -1503,9 +1529,14 @@ FORCEINLINE FVector FVector::ClampSize(float Min, float Max) const
 	VecSize = FMath::Clamp(VecSize, Min, Max);
 
 	return VecSize * VecDir;
-}	
+}
 
-FORCEINLINE FVector FVector::ClampSize2D(float Min, float Max) const
+FORCEINLINE FVector FVector::ClampSize(float Min, float Max) const
+{
+	return GetClampedToSize(Min, Max);
+}
+
+FORCEINLINE FVector FVector::GetClampedToSize2D(float Min, float Max) const
 {
 	float VecSize2D = Size2D();
 	const FVector VecDir = (VecSize2D > SMALL_NUMBER) ? (*this/VecSize2D) : FVector::ZeroVector;
@@ -1513,10 +1544,15 @@ FORCEINLINE FVector FVector::ClampSize2D(float Min, float Max) const
 	VecSize2D = FMath::Clamp(VecSize2D, Min, Max);
 
 	return FVector(VecSize2D * VecDir.X, VecSize2D * VecDir.Y, Z);
-}	
+}
+
+FORCEINLINE FVector FVector::ClampSize2D(float Min, float Max) const
+{
+	return GetClampedToSize2D(Min, Max);
+}
 
 
-FORCEINLINE FVector FVector::ClampMaxSize(float MaxSize) const
+FORCEINLINE FVector FVector::GetClampedToMaxSize(float MaxSize) const
 {
 	if (MaxSize < KINDA_SMALL_NUMBER)
 	{
@@ -1535,7 +1571,12 @@ FORCEINLINE FVector FVector::ClampMaxSize(float MaxSize) const
 	}	
 }
 
-FORCEINLINE FVector FVector::ClampMaxSize2D(float MaxSize) const
+FORCEINLINE FVector FVector::ClampMaxSize(float MaxSize) const
+{
+	return GetClampedToMaxSize(MaxSize);
+}
+
+FORCEINLINE FVector FVector::GetClampedToMaxSize2D(float MaxSize) const
 {
 	if (MaxSize < KINDA_SMALL_NUMBER)
 	{
@@ -1554,6 +1595,11 @@ FORCEINLINE FVector FVector::ClampMaxSize2D(float MaxSize) const
 	}
 }
 
+FORCEINLINE FVector FVector::ClampMaxSize2D(float MaxSize) const
+{
+	return GetClampedToMaxSize2D(MaxSize);
+}
+
 
 FORCEINLINE void FVector::AddBounded( const FVector& V, float Radius )
 {
@@ -1611,7 +1657,7 @@ FORCEINLINE FVector FVector::MirrorByVector( const FVector& MirrorNormal ) const
 	return *this - MirrorNormal * (2.f * (*this | MirrorNormal));
 }
 
-FORCEINLINE FVector FVector::SafeNormal(float Tolerance) const
+FORCEINLINE FVector FVector::GetSafeNormal(float Tolerance) const
 {
 	const float SquareSum = X*X + Y*Y + Z*Z;
 
@@ -1628,7 +1674,12 @@ FORCEINLINE FVector FVector::SafeNormal(float Tolerance) const
 	return FVector(X*Scale, Y*Scale, Z*Scale);
 }
 
-FORCEINLINE FVector FVector::SafeNormal2D(float Tolerance) const
+FORCEINLINE FVector FVector::SafeNormal(float Tolerance) const
+{
+	return GetSafeNormal(Tolerance);
+}
+
+FORCEINLINE FVector FVector::GetSafeNormal2D(float Tolerance) const
 {
 	const float SquareSum = X*X + Y*Y;
 
@@ -1653,6 +1704,11 @@ FORCEINLINE FVector FVector::SafeNormal2D(float Tolerance) const
 	return FVector(X*Scale, Y*Scale, 0.f);
 }
 
+FORCEINLINE FVector FVector::SafeNormal2D(float Tolerance) const
+{
+	return GetSafeNormal2D(Tolerance);
+}
+
 FORCEINLINE float FVector::CosineAngle2D(FVector B) const
 {
 	FVector A(*this);
@@ -1812,7 +1868,7 @@ FORCEINLINE float FVector::HeadingAngle() const
 	// Project Dir into Z plane.
 	FVector PlaneDir = *this;
 	PlaneDir.Z = 0.f;
-	PlaneDir = PlaneDir.SafeNormal();
+	PlaneDir = PlaneDir.GetSafeNormal();
 
 	float Angle = FMath::Acos(PlaneDir.X);
 
diff --git a/Engine/Source/Runtime/Core/Public/Math/Vector2D.h b/Engine/Source/Runtime/Core/Public/Math/Vector2D.h
--- a/Engine/Source/Runtime/Core/Public/Math/Vector2D.h
+++ b/Engine/Source/Runtime/Core/Public/Math/Vector2D.h
@@ -379,13 +379,16 @@ public:
 	 * @param Tolerance Minimum squared length of vector for normalization.
 	 * @return A normalized copy of the vector if safe, (0,0) otherwise.
 	 */
-	FVector2D SafeNormal( float Tolerance=SMALL_NUMBER ) const;
+	FVector2D GetSafeNormal( float Tolerance=SMALL_NUMBER ) const;
+
+	DEPRECATED(4.7, "Deprecated due to unclear name, use GetSafeNormal instead.")
+	FVector2D SafeNormal(float Tolerance = SMALL_NUMBER) const;
 
 	/**
 	 * Normalize this vector in-place if it is large enough, set it to (0,0) otherwise.
 	 *
 	 * @param Tolerance Minimum squared length of vector for normalization.
-	 * @see SafeNormal()
+	 * @see GetSafeNormal()
 	 */
 	void Normalize( float Tolerance=SMALL_NUMBER );
 
@@ -742,7 +745,7 @@ FORCEINLINE float FVector2D::SizeSquared() const
 }
 
 
-FORCEINLINE FVector2D FVector2D::SafeNormal(float Tolerance) const
+FORCEINLINE FVector2D FVector2D::GetSafeNormal(float Tolerance) const
 {	
 	const float SquareSum = X*X + Y*Y;
 	if( SquareSum > Tolerance )
@@ -754,6 +757,12 @@ FORCEINLINE FVector2D FVector2D::SafeNormal(float Tolerance) const
 }
 
 
+FORCEINLINE FVector2D FVector2D::SafeNormal(float Tolerance) const
+{
+	return GetSafeNormal(Tolerance);
+}
+
+
 FORCEINLINE void FVector2D::Normalize(float Tolerance)
 {
 	const float SquareSum = X*X + Y*Y;
diff --git a/Engine/Source/Runtime/Core/Public/Math/Vector4.h b/Engine/Source/Runtime/Core/Public/Math/Vector4.h
--- a/Engine/Source/Runtime/Core/Public/Math/Vector4.h
+++ b/Engine/Source/Runtime/Core/Public/Math/Vector4.h
@@ -289,6 +289,9 @@ public:
 	 * @param Tolerance Minimum squared length of vector for normalization.
 	 * @return A normalized copy of the vector or a zero vector.
 	 */
+	FORCEINLINE FVector4 GetSafeNormal(float Tolerance = SMALL_NUMBER) const;
+
+	DEPRECATED(4.7, "Deprecated due to unclear name, use GetSafeNormal instead.")
 	FORCEINLINE FVector4 SafeNormal( float Tolerance=SMALL_NUMBER ) const;
 
 	/**
@@ -296,6 +299,9 @@ public:
 	 *
 	 * @return Normalized version of vector.
 	 */
+	FORCEINLINE FVector4 GetUnsafeNormal3() const;
+
+	DEPRECATED(4.7, "Deprecated due to unclear name, use GetUnsafeNormal3 instead.")
 	FORCEINLINE FVector4 UnsafeNormal3( ) const;
 
 	/**
@@ -562,7 +568,7 @@ FORCEINLINE bool FVector4::InitFromString( const FString& InSourceString )
 }
 
 
-FORCEINLINE FVector4 FVector4::SafeNormal( float Tolerance ) const
+FORCEINLINE FVector4 FVector4::GetSafeNormal( float Tolerance ) const
 {
 	const float SquareSum = X*X + Y*Y + Z*Z;
 	if( SquareSum > Tolerance )
@@ -574,13 +580,25 @@ FORCEINLINE FVector4 FVector4::SafeNormal( float Tolerance ) const
 }
 
 
-FORCEINLINE FVector4 FVector4::UnsafeNormal3( ) const
+FORCEINLINE FVector4 FVector4::SafeNormal(float Tolerance) const
+{
+	return GetSafeNormal(Tolerance);
+}
+
+
+FORCEINLINE FVector4 FVector4::GetUnsafeNormal3( ) const
 {
 	const float Scale = FMath::InvSqrt(X*X+Y*Y+Z*Z);
 	return FVector4( X*Scale, Y*Scale, Z*Scale, 0.0f );
 }
 
 
+FORCEINLINE FVector4 FVector4::UnsafeNormal3() const
+{
+	return GetUnsafeNormal3();
+}
+
+
 FORCEINLINE float FVector4::Size3( ) const
 {
 	return FMath::Sqrt( X*X + Y*Y + Z*Z );
@@ -633,7 +651,7 @@ FORCEINLINE void FVector4::FindBestAxisVectors3( FVector4& Axis1, FVector4& Axis
 	if( NZ>NX && NZ>NY )	Axis1 = FVector4(1,0,0);
 	else					Axis1 = FVector4(0,0,1);
 
-	Axis1 = (Axis1 - *this * Dot3(Axis1, *this)).SafeNormal();
+	Axis1 = (Axis1 - *this * Dot3(Axis1, *this)).GetSafeNormal();
 	Axis2 = Axis1 ^ *this;
 }
 
diff --git a/Engine/Source/Runtime/Engine/Private/AI/Navigation/AvoidanceManager.cpp b/Engine/Source/Runtime/Engine/Private/AI/Navigation/AvoidanceManager.cpp
--- a/Engine/Source/Runtime/Engine/Private/AI/Navigation/AvoidanceManager.cpp
+++ b/Engine/Source/Runtime/Engine/Private/AI/Navigation/AvoidanceManager.cpp
@@ -353,7 +353,7 @@ FVector UAvoidanceManager::GetAvoidanceVelocity_Internal(const FNavAvoidanceData
 			float RadiusB = OtherObject.Radius + inAvoidanceData.Radius;
 
 			PointBRelative.Z = 0.0f;
-			TowardB = PointBRelative.SafeNormal2D();		//Don't care about height for this game. Rough height-checking will come in later, but even then it will be acceptable to do this.
+			TowardB = PointBRelative.GetSafeNormal2D();		//Don't care about height for this game. Rough height-checking will come in later, but even then it will be acceptable to do this.
 			if (TowardB.IsZero())
 			{
 				//Already intersecting, or aligned vertically, scrap this whole object.
diff --git a/Engine/Source/Runtime/Engine/Private/AI/Navigation/NavTestRenderingComponent.cpp b/Engine/Source/Runtime/Engine/Private/AI/Navigation/NavTestRenderingComponent.cpp
--- a/Engine/Source/Runtime/Engine/Private/AI/Navigation/NavTestRenderingComponent.cpp
+++ b/Engine/Source/Runtime/Engine/Private/AI/Navigation/NavTestRenderingComponent.cpp
@@ -107,7 +107,7 @@ public:
 
 					//DrawWireBox(PDI, FBox(ProjectedLocation-BoxExtent, ProjectedLocation+BoxExtent), ProjectedColor, false);
 					DrawWireBox(PDI, FBox(ActorLocation - BoxExtent, ActorLocation + BoxExtent), FColor::White, false);
-					const FVector LineEnd = ProjectedLocation - (ProjectedLocation - ActorLocation).SafeNormal()*BoxExtent.X;
+					const FVector LineEnd = ProjectedLocation - (ProjectedLocation - ActorLocation).GetSafeNormal()*BoxExtent.X;
 					PDI->DrawLine(LineEnd, ActorLocation, ProjectedColor, SDPG_World, 2.5);
 					DrawArrowHead(PDI, LineEnd, ActorLocation, 20.f, ProjectedColor, SDPG_World, 2.5f);
 
diff --git a/Engine/Source/Runtime/Engine/Private/AI/Navigation/NavigationPath.cpp b/Engine/Source/Runtime/Engine/Private/AI/Navigation/NavigationPath.cpp
--- a/Engine/Source/Runtime/Engine/Private/AI/Navigation/NavigationPath.cpp
+++ b/Engine/Source/Runtime/Engine/Private/AI/Navigation/NavigationPath.cpp
@@ -298,18 +298,18 @@ FVector FNavigationPath::GetSegmentDirection(uint32 SegmentEndIndex) const
 		{
 			if (SegmentEndIndex > 0)
 			{
-				Result = (PathPoints[SegmentEndIndex].Location - PathPoints[SegmentEndIndex - 1].Location).SafeNormal();
+				Result = (PathPoints[SegmentEndIndex].Location - PathPoints[SegmentEndIndex - 1].Location).GetSafeNormal();
 			}
 			else
 			{
 				// for '0'-th segment returns same as for 1st segment 
-				Result = (PathPoints[1].Location - PathPoints[0].Location).SafeNormal();
+				Result = (PathPoints[1].Location - PathPoints[0].Location).GetSafeNormal();
 			}
 		}
 		else if (SegmentEndIndex >= uint32(GetPathPoints().Num()))
 		{
 			// in this special case return direction of last segment
-			Result = (PathPoints[PathPoints.Num() - 1].Location - PathPoints[PathPoints.Num() - 2].Location).SafeNormal();
+			Result = (PathPoints[PathPoints.Num() - 1].Location - PathPoints[PathPoints.Num() - 2].Location).GetSafeNormal();
 		}
 	}
 
@@ -542,7 +542,7 @@ namespace
 			return true;
 		}
 
-		const FVector RayNormal = (StartTrace-EndTrace) .SafeNormal() * OffsetDistannce;
+		const FVector RayNormal = (StartTrace-EndTrace) .GetSafeNormal() * OffsetDistannce;
 		StartTrace = StartTrace + RayNormal;
 		EndTrace = EndTrace - RayNormal;
 
@@ -564,8 +564,8 @@ namespace
 			{
 				const float EdgeHalfLength = (CurrentEdge->Left - CurrentEdge->Right).Size() * 0.5;
 				const float Distance = FMath::Min(OffsetDistannce, EdgeHalfLength) *  0.1;
-				Left = CurrentEdge->Left + Distance * (CurrentEdge->Right - CurrentEdge->Left).SafeNormal();
-				Right = CurrentEdge->Right + Distance * (CurrentEdge->Left - CurrentEdge->Right).SafeNormal();
+				Left = CurrentEdge->Left + Distance * (CurrentEdge->Right - CurrentEdge->Left).GetSafeNormal();
+				Right = CurrentEdge->Right + Distance * (CurrentEdge->Left - CurrentEdge->Right).GetSafeNormal();
 				FVector ClosestPointOnRay, ClosestPointOnEdge;
 				FMath::SegmentDistToSegment(StartTrace, EndTrace, Right, Left, ClosestPointOnRay, ClosestPointOnEdge);
 #if DEBUG_DRAW_OFFSET
@@ -689,7 +689,7 @@ void FNavMeshPath::OffsetFromCorners(float Distance)
 			const FVector EdgePt0 = Edge->GetPoint(CloserPoint);
 			const FVector EdgePt1 = Edge->GetPoint((CloserPoint+1)%2);
 			const FVector EdgeDir = EdgePt1 - EdgePt0;
-			const FVector EdgeOffset = EdgeDir.SafeNormal() * ActualOffset;
+			const FVector EdgeOffset = EdgeDir.GetSafeNormal() * ActualOffset;
 			NewPathPoint.Location = EdgePt0 + EdgeOffset;
 			// update NodeRef (could be different if this is n-th pass on the same PathPoint
 			NewPathPoint.NodeRef = Edge->ToRef;
@@ -1005,17 +1005,17 @@ FVector FNavMeshPath::GetSegmentDirection(uint32 SegmentEndIndex) const
 		{
 			if (SegmentEndIndex > 0)
 			{
-				Result = (Corridor[SegmentEndIndex].GetMiddlePoint() - Corridor[SegmentEndIndex - 1].GetMiddlePoint()).SafeNormal();
+				Result = (Corridor[SegmentEndIndex].GetMiddlePoint() - Corridor[SegmentEndIndex - 1].GetMiddlePoint()).GetSafeNormal();
 			}
 			else
 			{
-				Result = (Corridor[0].GetMiddlePoint() - GetPathPoints()[0].Location).SafeNormal();
+				Result = (Corridor[0].GetMiddlePoint() - GetPathPoints()[0].Location).GetSafeNormal();
 			}
 		}
 		else if (SegmentEndIndex >= uint32(Corridor.Num()))
 		{
 			// in this special case return direction of last segment
-			Result = (Corridor[Corridor.Num() - 1].GetMiddlePoint() - GetPathPoints()[0].Location).SafeNormal();
+			Result = (Corridor[Corridor.Num() - 1].GetMiddlePoint() - GetPathPoints()[0].Location).GetSafeNormal();
 		}
 	}
 
diff --git a/Engine/Source/Runtime/Engine/Private/AI/Navigation/PImplRecastNavMesh.cpp b/Engine/Source/Runtime/Engine/Private/AI/Navigation/PImplRecastNavMesh.cpp
--- a/Engine/Source/Runtime/Engine/Private/AI/Navigation/PImplRecastNavMesh.cpp
+++ b/Engine/Source/Runtime/Engine/Private/AI/Navigation/PImplRecastNavMesh.cpp
@@ -681,7 +681,7 @@ void FPImplRecastNavMesh::Raycast2D(const FVector& StartLoc, const FVector& EndL
 	{
 		// start location is not on navmesh, treat it as a blocked raycast
 		RaycastResult.HitTime = 0.f;
-		RaycastResult.HitNormal = (StartLoc - EndLoc).SafeNormal();
+		RaycastResult.HitNormal = (StartLoc - EndLoc).GetSafeNormal();
 	}
 }
 
@@ -720,7 +720,7 @@ void FPImplRecastNavMesh::Raycast2D(NavNodeRef StartNode, const FVector& StartLo
 	{
 		// start location is not on navmesh, treat it as a blocked raycast
 		RaycastResult.HitTime = 0.f;
-		RaycastResult.HitNormal = (StartLoc - EndLoc).SafeNormal();
+		RaycastResult.HitNormal = (StartLoc - EndLoc).GetSafeNormal();
 	}
 }
 
@@ -2107,7 +2107,7 @@ void FPImplRecastNavMesh::GetDebugGeometry(FRecastDebugGeometry& OutGeometry, in
 						if (linkedIdx > i || TileIdx > (int32)ConstNavMesh->decodeClusterIdTile(link.ref))
 						{
 							FVector UpDir(0,0,1.0f);
-							FVector LinkDir = (LinkGeom.ToCluster - LinkGeom.FromCluster).SafeNormal();
+							FVector LinkDir = (LinkGeom.ToCluster - LinkGeom.FromCluster).GetSafeNormal();
 							FVector SideDir = FVector::CrossProduct(LinkDir, UpDir);
 							LinkGeom.FromCluster += SideDir * 40.0f;
 							LinkGeom.ToCluster += SideDir * 40.0f;
diff --git a/Engine/Source/Runtime/Engine/Private/AI/Navigation/RecastNavMesh.cpp b/Engine/Source/Runtime/Engine/Private/AI/Navigation/RecastNavMesh.cpp
--- a/Engine/Source/Runtime/Engine/Private/AI/Navigation/RecastNavMesh.cpp
+++ b/Engine/Source/Runtime/Engine/Private/AI/Navigation/RecastNavMesh.cpp
@@ -1431,7 +1431,7 @@ bool ARecastNavMesh::AdjustLocationWithFilter(const FVector& StartLoc, FVector&
 		// move it just a bit further - otherwise recast can still pick "wrong" poly when 
 		// later projecting StartLoc (meaning a poly we want to filter out with 
 		// QueryFilter here)
-		OutAdjustedLocation += (OutAdjustedLocation - StartLoc).SafeNormal() * 0.1f;
+		OutAdjustedLocation += (OutAdjustedLocation - StartLoc).GetSafeNormal() * 0.1f;
 		return true;
 	}
 }
diff --git a/Engine/Source/Runtime/Engine/Private/AI/Navigation/RecastNavMeshGenerator.cpp b/Engine/Source/Runtime/Engine/Private/AI/Navigation/RecastNavMeshGenerator.cpp
--- a/Engine/Source/Runtime/Engine/Private/AI/Navigation/RecastNavMeshGenerator.cpp
+++ b/Engine/Source/Runtime/Engine/Private/AI/Navigation/RecastNavMeshGenerator.cpp
@@ -1038,8 +1038,8 @@ FORCEINLINE void GrowConvexHull(const float ExpandBy, const TArray<FVector>& Ver
 		const FVector& V2 = AllVerts[Index + 1];
 		const FVector& V3 = AllVerts[Index + 2];
 
-		const FVector V01 = (V1 - V2).SafeNormal();
-		const FVector V12 = (V2 - V3).SafeNormal();
+		const FVector V01 = (V1 - V2).GetSafeNormal();
+		const FVector V12 = (V2 - V3).GetSafeNormal();
 		const FVector NV1 = Rotation90.RotateVector(V01);
 		const float d = FVector::DotProduct(NV1, V12);
 
@@ -1081,14 +1081,14 @@ FORCEINLINE void GrowConvexHull(const float ExpandBy, const TArray<FVector>& Ver
 		}
 		else
 		{
-			const FVector V01 = (V1 - V2).SafeNormal();
-			const FVector N1 = Rotation.RotateVector(V01).SafeNormal();
+			const FVector V01 = (V1 - V2).GetSafeNormal();
+			const FVector N1 = Rotation.RotateVector(V01).GetSafeNormal();
 			const FVector MoveDir1 = N1 * ExpandBy;
 			Line1 = FSimpleLine(V1 + MoveDir1, V2 + MoveDir1);
 		}
 
-		const FVector V12 = (V2 - V3).SafeNormal();
-		const FVector N2 = Rotation.RotateVector(V12).SafeNormal();
+		const FVector V12 = (V2 - V3).GetSafeNormal();
+		const FVector N2 = Rotation.RotateVector(V12).GetSafeNormal();
 		const FVector MoveDir2 = N2 * ExpandBy;
 		const FSimpleLine Line2(V2 + MoveDir2, V3 + MoveDir2);
 
@@ -1105,7 +1105,7 @@ FORCEINLINE void GrowConvexHull(const float ExpandBy, const TArray<FVector>& Ver
 			if (DistToNewVector > ExpansionThresholdSQ)
 			{
 				//clamp our point to not move to far from original location
-				const FVector HelpPos = V2 + VectorToNewPoint.SafeNormal2D() * ExpandBy * 1.4142;
+				const FVector HelpPos = V2 + VectorToNewPoint.GetSafeNormal2D() * ExpandBy * 1.4142;
 				OutResult.Add(HelpPos);
 			}
 			else
diff --git a/Engine/Source/Runtime/Engine/Private/Actor.cpp b/Engine/Source/Runtime/Engine/Private/Actor.cpp
--- a/Engine/Source/Runtime/Engine/Private/Actor.cpp
+++ b/Engine/Source/Runtime/Engine/Private/Actor.cpp
@@ -3501,7 +3501,7 @@ float AActor::GetDotProductTo(AActor* OtherActor)
 	{
 		FVector Dir = GetActorForwardVector();
 		FVector Offset = OtherActor->GetActorLocation() - GetActorLocation();
-		Offset = Offset.SafeNormal();
+		Offset = Offset.GetSafeNormal();
 		return FVector::DotProduct(Dir, Offset);
 	}
 	return -2.0;
@@ -3513,7 +3513,7 @@ float AActor::GetHorizontalDotProductTo(AActor* OtherActor)
 	{
 		FVector Dir = GetActorForwardVector();
 		FVector Offset = OtherActor->GetActorLocation() - GetActorLocation();
-		Offset = Offset.SafeNormal2D();
+		Offset = Offset.GetSafeNormal2D();
 		return FVector::DotProduct(Dir, Offset);
 	}
 	return -2.0;
diff --git a/Engine/Source/Runtime/Engine/Private/Animation/AnimCompress_RemoveLinearKeys.cpp b/Engine/Source/Runtime/Engine/Private/Animation/AnimCompress_RemoveLinearKeys.cpp
--- a/Engine/Source/Runtime/Engine/Private/Animation/AnimCompress_RemoveLinearKeys.cpp
+++ b/Engine/Source/Runtime/Engine/Private/Animation/AnimCompress_RemoveLinearKeys.cpp
@@ -740,7 +740,7 @@ void UAnimCompress_RemoveLinearKeys::ProcessAnimationTracks(
 							// if these are valid, we can continue
 							if (!CurrentHeading.IsNearlyZero() && !DesiredHeading.IsNearlyZero())
 							{
-								const float DotResult = CurrentHeading.SafeNormal() | DesiredHeading.SafeNormal();
+								const float DotResult = CurrentHeading.GetSafeNormal() | DesiredHeading.GetSafeNormal();
 
 								// limit the range we will retarget to something reasonable (~60 degrees)
 								if (DotResult < 1.0f && DotResult > 0.5f)
diff --git a/Engine/Source/Runtime/Engine/Private/Animation/AnimInstance.cpp b/Engine/Source/Runtime/Engine/Private/Animation/AnimInstance.cpp
--- a/Engine/Source/Runtime/Engine/Private/Animation/AnimInstance.cpp
+++ b/Engine/Source/Runtime/Engine/Private/Animation/AnimInstance.cpp
@@ -2252,7 +2252,7 @@ float UAnimInstance::CalculateDirection(const FVector& Velocity, const FRotator&
 	FMatrix RotMatrix = FRotationMatrix(BaseRotation);
 	FVector ForwardVector = RotMatrix.GetScaledAxis(EAxis::X);
 	FVector RightVector = RotMatrix.GetScaledAxis(EAxis::Y);
-	FVector NormalizedVel = Velocity.SafeNormal();
+	FVector NormalizedVel = Velocity.GetSafeNormal();
 	ForwardVector.Z = RightVector.Z = NormalizedVel.Z = 0.f;
 
 	// get a cos(alpha) of forward vector vs velocity
diff --git a/Engine/Source/Runtime/Engine/Private/Animation/AnimNode_Fabrik.cpp b/Engine/Source/Runtime/Engine/Private/Animation/AnimNode_Fabrik.cpp
--- a/Engine/Source/Runtime/Engine/Private/Animation/AnimNode_Fabrik.cpp
+++ b/Engine/Source/Runtime/Engine/Private/Animation/AnimNode_Fabrik.cpp
@@ -111,7 +111,7 @@ void FAnimNode_Fabrik::EvaluateBoneTransforms(USkeletalMeshComponent* SkelComp,
 		{
 			FABRIKChainLink const & ParentLink = Chain[LinkIndex - 1];
 			FABRIKChainLink & CurrentLink = Chain[LinkIndex];
-			CurrentLink.Position = ParentLink.Position + (CSEffectorLocation - ParentLink.Position).UnsafeNormal() * CurrentLink.Length;
+			CurrentLink.Position = ParentLink.Position + (CSEffectorLocation - ParentLink.Position).GetUnsafeNormal() * CurrentLink.Length;
 		}
 		bBoneLocationUpdated = true;
 	}
@@ -135,7 +135,7 @@ void FAnimNode_Fabrik::EvaluateBoneTransforms(USkeletalMeshComponent* SkelComp,
 					FABRIKChainLink & CurrentLink = Chain[LinkIndex];
 					FABRIKChainLink const & ChildLink = Chain[LinkIndex + 1];
 
-					CurrentLink.Position = ChildLink.Position + (CurrentLink.Position - ChildLink.Position).UnsafeNormal() * ChildLink.Length;
+					CurrentLink.Position = ChildLink.Position + (CurrentLink.Position - ChildLink.Position).GetUnsafeNormal() * ChildLink.Length;
 				}
 
 				// "Backward Reaching" stage - adjust bones from root.
@@ -144,7 +144,7 @@ void FAnimNode_Fabrik::EvaluateBoneTransforms(USkeletalMeshComponent* SkelComp,
 					FABRIKChainLink const & ParentLink = Chain[LinkIndex - 1];
 					FABRIKChainLink & CurrentLink = Chain[LinkIndex];
 
-					CurrentLink.Position = ParentLink.Position + (CurrentLink.Position - ParentLink.Position).UnsafeNormal() * CurrentLink.Length;
+					CurrentLink.Position = ParentLink.Position + (CurrentLink.Position - ParentLink.Position).GetUnsafeNormal() * CurrentLink.Length;
 				}
 
 				// Re-check distance between tip location and effector location
@@ -157,7 +157,7 @@ void FAnimNode_Fabrik::EvaluateBoneTransforms(USkeletalMeshComponent* SkelComp,
 				FABRIKChainLink const & ParentLink = Chain[TipBoneLinkIndex - 1];
 				FABRIKChainLink & CurrentLink = Chain[TipBoneLinkIndex];
 
-				CurrentLink.Position = ParentLink.Position + (CurrentLink.Position - ParentLink.Position).UnsafeNormal() * CurrentLink.Length;
+				CurrentLink.Position = ParentLink.Position + (CurrentLink.Position - ParentLink.Position).GetUnsafeNormal() * CurrentLink.Length;
 			}
 
 			bBoneLocationUpdated = true;
@@ -188,13 +188,13 @@ void FAnimNode_Fabrik::EvaluateBoneTransforms(USkeletalMeshComponent* SkelComp,
 			FABRIKChainLink const & ChildLink = Chain[LinkIndex + 1];
 
 			// Calculate pre-translation vector between this bone and child
-			FVector const OldDir = (GetCurrentLocation(MeshBases, ChildLink.BoneIndex) - GetCurrentLocation(MeshBases, CurrentLink.BoneIndex)).UnsafeNormal();
+			FVector const OldDir = (GetCurrentLocation(MeshBases, ChildLink.BoneIndex) - GetCurrentLocation(MeshBases, CurrentLink.BoneIndex)).GetUnsafeNormal();
 
 			// Get vector from the post-translation bone to it's child
-			FVector const NewDir = (ChildLink.Position - CurrentLink.Position).UnsafeNormal();
+			FVector const NewDir = (ChildLink.Position - CurrentLink.Position).GetUnsafeNormal();
 
 			// Calculate axis of rotation from pre-translation vector to post-translation vector
-			FVector const RotationAxis = FVector::CrossProduct(OldDir, NewDir).SafeNormal();
+			FVector const RotationAxis = FVector::CrossProduct(OldDir, NewDir).GetSafeNormal();
 			float const RotationAngle = FMath::Acos(FVector::DotProduct(OldDir, NewDir));
 			FQuat const DeltaRotation = FQuat(RotationAxis, RotationAngle);
 			// We're going to multiply it, in order to not have to re-normalize the final quaternion, it has to be a unit quaternion.
diff --git a/Engine/Source/Runtime/Engine/Private/Animation/AnimNode_SpringBone.cpp b/Engine/Source/Runtime/Engine/Private/Animation/AnimNode_SpringBone.cpp
--- a/Engine/Source/Runtime/Engine/Private/Animation/AnimNode_SpringBone.cpp
+++ b/Engine/Source/Runtime/Engine/Private/Animation/AnimNode_SpringBone.cpp
@@ -137,7 +137,7 @@ void FAnimNode_SpringBone::EvaluateBoneTransforms(USkeletalMeshComponent* SkelCo
 			// Too far away - project back onto sphere around target.
 			if (CurrentDisp.Size() > MaxDisplacement)
 			{
-				FVector DispDir = CurrentDisp.SafeNormal();
+				FVector DispDir = CurrentDisp.GetSafeNormal();
 				BoneLocation = TargetPos + (MaxDisplacement * DispDir);
 			}
 		}
diff --git a/Engine/Source/Runtime/Engine/Private/Animation/AnimNode_Trail.cpp b/Engine/Source/Runtime/Engine/Private/Animation/AnimNode_Trail.cpp
--- a/Engine/Source/Runtime/Engine/Private/Animation/AnimNode_Trail.cpp
+++ b/Engine/Source/Runtime/Engine/Private/Animation/AnimNode_Trail.cpp
@@ -173,10 +173,10 @@ void FAnimNode_Trail::EvaluateBoneTransforms(USkeletalMeshComponent* SkelComp, c
 
 		// Calculate the direction that parent bone is currently pointing.
 		FVector CurrentBoneDir = OutBoneTransforms[i-1].Transform.TransformVector( GetAlignVector(ChainBoneAxis, bInvertChainBoneAxis) );
-		CurrentBoneDir = CurrentBoneDir.SafeNormal(SMALL_NUMBER);
+		CurrentBoneDir = CurrentBoneDir.GetSafeNormal(SMALL_NUMBER);
 
 		// Calculate vector from parent to child.
-		FVector NewBoneDir = FVector(OutBoneTransforms[i].Transform.GetTranslation() - OutBoneTransforms[i - 1].Transform.GetTranslation()).SafeNormal(SMALL_NUMBER);
+		FVector NewBoneDir = FVector(OutBoneTransforms[i].Transform.GetTranslation() - OutBoneTransforms[i - 1].Transform.GetTranslation()).GetSafeNormal(SMALL_NUMBER);
 
 		// Calculate a quaternion that gets us from our current rotation to the desired one.
 		FQuat DeltaLookQuat = FQuat::FindBetween(CurrentBoneDir, NewBoneDir);
diff --git a/Engine/Source/Runtime/Engine/Private/Animation/AnimNode_TwoBoneIK.cpp b/Engine/Source/Runtime/Engine/Private/Animation/AnimNode_TwoBoneIK.cpp
--- a/Engine/Source/Runtime/Engine/Private/Animation/AnimNode_TwoBoneIK.cpp
+++ b/Engine/Source/Runtime/Engine/Private/Animation/AnimNode_TwoBoneIK.cpp
@@ -225,8 +225,8 @@ void FAnimNode_TwoBoneIK::EvaluateBoneTransforms(USkeletalMeshComponent* SkelCom
 	// Update transform for upper bone.
 	{
 		// Get difference in direction for old and new joint orientations
-		FVector const OldDir = (InitialJointPos - RootPos).SafeNormal();
-		FVector const NewDir = (OutJointPos - RootPos).SafeNormal();
+		FVector const OldDir = (InitialJointPos - RootPos).GetSafeNormal();
+		FVector const NewDir = (OutJointPos - RootPos).GetSafeNormal();
 		// Find Delta Rotation take takes us from Old to New dir
 		FQuat const DeltaRotation = FQuat::FindBetween(OldDir, NewDir);
 		// Rotate our Joint quaternion by this delta rotation
@@ -241,8 +241,8 @@ void FAnimNode_TwoBoneIK::EvaluateBoneTransforms(USkeletalMeshComponent* SkelCom
 	// Update transform for lower bone.
 	{
 		// Get difference in direction for old and new joint orientations
-		FVector const OldDir = (InitialEndPos - InitialJointPos).SafeNormal();
-		FVector const NewDir = (OutEndPos - OutJointPos).SafeNormal();
+		FVector const OldDir = (InitialEndPos - InitialJointPos).GetSafeNormal();
+		FVector const NewDir = (OutEndPos - OutJointPos).GetSafeNormal();
 
 		// Find Delta Rotation take takes us from Old to New dir
 		FQuat const DeltaRotation = FQuat::FindBetween(OldDir, NewDir);
diff --git a/Engine/Source/Runtime/Engine/Private/BatchedElements.cpp b/Engine/Source/Runtime/Engine/Private/BatchedElements.cpp
--- a/Engine/Source/Runtime/Engine/Private/BatchedElements.cpp
+++ b/Engine/Source/Runtime/Engine/Private/BatchedElements.cpp
@@ -678,9 +678,9 @@ bool FBatchedElements::Draw(FRHICommandList& RHICmdList, ERHIFeatureLevel::Type
 	if( HasPrimsToDraw() )
 	{
 		FMatrix InvTransform = Transform.Inverse();
-		FVector CameraX = InvTransform.TransformVector(FVector(1,0,0)).SafeNormal();
-		FVector CameraY = InvTransform.TransformVector(FVector(0,1,0)).SafeNormal();
-		FVector CameraZ = InvTransform.TransformVector(FVector(0,0,1)).SafeNormal();
+		FVector CameraX = InvTransform.TransformVector(FVector(1,0,0)).GetSafeNormal();
+		FVector CameraY = InvTransform.TransformVector(FVector(0,1,0)).GetSafeNormal();
+		FVector CameraZ = InvTransform.TransformVector(FVector(0,0,1)).GetSafeNormal();
 
 		RHICmdList.SetRasterizerState(TStaticRasterizerState<FM_Solid, CM_None>::GetRHI());
 
diff --git a/Engine/Source/Runtime/Engine/Private/CharacterMovementComponent.cpp b/Engine/Source/Runtime/Engine/Private/CharacterMovementComponent.cpp
--- a/Engine/Source/Runtime/Engine/Private/CharacterMovementComponent.cpp
+++ b/Engine/Source/Runtime/Engine/Private/CharacterMovementComponent.cpp
@@ -541,7 +541,7 @@ void UCharacterMovementComponent::JumpOff(AActor* MovementBaseActor)
 			Velocity += MaxSpeed * GetBestDirectionOffActor(MovementBaseActor);
 			if ( Velocity.Size2D() > MaxSpeed )
 			{
-				Velocity = MaxSpeed * Velocity.SafeNormal();
+				Velocity = MaxSpeed * Velocity.GetSafeNormal();
 			}
 			Velocity.Z = JumpOffJumpZFactor * JumpZVelocity;
 			SetMovementMode(MOVE_Falling);
@@ -555,7 +555,7 @@ FVector UCharacterMovementComponent::GetBestDirectionOffActor(AActor* BaseActor)
 	// By default, just pick a random direction.  Derived character classes can choose to do more complex calculations,
 	// such as finding the shortest distance to move in based on the BaseActor's Bounding Volume.
 	const float RandAngle = FMath::DegreesToRadians(GetNetworkSafeRandomAngleDegrees());
-	return FVector(FMath::Cos(RandAngle), FMath::Sin(RandAngle), 0.5f).SafeNormal();
+	return FVector(FMath::Cos(RandAngle), FMath::Sin(RandAngle), 0.5f).GetSafeNormal();
 }
 
 float UCharacterMovementComponent::GetNetworkSafeRandomAngleDegrees() const
@@ -757,7 +757,7 @@ void UCharacterMovementComponent::PerformAirControlForPathFollowing(FVector Dire
 			{
 				float Dist2D = Direction.Size2D();
 				//Direction.Z = 0.f;
-				Acceleration = Direction.SafeNormal() * GetMaxAcceleration();
+				Acceleration = Direction.GetSafeNormal() * GetMaxAcceleration();
 
 				if ( (Dist2D < 0.5f * FMath::Abs(Direction.Z)) && ((Velocity | Direction) > 0.5f*FMath::Square(Dist2D)) )
 				{
@@ -1100,7 +1100,7 @@ void UCharacterMovementComponent::SimulateMovement(float DeltaSeconds)
 			return;
 		}
 
-		Acceleration = Velocity.SafeNormal();	// Not currently used for simulated movement
+		Acceleration = Velocity.GetSafeNormal();	// Not currently used for simulated movement
 		AnalogInputModifier = 1.0f;				// Not currently used for simulated movement
 
 		MaybeUpdateBasedMovement(DeltaSeconds);
@@ -1899,7 +1899,7 @@ float UCharacterMovementComponent::SlideAlongSurface(const FVector& Delta, float
 		{
 			if (!IsWalkable(Hit))
 			{
-				Normal = Normal.SafeNormal2D();
+				Normal = Normal.GetSafeNormal2D();
 			}
 		}
 		else if (Normal.Z < -KINDA_SMALL_NUMBER)
@@ -1914,7 +1914,7 @@ float UCharacterMovementComponent::SlideAlongSurface(const FVector& Delta, float
 					Normal = FloorNormal;
 				}
 				
-				Normal = Normal.SafeNormal2D();
+				Normal = Normal.GetSafeNormal2D();
 			}
 		}
 	}
@@ -1937,7 +1937,7 @@ void UCharacterMovementComponent::TwoWallAdjust(FVector &Delta, const FHitResult
 			{
 				// Maintain horizontal velocity
 				const float Time = (1.f - Hit.Time);
-				const FVector ScaledDelta = Delta.SafeNormal() * InDelta.Size();
+				const FVector ScaledDelta = Delta.GetSafeNormal() * InDelta.Size();
 				Delta = FVector(InDelta.X, InDelta.Y, ScaledDelta.Z / Hit.Normal.Z) * Time;
 			}
 			else
@@ -1995,7 +1995,7 @@ FVector UCharacterMovementComponent::HandleSlopeBoosting(const FVector& SlideRes
 
 			// Make remaining portion of original result horizontal and parallel to impact normal.
 			const FVector RemainderXY = (SlideResult - Result) * FVector(1.f, 1.f, 0.f);
-			const FVector NormalXY = Normal.SafeNormal2D();
+			const FVector NormalXY = Normal.GetSafeNormal2D();
 			const FVector Adjust = Super::ComputeSlideVector(RemainderXY, 1.f, NormalXY, Hit);
 			Result += Adjust;
 		}
@@ -2022,7 +2022,7 @@ FVector UCharacterMovementComponent::NewFallVelocity(const FVector& InitialVeloc
 		// Apply gravity.
 		Result += Gravity * DeltaTime;
 
-		const FVector GravityDir = Gravity.SafeNormal();
+		const FVector GravityDir = Gravity.GetSafeNormal();
 		const float TerminalLimit = FMath::Abs(GetPhysicsVolume()->TerminalVelocity);
 
 		// Don't exceed terminal velocity.
@@ -2141,7 +2141,7 @@ void UCharacterMovementComponent::CalcVelocity(float DeltaTime, float Friction,
 	float RequestedSpeed = 0.0f;
 	if (ApplyRequestedMove(DeltaTime, MaxAccel, MaxSpeed, Friction, BrakingDeceleration, RequestedAcceleration, RequestedSpeed))
 	{
-		RequestedAcceleration = RequestedAcceleration.ClampMaxSize(MaxAccel);
+		RequestedAcceleration = RequestedAcceleration.GetClampedToMaxSize(MaxAccel);
 		bZeroRequestedAcceleration = false;
 	}
 
@@ -2237,7 +2237,7 @@ bool UCharacterMovementComponent::ApplyRequestedMove(float DeltaTime, float MaxA
 
 			// How much do we need to accelerate to get to the new velocity?
 			NewAcceleration = ((MoveVelocity - Velocity) / DeltaTime);
-			NewAcceleration = NewAcceleration.ClampMaxSize(MaxAccel);
+			NewAcceleration = NewAcceleration.GetClampedToMaxSize(MaxAccel);
 		}
 		else
 		{
@@ -2264,7 +2264,7 @@ void UCharacterMovementComponent::RequestDirectMove(const FVector& MoveVelocity,
 
 	if (IsFalling())
 	{
-		const FVector FallVelocity = MoveVelocity.ClampMaxSize(GetMaxSpeed());
+		const FVector FallVelocity = MoveVelocity.GetClampedToMaxSize(GetMaxSpeed());
 		PerformAirControlForPathFollowing(FallVelocity, FallVelocity.Z);
 		return;
 	}
@@ -2627,7 +2627,7 @@ void UCharacterMovementComponent::PhysFlying(float deltaTime, int32 Iterations)
 	if( Hit.Time < 1.f && CharacterOwner )
 	{
 		const FVector GravDir = FVector(0.f,0.f,-1.f);
-		const FVector VelDir = Velocity.SafeNormal();
+		const FVector VelDir = Velocity.GetSafeNormal();
 		const float UpDown = GravDir | VelDir;
 
 		bool bSteppedUp = false;
@@ -2695,7 +2695,7 @@ void UCharacterMovementComponent::PhysSwimming(float deltaTime, int32 Iterations
 	if ( Hit.Time < 1.f && CharacterOwner)
 	{
 		const FVector GravDir = FVector(0.f,0.f,-1.f);
-		const FVector VelDir = Velocity.SafeNormal();
+		const FVector VelDir = Velocity.GetSafeNormal();
 		const float UpDown = GravDir | VelDir;
 
 		bool bSteppedUp = false;
@@ -2761,7 +2761,7 @@ void UCharacterMovementComponent::StartSwimming(FVector OldLocation, FVector Old
 	{
 		Velocity = (CharacterOwner->GetActorLocation() - OldLocation)/timeTick; //actual average velocity
 		Velocity = 2.f*Velocity - OldVelocity; //end velocity has 2* accel of avg
-		Velocity.ClampMaxSize(GetPhysicsVolume()->TerminalVelocity);
+		Velocity = Velocity.GetClampedToMaxSize(GetPhysicsVolume()->TerminalVelocity);
 	}
 	const FVector End = FindWaterLine(CharacterOwner->GetActorLocation(), OldLocation);
 	float waterTime = 0.f;
@@ -2824,7 +2824,7 @@ FVector UCharacterMovementComponent::FindWaterLine(FVector InWater, FVector Outo
 			APhysicsVolume *W = Cast<APhysicsVolume>(Check.GetActor());
 			if ( W && W->bWaterVolume )
 			{
-				FVector Dir = (InWater - OutofWater).SafeNormal();
+				FVector Dir = (InWater - OutofWater).GetSafeNormal();
 				Result = Check.Location;
 				if ( W == GetPhysicsVolume() )
 					Result += 0.1f * Dir;
@@ -2856,7 +2856,7 @@ FVector UCharacterMovementComponent::GetFallingLateralAcceleration(float DeltaTi
 	if (!HasRootMotion() && FallAcceleration.SizeSquared2D() > 0.f)
 	{
 		FallAcceleration = GetAirControl(DeltaTime, AirControl, FallAcceleration);
-		FallAcceleration = FallAcceleration.ClampMaxSize(GetMaxAcceleration());
+		FallAcceleration = FallAcceleration.GetClampedToMaxSize(GetMaxAcceleration());
 	}
 
 	return FallAcceleration;
@@ -3093,10 +3093,10 @@ void UCharacterMovementComponent::PhysFalling(float deltaTime, int32 Iterations)
 						if ( Hit.Time == 0 )
 						{
 							// if we are stuck then try to side step
-							FVector SideDelta = (OldHitNormal + Hit.ImpactNormal).SafeNormal2D();
+							FVector SideDelta = (OldHitNormal + Hit.ImpactNormal).GetSafeNormal2D();
 							if ( SideDelta.IsNearlyZero() )
 							{
-								SideDelta = FVector(OldHitNormal.Y, -OldHitNormal.X, 0).SafeNormal();
+								SideDelta = FVector(OldHitNormal.Y, -OldHitNormal.X, 0).GetSafeNormal();
 							}
 							SafeMoveUpdatedComponent( SideDelta, PawnRotation, true, Hit);
 						}
@@ -3177,7 +3177,7 @@ FVector UCharacterMovementComponent::LimitAirControl(float DeltaTime, const FVec
 			if (FVector::DotProduct(FallAcceleration, HitResult.Normal) < 0.f)
 			{
 				// Allow movement parallel to the wall, but not into it because that may push us up.
-				const FVector Normal2D = HitResult.Normal.SafeNormal2D();
+				const FVector Normal2D = HitResult.Normal.GetSafeNormal2D();
 				Result = FVector::VectorPlaneProject(FallAcceleration, Normal2D);
 			}
 		}
@@ -3364,7 +3364,7 @@ FVector UCharacterMovementComponent::ComputeGroundMovementDelta(const FVector& D
 		}
 		else
 		{
-			return RampMovement.SafeNormal() * Delta.Size();
+			return RampMovement.GetSafeNormal() * Delta.Size();
 		}
 	}
 
@@ -3456,7 +3456,7 @@ void UCharacterMovementComponent::MaintainHorizontalGroundVelocity()
 		else
 		{
 			// Rescale velocity to be horizontal but maintain magnitude of last update.
-			Velocity = Velocity.SafeNormal2D() * Velocity.Size();
+			Velocity = Velocity.GetSafeNormal2D() * Velocity.Size();
 		}
 	}
 }
@@ -3858,7 +3858,7 @@ FRotator UCharacterMovementComponent::ComputeOrientToMovementRotation(const FRot
 		// AI path following request can orient us in that direction (it's effectively an acceleration)
 		if (bHasRequestedVelocity && RequestedVelocity.SizeSquared() > KINDA_SMALL_NUMBER)
 		{
-			return RequestedVelocity.SafeNormal().Rotation();
+			return RequestedVelocity.GetSafeNormal().Rotation();
 		}
 
 		// Don't change rotation if there is no acceleration.
@@ -3866,7 +3866,7 @@ FRotator UCharacterMovementComponent::ComputeOrientToMovementRotation(const FRot
 	}
 
 	// Rotate toward direction of acceleration.
-	return Acceleration.SafeNormal().Rotation();
+	return Acceleration.GetSafeNormal().Rotation();
 }
 
 void UCharacterMovementComponent::PhysicsRotation(float DeltaTime)
@@ -4008,7 +4008,7 @@ bool UCharacterMovementComponent::CheckWaterJump(FVector CheckPoint, FVector& Wa
 	}
 	// check if there is a wall directly in front of the swimming pawn
 	CheckPoint.Z = 0.f;
-	FVector CheckNorm = CheckPoint.SafeNormal();
+	FVector CheckNorm = CheckPoint.GetSafeNormal();
 	float PawnCapsuleRadius, PawnCapsuleHalfHeight;
 	CharacterOwner->GetCapsuleComponent()->GetScaledCapsuleSize(PawnCapsuleRadius, PawnCapsuleHalfHeight);
 	CheckPoint = CharacterOwner->GetActorLocation() + 1.2f * PawnCapsuleRadius * CheckNorm;
@@ -4121,7 +4121,7 @@ void UCharacterMovementComponent::MoveSmooth(const FVector& InVelocity, const fl
 					if (FMath::Abs(Hit.ImpactNormal.Z) < 0.2f)
 					{
 						const FVector GravDir = FVector(0.f,0.f,-1.f);
-						const FVector DesiredDir = Delta.SafeNormal();
+						const FVector DesiredDir = Delta.GetSafeNormal();
 						const float UpDown = GravDir | DesiredDir;
 						if ((UpDown < 0.5f) && (UpDown > -0.2f))
 						{
@@ -4697,7 +4697,7 @@ bool UCharacterMovementComponent::StepUp(const FVector& InGravDir, const FVector
 		return false;
 	}
 
-	const FVector GravDir = InGravDir.SafeNormal();
+	const FVector GravDir = InGravDir.GetSafeNormal();
 	if (GravDir.IsZero())
 	{
 		return false;
@@ -4936,7 +4936,7 @@ void UCharacterMovementComponent::ApplyImpactPhysicsForces(const FHitResult& Imp
 			float PushForceModificator = 1.0f;
 
 			const FVector ComponentVelocity = ImpactComponent->GetPhysicsLinearVelocity();
-			const FVector VirtualVelocity = ImpactAcceleration.IsZero() ? ImpactVelocity : ImpactAcceleration.SafeNormal() * GetMaxSpeed();
+			const FVector VirtualVelocity = ImpactAcceleration.IsZero() ? ImpactVelocity : ImpactAcceleration.GetSafeNormal() * GetMaxSpeed();
 
 			float Dot = 0.0f;
 
@@ -5131,7 +5131,7 @@ void UCharacterMovementComponent::SmoothCorrection(const FVector& OldLocation)
 		{
 			ClientData->MeshTranslationOffset = (DistSq > FMath::Square(ClientData->NoSmoothNetUpdateDist)) 
 				? FVector::ZeroVector 
-				: ClientData->MeshTranslationOffset + ClientData->MaxSmoothNetUpdateDist * (OldLocation - CharacterOwner->GetActorLocation()).SafeNormal();	
+				: ClientData->MeshTranslationOffset + ClientData->MaxSmoothNetUpdateDist * (OldLocation - CharacterOwner->GetActorLocation()).GetSafeNormal();
 		}
 		else
 		{
@@ -6219,8 +6219,8 @@ void UCharacterMovementComponent::CapsuleTouched( AActor* Other, UPrimitiveCompo
 	{
 		const FVector OtherLoc = OtherComp->GetComponentLocation();
 		const FVector Loc = UpdatedComponent->GetComponentLocation();
-		FVector ImpulseDir = FVector(OtherLoc.X - Loc.X, OtherLoc.Y - Loc.Y, 0.25f).SafeNormal();
-		ImpulseDir = (ImpulseDir + Velocity.SafeNormal2D()) * 0.5f;
+		FVector ImpulseDir = FVector(OtherLoc.X - Loc.X, OtherLoc.Y - Loc.Y, 0.25f).GetSafeNormal();
+		ImpulseDir = (ImpulseDir + Velocity.GetSafeNormal2D()) * 0.5f;
 		ImpulseDir.Normalize();
 
 		FName BoneName = NAME_None;
@@ -6396,7 +6396,7 @@ void UCharacterMovementComponent::AddRadialForce(const FVector& Origin, float Ra
 		return;
 	}
 
-	Delta = Delta.SafeNormal();
+	Delta = Delta.GetSafeNormal();
 
 	float ForceMagnitude = Strength;
 	if (Falloff == RIF_Linear && Radius > 0.0f)
@@ -6418,7 +6418,7 @@ void UCharacterMovementComponent::AddRadialImpulse(const FVector& Origin, float
 		return;
 	}
 
-	Delta = Delta.SafeNormal();
+	Delta = Delta.GetSafeNormal();
 
 	float ImpulseMagnitude = Strength;
 	if (Falloff == RIF_Linear && Radius > 0.0f)
diff --git a/Engine/Source/Runtime/Engine/Private/Collision/CollisionConversions.cpp b/Engine/Source/Runtime/Engine/Private/Collision/CollisionConversions.cpp
--- a/Engine/Source/Runtime/Engine/Private/Collision/CollisionConversions.cpp
+++ b/Engine/Source/Runtime/Engine/Private/Collision/CollisionConversions.cpp
@@ -85,7 +85,7 @@ static PxVec3 TransformNormalToShapeSpace(const PxMeshScale& meshScale, const Px
 static bool FindSimpleOpposingNormal(const PxLocationHit& PHit, const FVector& TraceDirectionDenorm, FVector& OutNormal)
 {
 	const bool bNormalData = PHit.flags & PxHitFlag::eNORMAL;
-	OutNormal = bNormalData ? P2UVector(PHit.normal) : -TraceDirectionDenorm.SafeNormal();
+	OutNormal = bNormalData ? P2UVector(PHit.normal) : -TraceDirectionDenorm.GetSafeNormal();
 	return true;
 }
 
@@ -352,7 +352,7 @@ void ConvertQueryImpactHit(const PxLocationHit& PHit, FHitResult& OutResult, flo
 	// Other info
 	OutResult.Location = SafeLocationToFitShape;
 	OutResult.ImpactPoint = (PHit.flags & PxHitFlag::ePOSITION) ? P2UVector(PHit.position) : StartLoc;
-	OutResult.Normal = (PHit.flags & PxHitFlag::eNORMAL) ? P2UVector(PHit.normal) : -TraceStartToEnd.SafeNormal();
+	OutResult.Normal = (PHit.flags & PxHitFlag::eNORMAL) ? P2UVector(PHit.normal) : -TraceStartToEnd.GetSafeNormal();
 	OutResult.ImpactNormal = OutResult.Normal;
 
 	OutResult.TraceStart = StartLoc;
@@ -680,7 +680,7 @@ static bool ConvertOverlappedShapeToImpactHit(const PxLocationHit& PHit, const F
 					PClosestPoint = PxShapeExt::getWorldBounds(*PShape, *PActor).getCenter(); 
 				}
 
-				OutResult.ImpactNormal = (OutResult.Location - P2UVector(PClosestPoint)).SafeNormal();
+				OutResult.ImpactNormal = (OutResult.Location - P2UVector(PClosestPoint)).GetSafeNormal();
 			}
 		}
 	}
diff --git a/Engine/Source/Runtime/Engine/Private/Collision/PhysXCollision.cpp b/Engine/Source/Runtime/Engine/Private/Collision/PhysXCollision.cpp
--- a/Engine/Source/Runtime/Engine/Private/Collision/PhysXCollision.cpp
+++ b/Engine/Source/Runtime/Engine/Private/Collision/PhysXCollision.cpp
@@ -205,7 +205,7 @@ public:
 
 		if (ResponseMode != ResponseNone)
 		{
-			ProcessHit(Fixture, Fraction, FPhysicsIntegration2D::ConvertBoxVectorToUnreal(Normal).SafeNormal(), ResponseMode);
+			ProcessHit(Fixture, Fraction, FPhysicsIntegration2D::ConvertBoxVectorToUnreal(Normal).GetSafeNormal(), ResponseMode);
 		}
 
 		switch (ResponseMode)
@@ -274,7 +274,7 @@ protected:
 			PointQueryStartTime = FPhysicsIntegration2D::ConvertUnrealVectorToPerpendicularDistance(Start);
 			const float PointQueryEndTime = FPhysicsIntegration2D::ConvertUnrealVectorToPerpendicularDistance(End);
 			PointQueryScaleFactor = 1.0f / (PointQueryEndTime - PointQueryStartTime);
-			PointQueryNormal = (Start - End).SafeNormal();
+			PointQueryNormal = (Start - End).GetSafeNormal();
 
 			// Run the query
 			b2AABB QueryBox;
diff --git a/Engine/Source/Runtime/Engine/Private/DestructibleComponent.cpp b/Engine/Source/Runtime/Engine/Private/DestructibleComponent.cpp
--- a/Engine/Source/Runtime/Engine/Private/DestructibleComponent.cpp
+++ b/Engine/Source/Runtime/Engine/Private/DestructibleComponent.cpp
@@ -968,7 +968,7 @@ void UDestructibleComponent::ApplyDamage(float DamageAmount, const FVector& HitL
 #if WITH_APEX
 	if (ApexDestructibleActor != NULL)
 	{
-		const FVector& NormalizedImpactDir = ImpulseDir.SafeNormal();
+		const FVector& NormalizedImpactDir = ImpulseDir.GetSafeNormal();
 
 		// Transfer damage information to the APEX NxDestructibleActor interface
 		ApexDestructibleActor->applyDamage(DamageAmount, ImpulseStrength, U2PVector( HitLocation ), U2PVector( ImpulseDir ));
diff --git a/Engine/Source/Runtime/Engine/Private/DirectionalLightComponent.cpp b/Engine/Source/Runtime/Engine/Private/DirectionalLightComponent.cpp
--- a/Engine/Source/Runtime/Engine/Private/DirectionalLightComponent.cpp
+++ b/Engine/Source/Runtime/Engine/Private/DirectionalLightComponent.cpp
@@ -226,7 +226,7 @@ public:
 		OutInitializer.bDirectionalLight = true;
 		OutInitializer.bOnePassPointLightShadow = false;
 		OutInitializer.PreShadowTranslation = -Bounds.Center;
-		OutInitializer.WorldToLight = FInverseRotationMatrix(FVector(WorldToLight.M[0][0],WorldToLight.M[1][0],WorldToLight.M[2][0]).SafeNormal().Rotation());
+		OutInitializer.WorldToLight = FInverseRotationMatrix(FVector(WorldToLight.M[0][0],WorldToLight.M[1][0],WorldToLight.M[2][0]).GetSafeNormal().Rotation());
 		OutInitializer.Scales = FVector(1.0f,1.0f / Bounds.W,1.0f / Bounds.W);
 		OutInitializer.FaceDirection = FVector(1,0,0);
 		OutInitializer.SubjectBounds = FBoxSphereBounds(FVector::ZeroVector,SubjectBounds.BoxExtent,SubjectBounds.SphereRadius);
@@ -253,7 +253,7 @@ public:
 
 		OutInitializer.bDirectionalLight = true;
 		OutInitializer.PreShadowTranslation = -LightPropagationVolumeBounds.GetCenter();
-		OutInitializer.WorldToLight = FInverseRotationMatrix(FVector(WorldToLight.M[0][0],WorldToLight.M[1][0],WorldToLight.M[2][0]).SafeNormal().Rotation());
+		OutInitializer.WorldToLight = FInverseRotationMatrix(FVector(WorldToLight.M[0][0],WorldToLight.M[1][0],WorldToLight.M[2][0]).GetSafeNormal().Rotation());
 		OutInitializer.Scales = FVector(1.0f,1.0f / LpvExtent,1.0f / LpvExtent);
 		OutInitializer.FaceDirection = FVector(1,0,0);
 		OutInitializer.SubjectBounds = FBoxSphereBounds( FVector::ZeroVector, LightPropagationVolumeBounds.GetExtent(), FMath::Sqrt( LpvExtent * LpvExtent * 3.0f ) );
@@ -299,7 +299,7 @@ public:
 	{
 		OutInitializer.bDirectionalLight = true;
 		OutInitializer.PreShadowTranslation = -SubjectBounds.Origin;
-		OutInitializer.WorldToLight = FInverseRotationMatrix(FVector(WorldToLight.M[0][0],WorldToLight.M[1][0],WorldToLight.M[2][0]).SafeNormal().Rotation());
+		OutInitializer.WorldToLight = FInverseRotationMatrix(FVector(WorldToLight.M[0][0],WorldToLight.M[1][0],WorldToLight.M[2][0]).GetSafeNormal().Rotation());
 		OutInitializer.Scales = FVector(1.0f,1.0f / SubjectBounds.SphereRadius,1.0f / SubjectBounds.SphereRadius);
 		OutInitializer.FaceDirection = FVector(1,0,0);
 		OutInitializer.SubjectBounds = FBoxSphereBounds(FVector::ZeroVector,SubjectBounds.BoxExtent,SubjectBounds.SphereRadius);
diff --git a/Engine/Source/Runtime/Engine/Private/DrawDebugHelpers.cpp b/Engine/Source/Runtime/Engine/Private/DrawDebugHelpers.cpp
--- a/Engine/Source/Runtime/Engine/Private/DrawDebugHelpers.cpp
+++ b/Engine/Source/Runtime/Engine/Private/DrawDebugHelpers.cpp
@@ -440,7 +440,7 @@ void DrawDebugCylinder(const UWorld* InWorld, FVector const& Start, FVector cons
 		float Angle = AngleInc;
 
 		// Default for Axis is up
-		FVector Axis = (End - Start).SafeNormal();
+		FVector Axis = (End - Start).GetSafeNormal();
 		if( Axis.IsZero() )
 		{
 			Axis = FVector(0.f, 0.f, 1.f);
@@ -584,7 +584,7 @@ void DrawDebugCone(const UWorld* InWorld, FVector const& Origin, FVector const&
 
 		// Calculate transform for cone.
 		FVector YAxis, ZAxis;
-		FVector DirectionNorm = Direction.SafeNormal();
+		FVector DirectionNorm = Direction.GetSafeNormal();
 		DirectionNorm.FindBestAxisVectors(YAxis, ZAxis);
 		const FMatrix ConeToWorld = FScaleMatrix(FVector(Length)) * FMatrix(DirectionNorm, YAxis, ZAxis, Origin);
 
diff --git a/Engine/Source/Runtime/Engine/Private/EngineTypes.cpp b/Engine/Source/Runtime/Engine/Private/EngineTypes.cpp
--- a/Engine/Source/Runtime/Engine/Private/EngineTypes.cpp
+++ b/Engine/Source/Runtime/Engine/Private/EngineTypes.cpp
@@ -49,7 +49,7 @@ void FDamageEvent::GetBestHitInfo(AActor const* HitActor, AActor const* HitInsti
 		
 		// assume hit came from instigator's location
 		OutImpulseDir = HitInstigator ? 
-			( OutHitInfo.ImpactPoint - HitInstigator->GetActorLocation() ).SafeNormal()
+			( OutHitInfo.ImpactPoint - HitInstigator->GetActorLocation() ).GetSafeNormal()
 			: FVector::ZeroVector;
 
 		// assume normal points back toward instigator
@@ -72,7 +72,7 @@ void FRadialDamageEvent::GetBestHitInfo(AActor const* HitActor, AActor const* Hi
 
 	// for now, just return the first one
 	OutHitInfo = ComponentHits[0];
-	OutImpulseDir = (OutHitInfo.ImpactPoint - Origin).SafeNormal();
+	OutImpulseDir = (OutHitInfo.ImpactPoint - Origin).GetSafeNormal();
 }
 
 
diff --git a/Engine/Source/Runtime/Engine/Private/FloatingPawnMovement.cpp b/Engine/Source/Runtime/Engine/Private/FloatingPawnMovement.cpp
--- a/Engine/Source/Runtime/Engine/Private/FloatingPawnMovement.cpp
+++ b/Engine/Source/Runtime/Engine/Private/FloatingPawnMovement.cpp
@@ -86,7 +86,7 @@ bool UFloatingPawnMovement::LimitWorldBounds()
 
 void UFloatingPawnMovement::ApplyControlInputToVelocity(float DeltaTime)
 {
-	const FVector ControlAcceleration = GetPendingInputVector().ClampMaxSize(1.f);
+	const FVector ControlAcceleration = GetPendingInputVector().GetClampedToMaxSize(1.f);
 
 	const float AnalogInputModifier = (ControlAcceleration.SizeSquared() > 0.f ? ControlAcceleration.Size() : 0.f);
 	const float MaxPawnSpeed = GetMaxSpeed() * AnalogInputModifier;
@@ -107,12 +107,12 @@ void UFloatingPawnMovement::ApplyControlInputToVelocity(float DeltaTime)
 		{
 			const FVector OldVelocity = Velocity;
 			const float VelSize = FMath::Max(Velocity.Size() - FMath::Abs(Deceleration) * DeltaTime, 0.f);
-			Velocity = Velocity.SafeNormal() * VelSize;
+			Velocity = Velocity.GetSafeNormal() * VelSize;
 
 			// Don't allow braking to lower us below max speed if we started above it.
 			if (bExceedingMaxSpeed && Velocity.SizeSquared() < FMath::Square(MaxPawnSpeed))
 			{
-				Velocity = OldVelocity.SafeNormal() * MaxPawnSpeed;
+				Velocity = OldVelocity.GetSafeNormal() * MaxPawnSpeed;
 			}
 		}
 	}
@@ -120,7 +120,7 @@ void UFloatingPawnMovement::ApplyControlInputToVelocity(float DeltaTime)
 	// Apply acceleration and clamp velocity magnitude.
 	const float NewMaxSpeed = (IsExceedingMaxSpeed(MaxPawnSpeed)) ? Velocity.Size() : MaxPawnSpeed;
 	Velocity += ControlAcceleration * FMath::Abs(Acceleration) * DeltaTime;
-	Velocity = Velocity.ClampMaxSize(NewMaxSpeed);
+	Velocity = Velocity.GetClampedToMaxSize(NewMaxSpeed);
 
 	ConsumeInputVector();
 }
diff --git a/Engine/Source/Runtime/Engine/Private/FoliageComponent.cpp b/Engine/Source/Runtime/Engine/Private/FoliageComponent.cpp
--- a/Engine/Source/Runtime/Engine/Private/FoliageComponent.cpp
+++ b/Engine/Source/Runtime/Engine/Private/FoliageComponent.cpp
@@ -83,7 +83,7 @@ float AInteractiveFoliageActor::TakeDamage(float DamageAmount, FDamageEvent cons
 	DamageEvent.GetBestHitInfo(this, (EventInstigator ? EventInstigator->GetPawn() : NULL), Hit, ImpulseDir);
 
 	// Discard the magnitude of the momentum and use Damage as the length instead
-	FVector DamageImpulse = ImpulseDir.SafeNormal() * DamageAmount * FoliageDamageImpulseScale;
+	FVector DamageImpulse = ImpulseDir.GetSafeNormal() * DamageAmount * FoliageDamageImpulseScale;
 
 	// Apply force magnitude clamps
 	DamageImpulse.X = FMath::Clamp(DamageImpulse.X, -MaxDamageImpulse, MaxDamageImpulse);
@@ -108,7 +108,7 @@ void AInteractiveFoliageActor::CapsuleTouched(AActor* Other, UPrimitiveComponent
 			const FVector CenterToTouching = FVector(TouchingActorCapsule->Bounds.Origin.X, TouchingActorCapsule->Bounds.Origin.Y, CapsuleComponent->Bounds.Origin.Z) - CapsuleComponent->Bounds.Origin;
 			// Keep track of the first position on the collision cylinder that the touching actor intersected
 			//@todo - need to handle multiple touching actors
-			TouchingActorEntryPosition = GetRootComponent()->Bounds.Origin + CenterToTouching.SafeNormal() * CapsuleComponent->GetScaledCapsuleRadius();
+			TouchingActorEntryPosition = GetRootComponent()->Bounds.Origin + CenterToTouching.GetSafeNormal() * CapsuleComponent->GetScaledCapsuleRadius();
 		}
 		// Bring this actor out of stasis so that it gets ticked now that a force has been applied
 		SetActorTickEnabled(true);
@@ -164,11 +164,11 @@ void AInteractiveFoliageActor::Tick(float DeltaSeconds)
 					// Which prevents strange movement that results from just comparing cylinder centers.
 					const FVector ProjectedTouchingActorPosition = (TouchingActorPosition - OppositeTouchingEntryPosition).ProjectOnTo(TouchingActorEntryPosition - OppositeTouchingEntryPosition) + OppositeTouchingEntryPosition;
 					// Find the furthest position on the cylinder of the touching actor from OppositeTouchingEntryPosition
-					const FVector TouchingActorFurthestPosition = ProjectedTouchingActorPosition + (TouchingActorEntryPosition - OppositeTouchingEntryPosition).SafeNormal() * TouchingActorCapsule->GetScaledCapsuleRadius();
+					const FVector TouchingActorFurthestPosition = ProjectedTouchingActorPosition + (TouchingActorEntryPosition - OppositeTouchingEntryPosition).GetSafeNormal() * TouchingActorCapsule->GetScaledCapsuleRadius();
 					// Construct the impulse as the distance between the furthest cylinder positions minus the two cylinder's diameters
 					const FVector ImpulseDirection = 
 						- (OppositeTouchingEntryPosition - TouchingActorFurthestPosition 
-						- (OppositeTouchingEntryPosition - TouchingActorFurthestPosition).SafeNormal() * 2.0f * (TouchingActorCapsule->GetScaledCapsuleRadius() + CapsuleComponent->GetScaledCapsuleRadius()));
+						- (OppositeTouchingEntryPosition - TouchingActorFurthestPosition).GetSafeNormal() * 2.0f * (TouchingActorCapsule->GetScaledCapsuleRadius() + CapsuleComponent->GetScaledCapsuleRadius()));
 
 					//DrawDebugLine(GetWorld(), GetRootComponent()->Bounds.Origin + FVector(0,0,100), GetRootComponent()->Bounds.Origin + ImpulseDirection + FVector(0,0,100), 100, 255, 100, false);
 
@@ -186,7 +186,7 @@ void AInteractiveFoliageActor::Tick(float DeltaSeconds)
 		FoliageForce += -FoliageStiffness * FoliagePosition;
 		// Apply spring quadratic stiffness, which increases in magnitude with the square of the distance to the origin
 		// This prevents the spring from being displaced too much by touch and damage forces
-		FoliageForce += -FoliageStiffnessQuadratic * FoliagePosition.SizeSquared() * FoliagePosition.SafeNormal();
+		FoliageForce += -FoliageStiffnessQuadratic * FoliagePosition.SizeSquared() * FoliagePosition.GetSafeNormal();
 		// Apply spring damping, which is like air resistance and causes the spring to lose energy over time
 		FoliageForce += -FoliageDamping * FoliageVelocity;
 
@@ -212,7 +212,7 @@ void AInteractiveFoliageActor::Tick(float DeltaSeconds)
 		const float RotationAngle = -FMath::Asin(FoliagePosition.Size() / IntersectionHeight);
 		// Use a rotation angle perpendicular to the impulse direction and the z axis
 		const FVector NormalizedRotationAxis = FoliagePosition.SizeSquared() > KINDA_SMALL_NUMBER ? 
-			(FoliagePosition ^ FVector(0,0,1)).SafeNormal() :
+			(FoliagePosition ^ FVector(0,0,1)).GetSafeNormal() :
 			FVector(0,0,1);
 
 		// Propagate the new rotation axis and angle to the rendering thread
diff --git a/Engine/Source/Runtime/Engine/Private/GameplayStatics.cpp b/Engine/Source/Runtime/Engine/Private/GameplayStatics.cpp
--- a/Engine/Source/Runtime/Engine/Private/GameplayStatics.cpp
+++ b/Engine/Source/Runtime/Engine/Private/GameplayStatics.cpp
@@ -189,7 +189,7 @@ static bool ComponentIsDamageableFrom(UPrimitiveComponent* VictimComp, FVector c
 	// didn't hit anything, assume nothing blocking the damage and victim is consequently visible
 	// but since we don't have a hit result to pass back, construct a simple one, modeling the damage as having hit a point at the component's center.
 	FVector const FakeHitLoc = VictimComp->GetComponentLocation();
-	FVector const FakeHitNorm = (Origin - FakeHitLoc).SafeNormal();		// normal points back toward the epicenter
+	FVector const FakeHitNorm = (Origin - FakeHitLoc).GetSafeNormal();		// normal points back toward the epicenter
 	OutHitResult = FHitResult(VictimComp->GetOwner(), VictimComp, FakeHitLoc, FakeHitNorm);
 	return true;
 }
@@ -1132,7 +1132,7 @@ bool UGameplayStatics::BlueprintSuggestProjectileVelocity(UObject* WorldContextO
 bool UGameplayStatics::SuggestProjectileVelocity(UObject* WorldContextObject, FVector& OutTossVelocity, FVector Start, FVector End, float TossSpeed, bool bFavorHighArc, float CollisionRadius, float OverrideGravityZ, ESuggestProjVelocityTraceOption::Type TraceOption, const FCollisionResponseParams& ResponseParam, const TArray<AActor*>& ActorsToIgnore, bool bDrawDebug)
 {
 	const FVector FlightDelta = End - Start;
-	const FVector DirXY = FlightDelta.SafeNormal2D();
+	const FVector DirXY = FlightDelta.GetSafeNormal2D();
 	const float DeltaXY = FlightDelta.Size2D();
 
 	const float DeltaZ = FlightDelta.Z;
diff --git a/Engine/Source/Runtime/Engine/Private/Interpolation.cpp b/Engine/Source/Runtime/Engine/Private/Interpolation.cpp
--- a/Engine/Source/Runtime/Engine/Private/Interpolation.cpp
+++ b/Engine/Source/Runtime/Engine/Private/Interpolation.cpp
@@ -4620,7 +4620,7 @@ FRotator UInterpTrackMove::GetLookAtRotation(UInterpTrackInst* TrInst)
 			}
 
 			// Find Rotator that points at LookAtActor
-			FVector LookDir = (LookAtActor->GetActorLocation() - Actor->GetActorLocation()).SafeNormal();
+			FVector LookDir = (LookAtActor->GetActorLocation() - Actor->GetActorLocation()).GetSafeNormal();
 			LookAtRot = LookDir.Rotation();
 		}
 	}
diff --git a/Engine/Source/Runtime/Engine/Private/KismetMathLibrary.cpp b/Engine/Source/Runtime/Engine/Private/KismetMathLibrary.cpp
--- a/Engine/Source/Runtime/Engine/Private/KismetMathLibrary.cpp
+++ b/Engine/Source/Runtime/Engine/Private/KismetMathLibrary.cpp
@@ -690,7 +690,7 @@ FVector UKismetMathLibrary::GreaterGreater_VectorRotator(FVector A, FRotator B)
 
 FVector  UKismetMathLibrary::RotateAngleAxis(FVector InVect, float AngleDeg, FVector Axis)
 {
-	return InVect.RotateAngleAxis(AngleDeg, Axis.SafeNormal());
+	return InVect.RotateAngleAxis(AngleDeg, Axis.GetSafeNormal());
 }
 
 bool UKismetMathLibrary::EqualEqual_VectorVector(FVector A, FVector B, float ErrorTolerance)
@@ -735,12 +735,12 @@ float UKismetMathLibrary::VSize2DSquared(FVector2D A)
 
 FVector UKismetMathLibrary::Normal(FVector A)
 {
-	return A.SafeNormal();
+	return A.GetSafeNormal();
 }
 
 FVector2D UKismetMathLibrary::Normal2D(FVector2D A)
 {
-	return A.SafeNormal();
+	return A.GetSafeNormal();
 }
 
 FVector UKismetMathLibrary::VLerp(FVector A, FVector B, float V)
@@ -833,7 +833,7 @@ FVector UKismetMathLibrary::NegateVector(FVector A)
 
 FVector UKismetMathLibrary::ClampVectorSize(FVector A, float Min, float Max)
 {
-	return A.ClampSize(Min, Max);
+	return A.GetClampedToSize(Min, Max);
 }
 
 float UKismetMathLibrary::GetMinElement(FVector A)
@@ -867,7 +867,7 @@ FVector UKismetMathLibrary::GetVectorArrayAverage(const TArray<FVector>& Vectors
 /** Find the unit direction vector from one position to another. */
 FVector UKismetMathLibrary::GetDirectionVector(FVector From, FVector To)
 {
-	return (To - From).SafeNormal();
+	return (To - From).GetSafeNormal();
 }
 
 
@@ -939,7 +939,7 @@ FRotator UKismetMathLibrary::NormalizedDeltaRotator(FRotator A, FRotator B)
 
 FRotator UKismetMathLibrary::RotatorFromAxisAndAngle(FVector Axis, float Angle)
 {
-	FVector SafeAxis = Axis.SafeNormal(); // Make sure axis is unit length
+	FVector SafeAxis = Axis.GetSafeNormal(); // Make sure axis is unit length
 	return FQuat(SafeAxis, FMath::DegreesToRadians(Angle)).Rotator();
 }
 
@@ -1880,7 +1880,7 @@ void UKismetMathLibrary::MinimumAreaRectangle(class UObject* WorldContextObject,
 	// Minimum area rectangle as computed by http://www.geometrictools.com/Documentation/MinimumAreaRectangle.pdf
 	for( int32 Idx = 1; Idx < PolyVertIndices.Num() - 1; ++Idx )
 	{
-		SupportVectorA = (TransformedVerts[PolyVertIndices[Idx]] - TransformedVerts[PolyVertIndices[Idx-1]]).SafeNormal();
+		SupportVectorA = (TransformedVerts[PolyVertIndices[Idx]] - TransformedVerts[PolyVertIndices[Idx-1]]).GetSafeNormal();
 		SupportVectorA.Z = 0.f;
 		SupportVectorB.X = -SupportVectorA.Y;
 		SupportVectorB.Y = SupportVectorA.X;
diff --git a/Engine/Source/Runtime/Engine/Private/Level.cpp b/Engine/Source/Runtime/Engine/Private/Level.cpp
--- a/Engine/Source/Runtime/Engine/Private/Level.cpp
+++ b/Engine/Source/Runtime/Engine/Private/Level.cpp
@@ -995,9 +995,9 @@ void ULevel::CreateModelComponents()
 					const FVector& TextureY = Model->Vectors[NodeSurf.vTextureV];
 					const int32 BaseVertexIndex = NodeGroup->Vertices.Num();
 					// Compute the surface's tangent basis.
-					FVector NodeTangentX = Model->Vectors[NodeSurf.vTextureU].SafeNormal();
-					FVector NodeTangentY = Model->Vectors[NodeSurf.vTextureV].SafeNormal();
-					FVector NodeTangentZ = Model->Vectors[NodeSurf.vNormal].SafeNormal();
+					FVector NodeTangentX = Model->Vectors[NodeSurf.vTextureU].GetSafeNormal();
+					FVector NodeTangentY = Model->Vectors[NodeSurf.vTextureV].GetSafeNormal();
+					FVector NodeTangentZ = Model->Vectors[NodeSurf.vNormal].GetSafeNormal();
 
 					// Generate the node's vertices.
 					for(uint32 VertexIndex = 0;VertexIndex < Node.NumVertices;VertexIndex++)
diff --git a/Engine/Source/Runtime/Engine/Private/MaterialBillboardComponent.cpp b/Engine/Source/Runtime/Engine/Private/MaterialBillboardComponent.cpp
--- a/Engine/Source/Runtime/Engine/Private/MaterialBillboardComponent.cpp
+++ b/Engine/Source/Runtime/Engine/Private/MaterialBillboardComponent.cpp
@@ -180,8 +180,8 @@ public:
 						for(uint32 VertexIndex = 0;VertexIndex < 4;++VertexIndex)
 						{
 							VertexArray.Vertices[VertexIndex].Color = Color;
-							VertexArray.Vertices[VertexIndex].TangentX = FPackedNormal(LocalCameraRight.SafeNormal());
-							VertexArray.Vertices[VertexIndex].TangentZ = FPackedNormal(-LocalCameraForward.SafeNormal());
+							VertexArray.Vertices[VertexIndex].TangentX = FPackedNormal(LocalCameraRight.GetSafeNormal());
+							VertexArray.Vertices[VertexIndex].TangentZ = FPackedNormal(-LocalCameraForward.GetSafeNormal());
 						}
 
 						// Set up the sprite vertex positions and texture coordinates.
diff --git a/Engine/Source/Runtime/Engine/Private/ModelLight.cpp b/Engine/Source/Runtime/Engine/Private/ModelLight.cpp
--- a/Engine/Source/Runtime/Engine/Private/ModelLight.cpp
+++ b/Engine/Source/Runtime/Engine/Private/ModelLight.cpp
@@ -70,7 +70,7 @@ static bool SphereOnNode(UModel* Model,uint32 NodeIndex,FVector Point,float Radi
 			continue;
 
 		// If point is not behind all the planes created by this polys edges, it's outside the poly.
-		if(FVector::PointPlaneDist(Point,Model->Points[Model->Verts[Node.iVertPool + VertexIndex].pVertex],EdgeNormal.SafeNormal()) > Radius)
+		if(FVector::PointPlaneDist(Point,Model->Points[Model->Verts[Node.iVertPool + VertexIndex].pVertex],EdgeNormal.GetSafeNormal()) > Radius)
 			return 0;
 	}
 
diff --git a/Engine/Source/Runtime/Engine/Private/MovementComponent.cpp b/Engine/Source/Runtime/Engine/Private/MovementComponent.cpp
--- a/Engine/Source/Runtime/Engine/Private/MovementComponent.cpp
+++ b/Engine/Source/Runtime/Engine/Private/MovementComponent.cpp
@@ -120,7 +120,7 @@ void UMovementComponent::OnRegister()
 		SetPlaneConstraintAxisSetting(PlaneConstraintAxisSetting);
 	}
 
-	PlaneConstraintNormal = PlaneConstraintNormal.SafeNormal();
+	PlaneConstraintNormal = PlaneConstraintNormal.GetSafeNormal();
 
 	if (bSnapToPlaneAtStart)
 	{
@@ -348,13 +348,13 @@ void UMovementComponent::SetPlaneConstraintAxisSetting(EPlaneConstraintAxisSetti
 
 void UMovementComponent::SetPlaneConstraintNormal(FVector PlaneNormal)
 {
-	PlaneConstraintNormal = PlaneNormal.SafeNormal();
+	PlaneConstraintNormal = PlaneNormal.GetSafeNormal();
 	PlaneConstraintAxisSetting = EPlaneConstraintAxisSetting::Custom;
 }
 
 void UMovementComponent::SetPlaneConstraintFromVectors(FVector Forward, FVector Up)
 {
-	PlaneConstraintNormal = (Up ^ Forward).SafeNormal();
+	PlaneConstraintNormal = (Up ^ Forward).GetSafeNormal();
 }
 
 void UMovementComponent::SetPlaneConstraintOrigin(FVector PlaneOrigin)
@@ -398,7 +398,7 @@ FVector UMovementComponent::ConstrainNormalToPlane(FVector Normal) const
 {
 	if (bConstrainToPlane)
 	{
-		Normal = FVector::VectorPlaneProject(Normal, PlaneConstraintNormal).SafeNormal();
+		Normal = FVector::VectorPlaneProject(Normal, PlaneConstraintNormal).GetSafeNormal();
 	}
 
 	return Normal;
@@ -621,7 +621,7 @@ void UMovementComponent::TwoWallAdjust(FVector& OutDelta, const FHitResult& Hit,
 	{
 		const FVector DesiredDir = Delta;
 		FVector NewDir = (HitNormal ^ OldHitNormal);
-		NewDir = NewDir.SafeNormal();
+		NewDir = NewDir.GetSafeNormal();
 		Delta = (Delta | NewDir) * (1.f - Hit.Time) * NewDir;
 		if ((DesiredDir | Delta) < 0.f)
 		{
diff --git a/Engine/Source/Runtime/Engine/Private/Particles/ParticleModules_Collision.cpp b/Engine/Source/Runtime/Engine/Private/Particles/ParticleModules_Collision.cpp
--- a/Engine/Source/Runtime/Engine/Private/Particles/ParticleModules_Collision.cpp
+++ b/Engine/Source/Runtime/Engine/Private/Particles/ParticleModules_Collision.cpp
@@ -346,7 +346,7 @@ void UParticleModuleCollision::Update(FParticleEmitterInstance* Owner, int32 Off
 		{
 			OldLocation	= Particle.OldLocation;
 		}
-		FVector	Direction = (Location - OldLocation).SafeNormal();
+		FVector	Direction = (Location - OldLocation).GetSafeNormal();
 
 		// Determine the size
 		FVector Size = Particle.Size * ParentScale;
diff --git a/Engine/Source/Runtime/Engine/Private/Particles/ParticleModules_Location.cpp b/Engine/Source/Runtime/Engine/Private/Particles/ParticleModules_Location.cpp
--- a/Engine/Source/Runtime/Engine/Private/Particles/ParticleModules_Location.cpp
+++ b/Engine/Source/Runtime/Engine/Private/Particles/ParticleModules_Location.cpp
@@ -1083,7 +1083,7 @@ void UParticleModuleLocationPrimitiveCylinder::SpawnEx(FParticleEmitterInstance*
 	vNormalizedDir.Normalize();
 
 	FVector2D vUnitDir2D(vUnitDir[RadialIndex0], vUnitDir[RadialIndex1]);
-	FVector2D vNormalizedDir2D = vUnitDir2D.SafeNormal();
+	FVector2D vNormalizedDir2D = vUnitDir2D.GetSafeNormal();
 
 	// Determine the position
 	// Always want Z in the [-Height, Height] range
@@ -1207,9 +1207,9 @@ void UParticleModuleLocationPrimitiveCylinder::Render3DPreview(FParticleEmitterI
 	FVector	TransformedAxis[3];
 	FVector	Axis[3];
 
-	TransformedAxis[0] = LocalToWorld.TransformVector(FVector(1.0f, 0.0f, 0.0f)).SafeNormal();
-	TransformedAxis[1] = LocalToWorld.TransformVector(FVector(0.0f, 1.0f, 0.0f)).SafeNormal();
-	TransformedAxis[2] = LocalToWorld.TransformVector(FVector(0.0f, 0.0f, 1.0f)).SafeNormal();
+	TransformedAxis[0] = LocalToWorld.TransformVector(FVector(1.0f, 0.0f, 0.0f)).GetSafeNormal();
+	TransformedAxis[1] = LocalToWorld.TransformVector(FVector(0.0f, 1.0f, 0.0f)).GetSafeNormal();
+	TransformedAxis[2] = LocalToWorld.TransformVector(FVector(0.0f, 0.0f, 1.0f)).GetSafeNormal();
 
 	switch (HeightAxis)
 	{
diff --git a/Engine/Source/Runtime/Engine/Private/Particles/ParticleModules_Velocity.cpp b/Engine/Source/Runtime/Engine/Private/Particles/ParticleModules_Velocity.cpp
--- a/Engine/Source/Runtime/Engine/Private/Particles/ParticleModules_Velocity.cpp
+++ b/Engine/Source/Runtime/Engine/Private/Particles/ParticleModules_Velocity.cpp
@@ -93,7 +93,7 @@ void UParticleModuleVelocity::SpawnEx(FParticleEmitterInstance* Owner, int32 Off
 	SPAWN_INIT;
 	{
 		FVector Vel = StartVelocity.GetValue(Owner->EmitterTime, Owner->Component, 0, InRandomStream);
-		FVector FromOrigin = (Particle.Location - Owner->EmitterToSimulation.GetOrigin()).SafeNormal();
+		FVector FromOrigin = (Particle.Location - Owner->EmitterToSimulation.GetOrigin()).GetSafeNormal();
 
 		FVector OwnerScale(1.0f);
 		if ((bApplyOwnerScale == true) && Owner && Owner->Component)
@@ -476,7 +476,7 @@ void UParticleModuleVelocityCone::SpawnEx(FParticleEmitterInstance* Owner, int32
 		const FVector DefaultSpawnDirection = DefaultDirectionRotation.TransformVector(DefaultDirection);
 
 		// Orientate the cone along the direction vector		
-		const FVector ForwardDirection = (Direction != FVector::ZeroVector)? Direction.SafeNormal(): DefaultDirection;
+		const FVector ForwardDirection = (Direction != FVector::ZeroVector)? Direction.GetSafeNormal(): DefaultDirection;
 		FVector UpDirection(0.0f, 0.0f, 1.0f);
 		FVector RightDirection(1.0f, 0.0f, 0.0f);
 
@@ -493,8 +493,8 @@ void UParticleModuleVelocityCone::SpawnEx(FParticleEmitterInstance* Owner, int32
 
 		FMatrix DirectionRotation;
 		DirectionRotation.SetIdentity();
-		DirectionRotation.SetAxis(0, RightDirection.SafeNormal());
-		DirectionRotation.SetAxis(1, UpDirection.SafeNormal());
+		DirectionRotation.SetAxis(0, RightDirection.GetSafeNormal());
+		DirectionRotation.SetAxis(1, UpDirection.GetSafeNormal());
 		DirectionRotation.SetAxis(2, ForwardDirection);
 		FVector SpawnDirection = DirectionRotation.TransformVector(DefaultSpawnDirection);
 	
@@ -543,7 +543,7 @@ void UParticleModuleVelocityCone::Render3DPreview(FParticleEmitterInstance* Owne
 
 	// Calculate direction transform
 	const FVector DefaultDirection(0.0f, 0.0f, 1.0f);
-	const FVector ForwardDirection = (Direction != FVector::ZeroVector)? Direction.SafeNormal(): DefaultDirection;
+	const FVector ForwardDirection = (Direction != FVector::ZeroVector)? Direction.GetSafeNormal(): DefaultDirection;
 	FVector UpDirection(0.0f, 0.0f, 1.0f);
 	FVector RightDirection(1.0f, 0.0f, 0.0f);
 
@@ -560,8 +560,8 @@ void UParticleModuleVelocityCone::Render3DPreview(FParticleEmitterInstance* Owne
 
 	FMatrix DirectionRotation;
 	DirectionRotation.SetIdentity();
-	DirectionRotation.SetAxis(0, RightDirection.SafeNormal());
-	DirectionRotation.SetAxis(1, UpDirection.SafeNormal());
+	DirectionRotation.SetAxis(0, RightDirection.GetSafeNormal());
+	DirectionRotation.SetAxis(1, UpDirection.GetSafeNormal());
 	DirectionRotation.SetAxis(2, ForwardDirection);
 
 	// Calculate the owning actor's scale and rotation
diff --git a/Engine/Source/Runtime/Engine/Private/Particles/ParticleSystemRender.cpp b/Engine/Source/Runtime/Engine/Private/Particles/ParticleSystemRender.cpp
--- a/Engine/Source/Runtime/Engine/Private/Particles/ParticleSystemRender.cpp
+++ b/Engine/Source/Runtime/Engine/Private/Particles/ParticleSystemRender.cpp
@@ -2194,8 +2194,8 @@ FParticleBeamTrailUniformBufferRef CreateBeamTrailUniformBuffer(
 	const EParticleAxisLock LockAxisFlag = (EParticleAxisLock)SourceData->LockAxisFlag;
 	if (LockAxisFlag == EPAL_NONE)
 	{
-		CameraUp	= -View->InvViewProjectionMatrix.TransformVector(FVector(1.0f,0.0f,0.0f)).SafeNormal();
-		CameraRight	= -View->InvViewProjectionMatrix.TransformVector(FVector(0.0f,1.0f,0.0f)).SafeNormal();
+		CameraUp	= -View->InvViewProjectionMatrix.TransformVector(FVector(1.0f,0.0f,0.0f)).GetSafeNormal();
+		CameraRight	= -View->InvViewProjectionMatrix.TransformVector(FVector(0.0f,1.0f,0.0f)).GetSafeNormal();
 	}
 	else
 	{
diff --git a/Engine/Source/Runtime/Engine/Private/PhysicsEngine/BodyInstance.cpp b/Engine/Source/Runtime/Engine/Private/PhysicsEngine/BodyInstance.cpp
--- a/Engine/Source/Runtime/Engine/Private/PhysicsEngine/BodyInstance.cpp
+++ b/Engine/Source/Runtime/Engine/Private/PhysicsEngine/BodyInstance.cpp
@@ -613,7 +613,7 @@ void FBodyInstance::CreateDOFLock()
 		DOFConstraint->LinearYMotion = ELinearConstraintMotion::LCM_Free;
 		DOFConstraint->LinearZMotion = ELinearConstraintMotion::LCM_Free;
 
-		FVector Normal = LockedAxis.SafeNormal();
+		FVector Normal = LockedAxis.GetSafeNormal();
 		FVector Sec;
 		FVector Garbage;
 		Normal.FindBestAxisVectors(Garbage, Sec);
diff --git a/Engine/Source/Runtime/Engine/Private/PhysicsEngine/PhysDrawing.cpp b/Engine/Source/Runtime/Engine/Private/PhysicsEngine/PhysDrawing.cpp
--- a/Engine/Source/Runtime/Engine/Private/PhysicsEngine/PhysDrawing.cpp
+++ b/Engine/Source/Runtime/Engine/Private/PhysicsEngine/PhysDrawing.cpp
@@ -429,9 +429,9 @@ void FKConvexElem::AddCachedSolidConvexGeom(TArray<FDynamicMeshVertex>& VertexBu
 			const PxU8* indices = PIndexBuffer + Data.mIndexBase;
 
 			// create tangents from the first and second vertices of each polygon
-			const FVector TangentX = P2UVector(PVertices[indices[1]]-PVertices[indices[0]]).SafeNormal();
-			const FVector TangentZ = FVector(Data.mPlane[0], Data.mPlane[1], Data.mPlane[2]).SafeNormal();
-			const FVector TangentY = (TangentX ^ TangentZ).SafeNormal();
+			const FVector TangentX = P2UVector(PVertices[indices[1]]-PVertices[indices[0]]).GetSafeNormal();
+			const FVector TangentZ = FVector(Data.mPlane[0], Data.mPlane[1], Data.mPlane[2]).GetSafeNormal();
+			const FVector TangentY = (TangentX ^ TangentZ).GetSafeNormal();
 
 			// add vertices 
 			for(PxU32 j=0;j<Data.mNbVerts;j++)
diff --git a/Engine/Source/Runtime/Engine/Private/PointLightSceneProxy.h b/Engine/Source/Runtime/Engine/Private/PointLightSceneProxy.h
--- a/Engine/Source/Runtime/Engine/Private/PointLightSceneProxy.h
+++ b/Engine/Source/Runtime/Engine/Private/PointLightSceneProxy.h
@@ -230,7 +230,7 @@ public:
 		if (LightDistance <= SubjectRadius * ShadowRadiusMultiplier)
 		{
 			// Make the primitive fit in a single < 90 degree FOV projection.
-			LightVector = SubjectRadius * LightVector.SafeNormal() * ShadowRadiusMultiplier;
+			LightVector = SubjectRadius * LightVector.GetSafeNormal() * ShadowRadiusMultiplier;
 			LightPosition = (SubjectBounds.Origin - LightVector );
 			LightDistance = SubjectRadius * ShadowRadiusMultiplier;
 			SilhouetteRadius = 1.0f;
diff --git a/Engine/Source/Runtime/Engine/Private/Polygon.cpp b/Engine/Source/Runtime/Engine/Private/Polygon.cpp
--- a/Engine/Source/Runtime/Engine/Private/Polygon.cpp
+++ b/Engine/Source/Runtime/Engine/Private/Polygon.cpp
@@ -459,7 +459,7 @@ void FPoly::Transform
 
 	// Transform normal.  Since the transformation coordinate system is
 	// orthogonal but not orthonormal, it has to be renormalized here.
-	Normal = Normal.SafeNormal();
+	Normal = Normal.GetSafeNormal();
 
 }
 
@@ -484,7 +484,7 @@ void FPoly::Rotate
 
 	// Rotate the normal.
 	Normal = FRotationMatrix(Rotation).TransformVector(Normal);
-	Normal = Normal.SafeNormal();
+	Normal = Normal.GetSafeNormal();
 }
 
 
@@ -507,7 +507,7 @@ void FPoly::Scale
 	TextureV *= Scale;
 
 	// Renormalize the normal.
-	Normal = Normal.SafeNormal();
+	Normal = Normal.GetSafeNormal();
 }
 
 
@@ -839,7 +839,7 @@ FPoly FPoly::BuildAndCutInfiniteFPoly(const FPlane& InPlane, const TArray<FPlane
 	{
 		const FPlane* Plane = &InCutPlanes[p];
 
-		result = PolyMerged.SplitWithPlane( Plane->SafeNormal() * Plane->W, Plane->SafeNormal(), &Front, &Back, 1 );
+		result = PolyMerged.SplitWithPlane( Plane->GetSafeNormal() * Plane->W, Plane->GetSafeNormal(), &Front, &Back, 1 );
 
 		if( result == SP_Split )
 		{
@@ -933,8 +933,8 @@ int32 FPoly::Finalize( ABrush* InOwner, int32 NoError )
 	{
 		for( int32 i=1; i<Vertices.Num(); i++ )
 		{
-			TextureU = ((Vertices[0] - Vertices[i]) ^ Normal).SafeNormal();
-			TextureV = (Normal ^ TextureU).SafeNormal();
+			TextureU = ((Vertices[0] - Vertices[i]) ^ Normal).GetSafeNormal();
+			TextureV = (Normal ^ TextureU).GetSafeNormal();
 			if( TextureU.SizeSquared()!=0 && TextureV.SizeSquared()!=0 )
 				break;
 		}
diff --git a/Engine/Source/Runtime/Engine/Private/PrimitiveComponent.cpp b/Engine/Source/Runtime/Engine/Private/PrimitiveComponent.cpp
--- a/Engine/Source/Runtime/Engine/Private/PrimitiveComponent.cpp
+++ b/Engine/Source/Runtime/Engine/Private/PrimitiveComponent.cpp
@@ -640,7 +640,7 @@ void UPrimitiveComponent::ReceiveComponentDamage(float DamageAmount, FDamageEven
 		FPointDamageEvent* const PointDamageEvent = (FPointDamageEvent*) &DamageEvent;
 		if((DamageTypeCDO->DamageImpulse > 0.f) && !PointDamageEvent->ShotDirection.IsNearlyZero())
 		{
-			FVector const ImpulseToApply = PointDamageEvent->ShotDirection.SafeNormal() * DamageTypeCDO->DamageImpulse;
+			FVector const ImpulseToApply = PointDamageEvent->ShotDirection.GetSafeNormal() * DamageTypeCDO->DamageImpulse;
 			AddImpulseAtLocation( ImpulseToApply, PointDamageEvent->HitInfo.ImpactPoint, PointDamageEvent->HitInfo.BoneName );
 		}
 	}
@@ -1178,7 +1178,7 @@ static bool ShouldIgnoreHitResult(const UWorld* InWorld, FHitResult const& TestH
 			const float DotTolerance = CVarInitialOverlapTolerance.GetValueOnGameThread();
 
 			// Dot product of movement direction against 'exit' direction
-			const FVector MovementDir = MovementDirDenormalized.SafeNormal();
+			const FVector MovementDir = MovementDirDenormalized.GetSafeNormal();
 			const float MoveDot = (TestHit.ImpactNormal | MovementDir);
 
 			const bool bMovingOut = MoveDot > DotTolerance;
diff --git a/Engine/Source/Runtime/Engine/Private/PrimitiveComponentPhysics.cpp b/Engine/Source/Runtime/Engine/Private/PrimitiveComponentPhysics.cpp
--- a/Engine/Source/Runtime/Engine/Private/PrimitiveComponentPhysics.cpp
+++ b/Engine/Source/Runtime/Engine/Private/PrimitiveComponentPhysics.cpp
@@ -72,7 +72,7 @@ bool UPrimitiveComponent::ApplyRigidBodyState(const FRigidBodyState& NewState, c
 		if (FMath::Abs(DeltaAng) < ErrorCorrection.AngularDeltaThreshold )
 		{
 			UpdatedQuat = FMath::Lerp(CurrentState.Quaternion, NewState.Quaternion, ErrorCorrection.AngularInterpAlpha);
-			FixAngVel = DeltaAxis.SafeNormal() * FMath::RadiansToDegrees(DeltaAng) * (1.f - ErrorCorrection.AngularInterpAlpha) * ErrorCorrection.AngularRecipFixTime;
+			FixAngVel = DeltaAxis.GetSafeNormal() * FMath::RadiansToDegrees(DeltaAng) * (1.f - ErrorCorrection.AngularInterpAlpha) * ErrorCorrection.AngularRecipFixTime;
 		}
 
 		/////// BODY UPDATE ///////
diff --git a/Engine/Source/Runtime/Engine/Private/PrimitiveDrawingUtils.cpp b/Engine/Source/Runtime/Engine/Private/PrimitiveDrawingUtils.cpp
--- a/Engine/Source/Runtime/Engine/Private/PrimitiveDrawingUtils.cpp
+++ b/Engine/Source/Runtime/Engine/Private/PrimitiveDrawingUtils.cpp
@@ -425,14 +425,14 @@ void BuildConeVerts(float Angle1, float Angle2, float Scale, float XOffset, int3
 		V1.TextureCoordinate.X = 1.0f;
 		V1.TextureCoordinate.Y = (float)i / NumSides;
 		FVector TriTangentZPrev = ConeVerts[i] ^ ConeVerts[i == 0 ? NumSides - 1 : i - 1]; // Normal of the previous face connected to this face
-		V1.SetTangents(TriTangentX, TriTangentY, (TriTangentZPrev + TriTangentZ).SafeNormal());
+		V1.SetTangents(TriTangentX, TriTangentY, (TriTangentZPrev + TriTangentZ).GetSafeNormal());
 		int32 I1 = OutVerts.Add(V1);
 
 		V2.Position = ConeVerts[(i + 1) % NumSides];
 		V2.TextureCoordinate.X = 1.0f;
 		V2.TextureCoordinate.Y = (float)((i + 1) % NumSides) / NumSides;
 		FVector TriTangentZNext = ConeVerts[(i + 2) % NumSides] ^ ConeVerts[(i + 1) % NumSides]; // Normal of the next face connected to this face
-		V2.SetTangents(TriTangentX, TriTangentY, (TriTangentZNext + TriTangentZ).SafeNormal());
+		V2.SetTangents(TriTangentX, TriTangentY, (TriTangentZNext + TriTangentZ).GetSafeNormal());
 		int32 I2 = OutVerts.Add(V2);
 
 		// Flip winding for negative scale
@@ -1004,9 +1004,9 @@ static void DrawHalfCircle(FPrimitiveDrawInterface* PDI, const FVector& Base, co
 void DrawWireCapsule(FPrimitiveDrawInterface* PDI, const FVector& Base, const FVector& X, const FVector& Y, const FVector& Z, const FLinearColor& Color, float Radius, float HalfHeight, int32 NumSides, uint8 DepthPriority, float Thickness, float DepthBias, bool bScreenSpace)
 {
 	const FVector Origin = Base;
-	const FVector XAxis = X.SafeNormal();
-	const FVector YAxis = Y.SafeNormal();
-	const FVector ZAxis = Z.SafeNormal();
+	const FVector XAxis = X.GetSafeNormal();
+	const FVector YAxis = Y.GetSafeNormal();
+	const FVector ZAxis = Z.GetSafeNormal();
 
 	// because we are drawing a capsule we have to have room for the "domed caps"
 	const float XScale = X.Size();
diff --git a/Engine/Source/Runtime/Engine/Private/ProjectileMovementComponent.cpp b/Engine/Source/Runtime/Engine/Private/ProjectileMovementComponent.cpp
--- a/Engine/Source/Runtime/Engine/Private/ProjectileMovementComponent.cpp
+++ b/Engine/Source/Runtime/Engine/Private/ProjectileMovementComponent.cpp
@@ -62,7 +62,7 @@ void UProjectileMovementComponent::InitializeComponent()
 		// InitialSpeed > 0 overrides initial velocity magnitude.
 		if (InitialSpeed > 0.f)
 		{
-			Velocity = Velocity.SafeNormal() * InitialSpeed;
+			Velocity = Velocity.GetSafeNormal() * InitialSpeed;
 		}
 
 		if (bInitialVelocityInLocalSpace)
@@ -182,7 +182,7 @@ void UProjectileMovementComponent::TickComponent(float DeltaTime, enum ELevelTic
 			// if velocity still into wall (after HandleHitWall() had a chance to adjust), slide along wall
 			const float DotTolerance = 0.01f;
 			bIsSliding = (bMultiHit && FVector::Coincident(PreviousHitNormal, Normal)) ||
-						 ((Velocity.SafeNormal() | Normal) <= DotTolerance);
+						 ((Velocity.GetSafeNormal() | Normal) <= DotTolerance);
 			
 			if (bIsSliding)
 			{
@@ -190,7 +190,7 @@ void UProjectileMovementComponent::TickComponent(float DeltaTime, enum ELevelTic
 				{
 					//90 degree or less corner, so use cross product for direction
 					FVector NewDir = (Normal ^ PreviousHitNormal);
-					NewDir = NewDir.SafeNormal();
+					NewDir = NewDir.GetSafeNormal();
 					Velocity = Velocity.ProjectOnToNormal(NewDir);
 					if ((OldVelocity | Velocity) < 0.f)
 					{
@@ -268,7 +268,7 @@ bool UProjectileMovementComponent::HandleSliding(FHitResult& Hit, float& SubTick
 			const FVector ProjectedForce = FVector::VectorPlaneProject(Force, OldHitNormal);
 			const FVector NewVelocity = Velocity + ProjectedForce;
 
-			const FVector FrictionForce = -NewVelocity.SafeNormal() * FMath::Min(-ForceDotN * Friction, NewVelocity.Size());
+			const FVector FrictionForce = -NewVelocity.GetSafeNormal() * FMath::Min(-ForceDotN * Friction, NewVelocity.Size());
 			Velocity = ConstrainDirectionToPlane(NewVelocity + FrictionForce);
 		}
 		else
@@ -319,7 +319,7 @@ FVector UProjectileMovementComponent::LimitVelocity(FVector NewVelocity) const
 	const float CurrentMaxSpeed = GetMaxSpeed();
 	if (CurrentMaxSpeed > 0.f)
 	{
-		NewVelocity = NewVelocity.ClampMaxSize(CurrentMaxSpeed);
+		NewVelocity = NewVelocity.GetClampedToMaxSize(CurrentMaxSpeed);
 	}
 
 	return ConstrainDirectionToPlane(NewVelocity);
@@ -357,7 +357,7 @@ FVector UProjectileMovementComponent::ComputeAcceleration(const FVector& InVeloc
 // Allow the projectile to track towards its homing target.
 FVector UProjectileMovementComponent::ComputeHomingAcceleration(const FVector& InVelocity, float DeltaTime) const
 {
-	FVector HomingAcceleration = ((HomingTargetComponent->GetComponentLocation() - UpdatedComponent->GetComponentLocation()).SafeNormal() * HomingAccelerationMagnitude);
+	FVector HomingAcceleration = ((HomingTargetComponent->GetComponentLocation() - UpdatedComponent->GetComponentLocation()).GetSafeNormal() * HomingAccelerationMagnitude);
 	return HomingAcceleration;
 }
 
diff --git a/Engine/Source/Runtime/Engine/Private/ReflectionCaptureComponent.cpp b/Engine/Source/Runtime/Engine/Private/ReflectionCaptureComponent.cpp
--- a/Engine/Source/Runtime/Engine/Private/ReflectionCaptureComponent.cpp
+++ b/Engine/Source/Runtime/Engine/Private/ReflectionCaptureComponent.cpp
@@ -1560,12 +1560,12 @@ void FReflectionCaptureProxy::SetTransform(const FMatrix& InTransform)
 	const FVector4 PlaneNormal = InTransform.TransformVector(ForwardVector);
 
 	// Normalize the plane
-	ReflectionPlane = FPlane(Position, FVector(PlaneNormal).SafeNormal());
+	ReflectionPlane = FPlane(Position, FVector(PlaneNormal).GetSafeNormal());
 	const FVector ReflectionXAxis = InTransform.TransformVector(RightVector);
 	const FVector ScaleVector = InTransform.GetScaleVector();
 	BoxScales = ScaleVector;
 	// Include the owner's draw scale in the axes
-	ReflectionXAxisAndYScale = ReflectionXAxis.SafeNormal() * ScaleVector.Y;
+	ReflectionXAxisAndYScale = ReflectionXAxis.GetSafeNormal() * ScaleVector.Y;
 	ReflectionXAxisAndYScale.W = ScaleVector.Y / ScaleVector.Z;
 }
 
diff --git a/Engine/Source/Runtime/Engine/Private/SceneView.cpp b/Engine/Source/Runtime/Engine/Private/SceneView.cpp
--- a/Engine/Source/Runtime/Engine/Private/SceneView.cpp
+++ b/Engine/Source/Runtime/Engine/Private/SceneView.cpp
@@ -318,12 +318,12 @@ FSceneView::FSceneView(const FSceneViewInitOptions& InitOptions)
 	else if (InitOptions.bUseFauxOrthoViewPos)
 	{
 		float DistanceToViewOrigin = WORLD_MAX;
-		ViewMatrices.ViewOrigin = FVector4(InvViewMatrix.TransformVector(FVector(0,0,-1)).SafeNormal()*DistanceToViewOrigin,1) + InvViewMatrix.GetOrigin();
+		ViewMatrices.ViewOrigin = FVector4(InvViewMatrix.TransformVector(FVector(0,0,-1)).GetSafeNormal()*DistanceToViewOrigin,1) + InvViewMatrix.GetOrigin();
 	}
 #endif
 	else
 	{
-		ViewMatrices.ViewOrigin = FVector4(InvViewMatrix.TransformVector(FVector(0,0,-1)).SafeNormal(),0);
+		ViewMatrices.ViewOrigin = FVector4(InvViewMatrix.TransformVector(FVector(0,0,-1)).GetSafeNormal(),0);
 		// to avoid issues with view dependent effect (e.g. Frensel)
 		ApplyPreViewTranslation = false;
 	}
@@ -616,7 +616,7 @@ void FSceneView::DeprojectScreenToWorld(const FVector2D& ScreenPos, const FIntRe
 		RayEndViewSpace /= HGRayEndViewSpace.W;
 	}
 	FVector RayDirViewSpace = RayEndViewSpace - RayStartViewSpace;
-	RayDirViewSpace = RayDirViewSpace.SafeNormal();
+	RayDirViewSpace = RayDirViewSpace.GetSafeNormal();
 
 	// The view transform does not have projection, so we can use the standard functions that deal with vectors and normals (normals
 	// are vectors that do not use the translational part of a rotation/translation)
@@ -626,7 +626,7 @@ void FSceneView::DeprojectScreenToWorld(const FVector2D& ScreenPos, const FIntRe
 	// Finally, store the results in the hitcheck inputs.  The start position is the eye, and the end position
 	// is the eye plus a long distance in the direction the mouse is pointing.
 	out_WorldOrigin = RayStartWorldSpace;
-	out_WorldDirection = RayDirWorldSpace.SafeNormal();
+	out_WorldDirection = RayDirWorldSpace.GetSafeNormal();
 }
 
 #define LERP_PP(NAME) if(Src.bOverride_ ## NAME)	Dest . NAME = FMath::Lerp(Dest . NAME, Src . NAME, Weight);
diff --git a/Engine/Source/Runtime/Engine/Private/SkeletalRenderCPUSkin.cpp b/Engine/Source/Runtime/Engine/Private/SkeletalRenderCPUSkin.cpp
--- a/Engine/Source/Runtime/Engine/Private/SkeletalRenderCPUSkin.cpp
+++ b/Engine/Source/Runtime/Engine/Private/SkeletalRenderCPUSkin.cpp
@@ -419,17 +419,17 @@ void FSkeletalMeshObjectCPUSkin::DrawVertexElements(FPrimitiveDrawInterface* PDI
 
 		if( bDrawNormals )
 		{
-			PDI->DrawLine( WorldPos, WorldPos+LocalToWorldInverseTranspose.TransformVector( (FVector)(Normal) ).SafeNormal() * Len, FLinearColor( 0.0f, 1.0f, 0.0f), SDPG_World );
+			PDI->DrawLine( WorldPos, WorldPos+LocalToWorldInverseTranspose.TransformVector( (FVector)(Normal) ).GetSafeNormal() * Len, FLinearColor( 0.0f, 1.0f, 0.0f), SDPG_World );
 		}
 
 		if( bDrawTangents )
 		{
-			PDI->DrawLine( WorldPos, WorldPos+LocalToWorldInverseTranspose.TransformVector( Tangent ).SafeNormal() * Len, FLinearColor( 1.0f, 0.0f, 0.0f), SDPG_World );
+			PDI->DrawLine( WorldPos, WorldPos+LocalToWorldInverseTranspose.TransformVector( Tangent ).GetSafeNormal() * Len, FLinearColor( 1.0f, 0.0f, 0.0f), SDPG_World );
 		}
 
 		if( bDrawBinormals )
 		{
-			PDI->DrawLine( WorldPos, WorldPos+LocalToWorldInverseTranspose.TransformVector( Binormal ).SafeNormal() * Len, FLinearColor( 0.0f, 0.0f, 1.0f), SDPG_World );
+			PDI->DrawLine( WorldPos, WorldPos+LocalToWorldInverseTranspose.TransformVector( Binormal ).GetSafeNormal() * Len, FLinearColor( 0.0f, 0.0f, 1.0f), SDPG_World );
 		}
 	}
 }
@@ -560,7 +560,7 @@ FORCEINLINE void ApplyMorphBlend( VertexType& DestVertex, const FVertexAnimDelta
 	// Save W before = operator. That overwrites W to be 127.
 	uint8 W = DestVertex.TangentZ.Vector.W;
 	// add normal offset. can only apply normal deltas up to a weight of 1
-	DestVertex.TangentZ = FVector(FVector(DestVertex.TangentZ) + SrcMorph.TangentZDelta * FMath::Min(Weight,1.0f)).UnsafeNormal();
+	DestVertex.TangentZ = FVector(FVector(DestVertex.TangentZ) + SrcMorph.TangentZDelta * FMath::Min(Weight,1.0f)).GetUnsafeNormal();
 	// Recover W
 	DestVertex.TangentZ.Vector.W = W;
 } 
diff --git a/Engine/Source/Runtime/Engine/Private/SoundAttenuation.cpp b/Engine/Source/Runtime/Engine/Private/SoundAttenuation.cpp
--- a/Engine/Source/Runtime/Engine/Private/SoundAttenuation.cpp
+++ b/Engine/Source/Runtime/Engine/Private/SoundAttenuation.cpp
@@ -161,7 +161,7 @@ float FAttenuationSettings::AttenuationEvalCone(const FTransform& SoundTransform
 
 	if (VolumeMultiplier > 0.f)
 	{
-		const float theta = FMath::RadiansToDegrees(fabsf(FMath::Acos( FVector::DotProduct(SoundForward, (ListenerLocation - Origin).SafeNormal()))));
+		const float theta = FMath::RadiansToDegrees(fabsf(FMath::Acos( FVector::DotProduct(SoundForward, (ListenerLocation - Origin).GetSafeNormal()))));
 		VolumeMultiplier *= AttenuationEval(DistanceAlgorithm, theta - AttenuationShapeExtents.Y, AttenuationShapeExtents.Z, dBAttenuationAtMax);
 	}
 
diff --git a/Engine/Source/Runtime/Engine/Private/SoundNodeDoppler.cpp b/Engine/Source/Runtime/Engine/Private/SoundNodeDoppler.cpp
--- a/Engine/Source/Runtime/Engine/Private/SoundNodeDoppler.cpp
+++ b/Engine/Source/Runtime/Engine/Private/SoundNodeDoppler.cpp
@@ -26,7 +26,7 @@ float USoundNodeDoppler::GetDopplerPitchMultiplier(FListener const& InListener,
 {
 	static const float SpeedOfSoundInAirAtSeaLevel = 33000.f;		// cm/sec
 
-	FVector const SourceToListenerNorm = (InListener.Transform.GetTranslation() - Location).SafeNormal();
+	FVector const SourceToListenerNorm = (InListener.Transform.GetTranslation() - Location).GetSafeNormal();
 
 	// find source and listener speeds along the line between them
 	float const SourceVelMagTorwardListener = Velocity | SourceToListenerNorm;
diff --git a/Engine/Source/Runtime/Engine/Private/SplineComponent.cpp b/Engine/Source/Runtime/Engine/Private/SplineComponent.cpp
--- a/Engine/Source/Runtime/Engine/Private/SplineComponent.cpp
+++ b/Engine/Source/Runtime/Engine/Private/SplineComponent.cpp
@@ -498,7 +498,7 @@ FVector USplineComponent::GetWorldTangentAtDistanceAlongSpline(float Distance) c
 FVector USplineComponent::GetWorldDirectionAtDistanceAlongSpline(float Distance) const
 {
 	const float Param = SplineReparamTable.Eval(Distance, 0.f);
-	const FVector Tangent = SplineInfo.EvalDerivative(Param, FVector::ZeroVector).SafeNormal();
+	const FVector Tangent = SplineInfo.EvalDerivative(Param, FVector::ZeroVector).GetSafeNormal();
 	return ComponentToWorld.TransformVectorNoScale(Tangent);
 }
 
@@ -541,7 +541,7 @@ FVector USplineComponent::GetWorldDirectionAtTime(float Time, bool bUseConstantV
 	}
 
 	const float TimeMultiplier = (SplineInfo.Points.Num() - 1.0f) / Duration;
-	const FVector Tangent = SplineInfo.EvalDerivative(Time * TimeMultiplier, FVector::ZeroVector).SafeNormal();
+	const FVector Tangent = SplineInfo.EvalDerivative(Time * TimeMultiplier, FVector::ZeroVector).GetSafeNormal();
 	return ComponentToWorld.TransformVectorNoScale(Tangent);
 }
 
diff --git a/Engine/Source/Runtime/Engine/Private/SplineMeshComponent.cpp b/Engine/Source/Runtime/Engine/Private/SplineMeshComponent.cpp
--- a/Engine/Source/Runtime/Engine/Private/SplineMeshComponent.cpp
+++ b/Engine/Source/Runtime/Engine/Private/SplineMeshComponent.cpp
@@ -500,7 +500,7 @@ static FVector SplineEvalDir(const FVector& StartPos, const FVector& StartTangen
 
 	const float A2 = A  * A;
 
-	return ((C * A2) + (D * A) + E).SafeNormal();
+	return ((C * A2) + (D * A) + E).GetSafeNormal();
 }
 
 
@@ -520,8 +520,8 @@ FTransform USplineMeshComponent::CalcSliceTransform(const float DistanceAlong) c
 	const FVector SplineDir = SplineEvalDir( SplineParams.StartPos, SplineParams.StartTangent, SplineParams.EndPos, SplineParams.EndTangent, Alpha );
 
 	// Find base frenet frame
-	const FVector BaseXVec = (SplineUpDir ^ SplineDir).SafeNormal();
-	const FVector BaseYVec = (SplineDir ^ BaseXVec).SafeNormal();
+	const FVector BaseXVec = (SplineUpDir ^ SplineDir).GetSafeNormal();
+	const FVector BaseYVec = (SplineDir ^ BaseXVec).GetSafeNormal();
 
 	// Offset the spline by the desired amount
 	const FVector2D SliceOffset = FMath::Lerp<FVector2D>(SplineParams.StartOffset, SplineParams.EndOffset, HermiteAlpha);
diff --git a/Engine/Source/Runtime/Engine/Private/StaticMeshLight.cpp b/Engine/Source/Runtime/Engine/Private/StaticMeshLight.cpp
--- a/Engine/Source/Runtime/Engine/Private/StaticMeshLight.cpp
+++ b/Engine/Source/Runtime/Engine/Private/StaticMeshLight.cpp
@@ -28,9 +28,9 @@ static void GetStaticLightingVertex(
 	)
 {
 	OutVertex.WorldPosition = LocalToWorld.TransformPosition(PositionVertexBuffer.VertexPosition(VertexIndex));
-	OutVertex.WorldTangentX = LocalToWorld.TransformVector(VertexBuffer.VertexTangentX(VertexIndex)).SafeNormal();
-	OutVertex.WorldTangentY = LocalToWorld.TransformVector(VertexBuffer.VertexTangentY(VertexIndex)).SafeNormal();
-	OutVertex.WorldTangentZ = LocalToWorldInverseTranspose.TransformVector(VertexBuffer.VertexTangentZ(VertexIndex)).SafeNormal();
+	OutVertex.WorldTangentX = LocalToWorld.TransformVector(VertexBuffer.VertexTangentX(VertexIndex)).GetSafeNormal();
+	OutVertex.WorldTangentY = LocalToWorld.TransformVector(VertexBuffer.VertexTangentY(VertexIndex)).GetSafeNormal();
+	OutVertex.WorldTangentZ = LocalToWorldInverseTranspose.TransformVector(VertexBuffer.VertexTangentZ(VertexIndex)).GetSafeNormal();
 
 	checkSlow(VertexBuffer.GetNumTexCoords() <= ARRAY_COUNT(OutVertex.TextureCoordinates));
 	for(uint32 LightmapTextureCoordinateIndex = 0;LightmapTextureCoordinateIndex < VertexBuffer.GetNumTexCoords();LightmapTextureCoordinateIndex++)
diff --git a/Engine/Source/Runtime/Engine/Private/UnrealEngine.cpp b/Engine/Source/Runtime/Engine/Private/UnrealEngine.cpp
--- a/Engine/Source/Runtime/Engine/Private/UnrealEngine.cpp
+++ b/Engine/Source/Runtime/Engine/Private/UnrealEngine.cpp
@@ -6245,7 +6245,7 @@ static void DrawVolumeOnCanvas(const AVolume* Volume, FCanvas* Canvas, const FVe
 				// We only want to draw faces pointing up
 				const FVector Edge0 = V1 - V0;
 				const FVector Edge1 = V2 - V1;
-				const FVector Normal = (Edge1 ^ Edge0).SafeNormal();
+				const FVector Normal = (Edge1 ^ Edge0).GetSafeNormal();
 				if(Normal.Z > 0.01)
 				{
 					// Transform as 2d points in 'map space'
diff --git a/Engine/Source/Runtime/Engine/Public/Collision.h b/Engine/Source/Runtime/Engine/Public/Collision.h
--- a/Engine/Source/Runtime/Engine/Public/Collision.h
+++ b/Engine/Source/Runtime/Engine/Public/Collision.h
@@ -194,7 +194,7 @@ FORCEINLINE bool LineCheckWithTriangle(FHitResult& Result,const FVector& V1,cons
 		return false;
 	}
 
-	Result.Normal = ((V3-V2)^(V2-V1)).SafeNormal();
+	Result.Normal = ((V3-V2)^(V2-V1)).GetSafeNormal();
 	Result.Time = ((V1 - Start)|Result.Normal) / (Result.Normal|Direction);
 
 	return true;
diff --git a/Engine/Source/Runtime/Engine/Public/kDOP.h b/Engine/Source/Runtime/Engine/Public/kDOP.h
--- a/Engine/Source/Runtime/Engine/Public/kDOP.h
+++ b/Engine/Source/Runtime/Engine/Public/kDOP.h
@@ -291,7 +291,7 @@ struct FkDOPBuildCollisionTriangle
 	}
 	FVector4 GetLocalNormal() const
 	{
-		FVector4 LocalNormal = ((V1 - V2) ^ (V0 - V2)).SafeNormal();
+		FVector4 LocalNormal = ((V1 - V2) ^ (V0 - V2)).GetSafeNormal();
 		LocalNormal.W = Dot3(V0, LocalNormal);
 		return LocalNormal;
 	}
@@ -1080,7 +1080,7 @@ template <typename COLL_DATA_PROVIDER, typename KDOP_IDX_TYPE> struct TkDOPLineC
 	FORCEINLINE FVector4 GetHitNormal(void)
 	{
 		// Transform the hit back into world space using the transpose adjoint
-		FVector4 Normal = TkDOPCollisionCheck<COLL_DATA_PROVIDER,KDOP_IDX_TYPE>::CollDataProvider.GetLocalToWorldTransposeAdjoint().TransformVector(LocalHitNormal).SafeNormal();
+		FVector4 Normal = TkDOPCollisionCheck<COLL_DATA_PROVIDER,KDOP_IDX_TYPE>::CollDataProvider.GetLocalToWorldTransposeAdjoint().TransformVector(LocalHitNormal).GetSafeNormal();
 		// Flip the normal if the triangle is inverted
 		if (TkDOPCollisionCheck<COLL_DATA_PROVIDER,KDOP_IDX_TYPE>::CollDataProvider.GetDeterminant() < 0.f)
 		{
diff --git a/Engine/Source/Runtime/GameplayAbilities/Private/Abilities/GameplayAbilityTargetActor_GroundTrace.cpp b/Engine/Source/Runtime/GameplayAbilities/Private/Abilities/GameplayAbilityTargetActor_GroundTrace.cpp
--- a/Engine/Source/Runtime/GameplayAbilities/Private/Abilities/GameplayAbilityTargetActor_GroundTrace.cpp
+++ b/Engine/Source/Runtime/GameplayAbilities/Private/Abilities/GameplayAbilityTargetActor_GroundTrace.cpp
@@ -49,7 +49,7 @@ bool AGameplayAbilityTargetActor_GroundTrace::AdjustCollisionResultForShape(cons
 	UWorld *ThisWorld = GetWorld();
 	//Pull back toward player to find a better spot, accounting for the width of our object
 	FVector Movement = (OriginalEndPoint - OriginalStartPoint);
-	FVector MovementDirection = Movement.SafeNormal();
+	FVector MovementDirection = Movement.GetSafeNormal();
 	float MovementMagnitude2D = Movement.Size2D();
 
 	if (bDebug)
@@ -158,7 +158,7 @@ FHitResult AGameplayAbilityTargetActor_GroundTrace::PerformTrace(AActor* InSourc
 	}
 
 	//Second trace, straight down. Consider using InSourceActor->GetWorld()->NavigationSystem->ProjectPointToNavigation() instead of just going straight down in the case of movement abilities (flag/bool).
-	TraceStart = ReturnHitResult.Location - (TraceEnd - TraceStart).SafeNormal();		//Pull back very slightly to avoid scraping down walls
+	TraceStart = ReturnHitResult.Location - (TraceEnd - TraceStart).GetSafeNormal();		//Pull back very slightly to avoid scraping down walls
 	TraceEnd = TraceStart;
 	TraceStart.Z += CollisionHeightOffset;
 	TraceEnd.Z -= 99999.0f;
diff --git a/Engine/Source/Runtime/GameplayAbilities/Private/Abilities/GameplayAbilityTargetActor_Trace.cpp b/Engine/Source/Runtime/GameplayAbilities/Private/Abilities/GameplayAbilityTargetActor_Trace.cpp
--- a/Engine/Source/Runtime/GameplayAbilities/Private/Abilities/GameplayAbilityTargetActor_Trace.cpp
+++ b/Engine/Source/Runtime/GameplayAbilities/Private/Abilities/GameplayAbilityTargetActor_Trace.cpp
@@ -107,7 +107,7 @@ void AGameplayAbilityTargetActor_Trace::AimWithPlayerController(AActor* InSource
 		}
 
 		//Readjust so we have a full-length line going through the predicted target point.
-		FVector AimDirection = (TraceEnd - TraceStart).SafeNormal();
+		FVector AimDirection = (TraceEnd - TraceStart).GetSafeNormal();
 		if (AimDirection.SizeSquared() > 0.0f)
 		{
 			TraceEnd = TraceStart + (AimDirection * MaxRange);
diff --git a/Engine/Source/Runtime/GameplayAbilities/Private/Abilities/GameplayAbilityWorldReticle.cpp b/Engine/Source/Runtime/GameplayAbilities/Private/Abilities/GameplayAbilityWorldReticle.cpp
--- a/Engine/Source/Runtime/GameplayAbilities/Private/Abilities/GameplayAbilityWorldReticle.cpp
+++ b/Engine/Source/Runtime/GameplayAbilities/Private/Abilities/GameplayAbilityWorldReticle.cpp
@@ -69,10 +69,10 @@ void AGameplayAbilityWorldReticle::FaceTowardSource(bool bFaceIn2D)
 	{
 		if (bFaceIn2D)
 		{
-			FVector FacingVector = (TargetingActor->StartLocation.GetTargetingTransform().GetLocation() - GetActorLocation()).SafeNormal2D();
+			FVector FacingVector = (TargetingActor->StartLocation.GetTargetingTransform().GetLocation() - GetActorLocation()).GetSafeNormal2D();
 			if (FacingVector.IsZero())
 			{
-				FacingVector = -GetActorForwardVector().SafeNormal2D();
+				FacingVector = -GetActorForwardVector().GetSafeNormal2D();
 			}
 			if (!FacingVector.IsZero())
 			{
@@ -81,10 +81,10 @@ void AGameplayAbilityWorldReticle::FaceTowardSource(bool bFaceIn2D)
 		}
 		else
 		{
-			FVector FacingVector = (TargetingActor->StartLocation.GetTargetingTransform().GetLocation() - GetActorLocation()).SafeNormal();
+			FVector FacingVector = (TargetingActor->StartLocation.GetTargetingTransform().GetLocation() - GetActorLocation()).GetSafeNormal();
 			if (FacingVector.IsZero())
 			{
-				FacingVector = -GetActorForwardVector().SafeNormal();
+				FacingVector = -GetActorForwardVector().GetSafeNormal();
 			}
 			SetActorRotation(FacingVector.Rotation());
 		}
diff --git a/Engine/Source/Runtime/GameplayAbilities/Private/Abilities/Tasks/AbilityTask_WaitVelocityChange.cpp b/Engine/Source/Runtime/GameplayAbilities/Private/Abilities/Tasks/AbilityTask_WaitVelocityChange.cpp
--- a/Engine/Source/Runtime/GameplayAbilities/Private/Abilities/Tasks/AbilityTask_WaitVelocityChange.cpp
+++ b/Engine/Source/Runtime/GameplayAbilities/Private/Abilities/Tasks/AbilityTask_WaitVelocityChange.cpp
@@ -33,7 +33,7 @@ UAbilityTask_WaitVelocityChange* UAbilityTask_WaitVelocityChange::CreateWaitVelo
 	auto MyObj = NewTask<UAbilityTask_WaitVelocityChange>(WorldContextObject);
 
 	MyObj->MinimumMagnitude = InMinimumMagnitude;
-	MyObj->Direction = InDirection.SafeNormal();
+	MyObj->Direction = InDirection.GetSafeNormal();
 	
 
 	return MyObj;
diff --git a/Engine/Source/Runtime/GameplayAbilities/Private/AbilitySystemBlueprintLibrary.cpp b/Engine/Source/Runtime/GameplayAbilities/Private/AbilitySystemBlueprintLibrary.cpp
--- a/Engine/Source/Runtime/GameplayAbilities/Private/AbilitySystemBlueprintLibrary.cpp
+++ b/Engine/Source/Runtime/GameplayAbilities/Private/AbilitySystemBlueprintLibrary.cpp
@@ -235,7 +235,7 @@ FTransform UAbilitySystemBlueprintLibrary::GetTargetDataOrigin(FGameplayAbilityT
 			const FHitResult* HitResultPtr = Data->GetHitResult();
 			FTransform ReturnTransform;
 			ReturnTransform.SetLocation(HitResultPtr->TraceStart);
-			ReturnTransform.SetRotation((HitResultPtr->Location - HitResultPtr->TraceStart).SafeNormal().Rotation().Quaternion());
+			ReturnTransform.SetRotation((HitResultPtr->Location - HitResultPtr->TraceStart).GetSafeNormal().Rotation().Quaternion());
 			return ReturnTransform;
 		}
 	}
diff --git a/Engine/Source/Runtime/GameplayAbilities/Public/Abilities/GameplayAbilityTargetTypes.h b/Engine/Source/Runtime/GameplayAbilities/Public/Abilities/GameplayAbilityTargetTypes.h
--- a/Engine/Source/Runtime/GameplayAbilities/Public/Abilities/GameplayAbilityTargetTypes.h
+++ b/Engine/Source/Runtime/GameplayAbilities/Public/Abilities/GameplayAbilityTargetTypes.h
@@ -443,7 +443,7 @@ struct GAMEPLAYABILITIES_API FGameplayAbilityTargetData_ActorArray : public FGam
 		{
 			if (TargetActorArray[i].IsValid())
 			{
-				FVector Direction = (TargetActorArray[i].Get()->GetActorLocation() - ReturnTransform.GetLocation()).SafeNormal();
+				FVector Direction = (TargetActorArray[i].Get()->GetActorLocation() - ReturnTransform.GetLocation()).GetSafeNormal();
 				if (Direction.IsNormalized())
 				{
 					ReturnTransform.SetRotation(Direction.Rotation().Quaternion());
diff --git a/Engine/Source/Runtime/Landscape/Private/LandscapeEdit.cpp b/Engine/Source/Runtime/Landscape/Private/LandscapeEdit.cpp
--- a/Engine/Source/Runtime/Landscape/Private/LandscapeEdit.cpp
+++ b/Engine/Source/Runtime/Landscape/Private/LandscapeEdit.cpp
@@ -1831,8 +1831,8 @@ void ALandscapeProxy::Import(FGuid Guid, int32 VertsX, int32 VertsY,
 			FVector Vert10 = FVector(1.f, 0.f, ((float)HEIGHTDATA(QuadX + 1, QuadY + 0) - 32768.f)*LANDSCAPE_ZSCALE) * DrawScale3D;
 			FVector Vert11 = FVector(1.f, 1.f, ((float)HEIGHTDATA(QuadX + 1, QuadY + 1) - 32768.f)*LANDSCAPE_ZSCALE) * DrawScale3D;
 
-			FVector FaceNormal1 = ((Vert00 - Vert10) ^ (Vert10 - Vert11)).SafeNormal();
-			FVector FaceNormal2 = ((Vert11 - Vert01) ^ (Vert01 - Vert00)).SafeNormal();
+			FVector FaceNormal1 = ((Vert00 - Vert10) ^ (Vert10 - Vert11)).GetSafeNormal();
+			FVector FaceNormal2 = ((Vert11 - Vert01) ^ (Vert01 - Vert00)).GetSafeNormal();
 
 			// contribute to the vertex normals.
 			VertexNormals[(QuadX + 1 + (NumPatchesX + 1)*(QuadY + 0))] += FaceNormal1;
@@ -2161,7 +2161,7 @@ void ALandscapeProxy::Import(FGuid Guid, int32 VertsX, int32 VertsY,
 
 							// copy height and normal data
 							uint16 HeightValue = HEIGHTDATA(CompX + LandscapeComponent->GetSectionBase().X, CompY + LandscapeComponent->GetSectionBase().Y);
-							FVector Normal = VertexNormals[CompX + LandscapeComponent->GetSectionBase().X + (NumPatchesX + 1)*(CompY + LandscapeComponent->GetSectionBase().Y)].SafeNormal();
+							FVector Normal = VertexNormals[CompX + LandscapeComponent->GetSectionBase().X + (NumPatchesX + 1)*(CompY + LandscapeComponent->GetSectionBase().Y)].GetSafeNormal();
 
 							HeightmapInfo.HeightmapTextureMipData[0][HeightTexDataIdx].R = HeightValue >> 8;
 							HeightmapInfo.HeightmapTextureMipData[0][HeightTexDataIdx].G = HeightValue & 255;
diff --git a/Engine/Source/Runtime/Landscape/Private/LandscapeEditInterface.cpp b/Engine/Source/Runtime/Landscape/Private/LandscapeEditInterface.cpp
--- a/Engine/Source/Runtime/Landscape/Private/LandscapeEditInterface.cpp
+++ b/Engine/Source/Runtime/Landscape/Private/LandscapeEditInterface.cpp
@@ -236,8 +236,8 @@ void FLandscapeEditDataInterface::SetHeightData(int32 X1, int32 Y1, int32 X2, in
 				FVector Vert10 = FVector(XYOffsets[(X+0) + NumVertsX*(Y+0)].X+1.f,	XYOffsets[(X+0) + NumVertsX*(Y+0)].Y, ((float)Data[(X+1) + Stride*(Y+0)] - 32768.f)*LANDSCAPE_ZSCALE) * DrawScale;
 				FVector Vert11 = FVector(XYOffsets[(X+0) + NumVertsX*(Y+0)].X+1.f,	XYOffsets[(X+0) + NumVertsX*(Y+0)].Y+1.f, ((float)Data[(X+1) + Stride*(Y+1)] - 32768.f)*LANDSCAPE_ZSCALE) * DrawScale;
 
-				FVector FaceNormal1 = ((Vert00-Vert10) ^ (Vert10-Vert11)).SafeNormal();
-				FVector FaceNormal2 = ((Vert11-Vert01) ^ (Vert01-Vert00)).SafeNormal(); 
+				FVector FaceNormal1 = ((Vert00-Vert10) ^ (Vert10-Vert11)).GetSafeNormal();
+				FVector FaceNormal2 = ((Vert11-Vert01) ^ (Vert01-Vert00)).GetSafeNormal(); 
 
 				// contribute to the vertex normals.
 				VertexNormals[(X+1 + NumVertsX*(Y+0))] += FaceNormal1;
@@ -348,7 +348,7 @@ void FLandscapeEditDataInterface::SetHeightData(int32 X1, int32 Y1, int32 X2, in
 							// Update normals if we're not on an edge vertex
 							if( VertexNormals && LandscapeX > X1 && LandscapeX < X2 && LandscapeY > Y1 && LandscapeY < Y2 )
 							{
-								FVector Normal = VertexNormals[DataIndex].SafeNormal();
+								FVector Normal = VertexNormals[DataIndex].GetSafeNormal();
 								TexData.B = FMath::RoundToInt( 127.5f * (Normal.X + 1.f) );
 								TexData.A = FMath::RoundToInt( 127.5f * (Normal.Y + 1.f) );
 							}
@@ -460,8 +460,8 @@ void FLandscapeEditDataInterface::RecalculateNormals()
 				FVector Vert10 = FVector(XYOffsets[(X+0) + Stride*(Y+0)].X+1.f,	XYOffsets[(X+0) + Stride*(Y+0)].Y, ((float)HeightData[(X+1) + Stride*(Y+0)] - 32768.f)*LANDSCAPE_ZSCALE) * DrawScale;
 				FVector Vert11 = FVector(XYOffsets[(X+0) + Stride*(Y+0)].X+1.f,	XYOffsets[(X+0) + Stride*(Y+0)].Y+1.f,((float)HeightData[(X+1) + Stride*(Y+1)] - 32768.f)*LANDSCAPE_ZSCALE) * DrawScale;
 
-				FVector FaceNormal1 = ((Vert00-Vert10) ^ (Vert10-Vert11)).SafeNormal();
-				FVector FaceNormal2 = ((Vert11-Vert01) ^ (Vert01-Vert00)).SafeNormal(); 
+				FVector FaceNormal1 = ((Vert00-Vert10) ^ (Vert10-Vert11)).GetSafeNormal();
+				FVector FaceNormal2 = ((Vert11-Vert01) ^ (Vert01-Vert00)).GetSafeNormal(); 
 
 				// contribute to the vertex normals.
 				VertexNormals[(X+1 + Stride*(Y+0))] += FaceNormal1;
@@ -498,7 +498,7 @@ void FLandscapeEditDataInterface::RecalculateNormals()
 						FColor& TexData = HeightmapTextureData[ TexX + TexY * SizeU ];
 
 						// Update the texture
-						FVector Normal = VertexNormals[DataIndex].SafeNormal();
+						FVector Normal = VertexNormals[DataIndex].GetSafeNormal();
 						TexData.B = FMath::RoundToInt( 127.5f * (Normal.X + 1.f) );
 						TexData.A = FMath::RoundToInt( 127.5f * (Normal.Y + 1.f) );
 					}
diff --git a/Engine/Source/Runtime/Landscape/Private/LandscapeGizmoActor.cpp b/Engine/Source/Runtime/Landscape/Private/LandscapeGizmoActor.cpp
--- a/Engine/Source/Runtime/Landscape/Private/LandscapeGizmoActor.cpp
+++ b/Engine/Source/Runtime/Landscape/Private/LandscapeGizmoActor.cpp
@@ -189,7 +189,7 @@ public:
 						SampledPos.Z = Gizmo->GetLandscapeHeight(SampledPos.Z);
 
 						FVector SampledNormal = NormalM.TransformVector(Gizmo->SampledNormal[X + Y * ALandscapeGizmoActiveActor::DataTexSize]);
-						SampledNormal = SampledNormal.SafeNormal();
+						SampledNormal = SampledNormal.GetSafeNormal();
 
 						SampledPositions.Add(SampledPos);
 						SampledNormals.Add(SampledNormal);
@@ -277,7 +277,7 @@ public:
 							{
 								FVector SampledNormal = SampledNormals[X + Y * SampleSizeX];
 								FVector TangentX(SampledNormal.Z, 0, -SampledNormal.X);
-								TangentX = TangentX.SafeNormal();
+								TangentX = TangentX.GetSafeNormal();
 
 								MeshBuilder.AddVertex(SampledPositions[X + Y * SampleSizeX], FVector2D((float)X / (SampleSizeX), (float)Y / (SampleSizeY)), TangentX, SampledNormal^TangentX, SampledNormal, FColor(255, 255, 255) );
 							}
@@ -744,8 +744,8 @@ void ALandscapeGizmoActiveActor::CalcNormal()
 				FVector Vert10 = SampledHeight[X+1 + Y*DataTexSize];
 				FVector Vert11 = SampledHeight[X+1 + (Y+1)*DataTexSize];
 
-				FVector FaceNormal1 = ((Vert00-Vert10) ^ (Vert10-Vert11)).SafeNormal();
-				FVector FaceNormal2 = ((Vert11-Vert01) ^ (Vert01-Vert00)).SafeNormal(); 
+				FVector FaceNormal1 = ((Vert00-Vert10) ^ (Vert10-Vert11)).GetSafeNormal();
+				FVector FaceNormal2 = ((Vert11-Vert01) ^ (Vert01-Vert00)).GetSafeNormal(); 
 
 				// contribute to the vertex normals.
 				SampledNormal[X + Y*DataTexSize] += FaceNormal1;
@@ -758,7 +758,7 @@ void ALandscapeGizmoActiveActor::CalcNormal()
 		{
 			for (int32 X = 0; X < SampleSizeX; ++X)
 			{
-				SampledNormal[X + Y*DataTexSize] = SampledNormal[X + Y*DataTexSize].SafeNormal();
+				SampledNormal[X + Y*DataTexSize] = SampledNormal[X + Y*DataTexSize].GetSafeNormal();
 			}
 		}
 	}
diff --git a/Engine/Source/Runtime/Landscape/Private/LandscapeSplines.cpp b/Engine/Source/Runtime/Landscape/Private/LandscapeSplines.cpp
--- a/Engine/Source/Runtime/Landscape/Private/LandscapeSplines.cpp
+++ b/Engine/Source/Runtime/Landscape/Private/LandscapeSplines.cpp
@@ -723,7 +723,7 @@ void ULandscapeSplineControlPoint::AutoSetConnections(bool bIncludingValid)
 				FVector StartLocation; FRotator StartRotation;
 				NearConnection.ControlPoint->GetConnectionLocationAndRotation(NearConnection.SocketName, StartLocation, StartRotation);
 
-				if (FVector::DotProduct((EndLocation - StartLocation).SafeNormal(), StartRotation.Vector()) < 0)
+				if (FVector::DotProduct((EndLocation - StartLocation).GetSafeNormal(), StartRotation.Vector()) < 0)
 				{
 					NearConnection.TangentLen = -NearConnection.TangentLen;
 				}
@@ -856,7 +856,7 @@ void ULandscapeSplineControlPoint::UpdateSplinePoints(bool bUpdateCollision, boo
 
 			const float Roll = FMath::DegreesToRadians(StartRotation.Roll);
 			const FVector Tangent = StartRotation.Vector();
-			const FVector BiNormal = FQuat(Tangent, -Roll).RotateVector((Tangent ^ FVector(0, 0, -1)).SafeNormal());
+			const FVector BiNormal = FQuat(Tangent, -Roll).RotateVector((Tangent ^ FVector(0, 0, -1)).GetSafeNormal());
 			const FVector LeftPos = StartLocation - BiNormal * Width;
 			const FVector RightPos = StartLocation + BiNormal * Width;
 			const FVector FalloffLeftPos = StartLocation - BiNormal * (Width + SideFalloff);
@@ -877,7 +877,7 @@ void ULandscapeSplineControlPoint::UpdateSplinePoints(bool bUpdateCollision, boo
 
 		const float Roll = FMath::DegreesToRadians(StartRotation.Roll);
 		const FVector Tangent = StartRotation.Vector();
-		const FVector BiNormal = FQuat(Tangent, -Roll).RotateVector((Tangent ^ FVector(0, 0, -1)).SafeNormal());
+		const FVector BiNormal = FQuat(Tangent, -Roll).RotateVector((Tangent ^ FVector(0, 0, -1)).GetSafeNormal());
 		const FVector LeftPos = StartLocation - BiNormal * Width;
 		const FVector RightPos = StartLocation + BiNormal * Width;
 		const FVector FalloffLeftPos = StartLocation - BiNormal * (Width + SideFalloff);
@@ -1109,11 +1109,11 @@ void ULandscapeSplineSegment::AutoFlipTangents()
 	Connections[1].ControlPoint->GetConnectionLocationAndRotation(Connections[1].SocketName, EndLocation, EndRotation);
 
 	// Flipping the tangent is only allowed if not using a socket
-	if (Connections[0].SocketName == NAME_None && FVector::DotProduct((EndLocation - StartLocation).SafeNormal() * Connections[0].TangentLen, StartRotation.Vector()) < 0)
+	if (Connections[0].SocketName == NAME_None && FVector::DotProduct((EndLocation - StartLocation).GetSafeNormal() * Connections[0].TangentLen, StartRotation.Vector()) < 0)
 	{
 		Connections[0].TangentLen = -Connections[0].TangentLen;
 	}
-	if (Connections[1].SocketName == NAME_None && FVector::DotProduct((StartLocation - EndLocation).SafeNormal() * Connections[1].TangentLen, EndRotation.Vector()) < 0)
+	if (Connections[1].SocketName == NAME_None && FVector::DotProduct((StartLocation - EndLocation).GetSafeNormal() * Connections[1].TangentLen, EndRotation.Vector()) < 0)
 	{
 		Connections[1].TangentLen = -Connections[1].TangentLen;
 	}
@@ -1238,8 +1238,8 @@ void ULandscapeSplineSegment::UpdateSplinePoints(bool bUpdateCollision)
 		const float NewKeyFalloff = FMath::Lerp(StartSideFalloff, EndSideFalloff, NewKeyCosInterp);
 		const float NewKeyRoll = FMath::Lerp(StartRoll, EndRoll, NewKeyCosInterp);
 		const FVector NewKeyPos = SplineInfo.Eval(NewKeyTime, FVector::ZeroVector);
-		const FVector NewKeyTangent = SplineInfo.EvalDerivative(NewKeyTime, FVector::ZeroVector).SafeNormal();
-		const FVector NewKeyBiNormal = FQuat(NewKeyTangent, -NewKeyRoll).RotateVector((NewKeyTangent ^ FVector(0, 0, -1)).SafeNormal());
+		const FVector NewKeyTangent = SplineInfo.EvalDerivative(NewKeyTime, FVector::ZeroVector).GetSafeNormal();
+		const FVector NewKeyBiNormal = FQuat(NewKeyTangent, -NewKeyRoll).RotateVector((NewKeyTangent ^ FVector(0, 0, -1)).GetSafeNormal());
 		const FVector NewKeyLeftPos = NewKeyPos - NewKeyBiNormal * NewKeyWidth;
 		const FVector NewKeyRightPos = NewKeyPos + NewKeyBiNormal * NewKeyWidth;
 		const FVector NewKeyFalloffLeftPos = NewKeyPos - NewKeyBiNormal * (NewKeyWidth + NewKeyFalloff);
@@ -1261,8 +1261,8 @@ void ULandscapeSplineSegment::UpdateSplinePoints(bool bUpdateCollision)
 				const float NewFalloff = FMath::Lerp(StartSideFalloff, EndSideFalloff, NewCosInterp);
 				const float NewRoll = FMath::Lerp(StartRoll, EndRoll, NewCosInterp);
 				const FVector NewPos = SplineInfo.Eval(NewTime, FVector::ZeroVector);
-				const FVector NewTangent = SplineInfo.EvalDerivative(NewTime, FVector::ZeroVector).SafeNormal();
-				const FVector NewBiNormal = FQuat(NewTangent, -NewRoll).RotateVector((NewTangent ^ FVector(0, 0, -1)).SafeNormal());
+				const FVector NewTangent = SplineInfo.EvalDerivative(NewTime, FVector::ZeroVector).GetSafeNormal();
+				const FVector NewBiNormal = FQuat(NewTangent, -NewRoll).RotateVector((NewTangent ^ FVector(0, 0, -1)).GetSafeNormal());
 				const FVector NewLeftPos = NewPos - NewBiNormal * NewWidth;
 				const FVector NewRightPos = NewPos + NewBiNormal * NewWidth;
 				const FVector NewFalloffLeftPos = NewPos - NewBiNormal * (NewWidth + NewFalloff);
@@ -1679,8 +1679,8 @@ bool ULandscapeSplineSegment::FixSelfIntersection(FVector FLandscapeSplineInterp
 		{
 			const FLandscapeSplineInterpPoint& CurrentPoint = Points[i];
 			const FLandscapeSplineInterpPoint& NextPoint = Points[i+1];
-			const FVector Direction = (NextPoint.Center - CurrentPoint.Center).SafeNormal();
-			const FVector SideDirection = (NextPoint.*Side - CurrentPoint.*Side).SafeNormal();
+			const FVector Direction = (NextPoint.Center - CurrentPoint.Center).GetSafeNormal();
+			const FVector SideDirection = (NextPoint.*Side - CurrentPoint.*Side).GetSafeNormal();
 			bReversed = (SideDirection | Direction) < 0;
 		}
 
diff --git a/Engine/Source/Runtime/Launch/Private/Android/LaunchAndroid.cpp b/Engine/Source/Runtime/Launch/Private/Android/LaunchAndroid.cpp
--- a/Engine/Source/Runtime/Launch/Private/Android/LaunchAndroid.cpp
+++ b/Engine/Source/Runtime/Launch/Private/Android/LaunchAndroid.cpp
@@ -922,7 +922,7 @@ static int HandleSensorEvents(int fd, int events, void* data)
 
 		// Calc the tilt from the accelerometer as it's not
 		// available directly.
-		FVector accelerometer_dir = -current_accelerometer.SafeNormal();
+		FVector accelerometer_dir = -current_accelerometer.GetSafeNormal();
 		float current_pitch
 			= FMath::Atan2(accelerometer_dir.Y, accelerometer_dir.Z);
 		float current_roll
diff --git a/Engine/Source/Runtime/Mac/CoreAudio/Private/CoreAudioSource.cpp b/Engine/Source/Runtime/Mac/CoreAudio/Private/CoreAudioSource.cpp
--- a/Engine/Source/Runtime/Mac/CoreAudio/Private/CoreAudioSource.cpp
+++ b/Engine/Source/Runtime/Mac/CoreAudio/Private/CoreAudioSource.cpp
@@ -283,7 +283,7 @@ void FCoreAudioSoundSource::Update( void )
 		}
 		else if( WaveInstance->bUseSpatialization )
 		{
-			FVector Direction = AudioDevice->InverseTransform.TransformPosition(WaveInstance->Location).SafeNormal();
+			FVector Direction = AudioDevice->InverseTransform.TransformPosition(WaveInstance->Location).GetSafeNormal();
 
 			FVector EmitterPosition;
 			EmitterPosition.X = -Direction.Z;
diff --git a/Engine/Source/Runtime/Renderer/Private/DistanceFieldSurfaceCacheLighting.cpp b/Engine/Source/Runtime/Renderer/Private/DistanceFieldSurfaceCacheLighting.cpp
--- a/Engine/Source/Runtime/Renderer/Private/DistanceFieldSurfaceCacheLighting.cpp
+++ b/Engine/Source/Runtime/Renderer/Private/DistanceFieldSurfaceCacheLighting.cpp
@@ -3679,7 +3679,7 @@ FVector GetUnitPosition2(FRandomStream& RandomStream)
 /** Generates a pseudo-random unit vector, uniformly distributed over all directions. */
 FVector GetUnitVector2(FRandomStream& RandomStream)
 {
-	return GetUnitPosition2(RandomStream).UnsafeNormal();
+	return GetUnitPosition2(RandomStream).GetUnsafeNormal();
 }
 
 void GenerateBestSpacedVectors()
@@ -3721,14 +3721,14 @@ void GenerateBestSpacedVectors()
 						if (Dot > 0)
 						{
 							// Repulsion force
-							Force += .001f * Distance.SafeNormal() * Dot * Dot * Dot * Dot;
+							Force += .001f * Distance.GetSafeNormal() * Dot * Dot * Dot * Dot;
 						}
 					}
 				}
 
 				FVector NewPosition = OriginalSpacedVectors9[i] + Force;
 				NewPosition.Z = FMath::Max(NewPosition.Z, MinZ);
-				NewPosition = NewPosition.SafeNormal();
+				NewPosition = NewPosition.GetSafeNormal();
 				OriginalSpacedVectors9[i] = NewPosition;
 			}
 		}
diff --git a/Engine/Source/Runtime/Renderer/Private/ShadowSetup.cpp b/Engine/Source/Runtime/Renderer/Private/ShadowSetup.cpp
--- a/Engine/Source/Runtime/Renderer/Private/ShadowSetup.cpp
+++ b/Engine/Source/Runtime/Renderer/Private/ShadowSetup.cpp
@@ -257,8 +257,8 @@ static bool GetBestShadowTransform(const FVector& ZAxis,const FBoundingBoxVertex
 		const FVector Point = PointsPtr[EdgesPtr[EdgeIndex].FirstEdgeIndex];
 		const FVector OtherPoint = PointsPtr[EdgesPtr[EdgeIndex].SecondEdgeIndex];
 		const FVector PointDelta = OtherPoint - Point;
-		const FVector TrialXAxis = (PointDelta - ZAxis * (PointDelta | ZAxis)).SafeNormal();
-		const FVector TrialYAxis = (ZAxis ^ TrialXAxis).SafeNormal();
+		const FVector TrialXAxis = (PointDelta - ZAxis * (PointDelta | ZAxis)).GetSafeNormal();
+		const FVector TrialYAxis = (ZAxis ^ TrialXAxis).GetSafeNormal();
 
 		// Calculate the size of the projection of the bounds onto this axis and an axis orthogonal to it and the Z axis.
 		float MinProjectedX = FLT_MAX;
@@ -414,7 +414,7 @@ FProjectedShadowInfo::FProjectedShadowInfo(
 	FMatrix LightToShadow;
 	float AspectRatio;
 
-	if (GetBestShadowTransform(Initializer.FaceDirection.SafeNormal(),ProjectedBoundsPoints,BoundsEdges,AspectRatio,LightToShadow))
+	if (GetBestShadowTransform(Initializer.FaceDirection.GetSafeNormal(), ProjectedBoundsPoints, BoundsEdges, AspectRatio, LightToShadow))
 	{
 		bValidTransform = true;
 		const FMatrix WorldToShadow = WorldToLightScaled * LightToShadow;
@@ -516,7 +516,7 @@ FProjectedShadowInfo::FProjectedShadowInfo(
 	FVector	XAxis, YAxis;
 	Initializer.FaceDirection.FindBestAxisVectors(XAxis,YAxis);
 	const FMatrix WorldToLightScaled = Initializer.WorldToLight * FScaleMatrix(Initializer.Scales);
-	const FMatrix WorldToFace = WorldToLightScaled * FBasisVectorMatrix(-XAxis,YAxis,Initializer.FaceDirection.SafeNormal(),FVector::ZeroVector);
+	const FMatrix WorldToFace = WorldToLightScaled * FBasisVectorMatrix(-XAxis,YAxis,Initializer.FaceDirection.GetSafeNormal(),FVector::ZeroVector);
 
 	MaxSubjectZ = WorldToFace.TransformPosition(Initializer.SubjectBounds.Origin).Z + Initializer.SubjectBounds.SphereRadius;
 	MinSubjectZ = FMath::Max(MaxSubjectZ - Initializer.SubjectBounds.SphereRadius * 2,Initializer.MinLightW);
diff --git a/Engine/Source/Runtime/Renderer/Private/SystemTextures.cpp b/Engine/Source/Runtime/Renderer/Private/SystemTextures.cpp
--- a/Engine/Source/Runtime/Renderer/Private/SystemTextures.cpp
+++ b/Engine/Source/Runtime/Renderer/Private/SystemTextures.cpp
@@ -336,7 +336,7 @@ void FSystemTextures::InitializeTextures(FRHICommandListImmediate& RHICmdList, E
 							float SinTheta = FMath::Sqrt(1.0f - CosTheta * CosTheta);
 
 							FVector L(SinTheta * FMath::Cos(Phi), SinTheta * FMath::Sin(Phi), CosTheta);
-							FVector H = ( V + L ).UnsafeNormal();
+							FVector H = (V + L).GetUnsafeNormal();
 
 							float NoL = FMath::Max(L.Z, 0.0f);
 							float NoH = FMath::Max(H.Z, 0.0f);
diff --git a/Engine/Source/Runtime/Slate/Private/Widgets/Input/SVirtualJoystick.cpp b/Engine/Source/Runtime/Slate/Private/Widgets/Input/SVirtualJoystick.cpp
--- a/Engine/Source/Runtime/Slate/Private/Widgets/Input/SVirtualJoystick.cpp
+++ b/Engine/Source/Runtime/Slate/Private/Widgets/Input/SVirtualJoystick.cpp
@@ -359,7 +359,7 @@ void SVirtualJoystick::Tick( const FGeometry& AllottedGeometry, const double InC
 			// now pass the fake joystick events to the game
 			// Assume that joystick size is all equal
 			float JoysticInputSize = Control.ThumbPosition.Size() * 2.f / Control.VisualSize.X;
-			FVector2D NormalizedOffset = Control.ThumbPosition.SafeNormal() * Control.InputScale * JoysticInputSize;
+			FVector2D NormalizedOffset = Control.ThumbPosition.GetSafeNormal() * Control.InputScale * JoysticInputSize;
 			EControllerButtons::Type XAxis = ControlIndex == 0 ? EControllerButtons::LeftAnalogX : EControllerButtons::RightAnalogX;
 			EControllerButtons::Type YAxis = ControlIndex == 0 ? EControllerButtons::LeftAnalogY : EControllerButtons::RightAnalogY;
 
diff --git a/Engine/Source/Runtime/SlateCore/Private/Rendering/ElementBatcher.cpp b/Engine/Source/Runtime/SlateCore/Private/Rendering/ElementBatcher.cpp
--- a/Engine/Source/Runtime/SlateCore/Private/Rendering/ElementBatcher.cpp
+++ b/Engine/Source/Runtime/SlateCore/Private/Rendering/ElementBatcher.cpp
@@ -823,7 +823,7 @@ void FSlateElementBatcher::AddSplineElement( const FSlateDrawElement& DrawElemen
 	const FVector2D EndDir = InPayload.EndDir;
 	
 	// Compute the normal to the line
-	FVector2D Normal = FVector2D( StartPt.Y - EndPt.Y, EndPt.X - StartPt.X ).SafeNormal();
+	FVector2D Normal = FVector2D( StartPt.Y - EndPt.Y, EndPt.X - StartPt.X ).GetSafeNormal();
 
 	FVector2D Up = Normal * HalfThickness;
 
@@ -850,7 +850,7 @@ void FSlateElementBatcher::AddSplineElement( const FSlateDrawElement& DrawElemen
 		int32 IndexStart = BatchVertices.Num();
 
 		// Compute the normal to the line
-		FVector2D SegmentNormal = FVector2D( StartPos.Y - EndPos.Y, EndPos.X - StartPos.X ).SafeNormal();
+		FVector2D SegmentNormal = FVector2D( StartPos.Y - EndPos.Y, EndPos.X - StartPos.X ).GetSafeNormal();
 
 		// Create the new vertices for the thick line segment
 		Up = SegmentNormal * HalfThickness;
@@ -968,7 +968,7 @@ void FSlateElementBatcher::AddLineElement( const FSlateDrawElement& DrawElement
 		FVector2D StartPos = Points[0];
 		FVector2D EndPos = Points[1];
 
-		FVector2D Normal = FVector2D( StartPos.Y - EndPos.Y, EndPos.X - StartPos.X ).SafeNormal();
+		FVector2D Normal = FVector2D( StartPos.Y - EndPos.Y, EndPos.X - StartPos.X ).GetSafeNormal();
 
 		FVector2D Up = Normal * HalfThickness;
 
@@ -985,7 +985,7 @@ void FSlateElementBatcher::AddLineElement( const FSlateDrawElement& DrawElement
 			uint32 IndexStart = BatchVertices.Num();
 
 			// Compute the normal to the line
-			Normal = FVector2D( StartPos.Y - EndPos.Y, EndPos.X - StartPos.X ).SafeNormal();
+			Normal = FVector2D( StartPos.Y - EndPos.Y, EndPos.X - StartPos.X ).GetSafeNormal();
 
 			// Create the new vertices for the thick line segment
 			Up = Normal * HalfThickness;
@@ -1000,7 +1000,7 @@ void FSlateElementBatcher::AddLineElement( const FSlateDrawElement& DrawElement
 				const FVector2D NextEndPos = Points[Point+1];
 
 				// The normal of the next segment
-				const FVector2D NextNormal = FVector2D( EndPos.Y - NextEndPos.Y, NextEndPos.X - EndPos.X ).SafeNormal();
+				const FVector2D NextNormal = FVector2D( EndPos.Y - NextEndPos.Y, NextEndPos.X - EndPos.X ).GetSafeNormal();
 
 				// The next amount to adjust the vertices by 
 				FVector2D NextUp = NextNormal * HalfThickness;
diff --git a/Engine/Source/Runtime/Windows/XAudio2/Private/XAudio2Source.cpp b/Engine/Source/Runtime/Windows/XAudio2/Private/XAudio2Source.cpp
--- a/Engine/Source/Runtime/Windows/XAudio2/Private/XAudio2Source.cpp
+++ b/Engine/Source/Runtime/Windows/XAudio2/Private/XAudio2Source.cpp
@@ -417,7 +417,7 @@ void FXAudio2SoundSource::GetChannelVolumes( float ChannelVolumes[CHANNELOUT_COU
 			if( WaveInstance->bUseSpatialization )
 			{
 				FVector UnnormalizedDirection = AudioDevice->InverseTransform.TransformPosition( WaveInstance->Location );
-				Direction = UnnormalizedDirection.SafeNormal();
+				Direction = UnnormalizedDirection.GetSafeNormal();
 				float Distance = UnnormalizedDirection.Size();
 				NormalizedOmniRadius = (Distance > 0) ? (WaveInstance->OmniRadius / Distance) : FLT_MAX;
 			}
