Expose IsTickEnabled to blueprints

diff --git a/Engine/Source/Runtime/Engine/Classes/GameFramework/Actor.h b/Engine/Source/Runtime/Engine/Classes/GameFramework/Actor.h
--- a/Engine/Source/Runtime/Engine/Classes/GameFramework/Actor.h
+++ b/Engine/Source/Runtime/Engine/Classes/GameFramework/Actor.h
@@ -1492,6 +1492,7 @@ public:
 	virtual void SetActorTickEnabled(bool bEnabled);
 
 	/**  Returns whether this actor has tick enabled or not	 */
+	UFUNCTION(BlueprintCallable, Category="Utilities")
 	bool IsActorTickEnabled() const;
 
 	/**
