Fix warning due to deprecated function.

diff --git a/Engine/Source/Runtime/Engine/Private/GameFramework/SpringArmComponent.cpp b/Engine/Source/Runtime/Engine/Private/GameFramework/SpringArmComponent.cpp
--- a/Engine/Source/Runtime/Engine/Private/GameFramework/SpringArmComponent.cpp
+++ b/Engine/Source/Runtime/Engine/Private/GameFramework/SpringArmComponent.cpp
@@ -123,7 +123,7 @@ void USpringArmComponent::UpdateDesiredArmLocation(bool bDoTrace, bool bDoLocati
 			const FVector FromOrigin = DesiredLoc - ArmOrigin;
 			if (FromOrigin.SizeSquared() > FMath::Square(CameraLagMaxDistance))
 			{
-				DesiredLoc = ArmOrigin + FromOrigin.ClampMaxSize(CameraLagMaxDistance);
+				DesiredLoc = ArmOrigin + FromOrigin.GetClampedToMaxSize(CameraLagMaxDistance);
 				bClampedDist = true;
 			}
 		}		
