- Rename OverrideMaterials from Materials since that is overriden materials.

- Added GetNumOverrideMaterials to return overriden materials.

- Reverted GetNumMaterials to previous behavior - they return Mesh->Materials count since it can't have more than that but only if Mesh is available.

diff --git a/Engine/Plugins/2D/Paper2D/Source/Paper2D/Private/PaperFlipbookComponent.cpp b/Engine/Plugins/2D/Paper2D/Source/Paper2D/Private/PaperFlipbookComponent.cpp
--- a/Engine/Plugins/2D/Paper2D/Source/Paper2D/Private/PaperFlipbookComponent.cpp
+++ b/Engine/Plugins/2D/Paper2D/Source/Paper2D/Private/PaperFlipbookComponent.cpp
@@ -148,9 +148,9 @@ void UPaperFlipbookComponent::GetUsedTextures(TArray<UTexture*>& OutTextures, EM
 
 UMaterialInterface* UPaperFlipbookComponent::GetMaterial(int32 MaterialIndex) const
 {
-	if (Materials.IsValidIndex(MaterialIndex) && (Materials[MaterialIndex] != nullptr))
+	if (OverrideMaterials.IsValidIndex(MaterialIndex) && (OverrideMaterials[MaterialIndex] != nullptr))
 	{
-		return Materials[MaterialIndex];
+		return OverrideMaterials[MaterialIndex];
 	}
 	else if (SourceFlipbook != nullptr)
 	{
@@ -173,7 +173,7 @@ void UPaperFlipbookComponent::GetStreamingTextureInfo(TArray<FStreamingTexturePr
 
 int32 UPaperFlipbookComponent::GetNumMaterials() const
 {
-	return FMath::Max<int32>(Materials.Num(), 1);
+	return FMath::Max<int32>(OverrideMaterials.Num(), 1);
 }
 
 void UPaperFlipbookComponent::CalculateCurrentFrame()
diff --git a/Engine/Plugins/2D/Paper2D/Source/Paper2D/Private/PaperSpriteComponent.cpp b/Engine/Plugins/2D/Paper2D/Source/Paper2D/Private/PaperSpriteComponent.cpp
--- a/Engine/Plugins/2D/Paper2D/Source/Paper2D/Private/PaperSpriteComponent.cpp
+++ b/Engine/Plugins/2D/Paper2D/Source/Paper2D/Private/PaperSpriteComponent.cpp
@@ -227,9 +227,9 @@ void UPaperSpriteComponent::GetUsedTextures(TArray<UTexture*>& OutTextures, EMat
 
 UMaterialInterface* UPaperSpriteComponent::GetMaterial(int32 MaterialIndex) const
 {
-	if (Materials.IsValidIndex(MaterialIndex) && (Materials[MaterialIndex] != nullptr))
+	if (OverrideMaterials.IsValidIndex(MaterialIndex) && (OverrideMaterials[MaterialIndex] != nullptr))
 	{
-		return Materials[MaterialIndex];
+		return OverrideMaterials[MaterialIndex];
 	}
 	else if (SourceSprite != nullptr)
 	{
@@ -254,11 +254,11 @@ int32 UPaperSpriteComponent::GetNumMaterials() const
 {
 	if (SourceSprite != nullptr)
 	{
-		return FMath::Max<int32>(Materials.Num(), SourceSprite->GetNumMaterials());
+		return FMath::Max<int32>(OverrideMaterials.Num(), SourceSprite->GetNumMaterials());
 	}
 	else
 	{
-		return FMath::Max<int32>(Materials.Num(), 1);
+		return FMath::Max<int32>(OverrideMaterials.Num(), 1);
 	}
 }
 
diff --git a/Engine/Source/Editor/DetailCustomizations/Private/MeshComponentDetails.cpp b/Engine/Source/Editor/DetailCustomizations/Private/MeshComponentDetails.cpp
--- a/Engine/Source/Editor/DetailCustomizations/Private/MeshComponentDetails.cpp
+++ b/Engine/Source/Editor/DetailCustomizations/Private/MeshComponentDetails.cpp
@@ -12,7 +12,7 @@ TSharedRef<IDetailCustomization> FMeshComponentDetails::MakeInstance()
 void FMeshComponentDetails::CustomizeDetails( IDetailLayoutBuilder& DetailLayout )
 {
 	RenderingCategory = &DetailLayout.EditCategory("Rendering");
-	TSharedRef<IPropertyHandle> MaterialProperty = DetailLayout.GetProperty( GET_MEMBER_NAME_CHECKED(UMeshComponent, Materials) );
+	TSharedRef<IPropertyHandle> MaterialProperty = DetailLayout.GetProperty( GET_MEMBER_NAME_CHECKED(UMeshComponent, OverrideMaterials) );
 
 	if( MaterialProperty->IsValidHandle() )
 	{
diff --git a/Engine/Source/Editor/FoliageEdit/Private/FoliageEdMode.cpp b/Engine/Source/Editor/FoliageEdit/Private/FoliageEdMode.cpp
--- a/Engine/Source/Editor/FoliageEdit/Private/FoliageEdMode.cpp
+++ b/Engine/Source/Editor/FoliageEdit/Private/FoliageEdMode.cpp
@@ -56,7 +56,7 @@ FEdModeFoliage::FEdModeFoliage()
 	SphereBrushComponent->SetCollisionProfileName(UCollisionProfile::NoCollision_ProfileName);
 	SphereBrushComponent->SetCollisionObjectType(ECC_WorldDynamic);
 	SphereBrushComponent->StaticMesh = StaticMesh;
-	SphereBrushComponent->Materials.Add(BrushMaterial);
+	SphereBrushComponent->OverrideMaterials.Add(BrushMaterial);
 	SphereBrushComponent->SetAbsolute(true, true, true);
 
 	bBrushTraceValid = false;
diff --git a/Engine/Source/Editor/MaterialEditor/Private/SMaterialEditorViewport.cpp b/Engine/Source/Editor/MaterialEditor/Private/SMaterialEditorViewport.cpp
--- a/Engine/Source/Editor/MaterialEditor/Private/SMaterialEditorViewport.cpp
+++ b/Engine/Source/Editor/MaterialEditor/Private/SMaterialEditorViewport.cpp
@@ -207,8 +207,8 @@ void SMaterialEditorViewport::Construct(const FArguments& InArgs)
 
 SMaterialEditorViewport::~SMaterialEditorViewport()
 {
-	PreviewMeshComponent->Materials.Empty();
-	PreviewSkeletalMeshComponent->Materials.Empty();
+	PreviewMeshComponent->OverrideMaterials.Empty();
+	PreviewSkeletalMeshComponent->OverrideMaterials.Empty();
 
 	if (EditorViewportClient.IsValid())
 	{
@@ -322,10 +322,10 @@ void SMaterialEditorViewport::SetPreviewMaterial(UMaterialInterface* InMaterialI
 	check( PreviewMeshComponent );
 	check( PreviewSkeletalMeshComponent );
 
-	PreviewMeshComponent->Materials.Empty();
-	PreviewMeshComponent->Materials.Add( InMaterialInterface );
-	PreviewSkeletalMeshComponent->Materials.Empty();
-	PreviewSkeletalMeshComponent->Materials.Add( InMaterialInterface );
+	PreviewMeshComponent->OverrideMaterials.Empty();
+	PreviewMeshComponent->OverrideMaterials.Add( InMaterialInterface );
+	PreviewSkeletalMeshComponent->OverrideMaterials.Empty();
+	PreviewSkeletalMeshComponent->OverrideMaterials.Add( InMaterialInterface );
 }
 
 void SMaterialEditorViewport::ToggleRealtime()
diff --git a/Engine/Source/Editor/UnrealEd/Private/AssetSelection.cpp b/Engine/Source/Editor/UnrealEd/Private/AssetSelection.cpp
--- a/Engine/Source/Editor/UnrealEd/Private/AssetSelection.cpp
+++ b/Engine/Source/Editor/UnrealEd/Private/AssetSelection.cpp
@@ -747,7 +747,7 @@ bool FActorFactoryAssetProxy::ApplyMaterialToActor( AActor* TargetActor, UMateri
 						// OK, we need to figure out how many material slots this mesh component/static mesh has.
 						// Start with the actor's material count, then drill into the static/skeletal mesh to make sure 
 						// we have the right total.
-						int32 MaterialCount = FMath::Max(FoundMeshComponent->Materials.Num(), FoundMeshComponent->GetNumMaterials());
+						int32 MaterialCount = FMath::Max(FoundMeshComponent->OverrideMaterials.Num(), FoundMeshComponent->GetNumMaterials());
 
 						// Any materials to overwrite?
 						if( MaterialCount > 0 )
diff --git a/Engine/Source/Editor/UnrealEd/Private/Editor.cpp b/Engine/Source/Editor/UnrealEd/Private/Editor.cpp
--- a/Engine/Source/Editor/UnrealEd/Private/Editor.cpp
+++ b/Engine/Source/Editor/UnrealEd/Private/Editor.cpp
@@ -3378,7 +3378,7 @@ struct FConvertStaticMeshActorInfo
 	// Component properties.
 	UStaticMesh*						StaticMesh;
 	USkeletalMesh*						SkeletalMesh;
-	TArray<UMaterialInterface*>			Materials;
+	TArray<UMaterialInterface*>			OverrideMaterials;
 	TArray<FGuid>						IrrelevantLights;
 	float								CachedMaxDrawDistance;
 	bool								CastShadow;
@@ -3431,7 +3431,7 @@ struct FConvertStaticMeshActorInfo
 
 		// Copy over component properties.
 		StaticMesh				= MeshComp->StaticMesh;
-		Materials				= MeshComp->Materials;
+		OverrideMaterials		= MeshComp->OverrideMaterials;
 		IrrelevantLights		= MeshComp->IrrelevantLights;
 		CachedMaxDrawDistance	= MeshComp->CachedMaxDrawDistance;
 		CastShadow				= MeshComp->CastShadow;
@@ -3475,7 +3475,7 @@ struct FConvertStaticMeshActorInfo
 
 		// Set component properties.
 		if ( bComponentPropsDifferFromDefaults[0] ) MeshComp->StaticMesh			= StaticMesh;
-		if ( bComponentPropsDifferFromDefaults[1] ) MeshComp->Materials				= Materials;
+		if ( bComponentPropsDifferFromDefaults[1] ) MeshComp->OverrideMaterials		= OverrideMaterials;
 		if ( bComponentPropsDifferFromDefaults[2] ) MeshComp->IrrelevantLights		= IrrelevantLights;
 		if ( bComponentPropsDifferFromDefaults[3] ) MeshComp->CachedMaxDrawDistance	= CachedMaxDrawDistance;
 		if ( bComponentPropsDifferFromDefaults[4] ) MeshComp->CastShadow			= CastShadow;
@@ -3534,7 +3534,7 @@ struct FConvertStaticMeshActorInfo
 
 		// Copy over component properties.
 		SkeletalMesh			= MeshComp->SkeletalMesh;
-		Materials				= MeshComp->Materials;
+		OverrideMaterials		= MeshComp->OverrideMaterials;
 		CachedMaxDrawDistance	= MeshComp->CachedMaxDrawDistance;
 		CastShadow				= MeshComp->CastShadow;
 
@@ -3558,7 +3558,7 @@ struct FConvertStaticMeshActorInfo
 
 		// Set component properties.
 		if ( bComponentPropsDifferFromDefaults[0] ) MeshComp->SkeletalMesh			= SkeletalMesh;
-		if ( bComponentPropsDifferFromDefaults[1] ) MeshComp->Materials				= Materials;
+		if ( bComponentPropsDifferFromDefaults[1] ) MeshComp->OverrideMaterials		= OverrideMaterials;
 		if ( bComponentPropsDifferFromDefaults[3] ) MeshComp->CachedMaxDrawDistance	= CachedMaxDrawDistance;
 		if ( bComponentPropsDifferFromDefaults[4] ) MeshComp->CastShadow			= CastShadow;
 		if ( bComponentPropsDifferFromDefaults[5] ) MeshComp->BodyInstance.CopyBodyInstancePropertiesFrom(&BodyInstance);
diff --git a/Engine/Source/Editor/UnrealEd/Private/ThumbnailHelpers.cpp b/Engine/Source/Editor/UnrealEd/Private/ThumbnailHelpers.cpp
--- a/Engine/Source/Editor/UnrealEd/Private/ThumbnailHelpers.cpp
+++ b/Engine/Source/Editor/UnrealEd/Private/ThumbnailHelpers.cpp
@@ -381,7 +381,7 @@ FSkeletalMeshThumbnailScene::FSkeletalMeshThumbnailScene()
 
 void FSkeletalMeshThumbnailScene::SetSkeletalMesh(USkeletalMesh* InSkeletalMesh)
 {
-	PreviewActor->GetSkeletalMeshComponent()->Materials.Empty();
+	PreviewActor->GetSkeletalMeshComponent()->OverrideMaterials.Empty();
 	PreviewActor->GetSkeletalMeshComponent()->SetSkeletalMesh(InSkeletalMesh);
 
 	if ( InSkeletalMesh )
@@ -524,7 +524,7 @@ FAnimationSequenceThumbnailScene::FAnimationSequenceThumbnailScene()
 
 bool FAnimationSequenceThumbnailScene::SetAnimation(UAnimSequenceBase* InAnimation)
 {
-	PreviewActor->GetSkeletalMeshComponent()->Materials.Empty();
+	PreviewActor->GetSkeletalMeshComponent()->OverrideMaterials.Empty();
 
 	bool bSetSucessfully = false;
 
@@ -648,7 +648,7 @@ FBlendSpaceThumbnailScene::FBlendSpaceThumbnailScene()
 
 bool FBlendSpaceThumbnailScene::SetBlendSpace(class UBlendSpaceBase* InBlendSpace)
 {
-	PreviewActor->GetSkeletalMeshComponent()->Materials.Empty();
+	PreviewActor->GetSkeletalMeshComponent()->OverrideMaterials.Empty();
 
 	bool bSetSucessfully = false;
 
@@ -774,7 +774,7 @@ FAnimBlueprintThumbnailScene::FAnimBlueprintThumbnailScene()
 
 bool FAnimBlueprintThumbnailScene::SetAnimBlueprint(class UAnimBlueprint* InBlueprint)
 {
-	PreviewActor->GetSkeletalMeshComponent()->Materials.Empty();
+	PreviewActor->GetSkeletalMeshComponent()->OverrideMaterials.Empty();
 
 	bool bSetSucessfully = false;
 
diff --git a/Engine/Source/Runtime/Engine/Classes/Components/MeshComponent.h b/Engine/Source/Runtime/Engine/Classes/Components/MeshComponent.h
--- a/Engine/Source/Runtime/Engine/Classes/Components/MeshComponent.h
+++ b/Engine/Source/Runtime/Engine/Classes/Components/MeshComponent.h
@@ -21,11 +21,13 @@ class ENGINE_API UMeshComponent : public UPrimitiveComponent
 
 	/** Per-Component material overrides.  These must NOT be set directly or a race condition can occur between GC and the rendering thread. */
 	UPROPERTY(EditAnywhere, Category=Rendering, Meta=(ToolTip="Material overrides."))
-	TArray<class UMaterialInterface*> Materials;
+	TArray<class UMaterialInterface*> OverrideMaterials;
 	
 	UFUNCTION(BlueprintCallable, Category="Components|Mesh")
 	virtual TArray<class UMaterialInterface*> GetMaterials() const;
 
+	/** Returns override Materials count */
+	virtual int32 GetNumOverrideMaterials() const;
 
 	// Begin UObject Interface
 	virtual void BeginDestroy() override;
diff --git a/Engine/Source/Runtime/Engine/Private/MeshComponent.cpp b/Engine/Source/Runtime/Engine/Private/MeshComponent.cpp
--- a/Engine/Source/Runtime/Engine/Private/MeshComponent.cpp
+++ b/Engine/Source/Runtime/Engine/Private/MeshComponent.cpp
@@ -21,9 +21,9 @@ void UMeshComponent::BeginDestroy()
 
 UMaterialInterface* UMeshComponent::GetMaterial(int32 ElementIndex) const
 {
-	if (Materials.IsValidIndex(ElementIndex))
+	if (OverrideMaterials.IsValidIndex(ElementIndex))
 	{
-		return Materials[ElementIndex];
+		return OverrideMaterials[ElementIndex];
 	}
 	else
 	{
@@ -35,20 +35,20 @@ void UMeshComponent::SetMaterial(int32 ElementIndex, UMaterialInterface* Materia
 {
 	if (ElementIndex >= 0)
 	{
-		if (Materials.IsValidIndex(ElementIndex) && (Materials[ElementIndex] == Material))
+		if (OverrideMaterials.IsValidIndex(ElementIndex) && (OverrideMaterials[ElementIndex] == Material))
 		{
 			// Do nothing, the material is already set
 		}
 		else
 		{
 			// Grow the array if the new index is too large
-			if (Materials.Num() <= ElementIndex)
+			if (OverrideMaterials.Num() <= ElementIndex)
 			{
-				Materials.AddZeroed(ElementIndex + 1 - Materials.Num());
+				OverrideMaterials.AddZeroed(ElementIndex + 1 - OverrideMaterials.Num());
 			}
 
 			// Set the material and invalidate things
-			Materials[ElementIndex] = Material;
+			OverrideMaterials[ElementIndex] = Material;
 			MarkRenderStateDirty();
 
 			if (BodyInstance.IsValidBodyInstance())
@@ -75,6 +75,11 @@ FMaterialRelevance UMeshComponent::GetMaterialRelevance(ERHIFeatureLevel::Type I
 	return Result;
 }
 
+int32 UMeshComponent::GetNumOverrideMaterials() const
+{
+	return OverrideMaterials.Num();
+}
+
 int32 UMeshComponent::GetNumMaterials() const
 {
 	return 0;
@@ -132,5 +137,5 @@ void UMeshComponent::SetTextureForceResidentFlag( bool bForceMiplevelsToBeReside
 
 TArray<class UMaterialInterface*> UMeshComponent::GetMaterials() const
 {
-	return Materials;
+	return OverrideMaterials;
 }
diff --git a/Engine/Source/Runtime/Engine/Private/SkinnedMeshComponent.cpp b/Engine/Source/Runtime/Engine/Private/SkinnedMeshComponent.cpp
--- a/Engine/Source/Runtime/Engine/Private/SkinnedMeshComponent.cpp
+++ b/Engine/Source/Runtime/Engine/Private/SkinnedMeshComponent.cpp
@@ -316,27 +316,23 @@ void USkinnedMeshComponent::UpdateSlaveComponent()
 
 TArray<class UMaterialInterface*> USkinnedMeshComponent::GetMaterials() const
 {
-	TArray<class UMaterialInterface*> Materials = Super::GetMaterials();
+	TArray<class UMaterialInterface*> OutMaterials = Super::GetMaterials();
 
 	// if no material is overriden, look for mesh material;
-	if(Materials.Num() == 0 && SkeletalMesh)
+	if(OutMaterials.Num() == 0 && SkeletalMesh)
 	{
 		for (auto Material : SkeletalMesh->Materials)
 		{
-			Materials.Add(Material.MaterialInterface);
+			OutMaterials.Add(Material.MaterialInterface);
 		}
 	}
 
-	return Materials;
+	return OutMaterials;
 }
 
 int32 USkinnedMeshComponent::GetNumMaterials() const
 {
-	if (Materials.Num() > 0)
-	{
-		return Materials.Num();
-	}
-	else if (SkeletalMesh)
+	if (SkeletalMesh)
 	{
 		return SkeletalMesh->Materials.Num();
 	}
@@ -346,9 +342,9 @@ int32 USkinnedMeshComponent::GetNumMaterials() const
 
 UMaterialInterface* USkinnedMeshComponent::GetMaterial(int32 MaterialIndex) const
 {
-	if(MaterialIndex < Materials.Num() && Materials[MaterialIndex])
+	if(MaterialIndex < OverrideMaterials.Num() && OverrideMaterials[MaterialIndex])
 	{
-		return Materials[MaterialIndex];
+		return OverrideMaterials[MaterialIndex];
 	}
 	else if (SkeletalMesh && MaterialIndex < SkeletalMesh->Materials.Num() && SkeletalMesh->Materials[MaterialIndex].MaterialInterface)
 	{
@@ -370,7 +366,7 @@ void USkinnedMeshComponent::GetStreamingTextureInfo(TArray<FStreamingTexturePrim
 	{
 		const float LocalTexelFactor = SkeletalMesh->GetStreamingTextureFactor(0) * StreamingDistanceMultiplier;
 		const float WorldTexelFactor = LocalTexelFactor * ComponentToWorld.GetMaximumAxisScale();
-		const int32 NumMaterials = FMath::Max(SkeletalMesh->Materials.Num(), Materials.Num());
+		const int32 NumMaterials = FMath::Max(SkeletalMesh->Materials.Num(), OverrideMaterials.Num());
 		for( int32 MatIndex = 0; MatIndex < NumMaterials; MatIndex++ )
 		{
 			UMaterialInterface* const MaterialInterface = GetMaterial(MatIndex);
@@ -1354,7 +1350,7 @@ void USkinnedMeshComponent::GetUsedMaterials( TArray<UMaterialInterface*>& OutMa
 	if( SkeletalMesh )
 	{
 		// The max number of materials used is the max of the materials on the skeletal mesh and the materials on the mesh component
-		const int32 NumMaterials = FMath::Max( SkeletalMesh->Materials.Num(), Materials.Num() );
+		const int32 NumMaterials = FMath::Max( SkeletalMesh->Materials.Num(), OverrideMaterials.Num() );
 		for( int32 MatIdx = 0; MatIdx < NumMaterials; ++MatIdx )
 		{
 			// GetMaterial will determine the correct material to use for this index.  
diff --git a/Engine/Source/Runtime/Engine/Private/StaticMeshComponent.cpp b/Engine/Source/Runtime/Engine/Private/StaticMeshComponent.cpp
--- a/Engine/Source/Runtime/Engine/Private/StaticMeshComponent.cpp
+++ b/Engine/Source/Runtime/Engine/Private/StaticMeshComponent.cpp
@@ -341,10 +341,10 @@ void UStaticMeshComponent::CheckForErrors()
 			}
 		}
 
-		if (Materials.Num() > StaticMesh->Materials.Num())
+		if (OverrideMaterials.Num() > StaticMesh->Materials.Num())
 		{
 			FFormatNamedArguments Arguments;
-			Arguments.Add(TEXT("OverridenCount"), Materials.Num());
+			Arguments.Add(TEXT("OverridenCount"), OverrideMaterials.Num());
 			Arguments.Add(TEXT("ReferencedCount"), StaticMesh->Materials.Num());
 			Arguments.Add(TEXT("MeshName"), FText::FromString(StaticMesh->GetName()));
 			FMessageLog("MapCheck").Warning()
@@ -1145,7 +1145,7 @@ void UStaticMeshComponent::PostLoad()
 
 #if WITH_EDITORONLY_DATA
 	// Remap the materials array if the static mesh materials may have been remapped to remove zero triangle sections.
-	if (StaticMesh && GetLinkerUE4Version() < VER_UE4_REMOVE_ZERO_TRIANGLE_SECTIONS && Materials.Num())
+	if (StaticMesh && GetLinkerUE4Version() < VER_UE4_REMOVE_ZERO_TRIANGLE_SECTIONS && OverrideMaterials.Num())
 	{
 		StaticMesh->ConditionalPostLoad();
 		if (StaticMesh->HasValidRenderData()
@@ -1154,8 +1154,8 @@ void UStaticMeshComponent::PostLoad()
 			TArray<UMaterialInterface*> OldMaterials;
 			const TArray<int32>& MaterialIndexToImportIndex = StaticMesh->RenderData->MaterialIndexToImportIndex;
 
-			Exchange(Materials,OldMaterials);
-			Materials.Empty(MaterialIndexToImportIndex.Num());
+			Exchange(OverrideMaterials,OldMaterials);
+			OverrideMaterials.Empty(MaterialIndexToImportIndex.Num());
 			for (int32 MaterialIndex = 0; MaterialIndex < MaterialIndexToImportIndex.Num(); ++MaterialIndex)
 			{
 				UMaterialInterface* Material = NULL;
@@ -1164,13 +1164,13 @@ void UStaticMeshComponent::PostLoad()
 				{
 					Material = OldMaterials[OldMaterialIndex];
 				}
-				Materials.Add(Material);
+				OverrideMaterials.Add(Material);
 			}
 		}
 
-		if (Materials.Num() > StaticMesh->Materials.Num())
+		if (OverrideMaterials.Num() > StaticMesh->Materials.Num())
 		{
-			Materials.RemoveAt(StaticMesh->Materials.Num(), Materials.Num() - StaticMesh->Materials.Num());
+			OverrideMaterials.RemoveAt(StaticMesh->Materials.Num(), OverrideMaterials.Num() - StaticMesh->Materials.Num());
 		}
 	}
 #endif // #if WITH_EDITORONLY_DATA
@@ -1470,23 +1470,23 @@ int32 UStaticMeshComponent::GetNumMaterials() const
 
 TArray<class UMaterialInterface*> UStaticMeshComponent::GetMaterials() const
 {
-	TArray<class UMaterialInterface*> Materials = Super::GetMaterials();
+	TArray<class UMaterialInterface*> OutMaterials = Super::GetMaterials();
 
 	// if no material is overriden, look for mesh material;
-	if(Materials.Num() == 0 && StaticMesh)
+	if(OutMaterials.Num() == 0 && StaticMesh)
 	{
-		Materials = StaticMesh->Materials;
+		OutMaterials = StaticMesh->Materials;
 	}
 
-	return Materials;
+	return OutMaterials;
 }
 
 UMaterialInterface* UStaticMeshComponent::GetMaterial(int32 MaterialIndex) const
 {
 	// If we have a base materials array, use that
-	if(MaterialIndex < Materials.Num() && Materials[MaterialIndex])
+	if(MaterialIndex < OverrideMaterials.Num() && OverrideMaterials[MaterialIndex])
 	{
-		return Materials[MaterialIndex];
+		return OverrideMaterials[MaterialIndex];
 	}
 	// Otherwise get from static mesh
 	else
diff --git a/Engine/Source/Runtime/Landscape/Private/LandscapeSplines.cpp b/Engine/Source/Runtime/Landscape/Private/LandscapeSplines.cpp
--- a/Engine/Source/Runtime/Landscape/Private/LandscapeSplines.cpp
+++ b/Engine/Source/Runtime/Landscape/Private/LandscapeSplines.cpp
@@ -812,10 +812,10 @@ void ULandscapeSplineControlPoint::UpdateSplinePoints(bool bUpdateCollision, boo
 
 	if (MeshComponent != NULL)
 	{
-		if (MeshComponent->Materials != MaterialOverrides)
+		if (MeshComponent->OverrideMaterials != MaterialOverrides)
 		{
 			MeshComponent->Modify();
-			MeshComponent->Materials = MaterialOverrides;
+			MeshComponent->OverrideMaterials = MaterialOverrides;
 			MeshComponent->MarkRenderStateDirty();
 			if (MeshComponent->BodyInstance.IsValidBodyInstance())
 			{
@@ -1391,7 +1391,7 @@ void ULandscapeSplineSegment::UpdateSplinePoints(bool bUpdateCollision)
 
 			MeshComponent->SetStaticMesh(Mesh);
 
-			MeshComponent->Materials = MeshEntry->MaterialOverrides;
+			MeshComponent->OverrideMaterials = MeshEntry->MaterialOverrides;
 			MeshComponent->MarkRenderStateDirty();
 			if (MeshComponent->BodyInstance.IsValidBodyInstance())
 			{
