Make ULocalPlayer::ControllerId private
Const versions of many world context related functions
Remove UPROPERTY from a number of ULocalPlayer members

diff --git a/Engine/Source/Runtime/Engine/Classes/Engine/Engine.h b/Engine/Source/Runtime/Engine/Classes/Engine/Engine.h
--- a/Engine/Source/Runtime/Engine/Classes/Engine/Engine.h
+++ b/Engine/Source/Runtime/Engine/Classes/Engine/Engine.h
@@ -1780,8 +1780,8 @@ public:
 	TArray<class ULocalPlayer*>::TConstIterator	GetLocalPlayerIterator(UWorld *World);
 	TArray<class ULocalPlayer*>::TConstIterator GetLocalPlayerIterator(const UGameViewportClient *Viewport);
 
-	const TArray<class ULocalPlayer*>& GetGamePlayers(UWorld *World);
-	const TArray<class ULocalPlayer*>& GetGamePlayers(const UGameViewportClient *Viewport);
+	const TArray<class ULocalPlayer*>& GetGamePlayers(UWorld *World) const;
+	const TArray<class ULocalPlayer*>& GetGamePlayers(const UGameViewportClient *Viewport) const;
 
 	/**
 	 *	Returns the first ULocalPlayer that matches the given ControllerId. 
@@ -2042,10 +2042,10 @@ public:
 	 * @param	ControllerId	the game pad index of the player to search for
 	 * @return	The player that has the ControllerId specified, or nullptr if no players have that ControllerId
 	 */
-	ULocalPlayer* GetLocalPlayerFromControllerId( const UGameViewportClient* InViewport, int32 ControllerId );
-	ULocalPlayer* GetLocalPlayerFromControllerId( UWorld * InWorld, int32 ControllerId );
+	ULocalPlayer* GetLocalPlayerFromControllerId( const UGameViewportClient* InViewport, const int32 ControllerId ) const;
+	ULocalPlayer* GetLocalPlayerFromControllerId( UWorld * InWorld, const int32 ControllerId ) const;
 
-	void SwapControllerId(ULocalPlayer *NewPlayer, int32 CurrentControllerId, int32 NewControllerID);
+	void SwapControllerId(ULocalPlayer *NewPlayer, const int32 CurrentControllerId, const int32 NewControllerID) const;
 
 	/** 
 	 * Find a Local Player Controller, which may not exist at all if this is a server.
@@ -2405,14 +2405,28 @@ public:
 	FWorldContext* GetWorldContextFromHandle(const FName WorldContextHandle);
 	FWorldContext* GetWorldContextFromPIEInstance(const int32 PIEInstance);
 
-	FWorldContext& GetWorldContextFromWorldChecked(UWorld * InWorld);
+	const FWorldContext* GetWorldContextFromWorld(const UWorld* InWorld) const;
+	const FWorldContext* GetWorldContextFromGameViewport(const UGameViewportClient *InViewport) const;
+	const FWorldContext* GetWorldContextFromPendingNetGame(const UPendingNetGame *InPendingNetGame) const;	
+	const FWorldContext* GetWorldContextFromPendingNetGameNetDriver(const UNetDriver *InPendingNetGame) const;	
+	const FWorldContext* GetWorldContextFromHandle(const FName WorldContextHandle) const;
+	const FWorldContext* GetWorldContextFromPIEInstance(const int32 PIEInstance) const;
+
+	FWorldContext& GetWorldContextFromWorldChecked(const UWorld * InWorld);
 	FWorldContext& GetWorldContextFromGameViewportChecked(const UGameViewportClient *InViewport);
 	FWorldContext& GetWorldContextFromPendingNetGameChecked(const UPendingNetGame *InPendingNetGame);	
 	FWorldContext& GetWorldContextFromPendingNetGameNetDriverChecked(const UNetDriver *InPendingNetGame);	
 	FWorldContext& GetWorldContextFromHandleChecked(const FName WorldContextHandle);
 	FWorldContext& GetWorldContextFromPIEInstanceChecked(const int32 PIEInstance);
 
-	const TIndirectArray<FWorldContext>& GetWorldContexts() { return WorldList;	}
+	const FWorldContext& GetWorldContextFromWorldChecked(const UWorld * InWorld) const;
+	const FWorldContext& GetWorldContextFromGameViewportChecked(const UGameViewportClient *InViewport) const;
+	const FWorldContext& GetWorldContextFromPendingNetGameChecked(const UPendingNetGame *InPendingNetGame) const;	
+	const FWorldContext& GetWorldContextFromPendingNetGameNetDriverChecked(const UNetDriver *InPendingNetGame) const;	
+	const FWorldContext& GetWorldContextFromHandleChecked(const FName WorldContextHandle) const;
+	const FWorldContext& GetWorldContextFromPIEInstanceChecked(const int32 PIEInstance) const;
+
+	const TIndirectArray<FWorldContext>& GetWorldContexts() const { return WorldList;	}
 
 
 	/** Verify any remaining World(s) are valid after ::LoadMap destroys a world */
@@ -2424,7 +2438,7 @@ public:
 	bool ShouldAbsorbAuthorityOnlyEvent();
 	bool ShouldAbsorbCosmeticOnlyEvent();
 
-	UGameViewportClient* GameViewportForWorld(UWorld *InWorld);
+	UGameViewportClient* GameViewportForWorld(const UWorld *InWorld) const;
 
 protected:
 
diff --git a/Engine/Source/Runtime/Engine/Classes/Engine/GameInstance.h b/Engine/Source/Runtime/Engine/Classes/Engine/GameInstance.h
--- a/Engine/Source/Runtime/Engine/Classes/Engine/GameInstance.h
+++ b/Engine/Source/Runtime/Engine/Classes/Engine/GameInstance.h
@@ -143,7 +143,7 @@ public:
 	int32					GetNumLocalPlayers() const;
 	ULocalPlayer*			GetLocalPlayerByIndex(const int32 Index) const;
 	APlayerController*		GetFirstLocalPlayerController() const;
-	ULocalPlayer*			FindLocalPlayerFromControllerId(int32 ControllerId) const;
+	ULocalPlayer*			FindLocalPlayerFromControllerId(const int32 ControllerId) const;
 	ULocalPlayer*			FindLocalPlayerFromUniqueNetId(TSharedPtr< FUniqueNetId > UniqueNetId) const;
 	ULocalPlayer*			FindLocalPlayerFromUniqueNetId(const FUniqueNetId& UniqueNetId) const;
 	ULocalPlayer*			GetFirstGamePlayer() const;
diff --git a/Engine/Source/Runtime/Engine/Classes/Engine/LocalPlayer.h b/Engine/Source/Runtime/Engine/Classes/Engine/LocalPlayer.h
--- a/Engine/Source/Runtime/Engine/Classes/Engine/LocalPlayer.h
+++ b/Engine/Source/Runtime/Engine/Classes/Engine/LocalPlayer.h
@@ -134,10 +134,6 @@ class ENGINE_API ULocalPlayer : public UPlayer
 {
 	GENERATED_UCLASS_BODY()
 
-	/** The controller ID which this player accepts input from. */
-	UPROPERTY()
-	int32 ControllerId;
-
 	/** The FUniqueNetId which this player is associated with. */
 	TSharedPtr<class FUniqueNetId> CachedUniqueNetId;
 
@@ -146,15 +142,12 @@ class ENGINE_API ULocalPlayer : public UPlayer
 	class UGameViewportClient* ViewportClient;
 
 	/** The coordinates for the upper left corner of the master viewport subregion allocated to this player. 0-1 */
-	UPROPERTY()
 	FVector2D Origin;
 
 	/** The size of the master viewport subregion allocated to this player. 0-1 */
-	UPROPERTY()
 	FVector2D Size;
 
 	/** The location of the player's view the previous frame. */
-	UPROPERTY(transient)
 	FVector LastViewLocation;
 
 	/** How to constrain perspective viewport FOV */
@@ -174,12 +167,15 @@ private:
 	FSceneViewStateReference StereoViewState;
 
 	/** Class to manage online services */
-	UPROPERTY(transient)
+	UPROPERTY()
 	class UOnlineSession* OnlineSession;
 
 	/** @return OnlineSession class to use for this player controller  */
 	virtual TSubclassOf<UOnlineSession> GetOnlineSessionClass();
 
+	/** The controller ID which this player accepts input from. */
+	int32 ControllerId;
+
 public:
 	// UObject interface
 	virtual void PostInitProperties() override;
@@ -308,6 +304,11 @@ public:
 	 */
 	virtual void SetControllerId(int32 NewControllerId);
 
+	/**
+	 * Returns the controller ID for the player
+	 */
+	int32 GetControllerId() const { return ControllerId; }
+
 	/** 
 	 * Retrieves this player's name/tag from the online subsystem
 	 * if this function returns a non-empty string, the returned name will replace the "Name" URL parameter
diff --git a/Engine/Source/Runtime/Engine/Private/GameInstance.cpp b/Engine/Source/Runtime/Engine/Private/GameInstance.cpp
--- a/Engine/Source/Runtime/Engine/Private/GameInstance.cpp
+++ b/Engine/Source/Runtime/Engine/Private/GameInstance.cpp
@@ -485,7 +485,7 @@ bool UGameInstance::RemoveLocalPlayer(ULocalPlayer* ExistingPlayer)
 	// Do this after notifications, as some of them require the ViewportClient.
 	ExistingPlayer->ViewportClient = NULL;
 
-	UE_LOG(LogPlayerManagement, Log, TEXT("UGameInstance::RemovePlayer: Removed player %s with ControllerId %i at index %i (%i remaining players)"), *ExistingPlayer->GetName(), ExistingPlayer->ControllerId, OldIndex, LocalPlayers.Num());
+	UE_LOG(LogPlayerManagement, Log, TEXT("UGameInstance::RemovePlayer: Removed player %s with ControllerId %i at index %i (%i remaining players)"), *ExistingPlayer->GetName(), ExistingPlayer->GetControllerId(), OldIndex, LocalPlayers.Num());
 
 	return true;
 }
@@ -538,11 +538,11 @@ APlayerController* UGameInstance::GetFirstLocalPlayerController() const
 	return nullptr;
 }
 
-ULocalPlayer* UGameInstance::FindLocalPlayerFromControllerId(int32 ControllerId) const
+ULocalPlayer* UGameInstance::FindLocalPlayerFromControllerId(const int32 ControllerId) const
 {
 	for (ULocalPlayer * LP : LocalPlayers)
 	{
-		if (LP && (LP->ControllerId == ControllerId))
+		if (LP && (LP->GetControllerId() == ControllerId))
 		{
 			return LP;
 		}
diff --git a/Engine/Source/Runtime/Engine/Private/GameViewportClient.cpp b/Engine/Source/Runtime/Engine/Private/GameViewportClient.cpp
--- a/Engine/Source/Runtime/Engine/Private/GameViewportClient.cpp
+++ b/Engine/Source/Runtime/Engine/Private/GameViewportClient.cpp
@@ -1325,13 +1325,13 @@ void UGameViewportClient::PeekNetworkFailureMessages(UWorld *InWorld, UNetDriver
 void UGameViewportClient::SSSwapControllers()
 {
 #if !UE_BUILD_SHIPPING
-	int32 TmpControllerID = GetOuterUEngine()->GetFirstGamePlayer(this)->ControllerId;
+	const int32 TmpControllerID = GetOuterUEngine()->GetFirstGamePlayer(this)->GetControllerId();
 
 	for (int32 Idx=0; Idx<GetOuterUEngine()->GetNumGamePlayers(this)-1; ++Idx)
 	{
-		GetOuterUEngine()->GetGamePlayer(this, Idx)->ControllerId = GetOuterUEngine()->GetGamePlayer(this, Idx+1)->ControllerId;
+		GetOuterUEngine()->GetGamePlayer(this, Idx)->SetControllerId(GetOuterUEngine()->GetGamePlayer(this, Idx+1)->GetControllerId());
 	}
-	GetOuterUEngine()->GetGamePlayer(this, GetOuterUEngine()->GetNumGamePlayers(this)-1)->ControllerId = TmpControllerID;
+	GetOuterUEngine()->GetGamePlayer(this, GetOuterUEngine()->GetNumGamePlayers(this)-1)->SetControllerId(TmpControllerID);
 #endif
 }
 
diff --git a/Engine/Source/Runtime/Engine/Private/KismetSystemLibrary.cpp b/Engine/Source/Runtime/Engine/Private/KismetSystemLibrary.cpp
--- a/Engine/Source/Runtime/Engine/Private/KismetSystemLibrary.cpp
+++ b/Engine/Source/Runtime/Engine/Private/KismetSystemLibrary.cpp
@@ -2698,7 +2698,7 @@ void UKismetSystemLibrary::ShowPlatformSpecificAchievementsScreen(class APlayerC
 			ULocalPlayer* LocalPlayer = Cast<ULocalPlayer>(SpecificPlayer->Player);
 			if(LocalPlayer)
 			{
-				LocalUserNum = LocalPlayer->ControllerId;
+				LocalUserNum = LocalPlayer->GetControllerId();
 			}
 		}
 		ExternalUI->ShowAchievementsUI(LocalUserNum);
diff --git a/Engine/Source/Runtime/Engine/Private/PlayerController.cpp b/Engine/Source/Runtime/Engine/Private/PlayerController.cpp
--- a/Engine/Source/Runtime/Engine/Private/PlayerController.cpp
+++ b/Engine/Source/Runtime/Engine/Private/PlayerController.cpp
@@ -1340,7 +1340,7 @@ void APlayerController::EndPlay(const EEndPlayReason::Type EndPlayReason)
 		IForceFeedbackSystem* ForceFeedbackSystem = FSlateApplication::Get().GetForceFeedbackSystem();
 		if (ForceFeedbackSystem)
 		{
-			ForceFeedbackSystem->SetChannelValues(LocalPlayer->ControllerId, FForceFeedbackValues());
+			ForceFeedbackSystem->SetChannelValues(LocalPlayer->GetControllerId(), FForceFeedbackValues());
 		}
 	}
 
@@ -3009,11 +3009,11 @@ void APlayerController::ToggleSpeaking(bool bSpeaking)
 		{
 			if (bSpeaking)
 			{
-				VoiceInt->StartNetworkedVoice(LP->ControllerId);
+				VoiceInt->StartNetworkedVoice(LP->GetControllerId());
 			}
 			else
 			{
-				VoiceInt->StopNetworkedVoice(LP->ControllerId);
+				VoiceInt->StopNetworkedVoice(LP->GetControllerId());
 			}
 		}
 	}
@@ -3505,7 +3505,7 @@ void APlayerController::ProcessForceFeedback(const float DeltaTime, const bool b
 	IForceFeedbackSystem* ForceFeedbackSystem = FSlateApplication::Get().GetForceFeedbackSystem();
 	if (ForceFeedbackSystem)
 	{
-		ForceFeedbackSystem->SetChannelValues(CastChecked<ULocalPlayer>(Player)->ControllerId, (bForceFeedbackEnabled ? ForceFeedbackValues : FForceFeedbackValues()));
+		ForceFeedbackSystem->SetChannelValues(CastChecked<ULocalPlayer>(Player)->GetControllerId(), (bForceFeedbackEnabled ? ForceFeedbackValues : FForceFeedbackValues()));
 	}
 }
 
diff --git a/Engine/Source/Runtime/Engine/Private/PlayerMuteList.cpp b/Engine/Source/Runtime/Engine/Private/PlayerMuteList.cpp
--- a/Engine/Source/Runtime/Engine/Private/PlayerMuteList.cpp
+++ b/Engine/Source/Runtime/Engine/Private/PlayerMuteList.cpp
@@ -107,7 +107,7 @@ void FPlayerMuteList::ClientMutePlayer(APlayerController* OwningPC, const FUniqu
 		if (VoiceInt.IsValid())
 		{
 			// Have the voice subsystem mute this player
-			VoiceInt->MuteRemoteTalker(LP->ControllerId, *PlayerIdToMute, false);
+			VoiceInt->MuteRemoteTalker(LP->GetControllerId(), *PlayerIdToMute, false);
 		}
 	}
 }
@@ -128,7 +128,7 @@ void FPlayerMuteList::ClientUnmutePlayer(APlayerController* OwningPC, const FUni
 		if (VoiceInt.IsValid())
 		{
 			// Have the voice subsystem mute this player
-			VoiceInt->UnmuteRemoteTalker(LP->ControllerId, *PlayerIdToUnmute, false);
+			VoiceInt->UnmuteRemoteTalker(LP->GetControllerId(), *PlayerIdToUnmute, false);
 		}
 	}
 }
diff --git a/Engine/Source/Runtime/Engine/Private/UnrealEngine.cpp b/Engine/Source/Runtime/Engine/Private/UnrealEngine.cpp
--- a/Engine/Source/Runtime/Engine/Private/UnrealEngine.cpp
+++ b/Engine/Source/Runtime/Engine/Private/UnrealEngine.cpp
@@ -1792,41 +1792,40 @@ bool UEngine::IsStereoscopic3D()
 	return !GIsEditor && StereoRenderingDevice.IsValid() && StereoRenderingDevice->IsStereoEnabled();
 }
 
-ULocalPlayer* GetLocalPlayerFromControllerId_local(const TArray<class ULocalPlayer*>& GamePlayers, int32 ControllerId)
+ULocalPlayer* GetLocalPlayerFromControllerId_local(const TArray<class ULocalPlayer*>& GamePlayers, const int32 ControllerId)
 {
-	for ( int32 PlayerIndex = 0; PlayerIndex < GamePlayers.Num(); PlayerIndex++ )
+	for ( ULocalPlayer* const Player : GamePlayers )
 	{
-		ULocalPlayer* const Player = GamePlayers[PlayerIndex];
-		if ( Player && Player->ControllerId == ControllerId )
+		if ( Player && Player->GetControllerId() == ControllerId )
 		{
 			return Player;
 		}
 	}
 
-	return NULL;
+	return nullptr;
 }
 
-ULocalPlayer* UEngine::GetLocalPlayerFromControllerId( const UGameViewportClient* InViewport, int32 ControllerId )
+ULocalPlayer* UEngine::GetLocalPlayerFromControllerId( const UGameViewportClient* InViewport, const int32 ControllerId ) const
 {
-	if (GetWorldContextFromGameViewport(InViewport) != NULL)
+	if (GetWorldContextFromGameViewport(InViewport) != nullptr)
 	{
 		const TArray<class ULocalPlayer*>& GamePlayers = GetGamePlayers(InViewport);
 		return GetLocalPlayerFromControllerId_local(GamePlayers, ControllerId);
 	}
-	return NULL;
+	return nullptr;
 }
 
-ULocalPlayer* UEngine::GetLocalPlayerFromControllerId( UWorld * InWorld, int32 ControllerId )
+ULocalPlayer* UEngine::GetLocalPlayerFromControllerId( UWorld * InWorld, const int32 ControllerId ) const
 {
 	const TArray<class ULocalPlayer*>& GamePlayers = GetGamePlayers(InWorld);
 	return GetLocalPlayerFromControllerId_local(GamePlayers, ControllerId);
 }
 
-void UEngine::SwapControllerId(ULocalPlayer *NewPlayer, int32 CurrentControllerId, int32 NewControllerID)
+void UEngine::SwapControllerId(ULocalPlayer *NewPlayer, const int32 CurrentControllerId, const int32 NewControllerID) const
 {
-	for (auto It = WorldList.CreateIterator(); It; ++It)
+	for (auto It = WorldList.CreateConstIterator(); It; ++It)
 	{
-		if (It->OwningGameInstance == NULL)
+		if (It->OwningGameInstance == nullptr)
 		{
 			continue;
 		}
@@ -1836,11 +1835,11 @@ void UEngine::SwapControllerId(ULocalPlayer *NewPlayer, int32 CurrentControllerI
 		if (LocalPlayers.Contains(NewPlayer))
 		{
 			// This is the world context that NewPlayer belongs to, see if anyone is using his CurrentControllerId
-			for (int32 i=0; i < LocalPlayers.Num(); ++i)
+			for (ULocalPlayer* LocalPlayer : LocalPlayers)
 			{
-				if(LocalPlayers[i] && LocalPlayers[i]->ControllerId == NewControllerID)
+				if(LocalPlayer && LocalPlayer->GetControllerId() == NewControllerID)
 				{
-					LocalPlayers[i]->ControllerId = CurrentControllerId;
+					LocalPlayer->SetControllerId( CurrentControllerId);
 					return;
 				}
 			}
@@ -7263,7 +7262,7 @@ const TArray<class ULocalPlayer*>& HandleFakeLocalPlayersList()
 	return FakeEmptyLocalPlayers;
 }
 
-const TArray<class ULocalPlayer*>& UEngine::GetGamePlayers(UWorld *World)
+const TArray<class ULocalPlayer*>& UEngine::GetGamePlayers(UWorld *World) const
 {
 	const FWorldContext &Context = GetWorldContextFromWorldChecked(World);
 	if ( Context.OwningGameInstance == NULL )
@@ -7273,7 +7272,7 @@ const TArray<class ULocalPlayer*>& UEngine::GetGamePlayers(UWorld *World)
 	return Context.OwningGameInstance->GetLocalPlayers();
 }
 	
-const TArray<class ULocalPlayer*>& UEngine::GetGamePlayers(const UGameViewportClient *Viewport)
+const TArray<class ULocalPlayer*>& UEngine::GetGamePlayers(const UGameViewportClient *Viewport) const
 {
 	const FWorldContext &Context = GetWorldContextFromGameViewportChecked(Viewport);
 	if ( Context.OwningGameInstance == NULL )
@@ -7293,17 +7292,17 @@ ULocalPlayer* UEngine::FindFirstLocalPlayerFromControllerId(int32 ControllerId)
 			const TArray<class ULocalPlayer*> & LocalPlayers = Context.OwningGameInstance->GetLocalPlayers();
 
 			// Use this world context, look for the ULocalPlayer with this ControllerId
-			for (int32 i=0; i < LocalPlayers.Num(); ++i)
+			for (ULocalPlayer* LocalPlayer : LocalPlayers)
 			{
-				if (LocalPlayers[i] && LocalPlayers[i]->ControllerId == ControllerId)
+				if (LocalPlayer && LocalPlayer->GetControllerId() == ControllerId)
 				{
-					return LocalPlayers[i];
+					return LocalPlayer;
 				}
 			}
 		}
 	}
 
-	return NULL;
+	return nullptr;
 }
 
 int32 UEngine::GetNumGamePlayers(UWorld *InWorld)
@@ -9242,7 +9241,7 @@ FWorldContext& HandleInvalidWorldContext()
 	return GEngine->CreateNewWorldContext(EWorldType::None);
 }
 
-FWorldContext* UEngine::GetWorldContextFromHandle(FName WorldContextHandle)
+FWorldContext* UEngine::GetWorldContextFromHandle(const FName WorldContextHandle)
 {
 	for (FWorldContext& WorldContext : WorldList)
 	{
@@ -9251,10 +9250,22 @@ FWorldContext* UEngine::GetWorldContextFromHandle(FName WorldContextHandle)
 			return &WorldContext;
 		}
 	}
-	return NULL;
+	return nullptr;
+}
+
+const FWorldContext* UEngine::GetWorldContextFromHandle(const FName WorldContextHandle) const
+{
+	for (const FWorldContext& WorldContext : WorldList)
+	{
+		if (WorldContext.ContextHandle == WorldContextHandle)
+		{
+			return &WorldContext;
+		}
+	}
+	return nullptr;
 }
 
-FWorldContext& UEngine::GetWorldContextFromHandleChecked(FName WorldContextHandle)
+FWorldContext& UEngine::GetWorldContextFromHandleChecked(const FName WorldContextHandle)
 {
 	if (FWorldContext* WorldContext = GetWorldContextFromHandle(WorldContextHandle))
 	{
@@ -9265,6 +9276,17 @@ FWorldContext& UEngine::GetWorldContextFromHandleChecked(FName WorldContextHandl
 	return HandleInvalidWorldContext();
 }
 
+const FWorldContext& UEngine::GetWorldContextFromHandleChecked(const FName WorldContextHandle) const
+{
+	if (const FWorldContext* WorldContext = GetWorldContextFromHandle(WorldContextHandle))
+	{
+		return *WorldContext;
+	}
+
+	UE_LOG(LogLoad, Warning, TEXT("WorldContext requested with invalid context handle %s"), *WorldContextHandle.ToString());
+	return HandleInvalidWorldContext();
+}
+
 FWorldContext* UEngine::GetWorldContextFromWorld(const UWorld* InWorld)
 {
 	for (FWorldContext& WorldContext : WorldList)
@@ -9274,10 +9296,22 @@ FWorldContext* UEngine::GetWorldContextFromWorld(const UWorld* InWorld)
 			return &WorldContext;
 		}
 	}
-	return NULL;
+	return nullptr;
 }
 
-FWorldContext& UEngine::GetWorldContextFromWorldChecked(UWorld *InWorld)
+const FWorldContext* UEngine::GetWorldContextFromWorld(const UWorld* InWorld) const
+{
+	for (const FWorldContext& WorldContext : WorldList)
+	{
+		if (WorldContext.World() == InWorld)
+		{
+			return &WorldContext;
+		}
+	}
+	return nullptr;
+}
+
+FWorldContext& UEngine::GetWorldContextFromWorldChecked(const UWorld *InWorld)
 {
 	if (FWorldContext* WorldContext = GetWorldContextFromWorld(InWorld))
 	{
@@ -9286,9 +9320,18 @@ FWorldContext& UEngine::GetWorldContextFromWorldChecked(UWorld *InWorld)
 	return HandleInvalidWorldContext();
 }
 
-UGameViewportClient* UEngine::GameViewportForWorld(UWorld *InWorld)
+const FWorldContext& UEngine::GetWorldContextFromWorldChecked(const UWorld *InWorld) const
+{
+	if (const FWorldContext* WorldContext = GetWorldContextFromWorld(InWorld))
+	{
+		return *WorldContext;
+	}
+	return HandleInvalidWorldContext();
+}
+
+UGameViewportClient* UEngine::GameViewportForWorld(const UWorld *InWorld) const
 {
-	FWorldContext* Context = GetWorldContextFromWorld(InWorld);
+	const FWorldContext* Context = GetWorldContextFromWorld(InWorld);
 	return (Context ? Context->GameViewport : NULL);
 }
 
@@ -9301,7 +9344,19 @@ FWorldContext* UEngine::GetWorldContextFromGameViewport(const UGameViewportClien
 			return &WorldContext;
 		}
 	}
-	return NULL;
+	return nullptr;
+}
+
+const FWorldContext* UEngine::GetWorldContextFromGameViewport(const UGameViewportClient *InViewport) const
+{
+	for (const FWorldContext& WorldContext : WorldList)
+	{
+		if (WorldContext.GameViewport == InViewport)
+		{
+			return &WorldContext;
+		}
+	}
+	return nullptr;
 }
 
 FWorldContext& UEngine::GetWorldContextFromGameViewportChecked(const UGameViewportClient *InViewport)
@@ -9313,6 +9368,15 @@ FWorldContext& UEngine::GetWorldContextFromGameViewportChecked(const UGameViewpo
 	return HandleInvalidWorldContext();
 }
 
+const FWorldContext& UEngine::GetWorldContextFromGameViewportChecked(const UGameViewportClient *InViewport) const
+{
+	if (const FWorldContext* WorldContext = GetWorldContextFromGameViewport(InViewport))
+	{
+		return *WorldContext;
+	}
+	return HandleInvalidWorldContext();
+}
+
 FWorldContext* UEngine::GetWorldContextFromPendingNetGame(const UPendingNetGame *InPendingNetGame)
 {
 	for (FWorldContext& WorldContext : WorldList)
@@ -9322,7 +9386,19 @@ FWorldContext* UEngine::GetWorldContextFromPendingNetGame(const UPendingNetGame
 			return &WorldContext;
 		}
 	}
-	return NULL;
+	return nullptr;
+}
+
+const FWorldContext* UEngine::GetWorldContextFromPendingNetGame(const UPendingNetGame *InPendingNetGame) const
+{
+	for (const FWorldContext& WorldContext : WorldList)
+	{
+		if (WorldContext.PendingNetGame == InPendingNetGame)
+		{
+			return &WorldContext;
+		}
+	}
+	return nullptr;
 }
 
 FWorldContext& UEngine::GetWorldContextFromPendingNetGameChecked(const UPendingNetGame *InPendingNetGame)
@@ -9334,6 +9410,15 @@ FWorldContext& UEngine::GetWorldContextFromPendingNetGameChecked(const UPendingN
 	return HandleInvalidWorldContext();
 }
 
+const FWorldContext& UEngine::GetWorldContextFromPendingNetGameChecked(const UPendingNetGame *InPendingNetGame) const
+{
+	if (const FWorldContext* WorldContext = GetWorldContextFromPendingNetGame(InPendingNetGame))
+	{
+		return *WorldContext;
+	}
+	return HandleInvalidWorldContext();
+}
+
 FWorldContext* UEngine::GetWorldContextFromPendingNetGameNetDriver(const UNetDriver *InPendingNetDriver)
 {
 	for (FWorldContext& WorldContext : WorldList)
@@ -9343,8 +9428,21 @@ FWorldContext* UEngine::GetWorldContextFromPendingNetGameNetDriver(const UNetDri
 			return &WorldContext;
 		}
 	}
-	return NULL;
+	return nullptr;
 }
+
+const FWorldContext* UEngine::GetWorldContextFromPendingNetGameNetDriver(const UNetDriver *InPendingNetDriver) const
+{
+	for (const FWorldContext& WorldContext : WorldList)
+	{
+		if (WorldContext.PendingNetGame && WorldContext.PendingNetGame->NetDriver == InPendingNetDriver)
+		{
+			return &WorldContext;
+		}
+	}
+	return nullptr;
+}
+
 FWorldContext& UEngine::GetWorldContextFromPendingNetGameNetDriverChecked(const UNetDriver *InPendingNetDriver)
 {
 	if (FWorldContext* WorldContext = GetWorldContextFromPendingNetGameNetDriver(InPendingNetDriver))
@@ -9354,6 +9452,15 @@ FWorldContext& UEngine::GetWorldContextFromPendingNetGameNetDriverChecked(const
 	return HandleInvalidWorldContext();
 }
 
+const FWorldContext& UEngine::GetWorldContextFromPendingNetGameNetDriverChecked(const UNetDriver *InPendingNetDriver) const
+{
+	if (const FWorldContext* WorldContext = GetWorldContextFromPendingNetGameNetDriver(InPendingNetDriver))
+	{
+		return *WorldContext;
+	}
+	return HandleInvalidWorldContext();
+}
+
 FWorldContext* UEngine::GetWorldContextFromPIEInstance(const int32 PIEInstance)
 {
 	for (FWorldContext& WorldContext : WorldList)
@@ -9363,7 +9470,19 @@ FWorldContext* UEngine::GetWorldContextFromPIEInstance(const int32 PIEInstance)
 			return &WorldContext;
 		}
 	}
-	return NULL;
+	return nullptr;
+}
+
+const FWorldContext* UEngine::GetWorldContextFromPIEInstance(const int32 PIEInstance) const
+{
+	for (const FWorldContext& WorldContext : WorldList)
+	{
+		if (WorldContext.WorldType == EWorldType::PIE && WorldContext.PIEInstance == PIEInstance)
+		{
+			return &WorldContext;
+		}
+	}
+	return nullptr;
 }
 
 FWorldContext& UEngine::GetWorldContextFromPIEInstanceChecked(const int32 PIEInstance)
@@ -9375,6 +9494,15 @@ FWorldContext& UEngine::GetWorldContextFromPIEInstanceChecked(const int32 PIEIns
 	return HandleInvalidWorldContext();
 }
 
+const FWorldContext& UEngine::GetWorldContextFromPIEInstanceChecked(const int32 PIEInstance) const
+{
+	if (const FWorldContext* WorldContext = GetWorldContextFromPIEInstance(PIEInstance))
+	{
+		return *WorldContext;
+	}
+	return HandleInvalidWorldContext();
+}
+
 UPendingNetGame* UEngine::PendingNetGameFromWorld( UWorld* InWorld )
 {
 	return GetWorldContextFromWorldChecked(InWorld).PendingNetGame;
diff --git a/Engine/Source/Runtime/Online/OnlineSubsystemUtils/Private/ConnectionCallbackProxy.cpp b/Engine/Source/Runtime/Online/OnlineSubsystemUtils/Private/ConnectionCallbackProxy.cpp
--- a/Engine/Source/Runtime/Online/OnlineSubsystemUtils/Private/ConnectionCallbackProxy.cpp
+++ b/Engine/Source/Runtime/Online/OnlineSubsystemUtils/Private/ConnectionCallbackProxy.cpp
@@ -32,12 +32,12 @@ void UConnectionCallbackProxy::Activate()
 		IOnlineIdentityPtr OnlineIdentity = Helper.OnlineSub->GetIdentityInterface();
 		if (OnlineIdentity.IsValid())
 		{
-			ULocalPlayer * localPlayer = CastChecked<ULocalPlayer>(PlayerControllerWeakPtr.Get()->Player);
+			const int32 ControllerId = CastChecked<ULocalPlayer>(PlayerControllerWeakPtr.Get()->Player)->GetControllerId();
 
-			if (!OnlineIdentity->OnLoginCompleteDelegates[localPlayer->ControllerId].IsBoundToObject(this))
+			if (!OnlineIdentity->OnLoginCompleteDelegates[ControllerId].IsBoundToObject(this))
 			{
-				OnlineIdentity->AddOnLoginCompleteDelegate(localPlayer->ControllerId, OnLoginCompleteDelegate);
-				OnlineIdentity->Login(localPlayer->ControllerId, FOnlineAccountCredentials()); /// Probably need to supply real creds here somehow... doesn't apply for all imple however.
+				OnlineIdentity->AddOnLoginCompleteDelegate(ControllerId, OnLoginCompleteDelegate);
+				OnlineIdentity->Login(ControllerId, FOnlineAccountCredentials()); /// Probably need to supply real creds here somehow... doesn't apply for all imple however.
 			}
 			else
 			{
diff --git a/Engine/Source/Runtime/Online/OnlineSubsystemUtils/Private/OnlineSessionClient.cpp b/Engine/Source/Runtime/Online/OnlineSubsystemUtils/Private/OnlineSessionClient.cpp
--- a/Engine/Source/Runtime/Online/OnlineSubsystemUtils/Private/OnlineSessionClient.cpp
+++ b/Engine/Source/Runtime/Online/OnlineSubsystemUtils/Private/OnlineSessionClient.cpp
@@ -37,9 +37,9 @@ APlayerController* UOnlineSessionClient::GetPlayerController()
 int32 UOnlineSessionClient::GetControllerId()
 {
 	ULocalPlayer* LP = Cast<ULocalPlayer>(GetOuter());
-	if (LP != NULL)
+	if (LP != nullptr)
 	{
-		return LP->ControllerId;
+		return LP->GetControllerId();
 	}
 
 	return INVALID_CONTROLLERID;
