Header cleanup (UE-3564)

diff --git a/Engine/Config/BaseEngine.ini b/Engine/Config/BaseEngine.ini
--- a/Engine/Config/BaseEngine.ini
+++ b/Engine/Config/BaseEngine.ini
@@ -711,6 +711,11 @@ bUseStreamingPause=false
 +TaggedPropertyRedirects=(ClassName="CollisionResponseContainer",OldPropertyName="Static",NewPropertyName="WorldStatic")
 +TaggedPropertyRedirects=(ClassName="CollisionResponseContainer",OldPropertyName="Dynamic",NewPropertyName="WorldDynamic")
 
+;bodysetup enums moved out of header
++ActiveClassRedirects=(OldClassName="ECollisionTraceFlag",NewClassName="/Script/Engine.BodySetupEnums:ECollisionTraceFlag")
++ActiveClassRedirects=(OldClassName="EBodyCollisionResponse",NewClassName="/Script/Engine.BodySetupEnums:EBodyCollisionResponse")
++ActiveClassRedirects=(OldClassName="EPhysicsType",NewClassName="/Script/Engine.BodySetupEnums:EPhysicsType")
+
 +TaggedPropertyRedirects=(ClassName="BodyInstance",OldPropertyName="MovementChannel",NewPropertyName="ObjectType")
 +K2FieldRedirects=(OldFieldName="PrimitiveComponent.SetMovementChannel",NewFieldName="PrimitiveComponent.SetCollisionObjectType")
 +K2FieldRedirects=(OldFieldName="SceneComponent.AttachTo",NewFieldName="SceneComponent.K2_AttachTo")
diff --git a/Engine/Source/Runtime/Engine/Classes/GameFramework/KillZVolume.h b/Engine/Source/Runtime/Engine/Classes/GameFramework/KillZVolume.h
--- a/Engine/Source/Runtime/Engine/Classes/GameFramework/KillZVolume.h
+++ b/Engine/Source/Runtime/Engine/Classes/GameFramework/KillZVolume.h
@@ -4,6 +4,13 @@
 #include "GameFramework/PhysicsVolume.h"
 #include "KillZVolume.generated.h"
 
+
+/**
+* KillZVolume is a volume used to determine when actors should be killed. Killing logic is overridden in FellOutOfWorld
+* 
+* @see FellOutOfWorld
+*/
+
 UCLASS()
 class ENGINE_API AKillZVolume : public APhysicsVolume
 {
diff --git a/Engine/Source/Runtime/Engine/Classes/GameFramework/WheeledVehicle.h b/Engine/Source/Runtime/Engine/Classes/GameFramework/WheeledVehicle.h
--- a/Engine/Source/Runtime/Engine/Classes/GameFramework/WheeledVehicle.h
+++ b/Engine/Source/Runtime/Engine/Classes/GameFramework/WheeledVehicle.h
@@ -6,6 +6,16 @@
 
 class FDebugDisplayInfo;
 
+/**
+ * WheeledVehicle is the base wheeled vehicle pawn actor.
+ * By default it uses UWheeledVehicleMovementComponent4W for its simulation, but this can be overridden by inheriting from the class and modifying its constructor like so:
+ * Super(ObjectInitializer.SetDefautSubobjectClass<UMyMovement>(VehicleMovementComponentName))
+ * Where UMyMovement is the new movement type that inherits from UWheeledVehicleMovementComponent
+ * 
+ * @see https://docs.unrealengine.com/latest/INT/Engine/Physics/Vehicles/VehicleUserGuide/
+ * @see UWheeledVehicleMovementComponent4W
+ */
+
 UCLASS(abstract, config=Game, BlueprintType)
 class ENGINE_API AWheeledVehicle : public APawn
 {
diff --git a/Engine/Source/Runtime/Engine/Classes/PhysicsEngine/AggregateGeom.h b/Engine/Source/Runtime/Engine/Classes/PhysicsEngine/AggregateGeom.h
new file mode 100644
--- /dev/null
+++ b/Engine/Source/Runtime/Engine/Classes/PhysicsEngine/AggregateGeom.h
@@ -0,0 +1,77 @@
+// Copyright 1998-2014 Epic Games, Inc. All Rights Reserved.
+
+#pragma once
+
+#include "ConvexElem.h"
+#include "BoxElem.h"
+#include "SphereElem.h"
+#include "SphylElem.h"
+#include "AggregateGeom.generated.h"
+
+class FMaterialRenderProxy;
+class FPrimitiveDrawInterface;
+
+
+/** Container for an aggregate of collision shapes */
+USTRUCT()
+struct ENGINE_API FKAggregateGeom
+{
+	GENERATED_USTRUCT_BODY()
+
+	UPROPERTY(EditAnywhere, editfixedsize, Category = KAggregateGeom)
+	TArray<FKSphereElem> SphereElems;
+
+	UPROPERTY(EditAnywhere, editfixedsize, Category = KAggregateGeom)
+	TArray<FKBoxElem> BoxElems;
+
+	UPROPERTY(EditAnywhere, editfixedsize, Category = KAggregateGeom)
+	TArray<FKSphylElem> SphylElems;
+
+	UPROPERTY(EditAnywhere, editfixedsize, Category = KAggregateGeom)
+	TArray<FKConvexElem> ConvexElems;
+
+	class FKConvexGeomRenderInfo* RenderInfo;
+
+	FKAggregateGeom()
+		: RenderInfo(NULL)
+	{}
+	int32 GetElementCount() const
+	{
+		return SphereElems.Num() + SphylElems.Num() + BoxElems.Num() + ConvexElems.Num();
+	}
+
+	int32 GetElementCount(int32 Type) const;
+
+	void EmptyElements()
+	{
+		BoxElems.Empty();
+		ConvexElems.Empty();
+		SphylElems.Empty();
+		SphereElems.Empty();
+
+		FreeRenderInfo();
+	}
+
+
+
+	void Serialize(const FArchive& Ar);
+
+	void GetAggGeom(const FTransform& Transform, const FColor Color, const FMaterialRenderProxy* MatInst, bool bPerHullColor, bool bDrawSolid, bool bUseEditorDepthTest, int32 ViewIndex, class FMeshElementCollector& Collector) const;
+
+	/** Release the RenderInfo (if its there) and safely clean up any resources. Call on the game thread. */
+	void FreeRenderInfo();
+
+	FBox CalcAABB(const FTransform& Transform) const;
+
+	/**
+	* Calculates a tight box-sphere bounds for the aggregate geometry; this is more expensive than CalcAABB
+	* (tight meaning the sphere may be smaller than would be required to encompass the AABB, but all individual components lie within both the box and the sphere)
+	*
+	* @param Output The output box-sphere bounds calculated for this set of aggregate geometry
+	*	@param LocalToWorld Transform
+	*/
+	void CalcBoxSphereBounds(FBoxSphereBounds& Output, const FTransform& LocalToWorld) const;
+
+	/** Returns the volume of this element */
+	float GetVolume(const FVector& Scale3D) const;
+};
\ No newline at end of file
diff --git a/Engine/Source/Runtime/Engine/Classes/PhysicsEngine/BodySetup.h b/Engine/Source/Runtime/Engine/Classes/PhysicsEngine/BodySetup.h
--- a/Engine/Source/Runtime/Engine/Classes/PhysicsEngine/BodySetup.h
+++ b/Engine/Source/Runtime/Engine/Classes/PhysicsEngine/BodySetup.h
@@ -2,399 +2,24 @@
 
 #pragma once
 
+#include "AggregateGeom.h"
+#include "BodySetupEnums.h"
 #include "BodySetup.generated.h"
 
-struct FDynamicMeshVertex;
-class FMaterialRenderProxy;
-class FPrimitiveDrawInterface;
-
 namespace physx
 {
-	class PxConvexMesh;
 	class PxTriangleMesh;
 	class PxRigidActor;
 }
 
-UENUM()
-enum ECollisionTraceFlag
-{
-	// default, we keep simple/complex separate for each test
-	CTF_UseDefault UMETA(DisplayName="Default"),
-	// use simple collision for complex collision test
-	CTF_UseSimpleAsComplex UMETA(DisplayName="Use Simple Collision As Complex"),
-	// use complex collision (per poly) for simple collision test
-	CTF_UseComplexAsSimple UMETA(DisplayName="Use Complex Collision As Simple"),
-	CTF_MAX,
-};
-
-UENUM()
-enum EPhysicsType
-{
-	// follow owner option
-	PhysType_Default UMETA(DisplayName="Default"),	
-	// Do not follow owner, but make kinematic
-	PhysType_Kinematic	UMETA(DisplayName="Kinematic"),		
-	// Do not follow owner, but simulate
-	PhysType_Simulated	UMETA(DisplayName="Simulated")	
-};
-
-UENUM()
-namespace EBodyCollisionResponse
-{
-	enum Type
-	{
-		BodyCollision_Enabled UMETA(DisplayName="Enabled"), 
-		BodyCollision_Disabled UMETA(DisplayName="Disabled")//, 
-		//BodyCollision_Custom UMETA(DisplayName="Custom")
-	};
-}
-
-/** One convex hull, used for simplified collision. */
-USTRUCT()
-struct FKConvexElem
-{
-	GENERATED_USTRUCT_BODY()
-
-	/** Array of indices that make up the convex hull. */
-	UPROPERTY()
-	TArray<FVector> VertexData;
-
-	/** Bounding box of this convex hull. */
-	UPROPERTY()
-	FBox ElemBox;
-
-	/** Transform of this element */
-	UPROPERTY()
-	FTransform Transform;
-
-		/** Convex mesh for this body, created from cooked data in CreatePhysicsMeshes */
-		physx::PxConvexMesh*   ConvexMesh;	
-
-		/** Convex mesh for this body, flipped across X, created from cooked data in CreatePhysicsMeshes */
-		physx::PxConvexMesh*   ConvexMeshNegX;
-
-		FKConvexElem() 
-		: ElemBox(0)
-		, Transform(FTransform::Identity)
-		, ConvexMesh(NULL)
-		, ConvexMeshNegX(NULL)
-		{}
-
-		ENGINE_API void	DrawElemWire(class FPrimitiveDrawInterface* PDI, const FTransform& ElemTM, const FColor Color) const;
-
-		void AddCachedSolidConvexGeom(TArray<FDynamicMeshVertex>& VertexBuffer, TArray<int32>& IndexBuffer, const FColor VertexColor) const;
-
-		/** Reset the hull to empty all arrays */
-		ENGINE_API void	Reset();
-
-		/** Updates internal ElemBox based on current value of VertexData */
-		ENGINE_API void	UpdateElemBox();
-
-		/** Calculate a bounding box for this convex element with the specified transform and scale */
-	ENGINE_API FBox	CalcAABB(const FTransform& BoneTM, const FVector& Scale3D) const;
-
-		/** Utility for creating a convex hull from a set of planes. Will reset current state of this elem. */
-		bool	HullFromPlanes(const TArray<FPlane>& InPlanes, const TArray<FVector>& SnapVerts);
-
-		/** Returns the volume of this element */
-		float GetVolume(const FVector& Scale) const;
-
-		FTransform GetTransform() const
-		{
-			return Transform;
-		};
-
-		void SetTransform( const FTransform& InTransform )
-		{
-			ensure(InTransform.IsValid());
-			Transform = InTransform;
-		}
-
-		friend FArchive& operator<<(FArchive& Ar,FKConvexElem& Elem);
-	
-	ENGINE_API void ScaleElem(FVector DeltaSize, float MinSize);
-};
-
-/** Sphere shape used for collision */
-USTRUCT()
-struct FKSphereElem
-{
-	GENERATED_USTRUCT_BODY()
-
-	UPROPERTY()
-	FMatrix TM_DEPRECATED;
-
-	UPROPERTY(Category=KSphereElem, VisibleAnywhere)
-	FVector Center;
-
-	UPROPERTY(Category=KSphereElem, VisibleAnywhere)
-	float Radius;
-
-	FKSphereElem() 
-	: Center( FVector::ZeroVector )
-	, Radius(1)
-	{
-
-	}
-
-	FKSphereElem( float r ) 
-	: Center( FVector::ZeroVector )
-	, Radius(r)
-	{
-
-	}
-
-	void Serialize( const FArchive& Ar );
-
-	friend bool operator==( const FKSphereElem& LHS, const FKSphereElem& RHS )
-	{
-		return ( LHS.Center == RHS.Center &&
-			LHS.Radius == RHS.Radius );
-	}
-
-	// Utility function that builds an FTransform from the current data
-	FTransform GetTransform() const
-	{
-		return FTransform( Center );
-	};
-
-	void SetTransform(const FTransform& InTransform)
-	{
-		ensure(InTransform.IsValid());
-		Center = InTransform.GetLocation();
-	}
-
-	FORCEINLINE float GetVolume(const FVector& Scale) const { return 1.3333f * PI * FMath::Pow(Radius * Scale.GetMin(), 3); }
-	
-	ENGINE_API void	DrawElemWire(class FPrimitiveDrawInterface* PDI, const FTransform& ElemTM, float Scale, const FColor Color) const;
-	ENGINE_API void	DrawElemSolid(class FPrimitiveDrawInterface* PDI, const FTransform& ElemTM, float Scale, const FMaterialRenderProxy* MaterialRenderProxy) const;
-	ENGINE_API void GetElemSolid(const FTransform& ElemTM, float Scale, const FMaterialRenderProxy* MaterialRenderProxy, int32 ViewIndex, class FMeshElementCollector& Collector) const;
-	ENGINE_API FBox CalcAABB(const FTransform& BoneTM, float Scale) const;
-
-	ENGINE_API void ScaleElem(FVector DeltaSize, float MinSize);
-};
-
-/** Box shape used for collision */
-USTRUCT()
-struct FKBoxElem
-{
-	GENERATED_USTRUCT_BODY()
-
-	UPROPERTY()
-	FMatrix TM_DEPRECATED;
-
-	UPROPERTY(Category=KBoxElem, VisibleAnywhere)
-	FVector Center;
-
-	UPROPERTY(Category=KBoxElem, VisibleAnywhere)
-	FQuat Orientation;
-
-	UPROPERTY(Category=KBoxElem, VisibleAnywhere)
-	float X;
-
-	UPROPERTY(Category=KBoxElem, VisibleAnywhere)
-	float Y;
-
-	/** length (not radius) */
-	UPROPERTY(Category=KBoxElem, VisibleAnywhere)
-	float Z;
-
-
-	FKBoxElem()
-	: Center( FVector::ZeroVector )
-	, Orientation( FQuat::Identity )
-	, X(1), Y(1), Z(1)
-	{
-
-	}
-
-	FKBoxElem( float s )
-	: Center( FVector::ZeroVector )
-	, Orientation( FQuat::Identity )
-	, X(s), Y(s), Z(s)
-	{
-
-	}
-
-	FKBoxElem( float InX, float InY, float InZ ) 
-	: Center( FVector::ZeroVector )
-	, Orientation( FQuat::Identity )
-	, X(InX), Y(InY), Z(InZ)
-
-	{
-
-	}
-
-	void Serialize( const FArchive& Ar );
-
-	friend bool operator==( const FKBoxElem& LHS, const FKBoxElem& RHS )
-	{
-		return ( LHS.Center == RHS.Center &&
-			LHS.Orientation == RHS.Orientation &&
-			LHS.X == RHS.X &&
-			LHS.Y == RHS.Y &&
-			LHS.Z == RHS.Z );
-	};
-
-	// Utility function that builds an FTransform from the current data
-	FTransform GetTransform() const
-	{
-		return FTransform( Orientation, Center );
-	};
-
-	void SetTransform( const FTransform& InTransform )
-	{
-		ensure(InTransform.IsValid());
-		Orientation = InTransform.GetRotation();
-		Center = InTransform.GetLocation();
-	}
-
-	FORCEINLINE float GetVolume(const FVector& Scale) const { float MinScale = Scale.GetMin(); return (X * MinScale) * (Y * MinScale) * (Z * MinScale); }
-
-	ENGINE_API void	DrawElemWire(class FPrimitiveDrawInterface* PDI, const FTransform& ElemTM, float Scale, const FColor Color) const;
-	ENGINE_API void	DrawElemSolid(class FPrimitiveDrawInterface* PDI, const FTransform& ElemTM, float Scale, const FMaterialRenderProxy* MaterialRenderProxy) const;
-	ENGINE_API void GetElemSolid(const FTransform& ElemTM, float Scale, const FMaterialRenderProxy* MaterialRenderProxy, int32 ViewIndex, FMeshElementCollector& Collector) const;
-	ENGINE_API FBox CalcAABB(const FTransform& BoneTM, float Scale) const;
-
-	ENGINE_API void ScaleElem(FVector DeltaSize, float MinSize);
-};
-
-/** Capsule shape used for collision */
-USTRUCT()
-struct FKSphylElem
-{
-	GENERATED_USTRUCT_BODY()
-
-	/** The transform assumes the sphyl axis points down Z. */
-	UPROPERTY()
-	FMatrix TM_DEPRECATED;
-
-	UPROPERTY(Category=KBoxElem, VisibleAnywhere)
-	FVector Center;
-
-	UPROPERTY(Category=KBoxElem, VisibleAnywhere)
-	FQuat Orientation;
-
-	UPROPERTY(Category=KSphylElem, VisibleAnywhere)
-	float Radius;
-
-	/** This is of line-segment ie. add Radius to both ends to find total length. */
-	UPROPERTY(Category=KSphylElem, VisibleAnywhere)
-	float Length;
-
-	FKSphylElem()
-	: Center( FVector::ZeroVector )
-	, Orientation( FQuat::Identity )
-	, Radius(1), Length(1)
-
-	{
-
-	}
-
-	FKSphylElem( float InRadius, float InLength )
-	: Center( FVector::ZeroVector )
-	, Orientation( FQuat::Identity )
-	, Radius(InRadius), Length(InLength)
-	{
-
-	}
-
-	void Serialize( const FArchive& Ar );
-
-	friend bool operator==( const FKSphylElem& LHS, const FKSphylElem& RHS )
-	{
-		return ( LHS.Center == RHS.Center &&
-			LHS.Orientation == RHS.Orientation &&
-			LHS.Radius == RHS.Radius &&
-			LHS.Length == RHS.Length );
-	};
-
-	// Utility function that builds an FTransform from the current data
-	FTransform GetTransform() const
-	{
-		return FTransform( Orientation, Center );
-	};
-
-	void SetTransform( const FTransform& InTransform )
-	{
-		ensure(InTransform.IsValid());
-		Orientation = InTransform.GetRotation();
-		Center = InTransform.GetLocation();
-	}
-
-	FORCEINLINE float GetVolume(const FVector& Scale) const { float ScaledRadius = Radius * Scale.GetMin(); return PI * FMath::Square(ScaledRadius) * ( 1.3333f * ScaledRadius + (Length * Scale.GetMin())); }
-
-	ENGINE_API void	DrawElemWire(class FPrimitiveDrawInterface* PDI, const FTransform& ElemTM, float Scale, const FColor Color) const;
-	ENGINE_API void GetElemSolid(const FTransform& ElemTM, float Scale, const FMaterialRenderProxy* MaterialRenderProxy, int32 ViewIndex, FMeshElementCollector& Collector) const;
-	ENGINE_API void	DrawElemSolid(class FPrimitiveDrawInterface* PDI, const FTransform& ElemTM, float Scale, const FMaterialRenderProxy* MaterialRenderProxy) const;
-	ENGINE_API FBox CalcAABB(const FTransform& BoneTM, float Scale) const;
-
-	ENGINE_API void ScaleElem(FVector DeltaSize, float MinSize);
-};
-
-/** Container for an aggregate of collision shapes */
-USTRUCT()
-struct ENGINE_API FKAggregateGeom
-{
-	GENERATED_USTRUCT_BODY()
-
-	UPROPERTY(EditAnywhere, editfixedsize, Category=KAggregateGeom)
-	TArray<struct FKSphereElem> SphereElems;
-
-	UPROPERTY(EditAnywhere, editfixedsize, Category=KAggregateGeom)
-	TArray<struct FKBoxElem> BoxElems;
-
-	UPROPERTY(EditAnywhere, editfixedsize, Category=KAggregateGeom)
-	TArray<struct FKSphylElem> SphylElems;
-
-	UPROPERTY(EditAnywhere, editfixedsize, Category=KAggregateGeom)
-	TArray<struct FKConvexElem> ConvexElems;
-
-	class FKConvexGeomRenderInfo* RenderInfo;
-
-	FKAggregateGeom() 
-	: RenderInfo(NULL)
-	{}
-	int32 GetElementCount() const
-	{
-		return SphereElems.Num() + SphylElems.Num() + BoxElems.Num() + ConvexElems.Num();
-	}
-
-	int32 GetElementCount(int32 Type) const;
-
-	void EmptyElements()
-	{
-		BoxElems.Empty();
-		ConvexElems.Empty();
-		SphylElems.Empty();
-		SphereElems.Empty();
-
-		FreeRenderInfo();
-	}
-
-	
-
-	void Serialize( const FArchive& Ar );
-
-	void GetAggGeom(const FTransform& Transform, const FColor Color, const FMaterialRenderProxy* MatInst, bool bPerHullColor, bool bDrawSolid, bool bUseEditorDepthTest, int32 ViewIndex, class FMeshElementCollector& Collector) const;
-
-	/** Release the RenderInfo (if its there) and safely clean up any resources. Call on the game thread. */
-	void FreeRenderInfo();
-
-	FBox CalcAABB(const FTransform& Transform) const;
-
-	/**
-		* Calculates a tight box-sphere bounds for the aggregate geometry; this is more expensive than CalcAABB
-		* (tight meaning the sphere may be smaller than would be required to encompass the AABB, but all individual components lie within both the box and the sphere)
-		*
-		* @param Output The output box-sphere bounds calculated for this set of aggregate geometry
-		*	@param LocalToWorld Transform
-		*/
-	void CalcBoxSphereBounds(FBoxSphereBounds& Output, const FTransform& LocalToWorld) const;
-
-	/** Returns the volume of this element */
-	float GetVolume(const FVector& Scale3D) const;
-};
+/**
+ * BodySetup contains all collision information that is associated with a single asset.
+ * A single BodySetup instance is shared among many BodyInstances so that geometry data is not duplicated.
+ * Assets typically implement a GetBodySetup function that is used during physics state creation.
+ * 
+ * @see GetBodySetup
+ * @see FBodyInstance
+ */
 
 UCLASS(hidecategories=Object, MinimalAPI)
 class UBodySetup : public UObject
diff --git a/Engine/Source/Runtime/Engine/Classes/PhysicsEngine/BodySetupEnums.h b/Engine/Source/Runtime/Engine/Classes/PhysicsEngine/BodySetupEnums.h
new file mode 100644
--- /dev/null
+++ b/Engine/Source/Runtime/Engine/Classes/PhysicsEngine/BodySetupEnums.h
@@ -0,0 +1,39 @@
+// Copyright 1998-2014 Epic Games, Inc. All Rights Reserved.
+
+#pragma once
+
+#include "BodySetupEnums.generated.h"
+
+UENUM()
+enum ECollisionTraceFlag
+{
+	// default, we keep simple/complex separate for each test
+	CTF_UseDefault UMETA(DisplayName="Default"),
+	// use simple collision for complex collision test
+	CTF_UseSimpleAsComplex UMETA(DisplayName="Use Simple Collision As Complex"),
+	// use complex collision (per poly) for simple collision test
+	CTF_UseComplexAsSimple UMETA(DisplayName="Use Complex Collision As Simple"),
+	CTF_MAX,
+};
+
+UENUM()
+enum EPhysicsType
+{
+	// follow owner option
+	PhysType_Default UMETA(DisplayName="Default"),	
+	// Do not follow owner, but make kinematic
+	PhysType_Kinematic	UMETA(DisplayName="Kinematic"),		
+	// Do not follow owner, but simulate
+	PhysType_Simulated	UMETA(DisplayName="Simulated")	
+};
+
+UENUM()
+namespace EBodyCollisionResponse
+{
+	enum Type
+	{
+		BodyCollision_Enabled UMETA(DisplayName="Enabled"), 
+		BodyCollision_Disabled UMETA(DisplayName="Disabled")//, 
+		//BodyCollision_Custom UMETA(DisplayName="Custom")
+	};
+}
\ No newline at end of file
diff --git a/Engine/Source/Runtime/Engine/Classes/PhysicsEngine/BoxElem.h b/Engine/Source/Runtime/Engine/Classes/PhysicsEngine/BoxElem.h
new file mode 100644
--- /dev/null
+++ b/Engine/Source/Runtime/Engine/Classes/PhysicsEngine/BoxElem.h
@@ -0,0 +1,90 @@
+// Copyright 1998-2014 Epic Games, Inc. All Rights Reserved.
+
+#pragma once
+
+#include "BoxElem.generated.h"
+
+/** Box shape used for collision */
+USTRUCT()
+struct FKBoxElem
+{
+	GENERATED_USTRUCT_BODY()
+
+	UPROPERTY()
+	FMatrix TM_DEPRECATED;
+
+	UPROPERTY(Category=KBoxElem, VisibleAnywhere)
+	FVector Center;
+
+	UPROPERTY(Category=KBoxElem, VisibleAnywhere)
+	FQuat Orientation;
+
+	UPROPERTY(Category=KBoxElem, VisibleAnywhere)
+	float X;
+
+	UPROPERTY(Category=KBoxElem, VisibleAnywhere)
+	float Y;
+
+	/** length (not radius) */
+	UPROPERTY(Category=KBoxElem, VisibleAnywhere)
+	float Z;
+
+
+	FKBoxElem()
+	: Center( FVector::ZeroVector )
+	, Orientation( FQuat::Identity )
+	, X(1), Y(1), Z(1)
+	{
+
+	}
+
+	FKBoxElem( float s )
+	: Center( FVector::ZeroVector )
+	, Orientation( FQuat::Identity )
+	, X(s), Y(s), Z(s)
+	{
+
+	}
+
+	FKBoxElem( float InX, float InY, float InZ ) 
+	: Center( FVector::ZeroVector )
+	, Orientation( FQuat::Identity )
+	, X(InX), Y(InY), Z(InZ)
+
+	{
+
+	}
+
+	void Serialize( const FArchive& Ar );
+
+	friend bool operator==( const FKBoxElem& LHS, const FKBoxElem& RHS )
+	{
+		return ( LHS.Center == RHS.Center &&
+			LHS.Orientation == RHS.Orientation &&
+			LHS.X == RHS.X &&
+			LHS.Y == RHS.Y &&
+			LHS.Z == RHS.Z );
+	};
+
+	// Utility function that builds an FTransform from the current data
+	FTransform GetTransform() const
+	{
+		return FTransform( Orientation, Center );
+	};
+
+	void SetTransform( const FTransform& InTransform )
+	{
+		ensure(InTransform.IsValid());
+		Orientation = InTransform.GetRotation();
+		Center = InTransform.GetLocation();
+	}
+
+	FORCEINLINE float GetVolume(const FVector& Scale) const { float MinScale = Scale.GetMin(); return (X * MinScale) * (Y * MinScale) * (Z * MinScale); }
+
+	ENGINE_API void	DrawElemWire(class FPrimitiveDrawInterface* PDI, const FTransform& ElemTM, float Scale, const FColor Color) const;
+	ENGINE_API void	DrawElemSolid(class FPrimitiveDrawInterface* PDI, const FTransform& ElemTM, float Scale, const FMaterialRenderProxy* MaterialRenderProxy) const;
+	ENGINE_API void GetElemSolid(const FTransform& ElemTM, float Scale, const FMaterialRenderProxy* MaterialRenderProxy, int32 ViewIndex, FMeshElementCollector& Collector) const;
+	ENGINE_API FBox CalcAABB(const FTransform& BoneTM, float Scale) const;
+
+	ENGINE_API void ScaleElem(FVector DeltaSize, float MinSize);
+};
diff --git a/Engine/Source/Runtime/Engine/Classes/PhysicsEngine/ConvexElem.h b/Engine/Source/Runtime/Engine/Classes/PhysicsEngine/ConvexElem.h
new file mode 100644
--- /dev/null
+++ b/Engine/Source/Runtime/Engine/Classes/PhysicsEngine/ConvexElem.h
@@ -0,0 +1,78 @@
+// Copyright 1998-2014 Epic Games, Inc. All Rights Reserved.
+
+#pragma once
+
+#include "ConvexElem.generated.h"
+
+struct FDynamicMeshVertex;
+
+namespace physx
+{
+	class PxConvexMesh;
+}
+
+/** One convex hull, used for simplified collision. */
+USTRUCT()
+struct FKConvexElem
+{
+	GENERATED_USTRUCT_BODY()
+
+	/** Array of indices that make up the convex hull. */
+	UPROPERTY()
+	TArray<FVector> VertexData;
+
+	/** Bounding box of this convex hull. */
+	UPROPERTY()
+		FBox ElemBox;
+
+	/** Transform of this element */
+	UPROPERTY()
+		FTransform Transform;
+
+	/** Convex mesh for this body, created from cooked data in CreatePhysicsMeshes */
+	physx::PxConvexMesh*   ConvexMesh;
+
+	/** Convex mesh for this body, flipped across X, created from cooked data in CreatePhysicsMeshes */
+	physx::PxConvexMesh*   ConvexMeshNegX;
+
+	FKConvexElem()
+		: ElemBox(0)
+		, Transform(FTransform::Identity)
+		, ConvexMesh(NULL)
+		, ConvexMeshNegX(NULL)
+	{}
+
+	ENGINE_API void	DrawElemWire(class FPrimitiveDrawInterface* PDI, const FTransform& ElemTM, const FColor Color) const;
+
+	void AddCachedSolidConvexGeom(TArray<FDynamicMeshVertex>& VertexBuffer, TArray<int32>& IndexBuffer, const FColor VertexColor) const;
+
+	/** Reset the hull to empty all arrays */
+	ENGINE_API void	Reset();
+
+	/** Updates internal ElemBox based on current value of VertexData */
+	ENGINE_API void	UpdateElemBox();
+
+	/** Calculate a bounding box for this convex element with the specified transform and scale */
+	ENGINE_API FBox	CalcAABB(const FTransform& BoneTM, const FVector& Scale3D) const;
+
+	/** Utility for creating a convex hull from a set of planes. Will reset current state of this elem. */
+	bool	HullFromPlanes(const TArray<FPlane>& InPlanes, const TArray<FVector>& SnapVerts);
+
+	/** Returns the volume of this element */
+	float GetVolume(const FVector& Scale) const;
+
+	FTransform GetTransform() const
+	{
+		return Transform;
+	};
+
+	void SetTransform(const FTransform& InTransform)
+	{
+		ensure(InTransform.IsValid());
+		Transform = InTransform;
+	}
+
+	friend FArchive& operator<<(FArchive& Ar, FKConvexElem& Elem);
+
+	ENGINE_API void ScaleElem(FVector DeltaSize, float MinSize);
+};
\ No newline at end of file
diff --git a/Engine/Source/Runtime/Engine/Classes/PhysicsEngine/PhysicsAsset.h b/Engine/Source/Runtime/Engine/Classes/PhysicsEngine/PhysicsAsset.h
--- a/Engine/Source/Runtime/Engine/Classes/PhysicsEngine/PhysicsAsset.h
+++ b/Engine/Source/Runtime/Engine/Classes/PhysicsEngine/PhysicsAsset.h
@@ -2,73 +2,18 @@
 
 #pragma once
 
+#include "RigidBodyIndexPair.h"
 #include "PhysicsAsset.generated.h"
 
-#define		RB_MinSizeToLockDOF				(0.1)
-#define		RB_MinAngleToLockDOF			(5.0)
-
-/**
- * Endian save storage for a pair of rigid body indices used as a key in the CollisionDisableTable TMap.
- */
-struct FRigidBodyIndexPair
-{
-	/** Pair of indices */
-	int32		Indices[2];
-	
-	/** Default constructor required for use with TMap */
-	FRigidBodyIndexPair()
-	{}
-
-	/**
-	 * Constructor, taking unordered pair of indices and generating a key.
-	 *
-	 * @param	Index1	1st unordered index
-	 * @param	Index2	2nd unordered index
-	 */
-	FRigidBodyIndexPair( int32 Index1, int32 Index2 )
-	{
-		Indices[0] = FMath::Min( Index1, Index2 );
-		Indices[1] = FMath::Max( Index1, Index2 );
-	}
-
-	/**
-	 * == operator required for use with TMap
-	 *
-	 * @param	Other	FRigidBodyIndexPair to compare this one to.
-	 * @return	true if the passed in FRigidBodyIndexPair is identical to this one, false otherwise
-	 */
-	bool operator==( const FRigidBodyIndexPair &Other ) const
-	{
-		return (Indices[0] == Other.Indices[0]) && (Indices[1] == Other.Indices[1]);
-	}
-
-	/**
-	 * Serializes the rigid body index pair to the passed in archive.	
-	 *
-	 * @param	Ar		Archive to serialize data to.
-	 * @param	Pair	FRigidBodyIndexPair to serialize
-	 * @return	returns the passed in FArchive after using it for serialization
-	 */
-	friend FArchive& operator<< ( FArchive &Ar, FRigidBodyIndexPair &Pair )
-	{
-		Ar << Pair.Indices[0] << Pair.Indices[1];
-		return Ar;
-	}
-};
-
 /**
- * Generates a hash value in accordance with the uint64 implementation of GetTypeHash which is
- * required for backward compatibility as older versions of UPhysicsAssetInstance stored
- * a TMap<uint64,bool>.
- * 
- * @param	Pair	FRigidBodyIndexPair to calculate the hash value for
- * @return	the hash value for the passed in FRigidBodyIndexPair
+ * PhysicsAsset contains a set of rigid bodies and constraints that make up a single ragdoll.
+ * The asset is not limited to human ragdolls, and can be used for any physical simulation using bodies and constraints.
+ * A SkeletalMesh has a single PhysicsAsset, which allows for easily turning ragdoll physics on or off for many SkeletalMeshComponents
+ * The asset can be configured inside the Physics Asset Editor (PhAT).
+ *
+ * @see https://docs.unrealengine.com/latest/INT/Engine/Physics/PhAT/Reference/index.html
+ * @see USkeletalMesh
  */
-inline uint32 GetTypeHash( const FRigidBodyIndexPair Pair )
-{
-	return Pair.Indices[0] + (Pair.Indices[1] * 23);
-}
-
 
 UCLASS(hidecategories=Object, BlueprintType, MinimalAPI)
 class UPhysicsAsset : public UObject
@@ -178,3 +123,5 @@ public:
 	ENGINE_API void BodyFindConstraints(int32 BodyIndex, TArray<int32>& Constraints);
 };
 
+#define		RB_MinSizeToLockDOF				(0.1)
+#define		RB_MinAngleToLockDOF			(5.0)
\ No newline at end of file
diff --git a/Engine/Source/Runtime/Engine/Classes/PhysicsEngine/RigidBodyIndexPair.h b/Engine/Source/Runtime/Engine/Classes/PhysicsEngine/RigidBodyIndexPair.h
new file mode 100644
--- /dev/null
+++ b/Engine/Source/Runtime/Engine/Classes/PhysicsEngine/RigidBodyIndexPair.h
@@ -0,0 +1,65 @@
+// Copyright 1998-2014 Epic Games, Inc. All Rights Reserved.
+
+#pragma once
+
+/**
+ * Endian save storage for a pair of rigid body indices used as a key in the CollisionDisableTable TMap.
+ */
+struct FRigidBodyIndexPair
+{
+	/** Pair of indices */
+	int32		Indices[2];
+	
+	/** Default constructor required for use with TMap */
+	FRigidBodyIndexPair()
+	{}
+
+	/**
+	 * Constructor, taking unordered pair of indices and generating a key.
+	 *
+	 * @param	Index1	1st unordered index
+	 * @param	Index2	2nd unordered index
+	 */
+	FRigidBodyIndexPair( int32 Index1, int32 Index2 )
+	{
+		Indices[0] = FMath::Min( Index1, Index2 );
+		Indices[1] = FMath::Max( Index1, Index2 );
+	}
+
+	/**
+	 * == operator required for use with TMap
+	 *
+	 * @param	Other	FRigidBodyIndexPair to compare this one to.
+	 * @return	true if the passed in FRigidBodyIndexPair is identical to this one, false otherwise
+	 */
+	bool operator==( const FRigidBodyIndexPair &Other ) const
+	{
+		return (Indices[0] == Other.Indices[0]) && (Indices[1] == Other.Indices[1]);
+	}
+
+	/**
+	 * Serializes the rigid body index pair to the passed in archive.	
+	 *
+	 * @param	Ar		Archive to serialize data to.
+	 * @param	Pair	FRigidBodyIndexPair to serialize
+	 * @return	returns the passed in FArchive after using it for serialization
+	 */
+	friend FArchive& operator<< ( FArchive &Ar, FRigidBodyIndexPair &Pair )
+	{
+		Ar << Pair.Indices[0] << Pair.Indices[1];
+		return Ar;
+	}
+};
+
+/**
+ * Generates a hash value in accordance with the uint64 implementation of GetTypeHash which is
+ * required for backward compatibility as older versions of UPhysicsAssetInstance stored
+ * a TMap<uint64,bool>.
+ * 
+ * @param	Pair	FRigidBodyIndexPair to calculate the hash value for
+ * @return	the hash value for the passed in FRigidBodyIndexPair
+ */
+inline uint32 GetTypeHash( const FRigidBodyIndexPair Pair )
+{
+	return Pair.Indices[0] + (Pair.Indices[1] * 23);
+}
\ No newline at end of file
diff --git a/Engine/Source/Runtime/Engine/Classes/PhysicsEngine/SphereElem.h b/Engine/Source/Runtime/Engine/Classes/PhysicsEngine/SphereElem.h
new file mode 100644
--- /dev/null
+++ b/Engine/Source/Runtime/Engine/Classes/PhysicsEngine/SphereElem.h
@@ -0,0 +1,65 @@
+// Copyright 1998-2014 Epic Games, Inc. All Rights Reserved.
+
+#pragma once
+
+#include "SphereElem.generated.h"
+
+
+/** Sphere shape used for collision */
+USTRUCT()
+struct FKSphereElem
+{
+	GENERATED_USTRUCT_BODY()
+
+	UPROPERTY()
+	FMatrix TM_DEPRECATED;
+
+	UPROPERTY(Category=KSphereElem, VisibleAnywhere)
+	FVector Center;
+
+	UPROPERTY(Category=KSphereElem, VisibleAnywhere)
+	float Radius;
+
+	FKSphereElem() 
+	: Center( FVector::ZeroVector )
+	, Radius(1)
+	{
+
+	}
+
+	FKSphereElem( float r ) 
+	: Center( FVector::ZeroVector )
+	, Radius(r)
+	{
+
+	}
+
+	void Serialize( const FArchive& Ar );
+
+	friend bool operator==( const FKSphereElem& LHS, const FKSphereElem& RHS )
+	{
+		return ( LHS.Center == RHS.Center &&
+			LHS.Radius == RHS.Radius );
+	}
+
+	// Utility function that builds an FTransform from the current data
+	FTransform GetTransform() const
+	{
+		return FTransform( Center );
+	};
+
+	void SetTransform(const FTransform& InTransform)
+	{
+		ensure(InTransform.IsValid());
+		Center = InTransform.GetLocation();
+	}
+
+	FORCEINLINE float GetVolume(const FVector& Scale) const { return 1.3333f * PI * FMath::Pow(Radius * Scale.GetMin(), 3); }
+	
+	ENGINE_API void	DrawElemWire(class FPrimitiveDrawInterface* PDI, const FTransform& ElemTM, float Scale, const FColor Color) const;
+	ENGINE_API void	DrawElemSolid(class FPrimitiveDrawInterface* PDI, const FTransform& ElemTM, float Scale, const FMaterialRenderProxy* MaterialRenderProxy) const;
+	ENGINE_API void GetElemSolid(const FTransform& ElemTM, float Scale, const FMaterialRenderProxy* MaterialRenderProxy, int32 ViewIndex, class FMeshElementCollector& Collector) const;
+	ENGINE_API FBox CalcAABB(const FTransform& BoneTM, float Scale) const;
+
+	ENGINE_API void ScaleElem(FVector DeltaSize, float MinSize);
+};
\ No newline at end of file
diff --git a/Engine/Source/Runtime/Engine/Classes/PhysicsEngine/SphylElem.h b/Engine/Source/Runtime/Engine/Classes/PhysicsEngine/SphylElem.h
new file mode 100644
--- /dev/null
+++ b/Engine/Source/Runtime/Engine/Classes/PhysicsEngine/SphylElem.h
@@ -0,0 +1,81 @@
+// Copyright 1998-2014 Epic Games, Inc. All Rights Reserved.
+
+#pragma once
+
+#include "SphylElem.generated.h"
+
+class FMaterialRenderProxy;
+class FPrimitiveDrawInterface;
+
+/** Capsule shape used for collision */
+USTRUCT()
+struct FKSphylElem
+{
+	GENERATED_USTRUCT_BODY()
+
+	/** The transform assumes the sphyl axis points down Z. */
+	UPROPERTY()
+	FMatrix TM_DEPRECATED;
+
+	UPROPERTY(Category=KBoxElem, VisibleAnywhere)
+	FVector Center;
+
+	UPROPERTY(Category=KBoxElem, VisibleAnywhere)
+	FQuat Orientation;
+
+	UPROPERTY(Category=KSphylElem, VisibleAnywhere)
+	float Radius;
+
+	/** This is of line-segment ie. add Radius to both ends to find total length. */
+	UPROPERTY(Category=KSphylElem, VisibleAnywhere)
+	float Length;
+
+	FKSphylElem()
+	: Center( FVector::ZeroVector )
+	, Orientation( FQuat::Identity )
+	, Radius(1), Length(1)
+
+	{
+
+	}
+
+	FKSphylElem( float InRadius, float InLength )
+	: Center( FVector::ZeroVector )
+	, Orientation( FQuat::Identity )
+	, Radius(InRadius), Length(InLength)
+	{
+
+	}
+
+	void Serialize( const FArchive& Ar );
+
+	friend bool operator==( const FKSphylElem& LHS, const FKSphylElem& RHS )
+	{
+		return ( LHS.Center == RHS.Center &&
+			LHS.Orientation == RHS.Orientation &&
+			LHS.Radius == RHS.Radius &&
+			LHS.Length == RHS.Length );
+	};
+
+	// Utility function that builds an FTransform from the current data
+	FTransform GetTransform() const
+	{
+		return FTransform( Orientation, Center );
+	};
+
+	void SetTransform( const FTransform& InTransform )
+	{
+		ensure(InTransform.IsValid());
+		Orientation = InTransform.GetRotation();
+		Center = InTransform.GetLocation();
+	}
+
+	FORCEINLINE float GetVolume(const FVector& Scale) const { float ScaledRadius = Radius * Scale.GetMin(); return PI * FMath::Square(ScaledRadius) * ( 1.3333f * ScaledRadius + (Length * Scale.GetMin())); }
+
+	ENGINE_API void	DrawElemWire(class FPrimitiveDrawInterface* PDI, const FTransform& ElemTM, float Scale, const FColor Color) const;
+	ENGINE_API void GetElemSolid(const FTransform& ElemTM, float Scale, const FMaterialRenderProxy* MaterialRenderProxy, int32 ViewIndex, FMeshElementCollector& Collector) const;
+	ENGINE_API void	DrawElemSolid(class FPrimitiveDrawInterface* PDI, const FTransform& ElemTM, float Scale, const FMaterialRenderProxy* MaterialRenderProxy) const;
+	ENGINE_API FBox CalcAABB(const FTransform& BoneTM, float Scale) const;
+
+	ENGINE_API void ScaleElem(FVector DeltaSize, float MinSize);
+};
\ No newline at end of file
