Fixing this tooltip to match RandomIntegerInRangeFromStream

diff --git a/Engine/Source/Runtime/Engine/Classes/Kismet/KismetMathLibrary.h b/Engine/Source/Runtime/Engine/Classes/Kismet/KismetMathLibrary.h
--- a/Engine/Source/Runtime/Engine/Classes/Kismet/KismetMathLibrary.h
+++ b/Engine/Source/Runtime/Engine/Classes/Kismet/KismetMathLibrary.h
@@ -208,7 +208,7 @@ class UKismetMathLibrary : public UBlueprintFunctionLibrary
 	UFUNCTION(BlueprintPure, Category="Math|Random")
 	static int32 RandomInteger(int32 Max);
 
-	/** Return a random integer between Min and Max */
+	/** Return a random integer between Min and Max (>= Min and <= Max) */
 	UFUNCTION(BlueprintPure, Category="Math|Random")
 	static int32 RandomIntegerInRange(int32 Min, int32 Max);
 
