#ue4 - CharacterMovementComponent helper functions SafeNormalPrecise and ClampMaxSizePrecise renamed with "Get" prepended to indicate that they don't mutate but return a value instead.

diff --git a/Engine/Source/Runtime/Engine/Private/CharacterMovementComponent.cpp b/Engine/Source/Runtime/Engine/Private/CharacterMovementComponent.cpp
--- a/Engine/Source/Runtime/Engine/Private/CharacterMovementComponent.cpp
+++ b/Engine/Source/Runtime/Engine/Private/CharacterMovementComponent.cpp
@@ -76,7 +76,7 @@ static TAutoConsoleVariable<float> CVarNetCorrectionLifetime(
 #endif // !(UE_BUILD_SHIPPING || UE_BUILD_TEST)
 
 // Version that does not use inverse sqrt estimate, for higher precision.
-FORCEINLINE FVector SafeNormalPrecise(const FVector& V)
+FORCEINLINE FVector GetSafeNormalPrecise(const FVector& V)
 {
 	const float VSq = V.SizeSquared();
 	if (VSq < SMALL_NUMBER)
@@ -90,7 +90,7 @@ FORCEINLINE FVector SafeNormalPrecise(const FVector& V)
 }
 
 // Version that does not use inverse sqrt estimate, for higher precision.
-FORCEINLINE FVector ClampMaxSizePrecise(const FVector& V, float MaxSize)
+FORCEINLINE FVector GetClampedToMaxSizePrecise(const FVector& V, float MaxSize)
 {
 	if (MaxSize < KINDA_SMALL_NUMBER)
 	{
@@ -2151,11 +2151,11 @@ void UCharacterMovementComponent::CalcVelocity(float DeltaTime, float Friction,
 		// In consideration order for direction: Acceleration, then Velocity, then Pawn's rotation.
 		if (Acceleration.SizeSquared() > SMALL_NUMBER)
 		{
-			Acceleration = SafeNormalPrecise(Acceleration) * MaxAccel;
+			Acceleration = GetSafeNormalPrecise(Acceleration) * MaxAccel;
 		}
 		else 
 		{
-			Acceleration = MaxAccel * (Velocity.SizeSquared() < SMALL_NUMBER ? CharacterOwner->GetActorRotation().Vector() : SafeNormalPrecise(Velocity));
+			Acceleration = MaxAccel * (Velocity.SizeSquared() < SMALL_NUMBER ? CharacterOwner->GetActorRotation().Vector() : GetSafeNormalPrecise(Velocity));
 		}
 
 		AnalogInputModifier = 1.f;
@@ -2178,13 +2178,13 @@ void UCharacterMovementComponent::CalcVelocity(float DeltaTime, float Friction,
 		// Don't allow braking to lower us below max speed if we started above it.
 		if (bVelocityOverMax && Velocity.SizeSquared() < FMath::Square(MaxSpeed) && FVector::DotProduct(Acceleration, OldVelocity) > 0.0f)
 		{
-			Velocity = SafeNormalPrecise(OldVelocity) * MaxSpeed;
+			Velocity = GetSafeNormalPrecise(OldVelocity) * MaxSpeed;
 		}
 	}
 	else if (!bZeroAcceleration)
 	{
 		// Friction affects our ability to change direction. This is only done for input acceleration, not path following.
-		const FVector AccelDir = SafeNormalPrecise(Acceleration);
+		const FVector AccelDir = GetSafeNormalPrecise(Acceleration);
 		const float VelSize = Velocity.Size();
 		Velocity = Velocity - (Velocity - AccelDir * VelSize) * FMath::Min(DeltaTime * Friction, 1.f);
 	}
@@ -2199,7 +2199,7 @@ void UCharacterMovementComponent::CalcVelocity(float DeltaTime, float Friction,
 	const float NewMaxSpeed = (IsExceedingMaxSpeed(MaxSpeed)) ? Velocity.Size() : MaxSpeed;
 	Velocity += Acceleration * DeltaTime;
 	Velocity += RequestedAcceleration * DeltaTime;
-	Velocity = ClampMaxSizePrecise(Velocity, NewMaxSpeed);
+	Velocity = GetClampedToMaxSizePrecise(Velocity, NewMaxSpeed);
 
 	if (bUseRVOAvoidance)
 	{
@@ -2573,7 +2573,7 @@ void UCharacterMovementComponent::ApplyVelocityBraking(float DeltaTime, float Fr
 	const float MaxTimeStep = (1.0f / 33.0f);
 
 	// Decelerate to brake to a stop
-	const FVector RevAccel = (bZeroBraking ? FVector::ZeroVector : (-BrakingDeceleration * SafeNormalPrecise(Velocity)));
+	const FVector RevAccel = (bZeroBraking ? FVector::ZeroVector : (-BrakingDeceleration * GetSafeNormalPrecise(Velocity)));
 	while( RemainingTime >= MIN_TICK_TIME )
 	{
 		// Zero friction uses constant deceleration, so no need for iteration.
@@ -5057,7 +5057,7 @@ FVector UCharacterMovementComponent::ConstrainInputAcceleration(const FVector& I
 
 FVector UCharacterMovementComponent::ScaleInputAcceleration(const FVector& InputAcceleration) const
 {
-	return GetMaxAcceleration() * ClampMaxSizePrecise(InputAcceleration, 1.0f);
+	return GetMaxAcceleration() * GetClampedToMaxSizePrecise(InputAcceleration, 1.0f);
 }
 
 
