#ue4 - Fix for convex mesh not yielding useful MTD results for "zero distance" hits (hits right next to something).

- MTD fallback is now simulated for convex meshes (since we can't inflate them) by trying MTD pushed in the direction of trace, or back away from that direction, and then in axial directions (only against the shape the hit was reported for).

diff --git a/Engine/Source/Runtime/Engine/Private/Collision/CollisionConversions.cpp b/Engine/Source/Runtime/Engine/Private/Collision/CollisionConversions.cpp
--- a/Engine/Source/Runtime/Engine/Private/Collision/CollisionConversions.cpp
+++ b/Engine/Source/Runtime/Engine/Private/Collision/CollisionConversions.cpp
@@ -566,6 +566,49 @@ static bool ComputeInflatedMTD(const float MtdInflation, const PxLocationHit& PH
 			return ComputeInflatedMTD_Internal(MtdInflation, PHit, OutResult, QueryTM, InflatedSphere, PShapeWorldPose);
 		}
 
+		case PxGeometryType::eCONVEXMESH:
+		{
+			// We can't exactly inflate the mesh (not easily), so try jittering it a bit to get an MTD result.
+			PxVec3 TraceDir = U2PVector(OutResult.TraceEnd - OutResult.TraceStart);
+			TraceDir.normalizeSafe();
+
+			// Try forward (in trace direction)
+			PxTransform JitteredTM = PxTransform(QueryTM.p + (TraceDir * MtdInflation), QueryTM.q);
+			if (ComputeInflatedMTD_Internal(MtdInflation, PHit, OutResult, JitteredTM, Geom, PShapeWorldPose))
+			{
+				return true;
+			}
+
+			// Try backward (opposite trace direction)
+			JitteredTM = PxTransform(QueryTM.p - (TraceDir * MtdInflation), QueryTM.q);
+			if (ComputeInflatedMTD_Internal(MtdInflation, PHit, OutResult, JitteredTM, Geom, PShapeWorldPose))
+			{
+				return true;
+			}
+
+			// Try axial directions.
+			// Start with -Z because this is the most common case (objects on the floor).
+			for (int32 i=2; i >= 0; i--)
+			{
+				PxVec3 Jitter(0.f);
+				Jitter[i] = MtdInflation;
+
+				JitteredTM = PxTransform(QueryTM.p - Jitter, QueryTM.q);
+				if (ComputeInflatedMTD_Internal(MtdInflation, PHit, OutResult, JitteredTM, Geom, PShapeWorldPose))
+				{
+					return true;
+				}
+
+				JitteredTM = PxTransform(QueryTM.p + Jitter, QueryTM.q);
+				if (ComputeInflatedMTD_Internal(MtdInflation, PHit, OutResult, JitteredTM, Geom, PShapeWorldPose))
+				{
+					return true;
+				}
+			}
+
+			return false;
+		}
+
 		default:
 		{
 			return false;
@@ -619,31 +662,35 @@ static bool ConvertOverlappedShapeToImpactHit(const PxLocationHit& PHit, const F
 		}
 	}
 
-	// Zero-distance hits are often valid hits and we can extract the hit normal from the face index
+	// Zero-distance hits are often valid hits and we can extract the hit normal from a larger shape with MTD.
 	// For invalid normals we can try other methods as well (get overlapping triangles).
-	static const float MtdInflation = 0.125f;
 
 	if (PHit.distance == 0.f || !bValidNormal)
 	{
-		const PxTransform PShapeWorldPose = PxShapeExt::getGlobalPose(*PShape, *PActor); 
-		PxTriangleMeshGeometry PTriMeshGeom;
-		PxHeightFieldGeometry PHeightfieldGeom;
+		const PxTransform PShapeWorldPose = PxShapeExt::getGlobalPose(*PShape, *PActor);
 
-		PxGeometryType::Enum GeometryType = PShape->getGeometryType();
+		// Try MTD with a small inflation for better accuracy, then a larger one in case the first one fails due to precision issues.
+		static const float SmallMtdInflation = 0.250f;
+		static const float LargeMtdInflation = 2.500f;
 
-		if (PShape->getTriangleMeshGeometry(PTriMeshGeom) || PShape->getHeightFieldGeometry(PHeightfieldGeom))
+		if (ComputeInflatedMTD(SmallMtdInflation, PHit, OutResult, QueryTM, Geom, PShapeWorldPose) ||
+			ComputeInflatedMTD(LargeMtdInflation, PHit, OutResult, QueryTM, Geom, PShapeWorldPose))
 		{
-			if (ComputeInflatedMTD(MtdInflation, PHit, OutResult, QueryTM, Geom, PShapeWorldPose))
-			{
-				// Success
-			}
-			else
+			// Success
+		}
+		else
+		{
+			PxTriangleMeshGeometry PTriMeshGeom;
+			PxHeightFieldGeometry PHeightfieldGeom;
+
+			if (PShape->getTriangleMeshGeometry(PTriMeshGeom) || PShape->getHeightFieldGeometry(PHeightfieldGeom))
 			{
-				const bool bIsTriMesh = GeometryType == PxGeometryType::eTRIANGLEMESH;
+				PxGeometryType::Enum GeometryType = PShape->getGeometryType();
+				const bool bIsTriMesh = (GeometryType == PxGeometryType::eTRIANGLEMESH);
 				PxU32 HitTris[64];
 				bool bOverflow = false;
 
-				const int32 NumTrisHit =  bIsTriMesh? PxMeshQuery::findOverlapTriangleMesh(Geom, QueryTM, PTriMeshGeom,     PShapeWorldPose, HitTris, ARRAY_COUNT(HitTris), 0, bOverflow) :
+				const int32 NumTrisHit = bIsTriMesh ? PxMeshQuery::findOverlapTriangleMesh(Geom, QueryTM, PTriMeshGeom, PShapeWorldPose, HitTris, ARRAY_COUNT(HitTris), 0, bOverflow) :
 													  PxMeshQuery::findOverlapHeightField( Geom, QueryTM, PHeightfieldGeom, PShapeWorldPose, HitTris, ARRAY_COUNT(HitTris), 0, bOverflow);
 				if (NumTrisHit > 0)
 				{
@@ -658,16 +705,10 @@ static bool ConvertOverlappedShapeToImpactHit(const PxLocationHit& PHit, const F
 
 				}
 			}
-		}
-		else
-		{
-			// Not a tri-mesh or heightfield
-			if (ComputeInflatedMTD(MtdInflation, PHit, OutResult, QueryTM, Geom, PShapeWorldPose))
-			{
-				// Success
-			}
 			else
 			{
+				// Not a tri-mesh or heightfield
+
 				// MTD failed, use point distance. This is not ideal.
 				// Note: faceIndex seems to be unreliable for convex meshes in these cases, so not using FindGeomOpposingNormal() for them here.
 				PxGeometry& PGeom = PShape->getGeometry().any();
@@ -677,7 +718,7 @@ static bool ConvertOverlappedShapeToImpactHit(const PxLocationHit& PHit, const F
 				if (Distance < KINDA_SMALL_NUMBER)
 				{
 					UE_LOG(LogCollision, Verbose, TEXT("Warning: ConvertOverlappedShapeToImpactHit: Query origin inside shape, giving poor MTD."));
-					PClosestPoint = PxShapeExt::getWorldBounds(*PShape, *PActor).getCenter(); 
+					PClosestPoint = PxShapeExt::getWorldBounds(*PShape, *PActor).getCenter();
 				}
 
 				OutResult.ImpactNormal = (OutResult.Location - P2UVector(PClosestPoint)).GetSafeNormal();
