Fix for InstancedStaticMesh lightmap requantization

diff --git a/Engine/Source/Runtime/Engine/Private/LightMap.cpp b/Engine/Source/Runtime/Engine/Private/LightMap.cpp
--- a/Engine/Source/Runtime/Engine/Private/LightMap.cpp
+++ b/Engine/Source/Runtime/Engine/Private/LightMap.cpp
@@ -1341,6 +1341,10 @@ FLightMap2D* FLightMap2D::AllocateInstancedLightMap(UInstancedStaticMeshComponen
 				}
 			}
 		}
+
+		// Save new requantized Scale/Add
+		FMemory::Memcpy(SourceQuantizedData->Scale, Scale, sizeof(Scale));
+		FMemory::Memcpy(SourceQuantizedData->Add, Add, sizeof(Add));
 	}
 
 	FLightMapAllocationGroup AllocationGroup = FLightMapAllocationGroup();
