Added a simple, BP-assignable event to navigation system that gets triggered when a navigation data instance finishes data generation at runtime #UE4

diff --git a/Engine/Source/Runtime/Engine/Classes/AI/Navigation/NavigationSystem.h b/Engine/Source/Runtime/Engine/Classes/AI/Navigation/NavigationSystem.h
--- a/Engine/Source/Runtime/Engine/Classes/AI/Navigation/NavigationSystem.h
+++ b/Engine/Source/Runtime/Engine/Classes/AI/Navigation/NavigationSystem.h
@@ -40,7 +40,7 @@ ENGINE_API DECLARE_LOG_CATEGORY_EXTERN(LogNavigation, Warning, All);
 DECLARE_MULTICAST_DELEGATE_OneParam(FOnNavAreaChanged, const UClass* /*AreaClass*/);
 
 /** Delegate to let interested parties know that Nav Data has been registered */
-DECLARE_DYNAMIC_MULTICAST_DELEGATE_OneParam(FOnNavDataRegistered, ANavigationData*, NavData);
+DECLARE_DYNAMIC_MULTICAST_DELEGATE_OneParam(FOnNavDataGenerigEvent, ANavigationData*, NavData);
 
 namespace NavigationDebugDrawing
 {
@@ -190,8 +190,11 @@ class ENGINE_API UNavigationSystem : public UBlueprintFunctionLibrary
 	TArray<FNavigationBoundsUpdateRequest> PendingNavBoundsUpdates;
 
  	UPROPERTY(/*BlueprintAssignable, */Transient)
-	FOnNavDataRegistered OnNavDataRegisteredEvent;
+	FOnNavDataGenerigEvent OnNavDataRegisteredEvent;
 
+	UPROPERTY(BlueprintAssignable, Transient, meta = (displayname = OnNavigationGenerationFinished))
+	FOnNavDataGenerigEvent OnNavigationGenerationFinishedDelegate;
+	
 private:
 	TWeakObjectPtr<UCrowdManager> CrowdManager;
 
@@ -565,6 +568,8 @@ public:
 	/** Used to display "navigation building in progress" notify */
 	bool IsNavigationBuildInProgress(bool bCheckDirtyToo = true);
 
+	void OnNavigationGenerationFinished(ANavigationData& NavData);
+
 	/** Used to display "navigation building in progress" counter */
 	int32 GetNumRemainingBuildTasks() const;
 
diff --git a/Engine/Source/Runtime/Engine/Private/AI/Navigation/NavigationSystem.cpp b/Engine/Source/Runtime/Engine/Private/AI/Navigation/NavigationSystem.cpp
--- a/Engine/Source/Runtime/Engine/Private/AI/Navigation/NavigationSystem.cpp
+++ b/Engine/Source/Runtime/Engine/Private/AI/Navigation/NavigationSystem.cpp
@@ -2906,6 +2906,11 @@ bool UNavigationSystem::IsNavigationBuildInProgress(bool bCheckDirtyToo)
 	return bRet;
 }
 
+void UNavigationSystem::OnNavigationGenerationFinished(ANavigationData& NavData)
+{
+	OnNavigationGenerationFinishedDelegate.Broadcast(&NavData);
+}
+
 int32 UNavigationSystem::GetNumRemainingBuildTasks() const
 {
 	int32 NumTasks = 0;
diff --git a/Engine/Source/Runtime/Engine/Private/AI/Navigation/RecastNavMesh.cpp b/Engine/Source/Runtime/Engine/Private/AI/Navigation/RecastNavMesh.cpp
--- a/Engine/Source/Runtime/Engine/Private/AI/Navigation/RecastNavMesh.cpp
+++ b/Engine/Source/Runtime/Engine/Private/AI/Navigation/RecastNavMesh.cpp
@@ -1265,56 +1265,65 @@ URecastNavMeshDataChunk* ARecastNavMesh::GetNavigationDataChunk(ULevel* InLevel)
 
 void ARecastNavMesh::OnNavMeshGenerationFinished()
 {
-#if WITH_EDITOR
 	UWorld* World = GetWorld();
-	// For static navmeshes create navigation data holders in each streaming level
-	// so parts of navmesh can be streamed in/out with those levels
-	if (!World->IsGameWorld())
+
+	if (World != nullptr && World->IsPendingKill() == false)
 	{
-		const auto& Levels = World->GetLevels();
-		for (auto Level : Levels)
+#if WITH_EDITOR	
+		// For static navmeshes create navigation data holders in each streaming level
+		// so parts of navmesh can be streamed in/out with those levels
+		if (!World->IsGameWorld())
 		{
-			if (Level->IsPersistentLevel())
+			const auto& Levels = World->GetLevels();
+			for (auto Level : Levels)
 			{
-				continue;
-			}
-			
-			URecastNavMeshDataChunk* NavDataChunk = GetNavigationDataChunk(Level);
-				
-			if (!bRebuildAtRuntime)
-			{
-				// We use nav volumes that belongs to this streaming level to find tiles we want to save
-				TArray<int32> LevelTiles;
-				TArray<FBox> LevelNavBounds = GetNavigableBoundsInLevel(Level->GetOutermost()->GetFName());
-				RecastNavMeshImpl->GetNavMeshTilesIn(LevelNavBounds, LevelTiles);
-				
-				if (LevelTiles.Num())
+				if (Level->IsPersistentLevel())
 				{
-					// Create new chunk only if we have something to save in it			
-					if (NavDataChunk == nullptr)
+					continue;
+				}
+
+				URecastNavMeshDataChunk* NavDataChunk = GetNavigationDataChunk(Level);
+
+				if (!bRebuildAtRuntime)
+				{
+					// We use nav volumes that belongs to this streaming level to find tiles we want to save
+					TArray<int32> LevelTiles;
+					TArray<FBox> LevelNavBounds = GetNavigableBoundsInLevel(Level->GetOutermost()->GetFName());
+					RecastNavMeshImpl->GetNavMeshTilesIn(LevelNavBounds, LevelTiles);
+
+					if (LevelTiles.Num())
 					{
-						NavDataChunk = NewObject<URecastNavMeshDataChunk>(Level);
-						NavDataChunk->NavigationDataName = GetFName();
-						Level->NavDataChunks.Add(NavDataChunk);
+						// Create new chunk only if we have something to save in it			
+						if (NavDataChunk == nullptr)
+						{
+							NavDataChunk = NewObject<URecastNavMeshDataChunk>(Level);
+							NavDataChunk->NavigationDataName = GetFName();
+							Level->NavDataChunks.Add(NavDataChunk);
+						}
+
+						NavDataChunk->GatherTiles(RecastNavMeshImpl->DetourNavMesh, LevelTiles);
+						NavDataChunk->MarkPackageDirty();
+						continue;
 					}
-					
-					NavDataChunk->GatherTiles(RecastNavMeshImpl->DetourNavMesh, LevelTiles);
+				}
+
+				// stale data that is left in the level
+				if (NavDataChunk)
+				{
+					// clear it
+					NavDataChunk->ReleaseTiles();
 					NavDataChunk->MarkPackageDirty();
-					continue;
+					Level->NavDataChunks.Remove(NavDataChunk);
 				}
 			}
-			
-			// stale data that is left in the level
-			if (NavDataChunk)
-			{
-				// clear it
-				NavDataChunk->ReleaseTiles();
-				NavDataChunk->MarkPackageDirty();
-				Level->NavDataChunks.Remove(NavDataChunk);
-			}
 		}
-	}
 #endif// WITH_EDITOR
+
+		if (World->GetNavigationSystem())
+		{
+			World->GetNavigationSystem()->OnNavigationGenerationFinished(*this);
+		}
+	}
 }
 
 #if !UE_BUILD_SHIPPING
