Fixed memory leak

diff --git a/Engine/Source/Runtime/Engine/Private/Particles/ParticleSystemRender.cpp b/Engine/Source/Runtime/Engine/Private/Particles/ParticleSystemRender.cpp
--- a/Engine/Source/Runtime/Engine/Private/Particles/ParticleSystemRender.cpp
+++ b/Engine/Source/Runtime/Engine/Private/Particles/ParticleSystemRender.cpp
@@ -984,18 +984,18 @@ void FDynamicSpriteEmitterData::GetDynamicMeshElementsEmitter(const FParticleSys
 				Proxy->UpdateWorldSpacePrimitiveUniformBuffer();
 			}
 
-			FGlobalDynamicVertexBuffer::FAllocation* Allocation = new FGlobalDynamicVertexBuffer::FAllocation();
-			FGlobalDynamicVertexBuffer::FAllocation* DynamicParameterAllocation = new FGlobalDynamicVertexBuffer::FAllocation();
+			FGlobalDynamicVertexBuffer::FAllocation Allocation;
+			FGlobalDynamicVertexBuffer::FAllocation DynamicParameterAllocation;
 
 			// Allocate memory for render data.
-			*Allocation = FGlobalDynamicVertexBuffer::Get().Allocate( ParticleCount * VertexSize * VertexPerParticle );
+			Allocation = FGlobalDynamicVertexBuffer::Get().Allocate( ParticleCount * VertexSize * VertexPerParticle );
 
 			if (bUsesDynamicParameter)
 			{
-				*DynamicParameterAllocation = FGlobalDynamicVertexBuffer::Get().Allocate( ParticleCount * DynamicParameterVertexSize * VertexPerParticle );
+				DynamicParameterAllocation = FGlobalDynamicVertexBuffer::Get().Allocate( ParticleCount * DynamicParameterVertexSize * VertexPerParticle );
 			}
 
-			if ( Allocation->IsValid() && (!bUsesDynamicParameter || DynamicParameterAllocation->IsValid()) )
+			if ( Allocation.IsValid() && (!bUsesDynamicParameter || DynamicParameterAllocation.IsValid()) )
 			{
 				// Sort the particles if needed.
 				FParticleOrder* ParticleOrder = NULL;
@@ -1019,11 +1019,11 @@ void FDynamicSpriteEmitterData::GetDynamicMeshElementsEmitter(const FParticleSys
 				// Fill vertex buffers.
 				if(bInstanced)
 				{
-					GetVertexAndIndexData(Allocation->Buffer, DynamicParameterAllocation->Buffer, NULL, ParticleOrder, View->ViewMatrices.ViewOrigin, Proxy->GetLocalToWorld());
+					GetVertexAndIndexData(Allocation.Buffer, DynamicParameterAllocation.Buffer, NULL, ParticleOrder, View->ViewMatrices.ViewOrigin, Proxy->GetLocalToWorld());
 				}
 				else
 				{
-					GetVertexAndIndexDataNonInstanced(Allocation->Buffer, DynamicParameterAllocation->Buffer, NULL, ParticleOrder, View->ViewMatrices.ViewOrigin, Proxy->GetLocalToWorld());
+					GetVertexAndIndexDataNonInstanced(Allocation.Buffer, DynamicParameterAllocation.Buffer, NULL, ParticleOrder, View->ViewMatrices.ViewOrigin, Proxy->GetLocalToWorld());
 				}
 
 				// Create per-view uniform buffer.
@@ -1051,8 +1051,8 @@ void FDynamicSpriteEmitterData::GetDynamicMeshElementsEmitter(const FParticleSys
 
 				// Set the sprite uniform buffer for this view.
 				SpriteVertexFactory->SetSpriteUniformBuffer( CollectorResources.UniformBuffer );
-				SpriteVertexFactory->SetInstanceBuffer( Allocation->VertexBuffer, Allocation->VertexOffset, GetDynamicVertexStride(FeatureLevel), bInstanced );
-				SpriteVertexFactory->SetDynamicParameterBuffer( DynamicParameterAllocation ? DynamicParameterAllocation->VertexBuffer : NULL, DynamicParameterAllocation->VertexOffset, GetDynamicParameterVertexStride(), bInstanced );
+				SpriteVertexFactory->SetInstanceBuffer( Allocation.VertexBuffer, Allocation.VertexOffset, GetDynamicVertexStride(FeatureLevel), bInstanced );
+				SpriteVertexFactory->SetDynamicParameterBuffer( DynamicParameterAllocation.VertexBuffer, DynamicParameterAllocation.VertexOffset, GetDynamicParameterVertexStride(), bInstanced );
 
 				// Construct the mesh element to render.
 				FMeshBatch& Mesh = Collector.AllocateMesh();
