Fix physics issues created by double buffered bones

diff --git a/Engine/Source/Runtime/Engine/Private/PhysicsEngine/PhysAnim.cpp b/Engine/Source/Runtime/Engine/Private/PhysicsEngine/PhysAnim.cpp
--- a/Engine/Source/Runtime/Engine/Private/PhysicsEngine/PhysAnim.cpp
+++ b/Engine/Source/Runtime/Engine/Private/PhysicsEngine/PhysAnim.cpp
@@ -211,7 +211,7 @@ void USkeletalMeshComponent::BlendInPhysics()
 
 	// We now have all the animations blended together and final relative transforms for each bone.
 	// If we don't have or want any physics, we do nothing.
-	if( Bodies.Num() > 0 && (bBlendPhysics || ShouldBlendPhysicsBones()) )
+	if( Bodies.Num() > 0 )
 	{
 		BlendPhysicsBones( RequiredBones );
 		
diff --git a/Engine/Source/Runtime/Engine/Private/SkeletalMeshComponent.cpp b/Engine/Source/Runtime/Engine/Private/SkeletalMeshComponent.cpp
--- a/Engine/Source/Runtime/Engine/Private/SkeletalMeshComponent.cpp
+++ b/Engine/Source/Runtime/Engine/Private/SkeletalMeshComponent.cpp
@@ -1035,7 +1035,7 @@ void USkeletalMeshComponent::PostAnimEvaluation(FAnimationEvaluationContext& Eva
 	// this causes issue where a half of frame, physics position is fixed with anim pose, and the other half is real simulated position
 	// if you enable physics in tick, since that's before physics update, you'll get animation pose dominating physics pose, which isn't what you want. (Or what you'll see)
 	// so do not update transform if physics is on. This problem will be solved by double buffer, when we keep one buffer for intermediate, and the other buffer for result query
-	if( !IsSimulatingPhysics() )
+	if (!bBlendPhysics && !ShouldBlendPhysicsBones())
 	{
 		SCOPE_CYCLE_COUNTER(STAT_UpdateLocalToWorldAndOverlaps);
 
diff --git a/Engine/Source/Runtime/Engine/Private/SkeletalMeshComponentPhysics.cpp b/Engine/Source/Runtime/Engine/Private/SkeletalMeshComponentPhysics.cpp
--- a/Engine/Source/Runtime/Engine/Private/SkeletalMeshComponentPhysics.cpp
+++ b/Engine/Source/Runtime/Engine/Private/SkeletalMeshComponentPhysics.cpp
@@ -3078,14 +3078,14 @@ void USkeletalMeshComponent::PreClothTick(float DeltaTime)
 	// and run this if that is true or rendered
 	// that will at least reduce the chance of mismatch
 	// generally if you move your actor position, this has to happen to approximately match their bounds
-	if (Bodies.Num() > 0 && IsRegistered())
+	bool bShouldBlendPhys = ShouldBlendPhysicsBones() || bBlendPhysics;
+	if (bShouldBlendPhys)
 	{
-		BlendInPhysics();
-	}
-
-	if (IsSimulatingPhysics())
-	{
-		// If we aren't simulating we will have already flipped this
+		if (IsRegistered())
+		{
+			BlendInPhysics();
+		}
+		// If we aren't blending we will have already flipped this
 		FlipEditableSpaceBases();
 	}
 
