Added a basic path-following support to FloatingPawnMovement component #UE4

diff --git a/Engine/Source/Runtime/Engine/Private/FloatingPawnMovement.cpp b/Engine/Source/Runtime/Engine/Private/FloatingPawnMovement.cpp
--- a/Engine/Source/Runtime/Engine/Private/FloatingPawnMovement.cpp
+++ b/Engine/Source/Runtime/Engine/Private/FloatingPawnMovement.cpp
@@ -27,7 +27,17 @@ void UFloatingPawnMovement::TickComponent(float DeltaTime, enum ELevelTick TickT
 	const AController* Controller = PawnOwner->GetController();
 	if (Controller && Controller->IsLocalController())
 	{
-		ApplyControlInputToVelocity(DeltaTime);
+		if (Controller->IsLocalPlayerController())
+		{
+			ApplyControlInputToVelocity(DeltaTime);
+		}
+		// if it's not player controller, but we do have a controller, then it's AI
+		// and we need to limit the speed
+		else if (IsExceedingMaxSpeed(MaxSpeed))
+		{
+			Velocity = Velocity.GetUnsafeNormal() * MaxSpeed;
+		}
+
 		LimitWorldBounds();
 		bPositionCorrected = false;
 
