Added a way to specify ANavigationData class per NavigationSystem's supported agent #UE4

diff --git a/Engine/Source/Runtime/Engine/Classes/AI/Navigation/NavigationData.h b/Engine/Source/Runtime/Engine/Classes/AI/Navigation/NavigationData.h
--- a/Engine/Source/Runtime/Engine/Classes/AI/Navigation/NavigationData.h
+++ b/Engine/Source/Runtime/Engine/Classes/AI/Navigation/NavigationData.h
@@ -446,7 +446,7 @@ public:
 	void SetSupportsDefaultAgent(bool bIsDefault) { bSupportsDefaultAgent = bIsDefault; SetNavRenderingEnabled(bIsDefault); }
 	bool IsSupportingDefaultAgent() const { return bSupportsDefaultAgent; }
 
-	virtual bool DoesSupportAgent(const FNavAgentProperties& AgentProps) const { return false; }
+	virtual bool DoesSupportAgent(const FNavAgentProperties& AgentProps) const;
 
 protected:
 	virtual void FillConfig(FNavDataConfig& Dest) { Dest = NavDataConfig; }
diff --git a/Engine/Source/Runtime/Engine/Classes/AI/Navigation/NavigationSystem.h b/Engine/Source/Runtime/Engine/Classes/AI/Navigation/NavigationSystem.h
--- a/Engine/Source/Runtime/Engine/Classes/AI/Navigation/NavigationSystem.h
+++ b/Engine/Source/Runtime/Engine/Classes/AI/Navigation/NavigationSystem.h
@@ -293,6 +293,9 @@ public:
 	// Begin UObject Interface
 	virtual void PostInitProperties() override;
 	static void AddReferencedObjects(UObject* InThis, FReferenceCollector& Collector);
+#if WITH_EDITOR
+	virtual void PostEditChangeProperty(struct FPropertyChangedEvent& PropertyChangedEvent) override;
+#endif // WITH_EDITOR
 	// End UObject Interface
 
 	virtual void Tick(float DeltaSeconds);	
@@ -773,7 +776,7 @@ private:
 
 	/** constructs a navigation data instance of specified NavDataClass, in passed World
 	 *	for supplied NavConfig */
-	virtual ANavigationData* CreateNavigationDataInstance(TSubclassOf<ANavigationData> NavDataClass, UWorld* World, const FNavDataConfig& NavConfig);
+	virtual ANavigationData* CreateNavigationDataInstance(const FNavDataConfig& NavConfig);
 
 	/** Triggers navigation building on all eligible navigation data. */
 	void RebuildAll();
diff --git a/Engine/Source/Runtime/Engine/Classes/AI/Navigation/NavigationTypes.h b/Engine/Source/Runtime/Engine/Classes/AI/Navigation/NavigationTypes.h
--- a/Engine/Source/Runtime/Engine/Classes/AI/Navigation/NavigationTypes.h
+++ b/Engine/Source/Runtime/Engine/Classes/AI/Navigation/NavigationTypes.h
@@ -382,22 +382,22 @@ struct ENGINE_API FNavDataConfig : public FNavAgentProperties
 {
 	GENERATED_USTRUCT_BODY()
 
-	UPROPERTY(EditAnywhere, BlueprintReadWrite, Category=Display)
+	UPROPERTY(EditAnywhere, BlueprintReadWrite, Category = Display)
 	FName Name;
 
-	UPROPERTY(EditAnywhere, BlueprintReadWrite, Category=Display)
+	UPROPERTY(EditAnywhere, BlueprintReadWrite, Category = Display)
 	FColor Color;
 
-	UPROPERTY(EditAnywhere, BlueprintReadOnly, Category=Querying, config)
+	UPROPERTY(EditAnywhere, BlueprintReadOnly, Category = Querying, config)
 	FVector DefaultQueryExtent;
 
-	FNavDataConfig(float Radius = FNavigationSystem::FallbackAgentRadius, float Height = FNavigationSystem::FallbackAgentHeight)
-		: FNavAgentProperties(Radius, Height)
-		, Name(TEXT("Default"))
-		, Color(140,255,0,164)
-		, DefaultQueryExtent(DEFAULT_NAV_QUERY_EXTENT_HORIZONTAL, DEFAULT_NAV_QUERY_EXTENT_HORIZONTAL, DEFAULT_NAV_QUERY_EXTENT_VERTICAL)
-	{
-	}	
+	UPROPERTY(EditAnywhere, BlueprintReadOnly, Category = Navigation, config)
+	TSubclassOf<ANavigationData> NavigationDataClass;
+
+	UPROPERTY(config)
+	FStringClassReference NavigationDataClassName;
+
+	FNavDataConfig(float Radius = FNavigationSystem::FallbackAgentRadius, float Height = FNavigationSystem::FallbackAgentHeight);
 };
 
 struct FNavigationProjectionWork
diff --git a/Engine/Source/Runtime/Engine/Classes/AI/Navigation/RecastNavMesh.h b/Engine/Source/Runtime/Engine/Classes/AI/Navigation/RecastNavMesh.h
--- a/Engine/Source/Runtime/Engine/Classes/AI/Navigation/RecastNavMesh.h
+++ b/Engine/Source/Runtime/Engine/Classes/AI/Navigation/RecastNavMesh.h
@@ -850,10 +850,7 @@ public:
 
 	/** Runs A* pathfinding on navmesh and collect data for every step */
 	int32 DebugPathfinding(const FPathFindingQuery& Query, TArray<FRecastDebugPathfindingStep>& Steps);
-
-	/** Checks whether this instance of navmesh supports given Agent type */
-	virtual bool DoesSupportAgent(const FNavAgentProperties& AgentProps) const override;
-
+	
 	static const FRecastQueryFilter* GetNamedFilter(ERecastNamedFilter::Type FilterType);
 	FORCEINLINE static FNavPolyFlags GetNavLinkFlag() { return NavLinkFlag; }
 	
diff --git a/Engine/Source/Runtime/Engine/Private/AI/Navigation/NavigationData.cpp b/Engine/Source/Runtime/Engine/Private/AI/Navigation/NavigationData.cpp
--- a/Engine/Source/Runtime/Engine/Private/AI/Navigation/NavigationData.cpp
+++ b/Engine/Source/Runtime/Engine/Private/AI/Navigation/NavigationData.cpp
@@ -331,6 +331,11 @@ void ANavigationData::PostEditUndo()
 }
 #endif // WITH_EDITOR
 
+bool ANavigationData::DoesSupportAgent(const FNavAgentProperties& AgentProps) const
+{
+	return NavDataConfig.IsEquivalent(AgentProps);
+}
+
 void ANavigationData::Destroyed()
 {
 	UWorld* WorldOuter = GetWorld();
diff --git a/Engine/Source/Runtime/Engine/Private/AI/Navigation/NavigationSystem.cpp b/Engine/Source/Runtime/Engine/Private/AI/Navigation/NavigationSystem.cpp
--- a/Engine/Source/Runtime/Engine/Private/AI/Navigation/NavigationSystem.cpp
+++ b/Engine/Source/Runtime/Engine/Private/AI/Navigation/NavigationSystem.cpp
@@ -110,6 +110,17 @@ namespace NavigationDebugDrawing
 	const FVector PathNodeBoxExtent(16.f);
 }
 
+//----------------------------------------------------------------------//
+// FNavDataConfig
+//----------------------------------------------------------------------//
+FNavDataConfig::FNavDataConfig(float Radius, float Height)
+	: FNavAgentProperties(Radius, Height)
+	, Name(TEXT("Default"))
+	, Color(140, 255, 0, 164)
+	, DefaultQueryExtent(DEFAULT_NAV_QUERY_EXTENT_HORIZONTAL, DEFAULT_NAV_QUERY_EXTENT_HORIZONTAL, DEFAULT_NAV_QUERY_EXTENT_VERTICAL)
+	, NavigationDataClass(ARecastNavMesh::StaticClass())
+{
+}
 //----------------------------------------------------------------------//
 // FNavigationLockContext                                                                
 //----------------------------------------------------------------------//
@@ -294,7 +305,8 @@ void UNavigationSystem::DoInitialSetup()
 	if (AbstractNavData == NULL)
 	{
 		FNavDataConfig DummyConfig;
-		AbstractNavData = CreateNavigationDataInstance(AAbstractNavData::StaticClass(), NavWorld, DummyConfig);
+		DummyConfig.NavigationDataClass = AAbstractNavData::StaticClass();
+		AbstractNavData = CreateNavigationDataInstance(DummyConfig);
 	}
 
 	CreateCrowdManager();
@@ -332,8 +344,43 @@ void UNavigationSystem::PostInitProperties()
 		FCoreUObjectDelegates::PostLoadMap.AddUObject(this, &UNavigationSystem::OnPostLoadMap);
 		UNavigationSystem::NavigationDirtyEvent.AddUObject(this, &UNavigationSystem::OnNavigationDirtied);
 	}
+	else
+	{
+		// update SupportedActors' navigation classes
+		for (FNavDataConfig& SupportedAgentConfig : SupportedAgents)
+		{
+			if (SupportedAgentConfig.NavigationDataClassName.IsValid())
+			{
+				SupportedAgentConfig.NavigationDataClass = LoadClass<ANavigationData>(NULL, *SupportedAgentConfig.NavigationDataClassName.ToString(), NULL, LOAD_None, NULL);
+			}
+		}
+	}
 }
 
+#if WITH_EDITOR
+void UNavigationSystem::PostEditChangeProperty(FPropertyChangedEvent& PropertyChangedEvent)
+{
+	static const FName NAME_SupportedAgents = GET_MEMBER_NAME_CHECKED(UNavigationSystem, SupportedAgents);
+	static const FName NAME_NavigationDataClass = GET_MEMBER_NAME_CHECKED(FNavDataConfig, NavigationDataClass);
+
+	Super::PostEditChangeProperty(PropertyChangedEvent);
+
+	if (PropertyChangedEvent.Property)
+	{
+		FName PropName = PropertyChangedEvent.Property->GetFName();
+		if (PropName == NAME_NavigationDataClass)
+		{
+			int32 SupportedAgentIndex = PropertyChangedEvent.GetArrayIndex(NAME_SupportedAgents.ToString());
+			if (SupportedAgentIndex != INDEX_NONE)
+			{
+				// reflect the change to SupportedAgent's 
+				SupportedAgents[SupportedAgentIndex].NavigationDataClassName = FStringClassReference::GetOrCreateIDForClass(SupportedAgents[SupportedAgentIndex].NavigationDataClass);
+			}
+		}
+	}
+}
+#endif // WITH_EDITOR
+
 void UNavigationSystem::OnInitializeActors()
 {
 	
@@ -1469,7 +1516,8 @@ UNavigationSystem::ERegistrationResult UNavigationSystem::RegisterNavData(ANavig
 	if (NavConfig.IsValid() == false && SupportedAgents.Num() == 1)
 	{
 		// fill in AgentProps with whatever is the instance's setup
-		NavData->SetConfig((NavConfig = SupportedAgents[0]));
+		NavConfig = SupportedAgents[0];
+		NavData->SetConfig(SupportedAgents[0]);
 		NavData->SetSupportsDefaultAgent(true);	
 		NavData->ProcessNavAreas(NavAreaClasses, 0);
 	}
@@ -1482,19 +1530,19 @@ UNavigationSystem::ERegistrationResult UNavigationSystem::RegisterNavData(ANavig
 		{
 			// ok, so this navigation agent doesn't have its navmesh registered yet, but do we want to support it?
 			bool bAgentSupported = false;
-			FNavDataConfig* Agent = SupportedAgents.GetData();
-			for (int32 AgentIndex = 0; AgentIndex < SupportedAgents.Num(); ++AgentIndex, ++Agent)
+			
+			for (int32 AgentIndex = 0; AgentIndex < SupportedAgents.Num(); ++AgentIndex)
 			{
-				if (Agent->IsEquivalent(NavConfig) == true)
+				if (NavData->GetClass() == SupportedAgents[AgentIndex].NavigationDataClass && SupportedAgents[AgentIndex].IsEquivalent(NavConfig) == true)
 				{
 					// it's supported, then just in case it's not a precise match (IsEquivalent succeeds with some precision) 
 					// update NavData with supported Agent
 					bAgentSupported = true;
 
-					NavData->SetConfig(*Agent);
+					NavData->SetConfig(SupportedAgents[AgentIndex]);
 					if (!NavData->IsA(AAbstractNavData::StaticClass()))
 					{
-						AgentToNavDataMap.Add(*Agent, NavData);
+						AgentToNavDataMap.Add(SupportedAgents[AgentIndex], NavData);
 					}
 
 					NavData->SetSupportsDefaultAgent(AgentIndex == 0);
@@ -2700,14 +2748,13 @@ void UNavigationSystem::SpawnMissingNavigationData()
 	UWorld* NavWorld = GetWorld();
 
 	// 1. check whether any of required navigation data has already been instantiated
-	for (TActorIterator<ARecastNavMesh> It(NavWorld); It && NumberFound < SupportedAgentsCount; ++It)
+	for (TActorIterator<ANavigationData> It(NavWorld); It && NumberFound < SupportedAgentsCount; ++It)
 	{
-		ARecastNavMesh* Nav = (*It);
+		ANavigationData* Nav = (*It);
 		if (Nav != NULL && Nav->GetTypedOuter<UWorld>() == NavWorld && Nav->IsPendingKill() == false)
 		{
 			// find out which one it is
-			const FNavDataConfig* AgentProps = SupportedAgents.GetData();
-			for (int32 AgentIndex = 0; AgentIndex < SupportedAgentsCount; ++AgentIndex, ++AgentProps)
+			for (int32 AgentIndex = 0; AgentIndex < SupportedAgentsCount; ++AgentIndex)
 			{
 				if (AlreadyInstantiated[AgentIndex] == true)
 				{
@@ -2715,7 +2762,7 @@ void UNavigationSystem::SpawnMissingNavigationData()
 					continue;
 				}
 
-				if (Nav->DoesSupportAgent(*AgentProps) == true)
+				if (Nav->GetClass() == SupportedAgents[AgentIndex].NavigationDataClass && Nav->DoesSupportAgent(SupportedAgents[AgentIndex]) == true)
 				{
 					AlreadyInstantiated[AgentIndex] = true;
 					++NumberFound;
@@ -2734,19 +2781,13 @@ void UNavigationSystem::SpawnMissingNavigationData()
 			{
 				bool bHandled = false;
 
-				for (int32 NavClassIndex = 0; NavClassIndex < NavDataClasses.Num(); ++NavClassIndex)
-				{
-					ANavigationData* Instance = CreateNavigationDataInstance(NavDataClasses[NavClassIndex], NavWorld, SupportedAgents[AgentIndex]);
+				ANavigationData* Instance = CreateNavigationDataInstance(SupportedAgents[AgentIndex]);
 
-					if (Instance != NULL)
-					{
-						RequestRegistration(Instance);
-						bHandled = true;
-						break;
-					}
+				if (Instance != NULL)
+				{
+					RequestRegistration(Instance);
 				}
-
-				if (bHandled == false)
+				else 
 				{
 					UE_LOG(LogNavigation, Warning, TEXT("Was not able to create navigation data for SupportedAgent %s (index %d)")
 						, *(SupportedAgents[AgentIndex].Name.ToString()), AgentIndex);
@@ -2764,8 +2805,10 @@ void UNavigationSystem::SpawnMissingNavigationData()
 	}
 }
 
-ANavigationData* UNavigationSystem::CreateNavigationDataInstance(TSubclassOf<ANavigationData> NavDataClass, UWorld* World, const FNavDataConfig& NavConfig)
+ANavigationData* UNavigationSystem::CreateNavigationDataInstance(const FNavDataConfig& NavConfig)
 {
+	TSubclassOf<ANavigationData> NavDataClass = NavConfig.NavigationDataClass;
+	UWorld* World = GetWorld();
 	check(World);
 
 	FActorSpawnParameters SpawnInfo;
diff --git a/Engine/Source/Runtime/Engine/Private/AI/Navigation/RecastNavMesh.cpp b/Engine/Source/Runtime/Engine/Private/AI/Navigation/RecastNavMesh.cpp
--- a/Engine/Source/Runtime/Engine/Private/AI/Navigation/RecastNavMesh.cpp
+++ b/Engine/Source/Runtime/Engine/Private/AI/Navigation/RecastNavMesh.cpp
@@ -720,11 +720,6 @@ void ARecastNavMesh::FillConfig(FNavDataConfig& Dest)
 	Dest.AgentStepHeight = AgentMaxStepHeight;
 }
 
-bool ARecastNavMesh::DoesSupportAgent(const FNavAgentProperties& AgentProps) const
-{
-	return NavDataConfig.IsEquivalent(AgentProps);
-}
-
 void ARecastNavMesh::BeginBatchQuery() const
 {
 #if RECAST_ASYNC_REBUILDING
