Exposed SplineUpDir to Blueprints in SplineMeshComponent.

diff --git a/Engine/Source/Runtime/Engine/Classes/Components/SplineMeshComponent.h b/Engine/Source/Runtime/Engine/Classes/Components/SplineMeshComponent.h
--- a/Engine/Source/Runtime/Engine/Classes/Components/SplineMeshComponent.h
+++ b/Engine/Source/Runtime/Engine/Classes/Components/SplineMeshComponent.h
@@ -249,6 +249,14 @@ class ENGINE_API USplineMeshComponent : public UStaticMeshComponent, public IInt
 	UFUNCTION(BlueprintCallable, Category = SplineMesh)
 	void SetForwardAxis(ESplineMeshAxis::Type InForwardAxis);
 
+	/** Get the spline up direction */
+	UFUNCTION(BlueprintCallable, Category = SplineMesh)
+	FVector GetSplineUpDir() const;
+
+	/** Set the spline up direction */
+	UFUNCTION(BlueprintCallable, Category = SplineMesh)
+	void SetSplineUpDir(const FVector& InSplineUpDir);
+
 	// Destroys the body setup, used to clear collision if the mesh goes missing
 	void DestroyBodySetup();
 #if WITH_EDITOR
diff --git a/Engine/Source/Runtime/Engine/Private/SplineMeshComponent.cpp b/Engine/Source/Runtime/Engine/Private/SplineMeshComponent.cpp
--- a/Engine/Source/Runtime/Engine/Private/SplineMeshComponent.cpp
+++ b/Engine/Source/Runtime/Engine/Private/SplineMeshComponent.cpp
@@ -367,6 +367,17 @@ void USplineMeshComponent::SetForwardAxis(ESplineMeshAxis::Type InForwardAxis)
 	MarkSplineParamsDirty();
 }
 
+FVector USplineMeshComponent::GetSplineUpDir() const
+{
+	return SplineUpDir;
+}
+
+void USplineMeshComponent::SetSplineUpDir(const FVector& InSplineUpDir)
+{
+	SplineUpDir = InSplineUpDir.GetSafeNormal();
+	MarkSplineParamsDirty();
+}
+
 
 void USplineMeshComponent::MarkSplineParamsDirty()
 {
