Fixed UE-6050 "Single-entry logs do not show up in LogVisualizer"
Added yellow color for logs with warnings on LogVisualizer's timelines

diff --git a/Engine/Source/Developer/NewLogVisualizer/Private/STimelineBar.cpp b/Engine/Source/Developer/NewLogVisualizer/Private/STimelineBar.cpp
--- a/Engine/Source/Developer/NewLogVisualizer/Private/STimelineBar.cpp
+++ b/Engine/Source/Developer/NewLogVisualizer/Private/STimelineBar.cpp
@@ -305,7 +305,7 @@ int32 STimelineBar::OnPaint(const FPaintArgs& Args, const FGeometry& AllottedGeo
 			FillImage,
 			MyClippingRect,
 			DrawEffects,
-			ErrorTimeColor
+			WarningTimeColor
 			);
 	}
 
diff --git a/Engine/Source/Developer/NewLogVisualizer/Private/STimelinesContainer.cpp b/Engine/Source/Developer/NewLogVisualizer/Private/STimelinesContainer.cpp
--- a/Engine/Source/Developer/NewLogVisualizer/Private/STimelinesContainer.cpp
+++ b/Engine/Source/Developer/NewLogVisualizer/Private/STimelinesContainer.cpp
@@ -247,6 +247,7 @@ void STimelinesContainer::OnNewLogEntry(const FVisualLogDevice::FVisualLogEntryI
 		{
 			CurrentItem->AddEntry(Entry);
 			bCreateNew = false;
+			break;
 		}
 	}
 
diff --git a/Engine/Source/Developer/NewLogVisualizer/Private/SVisualLogger.cpp b/Engine/Source/Developer/NewLogVisualizer/Private/SVisualLogger.cpp
--- a/Engine/Source/Developer/NewLogVisualizer/Private/SVisualLogger.cpp
+++ b/Engine/Source/Developer/NewLogVisualizer/Private/SVisualLogger.cpp
@@ -263,7 +263,12 @@ void SVisualLogger::Construct(const FArguments& InArgs, const TSharedRef<SDockTa
 					)
 			)
 		);
-	//FGlobalTabmanager::Get()->RestoreFrom(Layout, TSharedPtr<SWindow>());
+	TabManager->SetOnPersistLayout(FTabManager::FOnPersistLayout::CreateRaw(this, &SVisualLogger::HandleTabManagerPersistLayout));
+
+	// Load the potentially previously saved new layout
+	TSharedRef<FTabManager::FLayout> UserConfiguredNewLayout = FLayoutSaveRestore::LoadFromConfig(GEditorLayoutIni, Layout);
+
+	TabManager->RestoreFrom(UserConfiguredNewLayout, TSharedPtr<SWindow>());
 
 	// Window Menu
 	FMenuBarBuilder MenuBarBuilder = FMenuBarBuilder(TSharedPtr<FUICommandList>());
@@ -308,7 +313,6 @@ void SVisualLogger::Construct(const FArguments& InArgs, const TSharedRef<SDockTa
 		];
 
 	VisualLoggerCanvasRenderer = MakeShareable(new FVisualLoggerCanvasRenderer(VisualLoggerInterface));
-	ConstructUnderMajorTab->SetOnPersistVisualState(SDockTab::FOnPersistVisualState::CreateRaw(this, &SVisualLogger::HandleMajorTabPersistVisualState));
 
 	UDebugDrawService::Register(TEXT("VisLog"), FDebugDrawDelegate::CreateRaw(VisualLoggerCanvasRenderer.Get(), &FVisualLoggerCanvasRenderer::DrawOnCanvas));
 	//UGameplayDebuggingComponent::OnDebuggingTargetChangedDelegate.AddSP(this, &SVisualLogger::SelectionChanged);
@@ -319,6 +323,10 @@ void SVisualLogger::HandleMajorTabPersistVisualState()
 	// save any settings here
 }
 
+void SVisualLogger::HandleTabManagerPersistLayout(const TSharedRef<FTabManager::FLayout>& LayoutToSave)
+{
+	FLayoutSaveRestore::SaveToConfig(GEditorLayoutIni, LayoutToSave);
+}
 
 void SVisualLogger::FillWindowMenu(FMenuBuilder& MenuBuilder, const TSharedPtr<FTabManager> TabManager)
 {
diff --git a/Engine/Source/Developer/NewLogVisualizer/Private/SVisualLogger.h b/Engine/Source/Developer/NewLogVisualizer/Private/SVisualLogger.h
--- a/Engine/Source/Developer/NewLogVisualizer/Private/SVisualLogger.h
+++ b/Engine/Source/Developer/NewLogVisualizer/Private/SVisualLogger.h
@@ -87,7 +87,7 @@ public:
 	bool HandleCameraCommandIsChecked() const;
 
 	TSharedPtr<SVisualLoggerFilters> GetVisualLoggerFilters() { return VisualLoggerFilters; }
-
+	void HandleTabManagerPersistLayout(const TSharedRef<FTabManager::FLayout>& LayoutToSave);
 protected:
 	// Holds the list of UI commands.
 	TSharedRef<FUICommandList> CommandList;
diff --git a/Engine/Source/Developer/NewLogVisualizer/Private/SVisualLoggerFilters.h b/Engine/Source/Developer/NewLogVisualizer/Private/SVisualLoggerFilters.h
--- a/Engine/Source/Developer/NewLogVisualizer/Private/SVisualLoggerFilters.h
+++ b/Engine/Source/Developer/NewLogVisualizer/Private/SVisualLoggerFilters.h
@@ -46,7 +46,7 @@ public:
 
 	TSharedRef<SWidget> MakeGraphsFilterMenu();
 	void CreateFiltersMenuCategoryForGraph(FMenuBuilder& MenuBuilder, FName MenuCategory) const;
-
+	bool HasFilters() { return Filters.Num() || GraphFilters.Num(); }
 protected:
 	/** The horizontal box which contains all the filters */
 	TSharedPtr<SWrapBox> FilterBox;
diff --git a/Engine/Source/Runtime/Engine/Private/VisualLogger/VisualLogger.cpp b/Engine/Source/Runtime/Engine/Private/VisualLogger/VisualLogger.cpp
--- a/Engine/Source/Runtime/Engine/Private/VisualLogger/VisualLogger.cpp
+++ b/Engine/Source/Runtime/Engine/Private/VisualLogger/VisualLogger.cpp
@@ -4,6 +4,16 @@
 #include "VisualLogger/VisualLogger.h"
 #include "VisualLogger/VisualLogger.h"
 #include "VisualLogger/VisualLoggerBinaryFileDevice.h"
+#if WITH_EDITOR
+#	include "Editor/UnrealEd/Public/EditorComponents.h"
+#	include "Editor/UnrealEd/Public/EditorReimportHandler.h"
+#	include "Editor/UnrealEd/Public/TexAlignTools.h"
+#	include "Editor/UnrealEd/Public/TickableEditorObject.h"
+#	include "UnrealEdClasses.h"
+#	include "Editor/UnrealEd/Public/Editor.h"
+#	include "Editor/UnrealEd/Public/EditorViewportClient.h"
+#endif
+
 
 #if ENABLE_VISUAL_LOG 
 
@@ -26,6 +36,24 @@ FVisualLogger::FVisualLogger()
 	}
 }
 
+UWorld* FVisualLogger::GetWorld()
+{
+	UWorld* World = NULL;
+	UEditorEngine *EEngine = Cast<UEditorEngine>(GEngine);
+	if (GIsEditor && EEngine != NULL)
+	{
+		// lets use PlayWorld during PIE/Simulate and regular world from editor otherwise, to draw debug information
+		World = EEngine->PlayWorld != NULL ? EEngine->PlayWorld : EEngine->GetEditorWorldContext().World();
+	}
+	else if (!GIsEditor)
+	{
+
+		World = GEngine->GetWorld();
+	}
+
+	return World;
+}
+
 void FVisualLogger::Shutdown()
 {
 	SetIsRecording(false);
diff --git a/Engine/Source/Runtime/Engine/Public/VisualLogger/VisualLogger.h b/Engine/Source/Runtime/Engine/Public/VisualLogger/VisualLogger.h
--- a/Engine/Source/Runtime/Engine/Public/VisualLogger/VisualLogger.h
+++ b/Engine/Source/Runtime/Engine/Public/VisualLogger/VisualLogger.h
@@ -212,6 +212,7 @@ public:
 private:
 	FVisualLogger();
 	virtual void Serialize(const TCHAR* V, ELogVerbosity::Type Verbosity, const class FName& Category) override { ensureMsgf(0, TEXT("Regular serialize is forbiden for visual logs")); }
+	class UWorld* GetWorld();
 
 protected:
 	/** Array of output devices to redirect to */
diff --git a/Engine/Source/Runtime/Engine/Public/VisualLogger/VisualLogger.inl b/Engine/Source/Runtime/Engine/Public/VisualLogger/VisualLogger.inl
--- a/Engine/Source/Runtime/Engine/Public/VisualLogger/VisualLogger.inl
+++ b/Engine/Source/Runtime/Engine/Public/VisualLogger/VisualLogger.inl
@@ -66,22 +66,19 @@ FVisualLogEntry* FVisualLogger::GetEntryToWrite(const class UObject* Object, flo
 	UObject * LogOwner = FVisualLogger::FindRedirection(Object);
 
 	bool InitializeNewEntry = false;
+
+	UWorld* World = GetWorld();
 	if (CurrentEntryPerObject.Contains(LogOwner))
 	{
 		CurrentEntry = &CurrentEntryPerObject[LogOwner];
-		if (TimeStamp > CurrentEntry->TimeStamp && ShouldCreate == ECreateIfNeeded::Create)
+		InitializeNewEntry = TimeStamp > CurrentEntry->TimeStamp && ShouldCreate == ECreateIfNeeded::Create;
+		if (World)
 		{
-			if (CurrentEntry->TimeStamp >= 0) //-1 means not initialized entry information
-			{
-				for (auto* Device : OutputDevices)
-				{
-					Device->Serialize(LogOwner, ObjectToNameMap[LogOwner], *CurrentEntry);
-				}
-			}
-			InitializeNewEntry = true;
+			World->GetTimerManager().ClearAllTimersForObject(this);
 		}
 	}
-	else
+
+	if(!CurrentEntry)
 	{
 		CurrentEntry = &CurrentEntryPerObject.Add(LogOwner);
 		ObjectToNameMap.Add(LogOwner, LogOwner->GetFName());
@@ -125,6 +122,27 @@ FVisualLogEntry* FVisualLogger::GetEntryToWrite(const class UObject* Object, flo
 				ObjectAsActor->GrabDebugSnapshot(CurrentEntry);
 			}
 		}
+
+		if (World)
+		{
+			//set next tick timer to flush obsolete/old entries
+			World->GetTimerManager().SetTimerForNextTick(FTimerDelegate::CreateLambda(
+				[this](){
+				for (auto& CurrentPair : CurrentEntryPerObject)
+				{
+					FVisualLogEntry* Entry = &CurrentPair.Value;
+					if (Entry->TimeStamp >= 0) // CurrentEntry->TimeStamp == -1 means it's not initialized entry information
+					{
+						for (auto* Device : OutputDevices)
+						{
+							Device->Serialize(CurrentPair.Key, ObjectToNameMap[CurrentPair.Key], *Entry);
+						}
+						Entry->Reset();
+					}
+				}
+			}
+			));
+		}
 	}
 
 	return CurrentEntry;
