added support for dynamic obstacles from static meshes without cooked convex collision data (bGatherConvexGeometry = false)

diff --git a/Engine/Source/Runtime/Engine/Classes/AI/Navigation/NavCollision.h b/Engine/Source/Runtime/Engine/Classes/AI/Navigation/NavCollision.h
--- a/Engine/Source/Runtime/Engine/Classes/AI/Navigation/NavCollision.h
+++ b/Engine/Source/Runtime/Engine/Classes/AI/Navigation/NavCollision.h
@@ -95,7 +95,7 @@ class ENGINE_API UNavCollision : public UObject
 	void DrawSimpleGeom(FPrimitiveDrawInterface* PDI, const FTransform& Transform, const FColor Color);
 
 	/** Get data for dynamic obstacle */
-	void GetNavigationModifier(FCompositeNavModifier& Modifier, const FTransform& LocalToWorld) const;
+	void GetNavigationModifier(FCompositeNavModifier& Modifier, const FTransform& LocalToWorld);
 
 	/** Read collisions data */
 	bool GatherCollision();
diff --git a/Engine/Source/Runtime/Engine/Private/AI/Navigation/NavCollision.cpp b/Engine/Source/Runtime/Engine/Private/AI/Navigation/NavCollision.cpp
--- a/Engine/Source/Runtime/Engine/Private/AI/Navigation/NavCollision.cpp
+++ b/Engine/Source/Runtime/Engine/Private/AI/Navigation/NavCollision.cpp
@@ -206,7 +206,7 @@ void UNavCollision::ClearCollision()
 	bHasConvexGeometry = false;
 }
 
-void UNavCollision::GetNavigationModifier(FCompositeNavModifier& Modifier, const FTransform& LocalToWorld) const
+void UNavCollision::GetNavigationModifier(FCompositeNavModifier& Modifier, const FTransform& LocalToWorld)
 {
 	const TSubclassOf<UNavArea> UseAreaClass = AreaClass ? AreaClass : UNavigationSystem::GetDefaultObstacleArea();
 
@@ -232,6 +232,12 @@ void UNavCollision::GetNavigationModifier(FCompositeNavModifier& Modifier, const
 		Modifier.Add(AreaMod);
 	}
 
+	// rebuild collision data if needed
+	if (!bHasConvexGeometry)
+	{
+		bHasConvexGeometry = GatherCollision();
+	}
+
 	if (ConvexCollision.VertexBuffer.Num() > 0)
 	{
 		int32 LastVertIndex = 0;
