Fixed CIS - not editor builds

diff --git a/Engine/Source/Runtime/Engine/Private/VisualLogger/VisualLogger.cpp b/Engine/Source/Runtime/Engine/Private/VisualLogger/VisualLogger.cpp
--- a/Engine/Source/Runtime/Engine/Private/VisualLogger/VisualLogger.cpp
+++ b/Engine/Source/Runtime/Engine/Private/VisualLogger/VisualLogger.cpp
@@ -39,13 +39,16 @@ FVisualLogger::FVisualLogger()
 UWorld* FVisualLogger::GetWorld()
 {
 	UWorld* World = NULL;
+#if WITH_EDITOR
 	UEditorEngine *EEngine = Cast<UEditorEngine>(GEngine);
 	if (GIsEditor && EEngine != NULL)
 	{
 		// lets use PlayWorld during PIE/Simulate and regular world from editor otherwise, to draw debug information
 		World = EEngine->PlayWorld != NULL ? EEngine->PlayWorld : EEngine->GetEditorWorldContext().World();
 	}
-	else if (!GIsEditor)
+	else 
+#endif
+	if (!GIsEditor)
 	{
 
 		World = GEngine->GetWorld();
