Improved cook in editor hitches.

diff --git a/Engine/Source/Runtime/Engine/Private/TextureDerivedData.cpp b/Engine/Source/Runtime/Engine/Private/TextureDerivedData.cpp
--- a/Engine/Source/Runtime/Engine/Private/TextureDerivedData.cpp
+++ b/Engine/Source/Runtime/Engine/Private/TextureDerivedData.cpp
@@ -49,6 +49,7 @@ namespace TextureDerivedDataTimings
 		SyncBlockTime,
 		BuildTextureTime,
 		SerializeCookedTime,
+		SyncDoesCachedDataExist,
 		NumTimings
 	};
 
@@ -60,7 +61,8 @@ namespace TextureDerivedDataTimings
 		TEXT("Asynchronous Block"),
 		TEXT("Synchronous Loads"),
 		TEXT("Build Textures"),
-		TEXT("Serialize Cooked")
+		TEXT("Serialize Cooked"),
+		TEXT("Sync Does Cached Data Exist")
 	};
 
 	void PrintTimings()
@@ -1439,13 +1441,13 @@ void UTexture::BeginCacheForCookedPlatformData( const ITargetPlatform *TargetPla
 			}
 		}
 
-
-
 		uint32 CacheFlags = ETextureCacheFlags::Async | ETextureCacheFlags::InlineMips;
 
 		// If source data is resident in memory then allow the texture to be built
 		// in a background thread.
-		if (Source.BulkData.IsBulkDataLoaded())
+		bool bAllowAsyncBuild = Source.BulkData.IsBulkDataLoaded();
+
+		if (bAllowAsyncBuild)
 		{
 			CacheFlags |= ETextureCacheFlags::AllowAsyncBuild;
 		}
@@ -1461,12 +1463,24 @@ void UTexture::BeginCacheForCookedPlatformData( const ITargetPlatform *TargetPla
 
 			if ( PlatformData == NULL )
 			{
+				uint32 CurrentCacheFlags = CacheFlags;
+				// if the cached data key exists already then we don't need to allowasync build
+				// if it doesn't then allow async builds
+				
+				{
+					TextureDerivedDataTimings::FScopedMeasurement Timer(TextureDerivedDataTimings::SyncDoesCachedDataExist);
+					if ( GetDerivedDataCacheRef().CachedDataProbablyExists( *DerivedDataKey ) == false )
+					{
+						CurrentCacheFlags |= ETextureCacheFlags::AllowAsyncBuild;
+					}
+				}
+
 				FTexturePlatformData* PlatformDataToCache;
 				PlatformDataToCache = new FTexturePlatformData();
 				PlatformDataToCache->Cache(
 					*this,
 					BuildSettingsToCache[SettingsIndex],
-					CacheFlags
+					CurrentCacheFlags
 					);
 				CookedPlatformData.Add( DerivedDataKey, PlatformDataToCache );
 			}
