UE-5668 Fix skelmeshcomp not returning local space bounds when identity is passed to CalcBounds

diff --git a/Engine/Source/Editor/PhAT/Private/PhATPreviewViewportClient.cpp b/Engine/Source/Editor/PhAT/Private/PhATPreviewViewportClient.cpp
--- a/Engine/Source/Editor/PhAT/Private/PhATPreviewViewportClient.cpp
+++ b/Engine/Source/Editor/PhAT/Private/PhATPreviewViewportClient.cpp
@@ -52,7 +52,7 @@ FPhATEdPreviewViewportClient::FPhATEdPreviewViewportClient(TWeakPtr<FPhAT> InPhA
 	EngineShowFlags.CompositeEditorPrimitives = true;
 
 	// Get actors asset collision bounding box, and move actor so its not intersection the floor plane at Z = 0.
-	FBox CollBox = SharedData->PhysicsAsset->CalcAABB(SharedData->EditorSkelComp);	
+	FBox CollBox = SharedData->PhysicsAsset->CalcAABB(SharedData->EditorSkelComp, SharedData->EditorSkelComp->ComponentToWorld);	
 	FVector SkelCompLocation = FVector(0, 0, -CollBox.Min.Z + SharedData->EditorSimOptions->FloorGap);
 
 	SharedData->EditorSkelComp->SetAbsolute(true, true, true);
@@ -60,7 +60,7 @@ FPhATEdPreviewViewportClient::FPhATEdPreviewViewportClient(TWeakPtr<FPhAT> InPhA
 	SharedData->ResetTM = SharedData->EditorSkelComp->GetComponentToWorld();
 
 	// Get new bounding box and set view based on that.
-	CollBox = SharedData->PhysicsAsset->CalcAABB(SharedData->EditorSkelComp);	
+	CollBox = SharedData->PhysicsAsset->CalcAABB(SharedData->EditorSkelComp, SharedData->EditorSkelComp->ComponentToWorld);	
 	FVector CollBoxExtent = CollBox.GetExtent();
 
 	// Take into account internal mesh translation/rotation/scaling etc.
diff --git a/Engine/Source/Runtime/Engine/Classes/Components/SkinnedMeshComponent.h b/Engine/Source/Runtime/Engine/Classes/Components/SkinnedMeshComponent.h
--- a/Engine/Source/Runtime/Engine/Classes/Components/SkinnedMeshComponent.h
+++ b/Engine/Source/Runtime/Engine/Classes/Components/SkinnedMeshComponent.h
@@ -704,6 +704,15 @@ public:
 	 */
 	FMatrix GetBoneMatrix( int32 BoneIndex ) const;
 
+	/** 
+	 * Get world space bone transform from bone index, also specifying the component transform to use
+	 * 
+	 * @param BoneIndex Index of the bone
+	 *
+	 * @return the transform of the bone at the specified index 
+	 */
+	FTransform GetBoneTransform( int32 BoneIndex, const FTransform& LocalToWorld ) const;
+
 	/** 
 	 * Get Bone Transform from index
 	 * 
diff --git a/Engine/Source/Runtime/Engine/Classes/PhysicsEngine/PhysicsAsset.h b/Engine/Source/Runtime/Engine/Classes/PhysicsEngine/PhysicsAsset.h
--- a/Engine/Source/Runtime/Engine/Classes/PhysicsEngine/PhysicsAsset.h
+++ b/Engine/Source/Runtime/Engine/Classes/PhysicsEngine/PhysicsAsset.h
@@ -82,7 +82,7 @@ public:
 
 	ENGINE_API void			GetNearestBodyIndicesBelow(TArray<int32> & OutBodyIndices, FName InBoneName, USkeletalMesh * InSkelMesh);
 
-	ENGINE_API FBox			CalcAABB(const class USkinnedMeshComponent* MeshComponent) const;
+	ENGINE_API FBox			CalcAABB(const class USkinnedMeshComponent* MeshComponent, const FTransform& LocalToWorld) const;
 
 	/** Clears physics meshes from all bodies */
 	ENGINE_API void ClearAllPhysicsMeshes();
diff --git a/Engine/Source/Runtime/Engine/Private/PhysicsEngine/PhysicsAsset.cpp b/Engine/Source/Runtime/Engine/Private/PhysicsEngine/PhysicsAsset.cpp
--- a/Engine/Source/Runtime/Engine/Private/PhysicsEngine/PhysicsAsset.cpp
+++ b/Engine/Source/Runtime/Engine/Private/PhysicsEngine/PhysicsAsset.cpp
@@ -116,7 +116,7 @@ void UPhysicsAsset::DisableCollision(int32 BodyIndexA, int32 BodyIndexB)
 	CollisionDisableTable.Add(Key, 0);
 }
 
-FBox UPhysicsAsset::CalcAABB(const USkinnedMeshComponent* MeshComp) const
+FBox UPhysicsAsset::CalcAABB(const USkinnedMeshComponent* MeshComp, const FTransform& LocalToWorld) const
 {
 	FBox Box(0);
 
@@ -125,7 +125,7 @@ FBox UPhysicsAsset::CalcAABB(const USkinnedMeshComponent* MeshComp) const
 		return Box;
 	}
 
-	FVector Scale3D = MeshComp->ComponentToWorld.GetScale3D();
+	FVector Scale3D = LocalToWorld.GetScale3D();
 	if( Scale3D.IsUniform() )
 	{
 		const TArray<int32>* BodyIndexRefs = NULL;
@@ -166,7 +166,7 @@ FBox UPhysicsAsset::CalcAABB(const USkinnedMeshComponent* MeshComp) const
 				int32 BoneIndex = MeshComp->GetBoneIndex(bs->BoneName);
 				if(BoneIndex != INDEX_NONE)
 				{
-					FTransform WorldBoneTransform = MeshComp->GetBoneTransform(BoneIndex);
+					FTransform WorldBoneTransform = MeshComp->GetBoneTransform(BoneIndex, LocalToWorld);
 					if(FMath::Abs(WorldBoneTransform.GetDeterminant()) > (float)KINDA_SMALL_NUMBER)
 					{
 						Box += bs->AggGeom.CalcAABB( WorldBoneTransform );
@@ -182,7 +182,7 @@ FBox UPhysicsAsset::CalcAABB(const USkinnedMeshComponent* MeshComp) const
 
 	if(!Box.IsValid)
 	{
-		Box = FBox( MeshComp->GetComponentLocation(), MeshComp->GetComponentLocation() );
+		Box = FBox( LocalToWorld.GetLocation(), LocalToWorld.GetLocation() );
 	}
 
 	return Box;
diff --git a/Engine/Source/Runtime/Engine/Private/SkinnedMeshComponent.cpp b/Engine/Source/Runtime/Engine/Private/SkinnedMeshComponent.cpp
--- a/Engine/Source/Runtime/Engine/Private/SkinnedMeshComponent.cpp
+++ b/Engine/Source/Runtime/Engine/Private/SkinnedMeshComponent.cpp
@@ -497,21 +497,18 @@ FBoxSphereBounds USkinnedMeshComponent::CalcMeshBound(const FVector& RootOffset,
 	// For AnimSet Viewer, use 'bounds preview' physics asset if present.
 	else if(SkeletalMesh && bHasPhysBodies && bCanUsePhysicsAsset)
 	{
-		// @fixme UE4 this does not use LocalToWorld entered but ComponentToWorld
-		NewBounds = FBoxSphereBounds(PhysicsAsset->CalcAABB(this));
+		NewBounds = FBoxSphereBounds(PhysicsAsset->CalcAABB(this, LocalToWorld));
 	}
 #endif // WITH_EDITOR
 	// If we have a PhysicsAsset (with at least one matching bone), and we can use it, do so to calc bounds.
 	else if( bHasPhysBodies && bCanUsePhysicsAsset && UsePhysicsAsset )
 	{
-		// @fixme UE4 this does not use LocalToWorld entered but ComponentToWorld
-		NewBounds = FBoxSphereBounds(PhysicsAsset->CalcAABB(this));
+		NewBounds = FBoxSphereBounds(PhysicsAsset->CalcAABB(this, LocalToWorld));
 	}
 	// Use MasterPoseComponent's PhysicsAsset, if we don't have one and it does
 	else if(MasterPoseComponent.IsValid() && bCanUsePhysicsAsset && bMasterHasPhysBodies)
 	{
-		// @fixme UE4 this does not use LocalToWorld entered but ComponentToWorld		
-		NewBounds = FBoxSphereBounds(MasterPhysicsAsset->CalcAABB(this));
+		NewBounds = FBoxSphereBounds(MasterPhysicsAsset->CalcAABB(this, LocalToWorld));
 	}
 	// Fallback is to use the one from the skeletal mesh. Usually pretty bad in terms of Accuracy of where the SkelMesh Bounds are located (i.e. usually bigger than it needs to be)
 	else if( SkeletalMesh )
@@ -585,16 +582,20 @@ FMatrix USkinnedMeshComponent::GetBoneMatrix(int32 BoneIdx) const
 	}
 }
 
-
 FTransform USkinnedMeshComponent::GetBoneTransform(int32 BoneIdx) const
 {
-	if ( !IsRegistered() )
+	if (!IsRegistered())
 	{
 		// if not registered, we don't have SpaceBases yet. 
 		// also ComponentToWorld isn't set yet (They're set from relativelocation, relativerotation, relativescale)
 		return FTransform::Identity;
 	}
 
+	return GetBoneTransform(BoneIdx, ComponentToWorld);
+}
+
+FTransform USkinnedMeshComponent::GetBoneTransform(int32 BoneIdx, const FTransform& LocalToWorld) const
+{
 	// Handle case of use a MasterPoseComponent - get bone matrix from there.
 	if(MasterPoseComponent.IsValid())
 	{
@@ -606,7 +607,7 @@ FTransform USkinnedMeshComponent::GetBoneTransform(int32 BoneIdx) const
 			if(	ParentBoneIndex != INDEX_NONE && 
 				ParentBoneIndex < MasterPoseComponent->GetNumSpaceBases())
 			{
-				return MasterPoseComponent->GetSpaceBases()[ParentBoneIndex] * ComponentToWorld;
+				return MasterPoseComponent->GetSpaceBases()[ParentBoneIndex] * LocalToWorld;
 			}
 			else
 			{
@@ -624,7 +625,7 @@ FTransform USkinnedMeshComponent::GetBoneTransform(int32 BoneIdx) const
 	{
 		if( GetNumSpaceBases() && BoneIdx < GetNumSpaceBases() )
 		{
-			return GetSpaceBases()[BoneIdx] * ComponentToWorld;
+			return GetSpaceBases()[BoneIdx] * LocalToWorld;
 		}
 		else
 		{
