Fix issue with double buffered blend spaces and bone interp

diff --git a/Engine/Source/Runtime/Engine/Private/SkeletalMeshComponent.cpp b/Engine/Source/Runtime/Engine/Private/SkeletalMeshComponent.cpp
--- a/Engine/Source/Runtime/Engine/Private/SkeletalMeshComponent.cpp
+++ b/Engine/Source/Runtime/Engine/Private/SkeletalMeshComponent.cpp
@@ -1020,6 +1020,12 @@ void USkeletalMeshComponent::PostAnimEvaluation(FAnimationEvaluationContext& Eva
 
 		const float Alpha = 0.25f + (1.f / float(FMath::Max(AnimUpdateRateParams.GetEvaluationRate(), 2) * 2));
 		FAnimationRuntime::LerpBoneTransforms(LocalAtoms, CachedLocalAtoms, Alpha, RequiredBones);
+		if (bDoubleBufferedBlendSpaces)
+		{
+			// We need to prep our space bases for interp (TODO: Would a new Lerp function that took
+			// separate input and output be quicker than current copy + lerp in place?)
+			GetEditableSpaceBases() = GetSpaceBases();
+		}
 		FAnimationRuntime::LerpBoneTransforms(GetEditableSpaceBases(), CachedSpaceBases, Alpha, RequiredBones);
 	}
 
