Height fog UI ranges

diff --git a/Engine/Source/Runtime/Engine/Classes/Components/ExponentialHeightFogComponent.h b/Engine/Source/Runtime/Engine/Classes/Components/ExponentialHeightFogComponent.h
--- a/Engine/Source/Runtime/Engine/Classes/Components/ExponentialHeightFogComponent.h
+++ b/Engine/Source/Runtime/Engine/Classes/Components/ExponentialHeightFogComponent.h
@@ -17,14 +17,14 @@ class ENGINE_API UExponentialHeightFogComponent : public USceneComponent
 	uint32 bEnabled_DEPRECATED:1;
 
 	/** Global density factor. */
-	UPROPERTY(BlueprintReadOnly, interp, Category=ExponentialHeightFogComponent)
+	UPROPERTY(BlueprintReadOnly, interp, Category=ExponentialHeightFogComponent, meta=(UIMin = "0", UIMax = ".05"))
 	float FogDensity;
 
 	UPROPERTY(BlueprintReadOnly, interp, Category=ExponentialHeightFogComponent)
 	FLinearColor FogInscatteringColor;
 
 	/** Controls the size of the directional inscattering cone, which is used to approximate inscattering from a directional light. */
-	UPROPERTY(BlueprintReadOnly, interp, Category=DirectionalInscattering)
+	UPROPERTY(BlueprintReadOnly, interp, Category=DirectionalInscattering, meta=(UIMin = "2", UIMax = "64"))
 	float DirectionalInscatteringExponent;
 
 	/** Controls the start distance from the viewer of the directional inscattering, which is used to approximate inscattering from a directional light. */
@@ -39,7 +39,7 @@ class ENGINE_API UExponentialHeightFogComponent : public USceneComponent
 	 * Height density factor, controls how the density increases as height decreases.  
 	 * Smaller values make the visible transition larger.
 	 */
-	UPROPERTY(BlueprintReadOnly, interp, Category=ExponentialHeightFogComponent)
+	UPROPERTY(BlueprintReadOnly, interp, Category=ExponentialHeightFogComponent, meta=(UIMin = "0.001", UIMax = "2"))
 	float FogHeightFalloff;
 
 	/** 
@@ -47,11 +47,11 @@ class ENGINE_API UExponentialHeightFogComponent : public USceneComponent
 	 * A value of 1 means the fog can become fully opaque at a distance and replace scene color completely,
 	 * A value of 0 means the fog color will not be factored in at all.
 	 */
-	UPROPERTY(BlueprintReadOnly, interp, Category=ExponentialHeightFogComponent)
+	UPROPERTY(BlueprintReadOnly, interp, Category=ExponentialHeightFogComponent, meta=(UIMin = "0", UIMax = "1"))
 	float FogMaxOpacity;
 
 	/** Distance from the camera that the fog will start, in world units. */
-	UPROPERTY(BlueprintReadOnly, interp, Category=ExponentialHeightFogComponent)
+	UPROPERTY(BlueprintReadOnly, interp, Category=ExponentialHeightFogComponent, meta=(UIMin = "0", UIMax = "5000"))
 	float StartDistance;
 
 	
@@ -100,4 +100,3 @@ public:
 };
 
 
-
