Fix for UE-6076 "NavMesh debug collision data rendering is broken"

diff --git a/Engine/Source/Runtime/Engine/Classes/AI/Navigation/RecastNavMesh.h b/Engine/Source/Runtime/Engine/Classes/AI/Navigation/RecastNavMesh.h
--- a/Engine/Source/Runtime/Engine/Classes/AI/Navigation/RecastNavMesh.h
+++ b/Engine/Source/Runtime/Engine/Classes/AI/Navigation/RecastNavMesh.h
@@ -444,6 +444,9 @@ class ENGINE_API ARecastNavMesh : public ANavigationData
 	UPROPERTY(config)
 	uint32 bDistinctlyDrawTilesBeingBuilt:1;
 
+	UPROPERTY(EditAnywhere, Category = Display)
+	uint32 bDrawNavMesh : 1;
+
 	/** vertical offset added to navmesh's debug representation for better readability */
 	UPROPERTY(EditAnywhere, Category=Display, config)
 	float DrawOffset;
diff --git a/Engine/Source/Runtime/Engine/Private/AI/Navigation/RecastNavMesh.cpp b/Engine/Source/Runtime/Engine/Private/AI/Navigation/RecastNavMesh.cpp
--- a/Engine/Source/Runtime/Engine/Private/AI/Navigation/RecastNavMesh.cpp
+++ b/Engine/Source/Runtime/Engine/Private/AI/Navigation/RecastNavMesh.cpp
@@ -198,6 +198,7 @@ ARecastNavMesh::ARecastNavMesh(const FObjectInitializer& ObjectInitializer)
 	, bDrawNavMeshEdges(true)
 	, bDrawNavLinks(true)
 	, bDistinctlyDrawTilesBeingBuilt(true)
+	, bDrawNavMesh(true)
 	, DrawOffset(10.f)
 	, TilePoolSize(1024)
 	, MaxSimplificationError(1.3f)	// from RecastDemo
diff --git a/Engine/Source/Runtime/Engine/Private/NavMeshRenderingComponent.cpp b/Engine/Source/Runtime/Engine/Private/NavMeshRenderingComponent.cpp
--- a/Engine/Source/Runtime/Engine/Private/NavMeshRenderingComponent.cpp
+++ b/Engine/Source/Runtime/Engine/Private/NavMeshRenderingComponent.cpp
@@ -183,35 +183,6 @@ void UNavMeshRenderingComponent::GatherData(struct FNavMeshSceneProxyData* Curre
 			}
 		}
 
-		if (NavMesh->bDrawPathCollidingGeometry)
-		{
-			// draw all geometry gathered in navoctree
-			const FNavigationOctree* NavOctree = NavMesh->GetWorld()->GetNavigationSystem()->GetNavOctree();
-
-			for (FNavigationOctree::TConstIterator<> It(*NavOctree); It.HasPendingNodes(); It.Advance())
-			{
-				const FNavigationOctree::FNode& Node = It.GetCurrentNode();
-				for (FNavigationOctree::ElementConstIt ElementIt(Node.GetElementIt()); ElementIt; ElementIt++)
-				{
-					const FNavigationOctreeElement& Element = *ElementIt;
-					if (Element.ShouldUseGeometry(&NavMesh->NavDataConfig) && Element.Data.CollisionData.Num())
-					{
-						const FRecastGeometryCache CachedGeometry(Element.Data.CollisionData.GetData());
-						AppendGeometry(CurrentData->PathCollidingGeomVerts, CurrentData->PathCollidingGeomIndices,
-							CachedGeometry.Verts, CachedGeometry.Header.NumVerts, CachedGeometry.Indices, CachedGeometry.Header.NumFaces);
-					}
-				}
-
-				FOREACH_OCTREE_CHILD_NODE(ChildRef)
-				{
-					if (Node.HasChild(ChildRef))
-					{
-						It.PushChild(ChildRef);
-					}
-				}
-			}
-		}
-		
 		// offset all navigation-link positions
 		if (!NavMesh->bDrawClusters)
 		{
@@ -367,7 +338,7 @@ void UNavMeshRenderingComponent::GatherData(struct FNavMeshSceneProxyData* Curre
 				CurrentData->MeshBuilders.Add(DebugMeshData);
 			}
 		}
-		else
+		else if (NavMesh->bDrawNavMesh)
 		{			
 			for (int32 AreaType = 0; AreaType < RECAST_MAX_AREAS; ++AreaType)
 			{
@@ -393,17 +364,39 @@ void UNavMeshRenderingComponent::GatherData(struct FNavMeshSceneProxyData* Curre
 			}
 		}
 
-		// Draw path generation input geometry
-		if (CurrentData->bDrawPathCollidingGeometry)
+		if (NavMesh->bDrawPathCollidingGeometry)
 		{
-			FNavMeshSceneProxyData::FDebugMeshData DebugMeshData;
-			for (int32 VertIdx=0; VertIdx < CurrentData->PathCollidingGeomVerts.Num(); ++VertIdx)
+			// draw all geometry gathered in navoctree
+			const FNavigationOctree* NavOctree = NavMesh->GetWorld()->GetNavigationSystem()->GetNavOctree();
+
+			TArray<FVector> PathCollidingGeomVerts;
+			TArray <int32> PathCollidingGeomIndices;
+			for (FNavigationOctree::TConstIterator<> It(*NavOctree); It.HasPendingNodes(); It.Advance())
+			{
+				const FNavigationOctree::FNode& Node = It.GetCurrentNode();
+				for (FNavigationOctree::ElementConstIt ElementIt(Node.GetElementIt()); ElementIt; ElementIt++)
+				{
+					const FNavigationOctreeElement& Element = *ElementIt;
+					if (Element.ShouldUseGeometry(&NavMesh->NavDataConfig) && Element.Data.CollisionData.Num())
+					{
+						const FRecastGeometryCache CachedGeometry(Element.Data.CollisionData.GetData());
+						AppendGeometry(PathCollidingGeomVerts, PathCollidingGeomIndices, CachedGeometry.Verts, CachedGeometry.Header.NumVerts, CachedGeometry.Indices, CachedGeometry.Header.NumFaces);
+					}
+				}
+				FOREACH_OCTREE_CHILD_NODE(ChildRef)
+				{
+					if (Node.HasChild(ChildRef))
+					{
+						It.PushChild(ChildRef);
+					}
+				}
+			}
+			CurrentData->PathCollidingGeomIndices = PathCollidingGeomIndices;
+			for (const auto& Vertex : PathCollidingGeomVerts)
 			{
-				AddVertexHelper(DebugMeshData, CurrentData->PathCollidingGeomVerts[VertIdx], NavMeshRenderColor_PathCollidingGeom);
+				CurrentData->PathCollidingGeomVerts.Add(FDynamicMeshVertex(Vertex));
 			}
-			DebugMeshData.Indices.Append(CurrentData->PathCollidingGeomIndices);
-			DebugMeshData.ClusterColor = NavMeshRenderColor_PathCollidingGeom;
-			CurrentData->MeshBuilders.Add(DebugMeshData);
+
 		}
 
 		if (CurrentData->NavMeshGeometry.BuiltMeshIndices.Num() > 0)
diff --git a/Engine/Source/Runtime/Engine/Private/NavMeshRenderingHelpers.h b/Engine/Source/Runtime/Engine/Private/NavMeshRenderingHelpers.h
--- a/Engine/Source/Runtime/Engine/Private/NavMeshRenderingHelpers.h
+++ b/Engine/Source/Runtime/Engine/Private/NavMeshRenderingHelpers.h
@@ -84,7 +84,7 @@ struct FNavMeshSceneProxyData : public TSharedFromThis<FNavMeshSceneProxyData, E
 	TArray<FDebugRenderSceneProxy::FDebugLine> ClusterLinkLines;
 
 	TArray<int32> PathCollidingGeomIndices;
-	TNavStatArray<FVector> PathCollidingGeomVerts;
+	TArray<FDynamicMeshVertex> PathCollidingGeomVerts;
 
 	TArray<FBoxCenterAndExtent>	OctreeBounds;
 		
@@ -274,7 +274,8 @@ public:
 			return;
 		}
 
-		MeshColors.Reserve(NumberOfMeshes);
+		MeshColors.Reserve(NumberOfMeshes + (ProxyData.bDrawPathCollidingGeometry && ProxyData.PathCollidingGeomVerts.Num() > 0 ? 1 : 0));
+		MeshBatchElements.Reserve(NumberOfMeshes);
 		const FMaterialRenderProxy* ParentMaterial = GEngine->DebugMeshMaterial->GetRenderProxy(false);
 		for (int32 Index = 0; Index < NumberOfMeshes; ++Index)
 		{
@@ -282,7 +283,7 @@ public:
 
 			FMeshBatchElement Element;
 			Element.FirstIndex = IndexBuffer.Indices.Num();
-			Element.NumPrimitives = CurrentMeshBuilder.Indices.Num() / 3;
+			Element.NumPrimitives = FMath::FloorToInt(CurrentMeshBuilder.Indices.Num() / 3);
 			Element.MinVertexIndex = VertexBuffer.Vertices.Num();
 			Element.MaxVertexIndex = Element.MinVertexIndex + CurrentMeshBuilder.Vertices.Num() - 1;
 			Element.IndexBuffer = &IndexBuffer;
@@ -294,28 +295,9 @@ public:
 			IndexBuffer.Indices.Append( CurrentMeshBuilder.Indices );
 		}
 
-		if (ProxyData.bDrawPathCollidingGeometry && ProxyData.PathCollidingGeomVerts.Num() > 0)
-		{
-			FMeshBatchElement Element;
-			Element.FirstIndex = IndexBuffer.Indices.Num();
-			Element.MinVertexIndex = VertexBuffer.Vertices.Num();
-			Element.IndexBuffer = &IndexBuffer;
-
-			for (const auto& CurrentLoc : ProxyData.PathCollidingGeomVerts)
-			{
-				VertexBuffer.Vertices.Add(FDynamicMeshVertex(CurrentLoc));
-			}
-			IndexBuffer.Indices.Append(ProxyData.PathCollidingGeomIndices);
-
-			Element.NumPrimitives = ProxyData.PathCollidingGeomIndices.Num() / 3;
-			Element.MaxVertexIndex = VertexBuffer.Vertices.Num() - 1;
-			MeshBatchElements.Add(Element);
-
-			MeshColors.Add(FColoredMaterialRenderProxy(ParentMaterial, NavMeshRenderColor_PathCollidingGeom));
-		}
+		MeshColors.Add(FColoredMaterialRenderProxy(ParentMaterial, NavMeshRenderColor_PathCollidingGeom));
 
 		VertexFactory.Init(&VertexBuffer);
-
 		BeginInitResource(&IndexBuffer);
 		BeginInitResource(&VertexBuffer);
 		BeginInitResource(&VertexFactory);
@@ -326,6 +308,7 @@ public:
 		VertexBuffer.ReleaseResource();
 		IndexBuffer.ReleaseResource();
 		VertexFactory.ReleaseResource();
+
 		ProxyData.Reset();
 	}
 
@@ -431,6 +414,15 @@ public:
 					Collector.AddMesh(ViewIndex, Mesh);
 				}
 
+				if (ProxyData.PathCollidingGeomIndices.Num() > 2)
+				{
+					FDynamicMeshBuilder MeshBuilder;
+					MeshBuilder.AddVertices(ProxyData.PathCollidingGeomVerts);
+					MeshBuilder.AddTriangles(ProxyData.PathCollidingGeomIndices);
+
+					MeshBuilder.GetMesh(FMatrix::Identity, &MeshColors[MeshBatchElements.Num()], SDPG_World, false, false, ViewIndex, Collector);
+				}
+
 				int32 Num = ProxyData.NavMeshEdgeLines.Num();
 				PDI->AddReserveLines(SDPG_World, Num, false, false);
 				PDI->AddReserveLines(SDPG_Foreground, Num, false, true);
@@ -587,6 +579,7 @@ private:
 	FNavMeshIndexBuffer IndexBuffer;
 	FNavMeshVertexBuffer VertexBuffer;
 	FNavMeshVertexFactory VertexFactory;
+
 	TArray <FColoredMaterialRenderProxy> MeshColors;
 	TArray<FMeshBatchElement> MeshBatchElements;
 
@@ -597,23 +590,21 @@ private:
 };
 #endif // WITH_RECAST
 
-FORCEINLINE void AppendGeometry(TNavStatArray<FVector>& OutVertexBuffer, TArray<int32>& OutIndexBuffer,
+FORCEINLINE void AppendGeometry(TArray<FVector>& OutVertexBuffer, TArray<int32>& OutIndexBuffer,
 								const float* Coords, int32 NumVerts, const int32* Faces, int32 NumFaces)
 {
 	const int32 VertIndexBase = OutVertexBuffer.Num();
-	for (int32 VertIdx = 0; VertIdx < NumVerts; ++VertIdx)
+	for (int32 VertIdx = 0; VertIdx < NumVerts*3; VertIdx+=3)
 	{
-		OutVertexBuffer.Add(Recast2UnrealPoint(&Coords[VertIdx * 3]));
+		OutVertexBuffer.Add(Recast2UnrealPoint(&Coords[VertIdx]));
 	}
 
 	const int32 FirstNewFaceVertexIndex = OutIndexBuffer.Num();
-	OutIndexBuffer.AddUninitialized(NumFaces * 3);
-	int32* DestFaceVertIndex = OutIndexBuffer.GetData() + FirstNewFaceVertexIndex;
-	const int32* SrcFaceVertIndex = Faces;
-	// copy with offset
-	for (int32 Index = 0; Index < NumFaces * 3; ++Index, ++DestFaceVertIndex, ++SrcFaceVertIndex)
+	const uint32 NumIndices = NumFaces * 3;
+	OutIndexBuffer.AddUninitialized(NumIndices);
+	for (uint32 Index = 0; Index < NumIndices; ++Index)
 	{
-		*DestFaceVertIndex = *SrcFaceVertIndex + VertIndexBase;
+		OutIndexBuffer[FirstNewFaceVertexIndex + Index] = VertIndexBase + Faces[Index];
 	}
 }
 
