Distance Field lighting features now support instanced static meshes including instanced foliage

diff --git a/Engine/Source/Runtime/Engine/Private/InstancedStaticMesh.cpp b/Engine/Source/Runtime/Engine/Private/InstancedStaticMesh.cpp
--- a/Engine/Source/Runtime/Engine/Private/InstancedStaticMesh.cpp
+++ b/Engine/Source/Runtime/Engine/Private/InstancedStaticMesh.cpp
@@ -322,8 +322,11 @@ void FStaticMeshInstanceBuffer::AllocateData()
 	CleanUp();
 
 	check(HasValidFeatureLevel());
+	static const auto CVar = IConsoleManager::Get().FindTConsoleVariableDataInt(TEXT("r.GenerateMeshDistanceFields"));
 	const bool bInstanced = RHISupportsInstancing(GetFeatureLevelShaderPlatform(GetFeatureLevel()));
-	const bool bNeedsCPUAccess = !bInstanced;
+	const bool bNeedsCPUAccess = !bInstanced 
+		// Distance field algorithms need access to instance data on the CPU
+		|| CVar->GetValueOnGameThread() != 0;
 	InstanceData = new FStaticMeshInstanceData(bNeedsCPUAccess);
 	// Calculate the vertex stride.
 	Stride = InstanceData->GetStride();
@@ -942,6 +945,8 @@ public:
 
 	virtual void GetDynamicMeshElements(const TArray<const FSceneView*>& Views, const FSceneViewFamily& ViewFamily, uint32 VisibilityMap, FMeshElementCollector& Collector) const override;
 
+	virtual void GetDistancefieldAtlasData(FBox& LocalVolumeBounds, FIntVector& OutBlockMin, FIntVector& OutBlockSize, bool& bOutBuiltAsIfTwoSided, bool& bMeshWasPlane, TArray<FMatrix>& ObjectLocalToWorldTransforms) const override;
+
 	virtual int32 GetNumMeshBatches() const override;
 
 	/** Sets up a shadow FMeshBatch for a specific LOD. */
@@ -1156,6 +1161,27 @@ bool FInstancedStaticMeshSceneProxy::GetWireframeMeshElement(int32 LODIndex, int
 	return false;
 }
 
+void FInstancedStaticMeshSceneProxy::GetDistancefieldAtlasData(FBox& LocalVolumeBounds, FIntVector& OutBlockMin, FIntVector& OutBlockSize, bool& bOutBuiltAsIfTwoSided, bool& bMeshWasPlane, TArray<FMatrix>& ObjectLocalToWorldTransforms) const
+{
+	FStaticMeshSceneProxy::GetDistancefieldAtlasData(LocalVolumeBounds, OutBlockMin, OutBlockSize, bOutBuiltAsIfTwoSided, bMeshWasPlane, ObjectLocalToWorldTransforms);
+
+	ObjectLocalToWorldTransforms.Reset();
+
+	const FInstancingUserData::FInstanceStream* InstanceStream = ((const FInstancingUserData::FInstanceStream*)InstancedRenderData.InstanceBuffer.GetRawData());
+	FVector4 FourthVector(0, 0, 0, 1);
+
+	for (uint32 InstanceIndex = 0; InstanceIndex < InstancedRenderData.InstanceBuffer.GetNumInstances(); InstanceIndex++)
+	{
+		const FMatrix TransposedInstanceToLocal(
+			(const FPlane&)InstanceStream[InstanceIndex].InstanceTransform[0], 
+			(const FPlane&)InstanceStream[InstanceIndex].InstanceTransform[1], 
+			(const FPlane&)InstanceStream[InstanceIndex].InstanceTransform[2], 
+			(const FPlane&)FourthVector);
+
+		ObjectLocalToWorldTransforms.Add(TransposedInstanceToLocal.GetTransposed() * GetLocalToWorld());
+	}
+}
+
 /*-----------------------------------------------------------------------------
 	UInstancedStaticMeshComponent
 -----------------------------------------------------------------------------*/
