Added comment about bUsedAsAtmosphereSunLight being required for exp height fog directional inscattering

diff --git a/Engine/Source/Runtime/Engine/Classes/Components/ExponentialHeightFogComponent.h b/Engine/Source/Runtime/Engine/Classes/Components/ExponentialHeightFogComponent.h
--- a/Engine/Source/Runtime/Engine/Classes/Components/ExponentialHeightFogComponent.h
+++ b/Engine/Source/Runtime/Engine/Classes/Components/ExponentialHeightFogComponent.h
@@ -23,15 +23,24 @@ class ENGINE_API UExponentialHeightFogComponent : public USceneComponent
 	UPROPERTY(BlueprintReadOnly, interp, Category=ExponentialHeightFogComponent)
 	FLinearColor FogInscatteringColor;
 
-	/** Controls the size of the directional inscattering cone, which is used to approximate inscattering from a directional light. */
+	/** 
+	 * Controls the size of the directional inscattering cone, which is used to approximate inscattering from a directional light.  
+	 * Note: there must be a directional light with bUsedAsAtmosphereSunLight enabled for DirectionalInscattering to be used.
+	 */
 	UPROPERTY(BlueprintReadOnly, interp, Category=DirectionalInscattering, meta=(UIMin = "2", UIMax = "64"))
 	float DirectionalInscatteringExponent;
 
-	/** Controls the start distance from the viewer of the directional inscattering, which is used to approximate inscattering from a directional light. */
+	/** 
+	 * Controls the start distance from the viewer of the directional inscattering, which is used to approximate inscattering from a directional light. 
+	 * Note: there must be a directional light with bUsedAsAtmosphereSunLight enabled for DirectionalInscattering to be used.
+	 */
 	UPROPERTY(BlueprintReadOnly, interp, Category=DirectionalInscattering)
 	float DirectionalInscatteringStartDistance;
 
-	/** Controls the color of the directional inscattering, which is used to approximate inscattering from a directional light. */
+	/** 
+	 * Controls the color of the directional inscattering, which is used to approximate inscattering from a directional light. 
+	 * Note: there must be a directional light with bUsedAsAtmosphereSunLight enabled for DirectionalInscattering to be used.
+	 */
 	UPROPERTY(BlueprintReadOnly, interp, Category=DirectionalInscattering)
 	FLinearColor DirectionalInscatteringColor;
 
