#anim : Fix retargeting looping animation.

diff --git a/Engine/Source/Runtime/Engine/Private/Animation/AnimSequence.cpp b/Engine/Source/Runtime/Engine/Private/Animation/AnimSequence.cpp
--- a/Engine/Source/Runtime/Engine/Private/Animation/AnimSequence.cpp
+++ b/Engine/Source/Runtime/Engine/Private/Animation/AnimSequence.cpp
@@ -94,6 +94,15 @@ bool FCurveTrack::CompressCurveWeights()
 	return false;
 }
 
+/////////////////////////////////////////////////////////////
+
+// since we want this change for hot fix, I can't change header file, 
+// next time move this to the header
+float GetIntervalPerKey(int32 NumFrames, float SequenceLength) 
+{
+	return (NumFrames > 1) ? (SequenceLength / (NumFrames-1)) : MINIMUM_ANIMATION_LENGTH;
+}
+
 /////////////////////////////////////////////////////
 // UAnimSequence
 
@@ -1729,7 +1738,7 @@ void UAnimSequence::UpgradeMorphTargetCurves()
 			// make sure this matches
 			// the way we import animation should cover 0 - (NumFrames-1) for SequenceLength;
 			// if only 1 frame, it will be written in the first frame
-			float TimeInterval = (NumFrames==1) ? 0.f : SequenceLength / (NumFrames-1);
+			float TimeInterval = GetIntervalPerKey(NumFrames, SequenceLength);
 
 			// copy all weights back
 			for (int32 Iter = 0 ; Iter<CurveTrack.CurveWeights.Num(); ++Iter)
@@ -1966,7 +1975,7 @@ void UAnimSequence::RemapTracksToNewSkeleton( USkeleton* NewSkeleton, bool bConv
 			ConvertedLocalSpaceAnimations.AddZeroed(SrcNumTracks);
 
 			int32 NumKeys = NumFrames;
-			float Interval = (NumKeys)? SequenceLength/NumKeys : 0.f;
+			float Interval = GetIntervalPerKey(NumFrames, SequenceLength);
 
 			// allocate arrays
 			for(int32 SrcTrackIndex=0; SrcTrackIndex<SrcNumTracks; ++SrcTrackIndex)
@@ -2188,7 +2197,7 @@ void UAnimSequence::RemapTracksToNewSkeleton( USkeleton* NewSkeleton, bool bConv
 			ConvertedSpaceBases.AddZeroed(NumBones);
 
 			int32 NumKeys = NumFrames;
-			float Interval = (NumKeys)? SequenceLength/NumKeys : 0.f;
+			float Interval = GetIntervalPerKey(NumFrames, SequenceLength);
 
 			// allocate arrays
 			for(int32 BoneIndex=0; BoneIndex<NumBones; ++BoneIndex)
@@ -2577,13 +2586,12 @@ void UAnimSequence::ReplaceReferredAnimations(const TMap<UAnimSequence*, UAnimSe
 bool UAnimSequence::AddLoopingInterpolation()
 {
 	int32 NumTracks = AnimationTrackNames.Num();
-	int32 NumKeys = NumFrames;
-	float Interval = (NumKeys)? SequenceLength/NumKeys : 0.f;
+	float Interval = GetIntervalPerKey(NumFrames, SequenceLength);
 
-	if(Interval > 0.f)
+	if(NumFrames > 0)
 	{
 		// added one more key
-		int32 NewNumKeys = NumKeys +1 ;
+		int32 NewNumKeys = NumFrames +1 ;
 
 		// now I need to calculate back to new animation data
 		for(int32 TrackIndex=0; TrackIndex<NumTracks; ++TrackIndex)
@@ -2638,7 +2646,7 @@ int32 UAnimSequence::GetSpaceBasedAnimationData(TArray< TArray<FTransform> >& An
 
 	// 2d array of animated time [boneindex][time key]
 	int32 NumKeys = NumFrames;
-	float Interval = (NumKeys) ? SequenceLength / NumKeys : 0.f;
+	float Interval = GetIntervalPerKey(NumFrames, SequenceLength);
 
 	// allocate arrays
 	for (int32 BoneIndex = 0; BoneIndex < NumBones; ++BoneIndex)
@@ -3190,7 +3198,7 @@ void UAnimSequence::BakeTrackCurvesToRawAnimation()
 				}
 
 				// NumFrames can't be zero (filtered earlier)
-				float Interval = (NumFrames > 1)? SequenceLength/(NumFrames-1) : MINIMUM_ANIMATION_LENGTH;
+				float Interval = GetIntervalPerKey(NumFrames, SequenceLength);
 
 				// now we have all data ready to apply
 				for(int32 KeyIndex=0; KeyIndex < NumFrames; ++KeyIndex)
