Remove Post Physics on Skeletal Mesh Component

diff --git a/Engine/Source/Runtime/Engine/Private/SkeletalMeshComponent.cpp b/Engine/Source/Runtime/Engine/Private/SkeletalMeshComponent.cpp
--- a/Engine/Source/Runtime/Engine/Private/SkeletalMeshComponent.cpp
+++ b/Engine/Source/Runtime/Engine/Private/SkeletalMeshComponent.cpp
@@ -98,7 +98,6 @@ USkeletalMeshComponent::USkeletalMeshComponent(const FObjectInitializer& ObjectI
 	bAutoActivate = true;
 	PrimaryComponentTick.bCanEverTick = true;
 	PrimaryComponentTick.TickGroup = TG_PrePhysics;
-	PostPhysicsComponentTick.bCanEverTick = true;
 	bWantsInitializeComponent = true;
 	GlobalAnimRateScale = 1.0f;
 	bNoSkeletonUpdate = false;
