#ue4 - Avoid registering PhysicsVolumeChanged events for movement components when the UpdatedComponent is pending kill. Also avoid registering for overlap events in CharacterMovementComponent when UpdatedComponent is pending kill.

- Fixes ensure triggered under these (rare) conditions.

UDN: https://udn.unrealengine.com/questions/225824/ucharactermovementcomponentphysicsvolumechanged-mi.html

diff --git a/Engine/Source/Runtime/Engine/Private/CharacterMovementComponent.cpp b/Engine/Source/Runtime/Engine/Private/CharacterMovementComponent.cpp
--- a/Engine/Source/Runtime/Engine/Private/CharacterMovementComponent.cpp
+++ b/Engine/Source/Runtime/Engine/Private/CharacterMovementComponent.cpp
@@ -372,7 +372,7 @@ void UCharacterMovementComponent::SetUpdatedComponent(UPrimitiveComponent* NewUp
 	bDeferUpdateMoveComponent = false;
 	DeferredUpdatedMoveComponent = NULL;
 
-	if (UpdatedComponent != NULL && UpdatedComponent->OnComponentBeginOverlap.IsBound())
+	if (IsValid(UpdatedComponent) && UpdatedComponent->OnComponentBeginOverlap.IsBound())
 	{
 		UpdatedComponent->OnComponentBeginOverlap.RemoveDynamic(this, &UCharacterMovementComponent::CapsuleTouched);
 	}
@@ -385,7 +385,7 @@ void UCharacterMovementComponent::SetUpdatedComponent(UPrimitiveComponent* NewUp
 		StopActiveMovement();
 	}
 
-	if (UpdatedComponent && bEnablePhysicsInteraction)
+	if (IsValid(UpdatedComponent) && bEnablePhysicsInteraction)
 	{
 		UpdatedComponent->OnComponentBeginOverlap.AddUniqueDynamic(this, &UCharacterMovementComponent::CapsuleTouched);
 	}
diff --git a/Engine/Source/Runtime/Engine/Private/MovementComponent.cpp b/Engine/Source/Runtime/Engine/Private/MovementComponent.cpp
--- a/Engine/Source/Runtime/Engine/Private/MovementComponent.cpp
+++ b/Engine/Source/Runtime/Engine/Private/MovementComponent.cpp
@@ -40,7 +40,7 @@ UMovementComponent::UMovementComponent(const FObjectInitializer& ObjectInitializ
 
 void UMovementComponent::SetUpdatedComponent(UPrimitiveComponent* NewUpdatedComponent)
 {
-	if ( UpdatedComponent && UpdatedComponent!=NewUpdatedComponent )
+	if (UpdatedComponent && UpdatedComponent != NewUpdatedComponent)
 	{
 		UpdatedComponent->bShouldUpdatePhysicsVolume = false;
 		if (!UpdatedComponent->IsPendingKill())
@@ -55,7 +55,7 @@ void UMovementComponent::SetUpdatedComponent(UPrimitiveComponent* NewUpdatedComp
 
 	UpdatedComponent = NewUpdatedComponent;
 
-	if ( UpdatedComponent != NULL )
+	if (IsValid(UpdatedComponent))
 	{
 		UpdatedComponent->bShouldUpdatePhysicsVolume = true;
 		UpdatedComponent->PhysicsVolumeChangedDelegate.AddUniqueDynamic(this, &UMovementComponent::PhysicsVolumeChanged);
