Framework: Exposing SetTickableWhenPaused and SetComponentTickEnabled on UActorComponent

diff --git a/Engine/Source/Runtime/Engine/Classes/Components/ActorComponent.h b/Engine/Source/Runtime/Engine/Classes/Components/ActorComponent.h
--- a/Engine/Source/Runtime/Engine/Classes/Components/ActorComponent.h
+++ b/Engine/Source/Runtime/Engine/Classes/Components/ActorComponent.h
@@ -153,6 +153,10 @@ public:
 	UFUNCTION(BlueprintCallable, Category="Components|Activation")
 	virtual bool IsActive() const;
 
+	/** Sets whether this component can tick when paused. */
+	UFUNCTION(BlueprintCallable, Category="Utilities")
+	void SetTickableWhenPaused(bool bTickableWhenPaused);
+
 	// Networking
 
 	/** This signifies the component can be ID'd by name over the network. This only needs to be called by engine code when constructing blueprint components. */
@@ -300,6 +304,7 @@ public:
 	 * 
 	 * @param	bEnabled - Rather it should be enabled or not
 	 */
+	UFUNCTION(BlueprintCallable, Category="Utilities")
 	virtual void SetComponentTickEnabled(bool bEnabled);
 	/** 
 	 * Spawns a task on GameThread that will call SetComponentTickEnabled
diff --git a/Engine/Source/Runtime/Engine/Private/ActorComponent.cpp b/Engine/Source/Runtime/Engine/Private/ActorComponent.cpp
--- a/Engine/Source/Runtime/Engine/Private/ActorComponent.cpp
+++ b/Engine/Source/Runtime/Engine/Private/ActorComponent.cpp
@@ -1105,6 +1105,11 @@ bool UActorComponent::IsActive() const
 	return bIsActive;
 }
 
+void UActorComponent::SetTickableWhenPaused(bool bTickableWhenPaused)
+{
+	PrimaryComponentTick.bTickEvenWhenPaused = bTickableWhenPaused;
+}
+
 bool UActorComponent::IsRunningUserConstructionScript() const
 {
 	const AActor* Owner = GetOwner();
