#ue4 - Avoid using a potentially invalid faceIndex when retreiving the physical material from landscape collision data.

diff --git a/Engine/Source/Runtime/Engine/Private/Collision/CollisionConversions.cpp b/Engine/Source/Runtime/Engine/Private/Collision/CollisionConversions.cpp
--- a/Engine/Source/Runtime/Engine/Private/Collision/CollisionConversions.cpp
+++ b/Engine/Source/Runtime/Engine/Private/Collision/CollisionConversions.cpp
@@ -284,7 +284,7 @@ static void SetHitResultFromShapeAndFaceIndex(const PxShape* PShape,  const PxRi
 		OutResult.Actor = OwningComponent->GetOwner();
 		OutResult.Component = OwningComponent;
 
-		if(bReturnPhysMat)
+		if (bReturnPhysMat && FaceIndex != InvalidQueryHit.faceIndex)
 		{
 			// @fixme: only do this for InGameThread, otherwise, this will be done in AsyncTrace
 			if ( IsInGameThread() )
@@ -381,10 +381,13 @@ void ConvertQueryImpactHit(const PxLocationHit& PHit, FHitResult& OutResult, flo
 	if( PHit.shape->getGeometryType() == PxGeometryType::eHEIGHTFIELD)
 	{
 		// Lookup physical material for heightfields
-		PxMaterial* HitMaterial = PHit.shape->getMaterialFromInternalFaceIndex(PHit.faceIndex);
-		if( HitMaterial != NULL )
+		if (PHit.faceIndex != InvalidQueryHit.faceIndex)
 		{
-			OutResult.PhysMaterial = FPhysxUserData::Get<UPhysicalMaterial>(HitMaterial->userData);
+			PxMaterial* HitMaterial = PHit.shape->getMaterialFromInternalFaceIndex(PHit.faceIndex);
+			if (HitMaterial != NULL)
+			{
+				OutResult.PhysMaterial = FPhysxUserData::Get<UPhysicalMaterial>(HitMaterial->userData);
+			}
 		}
 	}
 	else
