fixed doubled BSP data in navigation octree

diff --git a/Engine/Source/Runtime/Engine/Private/AI/Navigation/NavigationSystem.cpp b/Engine/Source/Runtime/Engine/Private/AI/Navigation/NavigationSystem.cpp
--- a/Engine/Source/Runtime/Engine/Private/AI/Navigation/NavigationSystem.cpp
+++ b/Engine/Source/Runtime/Engine/Private/AI/Navigation/NavigationSystem.cpp
@@ -3019,19 +3019,24 @@ void UNavigationSystem::OnLevelRemovedFromWorld(ULevel* InLevel, UWorld* InWorld
 void UNavigationSystem::AddLevelCollisionToOctree(ULevel* Level)
 {
 #if WITH_RECAST
-	const TArray<FVector>* LevelGeom = Level ? Level->GetStaticNavigableGeometry() : NULL;
-	if (LevelGeom && NavOctree)
+	if (Level && NavOctree)
 	{
-		FNavigationOctreeElement BSPElem;
-		FRecastNavMeshGenerator::ExportVertexSoupGeometry(*LevelGeom, BSPElem.Data);
+		const TArray<FVector>* LevelGeom = Level->GetStaticNavigableGeometry();
+		const FOctreeElementId* ElementId = GetObjectsNavOctreeId(Level);
 
-		const auto& Bounds = BSPElem.Data.Bounds;
-		if (!Bounds.GetExtent().IsNearlyZero())
+		if (LevelGeom && !ElementId)
 		{
-			NavOctree->AddNode(Level, NULL, Bounds, BSPElem);
-			AddDirtyArea(Bounds, ENavigationDirtyFlag::All);
+			FNavigationOctreeElement BSPElem;
+			FRecastNavMeshGenerator::ExportVertexSoupGeometry(*LevelGeom, BSPElem.Data);
 
-			UE_LOG(LogNavOctree, Log, TEXT("ADD %s"), *GetNameSafe(Level));
+			const auto& Bounds = BSPElem.Data.Bounds;
+			if (!Bounds.GetExtent().IsNearlyZero())
+			{
+				NavOctree->AddNode(Level, NULL, Bounds, BSPElem);
+				AddDirtyArea(Bounds, ENavigationDirtyFlag::All);
+
+				UE_LOG(LogNavOctree, Log, TEXT("ADD %s"), *GetNameSafe(Level));
+			}
 		}
 	}
 #endif// WITH_RECAST
