Instanced Array fix. https://udn.unrealengine.com/questions/218011/edit-inline-tarray-all-instance-data-lost-on-bluep.html

diff --git a/Engine/Source/Runtime/Engine/Private/UnrealEngine.cpp b/Engine/Source/Runtime/Engine/Private/UnrealEngine.cpp
--- a/Engine/Source/Runtime/Engine/Private/UnrealEngine.cpp
+++ b/Engine/Source/Runtime/Engine/Private/UnrealEngine.cpp
@@ -10047,35 +10047,49 @@ struct FFindInstancedReferenceSubobjectHelper
 		check(ContainerAddress);
 		for (UProperty* Prop = (Struct ? Struct->RefLink : NULL); Prop; Prop = Prop->NextRef)
 		{
-			if (Prop->HasAllPropertyFlags(CPF_PersistentInstance))
+			auto ArrayProperty = Cast<const UArrayProperty>(Prop);
+			if (ArrayProperty && Prop->HasAnyPropertyFlags(CPF_PersistentInstance | CPF_ContainsInstancedReference))
 			{
-				ensure(Prop->HasAllPropertyFlags(CPF_InstancedReference));
-				auto ObjectProperty = Cast<const UObjectProperty>(Prop);
-				if (ObjectProperty)
+				auto InnerStructProperty = Cast<const UStructProperty>(ArrayProperty->Inner);
+				auto InnerObjectProperty = Cast<const UObjectProperty>(ArrayProperty->Inner);
+				if (InnerStructProperty && InnerStructProperty->Struct)
 				{
-					for (int32 ArrayIdx = 0; ArrayIdx < Prop->ArrayDim; ++ArrayIdx)
+					FScriptArrayHelper_InContainer ArrayHelper(ArrayProperty, ContainerAddress);
+					for (int32 ElementIndex = 0; ElementIndex < ArrayHelper.Num(); ++ElementIndex)
 					{
-						auto ObjectValue = ObjectProperty->GetObjectPropertyValue_InContainer(ContainerAddress, ArrayIdx);
+						const uint8* ValueAddress = ArrayHelper.GetRawPtr(ElementIndex);
+						if (ValueAddress)
+						{
+							Get(InnerStructProperty->Struct, ValueAddress, OutObjects);
+						}
+					}
+				}
+				else if (InnerObjectProperty && InnerObjectProperty->HasAllPropertyFlags(CPF_PersistentInstance))
+				{
+					ensure(InnerObjectProperty->HasAllPropertyFlags(CPF_InstancedReference));
+					FScriptArrayHelper_InContainer ArrayHelper(ArrayProperty, ContainerAddress);
+					for (int32 ElementIndex = 0; ElementIndex < ArrayHelper.Num(); ++ElementIndex)
+					{
+						UObject* ObjectValue = InnerObjectProperty->GetObjectPropertyValue(ArrayHelper.GetRawPtr(ElementIndex));
 						if (ObjectValue)
 						{
 							OutObjects.Add(ObjectValue);
 						}
 					}
 				}
-				else
+			}
+			else if (Prop->HasAllPropertyFlags(CPF_PersistentInstance))
+			{
+				ensure(Prop->HasAllPropertyFlags(CPF_InstancedReference));
+				auto ObjectProperty = Cast<const UObjectProperty>(Prop);
+				if (ObjectProperty)
 				{
-					auto ArrayProperty = Cast<const UArrayProperty>(Prop);
-					auto InnerObjectProperty = ArrayProperty ? Cast<const UObjectProperty>(ArrayProperty->Inner) : NULL;
-					if (InnerObjectProperty)
+					for (int32 ArrayIdx = 0; ArrayIdx < Prop->ArrayDim; ++ArrayIdx)
 					{
-						FScriptArrayHelper_InContainer ArrayHelper(ArrayProperty, ContainerAddress);
-						for (int32 ElementIndex = 0; ElementIndex < ArrayHelper.Num(); ++ElementIndex)
+						auto ObjectValue = ObjectProperty->GetObjectPropertyValue_InContainer(ContainerAddress, ArrayIdx);
+						if (ObjectValue)
 						{
-							UObject** ObjectValuePtr = reinterpret_cast<UObject**>(ArrayHelper.GetRawPtr(ElementIndex));
-							if (ObjectValuePtr && *ObjectValuePtr)
-							{
-								OutObjects.Add(*ObjectValuePtr);
-							}
+							OutObjects.Add(ObjectValue);
 						}
 					}
 				}
@@ -10091,23 +10105,6 @@ struct FFindInstancedReferenceSubobjectHelper
 						Get(StructProperty->Struct, ValueAddress, OutObjects);
 					}
 				}
-				else
-				{
-					auto ArrayProperty = Cast<const UArrayProperty>(Prop);
-					auto InnerStructProperty = ArrayProperty ? Cast<const UStructProperty>(ArrayProperty->Inner) : NULL;
-					if (InnerStructProperty && InnerStructProperty->Struct)
-					{
-						FScriptArrayHelper_InContainer ArrayHelper(ArrayProperty, ContainerAddress);
-						for (int32 ElementIndex = 0; ElementIndex < ArrayHelper.Num(); ++ElementIndex)
-						{
-							const uint8* ValueAddress = ArrayHelper.GetRawPtr(ElementIndex);
-							if (ValueAddress)
-							{
-								Get(InnerStructProperty->Struct, ValueAddress, OutObjects);
-							}
-						}
-					}
-				}
 			}
 		}
 	}
