Remove logging accidentally checked in

diff --git a/Engine/Source/Runtime/Engine/Private/PhysicsEngine/PhysDrawing.cpp b/Engine/Source/Runtime/Engine/Private/PhysicsEngine/PhysDrawing.cpp
--- a/Engine/Source/Runtime/Engine/Private/PhysicsEngine/PhysDrawing.cpp
+++ b/Engine/Source/Runtime/Engine/Private/PhysicsEngine/PhysDrawing.cpp
@@ -581,8 +581,6 @@ void FKAggregateGeom::GetAggGeom(const FTransform& Transform, const FColor Color
 			// Cache collision vertex/index buffer
 			if(!RenderInfo)
 			{
-				UE_LOG(LogTemp, Log, TEXT("Create RenderInfo %x"), this);
-
 				//@todo - parallelrendering, remove const cast
 				FKAggregateGeom& ThisGeom = const_cast<FKAggregateGeom&>(*this);
 				ThisGeom.RenderInfo = new FKConvexGeomRenderInfo();
@@ -654,8 +652,6 @@ void FKAggregateGeom::FreeRenderInfo()
 	// See if we have rendering resources to free
 	if(RenderInfo)
 	{
-		UE_LOG(LogTemp, Log, TEXT("Create RenderInfo %x"), this);
-
 		// Should always have these if RenderInfo exists
 		check(RenderInfo->VertexBuffer);
 		check(RenderInfo->IndexBuffer);
