UE-6096 Add needed EngineBaseTypes.h include to ShowFlags.h

diff --git a/Engine/Source/Runtime/Engine/Public/ShowFlags.h b/Engine/Source/Runtime/Engine/Public/ShowFlags.h
--- a/Engine/Source/Runtime/Engine/Public/ShowFlags.h
+++ b/Engine/Source/Runtime/Engine/Public/ShowFlags.h
@@ -6,6 +6,7 @@
 
 #pragma once
 
+#include "Engine/EngineBaseTypes.h"
 
 /** for FShowFlagData, order is not important, adding new elements requires updating some editor code */
 enum EShowFlagGroup
