UE-6060 In CapsuleComponent PostEditChange, only clamp value that you changed, to help propagation
  Still issues though (value propagated before PostEditChange called), asked RobM to look at that part
Improve comments to make it clear capsule radius must be < half-height

diff --git a/Engine/Source/Runtime/Engine/Classes/Components/CapsuleComponent.h b/Engine/Source/Runtime/Engine/Classes/Components/CapsuleComponent.h
--- a/Engine/Source/Runtime/Engine/Classes/Components/CapsuleComponent.h
+++ b/Engine/Source/Runtime/Engine/Classes/Components/CapsuleComponent.h
@@ -13,11 +13,17 @@ class ENGINE_API UCapsuleComponent : public UShapeComponent
 	GENERATED_UCLASS_BODY()
 
 protected:
-	/** Half-height, i.e. from center of capsule to end of top or bottom hemisphere.  */
+	/** 
+	 *	Half-height, i.e. from center of capsule to end of top or bottom hemisphere.  
+	 *	This cannot be less than CapsuleRadius.
+	 */
 	UPROPERTY(EditAnywhere, BlueprintReadOnly, export, Category=Shape)
 	float CapsuleHalfHeight;
 
-	/** Radius of cap hemispheres and center cylinder. */
+	/** 
+	 *	Radius of cap hemispheres and center cylinder. 
+	 *	This cannot be more than CapsuleHalfHeight.
+	 */
 	UPROPERTY(EditAnywhere, BlueprintReadOnly, export, Category=Shape)
 	float CapsuleRadius;
 
diff --git a/Engine/Source/Runtime/Engine/Private/CapsuleComponent.cpp b/Engine/Source/Runtime/Engine/Private/CapsuleComponent.cpp
--- a/Engine/Source/Runtime/Engine/Private/CapsuleComponent.cpp
+++ b/Engine/Source/Runtime/Engine/Private/CapsuleComponent.cpp
@@ -112,13 +112,20 @@ void UCapsuleComponent::Serialize(FArchive& Ar)
 #if WITH_EDITOR
 void UCapsuleComponent::PostEditChangeProperty(FPropertyChangedEvent& PropertyChangedEvent)
 {
-	Super::PostEditChangeProperty(PropertyChangedEvent);
-
 	const FName PropertyName = PropertyChangedEvent.Property ? PropertyChangedEvent.Property->GetFName() : NAME_None;
-	if (PropertyName == FName(TEXT("CapsuleHalfHeight")) || PropertyName == FName(TEXT("CapsuleRadius")))
+
+	// We only want to modify the property that was changed at this point
+	// things like propagation from CDO to instances don't work correctly if changing one property causes a different property to change
+	if (PropertyName == GET_MEMBER_NAME_CHECKED(UCapsuleComponent, CapsuleHalfHeight))
 	{
 		CapsuleHalfHeight = FMath::Max(CapsuleHalfHeight, CapsuleRadius);
 	}
+	else if (PropertyName == GET_MEMBER_NAME_CHECKED(UCapsuleComponent, CapsuleRadius))
+	{
+		CapsuleRadius = FMath::Min(CapsuleHalfHeight, CapsuleRadius);
+	}
+
+	Super::PostEditChangeProperty(PropertyChangedEvent);
 }
 #endif	// WITH_EDITOR
 
