Remove some forward declarations which only work because of delayed template instantiation by the compiler (ie. they fail to compile if the forward-declarared class is not fully defined elsewhere before the end of the translation unit). Produces errors if you ever include them in isolation.

diff --git a/Engine/Source/Developer/Settings/Private/SettingsCategory.cpp b/Engine/Source/Developer/Settings/Private/SettingsCategory.cpp
--- a/Engine/Source/Developer/Settings/Private/SettingsCategory.cpp
+++ b/Engine/Source/Developer/Settings/Private/SettingsCategory.cpp
@@ -56,6 +56,11 @@ void FSettingsCategory::RemoveSection( const FName& SectionName )
 /* ISettingsCategory interface
  *****************************************************************************/
 
+ISettingsSectionPtr FSettingsCategory::GetSection( const FName& SectionName ) const
+{
+	return Sections.FindRef(SectionName);
+}
+
 int32 FSettingsCategory::GetSections( TArray<ISettingsSectionPtr>& OutSections ) const
 {
 	OutSections.Empty(Sections.Num());
diff --git a/Engine/Source/Developer/Settings/Private/SettingsCategory.h b/Engine/Source/Developer/Settings/Private/SettingsCategory.h
--- a/Engine/Source/Developer/Settings/Private/SettingsCategory.h
+++ b/Engine/Source/Developer/Settings/Private/SettingsCategory.h
@@ -87,10 +87,7 @@ public:
 		return Name;
 	}
 
-	virtual ISettingsSectionPtr GetSection( const FName& SectionName ) const override
-	{
-		return Sections.FindRef(SectionName);
-	}
+	virtual ISettingsSectionPtr GetSection( const FName& SectionName ) const override;
 
 	virtual int32 GetSections( TArray<ISettingsSectionPtr>& OutSections ) const override;
 
diff --git a/Engine/Source/Editor/GraphEditor/Private/SGraphPanel.cpp b/Engine/Source/Editor/GraphEditor/Private/SGraphPanel.cpp
--- a/Engine/Source/Editor/GraphEditor/Private/SGraphPanel.cpp
+++ b/Engine/Source/Editor/GraphEditor/Private/SGraphPanel.cpp
@@ -939,6 +939,11 @@ void SGraphPanel::AttachGraphEvents(TSharedPtr<SGraphNode> CreatedSubNode)
 	CreatedSubNode->SetTextCommittedEvent(OnTextCommitted);
 }
 
+const TSharedRef<SGraphNode> SGraphPanel::GetChild(int32 ChildIndex)
+{
+	return StaticCastSharedRef<SGraphNode>(Children[ChildIndex]);
+}
+
 void SGraphPanel::AddNode(UEdGraphNode* Node)
 {
 	TSharedPtr<SGraphNode> NewNode = FNodeFactory::CreateNodeWidget(Node);
diff --git a/Engine/Source/Editor/GraphEditor/Private/SGraphPanel.h b/Engine/Source/Editor/GraphEditor/Private/SGraphPanel.h
--- a/Engine/Source/Editor/GraphEditor/Private/SGraphPanel.h
+++ b/Engine/Source/Editor/GraphEditor/Private/SGraphPanel.h
@@ -160,10 +160,7 @@ protected:
 
 	void NotifyGraphChanged ( const struct FEdGraphEditAction& InAction);
 
-	const TSharedRef<SGraphNode> GetChild(int32 ChildIndex)
-	{
-		return StaticCastSharedRef<SGraphNode>(Children[ChildIndex]);
-	}
+	const TSharedRef<SGraphNode> GetChild(int32 ChildIndex);
 
 	/** Helper method to add a new node to the panel */
 	void AddNode (class UEdGraphNode* Node);
diff --git a/Engine/Source/Runtime/CoreUObject/Public/UObject/CoreObject.h b/Engine/Source/Runtime/CoreUObject/Public/UObject/CoreObject.h
--- a/Engine/Source/Runtime/CoreUObject/Public/UObject/CoreObject.h
+++ b/Engine/Source/Runtime/CoreUObject/Public/UObject/CoreObject.h
@@ -3,6 +3,7 @@
 #pragma once
 
 #include "ObjectBase.h"
+#include "WorldCompositionUtility.h"
 
 /**
  * Structure to hold information about an external packages objects used in cross-level references
@@ -86,7 +87,7 @@ public:
 	class UMetaData*	MetaData;
 	
 	// World browser information
-	TScopedPointer< class FWorldTileInfo > WorldTileInfo;
+	TScopedPointer< FWorldTileInfo > WorldTileInfo;
 
 #if WITH_EDITOR
 	/** Editor only: PIE instance ID this package belongs to, INDEX_NONE otherwise */
diff --git a/Engine/Source/Runtime/Engine/Public/ActiveSound.h b/Engine/Source/Runtime/Engine/Public/ActiveSound.h
--- a/Engine/Source/Runtime/Engine/Public/ActiveSound.h
+++ b/Engine/Source/Runtime/Engine/Public/ActiveSound.h
@@ -4,6 +4,7 @@
 
 #include "Sound/AudioVolume.h"
 #include "Sound/SoundAttenuation.h"
+#include "Components/AudioComponent.h"
 
 /**
  *	Struct used for gathering the final parameters to apply to a wave instance
@@ -173,7 +174,7 @@ public:
 	TMap<UPTRINT,uint32> SoundNodeOffsetMap;
 	TArray<uint8> SoundNodeData;
 
-	TArray<struct FAudioComponentParam> InstanceParameters;
+	TArray<FAudioComponentParam> InstanceParameters;
 
 	// Updates the wave instances to be played.
 	void UpdateWaveInstances( FAudioDevice* AudioDevice, TArray<FWaveInstance*> &WaveInstances, const float DeltaTime );
diff --git a/Engine/Source/Runtime/Engine/Public/SceneManagement.h b/Engine/Source/Runtime/Engine/Public/SceneManagement.h
--- a/Engine/Source/Runtime/Engine/Public/SceneManagement.h
+++ b/Engine/Source/Runtime/Engine/Public/SceneManagement.h
@@ -28,8 +28,6 @@ struct FDynamicMeshVertex;
 
 DECLARE_LOG_CATEGORY_EXTERN(LogBufferVisualization, Log, All);
 
-static const int MAX_FORWARD_SHADOWCASCADES = 2;
-
 // -----------------------------------------------------------------------------
 
 
@@ -1473,114 +1471,6 @@ private:
 
 
 
-/** 
- * Enumeration for currently used translucent lighting volume cascades 
- */
-enum ETranslucencyVolumeCascade
-{
-	TVC_Inner,
-	TVC_Outer,
-
-	TVC_MAX,
-};
-
-/** The uniform shader parameters associated with a view. */
-BEGIN_UNIFORM_BUFFER_STRUCT_WITH_CONSTRUCTOR(FViewUniformShaderParameters,ENGINE_API)
-	DECLARE_UNIFORM_BUFFER_STRUCT_MEMBER(FMatrix,TranslatedWorldToClip)
-	DECLARE_UNIFORM_BUFFER_STRUCT_MEMBER(FMatrix,WorldToClip)
-	DECLARE_UNIFORM_BUFFER_STRUCT_MEMBER(FMatrix,TranslatedWorldToView)
-	DECLARE_UNIFORM_BUFFER_STRUCT_MEMBER(FMatrix,ViewToTranslatedWorld)
-	DECLARE_UNIFORM_BUFFER_STRUCT_MEMBER(FMatrix,ViewToClip)
-	DECLARE_UNIFORM_BUFFER_STRUCT_MEMBER(FMatrix,ClipToView)
-	DECLARE_UNIFORM_BUFFER_STRUCT_MEMBER(FMatrix,ClipToTranslatedWorld)
-	DECLARE_UNIFORM_BUFFER_STRUCT_MEMBER(FMatrix,ScreenToWorld)
-	DECLARE_UNIFORM_BUFFER_STRUCT_MEMBER(FMatrix,ScreenToTranslatedWorld)
-	DECLARE_UNIFORM_BUFFER_STRUCT_MEMBER_EX(FVector,ViewForward, EShaderPrecisionModifier::Half)
-	DECLARE_UNIFORM_BUFFER_STRUCT_MEMBER_EX(FVector,ViewUp, EShaderPrecisionModifier::Half)
-	DECLARE_UNIFORM_BUFFER_STRUCT_MEMBER_EX(FVector,ViewRight, EShaderPrecisionModifier::Half)
-	DECLARE_UNIFORM_BUFFER_STRUCT_MEMBER(FVector4,InvDeviceZToWorldZTransform)
-	DECLARE_UNIFORM_BUFFER_STRUCT_MEMBER_EX(FVector4,ScreenPositionScaleBias, EShaderPrecisionModifier::Half)
-	DECLARE_UNIFORM_BUFFER_STRUCT_MEMBER_EX(FVector4,ViewRectMin, EShaderPrecisionModifier::Half)
-	DECLARE_UNIFORM_BUFFER_STRUCT_MEMBER(FVector4,ViewSizeAndSceneTexelSize)
-	DECLARE_UNIFORM_BUFFER_STRUCT_MEMBER(FVector4,ViewOrigin)
-	DECLARE_UNIFORM_BUFFER_STRUCT_MEMBER(FVector4,TranslatedViewOrigin)
-	// The exposure scale is just a scalar but needs to be a float4 to workaround a driver bug on IOS.
-	// After 4.2 we can put the workaround in the cross compiler.
-	DECLARE_UNIFORM_BUFFER_STRUCT_MEMBER_EX(FVector4,ExposureScale, EShaderPrecisionModifier::Half)
-	DECLARE_UNIFORM_BUFFER_STRUCT_MEMBER_EX(FVector4,DiffuseOverrideParameter, EShaderPrecisionModifier::Half)
-	DECLARE_UNIFORM_BUFFER_STRUCT_MEMBER_EX(FVector4,SpecularOverrideParameter, EShaderPrecisionModifier::Half)
-	DECLARE_UNIFORM_BUFFER_STRUCT_MEMBER_EX(FVector4,NormalOverrideParameter, EShaderPrecisionModifier::Half)
-	DECLARE_UNIFORM_BUFFER_STRUCT_MEMBER_EX(FVector2D,RoughnessOverrideParameter, EShaderPrecisionModifier::Half)
-	DECLARE_UNIFORM_BUFFER_STRUCT_MEMBER(FVector,PreViewTranslation)
-	DECLARE_UNIFORM_BUFFER_STRUCT_MEMBER_EX(float,OutOfBoundsMask, EShaderPrecisionModifier::Half)
-	DECLARE_UNIFORM_BUFFER_STRUCT_MEMBER(FVector,ViewOriginDelta)
-	DECLARE_UNIFORM_BUFFER_STRUCT_MEMBER(float,CullingSign)
-	DECLARE_UNIFORM_BUFFER_STRUCT_MEMBER_EX(float,NearPlane, EShaderPrecisionModifier::Half)
-	DECLARE_UNIFORM_BUFFER_STRUCT_MEMBER(float,AdaptiveTessellationFactor)
-	DECLARE_UNIFORM_BUFFER_STRUCT_MEMBER(float,GameTime)
-	DECLARE_UNIFORM_BUFFER_STRUCT_MEMBER(float,RealTime)
-	DECLARE_UNIFORM_BUFFER_STRUCT_MEMBER(uint32,Random)
-	DECLARE_UNIFORM_BUFFER_STRUCT_MEMBER(uint32,FrameNumber)
-	DECLARE_UNIFORM_BUFFER_STRUCT_MEMBER_EX(float,UseLightmaps, EShaderPrecisionModifier::Half)
-	DECLARE_UNIFORM_BUFFER_STRUCT_MEMBER_EX(float,UnlitViewmodeMask, EShaderPrecisionModifier::Half)
-	DECLARE_UNIFORM_BUFFER_STRUCT_MEMBER_EX(FLinearColor,DirectionalLightColor, EShaderPrecisionModifier::Half)
-	DECLARE_UNIFORM_BUFFER_STRUCT_MEMBER_EX(FVector,DirectionalLightDirection, EShaderPrecisionModifier::Half)
-	DECLARE_UNIFORM_BUFFER_STRUCT_MEMBER_EX(float, DirectionalLightShadowTransition, EShaderPrecisionModifier::Half)			
-	DECLARE_UNIFORM_BUFFER_STRUCT_MEMBER_EX(FVector4, DirectionalLightShadowSize, EShaderPrecisionModifier::Half)
-	DECLARE_UNIFORM_BUFFER_STRUCT_MEMBER_ARRAY(FMatrix, DirectionalLightScreenToShadow, [MAX_FORWARD_SHADOWCASCADES])
-	DECLARE_UNIFORM_BUFFER_STRUCT_MEMBER_EX(FVector4, DirectionalLightShadowDistances, EShaderPrecisionModifier::Half)
-	DECLARE_UNIFORM_BUFFER_STRUCT_MEMBER_EX(FLinearColor,UpperSkyColor, EShaderPrecisionModifier::Half)
-	DECLARE_UNIFORM_BUFFER_STRUCT_MEMBER_EX(FLinearColor,LowerSkyColor, EShaderPrecisionModifier::Half)
-	DECLARE_UNIFORM_BUFFER_STRUCT_MEMBER_ARRAY(FVector4,TranslucencyLightingVolumeMin,[TVC_MAX])
-	DECLARE_UNIFORM_BUFFER_STRUCT_MEMBER_ARRAY(FVector4,TranslucencyLightingVolumeInvSize,[TVC_MAX])
-	DECLARE_UNIFORM_BUFFER_STRUCT_MEMBER(FVector4,TemporalAAParams)
-	DECLARE_UNIFORM_BUFFER_STRUCT_MEMBER(float,DepthOfFieldFocalDistance)
-	DECLARE_UNIFORM_BUFFER_STRUCT_MEMBER(float,DepthOfFieldScale)
-	DECLARE_UNIFORM_BUFFER_STRUCT_MEMBER(float,DepthOfFieldFocalLength)
-	DECLARE_UNIFORM_BUFFER_STRUCT_MEMBER(float,DepthOfFieldFocalRegion)
-	DECLARE_UNIFORM_BUFFER_STRUCT_MEMBER(float,DepthOfFieldNearTransitionRegion)
-	DECLARE_UNIFORM_BUFFER_STRUCT_MEMBER(float,DepthOfFieldFarTransitionRegion)
-	DECLARE_UNIFORM_BUFFER_STRUCT_MEMBER(float,MotionBlurNormalizedToPixel)
-	DECLARE_UNIFORM_BUFFER_STRUCT_MEMBER(float,GeneralPurposeTweak)
-	DECLARE_UNIFORM_BUFFER_STRUCT_MEMBER_EX(float,DemosaicVposOffset, EShaderPrecisionModifier::Half)
-	DECLARE_UNIFORM_BUFFER_STRUCT_MEMBER(FMatrix,PrevProjection)
-	DECLARE_UNIFORM_BUFFER_STRUCT_MEMBER(FMatrix,PrevViewProj)
-	DECLARE_UNIFORM_BUFFER_STRUCT_MEMBER(FMatrix,PrevViewRotationProj)
-	DECLARE_UNIFORM_BUFFER_STRUCT_MEMBER(FMatrix,PrevTranslatedWorldToClip)
-	DECLARE_UNIFORM_BUFFER_STRUCT_MEMBER(FVector,PrevViewOrigin)
-	DECLARE_UNIFORM_BUFFER_STRUCT_MEMBER(FVector,PrevPreViewTranslation)
-	DECLARE_UNIFORM_BUFFER_STRUCT_MEMBER(FMatrix,PrevInvViewProj)
-	DECLARE_UNIFORM_BUFFER_STRUCT_MEMBER(FMatrix,PrevScreenToTranslatedWorld)
-	DECLARE_UNIFORM_BUFFER_STRUCT_MEMBER(FVector,IndirectLightingColorScale)
-	DECLARE_UNIFORM_BUFFER_STRUCT_MEMBER_EX(float,HdrMosaic, EShaderPrecisionModifier::Half)
-	DECLARE_UNIFORM_BUFFER_STRUCT_MEMBER(FVector,AtmosphericFogSunDirection)
-	DECLARE_UNIFORM_BUFFER_STRUCT_MEMBER_EX(float,AtmosphericFogSunPower, EShaderPrecisionModifier::Half)
-	DECLARE_UNIFORM_BUFFER_STRUCT_MEMBER_EX(float,AtmosphericFogPower, EShaderPrecisionModifier::Half)
-	DECLARE_UNIFORM_BUFFER_STRUCT_MEMBER_EX(float,AtmosphericFogDensityScale, EShaderPrecisionModifier::Half)
-	DECLARE_UNIFORM_BUFFER_STRUCT_MEMBER_EX(float,AtmosphericFogDensityOffset, EShaderPrecisionModifier::Half)
-	DECLARE_UNIFORM_BUFFER_STRUCT_MEMBER_EX(float,AtmosphericFogGroundOffset, EShaderPrecisionModifier::Half)
-	DECLARE_UNIFORM_BUFFER_STRUCT_MEMBER_EX(float,AtmosphericFogDistanceScale, EShaderPrecisionModifier::Half)
-	DECLARE_UNIFORM_BUFFER_STRUCT_MEMBER_EX(float,AtmosphericFogAltitudeScale, EShaderPrecisionModifier::Half)
-	DECLARE_UNIFORM_BUFFER_STRUCT_MEMBER_EX(float,AtmosphericFogHeightScaleRayleigh, EShaderPrecisionModifier::Half)
-	DECLARE_UNIFORM_BUFFER_STRUCT_MEMBER_EX(float,AtmosphericFogStartDistance, EShaderPrecisionModifier::Half)
-	DECLARE_UNIFORM_BUFFER_STRUCT_MEMBER_EX(float,AtmosphericFogDistanceOffset, EShaderPrecisionModifier::Half)
-	DECLARE_UNIFORM_BUFFER_STRUCT_MEMBER_EX(float,AtmosphericFogSunDiscScale, EShaderPrecisionModifier::Half)
-	DECLARE_UNIFORM_BUFFER_STRUCT_MEMBER(uint32,AtmosphericFogRenderMask)
-	DECLARE_UNIFORM_BUFFER_STRUCT_MEMBER(uint32,AtmosphericFogInscatterAltitudeSampleNum)
-	DECLARE_UNIFORM_BUFFER_STRUCT_MEMBER(FLinearColor,AtmosphericFogSunColor)
-	DECLARE_UNIFORM_BUFFER_STRUCT_MEMBER(FLinearColor,AmbientCubemapTint)//Used via a custom material node. DO NOT REMOVE.
-	DECLARE_UNIFORM_BUFFER_STRUCT_MEMBER(float,AmbientCubemapIntensity)//Used via a custom material node. DO NOT REMOVE.
-	DECLARE_UNIFORM_BUFFER_STRUCT_MEMBER(FVector2D,RenderTargetSize)
-	DECLARE_UNIFORM_BUFFER_STRUCT_MEMBER(float,SkyLightParameters)
-	DECLARE_UNIFORM_BUFFER_STRUCT_MEMBER(FVector4,SceneTextureMinMax)
-	DECLARE_UNIFORM_BUFFER_STRUCT_MEMBER(FLinearColor,SkyLightColor)
-	DECLARE_UNIFORM_BUFFER_STRUCT_MEMBER_ARRAY(FVector4,SkyIrradianceEnvironmentMap,[7])
-	DECLARE_UNIFORM_BUFFER_STRUCT_MEMBER(float, ES2PreviewMode)
-	DECLARE_UNIFORM_BUFFER_STRUCT_MEMBER_TEXTURE(Texture2D, DirectionalLightShadowTexture)	
-	DECLARE_UNIFORM_BUFFER_STRUCT_MEMBER_SAMPLER(SamplerState, DirectionalLightShadowSampler)
-END_UNIFORM_BUFFER_STRUCT(FViewUniformShaderParameters)
-
-
 //
 // Primitive drawing utility functions.
 //
diff --git a/Engine/Source/Runtime/Engine/Public/SceneView.h b/Engine/Source/Runtime/Engine/Public/SceneView.h
--- a/Engine/Source/Runtime/Engine/Public/SceneView.h
+++ b/Engine/Source/Runtime/Engine/Public/SceneView.h
@@ -9,6 +9,7 @@
 #include "FinalPostProcessSettings.h"
 #include "SceneInterface.h"
 #include "SceneTypes.h"
+#include "ShaderParameters.h"
 
 class FSceneViewStateInterface;
 class FViewUniformShaderParameters;
@@ -194,6 +195,115 @@ struct FViewMatrices
 
 //////////////////////////////////////////////////////////////////////////
 
+static const int MAX_FORWARD_SHADOWCASCADES = 2;
+
+/** 
+ * Enumeration for currently used translucent lighting volume cascades 
+ */
+enum ETranslucencyVolumeCascade
+{
+	TVC_Inner,
+	TVC_Outer,
+
+	TVC_MAX,
+};
+
+/** The uniform shader parameters associated with a view. */
+BEGIN_UNIFORM_BUFFER_STRUCT_WITH_CONSTRUCTOR(FViewUniformShaderParameters,ENGINE_API)
+	DECLARE_UNIFORM_BUFFER_STRUCT_MEMBER(FMatrix,TranslatedWorldToClip)
+	DECLARE_UNIFORM_BUFFER_STRUCT_MEMBER(FMatrix,WorldToClip)
+	DECLARE_UNIFORM_BUFFER_STRUCT_MEMBER(FMatrix,TranslatedWorldToView)
+	DECLARE_UNIFORM_BUFFER_STRUCT_MEMBER(FMatrix,ViewToTranslatedWorld)
+	DECLARE_UNIFORM_BUFFER_STRUCT_MEMBER(FMatrix,ViewToClip)
+	DECLARE_UNIFORM_BUFFER_STRUCT_MEMBER(FMatrix,ClipToView)
+	DECLARE_UNIFORM_BUFFER_STRUCT_MEMBER(FMatrix,ClipToTranslatedWorld)
+	DECLARE_UNIFORM_BUFFER_STRUCT_MEMBER(FMatrix,ScreenToWorld)
+	DECLARE_UNIFORM_BUFFER_STRUCT_MEMBER(FMatrix,ScreenToTranslatedWorld)
+	DECLARE_UNIFORM_BUFFER_STRUCT_MEMBER_EX(FVector,ViewForward, EShaderPrecisionModifier::Half)
+	DECLARE_UNIFORM_BUFFER_STRUCT_MEMBER_EX(FVector,ViewUp, EShaderPrecisionModifier::Half)
+	DECLARE_UNIFORM_BUFFER_STRUCT_MEMBER_EX(FVector,ViewRight, EShaderPrecisionModifier::Half)
+	DECLARE_UNIFORM_BUFFER_STRUCT_MEMBER(FVector4,InvDeviceZToWorldZTransform)
+	DECLARE_UNIFORM_BUFFER_STRUCT_MEMBER_EX(FVector4,ScreenPositionScaleBias, EShaderPrecisionModifier::Half)
+	DECLARE_UNIFORM_BUFFER_STRUCT_MEMBER_EX(FVector4,ViewRectMin, EShaderPrecisionModifier::Half)
+	DECLARE_UNIFORM_BUFFER_STRUCT_MEMBER(FVector4,ViewSizeAndSceneTexelSize)
+	DECLARE_UNIFORM_BUFFER_STRUCT_MEMBER(FVector4,ViewOrigin)
+	DECLARE_UNIFORM_BUFFER_STRUCT_MEMBER(FVector4,TranslatedViewOrigin)
+	// The exposure scale is just a scalar but needs to be a float4 to workaround a driver bug on IOS.
+	// After 4.2 we can put the workaround in the cross compiler.
+	DECLARE_UNIFORM_BUFFER_STRUCT_MEMBER_EX(FVector4,ExposureScale, EShaderPrecisionModifier::Half)
+	DECLARE_UNIFORM_BUFFER_STRUCT_MEMBER_EX(FVector4,DiffuseOverrideParameter, EShaderPrecisionModifier::Half)
+	DECLARE_UNIFORM_BUFFER_STRUCT_MEMBER_EX(FVector4,SpecularOverrideParameter, EShaderPrecisionModifier::Half)
+	DECLARE_UNIFORM_BUFFER_STRUCT_MEMBER_EX(FVector4,NormalOverrideParameter, EShaderPrecisionModifier::Half)
+	DECLARE_UNIFORM_BUFFER_STRUCT_MEMBER_EX(FVector2D,RoughnessOverrideParameter, EShaderPrecisionModifier::Half)
+	DECLARE_UNIFORM_BUFFER_STRUCT_MEMBER(FVector,PreViewTranslation)
+	DECLARE_UNIFORM_BUFFER_STRUCT_MEMBER_EX(float,OutOfBoundsMask, EShaderPrecisionModifier::Half)
+	DECLARE_UNIFORM_BUFFER_STRUCT_MEMBER(FVector,ViewOriginDelta)
+	DECLARE_UNIFORM_BUFFER_STRUCT_MEMBER(float,CullingSign)
+	DECLARE_UNIFORM_BUFFER_STRUCT_MEMBER_EX(float,NearPlane, EShaderPrecisionModifier::Half)
+	DECLARE_UNIFORM_BUFFER_STRUCT_MEMBER(float,AdaptiveTessellationFactor)
+	DECLARE_UNIFORM_BUFFER_STRUCT_MEMBER(float,GameTime)
+	DECLARE_UNIFORM_BUFFER_STRUCT_MEMBER(float,RealTime)
+	DECLARE_UNIFORM_BUFFER_STRUCT_MEMBER(uint32,Random)
+	DECLARE_UNIFORM_BUFFER_STRUCT_MEMBER(uint32,FrameNumber)
+	DECLARE_UNIFORM_BUFFER_STRUCT_MEMBER_EX(float,UseLightmaps, EShaderPrecisionModifier::Half)
+	DECLARE_UNIFORM_BUFFER_STRUCT_MEMBER_EX(float,UnlitViewmodeMask, EShaderPrecisionModifier::Half)
+	DECLARE_UNIFORM_BUFFER_STRUCT_MEMBER_EX(FLinearColor,DirectionalLightColor, EShaderPrecisionModifier::Half)
+	DECLARE_UNIFORM_BUFFER_STRUCT_MEMBER_EX(FVector,DirectionalLightDirection, EShaderPrecisionModifier::Half)
+	DECLARE_UNIFORM_BUFFER_STRUCT_MEMBER_EX(float, DirectionalLightShadowTransition, EShaderPrecisionModifier::Half)			
+	DECLARE_UNIFORM_BUFFER_STRUCT_MEMBER_EX(FVector4, DirectionalLightShadowSize, EShaderPrecisionModifier::Half)
+	DECLARE_UNIFORM_BUFFER_STRUCT_MEMBER_ARRAY(FMatrix, DirectionalLightScreenToShadow, [MAX_FORWARD_SHADOWCASCADES])
+	DECLARE_UNIFORM_BUFFER_STRUCT_MEMBER_EX(FVector4, DirectionalLightShadowDistances, EShaderPrecisionModifier::Half)
+	DECLARE_UNIFORM_BUFFER_STRUCT_MEMBER_EX(FLinearColor,UpperSkyColor, EShaderPrecisionModifier::Half)
+	DECLARE_UNIFORM_BUFFER_STRUCT_MEMBER_EX(FLinearColor,LowerSkyColor, EShaderPrecisionModifier::Half)
+	DECLARE_UNIFORM_BUFFER_STRUCT_MEMBER_ARRAY(FVector4,TranslucencyLightingVolumeMin,[TVC_MAX])
+	DECLARE_UNIFORM_BUFFER_STRUCT_MEMBER_ARRAY(FVector4,TranslucencyLightingVolumeInvSize,[TVC_MAX])
+	DECLARE_UNIFORM_BUFFER_STRUCT_MEMBER(FVector4,TemporalAAParams)
+	DECLARE_UNIFORM_BUFFER_STRUCT_MEMBER(float,DepthOfFieldFocalDistance)
+	DECLARE_UNIFORM_BUFFER_STRUCT_MEMBER(float,DepthOfFieldScale)
+	DECLARE_UNIFORM_BUFFER_STRUCT_MEMBER(float,DepthOfFieldFocalLength)
+	DECLARE_UNIFORM_BUFFER_STRUCT_MEMBER(float,DepthOfFieldFocalRegion)
+	DECLARE_UNIFORM_BUFFER_STRUCT_MEMBER(float,DepthOfFieldNearTransitionRegion)
+	DECLARE_UNIFORM_BUFFER_STRUCT_MEMBER(float,DepthOfFieldFarTransitionRegion)
+	DECLARE_UNIFORM_BUFFER_STRUCT_MEMBER(float,MotionBlurNormalizedToPixel)
+	DECLARE_UNIFORM_BUFFER_STRUCT_MEMBER(float,GeneralPurposeTweak)
+	DECLARE_UNIFORM_BUFFER_STRUCT_MEMBER_EX(float,DemosaicVposOffset, EShaderPrecisionModifier::Half)
+	DECLARE_UNIFORM_BUFFER_STRUCT_MEMBER(FMatrix,PrevProjection)
+	DECLARE_UNIFORM_BUFFER_STRUCT_MEMBER(FMatrix,PrevViewProj)
+	DECLARE_UNIFORM_BUFFER_STRUCT_MEMBER(FMatrix,PrevViewRotationProj)
+	DECLARE_UNIFORM_BUFFER_STRUCT_MEMBER(FMatrix,PrevTranslatedWorldToClip)
+	DECLARE_UNIFORM_BUFFER_STRUCT_MEMBER(FVector,PrevViewOrigin)
+	DECLARE_UNIFORM_BUFFER_STRUCT_MEMBER(FVector,PrevPreViewTranslation)
+	DECLARE_UNIFORM_BUFFER_STRUCT_MEMBER(FMatrix,PrevInvViewProj)
+	DECLARE_UNIFORM_BUFFER_STRUCT_MEMBER(FMatrix,PrevScreenToTranslatedWorld)
+	DECLARE_UNIFORM_BUFFER_STRUCT_MEMBER(FVector,IndirectLightingColorScale)
+	DECLARE_UNIFORM_BUFFER_STRUCT_MEMBER_EX(float,HdrMosaic, EShaderPrecisionModifier::Half)
+	DECLARE_UNIFORM_BUFFER_STRUCT_MEMBER(FVector,AtmosphericFogSunDirection)
+	DECLARE_UNIFORM_BUFFER_STRUCT_MEMBER_EX(float,AtmosphericFogSunPower, EShaderPrecisionModifier::Half)
+	DECLARE_UNIFORM_BUFFER_STRUCT_MEMBER_EX(float,AtmosphericFogPower, EShaderPrecisionModifier::Half)
+	DECLARE_UNIFORM_BUFFER_STRUCT_MEMBER_EX(float,AtmosphericFogDensityScale, EShaderPrecisionModifier::Half)
+	DECLARE_UNIFORM_BUFFER_STRUCT_MEMBER_EX(float,AtmosphericFogDensityOffset, EShaderPrecisionModifier::Half)
+	DECLARE_UNIFORM_BUFFER_STRUCT_MEMBER_EX(float,AtmosphericFogGroundOffset, EShaderPrecisionModifier::Half)
+	DECLARE_UNIFORM_BUFFER_STRUCT_MEMBER_EX(float,AtmosphericFogDistanceScale, EShaderPrecisionModifier::Half)
+	DECLARE_UNIFORM_BUFFER_STRUCT_MEMBER_EX(float,AtmosphericFogAltitudeScale, EShaderPrecisionModifier::Half)
+	DECLARE_UNIFORM_BUFFER_STRUCT_MEMBER_EX(float,AtmosphericFogHeightScaleRayleigh, EShaderPrecisionModifier::Half)
+	DECLARE_UNIFORM_BUFFER_STRUCT_MEMBER_EX(float,AtmosphericFogStartDistance, EShaderPrecisionModifier::Half)
+	DECLARE_UNIFORM_BUFFER_STRUCT_MEMBER_EX(float,AtmosphericFogDistanceOffset, EShaderPrecisionModifier::Half)
+	DECLARE_UNIFORM_BUFFER_STRUCT_MEMBER_EX(float,AtmosphericFogSunDiscScale, EShaderPrecisionModifier::Half)
+	DECLARE_UNIFORM_BUFFER_STRUCT_MEMBER(uint32,AtmosphericFogRenderMask)
+	DECLARE_UNIFORM_BUFFER_STRUCT_MEMBER(uint32,AtmosphericFogInscatterAltitudeSampleNum)
+	DECLARE_UNIFORM_BUFFER_STRUCT_MEMBER(FLinearColor,AtmosphericFogSunColor)
+	DECLARE_UNIFORM_BUFFER_STRUCT_MEMBER(FLinearColor,AmbientCubemapTint)//Used via a custom material node. DO NOT REMOVE.
+	DECLARE_UNIFORM_BUFFER_STRUCT_MEMBER(float,AmbientCubemapIntensity)//Used via a custom material node. DO NOT REMOVE.
+	DECLARE_UNIFORM_BUFFER_STRUCT_MEMBER(FVector2D,RenderTargetSize)
+	DECLARE_UNIFORM_BUFFER_STRUCT_MEMBER(float,SkyLightParameters)
+	DECLARE_UNIFORM_BUFFER_STRUCT_MEMBER(FVector4,SceneTextureMinMax)
+	DECLARE_UNIFORM_BUFFER_STRUCT_MEMBER(FLinearColor,SkyLightColor)
+	DECLARE_UNIFORM_BUFFER_STRUCT_MEMBER_ARRAY(FVector4,SkyIrradianceEnvironmentMap,[7])
+	DECLARE_UNIFORM_BUFFER_STRUCT_MEMBER(float, ES2PreviewMode)
+	DECLARE_UNIFORM_BUFFER_STRUCT_MEMBER_TEXTURE(Texture2D, DirectionalLightShadowTexture)	
+	DECLARE_UNIFORM_BUFFER_STRUCT_MEMBER_SAMPLER(SamplerState, DirectionalLightShadowSampler)
+END_UNIFORM_BUFFER_STRUCT(FViewUniformShaderParameters)
+
 /**
  * A projection from scene space into a 2D screen region.
  */
