Fixed issue where unsaved levels do not merge correctly with the Merge action in the Levels browser.

diff --git a/Engine/Source/Runtime/Engine/Private/World.cpp b/Engine/Source/Runtime/Engine/Private/World.cpp
--- a/Engine/Source/Runtime/Engine/Private/World.cpp
+++ b/Engine/Source/Runtime/Engine/Private/World.cpp
@@ -982,8 +982,9 @@ void UWorld::InitializeNewWorld(const InitializationValues IVS)
 	SpawnInfo.bNoCollisionFail = true;
 	// Set constant name for WorldSettings to make a network replication work between new worlds on host and client
 	SpawnInfo.Name = GEngine->WorldSettingsClass->GetFName();
-	SpawnActor( GEngine->WorldSettingsClass, NULL, NULL, SpawnInfo );
+	AActor* WorldSettings = SpawnActor( GEngine->WorldSettingsClass, NULL, NULL, SpawnInfo );
 	check(GetWorldSettings());
+	WorldSettings->SetIsTemporarilyHiddenInEditor(true);
 
 	// Initialize the world
 	InitWorld(IVS);
