Non-unity build fixes

diff --git a/Engine/Source/Runtime/Engine/Classes/PhysicsEngine/BoxElem.h b/Engine/Source/Runtime/Engine/Classes/PhysicsEngine/BoxElem.h
--- a/Engine/Source/Runtime/Engine/Classes/PhysicsEngine/BoxElem.h
+++ b/Engine/Source/Runtime/Engine/Classes/PhysicsEngine/BoxElem.h
@@ -4,6 +4,8 @@
 
 #include "BoxElem.generated.h"
 
+class FMeshElementCollector;
+
 /** Box shape used for collision */
 USTRUCT()
 struct FKBoxElem
