Slate : only game viewports request all their contents to use EPopupMethod::UseCurrentWindow. All other SViewports default to having no preference.

diff --git a/Engine/Source/Runtime/Engine/Classes/Engine/GameViewportClient.h b/Engine/Source/Runtime/Engine/Classes/Engine/GameViewportClient.h
--- a/Engine/Source/Runtime/Engine/Classes/Engine/GameViewportClient.h
+++ b/Engine/Source/Runtime/Engine/Classes/Engine/GameViewportClient.h
@@ -9,6 +9,7 @@
 #include "DebugDisplayProperty.h"
 #include "TitleSafeZone.h"
 #include "GameViewportDelegates.h"
+#include "SlateBasics.h"
 
 #include "GameViewportClient.generated.h"
 
@@ -587,6 +588,11 @@ public:
 		return bHideCursorDuringCapture;
 	}
 
+	virtual TOptional<EPopupMethod> OnQueryPopupMethod() const override
+	{
+		return EPopupMethod::UseCurrentWindow;
+	}
+
 private:
 	/**
 	 * Set a specific stat to either enabled or disabled (returns the number of remaining enabled stats)
diff --git a/Engine/Source/Runtime/Engine/Private/Slate/SceneViewport.cpp b/Engine/Source/Runtime/Engine/Private/Slate/SceneViewport.cpp
--- a/Engine/Source/Runtime/Engine/Private/Slate/SceneViewport.cpp
+++ b/Engine/Source/Runtime/Engine/Private/Slate/SceneViewport.cpp
@@ -663,6 +663,18 @@ FReply FSceneViewport::OnMotionDetected( const FGeometry& MyGeometry, const FMot
 	return CurrentReplyState;
 }
 
+TOptional<EPopupMethod> FSceneViewport::OnQueryPopupMethod() const
+{
+	if (ViewportClient != nullptr)
+	{
+		return ViewportClient->OnQueryPopupMethod();
+	}
+	else
+	{
+		return TOptional<EPopupMethod>();
+	}
+}
+
 void FSceneViewport::OnFinishedPointerInput()
 {
 	ProcessAccumulatedPointerInput();
diff --git a/Engine/Source/Runtime/Engine/Public/Slate/SceneViewport.h b/Engine/Source/Runtime/Engine/Public/Slate/SceneViewport.h
--- a/Engine/Source/Runtime/Engine/Public/Slate/SceneViewport.h
+++ b/Engine/Source/Runtime/Engine/Public/Slate/SceneViewport.h
@@ -33,10 +33,7 @@ public:
 	virtual void ShowSoftwareCursor( bool bVisible ) override { bIsSoftwareCursorVisible = bVisible; }
 	virtual void SetSoftwareCursorPosition( FVector2D Position ) override { SoftwareCursorPosition = Position; }
 	virtual bool IsSoftwareCursorVisible() const override { return bIsSoftwareCursorVisible; }
-	virtual FVector2D GetSoftwareCursorPosition() const override
-	{
-		return SoftwareCursorPosition;
-	}
+	virtual FVector2D GetSoftwareCursorPosition() const override { return SoftwareCursorPosition; }
 	virtual FCanvas* GetDebugCanvas() override;
 
 	/**
@@ -200,6 +197,7 @@ public:
 	virtual FReply OnTouchEnded( const FGeometry& MyGeometry, const FPointerEvent& InTouchEvent ) override;
 	virtual FReply OnTouchGesture( const FGeometry& MyGeometry, const FPointerEvent& InGestureEvent ) override;
 	virtual FReply OnMotionDetected( const FGeometry& MyGeometry, const FMotionEvent& InMotionEvent ) override;
+	virtual TOptional<EPopupMethod> OnQueryPopupMethod() const override;
 	virtual void OnFinishedPointerInput() override;
 	virtual FReply OnKeyDown( const FGeometry& InGeometry, const FKeyEvent& InKeyEvent ) override;
 	virtual FReply OnKeyUp( const FGeometry& InGeometry, const FKeyEvent& InKeyEvent ) override;
diff --git a/Engine/Source/Runtime/Engine/Public/UnrealClient.h b/Engine/Source/Runtime/Engine/Public/UnrealClient.h
--- a/Engine/Source/Runtime/Engine/Public/UnrealClient.h
+++ b/Engine/Source/Runtime/Engine/Public/UnrealClient.h
@@ -13,6 +13,7 @@
 
 class FCanvas;
 class FViewportClient;
+enum class EPopupMethod : uint8;
 
 /**
  * A render target.
@@ -886,6 +887,9 @@ public:
 	 * Gets whether or not the cursor is hidden when the viewport captures the mouse
 	 */
 	virtual bool HideCursorDuringCapture() { return false; }
+
+	/** Should we make new windows for popups or create an overlay in the current window. */
+	virtual TOptional<EPopupMethod> OnQueryPopupMethod() const { return TOptional<EPopupMethod>(); }
 };
 
 /** Tracks the viewport client that should process the stat command, can be NULL */
diff --git a/Engine/Source/Runtime/Slate/Private/Widgets/SViewport.cpp b/Engine/Source/Runtime/Slate/Private/Widgets/SViewport.cpp
--- a/Engine/Source/Runtime/Slate/Private/Widgets/SViewport.cpp
+++ b/Engine/Source/Runtime/Slate/Private/Widgets/SViewport.cpp
@@ -289,7 +289,15 @@ FReply SViewport::OnMotionDetected( const FGeometry& MyGeometry, const FMotionEv
 
 TOptional<EPopupMethod> SViewport::OnQueryPopupMethod() const
 {
-	return EPopupMethod::CreateNewWindow;
+	TSharedPtr<ISlateViewport> PinnedInterface = ViewportInterface.Pin();
+	if (PinnedInterface.IsValid())
+	{
+		return PinnedInterface->OnQueryPopupMethod();
+	}
+	else
+	{
+		return EPopupMethod::CreateNewWindow;
+	}	
 }
 
 void SViewport::OnFinishedPointerInput()
diff --git a/Engine/Source/Runtime/SlateCore/Public/Rendering/RenderingCommon.h b/Engine/Source/Runtime/SlateCore/Public/Rendering/RenderingCommon.h
--- a/Engine/Source/Runtime/SlateCore/Public/Rendering/RenderingCommon.h
+++ b/Engine/Source/Runtime/SlateCore/Public/Rendering/RenderingCommon.h
@@ -4,6 +4,8 @@
 
 struct FVector2D;
 class FSlateRect;
+enum class EPopupMethod : uint8;
+
 
 #define SLATE_USE_32BIT_INDICES !PLATFORM_USES_ES2
 #define SLATE_USE_FLOAT16 !PLATFORM_USES_ES2
@@ -470,6 +472,16 @@ public:
 	{
 	}
 
+	/**
+	 * Called to figure out whether we can make new windows for popups within this viewport.
+	 * Making windows allows us to have popups that go outside the parent window, but cannot
+	 * be used in fullscreen and do not have per-pixel alpha.
+	 */
+	virtual TOptional<EPopupMethod> OnQueryPopupMethod() const
+	{
+		return TOptional<EPopupMethod>();
+	}
+
 	/**
 	 * Called when navigation is requested
 	 * e.g. Left Joystick, Direction Pad, Arrow Keys can generate navigation events.
