Normally we disconnect a client when we receive suspicious packets, but in this case, we need to ease up until we find out why we're getting false positives during devlopment.

diff --git a/Engine/Source/Runtime/Engine/Private/NetConnection.cpp b/Engine/Source/Runtime/Engine/Private/NetConnection.cpp
--- a/Engine/Source/Runtime/Engine/Private/NetConnection.cpp
+++ b/Engine/Source/Runtime/Engine/Private/NetConnection.cpp
@@ -879,8 +879,9 @@ void UNetConnection::ReceivedPacket( FBitReader& Reader )
 				// Can't handle other channels until control channel exists.
 				if ( Channels[0] == NULL )
 				{
-					UE_LOG( LogNetTraffic, Error, TEXT( "UNetConnection::ReceivedPacket: Received bunch before connected. ChIndex: %i, ChType: %i" ), Bunch.ChIndex, Bunch.ChType );
-					Close();
+					UE_LOG( LogNetTraffic, Error, TEXT( "UNetConnection::ReceivedPacket: Received non-control bunch before control channel was created. ChIndex: %i, ChType: %i" ), Bunch.ChIndex, Bunch.ChType );
+					// Don't disconnect here until we figure out why we get false positives during development
+					//Close();
 					return;
 				}
 				// on the server, if we receive bunch data for a channel that doesn't exist while we're still logging in,
@@ -888,8 +889,9 @@ void UNetConnection::ReceivedPacket( FBitReader& Reader )
 				// so reject it
 				else if ( PlayerController == NULL && Driver->ClientConnections.Contains( this ) )
 				{
-					UE_LOG( LogNetTraffic, Error, TEXT( "UNetConnection::ReceivedPacket: Received non-control bunch during login process" ) );
-					Close();
+					UE_LOG( LogNetTraffic, Error, TEXT( "UNetConnection::ReceivedPacket: Received non-control bunch before player controller was assigned. ChIndex: %i, ChType: %i" ), Bunch.ChIndex, Bunch.ChType );
+					// Don't disconnect here until we figure out why we get false positives during development
+					//Close();
 					return;
 				}
 			}
