Fixed non-editor builds. (Fixes build error from CL 2373742)

diff --git a/Engine/Source/Runtime/Engine/Private/World.cpp b/Engine/Source/Runtime/Engine/Private/World.cpp
--- a/Engine/Source/Runtime/Engine/Private/World.cpp
+++ b/Engine/Source/Runtime/Engine/Private/World.cpp
@@ -984,7 +984,9 @@ void UWorld::InitializeNewWorld(const InitializationValues IVS)
 	SpawnInfo.Name = GEngine->WorldSettingsClass->GetFName();
 	AActor* WorldSettings = SpawnActor( GEngine->WorldSettingsClass, NULL, NULL, SpawnInfo );
 	check(GetWorldSettings());
+#if WITH_EDITOR
 	WorldSettings->SetIsTemporarilyHiddenInEditor(true);
+#endif
 
 	// Initialize the world
 	InitWorld(IVS);
