#Anim: Retargeting: Fixed when applying translation ratio to non-existing parent vs existing parent causing extreme long translation

diff --git a/Engine/Source/Runtime/Engine/Private/Animation/AnimSequence.cpp b/Engine/Source/Runtime/Engine/Private/Animation/AnimSequence.cpp
--- a/Engine/Source/Runtime/Engine/Private/Animation/AnimSequence.cpp
+++ b/Engine/Source/Runtime/Engine/Private/Animation/AnimSequence.cpp
@@ -1792,7 +1792,7 @@ int32 FindMeshBoneIndexFromBoneName(USkeleton * Skeleton, const FName &BoneName)
 
 	return BoneIndex;
 }
-void FillUpTransformBasedOnRig(USkeleton* Skeleton, TArray<FTransform>& NodeSpaceBases, TArray<FTransform> &Rotations, TArray<FTransform>& Translations)
+void FillUpTransformBasedOnRig(USkeleton* Skeleton, TArray<FTransform>& NodeSpaceBases, TArray<FTransform> &Rotations, TArray<FTransform>& Translations, TArray<bool>& TranslationParentFlags)
 {
 	TArray<FTransform> SpaceBases;
 	FAnimationRuntime::FillUpSpaceBasesRetargetBasePose(Skeleton, SpaceBases);
@@ -1814,6 +1814,9 @@ void FillUpTransformBasedOnRig(USkeleton* Skeleton, TArray<FTransform>& NodeSpac
 		Translations.Empty(NodeNum);
 		Translations.AddUninitialized(NodeNum);
 
+		TranslationParentFlags.Empty(Translations.Num());
+		TranslationParentFlags.AddZeroed(Translations.Num());
+		
 		const USkeletalMesh* PreviewMesh = Skeleton->GetPreviewMesh();
 
 		for ( int32 Index = 0; Index < NodeNum; ++Index )
@@ -1866,6 +1869,7 @@ void FillUpTransformBasedOnRig(USkeleton* Skeleton, TArray<FTransform>& NodeSpac
 						{
 							// I think translation has to include rotation, otherwise it won't work
 							Translations[Index] = SpaceBases[BoneIndex].GetRelativeTransform(SpaceBases[ParentBoneIndex]);
+							TranslationParentFlags[BoneIndex] = true;
 						}
 					}
 				}
@@ -1915,9 +1919,10 @@ void UAnimSequence::RemapTracksToNewSkeleton( USkeleton* NewSkeleton, bool bConv
 				// two ref poses. It is very important update ref pose before getting here. 
 				TArray<FTransform> NewRotations, NewSpaceBases, OldRotations;
 				TArray<FTransform> NewTranslations, OldSpaceBases, OldTranslations;
+				TArray<bool> NewTranslationParentFlags, OldTranslationParentFlags;
 				// get the spacebases transform
-				FillUpTransformBasedOnRig(NewSkeleton, NewSpaceBases, NewRotations, NewTranslations);
-				FillUpTransformBasedOnRig(OldSkeleton, OldSpaceBases, OldRotations, OldTranslations);
+				FillUpTransformBasedOnRig(NewSkeleton, NewSpaceBases, NewRotations, NewTranslations, NewTranslationParentFlags);
+				FillUpTransformBasedOnRig(OldSkeleton, OldSpaceBases, OldRotations, OldTranslations, OldTranslationParentFlags);
 
 				// now we'd like to get the relative transform from old to new ref pose in component space
 				// PK2*K2 = PK1*K1*theta where theta => P1*R1*theta = P2*R2 
@@ -1945,9 +1950,19 @@ void UAnimSequence::RemapTracksToNewSkeleton( USkeleton* NewSkeleton, bool bConv
 					FVector OldTranslation = OldTranslations[NodeIndex].GetTranslation();
 					FVector NewTranslation = NewTranslations[NodeIndex].GetTranslation();
 
-					float OldTranslationSize = OldTranslation.Size();
-					float NewTranslationSize = NewTranslation.Size();
-					OldToNewTranslationRatio[NodeIndex] = (FMath::IsNearlyZero(OldTranslationSize)) ? 1.f/*do not touch new translation size*/ : NewTranslationSize / OldTranslationSize;
+					if (NewTranslationParentFlags[NodeIndex] == OldTranslationParentFlags[NodeIndex])
+					{
+						// only do this if parent status matches, otherwise, you'll have invalid state 
+						// where one is based on shoulder, where the other is missing the shoulder node
+						float OldTranslationSize = OldTranslation.Size();
+						float NewTranslationSize = NewTranslation.Size();
+
+						OldToNewTranslationRatio[NodeIndex] = (FMath::IsNearlyZero(OldTranslationSize)) ? 1.f/*do not touch new translation size*/ : NewTranslationSize / OldTranslationSize;
+					}
+					else
+					{
+						OldToNewTranslationRatio[NodeIndex] = 1.f; // set to be 1, we don't know what it is
+					}
 
 					UE_LOG(LogAnimation, Verbose, TEXT("Retargeting (%s : %d) : OldtoNewTranslationRatio (%0.2f), Relative Transform (%s)"), *Nodes[NodeIndex].Name.ToString(), NodeIndex, 
 						OldToNewTranslationRatio[NodeIndex], *RelativeToNewSpaceBases[NodeIndex].ToString());
