Add the channel type and index to the log

diff --git a/Engine/Source/Runtime/Engine/Private/NetConnection.cpp b/Engine/Source/Runtime/Engine/Private/NetConnection.cpp
--- a/Engine/Source/Runtime/Engine/Private/NetConnection.cpp
+++ b/Engine/Source/Runtime/Engine/Private/NetConnection.cpp
@@ -879,7 +879,7 @@ void UNetConnection::ReceivedPacket( FBitReader& Reader )
 				// Can't handle other channels until control channel exists.
 				if ( Channels[0] == NULL )
 				{
-					UE_LOG( LogNetTraffic, Error, TEXT( "UNetConnection::ReceivedPacket: Received bunch before connected. ChIndex: %i, ChType: %i" ), Bunch.ChIndex, Bunch.ChType );
+					UE_LOG( LogNetTraffic, Error, TEXT( "UNetConnection::ReceivedPacket: Received non-control bunch before control channel was created. ChIndex: %i, ChType: %i" ), Bunch.ChIndex, Bunch.ChType );
 					Close();
 					return;
 				}
@@ -888,7 +888,7 @@ void UNetConnection::ReceivedPacket( FBitReader& Reader )
 				// so reject it
 				else if ( PlayerController == NULL && Driver->ClientConnections.Contains( this ) )
 				{
-					UE_LOG( LogNetTraffic, Error, TEXT( "UNetConnection::ReceivedPacket: Received non-control bunch during login process" ) );
+					UE_LOG( LogNetTraffic, Error, TEXT( "UNetConnection::ReceivedPacket: Received non-control bunch before player controller was assigned. ChIndex: %i, ChType: %i" ), Bunch.ChIndex, Bunch.ChType );
 					Close();
 					return;
 				}
