add additional logging for when multicast RPCs overflow

diff --git a/Engine/Source/Runtime/Engine/Private/DataReplication.cpp b/Engine/Source/Runtime/Engine/Private/DataReplication.cpp
--- a/Engine/Source/Runtime/Engine/Private/DataReplication.cpp
+++ b/Engine/Source/Runtime/Engine/Private/DataReplication.cpp
@@ -1047,9 +1047,12 @@ void FObjectReplicator::QueueRemoteFunctionBunch( UFunction* Func, FOutBunch &Bu
 	
 	if (++RemoteFuncInfo[InfoIdx].Calls > 2)
 	{
-		UE_LOG(LogNet, Log, TEXT("Too many calls to RPC %s within a single netupdate. Skipping. %s"), *Func->GetName(), *GetObject()->GetName() );
+		UE_LOG(LogNet, Log, TEXT("Too many calls to RPC %s within a single netupdate. Skipping. %s.  LastCallTime: %.2f. CurrentTime: %.2f. LastRelevantTime: %.2f. LastUpdateTime: %.2f "), 
+			*Func->GetName(), *GetObject()->GetName(), RemoteFuncInfo[InfoIdx].LastCallTime, OwningChannel->Connection->Driver->Time, OwningChannel->RelevantTime, OwningChannel->LastUpdateTime );
 		return;
 	}
+	
+	RemoteFuncInfo[InfoIdx].LastCallTime = OwningChannel->Connection->Driver->Time;
 
 	if (RemoteFunctions == NULL)
 	{
diff --git a/Engine/Source/Runtime/Engine/Public/Net/DataReplication.h b/Engine/Source/Runtime/Engine/Public/Net/DataReplication.h
--- a/Engine/Source/Runtime/Engine/Public/Net/DataReplication.h
+++ b/Engine/Source/Runtime/Engine/Public/Net/DataReplication.h
@@ -82,6 +82,7 @@ public:
 	{
 		FName	FuncName;
 		int32	Calls;
+		float	LastCallTime;
 	};
 
 	TArray< FRPCCallInfo >							RemoteFuncInfo;				// Meta information on pending net RPCs (to be sent)
