#Anim: Fixed a frame delay issue on blendspace evaluator node

https://jira.ol.epicgames.net/browse/UE-5547

diff --git a/Engine/Source/Runtime/Engine/Classes/Animation/AnimNode_BlendSpacePlayer.h b/Engine/Source/Runtime/Engine/Classes/Animation/AnimNode_BlendSpacePlayer.h
--- a/Engine/Source/Runtime/Engine/Classes/Animation/AnimNode_BlendSpacePlayer.h
+++ b/Engine/Source/Runtime/Engine/Classes/Animation/AnimNode_BlendSpacePlayer.h
@@ -64,4 +64,7 @@ public:
 	virtual void OverrideAsset(UAnimationAsset* NewAsset) override;
 	virtual void GatherDebugData(FNodeDebugData& DebugData) override;
 	// End of FAnimNode_Base interface
+
+protected:
+	void UpdateInternal(const FAnimationUpdateContext& Context);
 };
diff --git a/Engine/Source/Runtime/Engine/Private/Animation/AnimNode_BlendSpaceEvaluator.cpp b/Engine/Source/Runtime/Engine/Private/Animation/AnimNode_BlendSpaceEvaluator.cpp
--- a/Engine/Source/Runtime/Engine/Private/Animation/AnimNode_BlendSpaceEvaluator.cpp
+++ b/Engine/Source/Runtime/Engine/Private/Animation/AnimNode_BlendSpaceEvaluator.cpp
@@ -15,10 +15,11 @@ FAnimNode_BlendSpaceEvaluator::FAnimNode_BlendSpaceEvaluator()
 
 void FAnimNode_BlendSpaceEvaluator::Update(const FAnimationUpdateContext& Context)
 {
+	EvaluateGraphExposedInputs.Execute(Context);
 	InternalTimeAccumulator = NormalizedTime;
 	PlayRate = 0.f;
 
-	FAnimNode_BlendSpacePlayer::Update(Context);
+	UpdateInternal(Context);
 }
 
 
diff --git a/Engine/Source/Runtime/Engine/Private/Animation/AnimNode_BlendSpacePlayer.cpp b/Engine/Source/Runtime/Engine/Private/Animation/AnimNode_BlendSpacePlayer.cpp
--- a/Engine/Source/Runtime/Engine/Private/Animation/AnimNode_BlendSpacePlayer.cpp
+++ b/Engine/Source/Runtime/Engine/Private/Animation/AnimNode_BlendSpacePlayer.cpp
@@ -48,6 +48,11 @@ void FAnimNode_BlendSpacePlayer::Update(const FAnimationUpdateContext& Context)
 {
 	EvaluateGraphExposedInputs.Execute(Context);
 
+	UpdateInternal(Context);
+}
+
+void FAnimNode_BlendSpacePlayer::UpdateInternal(const FAnimationUpdateContext& Context)
+{
 	if ((BlendSpace != NULL) && (Context.AnimInstance->CurrentSkeleton->IsCompatible(BlendSpace->GetSkeleton())))
 	{
 		// Create a tick record and fill it out
