#ue4 - Fix for network clients moving incorrectly when they have "bRunPhysicsWithNoController=true" and they are possessed by a player.

- The server was running the movement code locally as well as when receiving the client's move, so would send a correction to the client causing jittery client updates.

Jira: UE-3581

diff --git a/Engine/Source/Runtime/Engine/Private/CharacterMovementComponent.cpp b/Engine/Source/Runtime/Engine/Private/CharacterMovementComponent.cpp
--- a/Engine/Source/Runtime/Engine/Private/CharacterMovementComponent.cpp
+++ b/Engine/Source/Runtime/Engine/Private/CharacterMovementComponent.cpp
@@ -816,7 +816,7 @@ void UCharacterMovementComponent::TickComponent(float DeltaTime, enum ELevelTick
 		}
 
 		// Allow root motion to move characters that have no controller.
-		if( CharacterOwner->IsLocallyControlled() || bRunPhysicsWithNoController || (!CharacterOwner->Controller && CharacterOwner->IsPlayingRootMotion()) )
+		if( CharacterOwner->IsLocallyControlled() || (!CharacterOwner->Controller && bRunPhysicsWithNoController) || (!CharacterOwner->Controller && CharacterOwner->IsPlayingRootMotion()) )
 		{
 			// We need to check the jump state before adjusting input acceleration, to minimize latency
 			// and to make sure acceleration respects our potentially new falling state.
