#ue4 - Characters placed in a level without an AIControllerClass but with CharacterMovement->bRunPhysicsWithNoController set to true will now initialize the default movement mode.

- Previously those Characters would float in the air because movement mode was initailized to "None".

Jira: UE-6143

diff --git a/Engine/Source/Runtime/Engine/Classes/GameFramework/CharacterMovementComponent.h b/Engine/Source/Runtime/Engine/Classes/GameFramework/CharacterMovementComponent.h
--- a/Engine/Source/Runtime/Engine/Classes/GameFramework/CharacterMovementComponent.h
+++ b/Engine/Source/Runtime/Engine/Classes/GameFramework/CharacterMovementComponent.h
@@ -336,7 +336,12 @@ public:
 	UPROPERTY()
 	uint32 bForceMaxAccel:1;    
 
-	/** When there is no Controller, Walking Physics abort and force a velocity and acceleration of 0. Set this to true to override. */
+	/**
+	 * If true, movement will be performed even if there is no Controller for the Character owner.
+	 * Normally without a Controller, movement will be aborted and velocity and acceleration are zeroed if the character is walking.
+	 * Characters that are spawned without a Controller but with this flag enabled will initialize the movement mode to DefaultLandMovementMode or DefaultWaterMovementMode appropriately.
+	 * @see DefaultLandMovementMode, DefaultWaterMovementMode
+	 */
 	UPROPERTY(Category="Character Movement", EditAnywhere, BlueprintReadWrite, AdvancedDisplay)
 	uint32 bRunPhysicsWithNoController:1;
 
@@ -524,11 +529,19 @@ public:
 	UPROPERTY(Category="Character Movement", VisibleInstanceOnly, BlueprintReadOnly)
 	FFindFloorResult CurrentFloor;
 
-	/** Default movement mode when not in water. Used at player startup or when teleported. */
+	/**
+	 * Default movement mode when not in water. Used at player startup or when teleported.
+	 * @see DefaultWaterMovementMode
+	 * @see bRunPhysicsWithNoController
+	 */
 	UPROPERTY(Category="Character Movement", EditAnywhere, BlueprintReadWrite)
 	TEnumAsByte<enum EMovementMode> DefaultLandMovementMode;
 
-	/** Default movement mode when in water. Used at player startup or when teleported. */
+	/**
+	 * Default movement mode when in water. Used at player startup or when teleported.
+	 * @see DefaultLandMovementMode
+	 * @see bRunPhysicsWithNoController
+	 */
 	UPROPERTY(Category="Character Movement", EditAnywhere, BlueprintReadWrite)
 	TEnumAsByte<enum EMovementMode> DefaultWaterMovementMode;
 
diff --git a/Engine/Source/Runtime/Engine/Private/Character.cpp b/Engine/Source/Runtime/Engine/Private/Character.cpp
--- a/Engine/Source/Runtime/Engine/Private/Character.cpp
+++ b/Engine/Source/Runtime/Engine/Private/Character.cpp
@@ -123,6 +123,14 @@ void ACharacter::PostInitializeComponents()
 		{
 			CharacterMovement->UpdateNavAgent(CapsuleComponent);
 		}
+
+		if (Controller == NULL && GetNetMode() != NM_Client)
+		{
+			if (CharacterMovement && CharacterMovement->bRunPhysicsWithNoController)
+			{
+				CharacterMovement->SetDefaultMovementMode();
+			}
+		}
 	}
 }
 
diff --git a/Engine/Source/Runtime/Engine/Private/CharacterMovementComponent.cpp b/Engine/Source/Runtime/Engine/Private/CharacterMovementComponent.cpp
--- a/Engine/Source/Runtime/Engine/Private/CharacterMovementComponent.cpp
+++ b/Engine/Source/Runtime/Engine/Private/CharacterMovementComponent.cpp
@@ -3469,7 +3469,7 @@ void UCharacterMovementComponent::PhysWalking(float deltaTime, int32 Iterations)
 		return;
 	}
 
-	if( (!CharacterOwner || !CharacterOwner->Controller) && !bRunPhysicsWithNoController && !HasRootMotion() )
+	if (!CharacterOwner || (!CharacterOwner->Controller && !bRunPhysicsWithNoController && !HasRootMotion()))
 	{
 		Acceleration = FVector::ZeroVector;
 		Velocity = FVector::ZeroVector;
@@ -3490,7 +3490,7 @@ void UCharacterMovementComponent::PhysWalking(float deltaTime, int32 Iterations)
 	float remainingTime = deltaTime;
 
 	// Perform the move
-	while ( (remainingTime >= MIN_TICK_TIME) && (Iterations < MaxSimulationIterations) && (CharacterOwner->Controller || bRunPhysicsWithNoController || HasRootMotion()) )
+	while ( (remainingTime >= MIN_TICK_TIME) && (Iterations < MaxSimulationIterations) && CharacterOwner && (CharacterOwner->Controller || bRunPhysicsWithNoController || HasRootMotion()) )
 	{
 		Iterations++;
 		bJustTeleported = false;
