#ANIM: Fixed crash with corrupted state of data

https://jira.ol.epicgames.net/browse/UE-5808

diff --git a/Engine/Source/Runtime/Engine/Private/Animation/AnimSequence.cpp b/Engine/Source/Runtime/Engine/Private/Animation/AnimSequence.cpp
--- a/Engine/Source/Runtime/Engine/Private/Animation/AnimSequence.cpp
+++ b/Engine/Source/Runtime/Engine/Private/Animation/AnimSequence.cpp
@@ -530,9 +530,17 @@ void UAnimSequence::VerifyTrackMap()
 			UE_LOG(LogAnimation, Warning, TEXT("RESAVE ANIMATION NEEDED(%s): Fixing track index."), *GetName());
 
 			const TArray<FBoneNode>& BoneTree = MySkeleton->GetBoneTree();
-			for(int32 I=0; I<NumTracks; ++I)
+			for(int32 I=NumTracks-1; I>=0; --I)
 			{
-				TrackToSkeletonMapTable[I].BoneTreeIndex = MySkeleton->GetReferenceSkeleton().FindBoneIndex(AnimationTrackNames[I]);
+				int32 BoneTreeIndex = MySkeleton->GetReferenceSkeleton().FindBoneIndex(AnimationTrackNames[I]);
+				if (BoneTreeIndex == INDEX_NONE)
+				{
+					RemoveTrack(I);
+				}
+				else
+				{
+					TrackToSkeletonMapTable[I].BoneTreeIndex = BoneTreeIndex;
+				}
 			}
 		}
 	}
@@ -2475,7 +2483,11 @@ void UAnimSequence::RemoveTrack(int32 TrackIndex)
 		RawAnimationData.RemoveAt(TrackIndex);
 		AnimationTrackNames.RemoveAt(TrackIndex);
 		TrackToSkeletonMapTable.RemoveAt(TrackIndex);
-		SourceRawAnimationData.RemoveAt(TrackIndex);
+		// source raw animation only exists if edited
+		if (SourceRawAnimationData.Num() > 0 )
+		{
+			SourceRawAnimationData.RemoveAt(TrackIndex);
+		}
 
 		check (RawAnimationData.Num() == AnimationTrackNames.Num() && AnimationTrackNames.Num() == TrackToSkeletonMapTable.Num() );
 	}
