Fixing potential crash + adding sanity check conditions + naming cleanup

diff --git a/Engine/Source/Runtime/Engine/Private/RepLayout.cpp b/Engine/Source/Runtime/Engine/Private/RepLayout.cpp
--- a/Engine/Source/Runtime/Engine/Private/RepLayout.cpp
+++ b/Engine/Source/Runtime/Engine/Private/RepLayout.cpp
@@ -283,7 +283,7 @@ public:
 		FScriptArray* ShadowArray	= (FScriptArray*)ShadowData;
 		FScriptArray* DataArray		= (FScriptArray*)Data;
 
-		TStackState StackState( CmdIndex + 1, Cmd.EndCmd - 1, ShadowArray, DataArray, (uint8*)ShadowArray->GetData(), (uint8*)DataArray->GetData() );
+		TStackState StackState( CmdIndex + 1, Cmd.EndCmd - 1, ShadowArray, DataArray, ShadowArray ? (uint8*)ShadowArray->GetData() : NULL, DataArray ? (uint8*)DataArray->GetData() : NULL );
 
 		static_cast< TImpl* >( this )->ProcessArrayCmd_r( PrevStackState, StackState, Cmd, CmdIndex, ShadowData, Data );
 	}
@@ -810,6 +810,8 @@ void FRepLayout::UpdateChangelistHistory( FRepState * RepState, UClass * ObjectC
 	// Remove any tiling in the history markers to keep them from wrapping over time
 	const int32 NewHistoryCount	= RepState->HistoryEnd - RepState->HistoryStart;
 
+	check( NewHistoryCount <= FRepState::MAX_CHANGE_HISTORY );
+
 	RepState->HistoryStart	= RepState->HistoryStart % FRepState::MAX_CHANGE_HISTORY;
 	RepState->HistoryEnd	= RepState->HistoryStart + NewHistoryCount;
 
@@ -1601,7 +1603,7 @@ public:
 		DirtyList2( InDirty2 ), 
 		DirtyListIndex1( 0 ),
 		DirtyListIndex2( 0 ),
-		Handle( 0 ),
+		CurrentHandle( 0 ),
 		bDirtyValid1( true ),
 		bDirtyValid2( true ),
 		MergedDirtyList( OutMergedDirty ),
@@ -1613,10 +1615,12 @@ public:
 
 	SHOULD_PROCESS_NEXT_CMD() 
 	{ 
-		Handle++;
+		CurrentHandle++;
+
+		check( CurrentHandle != 0 );
 
-		bLastDirty1Matches = bDirtyValid1 && DirtyList1[DirtyListIndex1] == Handle;
-		bLastDirty2Matches = bDirtyValid2 && DirtyList2[DirtyListIndex2] == Handle;
+		bLastDirty1Matches = bDirtyValid1 && DirtyList1[DirtyListIndex1] == CurrentHandle;
+		bLastDirty2Matches = bDirtyValid2 && DirtyList2[DirtyListIndex2] == CurrentHandle;
 
 		return bLastDirty1Matches || bLastDirty2Matches;
 	}
@@ -1631,7 +1635,7 @@ public:
 		const bool bDirty2Matches = bLastDirty2Matches;
 
 		// This will be a new merged dirty entry
-		MergedDirtyList.Add( Handle );
+		MergedDirtyList.Add( CurrentHandle );
 
 		const int32 OriginalMergedDirtyListIndex = MergedDirtyList.AddUninitialized();
 		check( OriginalMergedDirtyListIndex == MergedDirtyList.Num() - 1 );
@@ -1661,14 +1665,14 @@ public:
 		const int32 OldDirtyListIndex1 = DirtyListIndex1;
 		const int32 OldDirtyListIndex2 = DirtyListIndex2;
 
-		const int32 OldHandle = Handle;
-		Handle = 0;
+		const int32 OldHandle = CurrentHandle;
+		CurrentHandle = 0;
 
 		// Process the array elements
 		ProcessDataArrayElements_r( StackState, Cmd );
 
 		// Restore the handle
-		Handle = OldHandle;
+		CurrentHandle = OldHandle;
 
 		if ( bDirty1Matches )
 		{
@@ -1701,7 +1705,7 @@ public:
 		check( bLastDirty1Matches || bLastDirty2Matches )
 
 		// This will be a new merged dirty entry
-		MergedDirtyList.Add( Handle );
+		MergedDirtyList.Add( CurrentHandle );
 
 		// Advance matching dirty indices
 		if ( bLastDirty1Matches )
@@ -1720,7 +1724,7 @@ public:
 
 	int32					DirtyListIndex1;
 	int32					DirtyListIndex2;
-	uint16					Handle;
+	uint16					CurrentHandle;
 
 	bool					bDirtyValid1;
 	bool					bDirtyValid2;
@@ -1737,16 +1741,16 @@ void FRepLayout::MergeDirtyList( FRepState * RepState, const void* RESTRICT Data
 
 	MergedDirty.Empty();
 
-	FMergeDirtyListImpl ReceivePropertiesImpl( Dirty1, Dirty2, MergedDirty, Parents, Cmds );
+	FMergeDirtyListImpl MergePropertiesImpl( Dirty1, Dirty2, MergedDirty, Parents, Cmds );
 
 	// Even though one of these can be empty, we need to send the single one through, so we can prune it to the current shape of the tree
-	ReceivePropertiesImpl.bDirtyValid1 = Dirty1.Num() > 0;
-	ReceivePropertiesImpl.bDirtyValid2 = Dirty2.Num() > 0;
+	MergePropertiesImpl.bDirtyValid1 = Dirty1.Num() > 0;
+	MergePropertiesImpl.bDirtyValid2 = Dirty2.Num() > 0;
 
 	// Merge lists
-	ReceivePropertiesImpl.ProcessCmds( RepState, (uint8*)Data );
+	MergePropertiesImpl.ProcessCmds( RepState, (uint8*)Data );
 
-	ReceivePropertiesImpl.MergedDirtyList.Add( 0 );
+	MergePropertiesImpl.MergedDirtyList.Add( 0 );
 }
 
 void FRepLayout::SanityCheckChangeList_DynamicArray_r( 
