#ue4 - Add finer control for when AI Pawns (with an AIControllerClass set) are automatically possessed. This is done with a new AutoPossessAI enum: possess on level startup, when spawned at run-time, both, or neither.

- Previously the default (and only) behavior was to possess on level startup.
- Renamed existing "AutoPossess" to "AutoPossessPlayer" for clarity and to distinguish from the new AI flag.

diff --git a/Engine/Config/BaseEngine.ini b/Engine/Config/BaseEngine.ini
--- a/Engine/Config/BaseEngine.ini
+++ b/Engine/Config/BaseEngine.ini
@@ -1292,6 +1292,7 @@ bUseStreamingPause=false
 +ActiveClassRedirects=(OldClassName="PlayerReplicationInfo",NewClassName="/Script/Engine.PlayerState")
 +TaggedPropertyRedirects=(ClassName="Controller",OldPropertyName="PlayerReplicationInfo",NewPropertyName="PlayerState")
 +TaggedPropertyRedirects=(ClassName="Pawn",OldPropertyName="PlayerReplicationInfo",NewPropertyName="PlayerState")
++TaggedPropertyRedirects=(ClassName="Pawn",OldPropertyName="AutoPossess",NewPropertyName="AutoPossessPlayer")
 +TaggedPropertyRedirects=(ClassName="GameInfo",OldPropertyName="PlayerReplicationInfoClass",NewPropertyName="PlayerStateClass")
 +TaggedPropertyRedirects=(ClassName="GameMode",OldPropertyName="PlayerReplicationInfoClass",NewPropertyName="PlayerStateClass")
 
diff --git a/Engine/Source/Runtime/Core/Public/UObject/ObjectVersion.h b/Engine/Source/Runtime/Core/Public/UObject/ObjectVersion.h
--- a/Engine/Source/Runtime/Core/Public/UObject/ObjectVersion.h
+++ b/Engine/Source/Runtime/Core/Public/UObject/ObjectVersion.h
@@ -655,6 +655,8 @@ enum EUnrealEngineObjectUE4Version
 	VER_UE4_BLUEPRINT_ENFORCE_CONST_IN_FUNCTION_OVERRIDES,
 	// Added pivot to widget components, need to load old versions as a 0,0 pivot, new default is 0.5,0.5
 	VER_UE4_ADD_PIVOT_TO_WIDGET_COMPONENT,
+	// Added finer control over when AI Pawns are automatically possessed. Also renamed Pawn.AutoPossess to Pawn.AutoPossessPlayer indicate this was a setting for players and not AI.
+	VER_UE4_PAWN_AUTO_POSSESS_AI,
 	// -----<new versions can be added before this line>-------------------------------------------------
 	// - this needs to be the last line (see note below)
 	VER_UE4_AUTOMATIC_VERSION_PLUS_ONE,
diff --git a/Engine/Source/Runtime/Engine/Classes/Engine/EngineTypes.h b/Engine/Source/Runtime/Engine/Classes/Engine/EngineTypes.h
--- a/Engine/Source/Runtime/Engine/Classes/Engine/EngineTypes.h
+++ b/Engine/Source/Runtime/Engine/Classes/Engine/EngineTypes.h
@@ -2192,6 +2192,15 @@ namespace EAutoReceiveInput
 	};
 }
 
+UENUM()
+enum class EAutoPossessAI : uint8
+{
+	Disabled,				// Feature is disabled (do not automatically possess AI).
+	PlacedInWorld,			// Only possess by an AI Controller if Pawn is placed in the world.
+	Spawned,				// Only possess by an AI Controller if Pawn is spawned after the world has loaded.
+	PlacedInWorldOrSpawned, // Pawn is automatically possessed by an AI Controller whenever it is created.
+};
+
 UENUM(BlueprintType)
 namespace EEndPlayReason
 {
diff --git a/Engine/Source/Runtime/Engine/Classes/GameFramework/Pawn.h b/Engine/Source/Runtime/Engine/Classes/GameFramework/Pawn.h
--- a/Engine/Source/Runtime/Engine/Classes/GameFramework/Pawn.h
+++ b/Engine/Source/Runtime/Engine/Classes/GameFramework/Pawn.h
@@ -75,9 +75,26 @@ public:
 	UPROPERTY(EditAnywhere, BlueprintReadWrite, Category=Camera)
 	float BaseEyeHeight;
 
-	/* Specifies which player controller, if any, should automatically possess the pawn when the level starts or the pawn is spawned */
+	/**
+	 * Determines which PlayerController, if any, should automatically possess the pawn when the level starts or when the pawn is spawned.
+	 * @see AutoPossessAI
+	 */
+	UPROPERTY(EditAnywhere, Category=Pawn)
+	TEnumAsByte<EAutoReceiveInput::Type> AutoPossessPlayer;
+
+	/**
+	 * Determines when the Pawn creates and is possessed by an AI Controller (on level start, when spawned, etc).
+	 * Only possible if AIControllerClass is set, and ignored if AutoPossessPlayer is enabled.
+	 * @see AutoPossessPlayer
+	 */
 	UPROPERTY(EditAnywhere, Category=Pawn)
-	TEnumAsByte<EAutoReceiveInput::Type> AutoPossess;
+	EAutoPossessAI AutoPossessAI;
+
+	/**
+	 * Default class to use when pawn is controlled by AI.
+	 */
+	UPROPERTY(EditAnywhere, BlueprintReadWrite, meta=(DisplayName="AI Controller Class"), Category=Pawn)
+	TSubclassOf<AController> AIControllerClass;
 
 public:
 
@@ -99,10 +116,6 @@ public:
 	UFUNCTION(BlueprintCallable, BlueprintAuthorityOnly, Category=AI)
 	void PawnMakeNoise(float Loudness, FVector NoiseLocation, bool bUseNoiseMakerLocation = true, AActor* NoiseMaker = NULL);
 
-	/** default class to use when pawn is controlled by AI. */
-	UPROPERTY(EditAnywhere, BlueprintReadWrite, Category=AI)
-	TSubclassOf<AController>  AIControllerClass;
-
 	/** If Pawn is possessed by a player, points to his playerstate.  Needed for network play as controllers are not replicated to clients. */
 	UPROPERTY(replicatedUsing=OnRep_PlayerState, BlueprintReadOnly, Category="Pawn")
 	APlayerState* PlayerState;
diff --git a/Engine/Source/Runtime/Engine/Private/Pawn.cpp b/Engine/Source/Runtime/Engine/Private/Pawn.cpp
--- a/Engine/Source/Runtime/Engine/Private/Pawn.cpp
+++ b/Engine/Source/Runtime/Engine/Private/Pawn.cpp
@@ -35,6 +35,9 @@ APawn::APawn(const FObjectInitializer& ObjectInitializer)
 {
 	PrimaryActorTick.bCanEverTick = true;
 	PrimaryActorTick.TickGroup = TG_PrePhysics;
+
+	AutoPossessAI = EAutoPossessAI::PlacedInWorld;
+
 	if (HasAnyFlags(RF_ClassDefaultObject) && GetClass() == APawn::StaticClass())
 	{
 		AIControllerClass = LoadClass<AController>(NULL, *((UEngine*)(UEngine::StaticClass()->GetDefaultObject()))->AIControllerClassName.ToString(), NULL, LOAD_None, NULL);
@@ -68,9 +71,9 @@ void APawn::PreInitializeComponents()
 
 	Instigator = this;
 
-	if (AutoPossess != EAutoReceiveInput::Disabled && GetNetMode() != NM_Client )
+	if (AutoPossessPlayer != EAutoReceiveInput::Disabled && GetNetMode() != NM_Client )
 	{
-		const int32 PlayerIndex = int32(AutoPossess.GetValue()) - 1;
+		const int32 PlayerIndex = int32(AutoPossessPlayer.GetValue()) - 1;
 
 		APlayerController* PC = UGameplayStatics::GetPlayerController(this, PlayerIndex);
 		if (PC)
@@ -94,11 +97,16 @@ void APawn::PostInitializeComponents()
 	{
 		GetWorld()->AddPawn( this );
 
-		// automatically add controller to pawns which were placed in level
-		// NOTE: pawns spawned during gameplay are not automatically possessed by a controller
-		if ( GetWorld()->bStartup )
+		// Automatically add Controller to AI Pawns if we are allowed to.
+		if (AutoPossessAI != EAutoPossessAI::Disabled && Controller == NULL && GetNetMode() != NM_Client)
 		{
-			SpawnDefaultController();
+			const bool bPlacedInWorld = (GetWorld()->bStartup);
+			if ((AutoPossessAI == EAutoPossessAI::PlacedInWorldOrSpawned) ||
+				(AutoPossessAI == EAutoPossessAI::PlacedInWorld && bPlacedInWorld) ||
+				(AutoPossessAI == EAutoPossessAI::Spawned && !bPlacedInWorld))
+			{
+				SpawnDefaultController();
+			}
 		}
 
 		// update movement component's nav agent values
