PR #645: Optimize UpdateDrag (Contributed by [REDACTED_PERSON])
https://github.com/EpicGames/UnrealEngine/pull/645

diff --git a/Engine/Source/Runtime/Engine/Private/Vehicles/WheeledVehicleMovementComponent.cpp b/Engine/Source/Runtime/Engine/Private/Vehicles/WheeledVehicleMovementComponent.cpp
--- a/Engine/Source/Runtime/Engine/Private/Vehicles/WheeledVehicleMovementComponent.cpp
+++ b/Engine/Source/Runtime/Engine/Private/Vehicles/WheeledVehicleMovementComponent.cpp
@@ -714,16 +714,20 @@ void UWheeledVehicleMovementComponent::UpdateDrag(float DeltaTime)
 	if (PVehicle && UpdatedComponent)
 	{
 		float ForwardSpeed = GetForwardSpeed();
-		FVector GlobalForwardVector = UpdatedComponent->GetForwardVector();
-		FVector DragVector = -GlobalForwardVector;
-		float SpeedSquared = ForwardSpeed * ForwardSpeed;
-		float ChassisDragArea = ChassisHeight * ChassisWidth;
-		float AirDensity = 1.25 / (100 * 100 * 100); //kg/cm^3
-		float DragMag = 0.5f * AirDensity * SpeedSquared * DragCoefficient * ChassisDragArea;
-		DebugDragMagnitude = DragMag;
-		DragVector *= DragMag;
-		FBodyInstance* BodyInstance = UpdatedComponent->GetBodyInstance();
-		BodyInstance->AddForce(DragVector, false);
+		if (FMath::Abs(ForwardSpeed) > 1.f)
+		{
+			FVector GlobalForwardVector = UpdatedComponent->GetForwardVector();
+			FVector DragVector = -GlobalForwardVector;
+			float SpeedSquared = ForwardSpeed * ForwardSpeed;
+			float ChassisDragArea = ChassisHeight * ChassisWidth;
+			float AirDensity = 1.25 / (100 * 100 * 100); //kg/cm^3
+			float DragMag = 0.5f * AirDensity * SpeedSquared * DragCoefficient * ChassisDragArea;
+			DebugDragMagnitude = DragMag;
+			DragVector *= DragMag;
+			FBodyInstance* BodyInstance = UpdatedComponent->GetBodyInstance();
+			BodyInstance->AddForce(DragVector, false);
+		}
+		
 	}
 }
 
