Made sure audio volume tests use listener location instead of camera location.

diff --git a/Engine/Source/Runtime/Engine/Private/GameViewportClient.cpp b/Engine/Source/Runtime/Engine/Private/GameViewportClient.cpp
--- a/Engine/Source/Runtime/Engine/Private/GameViewportClient.cpp
+++ b/Engine/Source/Runtime/Engine/Private/GameViewportClient.cpp
@@ -854,7 +854,7 @@ void UGameViewportClient::Draw(FViewport* InViewport, FCanvas* SceneCanvas)
 
 								FReverbSettings PlayerReverbSettings;
 								FInteriorSettings PlayerInteriorSettings;
-								class AAudioVolume* PlayerAudioVolume = GetWorld()->GetAudioSettings( ViewLocation, &PlayerReverbSettings, &PlayerInteriorSettings );
+								class AAudioVolume* PlayerAudioVolume = GetWorld()->GetAudioSettings( Location, &PlayerReverbSettings, &PlayerInteriorSettings );
 
 								if (AudioVolume == nullptr || (PlayerAudioVolume != nullptr && PlayerAudioVolume->Priority > AudioVolume->Priority))
 								{
