changed initial navmesh rebuild to a config option of nav system, will reuse serialized data by default

diff --git a/Engine/Source/Runtime/Engine/Classes/AI/Navigation/NavigationSystem.h b/Engine/Source/Runtime/Engine/Classes/AI/Navigation/NavigationSystem.h
--- a/Engine/Source/Runtime/Engine/Classes/AI/Navigation/NavigationSystem.h
+++ b/Engine/Source/Runtime/Engine/Classes/AI/Navigation/NavigationSystem.h
@@ -171,7 +171,12 @@ class ENGINE_API UNavigationSystem : public UBlueprintFunctionLibrary
 	UPROPERTY(config, EditAnywhere, Category=NavigationSystem)
 	uint32 bSkipAgentHeightCheckWhenPickingNavData:1;
 
-	UPROPERTY(config, EditAnywhere, Category=Agents)
+	/** If set to true, navigation data will be always rebuilt from scratch after loading
+	 *  otherwise navigation system will reuse data saved with map */
+	UPROPERTY(config, EditAnywhere, Category = NavigationSystem)
+	uint32 bForceRebuildOnLoad : 1;
+
+	UPROPERTY(config, EditAnywhere, Category = Agents)
 	TArray<FNavDataConfig> SupportedAgents;
 	
 	/** update frequency for dirty areas on navmesh */
diff --git a/Engine/Source/Runtime/Engine/Private/AI/Navigation/NavigationSystem.cpp b/Engine/Source/Runtime/Engine/Private/AI/Navigation/NavigationSystem.cpp
--- a/Engine/Source/Runtime/Engine/Private/AI/Navigation/NavigationSystem.cpp
+++ b/Engine/Source/Runtime/Engine/Private/AI/Navigation/NavigationSystem.cpp
@@ -207,6 +207,7 @@ UNavigationSystem::UNavigationSystem(const FObjectInitializer& ObjectInitializer
 	, bWholeWorldNavigable(false)
 	, bAddPlayersToGenerationSeeds(true)
 	, bSkipAgentHeightCheckWhenPickingNavData(false)
+	, bForceRebuildOnLoad(false)
 	, DirtyAreasUpdateFreq(60)
 	, OperationMode(FNavigationSystem::InvalidMode)
 	, NavOctree(NULL)
@@ -444,11 +445,11 @@ void UNavigationSystem::OnWorldInitDone(FNavigationSystem::EMode Mode)
 		{
 			// don't lock navigation building in editor
 			bInitialBuildingLockActive = false;
-
-			// don't mark dirty areas after loading a map when automatic rebuild is disabled
-			bSkipDirtyAreasOnce = !bNavigationAutoUpdateEnabled;
 		}
-		
+
+		// don't mark dirty areas after loading a map if navigation system doesn't want to rebuilt from scratch
+		bSkipDirtyAreasOnce = !bForceRebuildOnLoad;
+
 		if (bAutoCreateNavigationData == true)
 		{
 			SpawnMissingNavigationData();
