fixed navmesh update when moving actor with other actors attached to it

diff --git a/Engine/Source/Runtime/Engine/Private/ActorEditor.cpp b/Engine/Source/Runtime/Engine/Private/ActorEditor.cpp
--- a/Engine/Source/Runtime/Engine/Private/ActorEditor.cpp
+++ b/Engine/Source/Runtime/Engine/Private/ActorEditor.cpp
@@ -131,6 +131,15 @@ void AActor::PostEditMove(bool bFinished)
 		// USceneComponent::UpdateNavigationData works only in game world
 		UNavigationSystem::UpdateNavOctreeBounds(this);
 		UNavigationSystem::UpdateNavOctreeAll(this);
+
+		TArray<AActor*> ParentedActors;
+		GetAttachedActors(ParentedActors);
+
+		for (int32 Idx = 0; Idx < ParentedActors.Num(); Idx++)
+		{
+			UNavigationSystem::UpdateNavOctreeBounds(ParentedActors[Idx]);
+			UNavigationSystem::UpdateNavOctreeAll(ParentedActors[Idx]);
+		}
 	}
 }
 
