Adding ability to modify listener position via blueprints

Can either set a component to attach to or update an absolute position/rotation every tick.

diff --git a/Engine/Source/Runtime/Engine/Classes/GameFramework/PlayerController.h b/Engine/Source/Runtime/Engine/Classes/GameFramework/PlayerController.h
--- a/Engine/Source/Runtime/Engine/Classes/GameFramework/PlayerController.h
+++ b/Engine/Source/Runtime/Engine/Classes/GameFramework/PlayerController.h
@@ -1234,8 +1234,33 @@ public:
 	/** get audio listener position and orientation */
 	virtual void GetAudioListenerPosition(FVector& OutLocation, FVector& OutFrontDir, FVector& OutRightDir);
 
+	/**
+	 * Used to override the default positioning of the audio listener
+	 * 
+	 * @param AttachToComponent Optional component to attach the audio listener to
+	 * @param Location Depending on whether Component is attached this is either an offset from its location or an absolute position
+	 * @param Rotation Depending on whether Component is attached this is either an offset from its rotation or an absolute rotation
+	 */
+	UFUNCTION(BlueprintCallable, Category="Game|Audio")
+	void SetAudioListenerOverride(USceneComponent* AttachToComponent, FVector Location, FRotator Rotation);
+
+	/**
+	 * Clear any overrides that have been applied to audio listener
+	 */
+	UFUNCTION(BlueprintCallable, Category = "Game|Audio")
+	void ClearAudioListenerOverride();
+
 protected:
+	/** Whether to override the normal audio listener positioning method */
+	uint32 bOverrideAudioListener:1;
+	/** Component that is currently driving the audio listener position/orientation */
+	USceneComponent* AudioListenerComponent;
+	/** Currently overridden location of audio listener */
+	FVector AudioListenerLocationOverride;
+	/** Currently overridden rotation of audio listener */
+	FRotator AudioListenerRotationOverride;
 
+protected:
 	/** Internal. */
 	void TickPlayerInput(const float DeltaSeconds, const bool bGamePaused);
 	virtual void ProcessPlayerInput(const float DeltaTime, const bool bGamePaused);
diff --git a/Engine/Source/Runtime/Engine/Private/PlayerController.cpp b/Engine/Source/Runtime/Engine/Private/PlayerController.cpp
--- a/Engine/Source/Runtime/Engine/Private/PlayerController.cpp
+++ b/Engine/Source/Runtime/Engine/Private/PlayerController.cpp
@@ -2495,7 +2495,25 @@ void APlayerController::GetAudioListenerPosition(FVector& OutLocation, FVector&
 {
 	FVector ViewLocation;
 	FRotator ViewRotation;
-	GetPlayerViewPoint(ViewLocation, ViewRotation);
+
+	if (bOverrideAudioListener)
+	{
+		if (AudioListenerComponent != nullptr)
+		{
+			ViewRotation = AudioListenerComponent->GetComponentRotation() + AudioListenerRotationOverride;
+			ViewLocation = AudioListenerComponent->GetComponentLocation() + ViewRotation.RotateVector(AudioListenerLocationOverride);
+		}
+		else
+		{
+			ViewLocation = AudioListenerLocationOverride;
+			ViewRotation = AudioListenerRotationOverride;
+		}
+	}
+	else
+	{
+		GetPlayerViewPoint(ViewLocation, ViewRotation);
+	}
+
 	const FRotationTranslationMatrix ViewRotationMatrix(ViewRotation, ViewLocation);
 
 	OutLocation = ViewLocation;
@@ -2503,6 +2521,20 @@ void APlayerController::GetAudioListenerPosition(FVector& OutLocation, FVector&
 	OutRightDir = ViewRotationMatrix.GetUnitAxis( EAxis::Y );
 }
 
+void APlayerController::SetAudioListenerOverride(USceneComponent* AttachedComponent, FVector Location, FRotator Rotation)
+{
+	bOverrideAudioListener = true;
+	AudioListenerComponent = AttachedComponent;
+	AudioListenerLocationOverride = Location;
+	AudioListenerRotationOverride = Rotation;
+}
+
+void APlayerController::ClearAudioListenerOverride()
+{
+	bOverrideAudioListener = false;
+	AudioListenerComponent = nullptr;
+}
+
 bool APlayerController::ServerCheckClientPossession_Validate()
 {
 	return true;
