Fix null ptr deref in player controller

diff --git a/Engine/Source/Runtime/Engine/Private/PlayerController.cpp b/Engine/Source/Runtime/Engine/Private/PlayerController.cpp
--- a/Engine/Source/Runtime/Engine/Private/PlayerController.cpp
+++ b/Engine/Source/Runtime/Engine/Private/PlayerController.cpp
@@ -1603,14 +1603,17 @@ bool APlayerController::SetPause( bool bPause, FCanUnpause CanUnpauseDelegate)
 	if (GetNetMode() != NM_Client)
 	{
 		AGameMode* const GameMode = GetWorld()->GetAuthGameMode();
-		if (bPause)
+		if (GameMode != nullptr)
 		{
-			// Pause gamepad rumbling too if needed
-			bResult = GameMode->SetPause(this, CanUnpauseDelegate);
-		}
-		else
-		{
-			GameMode->ClearPause();
+			if (bPause)
+			{
+				// Pause gamepad rumbling too if needed
+				bResult = GameMode->SetPause(this, CanUnpauseDelegate);
+			}
+			else
+			{
+				GameMode->ClearPause();
+			}
 		}
 	}
 	return bResult;
