#ue4 - Fix CharacterMovementComponent never accepting a new UpdatedComponent if it was set during movement and movement mode was "None".

diff --git a/Engine/Source/Runtime/Engine/Private/CharacterMovementComponent.cpp b/Engine/Source/Runtime/Engine/Private/CharacterMovementComponent.cpp
--- a/Engine/Source/Runtime/Engine/Private/CharacterMovementComponent.cpp
+++ b/Engine/Source/Runtime/Engine/Private/CharacterMovementComponent.cpp
@@ -1821,7 +1821,7 @@ void UCharacterMovementComponent::StartNewPhysics(float deltaTime, int32 Iterati
 	switch ( MovementMode )
 	{
 	case MOVE_None:
-		return;
+		break;
 	case MOVE_Walking:
 		PhysWalking(deltaTime, Iterations);
 		break;
