#ue4 - Don't hang on to stale references to UpdatedComponent in MovementComponent.

- Pending kill objects are not allowed to be used by SetUpdatedComponent.
- MovementComponent::TickComponent nulls out the UpdatedComponent if it is pending kill.

diff --git a/Engine/Source/Runtime/Engine/Classes/GameFramework/MovementComponent.h b/Engine/Source/Runtime/Engine/Classes/GameFramework/MovementComponent.h
--- a/Engine/Source/Runtime/Engine/Classes/GameFramework/MovementComponent.h
+++ b/Engine/Source/Runtime/Engine/Classes/GameFramework/MovementComponent.h
@@ -131,6 +131,7 @@ public:
 	uint32 bAutoRegisterUpdatedComponent:1;
 
 	// Begin ActorComponent interface 
+	virtual void TickComponent(float DeltaTime, enum ELevelTick TickType, FActorComponentTickFunction *ThisTickFunction) override;
 	virtual void RegisterComponentTickFunctions(bool bRegister) override;
 	virtual void PostLoad() override;
 
diff --git a/Engine/Source/Runtime/Engine/Private/MovementComponent.cpp b/Engine/Source/Runtime/Engine/Private/MovementComponent.cpp
--- a/Engine/Source/Runtime/Engine/Private/MovementComponent.cpp
+++ b/Engine/Source/Runtime/Engine/Private/MovementComponent.cpp
@@ -53,7 +53,8 @@ void UMovementComponent::SetUpdatedComponent(UPrimitiveComponent* NewUpdatedComp
 		UpdatedComponent->PrimaryComponentTick.RemovePrerequisite(this, PrimaryComponentTick); 
 	}
 
-	UpdatedComponent = NewUpdatedComponent;
+	// Don't assign pending kill components, but allow those to null out previous UpdatedComponent.
+	UpdatedComponent = IsValid(NewUpdatedComponent) ? NewUpdatedComponent : NULL;
 
 	if (IsValid(UpdatedComponent))
 	{
@@ -153,6 +154,18 @@ void UMovementComponent::UpdateTickRegistration()
 }
 
 
+void UMovementComponent::TickComponent(float DeltaTime, ELevelTick TickType, FActorComponentTickFunction *ThisTickFunction)
+{
+	Super::TickComponent(DeltaTime, TickType, ThisTickFunction);
+
+	// Don't hang on to stale references to a destroyed UpdatedComponent.
+	if (UpdatedComponent != NULL && UpdatedComponent->IsPendingKill())
+	{
+		SetUpdatedComponent(NULL);
+	}
+}
+
+
 void UMovementComponent::PostLoad()
 {
 	Super::PostLoad();
diff --git a/Engine/Source/Runtime/Engine/Private/ProjectileMovementComponent.cpp b/Engine/Source/Runtime/Engine/Private/ProjectileMovementComponent.cpp
--- a/Engine/Source/Runtime/Engine/Private/ProjectileMovementComponent.cpp
+++ b/Engine/Source/Runtime/Engine/Private/ProjectileMovementComponent.cpp
@@ -91,6 +91,7 @@ void UProjectileMovementComponent::InitializeComponent()
 void UProjectileMovementComponent::TickComponent(float DeltaTime, enum ELevelTick TickType, FActorComponentTickFunction *ThisTickFunction)
 {
 	QUICK_SCOPE_CYCLE_COUNTER( STAT_ProjectileMovementComponent_TickComponent );
+	Super::TickComponent(DeltaTime, TickType, ThisTickFunction);
 
 	// skip if don't want component updated when not rendered or updated component can't move
 	if ( ShouldSkipUpdate(DeltaTime) )
