#ue4 - Fix character movement when standing on a non-root component that is destroyed. Previously it would teleport the player near the origin.

Jira: UE-6151

diff --git a/Engine/Source/Runtime/Engine/Private/CharacterMovementComponent.cpp b/Engine/Source/Runtime/Engine/Private/CharacterMovementComponent.cpp
--- a/Engine/Source/Runtime/Engine/Private/CharacterMovementComponent.cpp
+++ b/Engine/Source/Runtime/Engine/Private/CharacterMovementComponent.cpp
@@ -1237,7 +1237,7 @@ void UCharacterMovementComponent::UpdateBasedMovement(float DeltaSeconds)
 		return;
 	}
 
-	if (!IsValid(MovementBase->GetOwner()))
+	if (!IsValid(MovementBase) || !IsValid(MovementBase->GetOwner()))
 	{
 		SetBase(NULL);
 		return;
