[UE-5444] custom event net replication flags are now properly inherited by custom event overrides in child Blueprints

change summary:
- modified UBlueprintGeneratedClass::Link() to ensure that any UFunction overrides are in sync with parent function net replication flags prior to linking
- modified UKismetCompilerContext::CreateFunctionStubForEvent() to inherit net replication flags from the parent function if the event node is an override

related issue(s):
- UE-3289

diff --git a/Engine/Source/Editor/KismetCompiler/Private/KismetCompiler.cpp b/Engine/Source/Editor/KismetCompiler/Private/KismetCompiler.cpp
--- a/Engine/Source/Editor/KismetCompiler/Private/KismetCompiler.cpp
+++ b/Engine/Source/Editor/KismetCompiler/Private/KismetCompiler.cpp
@@ -2316,9 +2316,19 @@ void FKismetCompilerContext::CreateFunctionStubForEvent(UK2Node_Event* SrcEventN
 		StubContext.MarkAsInternalOrCppUseOnly();
 	}
 
-	if ((SrcEventNode->FunctionFlags & FUNC_Net) > 0)
+	uint32 FunctionFlags = SrcEventNode->FunctionFlags;
+	if(SrcEventNode->bOverrideFunction && Blueprint->ParentClass != nullptr)
 	{
-		StubContext.MarkAsNetFunction(SrcEventNode->FunctionFlags);
+		const UFunction* ParentFunction = Blueprint->ParentClass->FindFunctionByName(SrcEventNode->GetFunctionName());
+		if(ParentFunction != nullptr)
+		{
+			FunctionFlags |= ParentFunction->FunctionFlags & FUNC_NetFuncFlags;
+		}
+	}
+
+	if ((FunctionFlags & FUNC_Net) > 0)
+	{
+		StubContext.MarkAsNetFunction(FunctionFlags);
 	}
 
 	// Create an entry point
diff --git a/Engine/Source/Runtime/Engine/Private/BlueprintGeneratedClass.cpp b/Engine/Source/Runtime/Engine/Private/BlueprintGeneratedClass.cpp
--- a/Engine/Source/Runtime/Engine/Private/BlueprintGeneratedClass.cpp
+++ b/Engine/Source/Runtime/Engine/Private/BlueprintGeneratedClass.cpp
@@ -483,6 +483,28 @@ void UBlueprintGeneratedClass::DestroyPersistentUberGraphFrame(UObject* Obj) con
 
 void UBlueprintGeneratedClass::Link(FArchive& Ar, bool bRelinkExistingProperties)
 {
+	// Ensure that function netflags equate to any super function in a parent BP prior to linking; it may have been changed by the user
+	// and won't be reflected in the child class until it is recompiled. Without this, UClass::Link() will assert if they are out of sync.
+	for(UField* Field = Children; Field; Field = Field->Next)
+	{
+		Ar.Preload(Field);
+
+		UFunction* Function = dynamic_cast<UFunction*>(Field);
+		if(Function != nullptr)
+		{
+			UFunction* ParentFunction = Function->GetSuperFunction();
+			if(ParentFunction != nullptr)
+			{
+				const uint32 ParentNetFlags = (ParentFunction->FunctionFlags & FUNC_NetFuncFlags);
+				if(ParentNetFlags != (Function->FunctionFlags & FUNC_NetFuncFlags))
+				{
+					Function->FunctionFlags &= ~FUNC_NetFuncFlags;
+					Function->FunctionFlags |= ParentNetFlags;
+				}
+			}
+		}
+	}
+
 	Super::Link(Ar, bRelinkExistingProperties);
 
 	if (UsePersistentUberGraphFrame() && UberGraphFunction)
