Merging RedirectURL into main from UT

//depot/UE4-UT/...

to //depot/UE4/...

diff --git a/Engine/Source/Runtime/Engine/Classes/Engine/EngineBaseTypes.h b/Engine/Source/Runtime/Engine/Classes/Engine/EngineBaseTypes.h
--- a/Engine/Source/Runtime/Engine/Classes/Engine/EngineBaseTypes.h
+++ b/Engine/Source/Runtime/Engine/Classes/Engine/EngineBaseTypes.h
@@ -596,6 +596,9 @@ struct ENGINE_API FURL
 	UPROPERTY()
 	FString Map;    // Map name, i.e. "SkyCity", default is "Entry".
 
+	UPROPERTY()
+	FString RedirectURL;    // Optional place to download Map if client does not possess it
+
 	UPROPERTY()
 	TArray<FString> Op;    // Options.
 
diff --git a/Engine/Source/Runtime/Engine/Classes/GameFramework/GameMode.h b/Engine/Source/Runtime/Engine/Classes/GameFramework/GameMode.h
--- a/Engine/Source/Runtime/Engine/Classes/GameFramework/GameMode.h
+++ b/Engine/Source/Runtime/Engine/Classes/GameFramework/GameMode.h
@@ -565,6 +565,9 @@ public:
 	 */
 	virtual void GetSeamlessTravelActorList(bool bToEntry, TArray<AActor*>& ActorList);
 
+	/** Allow the game to specify a place for clients to download MapName */
+	virtual FString GetRedirectURL(const FString& MapName) const;
+
 	/** 
 	 * used to swap a viewport/connection's PlayerControllers when seamless travelling and the new GameMode's
 	 * controller class is different than the previous
diff --git a/Engine/Source/Runtime/Engine/Private/GameMode.cpp b/Engine/Source/Runtime/Engine/Private/GameMode.cpp
--- a/Engine/Source/Runtime/Engine/Private/GameMode.cpp
+++ b/Engine/Source/Runtime/Engine/Private/GameMode.cpp
@@ -1719,3 +1719,7 @@ void AGameMode::DefaultTimer()
 {	
 }
 
+FString AGameMode::GetRedirectURL(const FString& MapName) const
+{
+	return FString();
+}
\ No newline at end of file
diff --git a/Engine/Source/Runtime/Engine/Private/PendingNetGame.cpp b/Engine/Source/Runtime/Engine/Private/PendingNetGame.cpp
--- a/Engine/Source/Runtime/Engine/Private/PendingNetGame.cpp
+++ b/Engine/Source/Runtime/Engine/Private/PendingNetGame.cpp
@@ -205,8 +205,9 @@ void UPendingNetGame::NotifyControlMessage(UNetConnection* Connection, uint8 Mes
 		{
 			// Server accepted connection.
 			FString GameName;
+			FString RedirectURL;
 
-			FNetControlMessage<NMT_Welcome>::Receive(Bunch, URL.Map, GameName);
+			FNetControlMessage<NMT_Welcome>::Receive(Bunch, URL.Map, GameName, RedirectURL);
 
 			//GEngine->NetworkRemapPath(this, URL.Map);
 
@@ -217,6 +218,7 @@ void UPendingNetGame::NotifyControlMessage(UNetConnection* Connection, uint8 Mes
 				FURL DefaultURL;
 				FURL TempURL( &DefaultURL, *URL.Map, TRAVEL_Partial );
 				URL.Map = TempURL.Map;
+				URL.RedirectURL = RedirectURL;
 				URL.Op.Append(TempURL.Op);
 			}
 
diff --git a/Engine/Source/Runtime/Engine/Private/UnrealEngine.cpp b/Engine/Source/Runtime/Engine/Private/UnrealEngine.cpp
--- a/Engine/Source/Runtime/Engine/Private/UnrealEngine.cpp
+++ b/Engine/Source/Runtime/Engine/Private/UnrealEngine.cpp
@@ -8303,32 +8303,40 @@ bool UEngine::TickWorldTravel(FWorldContext& Context, float DeltaSeconds)
 		}
 		else if( Context.PendingNetGame && Context.PendingNetGame->bSuccessfullyConnected && !Context.PendingNetGame->bSentJoinRequest )
 		{
-			// Attempt to load the map.
-			FString Error;
-			
-			const bool bLoadedMapSuccessfully = LoadMap( Context, Context.PendingNetGame->URL, Context.PendingNetGame, Error );
-				
-			if( !bLoadedMapSuccessfully || Error != TEXT("") )
+			if (!MakeSureMapNameIsValid(Context.PendingNetGame->URL.Map))
 			{
-				// we can't guarantee the current World is in a valid state, so travel to the default map
 				BrowseToDefaultMap(Context);
-				BroadcastTravelFailure(Context.World(), ETravelFailure::LoadMapFailure, Error);
-				check(Context.World() != NULL);
+				BroadcastTravelFailure(Context.World(), ETravelFailure::PackageMissing, Context.PendingNetGame->URL.RedirectURL);
 			}
 			else
 			{
-				// Show connecting message, cause precaching to occur.
-				TransitionType = TT_Connecting;
-					
-				RedrawViewports();
+				// Attempt to load the map.
+				FString Error;
 
-				// Send join.
-				Context.PendingNetGame->SendJoin();
-				Context.PendingNetGame->NetDriver = NULL;
-			}
+				const bool bLoadedMapSuccessfully = LoadMap(Context, Context.PendingNetGame->URL, Context.PendingNetGame, Error);
 
-			// Kill the pending level.
-			Context.PendingNetGame = NULL;
+				if (!bLoadedMapSuccessfully || Error != TEXT(""))
+				{
+					// we can't guarantee the current World is in a valid state, so travel to the default map
+					BrowseToDefaultMap(Context);
+					BroadcastTravelFailure(Context.World(), ETravelFailure::LoadMapFailure, Error);
+					check(Context.World() != NULL);
+				}
+				else
+				{
+					// Show connecting message, cause precaching to occur.
+					TransitionType = TT_Connecting;
+
+					RedrawViewports();
+
+					// Send join.
+					Context.PendingNetGame->SendJoin();
+					Context.PendingNetGame->NetDriver = NULL;
+				}
+
+				// Kill the pending level.
+				Context.PendingNetGame = NULL;
+			}
 		}
 	}
 	else if (TransitionType == TT_WaitingToConnect)
diff --git a/Engine/Source/Runtime/Engine/Private/World.cpp b/Engine/Source/Runtime/Engine/Private/World.cpp
--- a/Engine/Source/Runtime/Engine/Private/World.cpp
+++ b/Engine/Source/Runtime/Engine/Private/World.cpp
@@ -3549,12 +3549,14 @@ void UWorld::WelcomePlayer(UNetConnection* Connection)
 	Connection->ClientWorldPackageName = CurrentLevel->GetOutermost()->GetFName();
 
 	FString GameName;
+	FString RedirectURL;
 	if (AuthorityGameMode != NULL)
 	{
 		GameName = AuthorityGameMode->GetClass()->GetPathName();
+		RedirectURL = AuthorityGameMode->GetRedirectURL(LevelName);
 	}
 
-	FNetControlMessage<NMT_Welcome>::Send(Connection, LevelName, GameName);
+	FNetControlMessage<NMT_Welcome>::Send(Connection, LevelName, GameName, RedirectURL);
 	Connection->FlushNet();
 	// don't count initial join data for netspeed throttling
 	// as it's unnecessary, since connection won't be fully open until it all gets received, and this prevents later gameplay data from being delayed to "catch up"
diff --git a/Engine/Source/Runtime/Engine/Public/Net/DataChannel.h b/Engine/Source/Runtime/Engine/Public/Net/DataChannel.h
--- a/Engine/Source/Runtime/Engine/Public/Net/DataChannel.h
+++ b/Engine/Source/Runtime/Engine/Public/Net/DataChannel.h
@@ -371,7 +371,7 @@ private:
 
 // message type definitions
 DEFINE_CONTROL_CHANNEL_MESSAGE_FOURPARAM(Hello, 0, uint8, int32, int32, FGuid); // initial client connection message
-DEFINE_CONTROL_CHANNEL_MESSAGE_TWOPARAM(Welcome, 1, FString, FString); // server tells client they're ok'ed to load the server's level
+DEFINE_CONTROL_CHANNEL_MESSAGE_THREEPARAM(Welcome, 1, FString, FString, FString); // server tells client they're ok'ed to load the server's level
 DEFINE_CONTROL_CHANNEL_MESSAGE_TWOPARAM(Upgrade, 2, int32, int32); // server tells client their version is incompatible
 DEFINE_CONTROL_CHANNEL_MESSAGE_TWOPARAM(Challenge, 3, int32, FString); // server sends client challenge string to verify integrity
 DEFINE_CONTROL_CHANNEL_MESSAGE_ONEPARAM(Netspeed, 4, int32); // client sends requested transfer rate
