On windows change the default console key for German, French, and Spanish keyboard layouts

diff --git a/Engine/Source/Runtime/Engine/Private/UserInterface/InputSettings.cpp b/Engine/Source/Runtime/Engine/Private/UserInterface/InputSettings.cpp
--- a/Engine/Source/Runtime/Engine/Private/UserInterface/InputSettings.cpp
+++ b/Engine/Source/Runtime/Engine/Private/UserInterface/InputSettings.cpp
@@ -40,6 +40,31 @@ void UInputSettings::PostInitProperties()
 			AxisConfig.RemoveAtSwap(Index);
 		}
 	}
+
+#if PLATFORM_WINDOWS
+	// If the console key is set to the default we'll see about switching it to the keyboard default
+	// If they've mapped any additional keys, we'll just assume they've set it up in a way they desire
+	if (ConsoleKeys.Num() == 1 && ConsoleKeys[0] == EKeys::Tilde)
+	{
+		FKey DefaultConsoleKey = EKeys::Tilde;
+		switch(PRIMARYLANGID(GetKeyboardLayout(0)))
+		{
+		case LANG_FRENCH:
+			DefaultConsoleKey = FInputKeyManager::Get().GetKeyFromCodes(VK_OEM_7, 0);
+			break;
+
+		case LANG_GERMAN:
+			DefaultConsoleKey = EKeys::Caret;
+			break;
+
+		case LANG_SPANISH:
+			DefaultConsoleKey = FInputKeyManager::Get().GetKeyFromCodes(VK_OEM_5, 0);
+			break;
+		}
+
+		ConsoleKeys[0] = DefaultConsoleKey;
+	}
+#endif
 }
 
 #if WITH_EDITOR
