Fix compile warning

diff --git a/Engine/Source/Runtime/Engine/Private/UserInterface/InputSettings.cpp b/Engine/Source/Runtime/Engine/Private/UserInterface/InputSettings.cpp
--- a/Engine/Source/Runtime/Engine/Private/UserInterface/InputSettings.cpp
+++ b/Engine/Source/Runtime/Engine/Private/UserInterface/InputSettings.cpp
@@ -47,7 +47,7 @@ void UInputSettings::PostInitProperties()
 	if (ConsoleKeys.Num() == 1 && ConsoleKeys[0] == EKeys::Tilde)
 	{
 		FKey DefaultConsoleKey = EKeys::Tilde;
-		switch(PRIMARYLANGID(GetKeyboardLayout(0)))
+		switch(PRIMARYLANGID(LOWORD(GetKeyboardLayout(0))))
 		{
 		case LANG_FRENCH:
 			DefaultConsoleKey = FInputKeyManager::Get().GetKeyFromCodes(VK_OEM_7, 0);
