Fixed broken automation tests for Visual Logger (UE-6199 issue)

diff --git a/Engine/Source/Runtime/Engine/Private/VisualLogger/VisualLoggerAutomationTests.cpp b/Engine/Source/Runtime/Engine/Private/VisualLogger/VisualLoggerAutomationTests.cpp
--- a/Engine/Source/Runtime/Engine/Private/VisualLogger/VisualLoggerAutomationTests.cpp
+++ b/Engine/Source/Runtime/Engine/Private/VisualLogger/VisualLoggerAutomationTests.cpp
@@ -101,13 +101,12 @@ bool FVisualLogTest::RunTest(const FString& Parameters)
 
 	{
 		const FString TextToLog = TEXT("Simple text line to test if UE_VLOG_UELOG works fine");
+		float CurrentTimestamp = GWorld->TimeSeconds;
 		UE_VLOG_UELOG(GWorld, LogVisual, Log, TEXT("%s"), *TextToLog);
 		CHECK_SUCCESS(Context.Device.LastObject != GWorld);
 		CHECK_SUCCESS(Context.Device.LastEntry.TimeStamp == -1);
 		FVisualLogEntry* CurrentEntry = FVisualLogger::Get().GetEntryToWrite(GWorld, GWorld->TimeSeconds, ECreateIfNeeded::DontCreate);
 		
-		float CurrentTimestamp = GWorld->TimeSeconds;
-		for (int32 Index = 0; Index < 2; ++Index)
 		{
 			CHECK_SUCCESS(CurrentEntry != NULL);
 			CHECK_SUCCESS(CurrentEntry->TimeStamp == CurrentTimestamp);
@@ -117,6 +116,9 @@ bool FVisualLogTest::RunTest(const FString& Parameters)
 
 			FVisualLogger::Get().GetEntryToWrite(GWorld, CurrentTimestamp + 0.1); //generate new entry and serialize old one
 			CurrentEntry = &Context.Device.LastEntry;
+			CHECK_SUCCESS(CurrentEntry != NULL);
+			CHECK_SUCCESS(CurrentEntry->TimeStamp == -1);
+			CHECK_SUCCESS(CurrentEntry->LogLines.Num() == 0);
 		}
 	}
 
@@ -149,7 +151,6 @@ bool FVisualLogSegmentsTest::RunTest(const FString& Parameters)
 		FVisualLogEntry* CurrentEntry = FVisualLogger::Get().GetEntryToWrite(GWorld, GWorld->TimeSeconds, ECreateIfNeeded::DontCreate);
 
 		float CurrentTimestamp = GWorld->TimeSeconds;
-		for (int32 Index = 0; Index < 2; ++Index)
 		{
 			CHECK_SUCCESS(CurrentEntry != NULL);
 			CHECK_SUCCESS(CurrentEntry->TimeStamp == CurrentTimestamp);
@@ -161,6 +162,9 @@ bool FVisualLogSegmentsTest::RunTest(const FString& Parameters)
 
 			FVisualLogger::Get().GetEntryToWrite(GWorld, CurrentTimestamp + 0.1); //generate new entry and serialize old one
 			CurrentEntry = &Context.Device.LastEntry;
+			CHECK_SUCCESS(CurrentEntry != NULL);
+			CHECK_SUCCESS(CurrentEntry->TimeStamp == -1);
+			CHECK_SUCCESS(CurrentEntry->ElementsToDraw.Num() == 0);
 		}
 	}
 	return true;
@@ -207,7 +211,6 @@ bool FVisualLogEventsTest::RunTest(const FString& Parameters)
 	CurrentTimestamp = GWorld->TimeSeconds;
 	UE_VLOG_EVENTS(GWorld, NAME_None, EventTest, EventTest2, EventTest3);
 
-	for (int32 Index = 0; Index < 2; ++Index)
 	{
 		CHECK_SUCCESS(CurrentEntry != NULL);
 		CHECK_SUCCESS(CurrentEntry->TimeStamp == CurrentTimestamp);
@@ -225,6 +228,9 @@ bool FVisualLogEventsTest::RunTest(const FString& Parameters)
 
 		FVisualLogger::Get().GetEntryToWrite(GWorld, CurrentTimestamp+0.1); //generate new entry and serialize old one
 		CurrentEntry = &Context.Device.LastEntry;
+		CHECK_SUCCESS(CurrentEntry != NULL);
+		CHECK_SUCCESS(CurrentEntry->TimeStamp == -1);
+		CHECK_SUCCESS(CurrentEntry->Events.Num() == 0);
 	}
 
 	const FName EventTag1 = TEXT("ATLAS_C_0");
