Fixed not handling null navigation area classes in NavigationSystem.SupportedAgents #UE4

diff --git a/Engine/Source/Runtime/Engine/Private/AI/Navigation/NavigationSystem.cpp b/Engine/Source/Runtime/Engine/Private/AI/Navigation/NavigationSystem.cpp
--- a/Engine/Source/Runtime/Engine/Private/AI/Navigation/NavigationSystem.cpp
+++ b/Engine/Source/Runtime/Engine/Private/AI/Navigation/NavigationSystem.cpp
@@ -375,7 +375,14 @@ void UNavigationSystem::PostEditChangeProperty(FPropertyChangedEvent& PropertyCh
 			if (SupportedAgentIndex != INDEX_NONE)
 			{
 				// reflect the change to SupportedAgent's 
-				SupportedAgents[SupportedAgentIndex].NavigationDataClassName = FStringClassReference::GetOrCreateIDForClass(SupportedAgents[SupportedAgentIndex].NavigationDataClass);
+				if (SupportedAgents[SupportedAgentIndex].NavigationDataClass != nullptr)
+				{
+					SupportedAgents[SupportedAgentIndex].NavigationDataClassName = FStringClassReference::GetOrCreateIDForClass(SupportedAgents[SupportedAgentIndex].NavigationDataClass);
+				}
+				else
+				{
+					SupportedAgents[SupportedAgentIndex].NavigationDataClassName.Reset();
+				}
 			}
 		}
 	}
@@ -2778,7 +2785,7 @@ void UNavigationSystem::SpawnMissingNavigationData()
 	{
 		for (int32 AgentIndex = 0; AgentIndex < SupportedAgentsCount; ++AgentIndex)
 		{
-			if (AlreadyInstantiated[AgentIndex] == false)
+			if (AlreadyInstantiated[AgentIndex] == false && SupportedAgents[AgentIndex].NavigationDataClass != nullptr)
 			{
 				bool bHandled = false;
 
