Engine: Add optional Thickness and DepthBias parameters to DrawOrientedWireBox

diff --git a/Engine/Source/Runtime/Engine/Private/PrimitiveDrawingUtils.cpp b/Engine/Source/Runtime/Engine/Private/PrimitiveDrawingUtils.cpp
--- a/Engine/Source/Runtime/Engine/Private/PrimitiveDrawingUtils.cpp
+++ b/Engine/Source/Runtime/Engine/Private/PrimitiveDrawingUtils.cpp
@@ -1196,34 +1196,39 @@ void DrawWireChoppedCone(FPrimitiveDrawInterface* PDI,const FVector& Base,const
  * @param	DepthPriority	Depth priority for the cone.
  */
 
-void DrawOrientedWireBox(FPrimitiveDrawInterface* PDI,const FVector& Base,const FVector& X,const FVector& Y,const FVector& Z, FVector Extent, const FLinearColor& Color,uint8 DepthPriority)
+void DrawOrientedWireBox(FPrimitiveDrawInterface* PDI, const FVector& Base, const FVector& X, const FVector& Y, const FVector& Z, FVector Extent, const FLinearColor& Color, uint8 DepthPriority, float Thickness, float DepthBias, bool bScreenSpace)
 {
 	FVector	B[2],P,Q;
-	int32 i,j;
 
-	FMatrix m(X, Y, Z, Base);
+	FMatrix M(X, Y, Z, Base);
 	B[0] = -Extent;
 	B[1] = Extent;
 
-	for( i=0; i<2; i++ ) for( j=0; j<2; j++ )
+	for (int32 i = 0; i < 2; i++)
 	{
-		P.X=B[i].X; Q.X=B[i].X;
-		P.Y=B[j].Y; Q.Y=B[j].Y;
-		P.Z=B[0].Z; Q.Z=B[1].Z;
-		P = m.TransformPosition(P); Q = m.TransformPosition(Q);
-		PDI->DrawLine(P,Q,Color,DepthPriority);
-
-		P.Y=B[i].Y; Q.Y=B[i].Y;
-		P.Z=B[j].Z; Q.Z=B[j].Z;
-		P.X=B[0].X; Q.X=B[1].X;
-		P = m.TransformPosition(P); Q = m.TransformPosition(Q);
-		PDI->DrawLine(P,Q,Color,DepthPriority);
-
-		P.Z=B[i].Z; Q.Z=B[i].Z;
-		P.X=B[j].X; Q.X=B[j].X;
-		P.Y=B[0].Y; Q.Y=B[1].Y;
-		P = m.TransformPosition(P); Q = m.TransformPosition(Q);
-		PDI->DrawLine(P,Q,Color,DepthPriority);
+		for (int32 j = 0; j < 2; j++)
+		{
+			P.X = B[i].X; Q.X = B[i].X;
+			P.Y = B[j].Y; Q.Y = B[j].Y;
+			P.Z = B[0].Z; Q.Z = B[1].Z;
+			P = M.TransformPosition(P);
+			Q = M.TransformPosition(Q);
+			PDI->DrawLine(P, Q, Color, DepthPriority, Thickness, DepthBias, bScreenSpace);
+
+			P.Y = B[i].Y; Q.Y = B[i].Y;
+			P.Z = B[j].Z; Q.Z = B[j].Z;
+			P.X = B[0].X; Q.X = B[1].X;
+			P = M.TransformPosition(P);
+			Q = M.TransformPosition(Q);
+			PDI->DrawLine(P, Q, Color, DepthPriority, Thickness, DepthBias, bScreenSpace);
+
+			P.Z = B[i].Z; Q.Z = B[i].Z;
+			P.X = B[j].X; Q.X = B[j].X;
+			P.Y = B[0].Y; Q.Y = B[1].Y;
+			P = M.TransformPosition(P);
+			Q = M.TransformPosition(Q);
+			PDI->DrawLine(P, Q, Color, DepthPriority, Thickness, DepthBias, bScreenSpace);
+		}
 	}
 }
 
diff --git a/Engine/Source/Runtime/Engine/Public/SceneManagement.h b/Engine/Source/Runtime/Engine/Public/SceneManagement.h
--- a/Engine/Source/Runtime/Engine/Public/SceneManagement.h
+++ b/Engine/Source/Runtime/Engine/Public/SceneManagement.h
@@ -1514,13 +1514,13 @@ extern ENGINE_API void DrawWireSphere(class FPrimitiveDrawInterface* PDI, const
 extern ENGINE_API void DrawWireSphere(class FPrimitiveDrawInterface* PDI, const FTransform& Transform, const FLinearColor& Color, float Radius, int32 NumSides, uint8 DepthPriority, float Thickness = 0, float DepthBias = 0.0f, bool bScreenSpace = false);
 extern ENGINE_API void DrawWireSphereAutoSides(class FPrimitiveDrawInterface* PDI, const FVector& Base, const FLinearColor& Color, float Radius, uint8 DepthPriority, float Thickness = 0, float DepthBias = 0.0f, bool bScreenSpace = false);
 extern ENGINE_API void DrawWireSphereAutoSides(class FPrimitiveDrawInterface* PDI, const FTransform& Transform, const FLinearColor& Color, float Radius, uint8 DepthPriority, float Thickness = 0, float DepthBias = 0.0f, bool bScreenSpace = false);
-extern ENGINE_API void DrawWireCylinder(class FPrimitiveDrawInterface* PDI, const FVector& Base, const FVector& X, const FVector& Y, const FVector& Z, const FLinearColor& Color, float Radius, float HalfHeight, int32 NumSides, uint8 DepthPriority, float Thickness = 0, float DepthBias = 0.0f, bool bScreenSpace = false);
-extern ENGINE_API void DrawWireCapsule(class FPrimitiveDrawInterface* PDI, const FVector& Base, const FVector& X, const FVector& Y, const FVector& Z, const FLinearColor& Color, float Radius, float HalfHeight, int32 NumSides, uint8 DepthPriority, float Thickness = 0, float DepthBias = 0.0f, bool bScreenSpace = false);
+extern ENGINE_API void DrawWireCylinder(class FPrimitiveDrawInterface* PDI, const FVector& Base, const FVector& X, const FVector& Y, const FVector& Z, const FLinearColor& Color, float Radius, float HalfHeight, int32 NumSides, uint8 DepthPriority, float Thickness = 0.0f, float DepthBias = 0.0f, bool bScreenSpace = false);
+extern ENGINE_API void DrawWireCapsule(class FPrimitiveDrawInterface* PDI, const FVector& Base, const FVector& X, const FVector& Y, const FVector& Z, const FLinearColor& Color, float Radius, float HalfHeight, int32 NumSides, uint8 DepthPriority, float Thickness = 0.0f, float DepthBias = 0.0f, bool bScreenSpace = false);
 extern ENGINE_API void DrawWireChoppedCone(class FPrimitiveDrawInterface* PDI,const FVector& Base,const FVector& X,const FVector& Y,const FVector& Z,const FLinearColor& Color,float Radius,float TopRadius,float HalfHeight,int32 NumSides,uint8 DepthPriority);
-extern ENGINE_API void DrawWireCone(class FPrimitiveDrawInterface* PDI, TArray<FVector>& Verts, const FMatrix& Transform, float ConeRadius, float ConeAngle, int32 ConeSides, const FLinearColor& Color, uint8 DepthPriority, float Thickness = 0, float DepthBias = 0.0f, bool bScreenSpace = false);
-extern ENGINE_API void DrawWireCone(class FPrimitiveDrawInterface* PDI, TArray<FVector>& Verts, const FTransform& Transform, float ConeRadius, float ConeAngle, int32 ConeSides, const FLinearColor& Color, uint8 DepthPriority, float Thickness = 0, float DepthBias = 0.0f, bool bScreenSpace = false);
+extern ENGINE_API void DrawWireCone(class FPrimitiveDrawInterface* PDI, TArray<FVector>& Verts, const FMatrix& Transform, float ConeRadius, float ConeAngle, int32 ConeSides, const FLinearColor& Color, uint8 DepthPriority, float Thickness = 0.0f, float DepthBias = 0.0f, bool bScreenSpace = false);
+extern ENGINE_API void DrawWireCone(class FPrimitiveDrawInterface* PDI, TArray<FVector>& Verts, const FTransform& Transform, float ConeRadius, float ConeAngle, int32 ConeSides, const FLinearColor& Color, uint8 DepthPriority, float Thickness = 0.0f, float DepthBias = 0.0f, bool bScreenSpace = false);
 extern ENGINE_API void DrawWireSphereCappedCone(FPrimitiveDrawInterface* PDI, const FTransform& Transform, float ConeRadius, float ConeAngle, int32 ConeSides, int32 ArcFrequency, int32 CapSegments, const FLinearColor& Color, uint8 DepthPriority);
-extern ENGINE_API void DrawOrientedWireBox(class FPrimitiveDrawInterface* PDI,const FVector& Base,const FVector& X,const FVector& Y,const FVector& Z, FVector Extent, const FLinearColor& Color,uint8 DepthPriority);
+extern ENGINE_API void DrawOrientedWireBox(class FPrimitiveDrawInterface* PDI, const FVector& Base, const FVector& X, const FVector& Y, const FVector& Z, FVector Extent, const FLinearColor& Color, uint8 DepthPriority, float Thickness = 0.0f, float DepthBias = 0.0f, bool bScreenSpace = false);
 extern ENGINE_API void DrawDirectionalArrow(class FPrimitiveDrawInterface* PDI,const FMatrix& ArrowToWorld,const FLinearColor& InColor,float Length,float ArrowSize,uint8 DepthPriority);
 extern ENGINE_API void DrawConnectedArrow(class FPrimitiveDrawInterface* PDI, const FMatrix& ArrowToWorld, const FLinearColor& Color, float ArrowHeight, float ArrowWidth, uint8 DepthPriority, float Thickness = 0.5f, int32 NumSpokes = 6);
 extern ENGINE_API void DrawWireStar(class FPrimitiveDrawInterface* PDI,const FVector& Position, float Size, const FLinearColor& Color,uint8 DepthPriority);
