Remove editor only data from non editor build.

diff --git a/Engine/Source/Runtime/Engine/Classes/Matinee/InterpTrack.h b/Engine/Source/Runtime/Engine/Classes/Matinee/InterpTrack.h
--- a/Engine/Source/Runtime/Engine/Classes/Matinee/InterpTrack.h
+++ b/Engine/Source/Runtime/Engine/Classes/Matinee/InterpTrack.h
@@ -135,9 +135,10 @@ private:
 	uint32 bIsSelected:1;
 	
 public:
-	
+#if WITH_EDITORONLY_DATA
 	UPROPERTY(transient)
 	UTexture2D* TrackIcon;
+#endif // WITH_EDITORONLY_DATA
 
 	/** If true, the  AActor  this track is working on will have BeginAnimControl/FinishAnimControl called on it. */
 	UPROPERTY()
@@ -319,8 +320,10 @@ public:
 	 */
 	virtual const FString	GetSlateHelperClassName() const;
 
+#if WITH_EDITORONLY_DATA
 	/** @return	The icon to draw for this track in Matinee. */
 	virtual class UTexture2D* GetTrackIcon() const;
+#endif // WITH_EDITORONLY_DATA
 
 	/** @return  true if this track type works with static actors; false, otherwise. */
 	virtual bool AllowStaticActors() { return false; }
diff --git a/Engine/Source/Runtime/Engine/Classes/Matinee/InterpTrackAnimControl.h b/Engine/Source/Runtime/Engine/Classes/Matinee/InterpTrackAnimControl.h
--- a/Engine/Source/Runtime/Engine/Classes/Matinee/InterpTrackAnimControl.h
+++ b/Engine/Source/Runtime/Engine/Classes/Matinee/InterpTrackAnimControl.h
@@ -94,7 +94,9 @@ class UInterpTrackAnimControl : public UInterpTrackFloatBase
 	virtual void UpdateTrack(float NewPosition, UInterpTrackInst* TrInst, bool bJump) override;
 	virtual const FString	GetEdHelperClassName() const override;
 	virtual const FString	GetSlateHelperClassName() const override;
+#if WITH_EDITORONLY_DATA
 	virtual class UTexture2D* GetTrackIcon() const override;
+#endif // WITH_EDITORONLY_DATA
 	virtual void DrawTrack( FCanvas* Canvas, UInterpGroup* Group, const FInterpTrackDrawParams& Params ) override;
 	// End UInterpTrack interface.
 
diff --git a/Engine/Source/Runtime/Engine/Classes/Matinee/InterpTrackAudioMaster.h b/Engine/Source/Runtime/Engine/Classes/Matinee/InterpTrackAudioMaster.h
--- a/Engine/Source/Runtime/Engine/Classes/Matinee/InterpTrackAudioMaster.h
+++ b/Engine/Source/Runtime/Engine/Classes/Matinee/InterpTrackAudioMaster.h
@@ -17,7 +17,9 @@ class UInterpTrackAudioMaster : public UInterpTrackVectorBase
 	virtual void PreviewUpdateTrack(float NewPosition, UInterpTrackInst* TrInst) override;
 	virtual void UpdateTrack(float NewPosition, UInterpTrackInst* TrInst, bool bJump) override;
 	virtual void SetTrackToSensibleDefault() override;
+#if WITH_EDITORONLY_DATA
 	virtual class UTexture2D* GetTrackIcon() const override;
+#endif // WITH_EDITORONLY_DATA
 	// End UInterpTrack interface.
 
 	/** Return the sound volume scale for the specified time */
diff --git a/Engine/Source/Runtime/Engine/Classes/Matinee/InterpTrackBoolProp.h b/Engine/Source/Runtime/Engine/Classes/Matinee/InterpTrackBoolProp.h
--- a/Engine/Source/Runtime/Engine/Classes/Matinee/InterpTrackBoolProp.h
+++ b/Engine/Source/Runtime/Engine/Classes/Matinee/InterpTrackBoolProp.h
@@ -57,7 +57,9 @@ class UInterpTrackBoolProp : public UInterpTrack
 	virtual bool AllowStaticActors() override { return true; }
 	virtual const FString GetEdHelperClassName() const override;
 	virtual const FString GetSlateHelperClassName() const override;
+#if WITH_EDITORONLY_DATA
 	virtual class UTexture2D* GetTrackIcon() const override;
+#endif // WITH_EDITORONLY_DATA
 	// End UInterpTrack Interface
 };
 
diff --git a/Engine/Source/Runtime/Engine/Classes/Matinee/InterpTrackColorProp.h b/Engine/Source/Runtime/Engine/Classes/Matinee/InterpTrackColorProp.h
--- a/Engine/Source/Runtime/Engine/Classes/Matinee/InterpTrackColorProp.h
+++ b/Engine/Source/Runtime/Engine/Classes/Matinee/InterpTrackColorProp.h
@@ -23,7 +23,9 @@ class UInterpTrackColorProp : public UInterpTrackVectorBase
 	virtual void UpdateTrack(float NewPosition, UInterpTrackInst* TrInst, bool bJump) override;
 	virtual const FString	GetEdHelperClassName() const override;
 	virtual const FString	GetSlateHelperClassName() const override;
+#if WITH_EDITORONLY_DATA
 	virtual class UTexture2D* GetTrackIcon() const override;
+#endif // WITH_EDITORONLY_DATA
 	// End UInterpTrack interface
 };
 
diff --git a/Engine/Source/Runtime/Engine/Classes/Matinee/InterpTrackColorScale.h b/Engine/Source/Runtime/Engine/Classes/Matinee/InterpTrackColorScale.h
--- a/Engine/Source/Runtime/Engine/Classes/Matinee/InterpTrackColorScale.h
+++ b/Engine/Source/Runtime/Engine/Classes/Matinee/InterpTrackColorScale.h
@@ -21,7 +21,9 @@ class UInterpTrackColorScale : public UInterpTrackVectorBase
 	virtual void PreviewUpdateTrack(float NewPosition, UInterpTrackInst* TrInst) override;
 	virtual void UpdateTrack(float NewPosition, UInterpTrackInst* TrInst, bool bJump) override;
 	virtual void SetTrackToSensibleDefault() override;
+#if WITH_EDITORONLY_DATA
 	virtual class UTexture2D* GetTrackIcon() const override;
+#endif // WITH_EDITORONLY_DATA
 	// Begin UInterpTrack interface.
 
 	/** Return the blur alpha we want at the given time. */
diff --git a/Engine/Source/Runtime/Engine/Classes/Matinee/InterpTrackDirector.h b/Engine/Source/Runtime/Engine/Classes/Matinee/InterpTrackDirector.h
--- a/Engine/Source/Runtime/Engine/Classes/Matinee/InterpTrackDirector.h
+++ b/Engine/Source/Runtime/Engine/Classes/Matinee/InterpTrackDirector.h
@@ -83,7 +83,9 @@ public:
 	virtual void UpdateTrack( float NewPosition, UInterpTrackInst* TrackInst, bool bJump ) override;
 	virtual const FString GetEdHelperClassName() const override;
 	virtual const FString GetSlateHelperClassName() const override;
+#if WITH_EDITORONLY_DATA
 	virtual class UTexture2D* GetTrackIcon() const override;
+#endif // WITH_EDITORONLY_DATA
 	virtual void DrawTrack( FCanvas* Canvas, UInterpGroup* Group, const FInterpTrackDrawParams& Params ) override;
 	// End UInterpTrack Interface
 
diff --git a/Engine/Source/Runtime/Engine/Classes/Matinee/InterpTrackEvent.h b/Engine/Source/Runtime/Engine/Classes/Matinee/InterpTrackEvent.h
--- a/Engine/Source/Runtime/Engine/Classes/Matinee/InterpTrackEvent.h
+++ b/Engine/Source/Runtime/Engine/Classes/Matinee/InterpTrackEvent.h
@@ -67,7 +67,9 @@ class UInterpTrackEvent : public UInterpTrack
 	virtual void UpdateTrack( float NewPosition, UInterpTrackInst* TrackInst, bool bJump ) override;
 	virtual const FString GetEdHelperClassName() const override;
 	virtual const FString GetSlateHelperClassName() const override;
+#if WITH_EDITORONLY_DATA
 	virtual class UTexture2D* GetTrackIcon() const override;
+#endif // WITH_EDITORONLY_DATA
 	virtual bool AllowStaticActors() override { return true; }
 	virtual void DrawTrack( FCanvas* Canvas, UInterpGroup* Group, const FInterpTrackDrawParams& Params ) override;
 	// End UInterpTrack Interface
diff --git a/Engine/Source/Runtime/Engine/Classes/Matinee/InterpTrackFade.h b/Engine/Source/Runtime/Engine/Classes/Matinee/InterpTrackFade.h
--- a/Engine/Source/Runtime/Engine/Classes/Matinee/InterpTrackFade.h
+++ b/Engine/Source/Runtime/Engine/Classes/Matinee/InterpTrackFade.h
@@ -25,7 +25,9 @@ class UInterpTrackFade : public UInterpTrackFloatBase
 	virtual void UpdateKeyframe(int32 KeyIndex, UInterpTrackInst* TrInst) override;
 	virtual void PreviewUpdateTrack(float NewPosition, UInterpTrackInst* TrInst) override;
 	virtual void UpdateTrack(float NewPosition, UInterpTrackInst* TrInst, bool bJump) override;
+#if WITH_EDITORONLY_DATA
 	virtual class UTexture2D* GetTrackIcon() const override;
+#endif // WITH_EDITORONLY_DATA
 	// End UInterpTrack interface.
 
 	/** @return the amount of fading we want at the given time. */
diff --git a/Engine/Source/Runtime/Engine/Classes/Matinee/InterpTrackFloatProp.h b/Engine/Source/Runtime/Engine/Classes/Matinee/InterpTrackFloatProp.h
--- a/Engine/Source/Runtime/Engine/Classes/Matinee/InterpTrackFloatProp.h
+++ b/Engine/Source/Runtime/Engine/Classes/Matinee/InterpTrackFloatProp.h
@@ -22,7 +22,9 @@ class UInterpTrackFloatProp : public UInterpTrackFloatBase
 	virtual void UpdateTrack(float NewPosition, UInterpTrackInst* TrInst, bool bJump) override;
 	virtual const FString	GetEdHelperClassName() const override;
 	virtual const FString	GetSlateHelperClassName() const override;
+#if WITH_EDITORONLY_DATA
 	virtual class UTexture2D* GetTrackIcon() const override;
+#endif // WITH_EDITORONLY_DATA
 	virtual void ReduceKeys( float IntervalStart, float IntervalEnd, float Tolerance ) override;
 	// End InterpTrack interface.
 };
diff --git a/Engine/Source/Runtime/Engine/Classes/Matinee/InterpTrackLinearColorProp.h b/Engine/Source/Runtime/Engine/Classes/Matinee/InterpTrackLinearColorProp.h
--- a/Engine/Source/Runtime/Engine/Classes/Matinee/InterpTrackLinearColorProp.h
+++ b/Engine/Source/Runtime/Engine/Classes/Matinee/InterpTrackLinearColorProp.h
@@ -23,7 +23,9 @@ class UInterpTrackLinearColorProp : public UInterpTrackLinearColorBase
 	virtual void UpdateTrack(float NewPosition, UInterpTrackInst* TrInst, bool bJump) override;
 	virtual const FString	GetEdHelperClassName() const override;
 	virtual const FString	GetSlateHelperClassName() const override;
+#if WITH_EDITORONLY_DATA
 	virtual class UTexture2D* GetTrackIcon() const override;
+#endif // WITH_EDITORONLY_DATA
 	// End UInterpTrack interface.
 };
 
diff --git a/Engine/Source/Runtime/Engine/Classes/Matinee/InterpTrackMove.h b/Engine/Source/Runtime/Engine/Classes/Matinee/InterpTrackMove.h
--- a/Engine/Source/Runtime/Engine/Classes/Matinee/InterpTrackMove.h
+++ b/Engine/Source/Runtime/Engine/Classes/Matinee/InterpTrackMove.h
@@ -174,7 +174,9 @@ class UInterpTrackMove : public UInterpTrack
 	virtual void ConditionalPreviewUpdateTrack(float NewPosition, class UInterpTrackInst* TrInst) override;
 	virtual void PreviewUpdateTrack(float NewPosition, UInterpTrackInst* TrInst) override;
 	virtual void UpdateTrack(float NewPosition, UInterpTrackInst* TrInst, bool bJump) override;
+#if WITH_EDITORONLY_DATA
 	virtual class UTexture2D* GetTrackIcon() const override;
+#endif // WITH_EDITORONLY_DATA
 	virtual FColor GetKeyframeColor(int32 KeyIndex) const override;
 	virtual void DrawTrack( FCanvas* Canvas, UInterpGroup* Group, const FInterpTrackDrawParams& Params ) override;
 	virtual void Render3DTrack(UInterpTrackInst* TrInst, const FSceneView* View, class FPrimitiveDrawInterface* PDI, int32 TrackIndex, const FColor& TrackColor, TArray<struct FInterpEdSelKey>& SelectedKeys) override;
diff --git a/Engine/Source/Runtime/Engine/Classes/Matinee/InterpTrackMoveAxis.h b/Engine/Source/Runtime/Engine/Classes/Matinee/InterpTrackMoveAxis.h
--- a/Engine/Source/Runtime/Engine/Classes/Matinee/InterpTrackMoveAxis.h
+++ b/Engine/Source/Runtime/Engine/Classes/Matinee/InterpTrackMoveAxis.h
@@ -43,7 +43,9 @@ class UInterpTrackMoveAxis : public UInterpTrackFloatBase
 	virtual int32 SetKeyframeTime( int32 KeyIndex, float NewKeyTime, bool bUpdateOrder ) override;
 	virtual void RemoveKeyframe( int32 KeyIndex ) override;
 	virtual int32 DuplicateKeyframe( int32 KeyIndex, float NewKeyTime, UInterpTrack* ToTrack = NULL ) override;
+#if WITH_EDITORONLY_DATA
 	virtual class UTexture2D* GetTrackIcon() const override;
+#endif // WITH_EDITORONLY_DATA
 	virtual void ReduceKeys( float IntervalStart, float IntervalEnd, float Tolerance ) override;
 	// End UInterpTrack Interface
 	
diff --git a/Engine/Source/Runtime/Engine/Classes/Matinee/InterpTrackParticleReplay.h b/Engine/Source/Runtime/Engine/Classes/Matinee/InterpTrackParticleReplay.h
--- a/Engine/Source/Runtime/Engine/Classes/Matinee/InterpTrackParticleReplay.h
+++ b/Engine/Source/Runtime/Engine/Classes/Matinee/InterpTrackParticleReplay.h
@@ -71,7 +71,9 @@ class UInterpTrackParticleReplay : public UInterpTrack
 	virtual void UpdateTrack(float NewPosition, UInterpTrackInst* TrInst, bool bJump) override;
 	virtual const FString	GetEdHelperClassName() const override;
 	virtual const FString	GetSlateHelperClassName() const override;
+#if WITH_EDITORONLY_DATA
 	virtual class UTexture2D* GetTrackIcon() const override;
+#endif // WITH_EDITORONLY_DATA
 	virtual bool AllowStaticActors() override { return true; }
 	virtual void DrawTrack( FCanvas* Canvas, UInterpGroup* Group, const FInterpTrackDrawParams& Params );
 	// End UInterpTrack interface.
diff --git a/Engine/Source/Runtime/Engine/Classes/Matinee/InterpTrackSlomo.h b/Engine/Source/Runtime/Engine/Classes/Matinee/InterpTrackSlomo.h
--- a/Engine/Source/Runtime/Engine/Classes/Matinee/InterpTrackSlomo.h
+++ b/Engine/Source/Runtime/Engine/Classes/Matinee/InterpTrackSlomo.h
@@ -16,7 +16,9 @@ class UInterpTrackSlomo : public UInterpTrackFloatBase
 	virtual void PreviewUpdateTrack(float NewPosition, UInterpTrackInst* TrInst) override;
 	virtual void UpdateTrack(float NewPosition, UInterpTrackInst* TrInst, bool bJump) override;
 	virtual void SetTrackToSensibleDefault() override;
+#if WITH_EDITORONLY_DATA
 	virtual class UTexture2D* GetTrackIcon() const override;
+#endif // WITH_EDITORONLY_DATA
 	// End InterpTrack interface.
 
 	/** @return the slomo factor we want at the given time. */
diff --git a/Engine/Source/Runtime/Engine/Classes/Matinee/InterpTrackSound.h b/Engine/Source/Runtime/Engine/Classes/Matinee/InterpTrackSound.h
--- a/Engine/Source/Runtime/Engine/Classes/Matinee/InterpTrackSound.h
+++ b/Engine/Source/Runtime/Engine/Classes/Matinee/InterpTrackSound.h
@@ -90,7 +90,9 @@ class UInterpTrackSound : public UInterpTrackVectorBase
 	virtual void PreviewStopPlayback(class UInterpTrackInst* TrInst) override;
 	virtual const FString	GetEdHelperClassName() const override;
 	virtual const FString	GetSlateHelperClassName() const override;
+#if WITH_EDITORONLY_DATA
 	virtual class UTexture2D* GetTrackIcon() const override;
+#endif // WITH_EDITORONLY_DATA
 	virtual void DrawTrack( FCanvas* Canvas, UInterpGroup* Group, const FInterpTrackDrawParams& Params ) override;
 	virtual bool AllowStaticActors() override { return true; }
 	virtual void SetTrackToSensibleDefault() override;
diff --git a/Engine/Source/Runtime/Engine/Classes/Matinee/InterpTrackToggle.h b/Engine/Source/Runtime/Engine/Classes/Matinee/InterpTrackToggle.h
--- a/Engine/Source/Runtime/Engine/Classes/Matinee/InterpTrackToggle.h
+++ b/Engine/Source/Runtime/Engine/Classes/Matinee/InterpTrackToggle.h
@@ -93,7 +93,9 @@ class UInterpTrackToggle : public UInterpTrack
 	virtual void UpdateTrack(float NewPosition, UInterpTrackInst* TrInst, bool bJump) override;
 	virtual const FString	GetEdHelperClassName() const override;
 	virtual const FString	GetSlateHelperClassName() const override;
+#if WITH_EDITORONLY_DATA
 	virtual class UTexture2D* GetTrackIcon() const override;
+#endif // WITH_EDITORONLY_DATA
 	virtual bool AllowStaticActors() override { return true; }
 	virtual void DrawTrack( FCanvas* Canvas, UInterpGroup* Group, const FInterpTrackDrawParams& Params ) override;
 	// Begin UInterpTrack interface.
diff --git a/Engine/Source/Runtime/Engine/Classes/Matinee/InterpTrackVectorProp.h b/Engine/Source/Runtime/Engine/Classes/Matinee/InterpTrackVectorProp.h
--- a/Engine/Source/Runtime/Engine/Classes/Matinee/InterpTrackVectorProp.h
+++ b/Engine/Source/Runtime/Engine/Classes/Matinee/InterpTrackVectorProp.h
@@ -23,7 +23,9 @@ class UInterpTrackVectorProp : public UInterpTrackVectorBase
 	virtual void UpdateTrack(float NewPosition, UInterpTrackInst* TrInst, bool bJump) override;
 	virtual const FString	GetEdHelperClassName() const override;
 	virtual const FString	GetSlateHelperClassName() const override;
+#if WITH_EDITORONLY_DATA
 	virtual class UTexture2D* GetTrackIcon() const override;
+#endif // WITH_EDITORONLY_DATA
 	// End InterpTrack interface.
 };
 
diff --git a/Engine/Source/Runtime/Engine/Classes/Matinee/InterpTrackVisibility.h b/Engine/Source/Runtime/Engine/Classes/Matinee/InterpTrackVisibility.h
--- a/Engine/Source/Runtime/Engine/Classes/Matinee/InterpTrackVisibility.h
+++ b/Engine/Source/Runtime/Engine/Classes/Matinee/InterpTrackVisibility.h
@@ -98,7 +98,9 @@ class UInterpTrackVisibility : public UInterpTrack
 	virtual void UpdateTrack(float NewPosition, UInterpTrackInst* TrInst, bool bJump) override;
 	virtual const FString	GetEdHelperClassName() const override;
 	virtual const FString	GetSlateHelperClassName() const override;
+#if WITH_EDITORONLY_DATA
 	virtual class UTexture2D* GetTrackIcon() const override;
+#endif // WITH_EDITORONLY_DATA
 	virtual bool AllowStaticActors() override { return true; }
 	virtual void DrawTrack( FCanvas* Canvas, UInterpGroup* Group, const FInterpTrackDrawParams& Params ) override;
 	// End UInterpTrack interface.
diff --git a/Engine/Source/Runtime/Engine/Private/Interpolation.cpp b/Engine/Source/Runtime/Engine/Private/Interpolation.cpp
--- a/Engine/Source/Runtime/Engine/Private/Interpolation.cpp
+++ b/Engine/Source/Runtime/Engine/Private/Interpolation.cpp
@@ -2817,8 +2817,8 @@ UInterpTrack::UInterpTrack(const FObjectInitializer& ObjectInitializer)
 	bVisible = true;
 	SetSelected( false );
 	bIsRecording = false;
-	TrackIcon = Cast<UTexture2D>(StaticLoadObject( UTexture2D::StaticClass(), NULL, TEXT("/Engine/EditorMaterials/MatineeGroups/MAT_Groups_Float.MAT_Groups_Float"), NULL, LOAD_None, NULL ));
 #if WITH_EDITORONLY_DATA
+	TrackIcon = Cast<UTexture2D>(StaticLoadObject( UTexture2D::StaticClass(), NULL, TEXT("/Engine/EditorMaterials/MatineeGroups/MAT_Groups_Float.MAT_Groups_Float"), NULL, LOAD_None, NULL ));
 	bIsCollapsed = false;
 #endif // WITH_EDITORONLY_DATA
 }
@@ -3316,7 +3316,9 @@ UInterpTrackMove::UInterpTrackMove(const FObjectInitializer& ObjectInitializer)
 	RotMode = IMR_Keyframed;
 	bShowTranslationOnCurveEd = true;
 	bShowRotationOnCurveEd = false;
+#if WITH_EDITORONLY_DATA
 	TrackIcon = Cast<UTexture2D>(StaticLoadObject( UTexture2D::StaticClass(), NULL, TEXT("/Engine/EditorMaterials/MatineeGroups/MAT_Groups_Move.MAT_Groups_Move"), NULL, LOAD_None, NULL ));
+#endif // WITH_EDITORONLY_DATA
 
 #if WITH_EDITORONLY_DATA
 	int32 NewArrayIndex0 = SupportedSubTracks.Add(FSupportedSubTrackInfo());
@@ -4980,7 +4982,9 @@ UInterpTrackToggle::UInterpTrackToggle(const FObjectInitializer& ObjectInitializ
 	bFireEventsWhenBackwards = true;
 	bFireEventsWhenJumpingForwards = true;
 
+#if WITH_EDITORONLY_DATA
 	TrackIcon = Cast<UTexture2D>(StaticLoadObject( UTexture2D::StaticClass(), NULL, TEXT("/Engine/EditorMaterials/MAT_Groups_Toggle.MAT_Groups_Toggle"), NULL, LOAD_None, NULL ));
+#endif // WITH_EDITORONLY_DATA
 }
 
 int32 UInterpTrackToggle::AddKeyframe(float Time, UInterpTrackInst* TrInst, EInterpCurveMode InitInterpMode)
@@ -5439,7 +5443,9 @@ UInterpTrackFloatProp::UInterpTrackFloatProp(const FObjectInitializer& ObjectIni
 {
 	TrackInstClass = UInterpTrackInstFloatProp::StaticClass();
 	TrackTitle = TEXT("Float Property");
+#if WITH_EDITORONLY_DATA
 	TrackIcon = Cast<UTexture2D>(StaticLoadObject( UTexture2D::StaticClass(), NULL, TEXT("/Engine/EditorMaterials/MatineeGroups/MAT_Groups_Float.MAT_Groups_Float"), NULL, LOAD_None, NULL ));
+#endif // WITH_EDITORONLY_DATA
 }
 
 int32 UInterpTrackFloatProp::AddKeyframe(float Time, UInterpTrackInst* TrInst, EInterpCurveMode InitInterpMode)
@@ -5599,10 +5605,11 @@ void UInterpTrackInstFloatProp::InitTrackInst(UInterpTrack* Track)
 UInterpTrackVectorProp::UInterpTrackVectorProp(const FObjectInitializer& ObjectInitializer)
 	: Super(ObjectInitializer)
 {
-
 	TrackInstClass = UInterpTrackInstVectorProp::StaticClass();
 	TrackTitle = TEXT("Vector Property");
+#if WITH_EDITORONLY_DATA
 	TrackIcon = Cast<UTexture2D>(StaticLoadObject( UTexture2D::StaticClass(), NULL, TEXT("/Engine/EditorMaterials/MatineeGroups/MAT_Groups_Vector.MAT_Groups_Vector"), NULL, LOAD_None, NULL ));
+#endif // WITH_EDITORONLY_DATA
 }
 
 int32 UInterpTrackVectorProp::AddKeyframe(float Time, UInterpTrackInst* TrInst, EInterpCurveMode InitInterpMode)
@@ -5752,7 +5759,9 @@ UInterpTrackBoolProp::UInterpTrackBoolProp(const FObjectInitializer& ObjectIniti
 {
 	TrackInstClass = UInterpTrackInstBoolProp::StaticClass();
 	TrackTitle = TEXT("Bool Property");
+#if WITH_EDITORONLY_DATA
 	TrackIcon = Cast<UTexture2D>(StaticLoadObject( UTexture2D::StaticClass(), NULL, TEXT("/Engine/EditorMaterials/MatineeGroups/MAT_Groups_Float.MAT_Groups_Float"), NULL, LOAD_None, NULL ));
+#endif // WITH_EDITORONLY_DATA
 }
 
 int32 UInterpTrackBoolProp::AddKeyframe( float Time, UInterpTrackInst* TrackInst, EInterpCurveMode InitInterpMode )
@@ -5928,7 +5937,9 @@ UInterpTrackColorProp::UInterpTrackColorProp(const FObjectInitializer& ObjectIni
 
 	TrackInstClass = UInterpTrackInstColorProp::StaticClass();
 	TrackTitle = TEXT("Color Property");
+#if WITH_EDITORONLY_DATA
 	TrackIcon = Cast<UTexture2D>(StaticLoadObject( UTexture2D::StaticClass(), NULL, TEXT("/Engine/EditorMaterials/MatineeGroups/MAT_ColorTrack.MAT_ColorTrack"), NULL, LOAD_None, NULL ));
+#endif // WITH_EDITORONLY_DATA
 }
 
 int32 UInterpTrackColorProp::AddKeyframe(float Time, UInterpTrackInst* TrInst, EInterpCurveMode InitInterpMode)
@@ -6084,7 +6095,9 @@ UInterpTrackLinearColorProp::UInterpTrackLinearColorProp(const FObjectInitialize
 {
 	TrackInstClass = UInterpTrackInstLinearColorProp::StaticClass();
 	TrackTitle = TEXT("LinearColor Property");
+#if WITH_EDITORONLY_DATA
 	TrackIcon = Cast<UTexture2D>(StaticLoadObject( UTexture2D::StaticClass(), NULL, TEXT("/Engine/EditorMaterials/MatineeGroups/MAT_ColorTrack.MAT_ColorTrack"), NULL, LOAD_None, NULL ));
+#endif // WITH_EDITORONLY_DATA
 }
 
 int32 UInterpTrackLinearColorProp::AddKeyframe(float Time, UInterpTrackInst* TrInst, EInterpCurveMode InitInterpMode)
@@ -6241,7 +6254,9 @@ UInterpTrackEvent::UInterpTrackEvent(const FObjectInitializer& ObjectInitializer
 	TrackTitle = TEXT("Event");
 	bFireEventsWhenForwards = true;
 	bFireEventsWhenBackwards = true;
+#if WITH_EDITORONLY_DATA
 	TrackIcon = Cast<UTexture2D>(StaticLoadObject( UTexture2D::StaticClass(), NULL, TEXT("/Engine/EditorMaterials/MatineeGroups/MAT_Groups_Event.MAT_Groups_Event"), NULL, LOAD_None, NULL ));
+#endif // WITH_EDITORONLY_DATA
 }
 
 int32 UInterpTrackEvent::AddKeyframe(float Time, UInterpTrackInst* TrInst, EInterpCurveMode InitInterpMode)
@@ -6430,8 +6445,8 @@ UInterpTrackDirector::UInterpTrackDirector(const FObjectInitializer& ObjectIniti
 	TrackInstClass = UInterpTrackInstDirector::StaticClass();
 	TrackTitle = TEXT("Director");
 	bSimulateCameraCutsOnClients = true;
-	TrackIcon = Cast<UTexture2D>(StaticLoadObject( UTexture2D::StaticClass(), NULL, TEXT("/Engine/EditorMaterials/MatineeGroups/MAT_Groups_Director.MAT_Groups_Director"), NULL, LOAD_None, NULL ));
 #if WITH_EDITORONLY_DATA
+	TrackIcon = Cast<UTexture2D>(StaticLoadObject( UTexture2D::StaticClass(), NULL, TEXT("/Engine/EditorMaterials/MatineeGroups/MAT_Groups_Director.MAT_Groups_Director"), NULL, LOAD_None, NULL ));
 	PreviewCamera = NULL;
 #endif // WITH_EDITORONLY_DATA
 }
@@ -6765,7 +6780,9 @@ UInterpTrackFade::UInterpTrackFade(const FObjectInitializer& ObjectInitializer)
 	bDirGroupOnly = true;
 	TrackInstClass = UInterpTrackInstFade::StaticClass();
 	TrackTitle = TEXT("Fade");
+#if WITH_EDITORONLY_DATA
 	TrackIcon = Cast<UTexture2D>(StaticLoadObject( UTexture2D::StaticClass(), NULL, TEXT("/Engine/EditorMaterials/MatineeGroups/MAT_Groups_Fade.MAT_Groups_Fade"), NULL, LOAD_None, NULL ));
+#endif // WITH_EDITORONLY_DATA
 }
 
 int32 UInterpTrackFade::AddKeyframe(float Time, UInterpTrackInst* TrInst, EInterpCurveMode InitInterpMode)
@@ -6856,7 +6873,9 @@ UInterpTrackSlomo::UInterpTrackSlomo(const FObjectInitializer& ObjectInitializer
 	bDirGroupOnly = true;
 	TrackInstClass = UInterpTrackInstSlomo::StaticClass();
 	TrackTitle = TEXT("Slomo");
+#if WITH_EDITORONLY_DATA
 	TrackIcon = Cast<UTexture2D>(StaticLoadObject( UTexture2D::StaticClass(), NULL, TEXT("/Engine/EditorMaterials/MatineeGroups/MAT_Groups_Slomo.MAT_Groups_Slomo"), NULL, LOAD_None, NULL ));
+#endif // WITH_EDITORONLY_DATA
 }
 
 
@@ -6988,7 +7007,9 @@ UInterpTrackAnimControl::UInterpTrackAnimControl(const FObjectInitializer& Objec
 	TrackInstClass = UInterpTrackInstAnimControl::StaticClass();
 	TrackTitle = TEXT("Anim");
 	bIsAnimControlTrack = true;
+#if WITH_EDITORONLY_DATA
 	TrackIcon = Cast<UTexture2D>(StaticLoadObject( UTexture2D::StaticClass(), NULL, TEXT("/Engine/EditorMaterials/MatineeGroups/MAT_Groups_Anim.MAT_Groups_Anim"), NULL, LOAD_None, NULL ));
+#endif // WITH_EDITORONLY_DATA
 }
 
 void UInterpTrackAnimControl::PostLoad()
@@ -7589,7 +7610,9 @@ UInterpTrackSound::UInterpTrackSound(const FObjectInitializer& ObjectInitializer
 
 	TrackInstClass = UInterpTrackInstSound::StaticClass();
 	TrackTitle = TEXT("Sound");
+#if WITH_EDITORONLY_DATA
 	TrackIcon = Cast<UTexture2D>(StaticLoadObject( UTexture2D::StaticClass(), NULL, TEXT("/Engine/EditorMaterials/MatineeGroups/MAT_Groups_Sound.MAT_Groups_Sound"), NULL, LOAD_None, NULL ));
+#endif // WITH_EDITORONLY_DATA
 
 	bAttach = true;
 }
@@ -8813,7 +8836,9 @@ UInterpTrackColorScale::UInterpTrackColorScale(const FObjectInitializer& ObjectI
 	bDirGroupOnly = true;
 	TrackInstClass = UInterpTrackInstColorScale::StaticClass();
 	TrackTitle = TEXT("Color Scale");
+#if WITH_EDITORONLY_DATA
 	TrackIcon = Cast<UTexture2D>(StaticLoadObject( UTexture2D::StaticClass(), NULL, TEXT("/Engine/EditorMaterials/MatineeGroups/MAT_Groups_Fade.MAT_Groups_Fade"), NULL, LOAD_None, NULL ));
+#endif // WITH_EDITORONLY_DATA
 }
 
 int32 UInterpTrackColorScale::AddKeyframe(float Time, UInterpTrackInst* TrInst, EInterpCurveMode InitInterpMode)
@@ -8897,7 +8922,9 @@ UInterpTrackAudioMaster::UInterpTrackAudioMaster(const FObjectInitializer& Objec
 	bDirGroupOnly = true;
 	TrackInstClass = UInterpTrackInstAudioMaster::StaticClass();
 	TrackTitle = TEXT("Audio Master");
+#if WITH_EDITORONLY_DATA
 	TrackIcon = Cast<UTexture2D>(StaticLoadObject( UTexture2D::StaticClass(), NULL, TEXT("/Engine/EditorMaterials/MatineeGroups/MAT_Groups_AudioMaster.MAT_Groups_AudioMaster"), NULL, LOAD_None, NULL ));
+#endif // WITH_EDITORONLY_DATA
 }
 
 
@@ -8995,7 +9022,9 @@ UInterpTrackVisibility::UInterpTrackVisibility(const FObjectInitializer& ObjectI
 	bFireEventsWhenForwards = true;
 	bFireEventsWhenBackwards = true;
 	bFireEventsWhenJumpingForwards = true;
+#if WITH_EDITORONLY_DATA
 	TrackIcon = Cast<UTexture2D>(StaticLoadObject( UTexture2D::StaticClass(), NULL, TEXT("/Engine/EditorMaterials/MAT_Groups_Visibility.MAT_Groups_Visibility"), NULL, LOAD_None, NULL ));
+#endif // WITH_EDITORONLY_DATA
 }
 
 int32 UInterpTrackVisibility::AddKeyframe(float Time, UInterpTrackInst* TrInst, EInterpCurveMode InitInterpMode)
diff --git a/Engine/Source/Runtime/Engine/Private/InterpolationCurveEd.cpp b/Engine/Source/Runtime/Engine/Private/InterpolationCurveEd.cpp
--- a/Engine/Source/Runtime/Engine/Private/InterpolationCurveEd.cpp
+++ b/Engine/Source/Runtime/Engine/Private/InterpolationCurveEd.cpp
@@ -539,10 +539,12 @@ int32 UInterpTrackMoveAxis::SetKeyIn( int32 KeyIndex, float NewInVal )
 	return NewIndex;
 }
 
+#if WITH_EDITORONLY_DATA
 UTexture2D* UInterpTrackMoveAxis::GetTrackIcon() const
 {
 	return NULL;
 }
+#endif // WITH_EDITORONLY_DATA
 
 /*-----------------------------------------------------------------------------
 	UInterpTrackFloatBase
diff --git a/Engine/Source/Runtime/Engine/Private/InterpolationDraw.cpp b/Engine/Source/Runtime/Engine/Private/InterpolationDraw.cpp
--- a/Engine/Source/Runtime/Engine/Private/InterpolationDraw.cpp
+++ b/Engine/Source/Runtime/Engine/Private/InterpolationDraw.cpp
@@ -119,10 +119,12 @@ FColor UInterpTrack::GetKeyframeColor(int32 KeyIndex) const
 	return KeyNormalColor;
 }
 
+#if WITH_EDITORONLY_DATA
 UTexture2D* UInterpTrack::GetTrackIcon() const
 {
 	return TrackIcon;
 }
+#endif // WITH_EDITORONLY_DATA
 
 /*-----------------------------------------------------------------------------
   UInterpTrackMove
@@ -548,10 +550,12 @@ void UInterpTrackMove::Render3DTrack(UInterpTrackInst* TrInst,
 #endif // WITH_EDITORONLY_DATA
 }
 
+#if WITH_EDITORONLY_DATA
 UTexture2D* UInterpTrackMove::GetTrackIcon() const
 {
 	return TrackIcon;
 }
+#endif // WITH_EDITORONLY_DATA
 
 /*-----------------------------------------------------------------------------
 	UInterpTrackFloatBase
@@ -582,29 +586,35 @@ FColor UInterpTrackFloatBase::GetKeyframeColor(int32 KeyIndex) const
 	UInterpTrackFloatProp
 -----------------------------------------------------------------------------*/
 
+#if WITH_EDITORONLY_DATA
 UTexture2D* UInterpTrackFloatProp::GetTrackIcon() const
 {
 	return TrackIcon;
 }
+#endif // WITH_EDITORONLY_DATA
 
 /*-----------------------------------------------------------------------------
 	UInterpTrackBoolProp
 -----------------------------------------------------------------------------*/
 
+#if WITH_EDITORONLY_DATA
 UTexture2D* UInterpTrackBoolProp::GetTrackIcon() const
 {
 	return TrackIcon;
 }
+#endif // WITH_EDITORONLY_DATA
 
 
 /*-----------------------------------------------------------------------------
 	UInterpTrackToggle
 -----------------------------------------------------------------------------*/
 
+#if WITH_EDITORONLY_DATA
 UTexture2D* UInterpTrackToggle::GetTrackIcon() const
 {
 	return TrackIcon;
 }
+#endif // WITH_EDITORONLY_DATA
 
 void UInterpTrackToggle::DrawTrack( FCanvas* Canvas, UInterpGroup* Group, const FInterpTrackDrawParams& Params )
 {
@@ -796,29 +806,35 @@ FColor UInterpTrackLinearColorBase::GetKeyframeColor(int32 KeyIndex) const
 	UInterpTrackVectorProp
 -----------------------------------------------------------------------------*/
 
+#if WITH_EDITORONLY_DATA
 UTexture2D* UInterpTrackVectorProp::GetTrackIcon() const
 {
 	return TrackIcon;
 }
+#endif // WITH_EDITORONLY_DATA
 
 /*-----------------------------------------------------------------------------
 	UInterpTrackColorProp
 -----------------------------------------------------------------------------*/
 
+#if WITH_EDITORONLY_DATA
 UTexture2D* UInterpTrackColorProp::GetTrackIcon() const
 {
 	return TrackIcon;
 }
+#endif // WITH_EDITORONLY_DATA
 
 
 /*-----------------------------------------------------------------------------
 	UInterpTrackLinearColorProp
 -----------------------------------------------------------------------------*/
 
+#if WITH_EDITORONLY_DATA
 UTexture2D* UInterpTrackLinearColorProp::GetTrackIcon() const
 {
 	return TrackIcon;
 }
+#endif // WITH_EDITORONLY_DATA
 
 
 /*-----------------------------------------------------------------------------
@@ -856,10 +872,12 @@ void UInterpTrackEvent::DrawTrack( FCanvas* Canvas, UInterpGroup* Group, const F
 #endif // WITH_EDITORONLY_DATA
 }
 
+#if WITH_EDITORONLY_DATA
 UTexture2D* UInterpTrackEvent::GetTrackIcon() const
 {
 	return TrackIcon;
 }
+#endif // WITH_EDITORONLY_DATA
 
 /*-----------------------------------------------------------------------------
 	UInterpTrackDirector
@@ -940,10 +958,12 @@ void UInterpTrackDirector::DrawTrack( FCanvas* Canvas, UInterpGroup* Group, cons
 #endif // WITH_EDITORONLY_DATA
 }
 
+#if WITH_EDITORONLY_DATA
 UTexture2D* UInterpTrackDirector::GetTrackIcon() const
 {
 	return TrackIcon; 
 }
+#endif // WITH_EDITORONLY_DATA
 
 /*-----------------------------------------------------------------------------
 	UInterpTrackAnimControl
@@ -1329,10 +1349,12 @@ void UInterpTrackAnimControl::DrawTrack( FCanvas* Canvas, UInterpGroup* Group, c
 #endif // WITH_EDITORONLY_DATA
 }
 
+#if WITH_EDITORONLY_DATA
 UTexture2D* UInterpTrackAnimControl::GetTrackIcon() const
 {
 	return TrackIcon;
 }
+#endif // WITH_EDITORONLY_DATA
 
 /*-----------------------------------------------------------------------------
 	UInterpTrackSound
@@ -1449,58 +1471,70 @@ void UInterpTrackSound::DrawTrack( FCanvas* Canvas, UInterpGroup* Group, const F
 #endif // WITH_EDITORONLY_DATA
 }
 
+#if WITH_EDITORONLY_DATA
 UTexture2D* UInterpTrackSound::GetTrackIcon() const
 {
 	return TrackIcon;
 }
+#endif // WITH_EDITORONLY_DATA
 
 /*-----------------------------------------------------------------------------
 	UInterpTrackFade
 -----------------------------------------------------------------------------*/
 
+#if WITH_EDITORONLY_DATA
 UTexture2D* UInterpTrackFade::GetTrackIcon() const
 {
 	return TrackIcon;
 }
+#endif // WITH_EDITORONLY_DATA
 
 
 /*-----------------------------------------------------------------------------
 	UInterpTrackSlomo
 -----------------------------------------------------------------------------*/
 
+#if WITH_EDITORONLY_DATA
 UTexture2D* UInterpTrackSlomo::GetTrackIcon() const
 {
 	return TrackIcon;
 }
+#endif // WITH_EDITORONLY_DATA
 
 /*-----------------------------------------------------------------------------
 	UInterpTrackColorScale
 -----------------------------------------------------------------------------*/
 
+#if WITH_EDITORONLY_DATA
 UTexture2D* UInterpTrackColorScale::GetTrackIcon() const
 {
 	return TrackIcon;
 }
+#endif // WITH_EDITORONLY_DATA
 
 
 /*-----------------------------------------------------------------------------
 	UInterpTrackAudioMaster
 -----------------------------------------------------------------------------*/
 
+#if WITH_EDITORONLY_DATA
 UTexture2D* UInterpTrackAudioMaster::GetTrackIcon() const
 {
 	return TrackIcon;
 }
+#endif // WITH_EDITORONLY_DATA
 
 
 
 /*-----------------------------------------------------------------------------
 	UInterpTrackVisibility
 -----------------------------------------------------------------------------*/
+#if WITH_EDITORONLY_DATA
 UTexture2D* UInterpTrackVisibility::GetTrackIcon() const
 {
 	return TrackIcon;
 }
+#endif // WITH_EDITORONLY_DATA
 
 void UInterpTrackVisibility::DrawTrack( FCanvas* Canvas, UInterpGroup* Group, const FInterpTrackDrawParams& Params )
 {
@@ -1635,10 +1669,12 @@ void UInterpTrackVisibility::DrawTrack( FCanvas* Canvas, UInterpGroup* Group, co
 /*-----------------------------------------------------------------------------
 	UInterpTrackParticleReplay
 -----------------------------------------------------------------------------*/
+#if WITH_EDITORONLY_DATA
 UTexture2D* UInterpTrackParticleReplay::GetTrackIcon() const
 {
 	return (UTexture2D*)StaticLoadObject( UTexture2D::StaticClass(), NULL, TEXT("/Engine/EditorMaterials/MatineeGroups/MAT_Groups_ParticleReplay.MAT_Groups_ParticleReplay"), NULL, LOAD_None, NULL );
 }
+#endif // WITH_EDITORONLY_DATA
 
 
 void UInterpTrackParticleReplay::BeginDrag(FInterpEdInputData &InputData)
