Bring minimum engine version up to 214

diff --git a/Engine/Source/Runtime/Engine/Classes/Components/AudioComponent.h b/Engine/Source/Runtime/Engine/Classes/Components/AudioComponent.h
--- a/Engine/Source/Runtime/Engine/Classes/Components/AudioComponent.h
+++ b/Engine/Source/Runtime/Engine/Classes/Components/AudioComponent.h
@@ -39,9 +39,6 @@ struct FAudioComponentParam
 	UPROPERTY(EditAnywhere, BlueprintReadWrite, Category=AudioComponentParam)
 	int32 IntParam;
 
-	UPROPERTY()
-	class UDEPRECATED_SoundNodeWave* WaveParam_DEPRECATED;
-
 	// Value of the parameter when used as a sound wave
 	UPROPERTY(EditAnywhere, BlueprintReadWrite, Category=AudioComponentParam)
 	class USoundWave* SoundWaveParam;
@@ -51,7 +48,6 @@ struct FAudioComponentParam
 		: FloatParam(0)
 		, BoolParam(false)
 		, IntParam(0)
-		, WaveParam_DEPRECATED(NULL)
 		, SoundWaveParam(NULL)
 	{
 	}
@@ -81,10 +77,6 @@ class ENGINE_API UAudioComponent : public USceneComponent
 	UPROPERTY(EditAnywhere, Category=Sound, AdvancedDisplay)
 	USoundClass* SoundClassOverride;
 
-	/** Auto start this component on creation */
-	UPROPERTY()
-	uint32 bAutoPlay_DEPRECATED:1;
-
 	/** Auto destroy this component on completion */
 	UPROPERTY()
 	uint32 bAutoDestroy:1;
@@ -275,7 +267,6 @@ public:
 #if WITH_EDITOR
 	virtual void PostEditChangeProperty(FPropertyChangedEvent& PropertyChangedEvent) override;
 #endif // WITH_EDITOR
-	virtual void PostLoad() override;
 	virtual FString GetDetailedInfoInternal() const override;
 	// End UObject interface.
 
diff --git a/Engine/Source/Runtime/Engine/Classes/Components/ExponentialHeightFogComponent.h b/Engine/Source/Runtime/Engine/Classes/Components/ExponentialHeightFogComponent.h
--- a/Engine/Source/Runtime/Engine/Classes/Components/ExponentialHeightFogComponent.h
+++ b/Engine/Source/Runtime/Engine/Classes/Components/ExponentialHeightFogComponent.h
@@ -12,10 +12,6 @@ class ENGINE_API UExponentialHeightFogComponent : public USceneComponent
 {
 	GENERATED_UCLASS_BODY()
 
-	/** True if the fog is enabled. */
-	UPROPERTY()
-	uint32 bEnabled_DEPRECATED:1;
-
 	/** Global density factor. */
 	UPROPERTY(BlueprintReadOnly, interp, Category=ExponentialHeightFogComponent, meta=(UIMin = "0", UIMax = ".05"))
 	float FogDensity;
@@ -100,7 +96,6 @@ protected:
 
 public:
 	// Begin UObject Interface
-	virtual void PostLoad() override; 
 #if WITH_EDITOR
 	virtual void PostEditChangeProperty(FPropertyChangedEvent& PropertyChangedEvent) override;
 #endif // WITH_EDITOR
diff --git a/Engine/Source/Runtime/Engine/Classes/Components/LightComponent.h b/Engine/Source/Runtime/Engine/Classes/Components/LightComponent.h
--- a/Engine/Source/Runtime/Engine/Classes/Components/LightComponent.h
+++ b/Engine/Source/Runtime/Engine/Classes/Components/LightComponent.h
@@ -78,10 +78,6 @@ class ENGINE_API ULightComponent : public ULightComponentBase
 	UPROPERTY()
 	uint32 InverseSquaredFalloff_DEPRECATED:1;
 
-	/** Runtime toggleable enabled state - this has been replaced by USceneComponent::bVisible */
-	UPROPERTY()
-	uint32 bEnabled_DEPRECATED:1;
-
 	/** Whether the light is allowed to cast dynamic shadows from translucency. */
 	UPROPERTY(EditAnywhere, BlueprintReadOnly, Category=Light, AdvancedDisplay)
 	uint32 CastTranslucentShadows:1;
diff --git a/Engine/Source/Runtime/Engine/Classes/Components/PrimitiveComponent.h b/Engine/Source/Runtime/Engine/Classes/Components/PrimitiveComponent.h
--- a/Engine/Source/Runtime/Engine/Classes/Components/PrimitiveComponent.h
+++ b/Engine/Source/Runtime/Engine/Classes/Components/PrimitiveComponent.h
@@ -139,15 +139,7 @@ public:
 	UPROPERTY()
 	TEnumAsByte<enum ESceneDepthPriorityGroup> ViewOwnerDepthPriorityGroup;
 
-private:
-	/** DEPRECATED, use bGenerateOverlapEvents instead */
-	UPROPERTY()
-	uint32 bNoEncroachCheck_DEPRECATED:1;
-
 public:
-	UPROPERTY()
-	uint32 bDisableAllRigidBody_DEPRECATED:1;
-
 	/** 
 	 * Indicates if we'd like to create physics state all the time (for collision and simulation). 
 	 * If you set this to false, it still will create physics state if collision or simulation activated. 
@@ -217,14 +209,6 @@ public:
 	UPROPERTY(EditAnywhere, AdvancedDisplay, BlueprintReadOnly, Category = Rendering)
 	uint32 bRenderInMainPass:1;
 
-	/** Whether to completely hide the primitive in the game; if true, the primitive is not drawn, does not cast a shadow. */
-	UPROPERTY()
-	uint32 HiddenGame_DEPRECATED:1;
-
-	/** Whether to draw the primitive in game views, if the primitive isn't hidden. */
-	UPROPERTY()
-	uint32 DrawInGame_DEPRECATED:1;
-	
 	/** Whether the primitive receives decals. */
 	UPROPERTY(EditAnywhere, AdvancedDisplay, BlueprintReadOnly, Category=Rendering)
 	uint32 bReceivesDecals:1;
diff --git a/Engine/Source/Runtime/Engine/Classes/Components/SceneComponent.h b/Engine/Source/Runtime/Engine/Classes/Components/SceneComponent.h
--- a/Engine/Source/Runtime/Engine/Classes/Components/SceneComponent.h
+++ b/Engine/Source/Runtime/Engine/Classes/Components/SceneComponent.h
@@ -183,9 +183,6 @@ public:
 
 private:
 
-	UPROPERTY()
-	float RelativeScale_DEPRECATED;
-
 	bool NetUpdateTransform;
 
 	USceneComponent *NetOldAttachParent;
diff --git a/Engine/Source/Runtime/Engine/Classes/Components/SkeletalMeshComponent.h b/Engine/Source/Runtime/Engine/Classes/Components/SkeletalMeshComponent.h
--- a/Engine/Source/Runtime/Engine/Classes/Components/SkeletalMeshComponent.h
+++ b/Engine/Source/Runtime/Engine/Classes/Components/SkeletalMeshComponent.h
@@ -445,10 +445,6 @@ public:
 	UPROPERTY(EditAnywhere, AdvancedDisplay, BlueprintReadWrite, Category=SkeletalMesh)
 	uint32 bNoSkeletonUpdate:1;
 
-	/** If true, tick anim nodes even when our Owner has not been rendered recently  */
-	UPROPERTY()
-	uint32 bTickAnimationWhenNotRendered_DEPRECATED:1;
-
 	/** pauses this component's animations (doesn't tick them, but still refreshes bones) */
 	UPROPERTY(EditAnywhere, AdvancedDisplay, BlueprintReadWrite, Category=Animation)
 	uint32 bPauseAnims:1;
diff --git a/Engine/Source/Runtime/Engine/Classes/Components/SkinnedMeshComponent.h b/Engine/Source/Runtime/Engine/Classes/Components/SkinnedMeshComponent.h
--- a/Engine/Source/Runtime/Engine/Classes/Components/SkinnedMeshComponent.h
+++ b/Engine/Source/Runtime/Engine/Classes/Components/SkinnedMeshComponent.h
@@ -306,14 +306,6 @@ public:
 	UPROPERTY(EditAnywhere, AdvancedDisplay, BlueprintReadWrite, Category=SkeletalMesh)
 	uint32 bConsiderAllBodiesForBounds:1;
 
-	/** If true, update skeleton/attachments even when our Owner has not been rendered recently
-	 * @note if this is false, bone information may not be accurate, so be careful setting this to false if bone info is relevant to gameplay
-	 * @note you can use ForceSkelUpdate() to force an update
-	 * @note: In the output from SHOWSKELCOMPTICKTIME you want UpdatePoseTotal to be 0 when this is false for a specific component
-	 */
-	UPROPERTY()
-	uint32 bUpdateSkelWhenNotRendered_DEPRECATED:1;
-
 	/** This is update frequency flag even when our Owner has not been rendered recently
 	 * 
 	 * SMU_AlwaysTickPoseAndRefreshBones,			// Always Tick and Refresh BoneTransforms whether rendered or not
diff --git a/Engine/Source/Runtime/Engine/Classes/Curves/CurveBase.h b/Engine/Source/Runtime/Engine/Classes/Curves/CurveBase.h
--- a/Engine/Source/Runtime/Engine/Classes/Curves/CurveBase.h
+++ b/Engine/Source/Runtime/Engine/Classes/Curves/CurveBase.h
@@ -360,13 +360,6 @@ public:
 	/** Auto set tangents for any 'auto' keys in curve */
 	void AutoSetTangents(float Tension = 0.f);
 
-	// Legacy conversion
-
-	/** Util to convert from legacy FInterpCurveFloat */
-	static void ConvertInterpCurveFloat(const FInterpCurveFloat& InCurve, FRichCurve& OutCurve);
-	/** Util to convert from legacy FInterpCurveVector */
-	static void ConvertInterpCurveVector(const FInterpCurveVector& InCurve, FRichCurve OutCurves[3]);
-
 	/** Determine if two RichCurves are the same */
 	bool operator == (const FRichCurve& Curve) const;
 private:
diff --git a/Engine/Source/Runtime/Engine/Classes/Curves/CurveFloat.h b/Engine/Source/Runtime/Engine/Classes/Curves/CurveFloat.h
--- a/Engine/Source/Runtime/Engine/Classes/Curves/CurveFloat.h
+++ b/Engine/Source/Runtime/Engine/Classes/Curves/CurveFloat.h
@@ -28,10 +28,6 @@ class UCurveFloat : public UCurveBase
 {
 	GENERATED_UCLASS_BODY()
 
-	/** Legacy keyframe data. */
-	UPROPERTY()
-	FInterpCurveFloat FloatKeys_DEPRECATED;
-
 	/** Keyframe data */
 	UPROPERTY()
 	FRichCurve FloatCurve;
@@ -44,9 +40,6 @@ class UCurveFloat : public UCurveBase
 	UFUNCTION(BlueprintCallable, Category="Math|Curves")
 	ENGINE_API float GetFloatValue(float InTime) const;
 
-	// UObject interface
-	virtual void PostLoad() override;
-
 	// Begin FCurveOwnerInterface
 	virtual TArray<FRichCurveEditInfoConst> GetCurves() const override;
 	virtual TArray<FRichCurveEditInfo> GetCurves() override;
diff --git a/Engine/Source/Runtime/Engine/Classes/Curves/CurveVector.h b/Engine/Source/Runtime/Engine/Classes/Curves/CurveVector.h
--- a/Engine/Source/Runtime/Engine/Classes/Curves/CurveVector.h
+++ b/Engine/Source/Runtime/Engine/Classes/Curves/CurveVector.h
@@ -9,10 +9,6 @@ class UCurveVector : public UCurveBase
 {
 	GENERATED_UCLASS_BODY()
 
-	/** Legacy keyframe data. */
-	UPROPERTY()
-	FInterpCurveVector VectorKeys_DEPRECATED;
-
 	/** Keyframe data, one curve for X, Y and Z */
 	UPROPERTY()
 	FRichCurve FloatCurves[3];
@@ -21,9 +17,6 @@ class UCurveVector : public UCurveBase
 	UFUNCTION(BlueprintCallable, Category="Math|Curves")
 	ENGINE_API FVector GetVectorValue(float InTime) const;
 
-	// UObject interface
-	virtual void PostLoad() override;
-
 	// Begin FCurveOwnerInterface
 	virtual TArray<FRichCurveEditInfoConst> GetCurves() const override;
 	virtual TArray<FRichCurveEditInfo> GetCurves() override;
diff --git a/Engine/Source/Runtime/Engine/Classes/Distributions/Distribution.h b/Engine/Source/Runtime/Engine/Classes/Distributions/Distribution.h
--- a/Engine/Source/Runtime/Engine/Classes/Distributions/Distribution.h
+++ b/Engine/Source/Runtime/Engine/Classes/Distributions/Distribution.h
@@ -60,5 +60,8 @@ class UDistribution : public UObject, public FCurveEdInterface
 {
 	GENERATED_UCLASS_BODY()
 
+	/** Default value for initializing and checking correct values on UDistributions. */
+	static const float DefaultValue;
+
 };
 
diff --git a/Engine/Source/Runtime/Engine/Classes/Engine/Brush.h b/Engine/Source/Runtime/Engine/Classes/Engine/Brush.h
--- a/Engine/Source/Runtime/Engine/Classes/Engine/Brush.h
+++ b/Engine/Source/Runtime/Engine/Classes/Engine/Brush.h
@@ -77,10 +77,6 @@ class ENGINE_API ABrush : public AActor
 {
 	GENERATED_UCLASS_BODY()
 
-	/** CSG operation performed in editor. */
-	UPROPERTY()
-	TEnumAsByte<enum ECsgOper> CsgOper_DEPRECATED;
-
 	/** Type of brush */
 	UPROPERTY(EditAnywhere, Category=Brush)
 	TEnumAsByte<enum EBrushType> BrushType;
diff --git a/Engine/Source/Runtime/Engine/Classes/Engine/DestructibleMesh.h b/Engine/Source/Runtime/Engine/Classes/Engine/DestructibleMesh.h
--- a/Engine/Source/Runtime/Engine/Classes/Engine/DestructibleMesh.h
+++ b/Engine/Source/Runtime/Engine/Classes/Engine/DestructibleMesh.h
@@ -449,10 +449,6 @@ class UDestructibleMesh : public USkeletalMesh
 	UPROPERTY(EditAnywhere, Category=DestructibleMesh)
 	struct FDestructibleParameters DefaultDestructibleParameters;
 
-	/** The PhysicalMaterial used if there is no override in the component. */
-	UPROPERTY()
-	class UPhysicalMaterial* DestructiblePhysicalMaterial_DEPRECATED;
-
 	/** Fracture effects for each fracture level, unless overridden in the component. */
 	UPROPERTY(EditAnywhere, editfixedsize, Category=Effects)
 	TArray<struct FFractureEffect> FractureEffects;
diff --git a/Engine/Source/Runtime/Engine/Classes/Engine/EngineTypes.h b/Engine/Source/Runtime/Engine/Classes/Engine/EngineTypes.h
--- a/Engine/Source/Runtime/Engine/Classes/Engine/EngineTypes.h
+++ b/Engine/Source/Runtime/Engine/Classes/Engine/EngineTypes.h
@@ -1038,10 +1038,6 @@ struct FLightmassLightSettings
 {
 	GENERATED_USTRUCT_BODY()
 
-	/** Scale factor for the indirect lighting */
-	UPROPERTY()
-	float IndirectLightingScale_DEPRECATED;
-
 	/** 0 will be completely desaturated, 1 will be unchanged */
 	UPROPERTY(EditAnywhere, BlueprintReadWrite, Category=Lightmass, meta=(UIMin = "0.0", UIMax = "4.0"))
 	float IndirectLightingSaturation;
@@ -1051,8 +1047,7 @@ struct FLightmassLightSettings
 	float ShadowExponent;
 
 	FLightmassLightSettings()
-		: IndirectLightingScale_DEPRECATED(1.0f)
-		, IndirectLightingSaturation(1.0f)
+		: IndirectLightingSaturation(1.0f)
 		, ShadowExponent(2.0f)
 	{
 	}
diff --git a/Engine/Source/Runtime/Engine/Classes/Engine/InterpCurveEdSetup.h b/Engine/Source/Runtime/Engine/Classes/Engine/InterpCurveEdSetup.h
--- a/Engine/Source/Runtime/Engine/Classes/Engine/InterpCurveEdSetup.h
+++ b/Engine/Source/Runtime/Engine/Classes/Engine/InterpCurveEdSetup.h
@@ -110,7 +110,6 @@ class UInterpCurveEdSetup : public UObject
 
 	// Begin UObject interface
 	virtual void PostLoad() override;
-	virtual void Serialize(FArchive& Ar) override;
 	virtual bool NeedsLoadForClient() const override
 	{
 		return false;
diff --git a/Engine/Source/Runtime/Engine/Classes/Engine/SkeletalMesh.h b/Engine/Source/Runtime/Engine/Classes/Engine/SkeletalMesh.h
--- a/Engine/Source/Runtime/Engine/Classes/Engine/SkeletalMesh.h
+++ b/Engine/Source/Runtime/Engine/Classes/Engine/SkeletalMesh.h
@@ -299,10 +299,6 @@ struct FSkeletalMeshLODInfo
 	UPROPERTY()
 	TArray<bool> bEnableShadowCasting_DEPRECATED;
 
-	/** Per-section sorting options */
-	UPROPERTY()
-	TArray<TEnumAsByte<enum ETriangleSortOption> > TriangleSorting_DEPRECATED;
-
 	UPROPERTY(EditAnywhere, editfixedsize, Category=SkeletalMeshLODInfo)
 	TArray<struct FTriangleSortSettings> TriangleSortSettings;
 
@@ -634,11 +630,6 @@ public:
 	class UPhysicsAsset* PhysicsAsset;
 
 #if WITH_EDITORONLY_DATA
-	/** Asset used for previewing bounds in AnimSetViewer. Makes setting up LOD distance factors more reliable. 
-	 *  Removing this and use PhysicsAsset for preview from now on, but meanwhile, we will copy this info to PhysicsAsset (below)
-	 */
-	UPROPERTY()
-	class UPhysicsAsset* BoundsPreviewAsset_DEPRECATED;
 
 	/** Importing data and options used for this mesh */
 	UPROPERTY(EditAnywhere, Instanced, Category = Reimport)
@@ -672,12 +663,6 @@ public:
 	UPROPERTY(EditAnywhere, BlueprintReadOnly, Category=TextureStreaming)
 	float StreamingDistanceMultiplier;
 
-	/** This is serialized right now with strong pointer to Morph Target
-	 * Later on when memory is concerned, this can be transient array that is built based on MorphTarget loaded 
-	 */
-	UPROPERTY()
-	TArray<FMorphTargetMap> MorphTargetTable_DEPRECATED;
-
 	UPROPERTY(Category=Mesh, BlueprintReadWrite)
 	TArray<UMorphTarget*> MorphTargets;
 
diff --git a/Engine/Source/Runtime/Engine/Classes/Engine/StaticMesh.h b/Engine/Source/Runtime/Engine/Classes/Engine/StaticMesh.h
--- a/Engine/Source/Runtime/Engine/Classes/Engine/StaticMesh.h
+++ b/Engine/Source/Runtime/Engine/Classes/Engine/StaticMesh.h
@@ -109,51 +109,15 @@ struct FStaticMeshOptimizationSettings
 	/** Serialization for FStaticMeshOptimizationSettings. */
 	inline friend FArchive& operator<<( FArchive& Ar, FStaticMeshOptimizationSettings& Settings )
 	{
-		if( Ar.UE4Ver() < VER_UE4_ADDED_EXTRA_MESH_OPTIMIZATION_SETTINGS )
-		{
-			Ar << Settings.MaxDeviationPercentage;
-
-			//Remap Importance Settings
-			Ar << Settings.SilhouetteImportance;
-			Ar << Settings.TextureImportance;
-
-			//IL_Normal was previously the first enum value. We add the new index of IL_Normal to correctly offset the old values.
-			Settings.SilhouetteImportance += IL_Normal;
-			Settings.TextureImportance += IL_Normal;
-
-			//Remap NormalMode enum value to new threshold variable.
-			uint8 NormalMode;
-			Ar << NormalMode;
-
-			const float NormalThresholdTable[] =
-			{
-				60.0f, // Recompute
-				80.0f, // Recompute (Smooth)
-				45.0f  // Recompute (Hard)
-			};
-
-			if( NormalMode == NM_PreserveSmoothingGroups)
-			{
-				Settings.bRecalcNormals = false;
-			}
-			else
-			{
-				Settings.bRecalcNormals = true;
-				Settings.NormalsThreshold = NormalThresholdTable[NormalMode];
-			}
-		}
-		else
-		{
-			Ar << Settings.ReductionMethod;
-			Ar << Settings.MaxDeviationPercentage;
-			Ar << Settings.NumOfTrianglesPercentage;
-			Ar << Settings.SilhouetteImportance;
-			Ar << Settings.TextureImportance;
-			Ar << Settings.ShadingImportance;
-			Ar << Settings.bRecalcNormals;
-			Ar << Settings.NormalsThreshold;
-			Ar << Settings.WeldingThreshold;
-		}
+		Ar << Settings.ReductionMethod;
+		Ar << Settings.MaxDeviationPercentage;
+		Ar << Settings.NumOfTrianglesPercentage;
+		Ar << Settings.SilhouetteImportance;
+		Ar << Settings.TextureImportance;
+		Ar << Settings.ShadingImportance;
+		Ar << Settings.bRecalcNormals;
+		Ar << Settings.NormalsThreshold;
+		Ar << Settings.WeldingThreshold;
 
 		return Ar;
 	}
@@ -371,10 +335,6 @@ class UStaticMesh : public UObject, public IInterface_CollisionDataProvider, pub
 	UPROPERTY(EditAnywhere, transient, duplicatetransient, Instanced, Category = StaticMesh)
 	class UBodySetup* BodySetup;
 
-	// Artist-accessible options.
-	UPROPERTY()
-	uint32 UseFullPrecisionUVs_DEPRECATED:1;
-
 	/** True if mesh should use a less-conservative method of mip LOD texture factor computation.
 		requires mesh to be resaved to take effect as algorithm is applied on save. */
 	UPROPERTY(EditAnywhere, AdvancedDisplay, Category=StaticMesh, meta=(ToolTip="If true, use a less-conservative method of mip LOD texture factor computation.  Requires mesh to be resaved to take effect as algorithm is applied on save"))
@@ -650,11 +610,6 @@ public:
 	ENGINE_API static void GetLODGroupsDisplayNames(TArray<FText>& OutLODGroupsDisplayNames);
 
 private:
-	/**
-	 * Serializes in legacy source data and constructs the necessary source models.
-	 */
-	void SerializeLegacySouceData(class FArchive& Ar, const struct FBoxSphereBounds& LegacyBounds);
-
 	/**
 	 * Converts legacy LODDistance in the source models to Display Factor
 	 */
diff --git a/Engine/Source/Runtime/Engine/Classes/Engine/Texture.h b/Engine/Source/Runtime/Engine/Classes/Engine/Texture.h
--- a/Engine/Source/Runtime/Engine/Classes/Engine/Texture.h
+++ b/Engine/Source/Runtime/Engine/Classes/Engine/Texture.h
@@ -219,9 +219,6 @@ private:
 	/** Retrieve the size and offset for a source mip. The size includes all slices. */
 	int32 CalcMipOffset(int32 MipIndex) const;
 
-	/** Converts bulk data from legacy DDS. */
-	void ConvertFromLegacyDDS();
-
 	/** Uses a hash as the GUID, useful to prevent creating new GUIDs on load for legacy assets. */
 	void UseHashAsGuid();
 #endif
@@ -381,10 +378,6 @@ public:
 	UPROPERTY(EditAnywhere, BlueprintReadWrite, Category=Adjustments, meta=(DisplayName = "Max Alpha"))
 	float AdjustMaxAlpha;
 
-	/** Has uncompressed source art? */
-	UPROPERTY()
-	uint32 bIsSourceArtUncompressed_DEPRECATED:1;
-
 	/** If enabled, the texture's alpha channel will be discarded during compression */
 	UPROPERTY(EditAnywhere, BlueprintReadWrite, Category=Compression, meta=(DisplayName="Compress Without Alpha"))
 	uint32 CompressionNoAlpha:1;
@@ -396,9 +389,6 @@ public:
 	UPROPERTY(EditAnywhere, BlueprintReadWrite, Category=Compression)
 	uint32 DeferCompression:1;
 
-	UPROPERTY()
-	uint32 CompressionNoMipmaps_DEPRECATED:1;
-
 	/** The maximum resolution for generated textures. A value of 0 means the maximum size for the format on each platform, except HDR long/lat cubemaps, which default to a resolution of 512. */ 
 	UPROPERTY(EditAnywhere, BlueprintReadOnly, Category=Compression, meta=(DisplayName="Maximum Texture Size", ClampMin = "0.0"), AdvancedDisplay)
 	int32 MaxTextureSize;
@@ -415,27 +405,10 @@ public:
 	UPROPERTY(EditAnywhere, BlueprintReadWrite, Category=Texture, AdvancedDisplay)
 	uint32 bFlipGreenChannel:1;
 
-	/** If true, the color border pixels are preserved by mipmap generation.  One flag per color channel. */
-	UPROPERTY()
-	uint32 bPreserveBorderR_DEPRECATED:1;
-
-	UPROPERTY()
-	uint32 bPreserveBorderG_DEPRECATED:1;
-
-	UPROPERTY()
-	uint32 bPreserveBorderB_DEPRECATED:1;
-
-	UPROPERTY()
-	uint32 bPreserveBorderA_DEPRECATED:1;
-
 	/** For DXT1 textures, setting this will cause the texture to be twice the size, but better looking, on iPhone */
 	UPROPERTY()
 	uint32 bForcePVRTC4:1;
 
-	/* e.g. TSAT_Uncompressed, TSAT_PNGCompressed or TSAT_DDSFile */
-	UPROPERTY()
-	TEnumAsByte<enum ETextureSourceArtType> SourceArtType_DEPRECATED;
-
 	/** Per asset specific setting to define the mip-map generation properties like sharpening and kernel size. */
 	UPROPERTY(EditAnywhere, BlueprintReadWrite, Category=LevelOfDetail)
 	TEnumAsByte<enum TextureMipGenSettings> MipGenSettings;
@@ -734,13 +707,6 @@ public:
 #endif // WITH_EDITORONLY_DATA
 	}
 
-#if WITH_EDITORONLY_DATA
-	/**
-	 * Legacy serialization.
-	 */
-	void LegacySerialize(class FArchive& Ar, class FStripDataFlags& StripDataFlags);
-#endif
-
 	/**
 	 * Retrieves the pixel format enum for enum <-> string conversions.
 	 */
diff --git a/Engine/Source/Runtime/Engine/Classes/Engine/Texture2D.h b/Engine/Source/Runtime/Engine/Classes/Engine/Texture2D.h
--- a/Engine/Source/Runtime/Engine/Classes/Engine/Texture2D.h
+++ b/Engine/Source/Runtime/Engine/Classes/Engine/Texture2D.h
@@ -32,22 +32,6 @@ public:
 	/** Used for various timing measurements, e.g. streaming latency. */
 	float Timer;
 
-	/** The width of the texture. */
-	UPROPERTY()
-	int32 SizeX_DEPRECATED;
-
-	/** The height of the texture. */
-	UPROPERTY()
-	int32 SizeY_DEPRECATED;
-
-	/** The original width of the texture source art we imported from.			*/
-	UPROPERTY()
-	int32 OriginalSizeX_DEPRECATED;
-
-	/** The original height of the texture source art we imported from.			*/
-	UPROPERTY()
-	int32 OriginalSizeY_DEPRECATED;
-
 private:
 	/**
 	 * The imported size of the texture. Only valid on cooked builds when texture source is not
@@ -80,10 +64,6 @@ private:
 	uint32 bTemporarilyDisableStreaming:1;
 
 public:
-	/** true if the texture's mips should be stored directly and not use the derived data cache. Used by procedurally generated textures. */
-	UPROPERTY()
-	uint32 bDisableDerivedDataCache_DEPRECATED:1;
-
 	/** Whether the texture is currently streamable or not.						*/
 	UPROPERTY(transient, NonTransactional)
 	uint32 bIsStreamable:1;
@@ -106,10 +86,6 @@ public:
 	uint32 bHasBeenPaintedInEditor:1;
 #endif // WITH_EDITORONLY_DATA
 
-	/** The format of the texture data. */
-	UPROPERTY()
-	TEnumAsByte<enum EPixelFormat> Format_DEPRECATED;
-
 	/** The addressing mode to use for the X axis.								*/
 	UPROPERTY(EditAnywhere, BlueprintReadWrite, Category=Texture, meta=(DisplayName="X-axis Tiling Method"), AssetRegistrySearchable, AdvancedDisplay)
 	TEnumAsByte<enum TextureAddress> AddressX;
@@ -118,9 +94,6 @@ public:
 	UPROPERTY(EditAnywhere, BlueprintReadWrite, Category=Texture, meta=(DisplayName="Y-axis Tiling Method"), AssetRegistrySearchable, AdvancedDisplay)
 	TEnumAsByte<enum TextureAddress> AddressY;
 
-	/** ID generated whenever the texture is changed so that its bulk data can be updated in the TextureFileCache during cook */
-	FGuid TextureFileCacheGuid_DEPRECATED;
-
 public:
 	/** The derived data for this texture on this platform. */
 	FTexturePlatformData *PlatformData;
@@ -431,13 +404,6 @@ public:
 	/** creates and initializes a new Texture2D with the requested settings */
 	ENGINE_API static class UTexture2D* CreateTransient(int32 InSizeX, int32 InSizeY, EPixelFormat InFormat = PF_B8G8R8A8);
 
-#if WITH_EDITOR
-	/**
-	 * Legacy serialization.
-	 */
-	void LegacySerialize(class FArchive& Ar, class FStripDataFlags& StripDataFlags);
-#endif
-
 	/**
 	 * Gets the X size of the texture, in pixels
 	 */
diff --git a/Engine/Source/Runtime/Engine/Classes/Engine/World.h b/Engine/Source/Runtime/Engine/Classes/Engine/World.h
--- a/Engine/Source/Runtime/Engine/Classes/Engine/World.h
+++ b/Engine/Source/Runtime/Engine/Classes/Engine/World.h
@@ -504,10 +504,6 @@ class ENGINE_API UWorld : public UObject, public FNetworkNotify
 	GENERATED_UCLASS_BODY()
 
 #if WITH_EDITORONLY_DATA
-	// List of the layer names which were checked when the level was last saved
-	UPROPERTY(config) 
-	FString VisibleLayers_DEPRECATED;
-	
 	/** List of all the layers referenced by the world's actors */
 	UPROPERTY()
 	TArray< class ULayer* > Layers; 
diff --git a/Engine/Source/Runtime/Engine/Classes/GameFramework/Actor.h b/Engine/Source/Runtime/Engine/Classes/GameFramework/Actor.h
--- a/Engine/Source/Runtime/Engine/Classes/GameFramework/Actor.h
+++ b/Engine/Source/Runtime/Engine/Classes/GameFramework/Actor.h
@@ -912,10 +912,6 @@ public:
 	//==============================================================================
 	// Misc Blueprint support
 
-	/** Get the World in which this Actor exists. */
-	UFUNCTION(BlueprintCallable, meta=(DeprecatedFunction))
-	UWorld* K2_GetWorld() const;
-
 	/** 
 	 * Get CustomTimeDilation - this can be used for input control or speed control for slomo.
 	 * We don't want to scale input globally because input can be used for UI, which do not care for TimeDilation.
@@ -953,12 +949,10 @@ public:
 	//=============================================================================
 	// Sound functions.
 	
-	/* DEPRECATED - Use UGameplayStatics::PlaySoundAttached */
-	UFUNCTION(BlueprintCallable, Category="Audio", meta=(DeprecatedFunction))
+	DEPRECATED(4.0, "Actor::PlaySoundOnActor will be removed. Use UGameplayStatics::PlaySoundAttached instead.")
 	void PlaySoundOnActor(class USoundCue* InSoundCue, float VolumeMultiplier=1.f, float PitchMultiplier=1.f);
 
-	/* DEPRECATED - Use UGameplayStatics::PlaySoundAtLocation */
-	UFUNCTION(BlueprintCallable, Category="Audio", meta=(DeprecatedFunction))
+	DEPRECATED(4.0, "Actor::PlaySoundOnActor will be removed. Use UGameplayStatics::PlaySoundAtLocation instead.")
 	void PlaySoundAtLocation(class USoundCue* InSoundCue, FVector SoundLocation, float VolumeMultiplier=1.f, float PitchMultiplier=1.f);
 
 	//=============================================================================
diff --git a/Engine/Source/Runtime/Engine/Classes/GameFramework/WorldSettings.h b/Engine/Source/Runtime/Engine/Classes/GameFramework/WorldSettings.h
--- a/Engine/Source/Runtime/Engine/Classes/GameFramework/WorldSettings.h
+++ b/Engine/Source/Runtime/Engine/Classes/GameFramework/WorldSettings.h
@@ -359,17 +359,9 @@ class ENGINE_API AWorldSettings : public AInfo, public IInterface_AssetUserData
 	UPROPERTY()
 	class UBookMark* BookMarks[10];
 
-	/** Clip pad entries */
-	UPROPERTY()
-	TArray<class UDEPRECATED_ClipPadEntry*> ClipPadEntries_DEPRECATED;
-
 #endif // WITH_EDITORONLY_DATA
 	/************************************/
 
-	/** Level collection. ULevels are referenced by FName (Package name) to avoid serialized references. Also contains offsets in world units */
-	UPROPERTY()
-	TArray<class ULevelStreaming*> StreamingLevels_DEPRECATED;
-	
 	/** 
 	 * Normally 1 - scales real time passage.
 	 * Warning - most use cases should use GetEffectiveTimeDilation() instead of reading from this directly
diff --git a/Engine/Source/Runtime/Engine/Classes/Kismet/GameplayStatics.h b/Engine/Source/Runtime/Engine/Classes/Kismet/GameplayStatics.h
--- a/Engine/Source/Runtime/Engine/Classes/Kismet/GameplayStatics.h
+++ b/Engine/Source/Runtime/Engine/Classes/Kismet/GameplayStatics.h
@@ -319,10 +319,6 @@ class ENGINE_API UGameplayStatics : public UBlueprintFunctionLibrary
 	UFUNCTION(BlueprintCallable, Category="Audio", meta=(AdvancedDisplay = "2", UnsafeDuringActorConstruction = "true"))
 	static class UAudioComponent* PlayDialogueAttached(class UDialogueWave* Dialogue, const struct FDialogueContext& Context, class USceneComponent* AttachToComponent, FName AttachPointName = NAME_None, FVector Location = FVector(ForceInit), EAttachLocation::Type LocationType = EAttachLocation::KeepRelativeOffset, bool bStopWhenAttachedToDestroyed = false, float VolumeMultiplier = 1.f, float PitchMultiplier = 1.f, float StartTime = 0.f, class USoundAttenuation* AttenuationSettings = NULL);
 
-	/** DEPRECATED - Use GameplayStatics.PlaySoundAtLocation or GameplayStatics.PlaySoundAttached instead. */
-	UFUNCTION(BlueprintCallable, Category="Audio", meta=(WorldContext="WorldContextObject", DeprecatedFunction, DeprecationMessage="Use GameplayStatics.PlaySoundAtLocation or GameplayStatics.PlaySoundAttached instead."))
-	static void PlaySound(UObject* WorldContextObject, class USoundCue* InSoundCue, class USceneComponent* AttachComponent, FName AttachName=NAME_None, bool bFollow=false, float VolumeMultiplier=1.f, float PitchMultiplier=1.f);
-
 	// --- Audio Functions ----------------------------
 	/** Set the sound mix of the audio system for special EQing **/
 	UFUNCTION(BlueprintCallable, Category="Audio")
@@ -340,10 +336,6 @@ class ENGINE_API UGameplayStatics : public UBlueprintFunctionLibrary
 	UFUNCTION(BlueprintCallable, Category="Audio")
 	static void ClearSoundMixModifiers();
 
-	/** DEPRECATED - Use SetSoundMix that specifies SoundMix asset directly instead of by name. */
-	UFUNCTION(BlueprintCallable, meta=(DeprecatedFunction, DeprecationMessage="Use SetSoundMix that specifies SoundMix asset directly instead of by name.", FriendlyName = "SetSoundModeByName"))
-	static void SetSoundMode(FName SoundModeName);
-
 	/** Activates a Reverb Effect without the need for a volume
 	 * @param ReverbEffect Reverb Effect to use
 	 * @param TagName Tag to associate with Reverb Effect
diff --git a/Engine/Source/Runtime/Engine/Classes/Materials/Material.h b/Engine/Source/Runtime/Engine/Classes/Materials/Material.h
--- a/Engine/Source/Runtime/Engine/Classes/Materials/Material.h
+++ b/Engine/Source/Runtime/Engine/Classes/Materials/Material.h
@@ -289,9 +289,6 @@ class UMaterial : public UMaterialInterface
 	UPROPERTY()
 	FScalarMaterialInput Specular;
 
-	UPROPERTY()
-	FScalarMaterialInput SpecularPower_DEPRECATED;
-
 	UPROPERTY()
 	FScalarMaterialInput Roughness;
 
@@ -309,9 +306,6 @@ class UMaterial : public UMaterialInterface
 	UPROPERTY()
 	FScalarMaterialInput OpacityMask;
 
-	UPROPERTY()
-	float FresnelBaseReflectFraction_DEPRECATED;
-
 	/** 
 	 * The domain that the material's attributes will be evaluated in. 
 	 * Certain pieces of material functionality are only valid in certain domains, for example vertex normal is only valid on a surface.
@@ -478,16 +472,6 @@ public:
 	UPROPERTY(EditAnywhere, BlueprintReadWrite, Category=Material, meta=(DisplayName = "Emissive (Dynamic Area Light)"), AdvancedDisplay)
 	uint32 bUseEmissiveForDynamicAreaLighting : 1;
 
-	/** Whether material uses BaseColor, Metallic, Specular */
-	UPROPERTY()
-	uint32 bPhysicallyBasedInputs_DEPRECATED : 1;
-
-	UPROPERTY()
-	uint32 bUsedAsLightFunction_DEPRECATED:1;
-
-	UPROPERTY()
-	uint32 bUsedWithDeferredDecal_DEPRECATED:1;
-
 	/** 
 	 * This is a special usage flag that allows a material to be assignable to any primitive type.
 	 * This is useful for materials used by code to implement certain viewmodes, for example the default material or lighting only material.
diff --git a/Engine/Source/Runtime/Engine/Classes/Materials/MaterialExpressionBreakMaterialAttributes.h b/Engine/Source/Runtime/Engine/Classes/Materials/MaterialExpressionBreakMaterialAttributes.h
--- a/Engine/Source/Runtime/Engine/Classes/Materials/MaterialExpressionBreakMaterialAttributes.h
+++ b/Engine/Source/Runtime/Engine/Classes/Materials/MaterialExpressionBreakMaterialAttributes.h
@@ -16,10 +16,6 @@ class UMaterialExpressionBreakMaterialAttributes : public UMaterialExpression
  	UPROPERTY()
  	FMaterialAttributesInput MaterialAttributes;
  
-	// Begin UObject interface.
- 	virtual void Serialize( FArchive& Ar ) override;
-	// End UObject interface.
-
 	// Begin UMaterialExpression Interface
 	virtual int32 Compile(class FMaterialCompiler* Compiler, int32 OutputIndex, int32 MultiplexIndex) override;
 	virtual void GetCaption(TArray<FString>& OutCaptions) const override;
diff --git a/Engine/Source/Runtime/Engine/Classes/Materials/MaterialExpressionConstant3Vector.h b/Engine/Source/Runtime/Engine/Classes/Materials/MaterialExpressionConstant3Vector.h
--- a/Engine/Source/Runtime/Engine/Classes/Materials/MaterialExpressionConstant3Vector.h
+++ b/Engine/Source/Runtime/Engine/Classes/Materials/MaterialExpressionConstant3Vector.h
@@ -10,22 +10,9 @@ class UMaterialExpressionConstant3Vector : public UMaterialExpression
 {
 	GENERATED_UCLASS_BODY()
 
-	UPROPERTY()
-	float R_DEPRECATED;
-
-	UPROPERTY()
-	float G_DEPRECATED;
-
-	UPROPERTY()
-	float B_DEPRECATED;
-
  	UPROPERTY(EditAnywhere, Category=MaterialExpressionConstant3Vector, meta=(HideAlphaChannel))
 	FLinearColor Constant;
 
-	// Begin UObject interface.
-	virtual void PostLoad() override;
-	// End UObject interface.
-
 	// Begin UMaterialExpression Interface
 	virtual int32 Compile(class FMaterialCompiler* Compiler, int32 OutputIndex, int32 MultiplexIndex) override;
 	virtual void GetCaption(TArray<FString>& OutCaptions) const override;
diff --git a/Engine/Source/Runtime/Engine/Classes/Materials/MaterialExpressionConstant4Vector.h b/Engine/Source/Runtime/Engine/Classes/Materials/MaterialExpressionConstant4Vector.h
--- a/Engine/Source/Runtime/Engine/Classes/Materials/MaterialExpressionConstant4Vector.h
+++ b/Engine/Source/Runtime/Engine/Classes/Materials/MaterialExpressionConstant4Vector.h
@@ -10,25 +10,9 @@ class UMaterialExpressionConstant4Vector : public UMaterialExpression
 {
 	GENERATED_UCLASS_BODY()
 
-	UPROPERTY()
-	float R_DEPRECATED;
-
-	UPROPERTY()
-	float G_DEPRECATED;
-
-	UPROPERTY()
-	float B_DEPRECATED;
-
-	UPROPERTY()
-	float A_DEPRECATED;
-
 	UPROPERTY(EditAnywhere, Category=MaterialExpressionConstant4Vector)
 	FLinearColor Constant;
 
-	// Begin UObject interface.
-	virtual void PostLoad() override;
-	// End UObject interface.
-
 	// Begin UMaterialExpression Interface
 	virtual int32 Compile(class FMaterialCompiler* Compiler, int32 OutputIndex, int32 MultiplexIndex) override;
 	virtual void GetCaption(TArray<FString>& OutCaptions) const override;
diff --git a/Engine/Source/Runtime/Engine/Classes/Matinee/MatineeActor.h b/Engine/Source/Runtime/Engine/Classes/Matinee/MatineeActor.h
--- a/Engine/Source/Runtime/Engine/Classes/Matinee/MatineeActor.h
+++ b/Engine/Source/Runtime/Engine/Classes/Matinee/MatineeActor.h
@@ -41,9 +41,6 @@ struct FInterpGroupActorInfo
 {
 	GENERATED_USTRUCT_BODY()
 
-	UPROPERTY()
-	FName GroupName_DEPRECATED;
-
 	UPROPERTY(EditAnywhere, BlueprintReadWrite, Category=InterpGroupActorInfo)
 	FName ObjectName;
 
@@ -328,7 +325,6 @@ public:
 	virtual void PostEditChangeProperty( FPropertyChangedEvent& PropertyChangedEvent ) override; 
 	virtual bool CanEditChange( const UProperty* Property ) const override;
 #endif // WITH_EDITOR
-	virtual void PostLoadSubobjects( FObjectInstancingGraph* OuterInstanceGraph ) override;
 	// End UObject Interface
 
 	/** Increment track forwards by given timestep and iterate over each track updating any properties. */
diff --git a/Engine/Source/Runtime/Engine/Classes/Particles/Acceleration/ParticleModuleAcceleration.h b/Engine/Source/Runtime/Engine/Classes/Particles/Acceleration/ParticleModuleAcceleration.h
--- a/Engine/Source/Runtime/Engine/Classes/Particles/Acceleration/ParticleModuleAcceleration.h
+++ b/Engine/Source/Runtime/Engine/Classes/Particles/Acceleration/ParticleModuleAcceleration.h
@@ -36,7 +36,6 @@ class UParticleModuleAcceleration : public UParticleModuleAccelerationBase
 	virtual void PostEditChangeProperty(FPropertyChangedEvent& PropertyChangedEvent) override;
 #endif // WITH_EDITOR
 	virtual void PostInitProperties() override;
-	virtual void Serialize(FArchive& Ar) override;
 	//End UObject Interface
 
 	//Begin UParticleModule Interface
diff --git a/Engine/Source/Runtime/Engine/Classes/Particles/Location/ParticleModuleLocationDirect.h b/Engine/Source/Runtime/Engine/Classes/Particles/Location/ParticleModuleLocationDirect.h
--- a/Engine/Source/Runtime/Engine/Classes/Particles/Location/ParticleModuleLocationDirect.h
+++ b/Engine/Source/Runtime/Engine/Classes/Particles/Location/ParticleModuleLocationDirect.h
@@ -49,7 +49,6 @@ class UParticleModuleLocationDirect : public UParticleModuleLocationBase
 	virtual void PostEditChangeProperty(FPropertyChangedEvent& PropertyChangedEvent) override;
 #endif // WITH_EDITOR
 	virtual void	PostInitProperties() override;
-	virtual void	Serialize(FArchive& Ar) override;
 	// End UObject Interface
 
 	// Begin UParticleModule Interface
diff --git a/Engine/Source/Runtime/Engine/Classes/Particles/Parameter/ParticleModuleParameterDynamic.h b/Engine/Source/Runtime/Engine/Classes/Particles/Parameter/ParticleModuleParameterDynamic.h
--- a/Engine/Source/Runtime/Engine/Classes/Particles/Parameter/ParticleModuleParameterDynamic.h
+++ b/Engine/Source/Runtime/Engine/Classes/Particles/Parameter/ParticleModuleParameterDynamic.h
@@ -108,7 +108,6 @@ class UParticleModuleParameterDynamic : public UParticleModuleParameterBase
 	virtual void	PostEditChangeProperty(FPropertyChangedEvent& PropertyChangedEvent) override;
 #endif // WITH_EDITOR
 	virtual void	PostInitProperties() override;
-	virtual void	Serialize(FArchive& Ar) override;
 	// Begin UObject Interface
 
 	// Begin UParticleModule Interface
diff --git a/Engine/Source/Runtime/Engine/Classes/Particles/ParticleModule.h b/Engine/Source/Runtime/Engine/Classes/Particles/ParticleModule.h
--- a/Engine/Source/Runtime/Engine/Classes/Particles/ParticleModule.h
+++ b/Engine/Source/Runtime/Engine/Classes/Particles/ParticleModule.h
@@ -417,28 +417,6 @@ class UParticleModule : public UObject
 	{
 	}
 
-	/**
-	 *	Comparison routine...
-	 *	Intended for editor-use only, this function will return true if the given
-	 *	particle module settings are identical to the one the function is called on.
-	 *
-	 *	@param	InModule	The module to compare against.
-	 *
-	 *	@return	true		If the modules have all the relevant settings the same.
-	 *			false		If they don't.
-	 */
-	virtual bool	IsIdentical_Deprecated(const UParticleModule* InModule) const;
-
-	/**
-	 *	Used by the comparison routine to check for properties that are irrelevant.
-	 *
-	 *	@param	InPropName	The name of the property being checked.
-	 *
-	 *	@return	true		If the property is relevant.
-	 *			false		If it isn't.
-	 */
-	virtual bool	PropertyIsRelevantForIsIdentical_Deprecated(const FName& InPropName) const;
-
 	/**
 	 *	Generates a new module for LOD levels, setting the values appropriately.
 	 *	Note that the module returned could simply be the module it was called on.
diff --git a/Engine/Source/Runtime/Engine/Classes/Particles/ParticleModuleRequired.h b/Engine/Source/Runtime/Engine/Classes/Particles/ParticleModuleRequired.h
--- a/Engine/Source/Runtime/Engine/Classes/Particles/ParticleModuleRequired.h
+++ b/Engine/Source/Runtime/Engine/Classes/Particles/ParticleModuleRequired.h
@@ -118,10 +118,6 @@ class UParticleModuleRequired : public UParticleModule
 	UPROPERTY(EditAnywhere, Category=Emitter)
 	uint32 bKillOnCompleted:1;
 
-	/** Whether this emitter requires sorting as specified by artist.					*/
-	UPROPERTY()
-	uint32 bRequiresSorting_DEPRECATED:1;
-
 	/**
 	 *	The sorting mode to use for this emitter.
 	 *	PSORTMODE_None				- No sorting required.
@@ -344,7 +340,6 @@ class UParticleModuleRequired : public UParticleModule
 #endif // WITH_EDITOR
 	virtual void	PostLoad() override;
 	virtual void	PostInitProperties() override;
-	virtual void	Serialize(FArchive& Ar) override;
 	// End UObject Interface
 
 	// Begin UParticleModule Interface
diff --git a/Engine/Source/Runtime/Engine/Classes/Particles/ParticleSystem.h b/Engine/Source/Runtime/Engine/Classes/Particles/ParticleSystem.h
--- a/Engine/Source/Runtime/Engine/Classes/Particles/ParticleSystem.h
+++ b/Engine/Source/Runtime/Engine/Classes/Particles/ParticleSystem.h
@@ -145,10 +145,6 @@ class UParticleSystem : public UObject
 	float ThumbnailWarmup;
 
 #endif // WITH_EDITORONLY_DATA
-	/** Deprecated, ParticleSystemLOD::bLit is used instead. */
-	UPROPERTY()
-	uint32 bLit_DEPRECATED:1;
-
 	/** Used for curve editor to remember curve-editing setup.						*/
 	UPROPERTY(export)
 	class UInterpCurveEdSetup* CurveEdSetup;
diff --git a/Engine/Source/Runtime/Engine/Classes/Particles/Trail/ParticleModuleTrailSource.h b/Engine/Source/Runtime/Engine/Classes/Particles/Trail/ParticleModuleTrailSource.h
--- a/Engine/Source/Runtime/Engine/Classes/Particles/Trail/ParticleModuleTrailSource.h
+++ b/Engine/Source/Runtime/Engine/Classes/Particles/Trail/ParticleModuleTrailSource.h
@@ -82,7 +82,6 @@ class UParticleModuleTrailSource : public UParticleModuleTrailBase
 	virtual void	PostEditChangeProperty(FPropertyChangedEvent& PropertyChangedEvent) override;
 #endif // WITH_EDITOR
 	virtual void	PostInitProperties() override;
-	virtual void 	Serialize(FArchive& Ar) override;
 	// End UObject Interface
 
 	// Begin UParticleModule Interface
diff --git a/Engine/Source/Runtime/Engine/Classes/PhysicsEngine/BodyInstance.h b/Engine/Source/Runtime/Engine/Classes/PhysicsEngine/BodyInstance.h
--- a/Engine/Source/Runtime/Engine/Classes/PhysicsEngine/BodyInstance.h
+++ b/Engine/Source/Runtime/Engine/Classes/PhysicsEngine/BodyInstance.h
@@ -134,10 +134,6 @@ struct ENGINE_API FBodyInstance
 	/////////
 	// COLLISION SETTINGS
 
-	/** Enable any collision for this body */
-	UPROPERTY()
-	uint32 bEnableCollision_DEPRECATED:1;
-
 	/** Types of objects that this physics objects will collide with. */
 	// @todo : make this to be transient, so that it doesn't have to save anymore
 	// we have to still load them until resave
@@ -339,11 +335,6 @@ public:
 	/** Constructor **/
 	FBodyInstance();
 
-	// BodyInstance interface
-
-	/** Update CollisionEnabled in old assets (before VER_UE4_CHANGE_BENABLECOLLISION_TO_COLLISIONENABLED) using bEnableCollision_DEPRECATED */ 
-	void UpdateFromDeprecatedEnableCollision();
-
 	/**  
 	 * Update profile data if required
 	 * 
diff --git a/Engine/Source/Runtime/Engine/Classes/PhysicsEngine/ConstraintInstance.h b/Engine/Source/Runtime/Engine/Classes/PhysicsEngine/ConstraintInstance.h
--- a/Engine/Source/Runtime/Engine/Classes/PhysicsEngine/ConstraintInstance.h
+++ b/Engine/Source/Runtime/Engine/Classes/PhysicsEngine/ConstraintInstance.h
@@ -179,13 +179,6 @@ struct ENGINE_API FConstraintInstance
 	UPROPERTY(EditAnywhere, AdvancedDisplay, Category=Linear, meta=(editcondition = "bLinearBreakable", ClampMin = "0.0"))
 	float LinearBreakThreshold;
 
-	// ANGULAR DOF
-	UPROPERTY()
-	uint32 bSwingLimited_DEPRECATED:1;
-
-	UPROPERTY()
-	uint32 bTwistLimited_DEPRECATED:1;
-
 	/** Indicates whether rotation about the Z axis is allowed, blocked, or limited. If limited, the 
 		AngularLimit property will be used to determine the range of motion. See EAngularConstraintMotion. */
 	UPROPERTY(EditAnywhere, Category=Angular)
diff --git a/Engine/Source/Runtime/Engine/Classes/PhysicsEngine/PhysicsConstraintTemplate.h b/Engine/Source/Runtime/Engine/Classes/PhysicsEngine/PhysicsConstraintTemplate.h
--- a/Engine/Source/Runtime/Engine/Classes/PhysicsEngine/PhysicsConstraintTemplate.h
+++ b/Engine/Source/Runtime/Engine/Classes/PhysicsEngine/PhysicsConstraintTemplate.h
@@ -62,10 +62,6 @@ class UPhysicsConstraintTemplate : public UObject
 	UPROPERTY()
 	float LinearBreakThreshold_DEPRECATED;
 	UPROPERTY()
-	uint32 bSwingLimited_DEPRECATED:1;
-	UPROPERTY()
-	uint32 bTwistLimited_DEPRECATED:1;
-	UPROPERTY()
 	TEnumAsByte<enum EAngularConstraintMotion> AngularSwing1Motion_DEPRECATED;
 	UPROPERTY()
 	TEnumAsByte<enum EAngularConstraintMotion> AngularSwing2Motion_DEPRECATED;
@@ -99,7 +95,6 @@ class UPhysicsConstraintTemplate : public UObject
 
 
 	// Begin UObject interface.
-	virtual void PostLoad() override;
 	virtual void Serialize(FArchive& Ar) override;
 	// End UObject interface.
 
diff --git a/Engine/Source/Runtime/Engine/Classes/PhysicsEngine/PhysicsThrusterComponent.h b/Engine/Source/Runtime/Engine/Classes/PhysicsEngine/PhysicsThrusterComponent.h
--- a/Engine/Source/Runtime/Engine/Classes/PhysicsEngine/PhysicsThrusterComponent.h
+++ b/Engine/Source/Runtime/Engine/Classes/PhysicsEngine/PhysicsThrusterComponent.h
@@ -13,10 +13,6 @@ class UPhysicsThrusterComponent : public USceneComponent
 {
 	GENERATED_UCLASS_BODY()
 
-	/** If thrust should be applied at the moment. */
-	UPROPERTY()
-	uint32 bThrustEnabled_DEPRECATED:1;
-
 	/** Strength of thrust force applied to the base object. */
 	UPROPERTY(BlueprintReadWrite, interp, Category=Physics)
 	float ThrustStrength;
@@ -24,10 +20,6 @@ class UPhysicsThrusterComponent : public USceneComponent
 	// Begin UActorComponent Interface
 	virtual void TickComponent(float DeltaTime, enum ELevelTick TickType, FActorComponentTickFunction *ThisTickFunction) override;
 	// End UActorComponent Interface
-
-	// Begin UObject interface.
-	virtual void PostLoad() override;
-	// End UObject interface.
 };
 
 
diff --git a/Engine/Source/Runtime/Engine/Classes/PhysicsEngine/RadialForceComponent.h b/Engine/Source/Runtime/Engine/Classes/PhysicsEngine/RadialForceComponent.h
--- a/Engine/Source/Runtime/Engine/Classes/PhysicsEngine/RadialForceComponent.h
+++ b/Engine/Source/Runtime/Engine/Classes/PhysicsEngine/RadialForceComponent.h
@@ -40,10 +40,6 @@ class URadialForceComponent : public USceneComponent
 	UPROPERTY(EditAnywhere, BlueprintReadWrite, Category=Destructible)
 	float DestructibleDamage;
 
-	/** Is the force currently enabled */
-	UPROPERTY()
-	uint32 bForceEnabled_DEPRECATED:1;
-
 	/** Fire a single impulse */
 	UFUNCTION(BlueprintCallable, Category="Physics|Components|RadialForce")
 	virtual void FireImpulse();
diff --git a/Engine/Source/Runtime/Engine/Classes/Sound/AmbientSound.h b/Engine/Source/Runtime/Engine/Classes/Sound/AmbientSound.h
--- a/Engine/Source/Runtime/Engine/Classes/Sound/AmbientSound.h
+++ b/Engine/Source/Runtime/Engine/Classes/Sound/AmbientSound.h
@@ -9,9 +9,6 @@ class AAmbientSound : public AActor
 {
 	GENERATED_UCLASS_BODY()
 
-	UPROPERTY()
-	uint32 bAutoPlay_DEPRECATED:1;
-
 private_subobject:
 	/** Audio component that handles sound playing */
 	DEPRECATED_FORGAME(4.6, "AudioComponent should not be accessed directly, please use GetAudioComponent() function instead. AudioComponent will soon be private and your code will not compile.")
@@ -19,11 +16,6 @@ private_subobject:
 	class UAudioComponent* AudioComponent;
 public:
 	
-	UPROPERTY(instanced)
-	class UDEPRECATED_SoundNodeAmbient* SoundNodeInstance_DEPRECATED;
-
-	ENGINE_API static bool bUE4AudioRefactorMigrationUnderway;
-	ENGINE_API void MigrateSoundNodeInstance();	
 	ENGINE_API FString GetInternalSoundCueName();
 
 	// Begin AActor interface.
@@ -31,7 +23,6 @@ public:
 	virtual void CheckForErrors() override;
 	virtual bool GetReferencedContentObjects( TArray<UObject*>& Objects ) const override;
 #endif
-	virtual void PostLoad() override;
 	virtual void PostRegisterAllComponents() override;
 	// End AActor interface.
 
diff --git a/Engine/Source/Runtime/Engine/Classes/Sound/SoundBase.h b/Engine/Source/Runtime/Engine/Classes/Sound/SoundBase.h
--- a/Engine/Source/Runtime/Engine/Classes/Sound/SoundBase.h
+++ b/Engine/Source/Runtime/Engine/Classes/Sound/SoundBase.h
@@ -33,9 +33,6 @@ class ENGINE_API USoundBase : public UObject
 {
 	GENERATED_UCLASS_BODY()
 
-	UPROPERTY()
-	FName SoundClass_DEPRECATED;
-
 	/** Sound class this sound belongs to */
 	UPROPERTY(EditAnywhere, Category=Sound, meta=(DisplayName = "Sound Class"))
 	USoundClass* SoundClassObject;
diff --git a/Engine/Source/Runtime/Engine/Classes/Sound/SoundClass.h b/Engine/Source/Runtime/Engine/Classes/Sound/SoundClass.h
--- a/Engine/Source/Runtime/Engine/Classes/Sound/SoundClass.h
+++ b/Engine/Source/Runtime/Engine/Classes/Sound/SoundClass.h
@@ -171,9 +171,6 @@ class USoundClass : public UObject
 	UPROPERTY(EditAnywhere, Category=SoundClass)
 	TArray<struct FPassiveSoundMixModifier> PassiveSoundMixModifiers;
 
-	UPROPERTY()
-	TArray<FName> ChildClassNames_DEPRECATED;
-
 public:
 	UPROPERTY()
 	USoundClass* ParentClass;
diff --git a/Engine/Source/Runtime/Engine/Classes/Sound/SoundCue.h b/Engine/Source/Runtime/Engine/Classes/Sound/SoundCue.h
--- a/Engine/Source/Runtime/Engine/Classes/Sound/SoundCue.h
+++ b/Engine/Source/Runtime/Engine/Classes/Sound/SoundCue.h
@@ -63,8 +63,6 @@ class USoundCue : public USoundBase
 
 	/** EdGraph based representation of the SoundCue */
 	class USoundCueGraph* SoundCueGraph;
-
-	TMap<USoundNode*,FSoundNodeEditorData> EditorData_DEPRECATED;
 #endif
 
 private:
diff --git a/Engine/Source/Runtime/Engine/Classes/Sound/SoundMix.h b/Engine/Source/Runtime/Engine/Classes/Sound/SoundMix.h
--- a/Engine/Source/Runtime/Engine/Classes/Sound/SoundMix.h
+++ b/Engine/Source/Runtime/Engine/Classes/Sound/SoundMix.h
@@ -75,9 +75,6 @@ struct FSoundClassAdjuster
 {
 	GENERATED_USTRUCT_BODY()
 
-	UPROPERTY()
-	FName SoundClass_DEPRECATED;
-
 	/* The sound class this adjuster affects. */
 	UPROPERTY(EditAnywhere, Category=SoundClassAdjuster, DisplayName = "Sound Class" )
 	USoundClass* SoundClassObject;
diff --git a/Engine/Source/Runtime/Engine/Classes/Sound/SoundNode.h b/Engine/Source/Runtime/Engine/Classes/Sound/SoundNode.h
--- a/Engine/Source/Runtime/Engine/Classes/Sound/SoundNode.h
+++ b/Engine/Source/Runtime/Engine/Classes/Sound/SoundNode.h
@@ -51,14 +51,6 @@ class ENGINE_API USoundNode : public UObject
 	TArray<class USoundNode*> ChildNodes;
 
 #if WITH_EDITORONLY_DATA
-	/** X position of node in the editor, so old UEdGraphNode data not lost.*/
-	UPROPERTY()
-	int32 NodePosX_DEPRECATED;
-
-	/** Y position of node in the editor, so old UEdGraphNode data not lost. */
-	UPROPERTY()
-	int32 NodePosY_DEPRECATED;
-
 	/** Node's Graph representation, used to get position. */
 	class USoundCueGraphNode*	GraphNode;
 #endif
diff --git a/Engine/Source/Runtime/Engine/Classes/Sound/SoundNodeAttenuation.h b/Engine/Source/Runtime/Engine/Classes/Sound/SoundNodeAttenuation.h
--- a/Engine/Source/Runtime/Engine/Classes/Sound/SoundNodeAttenuation.h
+++ b/Engine/Source/Runtime/Engine/Classes/Sound/SoundNodeAttenuation.h
@@ -26,42 +26,7 @@ class USoundNodeAttenuation : public USoundNode
 	UPROPERTY(EditAnywhere, Category=Attenuation)
 	uint32 bOverrideAttenuation:1;
 
-	UPROPERTY()
-	uint32 bAttenuate_DEPRECATED:1;
-
-	UPROPERTY()
-	uint32 bSpatialize_DEPRECATED:1;
-
-	UPROPERTY()
-	float dBAttenuationAtMax_DEPRECATED;
-
-	UPROPERTY()
-	TEnumAsByte<enum ESoundDistanceModel> DistanceAlgorithm_DEPRECATED;
-
-	UPROPERTY()
-	TEnumAsByte<enum ESoundDistanceCalc> DistanceType_DEPRECATED;
-
-	UPROPERTY()
-	float RadiusMin_DEPRECATED;
-
-	UPROPERTY()
-	float RadiusMax_DEPRECATED;
-
-	UPROPERTY()
-	uint32 bAttenuateWithLPF_DEPRECATED:1;
-
-	UPROPERTY()
-	float LPFRadiusMin_DEPRECATED;
-
-	UPROPERTY()
-	float LPFRadiusMax_DEPRECATED;
-
-
 public:
-	// Begin UObject interface.
-	virtual void Serialize(FArchive& Ar) override;
-	// End UObject interface.
-
 	// Begin USoundNode interface. 
 	virtual void ParseNodes( class FAudioDevice* AudioDevice, const UPTRINT NodeWaveInstanceHash, FActiveSound& ActiveSound, const FSoundParseParameters& ParseParams, TArray<FWaveInstance*>& WaveInstances ) override;
 	virtual float MaxAudibleDistance( float CurrentMaxDistance ) override;
diff --git a/Engine/Source/Runtime/Engine/Classes/Sound/SoundNodeModulatorContinuous.h b/Engine/Source/Runtime/Engine/Classes/Sound/SoundNodeModulatorContinuous.h
--- a/Engine/Source/Runtime/Engine/Classes/Sound/SoundNodeModulatorContinuous.h
+++ b/Engine/Source/Runtime/Engine/Classes/Sound/SoundNodeModulatorContinuous.h
@@ -74,20 +74,7 @@ class USoundNodeModulatorContinuous : public USoundNode
 	UPROPERTY(EditAnywhere, Category=ContinuousModulator)
 	FModulatorContinuousParams VolumeModulationParams;
 
-	UPROPERTY()
-	struct FRawDistributionFloat PitchModulation_DEPRECATED;
-
-
-	UPROPERTY()
-	struct FRawDistributionFloat VolumeModulation_DEPRECATED;
-
-
-
 public:
-	// Begin UObject Interface
-	virtual void Serialize(FArchive& Ar) override;
-	// End UObject Interface
-
 	// Begin USoundNode interface.
 	virtual void ParseNodes( FAudioDevice* AudioDevice, const UPTRINT NodeWaveInstanceHash, FActiveSound& ActiveSound, const FSoundParseParameters& ParseParams, TArray<FWaveInstance*>& WaveInstances ) override;
 	// End USoundNode interface.
diff --git a/Engine/Source/Runtime/Engine/Private/Actor.cpp b/Engine/Source/Runtime/Engine/Private/Actor.cpp
--- a/Engine/Source/Runtime/Engine/Private/Actor.cpp
+++ b/Engine/Source/Runtime/Engine/Private/Actor.cpp
@@ -432,45 +432,6 @@ void AActor::PostLoad()
 		Owner->Children.Add(this);
 	}
 
-	// For components created before we added the bCreatedByConstructionScript flag, we assume if you are not a default subobject, you can from a construction script
-	if(GetLinkerUE4Version() < VER_UE4_ADD_CREATEDBYCONSTRUCTIONSCRIPT)
-	{
-		TArray<UActorComponent*> Components;
-		GetComponents(Components);
-
-		for(int32 i=0; i<Components.Num(); i++)
-		{
-			UActorComponent* Component = Components[i];
-			if (!Component->IsDefaultSubobject())
-			{
-				Component->bCreatedByConstructionScript = true;
-			}
-		}
-	}
-
-	// For actors created before SimpleConstructionScript's default RootComponent, and the default BadBlueprintSprite are flagged as RF_Transactional and bCreatedByConstructionScript
-	if(GetLinkerUE4Version() < VER_UE4_DEFAULT_ROOT_COMP_TRANSACTIONAL)
-	{
-		if(RootComponent != NULL && !RootComponent->IsDefaultSubobject())
-		{
-			if(RootComponent->GetClass()==USceneComponent::StaticClass())
-			{
-				RootComponent->SetFlags(RF_Transactional);
-				RootComponent->bCreatedByConstructionScript = true;
-			}
-			else if(RootComponent->GetClass()==UBillboardComponent::StaticClass())
-			{
-				static const FName BadSpriteName(TEXT("BadBlueprintSprite"));
-				UBillboardComponent* SpriteRootComponent = static_cast<UBillboardComponent*>(RootComponent);
-				if(SpriteRootComponent->Sprite->GetFName() == BadSpriteName)
-				{
-					SpriteRootComponent->SetFlags(RF_Transactional);
-					SpriteRootComponent->bCreatedByConstructionScript = true;
-				}
-			}
-		}
-	}
-
 	if (GetLinkerUE4Version() < VER_UE4_CONSUME_INPUT_PER_BIND)
 	{
 		bBlockInput = (InputConsumeOption_DEPRECATED == ICO_ConsumeAll);
@@ -3456,12 +3417,6 @@ void AActor::PreInitializeComponents()
 	}
 }
 
-UWorld* AActor::K2_GetWorld() const
-{
-	// If an actor is pending kill we don't consider it to be in a world
-	return (!IsPendingKill() ? GetLevel()->OwningWorld : NULL);
-}
-
 float AActor::GetActorTimeDilation() const
 {
 	// get actor custom time dilation
diff --git a/Engine/Source/Runtime/Engine/Private/AmbientSound.cpp b/Engine/Source/Runtime/Engine/Private/AmbientSound.cpp
--- a/Engine/Source/Runtime/Engine/Private/AmbientSound.cpp
+++ b/Engine/Source/Runtime/Engine/Private/AmbientSound.cpp
@@ -8,8 +8,6 @@
 #include "Sound/AmbientSound.h"
 #include "Sound/SoundBase.h"
 #include "Sound/SoundCue.h"
-#include "Sound/SoundNodeAmbient.h" // deprecated
-#include "Sound/SoundNodeAmbientNonLoop.h" // deprecated
 #include "Sound/SoundNodeMixer.h"
 #include "Sound/SoundNodeRandom.h"
 #include "Sound/SoundNodeLooping.h"
@@ -25,8 +23,6 @@
 /*-----------------------------------------------------------------------------
 	AAmbientSound implementation.
 -----------------------------------------------------------------------------*/
-bool AAmbientSound::bUE4AudioRefactorMigrationUnderway = false;
-
 AAmbientSound::AAmbientSound(const FObjectInitializer& ObjectInitializer)
 	: Super(ObjectInitializer)
 {
@@ -39,7 +35,6 @@ AAmbientSound::AAmbientSound(const FObjectInitializer& ObjectInitializer)
 
 	RootComponent = AudioComponent;
 
-	bAutoPlay_DEPRECATED = true;
 	bReplicates = false;
 	bHidden = true;
 	bCanBeDamaged = false;
@@ -104,185 +99,6 @@ FString AAmbientSound::GetInternalSoundCueName()
 	return CueName;
 }
 
-void AAmbientSound::MigrateSoundNodeInstance()
-{
-	if (SoundNodeInstance_DEPRECATED->SoundSlots.Num() == 0)
-	{
-		AudioComponent->Sound = NULL;
-		return;
-	}
-
-	bool bLoopAndDelay = SoundNodeInstance_DEPRECATED->GetClass() == UDEPRECATED_SoundNodeAmbientNonLoop::StaticClass();
-	bool bMixer = SoundNodeInstance_DEPRECATED->GetClass() == UDEPRECATED_SoundNodeAmbient::StaticClass();
-
-	int32 ColumnCount = 0;
-
-	USoundCue* SoundCue = ConstructObject<USoundCue>(USoundCue::StaticClass(), this, FName(*GetInternalSoundCueName()));
-	AudioComponent->Sound = SoundCue;
-
-	USoundNodeMixer* SoundMixer = NULL;
-	USoundNodeRandom* SoundRandom = NULL;
-	if (bMixer)
-	{
-		SoundMixer = SoundCue->ConstructSoundNode<USoundNodeMixer>();
-		SoundCue->FirstNode = SoundMixer;
-#if WITH_EDITOR
-		SoundMixer->PlaceNode(ColumnCount++, 0, 1);
-#endif //WITH_EDITOR
-	}
-	else
-	{
-		SoundRandom = SoundCue->ConstructSoundNode<USoundNodeRandom>();
-
-		if (bLoopAndDelay)
-		{
-			UDEPRECATED_SoundNodeAmbientNonLoop* SoundNodeNonLoop = CastChecked<UDEPRECATED_SoundNodeAmbientNonLoop>(SoundNodeInstance_DEPRECATED);
-			USoundNodeLooping* SoundLooping = SoundCue->ConstructSoundNode<USoundNodeLooping>();
-
-			SoundLooping->CreateStartingConnectors();
-			SoundCue->FirstNode = SoundLooping;
-#if WITH_EDITOR
-			SoundLooping->PlaceNode(ColumnCount++, 0, 1);
-#endif //WITH_EDITOR
-
-			if (SoundNodeNonLoop->DelayMin > 0.f || SoundNodeNonLoop->DelayMax > 0.f)
-			{
-				USoundNodeDelay* SoundDelay = SoundCue->ConstructSoundNode<USoundNodeDelay>();
-
-				SoundDelay->DelayMin = SoundNodeNonLoop->DelayMin;
-				SoundDelay->DelayMax = SoundNodeNonLoop->DelayMax;
-
-				SoundLooping->ChildNodes[0] = SoundDelay;
-#if WITH_EDITOR
-				SoundDelay->PlaceNode(ColumnCount++, 0, 1);
-#endif //WITH_EDITOR
-
-				SoundDelay->CreateStartingConnectors();
-				SoundDelay->ChildNodes[0] = SoundRandom;
-			}
-			else
-			{
-				SoundLooping->ChildNodes[0] = SoundRandom;
-			}
-		}
-		else
-		{
-			SoundCue->FirstNode = SoundRandom;
-		}
-#if WITH_EDITOR
-		SoundRandom->PlaceNode(ColumnCount++, 0, 1);
-#endif //WITH_EDITOR
-	}
-
-	for(int32 SlotIndex = 0;SlotIndex < SoundNodeInstance_DEPRECATED->SoundSlots.Num(); ++SlotIndex)
-	{
-		USoundNodeWavePlayer* WavePlayer = SoundCue->ConstructSoundNode<USoundNodeWavePlayer>();
-
-		WavePlayer->bLooping = bMixer;
-		WavePlayer->SoundWave = SoundNodeInstance_DEPRECATED->SoundSlots[SlotIndex].SoundWave;
-
-		if (bMixer)
-		{
-			SoundMixer->InsertChildNode(SlotIndex);
-			SoundMixer->InputVolume[SlotIndex] = SoundNodeInstance_DEPRECATED->SoundSlots[SlotIndex].VolumeScale;
-
-			if (SoundNodeInstance_DEPRECATED->SoundSlots[SlotIndex].PitchScale != 1.f)
-			{
-				USoundNodeModulator* SoundMod = SoundCue->ConstructSoundNode<USoundNodeModulator>();
-
-				SoundMod->VolumeMin = SoundMod->VolumeMax = 1.f;
-				SoundMod->PitchMin = SoundMod->PitchMax = SoundNodeInstance_DEPRECATED->SoundSlots[SlotIndex].PitchScale;
-
-				SoundMod->CreateStartingConnectors();
-				SoundMod->ChildNodes[0] = WavePlayer;
-#if WITH_EDITOR
-				SoundMod->PlaceNode(ColumnCount, SlotIndex, SoundNodeInstance_DEPRECATED->SoundSlots.Num());
-#endif //WITH_EDITOR
-
-				SoundMixer->ChildNodes[SlotIndex] = SoundMod;
-			}
-			else
-			{
-				SoundMixer->ChildNodes[SlotIndex] = WavePlayer;
-			}
-		}
-		else
-		{
-			SoundRandom->InsertChildNode(SlotIndex);
-			SoundRandom->Weights[SlotIndex] = SoundNodeInstance_DEPRECATED->SoundSlots[SlotIndex].Weight;
-
-			if (   SoundNodeInstance_DEPRECATED->SoundSlots[SlotIndex].VolumeScale != 1.f
-				|| SoundNodeInstance_DEPRECATED->SoundSlots[SlotIndex].PitchScale != 1.f)
-			{
-				USoundNodeModulator* SoundMod = SoundCue->ConstructSoundNode<USoundNodeModulator>();
-
-				SoundMod->VolumeMin = SoundMod->VolumeMax = SoundNodeInstance_DEPRECATED->SoundSlots[SlotIndex].VolumeScale;
-				SoundMod->PitchMin = SoundMod->PitchMax = SoundNodeInstance_DEPRECATED->SoundSlots[SlotIndex].PitchScale;
-
-				SoundMod->CreateStartingConnectors();
-				SoundMod->ChildNodes[0] = WavePlayer;
-#if WITH_EDITOR
-				SoundMod->PlaceNode(ColumnCount, SlotIndex, SoundNodeInstance_DEPRECATED->SoundSlots.Num());
-#endif //WITH_EDITOR
-
-				SoundRandom->ChildNodes[SlotIndex] = SoundMod;
-			}
-			else
-			{
-				SoundRandom->ChildNodes[SlotIndex] = WavePlayer;
-			}
-		}
-
-#if WITH_EDITOR
-		WavePlayer->PlaceNode(ColumnCount + 1, SlotIndex, SoundNodeInstance_DEPRECATED->SoundSlots.Num());
-#endif //WITH_EDITOR
-	}
-
-#if WITH_EDITOR
-	SoundCue->LinkGraphNodesFromSoundNodes();
-#endif //WITH_EDITORONLY_DATA
-}
-
-void AAmbientSound::PostLoad()
-{
-	Super::PostLoad();
-
-	if (GetLinkerUE4Version() < VER_UE4_UNIFY_AMBIENT_SOUND_ACTORS)
-	{
-		AudioComponent->ConditionalPostLoad();
-		AudioComponent->bAutoActivate = bAutoPlay_DEPRECATED;
-
-		if (SoundNodeInstance_DEPRECATED)
-		{
-			SoundNodeInstance_DEPRECATED->ConditionalPostLoad();
-
-			AudioComponent->bOverrideAttenuation = true;
-			AudioComponent->AttenuationOverrides.bAttenuate			= SoundNodeInstance_DEPRECATED->bAttenuate;
-			AudioComponent->AttenuationOverrides.bSpatialize		= SoundNodeInstance_DEPRECATED->bSpatialize;
-			AudioComponent->AttenuationOverrides.dBAttenuationAtMax	= SoundNodeInstance_DEPRECATED->dBAttenuationAtMax;
-			AudioComponent->AttenuationOverrides.DistanceAlgorithm	= SoundNodeInstance_DEPRECATED->DistanceModel;
-			AudioComponent->AttenuationOverrides.bAttenuateWithLPF	= SoundNodeInstance_DEPRECATED->bAttenuateWithLPF;
-			AudioComponent->AttenuationOverrides.LPFRadiusMin		= SoundNodeInstance_DEPRECATED->LPFRadiusMin;
-			AudioComponent->AttenuationOverrides.LPFRadiusMax		= SoundNodeInstance_DEPRECATED->LPFRadiusMax;
-
-			AudioComponent->AttenuationOverrides.RadiusMin_DEPRECATED = SoundNodeInstance_DEPRECATED->RadiusMin;
-			AudioComponent->AttenuationOverrides.RadiusMax_DEPRECATED = SoundNodeInstance_DEPRECATED->RadiusMax;
-			AudioComponent->AttenuationOverrides.FalloffDistance	  = SoundNodeInstance_DEPRECATED->RadiusMax - SoundNodeInstance_DEPRECATED->RadiusMin;
-			AudioComponent->AttenuationOverrides.AttenuationShapeExtents.X = SoundNodeInstance_DEPRECATED->RadiusMin;
-			
-			AudioComponent->VolumeModulationMin			= SoundNodeInstance_DEPRECATED->VolumeMin;
-			AudioComponent->VolumeModulationMax			= SoundNodeInstance_DEPRECATED->VolumeMax;
-			AudioComponent->PitchModulationMin			= SoundNodeInstance_DEPRECATED->PitchMin;
-			AudioComponent->PitchModulationMax			= SoundNodeInstance_DEPRECATED->PitchMax;
-
-			if (!bUE4AudioRefactorMigrationUnderway)
-			{
-				MigrateSoundNodeInstance();
-			}
-		}
-	}
-}
-
 void AAmbientSound::FadeIn(float FadeInDuration, float FadeVolumeLevel)
 {
 	if (AudioComponent)
diff --git a/Engine/Source/Runtime/Engine/Private/Animation/AnimEncoding.cpp b/Engine/Source/Runtime/Engine/Private/Animation/AnimEncoding.cpp
--- a/Engine/Source/Runtime/Engine/Private/Animation/AnimEncoding.cpp
+++ b/Engine/Source/Runtime/Engine/Private/Animation/AnimEncoding.cpp
@@ -288,7 +288,6 @@ void AnimEncodingLegacyBase::GetBoneAtom(
  *
  * @param	Seq					An Animation Sequence to contain the read data.
  * @param	MemoryReader		The MemoryReader object to read from.
- * @param	SourceArVersion		The version of the archive that the data is coming from.
  */
 [REDACTED_EMAIL]: Apparent VC10 compiler bug here causes an access violation in optimized builds
 #if _MSC_VER
@@ -296,8 +295,7 @@ void AnimEncodingLegacyBase::GetBoneAtom(
 #endif
 void AnimEncodingLegacyBase::ByteSwapIn(
 	UAnimSequence& Seq, 
-	FMemoryReader& MemoryReader,
-	int32 SourceArVersion)
+	FMemoryReader& MemoryReader)
 {
 	const int32 NumTracks = Seq.CompressedTrackOffsets.Num() / 4;
 
@@ -320,7 +318,7 @@ void AnimEncodingLegacyBase::ByteSwapIn(
 		checkSlow( (OffsetTrans % 4) == 0 && "CompressedByteStream not aligned to four bytes" );
 		uint8* TransTrackData = StreamBase + OffsetTrans;
 		checkSlow(Seq.TranslationCodec != NULL);
-		((AnimEncodingLegacyBase*)Seq.TranslationCodec)->ByteSwapTranslationIn(Seq, MemoryReader, TransTrackData, NumKeysTrans, SourceArVersion);
+		((AnimEncodingLegacyBase*)Seq.TranslationCodec)->ByteSwapTranslationIn(Seq, MemoryReader, TransTrackData, NumKeysTrans);
 
 		// Like the compressed byte stream, pad the serialization stream to four bytes.
 		// As a sanity check, each pad byte can be checked to be the PadSentinel.
@@ -330,7 +328,7 @@ void AnimEncodingLegacyBase::ByteSwapIn(
 		checkSlow( (OffsetRot % 4) == 0 && "CompressedByteStream not aligned to four bytes" );
 		uint8* RotTrackData = StreamBase + OffsetRot;
 		checkSlow(Seq.RotationCodec != NULL);
-		((AnimEncodingLegacyBase*)Seq.RotationCodec)->ByteSwapRotationIn(Seq, MemoryReader, RotTrackData, NumKeysRot, SourceArVersion);
+		((AnimEncodingLegacyBase*)Seq.RotationCodec)->ByteSwapRotationIn(Seq, MemoryReader, RotTrackData, NumKeysRot);
 
 		// Like the compressed byte stream, pad the serialization stream to four bytes.
 		// As a sanity check, each pad byte can be checked to be the PadSentinel.
@@ -345,7 +343,7 @@ void AnimEncodingLegacyBase::ByteSwapIn(
 			checkSlow( (OffsetScale % 4) == 0 && "CompressedByteStream not aligned to four bytes" );
 			uint8* ScaleTrackData = StreamBase + OffsetScale;
 			checkSlow(Seq.ScaleCodec != NULL);
-			((AnimEncodingLegacyBase*)Seq.ScaleCodec)->ByteSwapScaleIn(Seq, MemoryReader, ScaleTrackData, NumKeysScale, SourceArVersion);
+			((AnimEncodingLegacyBase*)Seq.ScaleCodec)->ByteSwapScaleIn(Seq, MemoryReader, ScaleTrackData, NumKeysScale);
 
 			// Like the compressed byte stream, pad the serialization stream to four bytes.
 			// As a sanity check, each pad byte can be checked to be the PadSentinel.
diff --git a/Engine/Source/Runtime/Engine/Private/Animation/AnimEncoding_ConstantKeyLerp.cpp b/Engine/Source/Runtime/Engine/Private/Animation/AnimEncoding_ConstantKeyLerp.cpp
--- a/Engine/Source/Runtime/Engine/Private/Animation/AnimEncoding_ConstantKeyLerp.cpp
+++ b/Engine/Source/Runtime/Engine/Private/Animation/AnimEncoding_ConstantKeyLerp.cpp
@@ -17,14 +17,12 @@
  * @param	MemoryReader	The FMemoryReader to read from.
  * @param	TrackData		The compressed data stream.
  * @param	NumKeys			The number of keys present in the stream.
- * @param	SourceArVersion	The version number of the source archive stream.
  */
 void AEFConstantKeyLerpShared::ByteSwapRotationIn(
 	UAnimSequence& Seq, 
 	FMemoryReader& MemoryReader,
 	uint8*& TrackData,
-	int32 NumKeys,
-	int32 SourceArVersion)
+	int32 NumKeys)
 {
 	// Calculate the effective compression (in a track with only one key, it's always stored lossless)
 	const int32 EffectiveFormat = (NumKeys == 1) ? ACF_Float96NoW : (int32) Seq.RotationCompressionFormat;
@@ -62,8 +60,7 @@ void AEFConstantKeyLerpShared::ByteSwapTranslationIn(
 	UAnimSequence& Seq, 
 	FMemoryReader& MemoryReader,
 	uint8*& TrackData,
-	int32 NumKeys,
-	int32 SourceArVersion)
+	int32 NumKeys)
 {
 	// Calculate the effective compression (in a track with only one key, it's always stored lossless)
 	const int32 EffectiveFormat = (NumKeys == 1) ? ACF_None : (int32) Seq.TranslationCompressionFormat;
@@ -101,8 +98,7 @@ void AEFConstantKeyLerpShared::ByteSwapScaleIn(
 	UAnimSequence& Seq, 
 	FMemoryReader& MemoryReader,
 	uint8*& TrackData,
-	int32 NumKeys,
-	int32 SourceArVersion)
+	int32 NumKeys)
 {
 	// Calculate the effective compression (in a track with only one key, it's always stored lossless)
 	const int32 EffectiveFormat = (NumKeys == 1) ? ACF_None : (int32) Seq.ScaleCompressionFormat;
diff --git a/Engine/Source/Runtime/Engine/Private/Animation/AnimEncoding_PerTrackCompression.cpp b/Engine/Source/Runtime/Engine/Private/Animation/AnimEncoding_PerTrackCompression.cpp
--- a/Engine/Source/Runtime/Engine/Private/Animation/AnimEncoding_PerTrackCompression.cpp
+++ b/Engine/Source/Runtime/Engine/Private/Animation/AnimEncoding_PerTrackCompression.cpp
@@ -380,12 +380,10 @@ void AEFPerTrackCompressionCodec::PreservePadding(uint8*& TrackData, FMemoryArch
  *
  * @param	Seq					An Animation Sequence to contain the read data.
  * @param	MemoryReader		The MemoryReader object to read from.
- * @param	SourceArVersion		The version of the archive that the data is coming from.
  */
 void AEFPerTrackCompressionCodec::ByteSwapIn(
 	UAnimSequence& Seq, 
-	FMemoryReader& MemoryReader,
-	int32 SourceArVersion)
+	FMemoryReader& MemoryReader)
 {
 	int32 OriginalNumBytes = MemoryReader.TotalSize();
 	Seq.CompressedByteStream.Empty(OriginalNumBytes);
diff --git a/Engine/Source/Runtime/Engine/Private/Animation/AnimEncoding_VariableKeyLerp.cpp b/Engine/Source/Runtime/Engine/Private/Animation/AnimEncoding_VariableKeyLerp.cpp
--- a/Engine/Source/Runtime/Engine/Private/Animation/AnimEncoding_VariableKeyLerp.cpp
+++ b/Engine/Source/Runtime/Engine/Private/Animation/AnimEncoding_VariableKeyLerp.cpp
@@ -17,16 +17,14 @@
  * @param	MemoryReader	The FMemoryReader to read from.
  * @param	TrackData		The compressed data stream.
  * @param	NumKeys			The number of keys present in the stream.
- * @param	SourceArVersion	The version number of the source archive stream.
  */
 void AEFVariableKeyLerpShared::ByteSwapRotationIn(
 	UAnimSequence& Seq, 
 	FMemoryReader& MemoryReader,
 	uint8*& TrackData,
-	int32 NumKeys,
-	int32 SourceArVersion)
+	int32 NumKeys)
 {
-	AEFConstantKeyLerpShared::ByteSwapRotationIn(Seq, MemoryReader, TrackData, NumKeys, SourceArVersion);
+	AEFConstantKeyLerpShared::ByteSwapRotationIn(Seq, MemoryReader, TrackData, NumKeys);
 
 	// Load the track table if present
 	if (NumKeys > 1)
@@ -56,10 +54,9 @@ void AEFVariableKeyLerpShared::ByteSwapTranslationIn(
 	UAnimSequence& Seq, 
 	FMemoryReader& MemoryReader,
 	uint8*& TrackData,
-	int32 NumKeys,
-	int32 SourceArVersion)
+	int32 NumKeys)
 {
-	AEFConstantKeyLerpShared::ByteSwapTranslationIn(Seq, MemoryReader, TrackData, NumKeys, SourceArVersion);
+	AEFConstantKeyLerpShared::ByteSwapTranslationIn(Seq, MemoryReader, TrackData, NumKeys);
 
 	// Load the track table if present
 	if (NumKeys > 1)
@@ -89,10 +86,9 @@ void AEFVariableKeyLerpShared::ByteSwapScaleIn(
 	UAnimSequence& Seq, 
 	FMemoryReader& MemoryReader,
 	uint8*& TrackData,
-	int32 NumKeys,
-	int32 SourceArVersion)
+	int32 NumKeys)
 {
-	AEFConstantKeyLerpShared::ByteSwapScaleIn(Seq, MemoryReader, TrackData, NumKeys, SourceArVersion);
+	AEFConstantKeyLerpShared::ByteSwapScaleIn(Seq, MemoryReader, TrackData, NumKeys);
 
 	// Load the track table if present
 	if (NumKeys > 1)
diff --git a/Engine/Source/Runtime/Engine/Private/Animation/AnimMontage.cpp b/Engine/Source/Runtime/Engine/Private/Animation/AnimMontage.cpp
--- a/Engine/Source/Runtime/Engine/Private/Animation/AnimMontage.cpp
+++ b/Engine/Source/Runtime/Engine/Private/Animation/AnimMontage.cpp
@@ -338,23 +338,9 @@ void UAnimMontage::PostLoad()
 	for ( auto SlotIter = SlotAnimTracks.CreateIterator() ; SlotIter ; ++SlotIter)
 	{
 		FAnimTrack & Track = (*SlotIter).AnimTrack;
-		for ( auto SegIter = Track.AnimSegments.CreateIterator() ; SegIter ; ++SegIter )
-		{
-			FAnimSegment & Segment = (*SegIter);
-			if ( Segment.AnimStartOffset_DEPRECATED!=0.f )
-			{
-				Segment.AnimStartTime = Segment.AnimStartOffset_DEPRECATED;
-				Segment.AnimStartOffset_DEPRECATED = 0.f;
-			}
-			if ( Segment.AnimEndOffset_DEPRECATED!=0.f )
-			{
-				Segment.AnimEndTime = Segment.AnimEndOffset_DEPRECATED;
-				Segment.AnimEndOffset_DEPRECATED = 0.f;
-			}
-		}
 		Track.ValidateSegmentTimes();
 
-		float CurrentCalculatedLength = CalculateSequenceLength();
+		const float CurrentCalculatedLength = CalculateSequenceLength();
 
 		if(CurrentCalculatedLength != SequenceLength)
 		{
@@ -367,21 +353,7 @@ void UAnimMontage::PostLoad()
 
 	for (auto& Composite : CompositeSections)
 	{
-		if (Composite.StarTime_DEPRECATED != 0.0f)
-		{
-			Composite.SetTime(Composite.StarTime_DEPRECATED);
-			Composite.StarTime_DEPRECATED = 0.0f;
-		}
-
-		if(Composite.StartTime_DEPRECATED != 0.0f)
-		{
-			Composite.Clear();
-			Composite.LinkMontage(this, Composite.StartTime_DEPRECATED);
-		}
-		else
-	{
-			Composite.LinkMontage(this, Composite.GetTime());
-		}
+		Composite.LinkMontage(this, Composite.GetTime());
 	}
 
 	bool bRootMotionEnabled = bEnableRootMotionTranslation || bEnableRootMotionRotation;
diff --git a/Engine/Source/Runtime/Engine/Private/Animation/AnimSequence.cpp b/Engine/Source/Runtime/Engine/Private/Animation/AnimSequence.cpp
--- a/Engine/Source/Runtime/Engine/Private/Animation/AnimSequence.cpp
+++ b/Engine/Source/Runtime/Engine/Private/Animation/AnimSequence.cpp
@@ -10,7 +10,7 @@
 #include "AnimationUtils.h"
 #include "AnimationRuntime.h"
 #include "TargetPlatform.h"
-#include "Animation/AnimCompress_RevertToRaw.h"
+#include "Animation/AnimCompress.h"
 #include "Animation/AnimNotifies/AnimNotify.h"
 #include "Animation/AnimNotifies/AnimNotifyState.h"
 #include "Animation/Rig.h"
@@ -236,7 +236,7 @@ void UAnimSequence::Serialize(FArchive& Ar)
 
 		// and then use the codecs to byte swap
 		check( RotationCodec != NULL );
-		((AnimEncoding*)RotationCodec)->ByteSwapIn(*this, MemoryReader, Ar.UE3Ver());
+		((AnimEncoding*)RotationCodec)->ByteSwapIn(*this, MemoryReader);
 	}
 	else if( Ar.IsSaving() || Ar.IsCountingMemory() )
 	{
@@ -320,14 +320,6 @@ void UAnimSequence::PostLoad()
 
 	Super::PostLoad();
 
-	// If RAW animation data exists, and needs to be recompressed, do so.
-	if( GIsEditor && GetLinkerUE4Version() < VER_UE4_REMOVE_REDUNDANT_KEY && RawAnimationData.Num() > 0 )
-	{
-		// Recompress Raw Animation data w/ lossless compression.
-		// If some keys have been removed, then recompress animsequence with its original compression algorithm
-		CompressRawAnimData();
-	}
-
 	// Ensure notifies are sorted.
 	SortNotifies();
 
@@ -415,17 +407,6 @@ void UAnimSequence::PostLoad()
 	}
 
 #if WITH_EDITORONLY_DATA
-	// swap out the deprecated revert to raw compression scheme with a least destructive compression scheme
-	if (GIsEditor && CompressionScheme && CompressionScheme->IsA(UDEPRECATED_AnimCompress_RevertToRaw::StaticClass()))
-	{
-		UE_LOG(LogAnimation, Warning, TEXT("AnimSequence %s (%s) uses the deprecated revert to RAW compression scheme. Using least destructive compression scheme instead"), *GetName(), *GetFullName());
-#if 0 //@todoanim: we won't compress in this case for now 
-		USkeletalMesh* DefaultSkeletalMesh = LoadObject<USkeletalMesh>(NULL, *AnimSet->BestRatioSkelMeshName.ToString(), NULL, LOAD_None, NULL);
-		UAnimCompress* NewAlgorithm = ConstructObject<UAnimCompress>( UAnimCompress_LeastDestructive::StaticClass() );
-		NewAlgorithm->Reduce(this, DefaultSkeletalMesh, false);
-#endif
-	}
-
 	bWasCompressedWithoutTranslations = false; //@todoanim: @fixmelh : AnimRotationOnly - GetAnimSet()->bAnimRotationOnly;
 #endif // WITH_EDITORONLY_DATA
 
@@ -437,31 +418,12 @@ void UAnimSequence::PostLoad()
 	VerifyTrackMap();
 #endif
 
-	// save track data for the case
-	FixAdditiveType();
-
 	if( IsRunningGame() )
 	{
 		// this probably will not show newly created animations in PIE but will show them in the game once they have been saved off
 		INC_DWORD_STAT_BY( STAT_AnimationMemory, GetResourceSize(EResourceSizeMode::Exclusive) );
 	}
 
-	// Convert Notifies to new data
-	if( GIsEditor && GetLinkerUE4Version() < VER_UE4_ANIMNOTIFY_NAMECHANGE && Notifies.Num() > 0 )
-	{
-		LOG_SCOPE_VERBOSITY_OVERRIDE(LogAnimation, ELogVerbosity::Warning);
- 		// convert animnotifies
- 		for (int32 I=0; I<Notifies.Num(); ++I)
- 		{
- 			if (Notifies[I].Notify!=NULL)
- 			{
-				FString Label = Notifies[I].Notify->GetClass()->GetName();
-				Label = Label.Replace(TEXT("AnimNotify_"), TEXT(""), ESearchCase::CaseSensitive);
-				Notifies[I].NotifyName = FName(*Label);
- 			}
- 		}
-	}
-
 #if WITH_EDITOR
 	if (USkeleton * CurrentSkeleton = GetSkeleton())
 	{
@@ -547,18 +509,6 @@ void UAnimSequence::VerifyTrackMap()
 }
 
 #endif // WITH_EDITOR
-void UAnimSequence::FixAdditiveType()
-{
-	if( GetLinkerUE4Version() < VER_UE4_ADDITIVE_TYPE_CHANGE )
-	{
-		if ( RefPoseType!=ABPT_None && AdditiveAnimType == AAT_None )
-		{
-			AdditiveAnimType = AAT_LocalSpaceBase;
-			MarkPackageDirty();
-		}
-	}
-}
-
 void UAnimSequence::BeginDestroy()
 {
 	Super::BeginDestroy();
@@ -1714,56 +1664,6 @@ bool UAnimSequence::CopyNotifies(UAnimSequence* SourceAnimSeq, UAnimSequence* De
 	return true;
 }
 
-void UAnimSequence::UpgradeMorphTargetCurves()
-{
-	// call super first, it will convert currently existing curve first
-	Super::UpgradeMorphTargetCurves();
-
-	// this also gets called for the imported data, so don't check version number
-	// until we fix importing data
-	for (auto CurveIter = CurveData_DEPRECATED.CreateConstIterator(); CurveIter; ++CurveIter)
-	{
-		const FCurveTrack& CurveTrack = (*CurveIter);
-
-		{
-			// Add curves into naming system 
-			FSmartNameMapping* NameMapping = GetSkeleton()->SmartNames.GetContainer(USkeleton::AnimCurveMappingName);
-			check(NameMapping); // This name mapping should always be present
-			USkeleton::AnimCurveUID NewId;
-			NameMapping->AddOrFindName(CurveTrack.CurveName, NewId);
-
-			// add it first, it won't add if duplicate
-			RawCurveData.AddCurveData(NewId);
-			
-			// get the last one I just added
-			FFloatCurve * NewCurve = static_cast<FFloatCurve*>(RawCurveData.GetCurveData(NewId, FRawCurveTracks::FloatType));
-			check (NewCurve);
-
-			// since the import code change, this might not match all the time
-			// but the current import should cover 0-(NumFrames-1) for whole SequenceLength
-			//ensure (CurveTrack.CurveWeights.Num() == NumFrames);
-
-			// make sure this matches
-			// the way we import animation should cover 0 - (NumFrames-1) for SequenceLength;
-			// if only 1 frame, it will be written in the first frame
-			float TimeInterval = GetIntervalPerKey(NumFrames, SequenceLength);
-
-			// copy all weights back
-			for (int32 Iter = 0 ; Iter<CurveTrack.CurveWeights.Num(); ++Iter)
-			{
-				// add the new value
-				NewCurve->FloatCurve.AddKey(TimeInterval*Iter, CurveTrack.CurveWeights[Iter]);
-			}
-
-			// set morph target flags on them, not editable by default
-			NewCurve->SetCurveTypeFlags(ACF_DrivesMorphTarget | ACF_TriggerEvent);
-		}
-	}
-
-	// now conversion is done, clear all previous data
-	CurveData_DEPRECATED.Empty();
-}
-
 bool UAnimSequence::IsValidAdditive() const		
 { 
 	if (AdditiveAnimType != AAT_None)
diff --git a/Engine/Source/Runtime/Engine/Private/Animation/AnimSequenceBase.cpp b/Engine/Source/Runtime/Engine/Private/Animation/AnimSequenceBase.cpp
--- a/Engine/Source/Runtime/Engine/Private/Animation/AnimSequenceBase.cpp
+++ b/Engine/Source/Runtime/Engine/Private/Animation/AnimSequenceBase.cpp
@@ -481,21 +481,6 @@ void UAnimSequenceBase::PostLoad()
 	// Convert Notifies to new data
 	if( GIsEditor && Notifies.Num() > 0 )
 	{
-		if(GetLinkerUE4Version() < VER_UE4_ANIMNOTIFY_NAMECHANGE)
-		{
-			LOG_SCOPE_VERBOSITY_OVERRIDE(LogAnimation, ELogVerbosity::Warning);
-			// convert animnotifies
-			for(int32 I = 0; I < Notifies.Num(); ++I)
-			{
-				if(Notifies[I].Notify != NULL)
-				{
-					FString Label = Notifies[I].Notify->GetClass()->GetName();
-					Label = Label.Replace(TEXT("AnimNotify_"), TEXT(""), ESearchCase::CaseSensitive);
-					Notifies[I].NotifyName = FName(*Label);
-				}
-			}
-		}
-
 		if(GetLinkerUE4Version() < VER_UE4_CLEAR_NOTIFY_TRIGGERS)
 		{
 			for(FAnimNotifyEvent Notify : Notifies)
@@ -509,10 +494,6 @@ void UAnimSequenceBase::PostLoad()
 		}
 	}
 
-	if ( GetLinkerUE4Version() < VER_UE4_MORPHTARGET_CURVE_INTEGRATION )
-	{
-		UpgradeMorphTargetCurves();
-	}
 	// Ensure notifies are sorted.
 	SortNotifies();
 
@@ -545,20 +526,6 @@ void UAnimSequenceBase::PostLoad()
 	}
 }
 
-void UAnimSequenceBase::UpgradeMorphTargetCurves()
-{
-	// make sure this doesn't get called by anywhere else or you'll have to check this before calling 
-	if ( GetLinkerUE4Version() < VER_UE4_MORPHTARGET_CURVE_INTEGRATION )
-	{
-		for ( auto CurveIter = RawCurveData.FloatCurves.CreateIterator(); CurveIter; ++CurveIter )
-		{
-			FFloatCurve & Curve = (*CurveIter);
-			// make sure previous curves has editable flag
-			Curve.SetCurveTypeFlag(ACF_DefaultCurve, true);
-		}	
-	}
-}
-
 void UAnimSequenceBase::SortNotifies()
 {
 	struct FCompareFAnimNotifyEvent
diff --git a/Engine/Source/Runtime/Engine/Private/Animation/AnimationUtils.cpp b/Engine/Source/Runtime/Engine/Private/Animation/AnimationUtils.cpp
--- a/Engine/Source/Runtime/Engine/Private/Animation/AnimationUtils.cpp
+++ b/Engine/Source/Runtime/Engine/Private/Animation/AnimationUtils.cpp
@@ -7,7 +7,6 @@
 #include "EnginePrivate.h"
 #include "Animation/AnimCompress_BitwiseCompressOnly.h"
 #include "Animation/AnimCompress_PerTrackCompression.h"
-#include "Animation/AnimCompress_RevertToRaw.h"
 #include "Animation/AnimCompress_LeastDestructive.h"
 #include "Animation/AnimCompress_RemoveEverySecondKey.h"
 #include "Animation/AnimSet.h"
@@ -745,12 +744,7 @@ void FAnimationUtils::CompressAnimSequenceExplicit(
 		{
 			UAnimCompress* OriginalCompressionAlgorithm = AnimSeq->CompressionScheme ? AnimSeq->CompressionScheme : FAnimationUtils::GetDefaultAnimationCompressionAlgorithm();
 
-			if( OriginalCompressionAlgorithm->IsA(UDEPRECATED_AnimCompress_RevertToRaw::StaticClass()) )
-			{
-				UE_LOG(LogAnimation, Warning, TEXT("FAnimationUtils::CompressAnimSequence %s (%s) Not allowed to revert to RAW. Using default compression scheme."), *AnimSeq->GetName(), *AnimSeq->GetFullName());
-				OriginalCompressionAlgorithm = FAnimationUtils::GetDefaultAnimationCompressionAlgorithm();
-			}
-			else if( OriginalCompressionAlgorithm->IsA(UAnimCompress_LeastDestructive::StaticClass()) )
+			if( OriginalCompressionAlgorithm->IsA(UAnimCompress_LeastDestructive::StaticClass()) )
 			{
 				UE_LOG(LogAnimation, Warning, TEXT("FAnimationUtils::CompressAnimSequence %s (%s) Not allowed to least destructive. Using default compression scheme."), *AnimSeq->GetName(), *AnimSeq->GetFullName());
 				OriginalCompressionAlgorithm = FAnimationUtils::GetDefaultAnimationCompressionAlgorithm();
diff --git a/Engine/Source/Runtime/Engine/Private/Animation/BlendSpaceBase.cpp b/Engine/Source/Runtime/Engine/Private/Animation/BlendSpaceBase.cpp
--- a/Engine/Source/Runtime/Engine/Private/Animation/BlendSpaceBase.cpp
+++ b/Engine/Source/Runtime/Engine/Private/Animation/BlendSpaceBase.cpp
@@ -25,34 +25,6 @@ void UBlendSpaceBase::PostLoad()
 
 	ValidateSampleData();
 
-	if ( GetLinkerUE4Version() < VER_UE4_ADDITIVE_TYPE_CHANGE )
-	{
-		// PostLoad isn't guaranteed to be called before getting here
-		// for all SampleData, so make sure you fix them up
-		for (int32 I=0; I<SampleData.Num(); ++I)
-		{
-			if ( SampleData[I].Animation )
-			{
-				SampleData[I].Animation->FixAdditiveType();
-			}
-		}
-
-#if WITH_EDITORONLY_DATA
-		if ( IsValidAdditive() && PreviewBasePose == NULL )
-		{
-			for (int32 I=0; I<SampleData.Num(); ++I)
-			{
-				if ( SampleData[I].Animation && SampleData[I].Animation->RefPoseSeq )
-				{
-					PreviewBasePose = SampleData[I].Animation->RefPoseSeq;
-					MarkPackageDirty();
-					break;
-				}
-			}
-		}
-#endif // WITH_EDITORONLY_DATA
-	}
-
 	InitializePerBoneBlend();
 }
 
diff --git a/Engine/Source/Runtime/Engine/Private/ArrowComponent.cpp b/Engine/Source/Runtime/Engine/Private/ArrowComponent.cpp
--- a/Engine/Source/Runtime/Engine/Private/ArrowComponent.cpp
+++ b/Engine/Source/Runtime/Engine/Private/ArrowComponent.cpp
@@ -284,7 +284,6 @@ UArrowComponent::UArrowComponent(const FObjectInitializer& ObjectInitializer)
 	};
 	static FConstructorStatics ConstructorStatics;
 
-	BodyInstance.bEnableCollision_DEPRECATED = false;
 	SetCollisionProfileName(UCollisionProfile::NoCollision_ProfileName);
 	ArrowColor = FColor(255, 0, 0, 255);
 
diff --git a/Engine/Source/Runtime/Engine/Private/Atmosphere/Atmosphere.cpp b/Engine/Source/Runtime/Engine/Private/Atmosphere/Atmosphere.cpp
--- a/Engine/Source/Runtime/Engine/Private/Atmosphere/Atmosphere.cpp
+++ b/Engine/Source/Runtime/Engine/Private/Atmosphere/Atmosphere.cpp
@@ -664,26 +664,23 @@ void UAtmosphericFogComponent::Serialize(FArchive& Ar)
 		IrradianceData.Serialize(Ar, this);
 	}
 
-	if (Ar.UE4Ver() >= VER_UE4_ATMOSPHERIC_FOG_CACHE_TEXTURE)
-	{
-		InscatterData.Serialize(Ar,this);
+	InscatterData.Serialize(Ar,this);
 
-		if (Ar.IsLoading())
-		{
-			int32 CounterVal;
-			Ar << CounterVal;
-			// Precomputation was not successful, just ignore it
-			if (CounterVal < EValid)
-			{
-				CounterVal = EInvalid;
-			}
-			PrecomputeCounter.Set(CounterVal);
-		}
-		else
+	if (Ar.IsLoading())
+	{
+		int32 CounterVal;
+		Ar << CounterVal;
+		// Precomputation was not successful, just ignore it
+		if (CounterVal < EValid)
 		{
-			int32 CounterVal = PrecomputeCounter.GetValue();
-			Ar << CounterVal;
+			CounterVal = EInvalid;
 		}
+		PrecomputeCounter.Set(CounterVal);
+	}
+	else
+	{
+		int32 CounterVal = PrecomputeCounter.GetValue();
+		Ar << CounterVal;
 	}
 
 	if (Ar.IsLoading() && Ar.UE4Ver() < VER_UE4_ATMOSPHERIC_FOG_CACHE_DATA && PrecomputeCounter.GetValue() == EValid) 
diff --git a/Engine/Source/Runtime/Engine/Private/AudioComponent.cpp b/Engine/Source/Runtime/Engine/Private/AudioComponent.cpp
--- a/Engine/Source/Runtime/Engine/Private/AudioComponent.cpp
+++ b/Engine/Source/Runtime/Engine/Private/AudioComponent.cpp
@@ -13,7 +13,6 @@ UAudioComponent::UAudioComponent(const FObjectInitializer& ObjectInitializer)
 	: Super(ObjectInitializer)
 {
 	bAutoDestroy = false;
-	bAutoPlay_DEPRECATED = false;
 	bAutoActivate = true;
 	bAllowSpatialization = true;
 	bStopWhenOwnerDestroyed = true;
@@ -388,16 +387,6 @@ void UAudioComponent::CollectAttenuationShapesForVisualization(TMultiMap<EAttenu
 	}
 }
 
-void UAudioComponent::PostLoad()
-{
-	Super::PostLoad();
-
-	if (GetLinkerUE4Version() < VER_UE4_CONFORM_COMPONENT_ACTIVATE_FLAG)
-	{
-		bAutoActivate = bAutoPlay_DEPRECATED;
-	}
-}
-
 void UAudioComponent::Activate(bool bReset)
 {
 	if (bReset || ShouldActivate()==true)
diff --git a/Engine/Source/Runtime/Engine/Private/AudioVolume.cpp b/Engine/Source/Runtime/Engine/Private/AudioVolume.cpp
--- a/Engine/Source/Runtime/Engine/Private/AudioVolume.cpp
+++ b/Engine/Source/Runtime/Engine/Private/AudioVolume.cpp
@@ -137,8 +137,6 @@ bool FInteriorSettings::operator!=(const FInteriorSettings& Other) const
 AAudioVolume::AAudioVolume(const FObjectInitializer& ObjectInitializer)
 	: Super(ObjectInitializer)
 {
-	GetBrushComponent()->BodyInstance.bEnableCollision_DEPRECATED = false;
-
 	GetBrushComponent()->SetCollisionProfileName(UCollisionProfile::NoCollision_ProfileName);
 	GetBrushComponent()->bAlwaysCreatePhysicsState = true;
 
diff --git a/Engine/Source/Runtime/Engine/Private/BillboardComponent.cpp b/Engine/Source/Runtime/Engine/Private/BillboardComponent.cpp
--- a/Engine/Source/Runtime/Engine/Private/BillboardComponent.cpp
+++ b/Engine/Source/Runtime/Engine/Private/BillboardComponent.cpp
@@ -246,7 +246,6 @@ UBillboardComponent::UBillboardComponent(const FObjectInitializer& ObjectInitial
 	};
 	static FConstructorStatics ConstructorStatics;
 
-	BodyInstance.bEnableCollision_DEPRECATED = false;
 	SetCollisionProfileName(UCollisionProfile::NoCollision_ProfileName);
 	Sprite = ConstructorStatics.SpriteTexture.Object;
 	bAbsoluteScale = true;
diff --git a/Engine/Source/Runtime/Engine/Private/BlockingVolume.cpp b/Engine/Source/Runtime/Engine/Private/BlockingVolume.cpp
--- a/Engine/Source/Runtime/Engine/Private/BlockingVolume.cpp
+++ b/Engine/Source/Runtime/Engine/Private/BlockingVolume.cpp
@@ -10,7 +10,6 @@ ABlockingVolume::ABlockingVolume(const FObjectInitializer& ObjectInitializer)
 	: Super(ObjectInitializer)
 {
 	GetBrushComponent()->bCanEverAffectNavigation = true;
-	GetBrushComponent()->BodyInstance.bEnableCollision_DEPRECATED = true;
 	GetBrushComponent()->SetCollisionProfileName(InvisibleWall_NAME);
 }
 
diff --git a/Engine/Source/Runtime/Engine/Private/Blueprint.cpp b/Engine/Source/Runtime/Engine/Private/Blueprint.cpp
--- a/Engine/Source/Runtime/Engine/Private/Blueprint.cpp
+++ b/Engine/Source/Runtime/Engine/Private/Blueprint.cpp
@@ -190,47 +190,42 @@ void UBlueprintCore::Serialize(FArchive& Ar)
 {
 	Super::Serialize(Ar);
 
-	if (Ar.UE4Ver() >= VER_UE4_BLUEPRINT_NATIVE_SERIALIZATION)
-	{
-		if (Ar.UE4Ver() >= VER_UE4_BLUEPRINT_CDO_MIGRATION)
-		{
-			Ar << bLegacyGeneratedClassIsAuthoritative;	
-		}
+	Ar << bLegacyGeneratedClassIsAuthoritative;	
 
-		if ((Ar.UE4Ver() < VER_UE4_BLUEPRINT_SKEL_CLASS_TRANSIENT_AGAIN)
-			&& (Ar.UE4Ver() != VER_UE4_BLUEPRINT_SKEL_TEMPORARY_TRANSIENT))
+	if ((Ar.UE4Ver() < VER_UE4_BLUEPRINT_SKEL_CLASS_TRANSIENT_AGAIN)
+		&& (Ar.UE4Ver() != VER_UE4_BLUEPRINT_SKEL_TEMPORARY_TRANSIENT))
+	{
+		Ar << SkeletonGeneratedClass;
+		if( SkeletonGeneratedClass )
 		{
-			Ar << SkeletonGeneratedClass;
-			if( SkeletonGeneratedClass )
+			// If we serialized in a skeleton class, make sure it and all its children are updated to be transient
+			SkeletonGeneratedClass->SetFlags(RF_Transient);
+			TArray<UObject*> SubObjs;
+			GetObjectsWithOuter(SkeletonGeneratedClass, SubObjs, true);
+			for(auto SubObjIt = SubObjs.CreateIterator(); SubObjIt; ++SubObjIt)
 			{
-				// If we serialized in a skeleton class, make sure it and all its children are updated to be transient
-				SkeletonGeneratedClass->SetFlags(RF_Transient);
-				TArray<UObject*> SubObjs;
-				GetObjectsWithOuter(SkeletonGeneratedClass, SubObjs, true);
-				for(auto SubObjIt = SubObjs.CreateIterator(); SubObjIt; ++SubObjIt)
-				{
-					(*SubObjIt)->SetFlags(RF_Transient);
-				}
+				(*SubObjIt)->SetFlags(RF_Transient);
 			}
+		}
 
-			// We only want to serialize in the GeneratedClass if the SkeletonClass didn't trigger a recompile
-			bool bSerializeGeneratedClass = true;
-			if (UBlueprint* BP = Cast<UBlueprint>(this))
-			{
-				bSerializeGeneratedClass = !Ar.IsLoading() || !BP->bHasBeenRegenerated;
-			}
+		// We only want to serialize in the GeneratedClass if the SkeletonClass didn't trigger a recompile
+		bool bSerializeGeneratedClass = true;
+		if (UBlueprint* BP = Cast<UBlueprint>(this))
+		{
+			bSerializeGeneratedClass = !Ar.IsLoading() || !BP->bHasBeenRegenerated;
+		}
 
-			if (bSerializeGeneratedClass)
-			{
-				Ar << GeneratedClass;
-			}
-			else if (Ar.IsLoading())
-			{
-				UClass* DummyClass = NULL;
-				Ar << DummyClass;
-			}
+		if (bSerializeGeneratedClass)
+		{
+			Ar << GeneratedClass;
+		}
+		else if (Ar.IsLoading())
+		{
+			UClass* DummyClass = NULL;
+			Ar << DummyClass;
 		}
 	}
+
 	if( Ar.ArIsLoading && !BlueprintGuid.IsValid() )
 	{
 		GenerateDeterministicGuid();
@@ -299,15 +294,6 @@ void UBlueprint::Serialize(FArchive& Ar)
 		}
 	}
 
-	if(Ar.IsLoading() && Ar.UE4Ver() < VER_UE4_ADD_KISMETVISIBLE)
-	{
-		for (int32 i = 0; i < NewVariables.Num(); ++i)
-		{
-			FBPVariableDescription& Variable = NewVariables[i];
-			Variable.PropertyFlags |= CPF_BlueprintVisible;
-		}
-	}
-
 	if (Ar.UE4Ver() < VER_UE4_K2NODE_REFERENCEGUIDS)
 	{
 		for (int32 Index = 0; Index < NewVariables.Num(); ++Index)
diff --git a/Engine/Source/Runtime/Engine/Private/Brush.cpp b/Engine/Source/Runtime/Engine/Private/Brush.cpp
--- a/Engine/Source/Runtime/Engine/Private/Brush.cpp
+++ b/Engine/Source/Runtime/Engine/Private/Brush.cpp
@@ -167,19 +167,6 @@ void ABrush::PostLoad()
 {
 	Super::PostLoad();
 
-	if( GetLinkerUE4Version() < VER_UE4_FIX_BSP_BRUSH_TYPE && BrushType == Brush_Default )
-	{
-		ECsgOper Oper = CsgOper_DEPRECATED;
-		if( Oper == CSG_Add )
-		{
-			BrushType = Brush_Add;
-		}
-		else if( Oper == CSG_Subtract )
-		{
-			BrushType = Brush_Subtract;
-		}
-	}
-
 #if WITH_EDITOR
 	if (BrushBuilder && BrushBuilder->GetOuter() != this)
 	{
diff --git a/Engine/Source/Runtime/Engine/Private/BrushShape.cpp b/Engine/Source/Runtime/Engine/Private/BrushShape.cpp
--- a/Engine/Source/Runtime/Engine/Private/BrushShape.cpp
+++ b/Engine/Source/Runtime/Engine/Private/BrushShape.cpp
@@ -7,8 +7,6 @@
 ABrushShape::ABrushShape(const FObjectInitializer& ObjectInitializer)
 	: Super(ObjectInitializer)
 {
-
-	GetBrushComponent()->BodyInstance.bEnableCollision_DEPRECATED = false;
 	GetBrushComponent()->SetCollisionProfileName(UCollisionProfile::NoCollision_ProfileName);
 	GetBrushComponent()->AlwaysLoadOnClient = true;
 	GetBrushComponent()->AlwaysLoadOnServer = false;
diff --git a/Engine/Source/Runtime/Engine/Private/Character.cpp b/Engine/Source/Runtime/Engine/Private/Character.cpp
--- a/Engine/Source/Runtime/Engine/Private/Character.cpp
+++ b/Engine/Source/Runtime/Engine/Private/Character.cpp
@@ -48,7 +48,6 @@ ACharacter::ACharacter(const FObjectInitializer& ObjectInitializer)
 
 	CapsuleComponent = ObjectInitializer.CreateDefaultSubobject<UCapsuleComponent>(this, ACharacter::CapsuleComponentName);
 	CapsuleComponent->InitCapsuleSize(34.0f, 88.0f);
-	CapsuleComponent->BodyInstance.bEnableCollision_DEPRECATED = true;
 
 	static FName CollisionProfileName(TEXT("Pawn"));
 	CapsuleComponent->SetCollisionProfileName(CollisionProfileName);
diff --git a/Engine/Source/Runtime/Engine/Private/CullDistanceVolume.cpp b/Engine/Source/Runtime/Engine/Private/CullDistanceVolume.cpp
--- a/Engine/Source/Runtime/Engine/Private/CullDistanceVolume.cpp
+++ b/Engine/Source/Runtime/Engine/Private/CullDistanceVolume.cpp
@@ -7,7 +7,6 @@
 ACullDistanceVolume::ACullDistanceVolume(const FObjectInitializer& ObjectInitializer)
 	: Super(ObjectInitializer)
 {
-	GetBrushComponent()->BodyInstance.bEnableCollision_DEPRECATED = false;
 	GetBrushComponent()->SetCollisionProfileName(UCollisionProfile::NoCollision_ProfileName);
 	GetBrushComponent()->bAlwaysCreatePhysicsState = true;
 
diff --git a/Engine/Source/Runtime/Engine/Private/CurveBase.cpp b/Engine/Source/Runtime/Engine/Private/CurveBase.cpp
--- a/Engine/Source/Runtime/Engine/Private/CurveBase.cpp
+++ b/Engine/Source/Runtime/Engine/Private/CurveBase.cpp
@@ -750,33 +750,6 @@ float FRichCurve::Eval( const float InTime, float DefaultValue ) const
 	return Keys[NumKeys-1].Value;
 }
 
-void FRichCurve::ConvertInterpCurveFloat(const FInterpCurveFloat& InCurve, FRichCurve& OutCurve)
-{
-	const int32 NumInKeys = InCurve.Points.Num();
-
-	OutCurve.Keys.Empty();
-	OutCurve.Keys.AddUninitialized(NumInKeys);
-	for(int32 i=0; i<NumInKeys; i++)
-	{
-		OutCurve.Keys[i] = FRichCurveKey(InCurve.Points[i]);
-	}
-}
-
-void FRichCurve::ConvertInterpCurveVector(const FInterpCurveVector& InCurve, FRichCurve OutCurves[3])
-{
-	const int32 NumInKeys = InCurve.Points.Num();
-
-	for(int32 AxisIdx = 0; AxisIdx < 3; AxisIdx++)
-	{
-		OutCurves[AxisIdx].Keys.Empty();
-		OutCurves[AxisIdx].Keys.AddUninitialized(NumInKeys);
-		for(int32 i=0; i<NumInKeys; i++)
-		{
-			OutCurves[AxisIdx].Keys[i] = FRichCurveKey(InCurve.Points[i], AxisIdx);
-		}
-	}
-}
-
 bool FRichCurve::operator==( const FRichCurve& Curve ) const
 {
 	if(Keys.Num() != Curve.Keys.Num())
diff --git a/Engine/Source/Runtime/Engine/Private/CurveFloat.cpp b/Engine/Source/Runtime/Engine/Private/CurveFloat.cpp
--- a/Engine/Source/Runtime/Engine/Private/CurveFloat.cpp
+++ b/Engine/Source/Runtime/Engine/Private/CurveFloat.cpp
@@ -45,16 +45,6 @@ float UCurveFloat::GetFloatValue( float InTime ) const
 	return FloatCurve.Eval(InTime);
 }
 
-void UCurveFloat::PostLoad()
-{
-	Super::PostLoad();
-
-	if(GetLinkerUE4Version() < VER_UE4_UCURVE_USING_RICHCURVES)
-	{
-		FRichCurve::ConvertInterpCurveFloat(FloatKeys_DEPRECATED, FloatCurve);
-	}
-}
-
 TArray<FRichCurveEditInfoConst> UCurveFloat::GetCurves() const
 {
 	TArray<FRichCurveEditInfoConst> Curves;
diff --git a/Engine/Source/Runtime/Engine/Private/CurveVector.cpp b/Engine/Source/Runtime/Engine/Private/CurveVector.cpp
--- a/Engine/Source/Runtime/Engine/Private/CurveVector.cpp
+++ b/Engine/Source/Runtime/Engine/Private/CurveVector.cpp
@@ -22,16 +22,6 @@ FVector UCurveVector::GetVectorValue( float InTime ) const
 	return Result;
 }
 
-void UCurveVector::PostLoad()
-{
-	Super::PostLoad();
-
-	if(GetLinkerUE4Version() < VER_UE4_UCURVE_USING_RICHCURVES)
-	{
-		FRichCurve::ConvertInterpCurveVector(VectorKeys_DEPRECATED, FloatCurves);
-	}
-}
-
 static const FName XCurveName(TEXT("X"));
 static const FName YCurveName(TEXT("Y"));
 static const FName ZCurveName(TEXT("Z"));
diff --git a/Engine/Source/Runtime/Engine/Private/DecalActor.cpp b/Engine/Source/Runtime/Engine/Private/DecalActor.cpp
--- a/Engine/Source/Runtime/Engine/Private/DecalActor.cpp
+++ b/Engine/Source/Runtime/Engine/Private/DecalActor.cpp
@@ -31,7 +31,6 @@ ADecalActor::ADecalActor(const FObjectInitializer& ObjectInitializer)
 	BoxComponent = ObjectInitializer.CreateEditorOnlyDefaultSubobject<UBoxComponent>(this, TEXT("DrawBox0"));
 	if (BoxComponent != nullptr)
 	{
-		BoxComponent->BodyInstance.bEnableCollision_DEPRECATED = false;
 		BoxComponent->SetCollisionProfileName(UCollisionProfile::NoCollision_ProfileName);
 
 		BoxComponent->ShapeColor = FColor(80, 80, 200, 255);
diff --git a/Engine/Source/Runtime/Engine/Private/DefaultPawn.cpp b/Engine/Source/Runtime/Engine/Private/DefaultPawn.cpp
--- a/Engine/Source/Runtime/Engine/Private/DefaultPawn.cpp
+++ b/Engine/Source/Runtime/Engine/Private/DefaultPawn.cpp
@@ -26,7 +26,6 @@ ADefaultPawn::ADefaultPawn(const FObjectInitializer& ObjectInitializer)
 
 	CollisionComponent = ObjectInitializer.CreateDefaultSubobject<USphereComponent>(this, ADefaultPawn::CollisionComponentName);
 	CollisionComponent->InitSphereRadius(35.0f);
-	CollisionComponent->BodyInstance.bEnableCollision_DEPRECATED = true;
 
 	static FName CollisionProfileName(TEXT("Pawn"));
 	CollisionComponent->SetCollisionProfileName(CollisionProfileName);
diff --git a/Engine/Source/Runtime/Engine/Private/DestructibleComponent.cpp b/Engine/Source/Runtime/Engine/Private/DestructibleComponent.cpp
--- a/Engine/Source/Runtime/Engine/Private/DestructibleComponent.cpp
+++ b/Engine/Source/Runtime/Engine/Private/DestructibleComponent.cpp
@@ -34,7 +34,6 @@ UDestructibleComponent::UDestructibleComponent(const FObjectInitializer& ObjectI
 	bHasCustomNavigableGeometry = EHasCustomNavigableGeometry::EvenIfNotCollidable;
 
 	BodyInstance.SetUseAsyncScene(true);
-	BodyInstance.bEnableCollision_DEPRECATED = true;
 	static FName CollisionProfileName(TEXT("Destructible"));
 	SetCollisionProfileName(CollisionProfileName);
 
diff --git a/Engine/Source/Runtime/Engine/Private/DestructibleFractureSettings.cpp b/Engine/Source/Runtime/Engine/Private/DestructibleFractureSettings.cpp
--- a/Engine/Source/Runtime/Engine/Private/DestructibleFractureSettings.cpp
+++ b/Engine/Source/Runtime/Engine/Private/DestructibleFractureSettings.cpp
@@ -257,19 +257,6 @@ void UDestructibleFractureSettings::Serialize(FArchive& Ar)
 			physx::PxFileBuf* Stream = GApexSDK->createMemoryReadStream(Buffer.GetData(), Size);
 			if (Stream != NULL)
 			{
-				if (Ar.UE4Ver() < VER_UE4_NX_DESTRUCTIBLE_ASSET_AUTHORING_LOAD_FIX)
-				{
-					// Create an NxParameterized serializer
-					NxParameterized::Serializer* Serializer = GApexSDK->createSerializer(NxParameterized::Serializer::NST_BINARY);
-					if (Serializer != NULL)
-					{
-						// Deserialize into a DeserializedData buffer
-						NxParameterized::Serializer::DeserializedData DeserializedData;
-						Serializer->deserialize(*Stream, DeserializedData);
-						// Release the NxParameterized serializer
-						Serializer->release();
-					}
-				}
 				// Now deserialize the ExplicitHierarchicalMeshes
 				FExplicitHierarchicalMeshEmbedding ExplicitHierarchicalMeshEmbedding;
 				ApexDestructibleAssetAuthoring->getExplicitHierarchicalMesh().deserialize(*Stream, ExplicitHierarchicalMeshEmbedding);
