Bring minimum engine version up to 214

diff --git a/Engine/Source/Runtime/Engine/Private/DestructibleMesh.cpp b/Engine/Source/Runtime/Engine/Private/DestructibleMesh.cpp
--- a/Engine/Source/Runtime/Engine/Private/DestructibleMesh.cpp
+++ b/Engine/Source/Runtime/Engine/Private/DestructibleMesh.cpp
@@ -32,14 +32,6 @@ void UDestructibleMesh::PostLoad()
 	CreateBodySetup();
 
 	// Fix location of destructible physical material
-	if (GetLinker() && (GetLinker()->UE4Ver() < VER_UE4_DESTRUCTIBLE_MESH_BODYSETUP_HOLDS_PHYSICAL_MATERIAL))
-	{
-		if (BodySetup != NULL)
-		{
-			BodySetup->PhysMaterial = DestructiblePhysicalMaterial_DEPRECATED;
-		}
-	}
-
 #if WITH_APEX
 	// destructible mesh needs to re-evaluate the number of max gpu bone count for each chunk
 	// we'll have to rechunk if that exceeds project setting
@@ -174,47 +166,6 @@ void UDestructibleMesh::Serialize(FArchive& Ar)
 			// Release our temporary objects
 			Serializer->release();
 			GApexSDK->releaseMemoryReadStream( *Stream );
-
-			// Fix chunk bounds which were not transformed upon import before version VER_UE4_NX_DESTRUCTIBLE_ASSET_CHUNK_BOUNDS_FIX
-			if (ApexDestructibleAsset != NULL && Ar.UE4Ver() < VER_UE4_NX_DESTRUCTIBLE_ASSET_CHUNK_BOUNDS_FIX)
-			{
-				// Name buffer used for parameterized array element lookup
-				char ArrayElementName[1024];
-
-				// Get the NxParameterized interface to the asset
-				NxParameterized::Interface* AssetParams = const_cast<NxParameterized::Interface*>(ApexDestructibleAsset->getAssetNxParameterized());
-				if (AssetParams != NULL)
-				{
-					/* render mesh asset (bounding boxes only) */
-					const physx::PxMat33 Basis = physx::PxMat33::createDiagonal(physx::PxVec3(1.0f, -1.0f, 1.0f));
-					NxParameterized::Interface* RenderMeshAssetParams;
-					const bool bFoundRMA = NxParameterized::getParamRef(*AssetParams, "renderMeshAsset", RenderMeshAssetParams);
-					check(bFoundRMA);
-					if (bFoundRMA)
-					{
-						physx::PxI32 PartBoundsCount;
-						const bool bFoundBounds = NxParameterized::getParamArraySize(*RenderMeshAssetParams, "partBounds", PartBoundsCount);
-						check(bFoundBounds);
-						if (bFoundBounds)
-						{
-							for (physx::PxI32 i = 0; i < PartBoundsCount; ++i)
-							{
-								FCStringAnsi::Sprintf( ArrayElementName, "partBounds[%d]", i );
-								NxParameterized::Handle PartBoundsHandle(*RenderMeshAssetParams);
-								bool bFoundBoundsHandle = NxParameterized::findParam(*RenderMeshAssetParams, ArrayElementName, PartBoundsHandle) != NULL;
-								check(bFoundBoundsHandle);
-								if (bFoundBoundsHandle)
-								{
-									physx::PxBounds3 PartBounds;
-									PartBoundsHandle.getParamBounds3( PartBounds );
-									PartBounds = physx::PxBounds3::transformSafe(Basis, PartBounds);
-									PartBoundsHandle.setParamBounds3( PartBounds );
-								}
-							}
-						}
-					}
-				}
-			}
 #endif
 		}
 	}
diff --git a/Engine/Source/Runtime/Engine/Private/Distributions.cpp b/Engine/Source/Runtime/Engine/Private/Distributions.cpp
--- a/Engine/Source/Runtime/Engine/Private/Distributions.cpp
+++ b/Engine/Source/Runtime/Engine/Private/Distributions.cpp
@@ -5,7 +5,6 @@
 =============================================================================*/
 
 #include "EnginePrivate.h"
-#include "DistributionHelpers.h"
 #include "Distributions/DistributionVectorParameterBase.h"
 #include "Distributions/DistributionVectorConstantCurve.h"
 #include "Distributions/DistributionVectorUniformCurve.h"
@@ -21,7 +20,7 @@
 // Moving UDistributions to PostInitProps to not be default sub-objects:
 // Small enough value to be rounded to 0.0 in the editor 
 // but significant enough to properly detect uninitialized defaults.
-const float FDistributionHelpers::DefaultValue = 1.2345E-20f;
+const float UDistribution::DefaultValue = 1.2345E-20f;
 
 ENGINE_API uint32 GDistributionType = 1;
 
@@ -1562,14 +1561,14 @@ void UDistributionFloatConstant::PostInitProperties()
 	{
 		// Set to a bogus value for distributions created before VER_UE4_MOVE_DISTRIBUITONS_TO_POSTINITPROPS
 		// to be able to restore to the previous default value.
-		Constant = FDistributionHelpers::DefaultValue;
+		Constant = UDistribution::DefaultValue;
 	}
 }
 
 void UDistributionFloatConstant::PostLoad()
 {
 	Super::PostLoad();
-	if (Constant == FDistributionHelpers::DefaultValue)
+	if (Constant == UDistribution::DefaultValue)
 	{
 		// Reset to default for distributions created before VER_UE4_MOVE_DISTRIBUITONS_TO_POSTINITPROPS.
 		Constant = 0.0f;
@@ -1855,8 +1854,8 @@ void UDistributionFloatUniform::PostInitProperties()
 	{
 		// Set to a bogus value for distributions created before VER_UE4_MOVE_DISTRIBUITONS_TO_POSTINITPROPS
 		// to be able to restore to the previous default value.
-		Min = FDistributionHelpers::DefaultValue;
-		Max = FDistributionHelpers::DefaultValue;
+		Min = UDistribution::DefaultValue;
+		Max = UDistribution::DefaultValue;
 	}
 }
 
@@ -1864,11 +1863,11 @@ void UDistributionFloatUniform::PostLoad()
 {
 	Super::PostLoad();
 	// Reset to default for distributions created before VER_UE4_MOVE_DISTRIBUITONS_TO_POSTINITPROPS.
-	if (Min == FDistributionHelpers::DefaultValue)
+	if (Min == UDistribution::DefaultValue)
 	{
 		Min = 0.0f;
 	}
-	if (Max == FDistributionHelpers::DefaultValue)
+	if (Max == UDistribution::DefaultValue)
 	{
 		Max = 0.0f;
 	}
@@ -2326,7 +2325,7 @@ void UDistributionVectorConstant::PostInitProperties()
 	{
 		// Set to a bogus value for distributions created before VER_UE4_MOVE_DISTRIBUITONS_TO_POSTINITPROPS
 		// to be able to restore to the previous default value.
-		Constant = FVector(FDistributionHelpers::DefaultValue);
+		Constant = FVector(UDistribution::DefaultValue);
 	}
 }
 
@@ -2334,7 +2333,7 @@ void UDistributionVectorConstant::PostLoad()
 {
 	Super::PostLoad();
 	// Reset to default for distributions created before VER_UE4_MOVE_DISTRIBUITONS_TO_POSTINITPROPS.
-	if (Constant == FVector(FDistributionHelpers::DefaultValue))
+	if (Constant == FVector(UDistribution::DefaultValue))
 	{
 		Constant = FVector::ZeroVector;
 	}
@@ -2933,8 +2932,8 @@ void UDistributionVectorUniform::PostInitProperties()
 	{
 		// Set to a bogus value for distributions created before VER_UE4_MOVE_DISTRIBUITONS_TO_POSTINITPROPS
 		// to be able to restore to the previous default value.
-		Min = FVector(FDistributionHelpers::DefaultValue);
-		Max = FVector(FDistributionHelpers::DefaultValue);
+		Min = FVector(UDistribution::DefaultValue);
+		Max = FVector(UDistribution::DefaultValue);
 	}
 }
 
@@ -2942,11 +2941,11 @@ void UDistributionVectorUniform::PostLoad()
 {
 	Super::PostLoad();
 	// Reset to default for distributions created before VER_UE4_MOVE_DISTRIBUITONS_TO_POSTINITPROPS.
-	if (Min == FVector(FDistributionHelpers::DefaultValue))
+	if (Min == FVector(UDistribution::DefaultValue))
 	{
 		Min = FVector::ZeroVector;
 	}
-	if (Max == FVector(FDistributionHelpers::DefaultValue))
+	if (Max == FVector(UDistribution::DefaultValue))
 	{
 		Max = FVector::ZeroVector;
 	}
diff --git a/Engine/Source/Runtime/Engine/Private/EdGraph/EdGraphNode.cpp b/Engine/Source/Runtime/Engine/Private/EdGraph/EdGraphNode.cpp
--- a/Engine/Source/Runtime/Engine/Private/EdGraph/EdGraphNode.cpp
+++ b/Engine/Source/Runtime/Engine/Private/EdGraph/EdGraphNode.cpp
@@ -240,11 +240,7 @@ void UEdGraphNode::PostLoad()
 	// Create Guid if not present (and not CDO)
 	if(!NodeGuid.IsValid() && !IsTemplate() && GetLinker() && GetLinker()->IsPersistent() && GetLinker()->IsLoading())
 	{
-		// _Should_ have a guid on all nodes after this version
-		if(GetLinkerUE4Version() >= VER_UE4_ADD_EDGRAPHNODE_GUID)
-		{
-			UE_LOG(LogBlueprint, Warning, TEXT("Node '%s' missing NodeGuid."), *GetPathName());
-		}
+		UE_LOG(LogBlueprint, Warning, TEXT("Node '%s' missing NodeGuid."), *GetPathName());
 
 		// Generate new one
 		CreateNewGuid();
diff --git a/Engine/Source/Runtime/Engine/Private/EngineUtils.cpp b/Engine/Source/Runtime/Engine/Private/EngineUtils.cpp
--- a/Engine/Source/Runtime/Engine/Private/EngineUtils.cpp
+++ b/Engine/Source/Runtime/Engine/Private/EngineUtils.cpp
@@ -477,11 +477,11 @@ void FConsoleOutputDevice::Serialize(const TCHAR* Text, ELogVerbosity::Type Verb
 /*-----------------------------------------------------------------------------
 	Serialized data stripping.
 -----------------------------------------------------------------------------*/
-FStripDataFlags::FStripDataFlags( class FArchive& Ar, uint8 InClassFlags /*= 0*/, int32 InVersion /*= VER_UE4_REMOVED_STRIP_DATA */ )
+FStripDataFlags::FStripDataFlags( class FArchive& Ar, uint8 InClassFlags /*= 0*/, int32 InVersion /*= VER_UE4_OLDEST_LOADABLE_PACKAGE */ )
 	: GlobalStripFlags( 0 )
 	, ClassStripFlags( 0 )
 {
-	check(InVersion >= VER_UE4_REMOVED_STRIP_DATA);
+	check(InVersion >= VER_UE4_OLDEST_LOADABLE_PACKAGE);
 	if (Ar.UE4Ver() >= InVersion)
 	{
 		if (Ar.IsCooking())
@@ -497,11 +497,11 @@ FStripDataFlags::FStripDataFlags( class FArchive& Ar, uint8 InClassFlags /*= 0*/
 	}
 }
 
-FStripDataFlags::FStripDataFlags( class FArchive& Ar, uint8 InGlobalFlags, uint8 InClassFlags, int32 InVersion /*= VER_UE4_REMOVED_STRIP_DATA */ )
+FStripDataFlags::FStripDataFlags( class FArchive& Ar, uint8 InGlobalFlags, uint8 InClassFlags, int32 InVersion /*= VER_UE4_OLDEST_LOADABLE_PACKAGE */ )
 	: GlobalStripFlags( 0 )
 	, ClassStripFlags( 0 )
 {
-	check(InVersion >= VER_UE4_REMOVED_STRIP_DATA);
+	check(InVersion >= VER_UE4_OLDEST_LOADABLE_PACKAGE);
 	if (Ar.UE4Ver() >= InVersion)
 	{
 		if (Ar.IsCooking())
diff --git a/Engine/Source/Runtime/Engine/Private/GameplayStatics.cpp b/Engine/Source/Runtime/Engine/Private/GameplayStatics.cpp
--- a/Engine/Source/Runtime/Engine/Private/GameplayStatics.cpp
+++ b/Engine/Source/Runtime/Engine/Private/GameplayStatics.cpp
@@ -794,34 +794,6 @@ class UAudioComponent* UGameplayStatics::PlayDialogueAttached(class UDialogueWav
 	return AudioComponent;
 }
 
-void UGameplayStatics::PlaySound(UObject* WorldContextObject, USoundCue* InSoundCue, USceneComponent* AttachComponent, FName AttachName, bool bFollow, float VolumeMultiplier, float PitchMultiplier)
-{
-	if (bFollow)
-	{
-		PlaySoundAttached(InSoundCue, AttachComponent, AttachName, FVector::ZeroVector, EAttachLocation::KeepRelativeOffset, false, VolumeMultiplier, PitchMultiplier);
-	}
-	else
-	{
-		if ( AttachComponent == NULL )
-		{
-			UE_LOG(LogScript, Warning, TEXT("UGameplayStatics::PlaySound: NULL AttachComponent specified!"));
-			return;
-		}
-		FVector SoundLocation;
-		if (AttachName == NAME_None)
-		{
-			SoundLocation = AttachComponent->GetComponentLocation();
-		}
-		else
-		{
-			FTransform const SocketTransform = AttachComponent->GetSocketTransform(AttachName,RTS_World);
-			SoundLocation = SocketTransform.GetTranslation();
-		}
-		UWorld* World = GEngine->GetWorldFromContextObject(WorldContextObject);
-		PlaySoundAtLocation(World, InSoundCue,SoundLocation, VolumeMultiplier, PitchMultiplier);
-	}
-}
-
 void UGameplayStatics::SetBaseSoundMix(USoundMix* InSoundMix)
 {
 	if (InSoundMix)
@@ -867,10 +839,6 @@ void UGameplayStatics::ClearSoundMixModifiers()
 	}
 }
 
-void UGameplayStatics::SetSoundMode(FName InSoundModeName)
-{
-}
-
 void UGameplayStatics::ActivateReverbEffect(class UReverbEffect* ReverbEffect, FName TagName, float Priority, float Volume, float FadeTime)
 {
 	if (ReverbEffect)
diff --git a/Engine/Source/Runtime/Engine/Private/HeightFogComponent.cpp b/Engine/Source/Runtime/Engine/Private/HeightFogComponent.cpp
--- a/Engine/Source/Runtime/Engine/Private/HeightFogComponent.cpp
+++ b/Engine/Source/Runtime/Engine/Private/HeightFogComponent.cpp
@@ -12,8 +12,6 @@
 UExponentialHeightFogComponent::UExponentialHeightFogComponent(const FObjectInitializer& ObjectInitializer)
 	: Super(ObjectInitializer)
 {
-	bEnabled_DEPRECATED = true;
-
 	FogInscatteringColor = FLinearColor(0.447f, 0.638f, 1.0f);
 
 	DirectionalInscatteringExponent = 4.0f;
@@ -54,16 +52,6 @@ void UExponentialHeightFogComponent::DestroyRenderState_Concurrent()
 	World->Scene->RemoveExponentialHeightFog(this);
 }
 
-void UExponentialHeightFogComponent::PostLoad()
-{
-	Super::PostLoad();
-
-	if (GetLinkerUE4Version() < VER_UE4_REMOVE_COMPONENT_ENABLED_FLAG)
-	{
-		bVisible = bEnabled_DEPRECATED;	
-	}
-}
-
 #if WITH_EDITOR
 void UExponentialHeightFogComponent::PostEditChangeProperty(FPropertyChangedEvent& PropertyChangedEvent)
 {
diff --git a/Engine/Source/Runtime/Engine/Private/InteractiveFoliageActor.cpp b/Engine/Source/Runtime/Engine/Private/InteractiveFoliageActor.cpp
--- a/Engine/Source/Runtime/Engine/Private/InteractiveFoliageActor.cpp
+++ b/Engine/Source/Runtime/Engine/Private/InteractiveFoliageActor.cpp
@@ -12,13 +12,11 @@ AInteractiveFoliageActor::AInteractiveFoliageActor(const FObjectInitializer& Obj
 {
 
 	UInteractiveFoliageComponent* FoliageMeshComponent = CastChecked<UInteractiveFoliageComponent>(GetStaticMeshComponent());
-	FoliageMeshComponent->BodyInstance.bEnableCollision_DEPRECATED = false;
 	FoliageMeshComponent->SetCollisionProfileName(UCollisionProfile::NoCollision_ProfileName);
 	FoliageMeshComponent->Mobility = EComponentMobility::Static;
 
 	CapsuleComponent = ObjectInitializer.CreateDefaultSubobject<UCapsuleComponent>(this, TEXT("CollisionCylinder"));
 	CapsuleComponent->InitCapsuleSize(60.0f, 200.0f);
-	CapsuleComponent->BodyInstance.bEnableCollision_DEPRECATED = true;
 	static FName CollisionProfileName(TEXT("OverlapAllDynamic"));
 	CapsuleComponent->SetCollisionProfileName(CollisionProfileName);
 	CapsuleComponent->Mobility = EComponentMobility::Static;
diff --git a/Engine/Source/Runtime/Engine/Private/InterpCurveEdSetup.cpp b/Engine/Source/Runtime/Engine/Private/InterpCurveEdSetup.cpp
--- a/Engine/Source/Runtime/Engine/Private/InterpCurveEdSetup.cpp
+++ b/Engine/Source/Runtime/Engine/Private/InterpCurveEdSetup.cpp
@@ -28,28 +28,6 @@ void UInterpCurveEdSetup::PostLoad()
 	}
 }
 
-void UInterpCurveEdSetup::Serialize(FArchive& Ar)
-{
-	Super::Serialize(Ar);
-
-	if (Ar.IsLoading() && GetLinker() && (GetLinker()->UE3Ver() < 203))
-	{
-		for (int32 i=0; i<Tabs.Num(); i++)
-		{
-			FCurveEdTab& Tab = Tabs[i];
-			for(int32 j=Tab.Curves.Num()-1; j>=0; j--)
-			{
-				FCurveEdEntry& Entry = Tab.Curves[j];
-
-				if (Entry.bHideCurve)
-				{
-					Entry.bHideCurve	= 0x00000001;
-				}
-			}
-		}
-	}
-}
-
 UInterpCurveEdSetup::UInterpCurveEdSetup(const FObjectInitializer& ObjectInitializer)
 	: Super(ObjectInitializer)
 {
diff --git a/Engine/Source/Runtime/Engine/Private/Interpolation.cpp b/Engine/Source/Runtime/Engine/Private/Interpolation.cpp
--- a/Engine/Source/Runtime/Engine/Private/Interpolation.cpp
+++ b/Engine/Source/Runtime/Engine/Private/Interpolation.cpp
@@ -1184,22 +1184,6 @@ bool AMatineeActor::CanEditChange( const UProperty* Property ) const
 }
 #endif	// WITH_EDITOR
 
-void AMatineeActor::PostLoadSubobjects( FObjectInstancingGraph* OuterInstanceGraph )
-{
-	Super::PostLoadSubobjects(OuterInstanceGraph);
-
-	if ( GetLinkerUE4Version() < VER_UE4_FIX_INTERPDATA_OUTERS )
-	{
-		// Fix up auto-created MatineeData such that it has the correct outer
-		// If the Outer of the MatineeData is the OuterMost (i.e the level) then we rename
-		if (MatineeData && GetOutermost() == MatineeData->GetOuter())
-		{
-			MatineeData->ConditionalPostLoad();
-			MatineeData->Rename(NULL, this, REN_DontCreateRedirectors | REN_NonTransactional | REN_ForceNoResetLoaders);
-		}
-	}
-}
-
 #if WITH_EDITOR
 
 void AMatineeActor::ValidateActorGroups()
diff --git a/Engine/Source/Runtime/Engine/Private/Level.cpp b/Engine/Source/Runtime/Engine/Private/Level.cpp
--- a/Engine/Source/Runtime/Engine/Private/Level.cpp
+++ b/Engine/Source/Runtime/Engine/Private/Level.cpp
@@ -9,7 +9,6 @@ Level.cpp: Level-related functions
 #include "Engine/LevelScriptBlueprint.h"
 #include "Engine/LevelScriptActor.h"
 #include "Engine/WorldComposition.h"
-#include "Sound/SoundNodeWave.h"
 #include "Net/UnrealNetwork.h"
 #include "Model.h"
 #include "StaticLighting.h"
@@ -42,65 +41,6 @@ ULevel implementation.
 /** Called when a level package has been dirtied. */
 FSimpleMulticastDelegate ULevel::LevelDirtiedEvent;
 
-//@deprecated with VER_SPLIT_SOUND_FROM_TEXTURE_STREAMING
-struct FStreamableResourceInstanceDeprecated
-{
-	FSphere BoundingSphere;
-	float TexelFactor;
-	friend FArchive& operator<<( FArchive& Ar, FStreamableResourceInstanceDeprecated& ResourceInstance )
-	{
-		Ar << ResourceInstance.BoundingSphere;
-		Ar << ResourceInstance.TexelFactor;
-		return Ar;
-	}
-};
-//@deprecated with VER_SPLIT_SOUND_FROM_TEXTURE_STREAMING
-struct FStreamableResourceInfoDeprecated
-{
-	UObject* Resource;
-	TArray<FStreamableResourceInstanceDeprecated> ResourceInstances;
-	friend FArchive& operator<<( FArchive& Ar, FStreamableResourceInfoDeprecated& ResourceInfo )
-	{
-		Ar << ResourceInfo.Resource;
-		Ar << ResourceInfo.ResourceInstances;
-		return Ar;
-	}
-};
-//@deprecated with VER_RENDERING_REFACTOR
-struct FStreamableSoundInstanceDeprecated
-{
-	FSphere BoundingSphere;
-	friend FArchive& operator<<( FArchive& Ar, FStreamableSoundInstanceDeprecated& SoundInstance )
-	{
-		Ar << SoundInstance.BoundingSphere;
-		return Ar;
-	}
-};
-//@deprecated with VER_RENDERING_REFACTOR
-struct FStreamableSoundInfoDeprecated
-{
-	UDEPRECATED_SoundNodeWave*	SoundNodeWave;
-	TArray<FStreamableSoundInstanceDeprecated> SoundInstances;
-	friend FArchive& operator<<( FArchive& Ar, FStreamableSoundInfoDeprecated& SoundInfo )
-	{
-		Ar << SoundInfo.SoundNodeWave;
-		Ar << SoundInfo.SoundInstances;
-		return Ar;
-	}
-};
-//@deprecated with VER_RENDERING_REFACTOR
-struct FStreamableTextureInfoDeprecated
-{
-	UTexture*							Texture;
-	TArray<FStreamableTextureInstance>	TextureInstances;
-	friend FArchive& operator<<( FArchive& Ar, FStreamableTextureInfoDeprecated& TextureInfo )
-	{
-		Ar << TextureInfo.Texture;
-		Ar << TextureInfo.TextureInstances;
-		return Ar;
-	}
-};
-
 int32 FPrecomputedVisibilityHandler::NextId = 0;
 
 /** Updates visibility stats. */
@@ -349,12 +289,6 @@ void ULevel::Serialize( FArchive& Ar )
 
 	Ar << ModelComponents;
 
-	if(Ar.UE4Ver() < VER_UE4_REMOVE_USEQUENCE)
-	{
-		TArray<UObject*> DummySequences;
-		Ar << DummySequences;
-	}
-
 	if(!Ar.IsFilterEditorOnly() || (Ar.UE4Ver() < VER_UE4_EDITORONLY_BLUEPRINTS) )
 	{
 #if WITH_EDITORONLY_DATA
@@ -395,7 +329,7 @@ void ULevel::Serialize( FArchive& Ar )
 		{
 			Ar << bIsCooked;
 		}
-		if (Ar.UE4Ver() >= VER_UE4_EXPLICIT_STREAMING_TEXTURE_BUILT && Ar.UE4Ver() < VER_UE4_REBUILD_TEXTURE_STREAMING_DATA_ON_LOAD)
+		if (Ar.UE4Ver() < VER_UE4_REBUILD_TEXTURE_STREAMING_DATA_ON_LOAD)
 		{
 			bool bTextureStreamingBuiltDummy = bIsCooked;
 			Ar << bTextureStreamingBuiltDummy;
@@ -447,19 +381,6 @@ void ULevel::Serialize( FArchive& Ar )
 	// serialize the nav list
 	Ar << NavListStart;
 	Ar << NavListEnd;
-	// and pylons
-	if (Ar.UE4Ver() < VER_UE4_REMOVED_OLD_NAVMESH)
-	{
-		if (Ar.IsLoading())
-		{
-			// pylons were removed
-			// so if reading old data, just read and discard these 2 pointers
-			AActor* LegacyPylonListStart;
-			AActor* LegacyPylonListEnd;
-			Ar << LegacyPylonListStart;
-			Ar << LegacyPylonListEnd;
-		}
-	}
 
 	if (HasAnyFlags(RF_ClassDefaultObject))
 	{
@@ -603,17 +524,6 @@ void ULevel::PostLoad()
 
 	UWorldComposition::OnLevelPostLoad(this);
 		
-#if WITH_EDITOR
-	if(GetLinkerUE4Version() < VER_UE4_CLEAR_STANDALONE_FROM_LEVEL_SCRIPT_BLUEPRINTS)
-	{
-		if (LevelScriptBlueprint)
-		{
-			LevelScriptBlueprint->ConditionalPostLoad();
-			LevelScriptBlueprint->ClearFlags(RF_Standalone);
-		}
-	}
-#endif //WITH_EDITOR
-
 #if WITH_EDITOR
 	Actors.Remove(nullptr);
 #endif
diff --git a/Engine/Source/Runtime/Engine/Private/LevelBounds.cpp b/Engine/Source/Runtime/Engine/Private/LevelBounds.cpp
--- a/Engine/Source/Runtime/Engine/Private/LevelBounds.cpp
+++ b/Engine/Source/Runtime/Engine/Private/LevelBounds.cpp
@@ -20,7 +20,6 @@ ALevelBounds::ALevelBounds(const FObjectInitializer& ObjectInitializer)
 	BoxComponent->bDrawOnlyIfSelected = true;
 	BoxComponent->bUseAttachParentBound = false;
 	BoxComponent->bUseEditorCompositing = true;
-	BoxComponent->BodyInstance.bEnableCollision_DEPRECATED = false;
 	BoxComponent->SetCollisionProfileName(UCollisionProfile::NoCollision_ProfileName);
 	BoxComponent->InitBoxExtent(FVector(0.5f, 0.5f, 0.5f));
 	
diff --git a/Engine/Source/Runtime/Engine/Private/LevelStreamingVolume.cpp b/Engine/Source/Runtime/Engine/Private/LevelStreamingVolume.cpp
--- a/Engine/Source/Runtime/Engine/Private/LevelStreamingVolume.cpp
+++ b/Engine/Source/Runtime/Engine/Private/LevelStreamingVolume.cpp
@@ -7,8 +7,6 @@
 ALevelStreamingVolume::ALevelStreamingVolume(const FObjectInitializer& ObjectInitializer)
 	: Super(ObjectInitializer)
 {
-	GetBrushComponent()->BodyInstance.bEnableCollision_DEPRECATED = false;
-
 	GetBrushComponent()->SetCollisionProfileName(UCollisionProfile::NoCollision_ProfileName);
 	GetBrushComponent()->bAlwaysCreatePhysicsState = true;
 
diff --git a/Engine/Source/Runtime/Engine/Private/LightComponent.cpp b/Engine/Source/Runtime/Engine/Private/LightComponent.cpp
--- a/Engine/Source/Runtime/Engine/Private/LightComponent.cpp
+++ b/Engine/Source/Runtime/Engine/Private/LightComponent.cpp
@@ -312,11 +312,10 @@ ULightComponent::ULightComponent(const FObjectInitializer& ObjectInitializer)
 	IndirectLightingIntensity = 1.0f;
 	ShadowBias = 0.5f;
 	ShadowSharpen = 0.0f;
-	bUseIESBrightness = (GetLinkerUE4Version() < VER_UE4_LIGHTS_USE_IES_BRIGHTNESS_DEFAULT_CHANGED);
+	bUseIESBrightness = false;
 	IESBrightnessScale = 1.0f;
 	IESTexture = NULL;
 
-	bEnabled_DEPRECATED = true;
 	bAffectTranslucentLighting = true;
 	LightFunctionScale = FVector(1024.0f, 1024.0f, 1024.0f);
 
@@ -408,11 +407,6 @@ void ULightComponent::PostLoad()
 {
 	Super::PostLoad();
 
-	if (GetLinkerUE4Version() < VER_UE4_REMOVE_COMPONENT_ENABLED_FLAG)
-	{
-		bVisible = bEnabled_DEPRECATED;	
-	}
-
 	if (LightFunctionMaterial && HasStaticLighting())
 	{
 		// Light functions can only be used on dynamic lights
diff --git a/Engine/Source/Runtime/Engine/Private/LightMap.cpp b/Engine/Source/Runtime/Engine/Private/LightMap.cpp
--- a/Engine/Source/Runtime/Engine/Private/LightMap.cpp
+++ b/Engine/Source/Runtime/Engine/Private/LightMap.cpp
@@ -1621,28 +1621,7 @@ void FLightMap2D::Serialize(FArchive& Ar)
 {
 	FLightMap::Serialize(Ar);
 
-	if (Ar.IsLoading() && Ar.UE4Ver() < VER_UE4_SH_LIGHTMAPS)
-	{
-		for(uint32 CoefficientIndex = 0;CoefficientIndex < 3;CoefficientIndex++)
-		{
-			ULightMapTexture2D* Dummy = NULL;
-			Ar << Dummy;
-			FVector Dummy2;
-			Ar << Dummy2;
-		}
-	}
-	else if( Ar.IsLoading() && Ar.UE4Ver() < VER_UE4_LIGHTMAP_COMPRESSION )
-	{
-		for( uint32 CoefficientIndex = 0; CoefficientIndex < 5; CoefficientIndex++ )
-		{
-			ULightMapTexture2D* Dummy = NULL;
-			Ar << Dummy;
-			FVector Dummy2;
-			Ar << Dummy2;
-			Ar << Dummy2;
-		}
-	}
-	else if( Ar.IsLoading() && Ar.UE4Ver() < VER_UE4_LOW_QUALITY_DIRECTIONAL_LIGHTMAPS )
+	if( Ar.IsLoading() && Ar.UE4Ver() < VER_UE4_LOW_QUALITY_DIRECTIONAL_LIGHTMAPS )
 	{
 		for(uint32 CoefficientIndex = 0;CoefficientIndex < 3;CoefficientIndex++)
 		{
@@ -1764,32 +1743,13 @@ void FLegacyLightMap1D::Serialize(FArchive& Ar)
 
 	DirectionalSamples.Serialize( Ar, Owner );
 
-	if (Ar.IsLoading() && Ar.UE4Ver() < VER_UE4_SH_LIGHTMAPS)
+	for (int32 ElementIndex = 0; ElementIndex < 5; ElementIndex++)
 	{
-		for (int32 ElementIndex = 0; ElementIndex < 3; ElementIndex++)
-		{
-			FVector Dummy;
-			Ar << Dummy;
-		}
-	}
-	else
-	{
-		for (int32 ElementIndex = 0; ElementIndex < 5; ElementIndex++)
-		{
-			FVector Dummy;
-			Ar << Dummy;
-		}
+		FVector Dummy;
+		Ar << Dummy;
 	}
 
-	if (Ar.IsLoading() && Ar.UE4Ver() < VER_UE4_SH_LIGHTMAPS)
-	{
-		TQuantizedLightSampleBulkData<FLegacyQuantizedSimpleLightSample> Dummy;
-		Dummy.Serialize(Ar, Owner);
-	}
-	else
-	{
-		SimpleSamples.Serialize( Ar, Owner );
-	}
+	SimpleSamples.Serialize( Ar, Owner );
 }
 
 /*-----------------------------------------------------------------------------
diff --git a/Engine/Source/Runtime/Engine/Private/MaterialBillboardComponent.cpp b/Engine/Source/Runtime/Engine/Private/MaterialBillboardComponent.cpp
--- a/Engine/Source/Runtime/Engine/Private/MaterialBillboardComponent.cpp
+++ b/Engine/Source/Runtime/Engine/Private/MaterialBillboardComponent.cpp
@@ -253,8 +253,6 @@ private:
 UMaterialBillboardComponent::UMaterialBillboardComponent(const FObjectInitializer& ObjectInitializer)
 : Super(ObjectInitializer)
 {
-	BodyInstance.bEnableCollision_DEPRECATED = false;
-
 	SetCollisionProfileName(UCollisionProfile::NoCollision_ProfileName);
 }
 
diff --git a/Engine/Source/Runtime/Engine/Private/Materials/Material.cpp b/Engine/Source/Runtime/Engine/Private/Materials/Material.cpp
--- a/Engine/Source/Runtime/Engine/Private/Materials/Material.cpp
+++ b/Engine/Source/Runtime/Engine/Private/Materials/Material.cpp
@@ -53,7 +53,6 @@ bool FMaterialsWithDirtyUsageFlags::IsUsageFlagDirty(EMaterialUsage UsageFlag)
 }
 
 FUObjectAnnotationSparseBool GMaterialsThatNeedSamplerFixup;
-FUObjectAnnotationSparseBool GMaterialsThatNeedPhysicalConversion;
 FUObjectAnnotationSparse<FMaterialsWithDirtyUsageFlags,true> GMaterialsWithDirtyUsageFlags;
 FUObjectAnnotationSparseBool GMaterialsThatNeedExpressionsFlipped;
 FUObjectAnnotationSparseBool GMaterialsThatNeedCoordinateCheck;
@@ -518,8 +517,6 @@ UMaterial::UMaterial(const FObjectInitializer& ObjectInitializer)
 	Opacity.Constant = 1.0f;
 	OpacityMask.Constant = 1.0f;
 	OpacityMaskClipValue = 0.3333f;
-	FresnelBaseReflectFraction_DEPRECATED = 0.04f;
-	bPhysicallyBasedInputs_DEPRECATED = true;
 	bUsedWithStaticLighting = false;
 	D3D11TessellationMode = MTM_NoTessellation;
 	bEnableCrackFreeDisplacement = false;
@@ -1705,11 +1702,6 @@ void UMaterial::Serialize(FArchive& Ar)
 	}
 
 #if WITH_EDITOR
-	if ( Ar.UE4Ver() < VER_UE4_PHYSICAL_MATERIAL_MODEL )
-	{
-		GMaterialsThatNeedPhysicalConversion.Set( this );
-	}
-
 	if (Ar.UE4Ver() < VER_UE4_FLIP_MATERIAL_COORDS)
 	{
 		GMaterialsThatNeedExpressionsFlipped.Set(this);
@@ -1758,51 +1750,6 @@ void UMaterial::PostDuplicate(bool bDuplicateForPIE)
 void UMaterial::BackwardsCompatibilityInputConversion()
 {
 #if WITH_EDITOR
-	if ( GMaterialsThatNeedPhysicalConversion.Get( this ) )
-	{
-		GMaterialsThatNeedPhysicalConversion.Clear( this );
-
-		Roughness.Constant = 0.4238f;
-
-		if( ShadingModel != MSM_Unlit )
-		{
-			// Multiply SpecularColor by FresnelBaseReflectFraction
-			if( SpecularColor_DEPRECATED.IsConnected() && FresnelBaseReflectFraction_DEPRECATED != 1.0f )
-			{
-				UMaterialExpressionMultiply* MulExpression = ConstructObject< UMaterialExpressionMultiply >( UMaterialExpressionMultiply::StaticClass(), this );
-				Expressions.Add( MulExpression );
-
-				MulExpression->MaterialExpressionEditorX += 450;
-				MulExpression->MaterialExpressionEditorY += 20;
-
-				MulExpression->Desc = TEXT("FresnelBaseReflectFraction");
-				MulExpression->ConstA = 1.0f;
-				MulExpression->ConstB = FresnelBaseReflectFraction_DEPRECATED;
-
-				MulExpression->A.Connect( SpecularColor_DEPRECATED.OutputIndex, SpecularColor_DEPRECATED.Expression );
-				SpecularColor_DEPRECATED.Connect( 0, MulExpression );
-			}
-
-			// Convert from SpecularPower to Roughness
-			if( SpecularPower_DEPRECATED.IsConnected() )
-			{
-				check( GPowerToRoughnessMaterialFunction );
-
-				UMaterialExpressionMaterialFunctionCall* FunctionExpression = ConstructObject< UMaterialExpressionMaterialFunctionCall >( UMaterialExpressionMaterialFunctionCall::StaticClass(), this );
-				Expressions.Add( FunctionExpression );
-
-				FunctionExpression->MaterialExpressionEditorX += 200;
-				FunctionExpression->MaterialExpressionEditorY += 100;
-
-				FunctionExpression->MaterialFunction = GPowerToRoughnessMaterialFunction;
-				FunctionExpression->UpdateFromFunctionResource();
-
-				FunctionExpression->GetInput(0)->Connect( SpecularPower_DEPRECATED.OutputIndex, SpecularPower_DEPRECATED.Expression );
-				Roughness.Connect( 0, FunctionExpression );
-			}
-		}
-	}
-
 	if( ShadingModel != MSM_Unlit )
 	{
 		bool bIsDS = DiffuseColor_DEPRECATED.IsConnected() || SpecularColor_DEPRECATED.IsConnected();
@@ -1962,18 +1909,6 @@ void UMaterial::PostLoad()
 		LightingGuidFixupMap.Add(GetLightingGuid(), this);
 	}
 
-	// Fix exclusive material usage flags moved to an enum.
-	if (bUsedAsLightFunction_DEPRECATED)
-	{
-		MaterialDomain = MD_LightFunction;
-	}
-	else if (bUsedWithDeferredDecal_DEPRECATED)
-	{
-		MaterialDomain = MD_DeferredDecal;
-	}
-	bUsedAsLightFunction_DEPRECATED = false;
-	bUsedWithDeferredDecal_DEPRECATED = false;
-
 	// Fix the shading model to be valid.  Loading a material saved with a shading model that has been removed will yield a MSM_MAX.
 	if(ShadingModel == MSM_MAX)
 	{
diff --git a/Engine/Source/Runtime/Engine/Private/Materials/MaterialExpressions.cpp b/Engine/Source/Runtime/Engine/Private/Materials/MaterialExpressions.cpp
--- a/Engine/Source/Runtime/Engine/Private/Materials/MaterialExpressions.cpp
+++ b/Engine/Source/Runtime/Engine/Private/Materials/MaterialExpressions.cpp
@@ -2220,18 +2220,6 @@ UMaterialExpressionConstant3Vector::UMaterialExpressionConstant3Vector(const FOb
 	bCollapsed = false;
 }
 
-void UMaterialExpressionConstant3Vector::PostLoad()
-{
-	Super::PostLoad();
-
-	if(GetLinkerUE4Version() < VER_UE4_CHANGE_MATERIAL_EXPRESSION_CONSTANTS_TO_LINEARCOLOR)
-	{
-		Constant.R = R_DEPRECATED;
-		Constant.G = G_DEPRECATED;
-		Constant.B = B_DEPRECATED;
-	}
-}
-
 int32 UMaterialExpressionConstant3Vector::Compile(class FMaterialCompiler* Compiler, int32 OutputIndex, int32 MultiplexIndex)
 {
 	return Compiler->Constant3(Constant.R,Constant.G,Constant.B);
@@ -2275,19 +2263,6 @@ UMaterialExpressionConstant4Vector::UMaterialExpressionConstant4Vector(const FOb
 	bCollapsed = false;
 }
 
-void UMaterialExpressionConstant4Vector::PostLoad()
-{
-	Super::PostLoad();
-
-	if(GetLinkerUE4Version() < VER_UE4_CHANGE_MATERIAL_EXPRESSION_CONSTANTS_TO_LINEARCOLOR)
-	{
-		Constant.R = R_DEPRECATED;
-		Constant.G = G_DEPRECATED;
-		Constant.B = B_DEPRECATED;
-		Constant.A = A_DEPRECATED;
-	}
-}
-
 int32 UMaterialExpressionConstant4Vector::Compile(class FMaterialCompiler* Compiler, int32 OutputIndex, int32 MultiplexIndex)
 {
 	return Compiler->Constant4(Constant.R,Constant.G,Constant.B,Constant.A);
@@ -3262,23 +3237,6 @@ UMaterialExpressionBreakMaterialAttributes::UMaterialExpressionBreakMaterialAttr
 	}
 }
 
-void  UMaterialExpressionBreakMaterialAttributes::Serialize( FArchive& Ar )
-{
-	Super::Serialize(Ar);
-
-	//Switch the connection over to the new input.
-	if( Ar.UE4Ver() < VER_UE4_MATERIAL_ATTRIBUTES_MULTIPLEX )
-	{
-		MaterialAttributes.Expression = Struct.Expression;
-		MaterialAttributes.OutputIndex = Struct.OutputIndex;
-		MaterialAttributes.Mask = Struct.Mask;
-		MaterialAttributes.MaskR = Struct.MaskR;
-		MaterialAttributes.MaskG = Struct.MaskG;
-		MaterialAttributes.MaskB = Struct.MaskB;
-		MaterialAttributes.MaskA = Struct.MaskA;
-	}
-}
-
 int32 UMaterialExpressionBreakMaterialAttributes::Compile(class FMaterialCompiler* Compiler, int32 OutputIndex, int32 MultiplexIndex)
 {
 	//Here we don't care about any multiplex index coming in.
diff --git a/Engine/Source/Runtime/Engine/Private/Materials/MaterialShader.cpp b/Engine/Source/Runtime/Engine/Private/Materials/MaterialShader.cpp
--- a/Engine/Source/Runtime/Engine/Private/Materials/MaterialShader.cpp
+++ b/Engine/Source/Runtime/Engine/Private/Materials/MaterialShader.cpp
@@ -312,12 +312,9 @@ void FMaterialShaderMapId::Serialize(FArchive& Ar)
 	// Backwards compatibility only works with FMaterialShaderMapId's stored in packages.  
 	// You must bump MATERIALSHADERMAP_DERIVEDDATA_VER as well if changing the serialization of FMaterialShaderMapId.
 
-	if (Ar.UE4Ver() >= VER_UE4_ADDED_MATERIALSHADERMAP_USAGE)
-	{
-		uint32 UsageInt = Usage;
-		Ar << UsageInt;
-		Usage = (EMaterialShaderMapUsage::Type)UsageInt;
-	}
+	uint32 UsageInt = Usage;
+	Ar << UsageInt;
+	Usage = (EMaterialShaderMapUsage::Type)UsageInt;
 
 	Ar << BaseMaterialId;
 
@@ -326,7 +323,7 @@ void FMaterialShaderMapId::Serialize(FArchive& Ar)
 		Ar << (int32&)QualityLevel;
 		Ar << (int32&)FeatureLevel;
 	}
-	else if (Ar.UE4Ver() >= VER_UE4_MATERIAL_QUALITY_LEVEL_SWITCH)
+	else
 	{
 		uint8 LegacyQualityLevel;
 		Ar << LegacyQualityLevel;
@@ -334,27 +331,21 @@ void FMaterialShaderMapId::Serialize(FArchive& Ar)
 
 	ParameterSet.Serialize(Ar);
 
-	if (Ar.UE4Ver() >= VER_UE4_FUNCTIONS_IN_SHADERMAPID)
-	{
-		Ar << ReferencedFunctions;
-	}
+	Ar << ReferencedFunctions;
 
 	if (Ar.UE4Ver() >= VER_UE4_COLLECTIONS_IN_SHADERMAPID)
 	{
 		Ar << ReferencedParameterCollections;
 	}
 
-	if (Ar.UE4Ver() >= VER_UE4_REMOVED_PERSHADER_DDC)
-	{
-		Ar << ShaderTypeDependencies;
-		Ar << VertexFactoryTypeDependencies;
-	}
+	Ar << ShaderTypeDependencies;
+	Ar << VertexFactoryTypeDependencies;
 
 	if (Ar.UE4Ver() >= VER_UE4_PURGED_FMATERIAL_COMPILE_OUTPUTS)
 	{
 		Ar << TextureReferencesHash;
 	}
-	else if (Ar.UE4Ver() >= VER_UE4_HASHED_MATERIAL_OUTPUT)
+	else
 	{
 		FSHAHash LegacyHash;
 		Ar << LegacyHash;
diff --git a/Engine/Source/Runtime/Engine/Private/Materials/MaterialShared.cpp b/Engine/Source/Runtime/Engine/Private/Materials/MaterialShared.cpp
--- a/Engine/Source/Runtime/Engine/Private/Materials/MaterialShared.cpp
+++ b/Engine/Source/Runtime/Engine/Private/Materials/MaterialShared.cpp
@@ -571,30 +571,12 @@ void FMaterial::LegacySerialize(FArchive& Ar)
 		QualityLevel = EMaterialQualityLevel::High;
 		Ar << Id_DEPRECATED;
 
-		if (Ar.UE4Ver() < VER_UE4_REMOVED_FMATERIAL_COMPILE_OUTPUTS)
-		{
-			uint32 Temp = 0;
-			Ar << Temp;
-		}
-
 		TArray<UTexture*> LegacyTextures;
 		Ar << LegacyTextures;
 
 		bool bTemp2;
 		Ar << bTemp2;
 
-		if (Ar.UE4Ver() < VER_UE4_REMOVED_FMATERIAL_COMPILE_OUTPUTS)
-		{
-			bool bUsesLightmapUVsTemp = false;
-			Ar << bUsesLightmapUVsTemp;
-
-			bool bUsesMaterialVertexPositionOffsetTemp = false;
-			Ar << bUsesMaterialVertexPositionOffsetTemp;
-
-			bool bUsesSphericalParticleOpacityTemp = false;
-			Ar << bUsesSphericalParticleOpacityTemp;
-		}
-
 		bool bTemp;
 		Ar << bTemp;
 
@@ -610,13 +592,8 @@ void FMaterial::LegacySerialize(FArchive& Ar)
 
 void FMaterial::SerializeInlineShaderMap(FArchive& Ar)
 {
-	bool bCooked = false;
-
-	if (Ar.UE4Ver() >= VER_UE4_INLINE_SHADERS)
-	{
-		bCooked = Ar.IsCooking();
-		Ar << bCooked;
-	}
+	bool bCooked = Ar.IsCooking();
+	Ar << bCooked;
 
 	if (FPlatformProperties::RequiresCookedData() && !bCooked && Ar.IsLoading())
 	{
@@ -663,8 +640,7 @@ void FMaterialResource::LegacySerialize(FArchive& Ar)
 {
 	FMaterial::LegacySerialize(Ar);
 
-	if (Ar.UE4Ver() < VER_UE4_PURGED_FMATERIAL_COMPILE_OUTPUTS
-		&& Ar.UE4Ver() >= VER_UE4_MATERIAL_BLEND_OVERRIDE)
+	if (Ar.UE4Ver() < VER_UE4_PURGED_FMATERIAL_COMPILE_OUTPUTS)
 	{
 		int32 BlendModeOverrideValueTemp = 0;
 		Ar << BlendModeOverrideValueTemp;
diff --git a/Engine/Source/Runtime/Engine/Private/ModelComponent.cpp b/Engine/Source/Runtime/Engine/Private/ModelComponent.cpp
--- a/Engine/Source/Runtime/Engine/Private/ModelComponent.cpp
+++ b/Engine/Source/Runtime/Engine/Private/ModelComponent.cpp
@@ -42,10 +42,7 @@ FModelElement::~FModelElement()
 FArchive& operator<<(FArchive& Ar,FModelElement& Element)
 {
 	Ar << Element.LightMap;
-	if( Ar.UE4Ver() >= VER_UE4_PRECOMPUTED_SHADOW_MAPS_BSP )
-	{
-		Ar << Element.ShadowMap;
-	}
+	Ar << Element.ShadowMap;
 	
 	Ar << (UObject*&)Element.Component << (UObject*&)Element.Material << Element.Nodes;
 	Ar << Element.IrrelevantLights;
diff --git a/Engine/Source/Runtime/Engine/Private/NavigationObjectBase.cpp b/Engine/Source/Runtime/Engine/Private/NavigationObjectBase.cpp
--- a/Engine/Source/Runtime/Engine/Private/NavigationObjectBase.cpp
+++ b/Engine/Source/Runtime/Engine/Private/NavigationObjectBase.cpp
@@ -14,7 +14,6 @@ ANavigationObjectBase::ANavigationObjectBase(const FObjectInitializer& ObjectIni
 	CapsuleComponent->ShapeColor = FColor(255, 138, 5, 255);
 	CapsuleComponent->bDrawOnlyIfSelected = true;
 	CapsuleComponent->InitCapsuleSize(50.0f, 50.0f);
-	CapsuleComponent->BodyInstance.bEnableCollision_DEPRECATED = false;
 	CapsuleComponent->SetCollisionProfileName(UCollisionProfile::NoCollision_ProfileName);
 	CapsuleComponent->bShouldCollideWhenPlacing = true;
 	CapsuleComponent->bShouldUpdatePhysicsVolume = true;
diff --git a/Engine/Source/Runtime/Engine/Private/Particles/ParticleBeamModules.cpp b/Engine/Source/Runtime/Engine/Private/Particles/ParticleBeamModules.cpp
--- a/Engine/Source/Runtime/Engine/Private/Particles/ParticleBeamModules.cpp
+++ b/Engine/Source/Runtime/Engine/Private/Particles/ParticleBeamModules.cpp
@@ -7,7 +7,6 @@
 #include "EnginePrivate.h"
 #include "Distributions/DistributionVectorConstant.h"
 #include "ParticleDefinitions.h"
-#include "../DistributionHelpers.h"
 #include "Particles/Beam/ParticleModuleBeamBase.h"
 #include "Particles/Beam/ParticleModuleBeamModifier.h"
 #include "Particles/Beam/ParticleModuleBeamNoise.h"
@@ -88,17 +87,6 @@ void UParticleModuleTypeDataBeam2::PostInitProperties()
 	}
 }
 
-void UParticleModuleTypeDataBeam2::Serialize(FArchive& Ar)
-{
-	Super::Serialize(Ar);
-	if (Ar.IsLoading() && Ar.UE4Ver() < VER_UE4_MOVE_DISTRIBUITONS_TO_POSTINITPROPS)
-	{
-		FDistributionHelpers::RestoreDefaultConstant(Distance.Distribution, TEXT("DistributionDistance"), 25.0f);
-		FDistributionHelpers::RestoreDefaultConstant(TaperFactor.Distribution, TEXT("DistributionTaperFactor"), 1.0f);
-		FDistributionHelpers::RestoreDefaultConstant(TaperScale.Distribution, TEXT("DistributionTaperScale"), 1.0f);
-	}
-}
-
 void UParticleModuleTypeDataBeam2::Spawn(FParticleEmitterInstance* Owner, int32 Offset, float SpawnTime, FBaseParticle* ParticleBase)
 {
 	// Make sure that the owner is a beam emitter instance.
@@ -959,17 +947,6 @@ void UParticleModuleBeamModifier::PostInitProperties()
 	}
 }
 
-void UParticleModuleBeamModifier::Serialize(FArchive& Ar)
-{
-	Super::Serialize(Ar);
-	if (Ar.IsLoading() && Ar.UE4Ver() < VER_UE4_MOVE_DISTRIBUITONS_TO_POSTINITPROPS)
-	{
-		FDistributionHelpers::RestoreDefaultConstant(Position.Distribution, TEXT("DistributionPosition"), FVector(0.0f, 0.0f, 0.0f));
-		FDistributionHelpers::RestoreDefaultConstant(Tangent.Distribution, TEXT("DistributionTangent"), FVector(0.0f, 0.0f, 0.0f));
-		FDistributionHelpers::RestoreDefaultConstant(Strength.Distribution, TEXT("DistributionStrength"), 0.0f);
-	}
-}
-
 void UParticleModuleBeamModifier::Spawn(FParticleEmitterInstance* Owner, int32 Offset, float SpawnTime, FBaseParticle* ParticleBase)
 {
 	// Make sure that the owner is a beam emitter instance.
@@ -1254,18 +1231,6 @@ void UParticleModuleBeamNoise::PostInitProperties()
 	}
 }
 
-void UParticleModuleBeamNoise::Serialize(FArchive& Ar)
-{
-	Super::Serialize(Ar);
-	if (Ar.IsLoading() && Ar.UE4Ver() < VER_UE4_MOVE_DISTRIBUITONS_TO_POSTINITPROPS)
-	{
-		FDistributionHelpers::RestoreDefaultConstant(NoiseSpeed.Distribution, TEXT("DistributionNoiseSpeed"), FVector(50.0f, 50.0f, 50.0f));
-		FDistributionHelpers::RestoreDefaultConstant(NoiseRange.Distribution, TEXT("DistributionNoiseRange"), FVector(50.0f, 50.0f, 50.0f));
-		FDistributionHelpers::RestoreDefaultConstant(NoiseRangeScale.Distribution, TEXT("DistributionNoiseRangeScale"), 1.0f);
-		FDistributionHelpers::RestoreDefaultConstant(NoiseTangentStrength.Distribution, TEXT("DistributionNoiseTangentStrength"), 250.0f);
-	}
-}
-
 void UParticleModuleBeamNoise::Spawn(FParticleEmitterInstance* Owner, int32 Offset, float SpawnTime, FBaseParticle* ParticleBase)
 {
 	if (bLowFreq_Enabled == false)
@@ -1577,17 +1542,6 @@ void UParticleModuleBeamSource::PostInitProperties()
 	}
 }
 
-void UParticleModuleBeamSource::Serialize(FArchive& Ar)
-{
-	Super::Serialize(Ar);
-	if (Ar.IsLoading() && Ar.UE4Ver() < VER_UE4_MOVE_DISTRIBUITONS_TO_POSTINITPROPS)
-	{
-		FDistributionHelpers::RestoreDefaultConstant(Source.Distribution, TEXT("DistributionSource"), FVector(50.0f, 50.0f, 50.0f));
-		FDistributionHelpers::RestoreDefaultConstant(SourceTangent.Distribution, TEXT("DistributionSourceTangent"), FVector(1.0f, 0.0f, 0.0f));
-		FDistributionHelpers::RestoreDefaultConstant(SourceStrength.Distribution, TEXT("DistributionSourceStrength"), 25.0f);
-	}
-}
-
 void UParticleModuleBeamSource::Spawn(FParticleEmitterInstance* Owner, int32 Offset, float SpawnTime, FBaseParticle* ParticleBase)
 {
 	// Make sure that the owner is a beam emitter instance.
@@ -2009,17 +1963,6 @@ void UParticleModuleBeamTarget::PostInitProperties()
 	}
 }
 
-void UParticleModuleBeamTarget::Serialize(FArchive& Ar)
-{
-	Super::Serialize(Ar);
-	if (Ar.IsLoading() && Ar.UE4Ver() < VER_UE4_MOVE_DISTRIBUITONS_TO_POSTINITPROPS)
-	{
-		FDistributionHelpers::RestoreDefaultConstant(Target.Distribution, TEXT("DistributionTarget"), FVector(50.0f, 50.0f, 50.0f));
-		FDistributionHelpers::RestoreDefaultConstant(TargetTangent.Distribution, TEXT("DistributionTargetTangent"), FVector(1.0f, 0.0f, 0.0f));
-		FDistributionHelpers::RestoreDefaultConstant(TargetStrength.Distribution, TEXT("DistributionTargetStrength"), 25.0f);
-	}
-}
-
 void UParticleModuleBeamTarget::Spawn(FParticleEmitterInstance* Owner, int32 Offset, float SpawnTime, FBaseParticle* ParticleBase)
 {
 	FParticleBeam2EmitterInstance*	BeamInst	= (FParticleBeam2EmitterInstance*)Owner;
diff --git a/Engine/Source/Runtime/Engine/Private/Particles/ParticleComponents.cpp b/Engine/Source/Runtime/Engine/Private/Particles/ParticleComponents.cpp
--- a/Engine/Source/Runtime/Engine/Private/Particles/ParticleComponents.cpp
+++ b/Engine/Source/Runtime/Engine/Private/Particles/ParticleComponents.cpp
@@ -1978,7 +1978,6 @@ UParticleSystem::UParticleSystem(const FObjectInitializer& ObjectInitializer)
 	UpdateTime_Delta = 1.0f/60.0f;
 	WarmupTime = 0.0f;
 	WarmupTickRate = 0.0f;
-	bLit_DEPRECATED = true;
 #if WITH_EDITORONLY_DATA
 	EditorLODSetting = 0;
 #endif // WITH_EDITORONLY_DATA
@@ -2103,9 +2102,6 @@ void UParticleSystem::PostLoad()
 {
 	Super::PostLoad();
 
-	//@todo. Put this in a better place??
-	bool bHadDeprecatedEmitters = false;
-
 	// Remove any old emitters
 	bHasPhysics = false;
 	for (int32 i = Emitters.Num() - 1; i >= 0; i--)
@@ -2160,11 +2156,6 @@ void UParticleSystem::PostLoad()
 		}
 	}
 
-	if ((bHadDeprecatedEmitters || (GetLinker() && (GetLinker()->UE3Ver() < 204))) && CurveEdSetup)
-	{
-		CurveEdSetup->ResetTabs();
-	}
-
 	if (LODSettings.Num() == 0)
 	{
 		if (Emitters.Num() > 0)
@@ -2201,7 +2192,6 @@ void UParticleSystem::PostLoad()
 	}
 
 #if WITH_EDITOR
-//	if (GetLinker() && (GetLinker()->UE3Ver() < VER_PARTICLE_LOD_DISTANCE_FIXUP))
 	// Due to there still being some ways that LODLevel counts get mismatched,
 	// when loading in the editor LOD levels will always be checked and fixed
 	// up... This can be removed once all the edge cases that lead to the
@@ -2931,7 +2921,6 @@ UParticleSystemComponent::UParticleSystemComponent(const FObjectInitializer& Obj
 #if WITH_EDITORONLY_DATA
 	EditorDetailMode = -1;
 #endif // WITH_EDITORONLY_DATA
-	BodyInstance.bEnableCollision_DEPRECATED = false;
 	SetCollisionProfileName(UCollisionProfile::NoCollision_ProfileName);
 	bGenerateOverlapEvents = false;
 
diff --git a/Engine/Source/Runtime/Engine/Private/Particles/ParticleModules.cpp b/Engine/Source/Runtime/Engine/Private/Particles/ParticleModules.cpp
--- a/Engine/Source/Runtime/Engine/Private/Particles/ParticleModules.cpp
+++ b/Engine/Source/Runtime/Engine/Private/Particles/ParticleModules.cpp
@@ -11,7 +11,6 @@
 #include "Distributions/DistributionVectorUniformCurve.h"
 #include "FXSystem.h"
 #include "ParticleDefinitions.h"
-#include "../DistributionHelpers.h"
 #include "Particles/Acceleration/ParticleModuleAcceleration.h"
 #include "Particles/Acceleration/ParticleModuleAccelerationBase.h"
 #include "Particles/Acceleration/ParticleModuleAccelerationConstant.h"
@@ -498,88 +497,6 @@ bool UParticleModule::ConvertVectorDistribution(UDistributionVector* VectorDist,
 	return true;
 }
 
-
-bool UParticleModule::IsIdentical_Deprecated(const UParticleModule* InModule) const
-{
-	// Valid module?
-	if (InModule == NULL)
-	{
-		return false;
-	}
-
-	// Same class?
-	if (InModule->GetClass() != GetClass())
-	{
-		return false;
-	}
-
-	for (UProperty* Prop = GetClass()->PropertyLink; Prop; Prop = Prop->PropertyLinkNext)
-	{
-		// only the properties that could have been modified in the editor should be compared
-		// (skipping the name and archetype properties, since name will almost always be different)
-		bool bConsiderProperty = Prop->ShouldDuplicateValue();
-		if (PropertyIsRelevantForIsIdentical_Deprecated(Prop->GetFName()) == false)
-		{
-			bConsiderProperty = false;
-		}
-
-		if (bConsiderProperty)
-		{
-			for (int32 i = 0; i < Prop->ArrayDim; i++)
-			{
-				if (!Prop->Identical_InContainer(this, InModule, i, PPF_DeepComparison))
-				{
-					return false;
-				}
-			}
-		}
-	}
-
-	return true;
-}
-
-
-bool UParticleModule::PropertyIsRelevantForIsIdentical_Deprecated(const FName& InPropName) const
-{
-	static TArray<FName> IdenticalIgnoreProperties_ParticleModule;
-	static TArray<FName> IdenticalIgnoreProperties_ParticleModuleRequired;
-	if( IdenticalIgnoreProperties_ParticleModule.Num() == 0 )
-	{
-		IdenticalIgnoreProperties_ParticleModule.Add(FName(TEXT("bSpawnModule")));
-		IdenticalIgnoreProperties_ParticleModule.Add(FName(TEXT("bUpdateModule")));
-		IdenticalIgnoreProperties_ParticleModule.Add(FName(TEXT("bFinalUpdateModule")));
-		IdenticalIgnoreProperties_ParticleModule.Add(FName(TEXT("bCurvesAsColor")));
-		IdenticalIgnoreProperties_ParticleModule.Add(FName(TEXT("b3DDrawMode")));
-		IdenticalIgnoreProperties_ParticleModule.Add(FName(TEXT("bSupported3DDrawMode")));
-		IdenticalIgnoreProperties_ParticleModule.Add(FName(TEXT("bEditable")));
-		IdenticalIgnoreProperties_ParticleModule.Add(FName(TEXT("ModuleEditorColor")));
-		IdenticalIgnoreProperties_ParticleModule.Add(FName(TEXT("IdenticalIgnoreProperties")));
-		IdenticalIgnoreProperties_ParticleModule.Add(FName(TEXT("LODValidity")));
-		IdenticalIgnoreProperties_ParticleModuleRequired.Add(FName(TEXT("SpawnRate")));
-		IdenticalIgnoreProperties_ParticleModuleRequired.Add(FName(TEXT("ParticleBurstMethod")));
-		IdenticalIgnoreProperties_ParticleModuleRequired.Add(FName(TEXT("BurstList")));
-	}
-	for (int32 IgnoreIndex = 0; IgnoreIndex < IdenticalIgnoreProperties_ParticleModule.Num(); IgnoreIndex++)
-	{
-		if (IdenticalIgnoreProperties_ParticleModule[IgnoreIndex] == InPropName)
-		{
-			return false;
-		}
-	}
-	if( IsA( UParticleModuleRequired::StaticClass() ) )
-	{
-		for (int32 IgnoreIndex = 0; IgnoreIndex < IdenticalIgnoreProperties_ParticleModuleRequired.Num(); IgnoreIndex++)
-		{
-			if (IdenticalIgnoreProperties_ParticleModuleRequired[IgnoreIndex] == InPropName)
-			{
-				return false;
-			}
-		}
-	}
-	return true;
-}
-
-
 UParticleModule* UParticleModule::GenerateLODModule(UParticleLODLevel* SourceLODLevel, UParticleLODLevel* DestLODLevel, float Percentage, 
 	bool bGenerateModuleData, bool bForceModuleConstruction)
 {
@@ -967,15 +884,6 @@ void UParticleModuleSourceMovement::PostInitProperties()
 	}
 }
 
-void UParticleModuleSourceMovement::Serialize(FArchive& Ar)
-{
-	Super::Serialize(Ar);
-	if (Ar.IsLoading() && Ar.UE4Ver() < VER_UE4_MOVE_DISTRIBUITONS_TO_POSTINITPROPS)
-	{
-		FDistributionHelpers::RestoreDefaultConstant(SourceMovementScale.Distribution, TEXT("DistributionSourceMovementScale"), FVector(1.0f, 1.0f, 1.0f));
-	}
-}
-
 void UParticleModuleSourceMovement::FinalUpdate(FParticleEmitterInstance* Owner, int32 Offset, float DeltaTime)
 {
 	Super::FinalUpdate(Owner, Offset, DeltaTime);
@@ -1103,15 +1011,6 @@ void UParticleModuleRequired::PostInitProperties()
 	}
 }
 
-void UParticleModuleRequired::Serialize(FArchive& Ar)
-{
-	Super::Serialize(Ar);
-	if (Ar.IsLoading() && Ar.UE4Ver() < VER_UE4_MOVE_DISTRIBUITONS_TO_POSTINITPROPS)
-	{
-		FDistributionHelpers::RestoreDefaultConstant(SpawnRate.Distribution, TEXT("RequiredDistributionSpawnRate"), 0.0f);
-	}
-}
-
 #if WITH_EDITOR
 void UParticleModuleRequired::PostEditChangeProperty(FPropertyChangedEvent& PropertyChangedEvent)
 {
@@ -1255,15 +1154,6 @@ void UParticleModuleMeshRotation::PostInitProperties()
 	}
 }
 
-void UParticleModuleMeshRotation::Serialize(FArchive& Ar)
-{
-	Super::Serialize(Ar);
-	if (Ar.IsLoading() && Ar.UE4Ver() < VER_UE4_MOVE_DISTRIBUITONS_TO_POSTINITPROPS)
-	{
-		FDistributionHelpers::RestoreDefaultUniform(StartRotation.Distribution, TEXT("DistributionStartRotation"), FVector(0.0f, 0.0f, 0.0f), FVector(1.0f, 1.0f, 1.0f));
-	}
-}
-
 void UParticleModuleMeshRotation::Spawn(FParticleEmitterInstance* Owner, int32 Offset, float SpawnTime, FBaseParticle* ParticleBase)
 {
 	SpawnEx(Owner, Offset, SpawnTime, NULL, ParticleBase);
@@ -1362,15 +1252,6 @@ void UParticleModuleMeshRotationRate::PostInitProperties()
 	}
 }
 
-void UParticleModuleMeshRotationRate::Serialize(FArchive& Ar)
-{
-	Super::Serialize(Ar);
-	if (Ar.IsLoading() && Ar.UE4Ver() < VER_UE4_MOVE_DISTRIBUITONS_TO_POSTINITPROPS)
-	{
-		FDistributionHelpers::RestoreDefaultUniform(StartRotationRate.Distribution, TEXT("DistributionStartRotationRate"), FVector(0.0f, 0.0f, 0.0f), FVector(360.0f, 360.0f, 360.0f));
-	}
-}
-
 void UParticleModuleMeshRotationRate::Spawn(FParticleEmitterInstance* Owner, int32 Offset, float SpawnTime, FBaseParticle* ParticleBase)
 {
 	SpawnEx(Owner, Offset, SpawnTime, NULL, ParticleBase);
@@ -1463,15 +1344,6 @@ void UParticleModuleMeshRotationRateMultiplyLife::PostInitProperties()
 	}
 }
 
-void UParticleModuleMeshRotationRateMultiplyLife::Serialize(FArchive& Ar)
-{
-	Super::Serialize(Ar);
-	if (Ar.IsLoading() && Ar.UE4Ver() < VER_UE4_MOVE_DISTRIBUITONS_TO_POSTINITPROPS)
-	{
-		FDistributionHelpers::RestoreDefaultConstant(LifeMultiplier.Distribution, TEXT("DistributionLifeMultiplier"), FVector::ZeroVector);
-	}
-}
-
 void UParticleModuleMeshRotationRateMultiplyLife::Spawn(FParticleEmitterInstance* Owner, int32 Offset, float SpawnTime, FBaseParticle* ParticleBase)
 {
 	const int32 MeshRotationOffset = Owner->GetMeshRotationOffset();
@@ -1632,15 +1504,6 @@ void UParticleModuleRotation::PostInitProperties()
 	}
 }
 
-void UParticleModuleRotation::Serialize(FArchive& Ar)
-{
-	Super::Serialize(Ar);
-	if (Ar.IsLoading() && Ar.UE4Ver() < VER_UE4_MOVE_DISTRIBUITONS_TO_POSTINITPROPS)
-	{
-		FDistributionHelpers::RestoreDefaultUniform(StartRotation.Distribution, TEXT("DistributionStartRotation"), 0.0f, 1.0f);
-	}
-}
-
 #if WITH_EDITOR
 void UParticleModuleRotation::PostEditChangeProperty(FPropertyChangedEvent& PropertyChangedEvent)
 {
@@ -1724,15 +1587,6 @@ void UParticleModuleRotationRate::PostInitProperties()
 	}
 }
 
-void UParticleModuleRotationRate::Serialize(FArchive& Ar)
-{
-	Super::Serialize(Ar);
-	if (Ar.IsLoading() && Ar.UE4Ver() < VER_UE4_MOVE_DISTRIBUITONS_TO_POSTINITPROPS)
-	{
-		FDistributionHelpers::RestoreDefaultConstant(StartRotationRate.Distribution, TEXT("DistributionStartRotationRate"), 0.0f);
-	}
-}
-
 void UParticleModuleRotationRate::CompileModule( FParticleEmitterBuildInfo& EmitterInfo )
 {
 	float MinRate;
@@ -1911,15 +1765,6 @@ void UParticleModuleSubUV::PostInitProperties()
 	}
 }
 
-void UParticleModuleSubUV::Serialize(FArchive& Ar)
-{
-	Super::Serialize(Ar);
-	if (Ar.IsLoading() && Ar.UE4Ver() < VER_UE4_MOVE_DISTRIBUITONS_TO_POSTINITPROPS)
-	{
-		FDistributionHelpers::RestoreDefaultConstant(SubImageIndex.Distribution, TEXT("DistributionSubImage"), 0.0f);
-	}
-}
-
 void UParticleModuleSubUV::CompileModule( FParticleEmitterBuildInfo& EmitterInfo )
 {
 	check( EmitterInfo.RequiredModule );
@@ -2131,15 +1976,6 @@ void UParticleModuleSubUVMovie::PostInitProperties()
 	}
 }
 
-void UParticleModuleSubUVMovie::Serialize(FArchive& Ar)
-{
-	Super::Serialize(Ar);
-	if (Ar.IsLoading() && Ar.UE4Ver() < VER_UE4_MOVE_DISTRIBUITONS_TO_POSTINITPROPS)
-	{
-		FDistributionHelpers::RestoreDefaultConstant(FrameRate.Distribution, TEXT("DistributionFrameRate"), 30.0f);
-	}
-}
-
 #if WITH_EDITOR
 void UParticleModuleSubUVMovie::PostEditChangeProperty(FPropertyChangedEvent& PropertyChangedEvent)
 {
@@ -2305,15 +2141,6 @@ void UParticleModuleRotationRateMultiplyLife::PostInitProperties()
 	}
 }
 
-void UParticleModuleRotationRateMultiplyLife::Serialize(FArchive& Ar)
-{
-	Super::Serialize(Ar);
-	if (Ar.IsLoading() && Ar.UE4Ver() < VER_UE4_MOVE_DISTRIBUITONS_TO_POSTINITPROPS)
-	{
-		FDistributionHelpers::RestoreDefaultConstant(LifeMultiplier.Distribution, TEXT("DistributionLifeMultiplier"), 0.0f);
-	}
-}
-
 #if WITH_EDITOR
 void UParticleModuleRotationRateMultiplyLife::PostEditChangeProperty(FPropertyChangedEvent& PropertyChangedEvent)
 {
@@ -2477,15 +2304,6 @@ void UParticleModuleAccelerationDrag::PostInitProperties()
 	}
 }
 
-void UParticleModuleAccelerationDrag::Serialize(FArchive& Ar)
-{
-	Super::Serialize(Ar);
-	if (Ar.IsLoading() && Ar.UE4Ver() < VER_UE4_MOVE_DISTRIBUITONS_TO_POSTINITPROPS)
-	{
-		FDistributionHelpers::RestoreDefaultConstant(DragCoefficient, TEXT("DistributionDragCoefficient"), 1.0f);
-	}
-}
-
 void UParticleModuleAccelerationDrag::CompileModule(FParticleEmitterBuildInfo& EmitterInfo)
 {
 	EmitterInfo.DragCoefficient.Initialize(DragCoefficient);
@@ -2551,16 +2369,6 @@ void UParticleModuleAccelerationDragScaleOverLife::PostInitProperties()
 	}
 }
 
-void UParticleModuleAccelerationDragScaleOverLife::Serialize(FArchive& Ar)
-{
-	Super::Serialize(Ar);
-	if (Ar.IsLoading() && Ar.UE4Ver() < VER_UE4_MOVE_DISTRIBUITONS_TO_POSTINITPROPS)
-	{
-		FDistributionHelpers::RestoreDefaultConstant(DragScale, TEXT("DistributionDragScale"), 1.0f);
-	}
-}
-
-
 void UParticleModuleAccelerationDragScaleOverLife::CompileModule(FParticleEmitterBuildInfo& EmitterInfo)
 {
 	EmitterInfo.DragScale.ScaleByDistribution(DragScale);
@@ -2617,15 +2425,6 @@ void UParticleModuleAttractorPointGravity::PostInitProperties()
 	}
 }
 
-void UParticleModuleAttractorPointGravity::Serialize(FArchive& Ar)
-{
-	Super::Serialize(Ar);
-	if (Ar.IsLoading() && Ar.UE4Ver() < VER_UE4_MOVE_DISTRIBUITONS_TO_POSTINITPROPS)
-	{
-		FDistributionHelpers::RestoreDefaultConstant(Strength, TEXT("DistributionStrength"), 1.0f);
-	}
-}
-
 void UParticleModuleAttractorPointGravity::CompileModule(FParticleEmitterBuildInfo& EmitterInfo)
 {
 	EmitterInfo.PointAttractorPosition = Position;
@@ -2681,15 +2480,6 @@ void UParticleModuleAcceleration::PostInitProperties()
 	}
 }
 
-void UParticleModuleAcceleration::Serialize(FArchive& Ar)
-{
-	Super::Serialize(Ar);
-	if (Ar.IsLoading() && Ar.UE4Ver() < VER_UE4_MOVE_DISTRIBUITONS_TO_POSTINITPROPS)
-	{
-		FDistributionHelpers::RestoreDefaultUniform(Acceleration.Distribution, TEXT("DistributionAcceleration"), FVector::ZeroVector, FVector::ZeroVector);
-	}
-}
-
 void UParticleModuleAcceleration::CompileModule( FParticleEmitterBuildInfo& EmitterInfo )
 {
 	EmitterInfo.ConstantAcceleration = Acceleration.GetValue();
@@ -3210,16 +3000,6 @@ void UParticleModuleKillBox::PostInitProperties()
 	}
 }
 
-void UParticleModuleKillBox::Serialize(FArchive& Ar)
-{
-	Super::Serialize(Ar);
-	if (Ar.IsLoading() && Ar.UE4Ver() < VER_UE4_MOVE_DISTRIBUITONS_TO_POSTINITPROPS)
-	{
-		FDistributionHelpers::RestoreDefaultConstant(LowerLeftCorner.Distribution, TEXT("DistributionLowerLeftCorner"), FVector::ZeroVector);
-		FDistributionHelpers::RestoreDefaultConstant(UpperRightCorner.Distribution, TEXT("DistributionUpperRightCorner"), FVector::ZeroVector);
-	}
-}
-
 #if WITH_EDITOR
 void UParticleModuleKillBox::PostEditChangeProperty(FPropertyChangedEvent& PropertyChangedEvent)
 {
@@ -3365,15 +3145,6 @@ void UParticleModuleKillHeight::PostInitProperties()
 	}
 }
 
-void UParticleModuleKillHeight::Serialize(FArchive& Ar)
-{
-	Super::Serialize(Ar);
-	if (Ar.IsLoading() && Ar.UE4Ver() < VER_UE4_MOVE_DISTRIBUITONS_TO_POSTINITPROPS)
-	{
-		FDistributionHelpers::RestoreDefaultConstant(Height.Distribution, TEXT("DistributionHeight"), 0.0f);
-	}
-}
-
 #if WITH_EDITOR
 void UParticleModuleKillHeight::PostEditChangeProperty(FPropertyChangedEvent& PropertyChangedEvent)
 {
@@ -3497,15 +3268,6 @@ void UParticleModuleLifetime::PostInitProperties()
 	}
 }
 
-void UParticleModuleLifetime::Serialize(FArchive& Ar)
-{
-	Super::Serialize(Ar);
-	if (Ar.IsLoading() && Ar.UE4Ver() < VER_UE4_MOVE_DISTRIBUITONS_TO_POSTINITPROPS)
-	{
-		FDistributionHelpers::RestoreDefaultUniform(Lifetime.Distribution, TEXT("DistributionLifetime"), 0.0f, 0.0f);
-	}
-}
-
 void UParticleModuleLifetime::CompileModule( struct FParticleEmitterBuildInfo& EmitterInfo )
 {
 	float MinLifetime;
@@ -3660,16 +3422,6 @@ void UParticleModuleAttractorLine::PostInitProperties()
 	}
 }
 
-void UParticleModuleAttractorLine::Serialize(FArchive& Ar)
-{
-	Super::Serialize(Ar);
-	if (Ar.IsLoading() && Ar.UE4Ver() < VER_UE4_MOVE_DISTRIBUITONS_TO_POSTINITPROPS)
-	{
-		FDistributionHelpers::RestoreDefaultConstant(Strength.Distribution, TEXT("DistributionStrength"), 0.0f);
-		FDistributionHelpers::RestoreDefaultConstant(Range.Distribution, TEXT("DistributionRange"), 0.0f);
-	}
-}
-
 #if WITH_EDITOR
 void UParticleModuleAttractorLine::PostEditChangeProperty(FPropertyChangedEvent& PropertyChangedEvent)
 {
@@ -3811,16 +3563,6 @@ void UParticleModuleAttractorParticle::PostInitProperties()
 	}
 }
 
-void UParticleModuleAttractorParticle::Serialize(FArchive& Ar)
-{
-	Super::Serialize(Ar);
-	if (Ar.IsLoading() && Ar.UE4Ver() < VER_UE4_MOVE_DISTRIBUITONS_TO_POSTINITPROPS)
-	{
-		FDistributionHelpers::RestoreDefaultConstant(Range.Distribution, TEXT("DistributionRange"), 0.0f);
-		FDistributionHelpers::RestoreDefaultConstant(Strength.Distribution, TEXT("DistributionStrength"), 0.0f);
-	}
-}
-
 #if WITH_EDITOR
 void UParticleModuleAttractorParticle::PostEditChangeProperty(FPropertyChangedEvent& PropertyChangedEvent)
 {
@@ -4069,17 +3811,6 @@ void UParticleModuleAttractorPoint::PostInitProperties()
 	}
 }
 
-void UParticleModuleAttractorPoint::Serialize(FArchive& Ar)
-{
-	Super::Serialize(Ar);
-	if (Ar.IsLoading() && Ar.UE4Ver() < VER_UE4_MOVE_DISTRIBUITONS_TO_POSTINITPROPS)
-	{
-		FDistributionHelpers::RestoreDefaultConstant(Position.Distribution, TEXT("DistributionPosition"), FVector::ZeroVector);
-		FDistributionHelpers::RestoreDefaultConstant(Range.Distribution, TEXT("DistributionRange"), 0.0f);
-		FDistributionHelpers::RestoreDefaultConstant(Strength.Distribution, TEXT("DistributionStrength"), 0.0f);
-	}
-}
-
 #if WITH_EDITOR
 void UParticleModuleAttractorPoint::PostEditChangeProperty(FPropertyChangedEvent& PropertyChangedEvent)
 {
diff --git a/Engine/Source/Runtime/Engine/Private/Particles/ParticleModules_Camera.cpp b/Engine/Source/Runtime/Engine/Private/Particles/ParticleModules_Camera.cpp
--- a/Engine/Source/Runtime/Engine/Private/Particles/ParticleModules_Camera.cpp
+++ b/Engine/Source/Runtime/Engine/Private/Particles/ParticleModules_Camera.cpp
@@ -6,7 +6,6 @@
 =============================================================================*/
 #include "EnginePrivate.h"
 #include "ParticleDefinitions.h"
-#include "../DistributionHelpers.h"
 #include "Particles/Camera/ParticleModuleCameraOffset.h"
 #include "Particles/TypeData/ParticleModuleTypeDataMesh.h"
 #include "Particles/ParticleLODLevel.h"
@@ -54,15 +53,6 @@ void UParticleModuleCameraOffset::PostInitProperties()
 	}
 }
 
-void UParticleModuleCameraOffset::Serialize(FArchive& Ar)
-{
-	Super::Serialize(Ar);
-	if (Ar.IsLoading() && Ar.UE4Ver() < VER_UE4_MOVE_DISTRIBUITONS_TO_POSTINITPROPS)
-	{
-		FDistributionHelpers::RestoreDefaultConstant(CameraOffset.Distribution, TEXT("DistributionCameraOffset"), 1.0f);
-	}
-}
-
 #if WITH_EDITOR
 void UParticleModuleCameraOffset::PostEditChangeProperty(FPropertyChangedEvent& PropertyChangedEvent)
 {
diff --git a/Engine/Source/Runtime/Engine/Private/Particles/ParticleModules_Collision.cpp b/Engine/Source/Runtime/Engine/Private/Particles/ParticleModules_Collision.cpp
--- a/Engine/Source/Runtime/Engine/Private/Particles/ParticleModules_Collision.cpp
+++ b/Engine/Source/Runtime/Engine/Private/Particles/ParticleModules_Collision.cpp
@@ -7,7 +7,6 @@
 #include "EnginePrivate.h"
 #include "Engine/TriggerBase.h"
 #include "ParticleDefinitions.h"
-#include "../DistributionHelpers.h"
 #include "Particles/Collision/ParticleModuleCollision.h"
 #include "Particles/Collision/ParticleModuleCollisionBase.h"
 #include "Particles/Collision/ParticleModuleCollisionGPU.h"
@@ -99,19 +98,6 @@ void UParticleModuleCollision::PostInitProperties()
 	}
 }
 
-void UParticleModuleCollision::Serialize(FArchive& Ar)
-{
-	Super::Serialize(Ar);
-	if (Ar.IsLoading() && Ar.UE4Ver() < VER_UE4_MOVE_DISTRIBUITONS_TO_POSTINITPROPS)
-	{
-		FDistributionHelpers::RestoreDefaultUniform(DampingFactor.Distribution, TEXT("DistributionDampingFactor"), FVector::ZeroVector, FVector::ZeroVector);
-		FDistributionHelpers::RestoreDefaultConstant(DampingFactorRotation.Distribution, TEXT("DistributionDampingFactorRotation"), FVector(1.0f, 1.0f, 1.0f));
-		FDistributionHelpers::RestoreDefaultUniform(MaxCollisions.Distribution, TEXT("DistributionMaxCollisions"), 0.0f, 0.0f);
-		FDistributionHelpers::RestoreDefaultConstant(ParticleMass.Distribution, TEXT("DistributionParticleMass"), 0.1f);
-		FDistributionHelpers::RestoreDefaultConstant(DelayAmount.Distribution, TEXT("DistributionDelayAmount"), 0.0f);
-	}
-}
-
 #if WITH_EDITOR
 void UParticleModuleCollision::PostEditChangeProperty(FPropertyChangedEvent& PropertyChangedEvent)
 {
diff --git a/Engine/Source/Runtime/Engine/Private/Particles/ParticleModules_Color.cpp b/Engine/Source/Runtime/Engine/Private/Particles/ParticleModules_Color.cpp
--- a/Engine/Source/Runtime/Engine/Private/Particles/ParticleModules_Color.cpp
+++ b/Engine/Source/Runtime/Engine/Private/Particles/ParticleModules_Color.cpp
@@ -8,7 +8,6 @@
 #include "Distributions/DistributionFloatParticleParameter.h"
 #include "Distributions/DistributionVectorParticleParameter.h"
 #include "ParticleDefinitions.h"
-#include "../DistributionHelpers.h"
 #include "Particles/Color/ParticleModuleColor.h"
 #include "Particles/Color/ParticleModuleColorBase.h"
 #include "Particles/Color/ParticleModuleColorOverLife.h"
@@ -68,17 +67,6 @@ void UParticleModuleColor::PostInitProperties()
 	}
 }
 
-void UParticleModuleColor::Serialize(FArchive& Ar)
-{
-	Super::Serialize(Ar);
-	if (Ar.IsLoading() && Ar.UE4Ver() < VER_UE4_MOVE_DISTRIBUITONS_TO_POSTINITPROPS)
-	{
-		FDistributionHelpers::RestoreDefaultConstant(StartColor.Distribution, TEXT("DistributionStartColor"), FVector::ZeroVector);
-		FDistributionHelpers::RestoreDefaultConstant(StartAlpha.Distribution, TEXT("DistributionStartAlpha"), 1.0f);
-	}
-}
-
-
 void UParticleModuleColor::CompileModule( FParticleEmitterBuildInfo& EmitterInfo )
 {
 	FVector InitialColor;
@@ -260,15 +248,6 @@ void UParticleModuleColorOverLife::PostInitProperties()
 	}
 }
 
-void UParticleModuleColorOverLife::Serialize(FArchive& Ar)
-{
-	Super::Serialize(Ar);
-	if (Ar.IsLoading() && Ar.UE4Ver() < VER_UE4_MOVE_DISTRIBUITONS_TO_POSTINITPROPS)
-	{
-		FDistributionHelpers::RestoreDefaultConstant(AlphaOverLife.Distribution, TEXT("DistributionAlphaOverLife"), 1.0f);
-	}
-}
-
 void UParticleModuleColorOverLife::CompileModule( FParticleEmitterBuildInfo& EmitterInfo )
 {
 	bool ScaleColor = true;
@@ -515,16 +494,6 @@ void UParticleModuleColorScaleOverLife::PostInitProperties()
 	}
 }
 
-void UParticleModuleColorScaleOverLife::Serialize(FArchive& Ar)
-{
-	Super::Serialize(Ar);
-	if (Ar.IsLoading() && Ar.UE4Ver() < VER_UE4_MOVE_DISTRIBUITONS_TO_POSTINITPROPS)
-	{
-		FDistributionHelpers::RestoreDefaultConstant(AlphaScaleOverLife.Distribution, TEXT("DistributionAlphaScaleOverLife"), 1.0f);
-	}
-}
-
-
 void UParticleModuleColorScaleOverLife::CompileModule( FParticleEmitterBuildInfo& EmitterInfo )
 {
 	bool ScaleColor = true;
diff --git a/Engine/Source/Runtime/Engine/Private/Particles/ParticleModules_Event.cpp b/Engine/Source/Runtime/Engine/Private/Particles/ParticleModules_Event.cpp
--- a/Engine/Source/Runtime/Engine/Private/Particles/ParticleModules_Event.cpp
+++ b/Engine/Source/Runtime/Engine/Private/Particles/ParticleModules_Event.cpp
@@ -5,7 +5,6 @@
 =============================================================================*/
 #include "EnginePrivate.h"
 #include "ParticleDefinitions.h"
-#include "../DistributionHelpers.h"
 #include "Particles/Event/ParticleModuleEventBase.h"
 #include "Particles/Event/ParticleModuleEventGenerator.h"
 #include "Particles/Event/ParticleModuleEventReceiverBase.h"
@@ -320,16 +319,6 @@ void UParticleModuleEventReceiverSpawn::PostInitProperties()
 	}
 }
 
-void UParticleModuleEventReceiverSpawn::Serialize(FArchive& Ar)
-{
-	Super::Serialize(Ar);
-	if (Ar.IsLoading() && Ar.UE4Ver() < VER_UE4_MOVE_DISTRIBUITONS_TO_POSTINITPROPS)
-	{
-		FDistributionHelpers::RestoreDefaultConstant(SpawnCount.Distribution, TEXT("RequiredDistributionSpawnCount"), 0.0f);
-		FDistributionHelpers::RestoreDefaultConstant(InheritVelocityScale.Distribution, TEXT("RequiredDistributionInheritVelocityScale"), FVector(1.0f, 1.0f, 1.0f));
-	}
-}
-
 #if WITH_EDITOR
 void UParticleModuleEventReceiverSpawn::PostEditChangeProperty(FPropertyChangedEvent& PropertyChangedEvent)
 {
diff --git a/Engine/Source/Runtime/Engine/Private/Particles/ParticleModules_Location.cpp b/Engine/Source/Runtime/Engine/Private/Particles/ParticleModules_Location.cpp
--- a/Engine/Source/Runtime/Engine/Private/Particles/ParticleModules_Location.cpp
+++ b/Engine/Source/Runtime/Engine/Private/Particles/ParticleModules_Location.cpp
@@ -9,7 +9,6 @@
 #include "Distributions/DistributionVectorConstantCurve.h"
 #include "ParticleDefinitions.h"
 #include "RawIndexBuffer.h"
-#include "../DistributionHelpers.h"
 #include "Particles/Location/ParticleModuleLocation.h"
 #include "Particles/Location/ParticleModuleLocationBoneSocket.h"
 #include "Particles/Location/ParticleModuleLocationDirect.h"
@@ -71,15 +70,6 @@ void UParticleModuleLocation::PostInitProperties()
 	}
 }
 
-void UParticleModuleLocation::Serialize(FArchive& Ar)
-{
-	Super::Serialize(Ar);
-	if (Ar.IsLoading() && Ar.UE4Ver() < VER_UE4_MOVE_DISTRIBUITONS_TO_POSTINITPROPS)
-	{
-		FDistributionHelpers::RestoreDefaultUniform(StartLocation.Distribution, TEXT("DistributionStartLocation"), FVector::ZeroVector, FVector::ZeroVector);
-	}
-}
-
 #if WITH_EDITOR
 void UParticleModuleLocation::PostEditChangeProperty(FPropertyChangedEvent& PropertyChangedEvent)
 {
@@ -319,18 +309,6 @@ void UParticleModuleLocationDirect::PostInitProperties()
 	}
 }
 
-void UParticleModuleLocationDirect::Serialize(FArchive& Ar)
-{
-	Super::Serialize(Ar);
-	if (Ar.IsLoading() && Ar.UE4Ver() < VER_UE4_MOVE_DISTRIBUITONS_TO_POSTINITPROPS)
-	{
-		FDistributionHelpers::RestoreDefaultUniform(Location.Distribution, TEXT("DistributionLocation"), FVector::ZeroVector, FVector::ZeroVector);
-		FDistributionHelpers::RestoreDefaultConstant(LocationOffset.Distribution, TEXT("DistributionLocationOffset"), FVector::ZeroVector);
-		FDistributionHelpers::RestoreDefaultConstant(ScaleFactor.Distribution, TEXT("DistributionScaleFactor"), FVector(1.0f, 1.0f, 1.0f));
-		FDistributionHelpers::RestoreDefaultUniform(Direction.Distribution, TEXT("DistributionDirection"), FVector::ZeroVector, FVector::ZeroVector);
-	}
-}
-
 #if WITH_EDITOR
 void UParticleModuleLocationDirect::PostEditChangeProperty(FPropertyChangedEvent& PropertyChangedEvent)
 {
@@ -692,16 +670,6 @@ void UParticleModuleLocationPrimitiveBase::PostInitProperties()
 	}
 }
 
-void UParticleModuleLocationPrimitiveBase::Serialize(FArchive& Ar)
-{
-	Super::Serialize(Ar);
-	if (Ar.IsLoading() && Ar.UE4Ver() < VER_UE4_MOVE_DISTRIBUITONS_TO_POSTINITPROPS)
-	{
-		FDistributionHelpers::RestoreDefaultConstant(VelocityScale.Distribution, TEXT("DistributionVelocityScale"), 1.0f);
-		FDistributionHelpers::RestoreDefaultConstant(StartLocation.Distribution, TEXT("DistributionStartLocation"), FVector(0.0f, 0.0f, 0.0f));
-	}
-}
-
 #if WITH_EDITOR
 void UParticleModuleLocationPrimitiveBase::PostEditChangeProperty(FPropertyChangedEvent& PropertyChangedEvent)
 {
@@ -832,18 +800,6 @@ void UParticleModuleLocationPrimitiveTriangle::PostInitProperties()
 	}
 }
 
-void UParticleModuleLocationPrimitiveTriangle::Serialize(FArchive& Ar)
-{
-	Super::Serialize(Ar);
-	if (Ar.IsLoading() && Ar.UE4Ver() < VER_UE4_MOVE_DISTRIBUITONS_TO_POSTINITPROPS)
-	{
-		FDistributionHelpers::RestoreDefaultConstant(StartOffset.Distribution, TEXT("DistributionOffset"), FVector::ZeroVector);
-		FDistributionHelpers::RestoreDefaultConstant(Height.Distribution, TEXT("DistributionHeight"), 50.0f);
-		FDistributionHelpers::RestoreDefaultConstant(Angle.Distribution, TEXT("DistributionAngle"), 90.0f);
-		FDistributionHelpers::RestoreDefaultConstant(Thickness.Distribution, TEXT("DistributionThickness"), 0.0f);
-	}
-}
-
 #if WITH_EDITOR
 void UParticleModuleLocationPrimitiveTriangle::PostEditChangeProperty(FPropertyChangedEvent& PropertyChangedEvent)
 {
@@ -997,16 +953,6 @@ void UParticleModuleLocationPrimitiveCylinder::PostInitProperties()
 	}
 }
 
-void UParticleModuleLocationPrimitiveCylinder::Serialize(FArchive& Ar)
-{
-	Super::Serialize(Ar);
-	if (Ar.IsLoading() && Ar.UE4Ver() < VER_UE4_MOVE_DISTRIBUITONS_TO_POSTINITPROPS)
-	{
-		FDistributionHelpers::RestoreDefaultConstant(StartRadius.Distribution, TEXT("DistributionStartRadius"), 50.0f);
-		FDistributionHelpers::RestoreDefaultConstant(StartHeight.Distribution, TEXT("DistributionStartHeight"), 50.0f);
-	}
-}
-
 #if WITH_EDITOR
 void UParticleModuleLocationPrimitiveCylinder::PostEditChangeProperty(FPropertyChangedEvent& PropertyChangedEvent)
 {
@@ -1306,15 +1252,6 @@ void UParticleModuleLocationPrimitiveSphere::PostInitProperties()
 	}
 }
 
-void UParticleModuleLocationPrimitiveSphere::Serialize(FArchive& Ar)
-{
-	Super::Serialize(Ar);
-	if (Ar.IsLoading() && Ar.UE4Ver() < VER_UE4_MOVE_DISTRIBUITONS_TO_POSTINITPROPS)
-	{
-		FDistributionHelpers::RestoreDefaultConstant(StartRadius.Distribution, TEXT("DistributionStartRadius"), 50.0f);
-	}
-}
-
 #if WITH_EDITOR
 void UParticleModuleLocationPrimitiveSphere::PostEditChangeProperty(FPropertyChangedEvent& PropertyChangedEvent)
 {
diff --git a/Engine/Source/Runtime/Engine/Private/Particles/ParticleModules_Orbit.cpp b/Engine/Source/Runtime/Engine/Private/Particles/ParticleModules_Orbit.cpp
--- a/Engine/Source/Runtime/Engine/Private/Particles/ParticleModules_Orbit.cpp
+++ b/Engine/Source/Runtime/Engine/Private/Particles/ParticleModules_Orbit.cpp
@@ -5,7 +5,6 @@
 =============================================================================*/
 #include "EnginePrivate.h"
 #include "ParticleDefinitions.h"
-#include "../DistributionHelpers.h"
 #include "Particles/Orbit/ParticleModuleOrbit.h"
 #include "Particles/TypeData/ParticleModuleTypeDataGpu.h"
 #include "Particles/ParticleLODLevel.h"
@@ -68,17 +67,6 @@ void UParticleModuleOrbit::PostInitProperties()
 	}
 }
 
-void UParticleModuleOrbit::Serialize(FArchive& Ar)
-{
-	Super::Serialize(Ar);
-	if (Ar.IsLoading() && Ar.UE4Ver() < VER_UE4_MOVE_DISTRIBUITONS_TO_POSTINITPROPS)
-	{
-		FDistributionHelpers::RestoreDefaultUniform(OffsetAmount.Distribution, TEXT("DistributionOffsetAmount"), FVector(0.0f, 0.0f, 0.0f), FVector(0.0f, 50.0f, 0.0f));
-		FDistributionHelpers::RestoreDefaultUniform(RotationAmount.Distribution, TEXT("DistributionRotationAmount"), FVector(0.0f, 0.0f, 0.0f), FVector(1.0f, 1.0f, 1.0f));
-		FDistributionHelpers::RestoreDefaultUniform(RotationRateAmount.Distribution, TEXT("DistributionRotationRateAmount"), FVector(0.0f, 0.0f, 0.0f), FVector(1.0f, 1.0f, 1.0f));
-	}
-}
-
 void UParticleModuleOrbit::CompileModule(FParticleEmitterBuildInfo& EmitterInfo)
 {
 	switch (ChainMode)
diff --git a/Engine/Source/Runtime/Engine/Private/Particles/ParticleModules_Parameter.cpp b/Engine/Source/Runtime/Engine/Private/Particles/ParticleModules_Parameter.cpp
--- a/Engine/Source/Runtime/Engine/Private/Particles/ParticleModules_Parameter.cpp
+++ b/Engine/Source/Runtime/Engine/Private/Particles/ParticleModules_Parameter.cpp
@@ -8,7 +8,6 @@
 #include "Materials/MaterialExpressionDynamicParameter.h"
 #include "Materials/MaterialInstanceConstant.h"
 #include "ParticleDefinitions.h"
-#include "../DistributionHelpers.h"
 #include "Particles/Parameter/ParticleModuleParameterBase.h"
 #include "Particles/Parameter/ParticleModuleParameterDynamic.h"
 #include "Particles/Parameter/ParticleModuleParameterDynamic_Seeded.h"
@@ -71,18 +70,6 @@ void UParticleModuleParameterDynamic::PostInitProperties()
 	}
 }
 
-void UParticleModuleParameterDynamic::Serialize(FArchive& Ar)
-{
-	Super::Serialize(Ar);
-	if (Ar.IsLoading() && Ar.UE4Ver() < VER_UE4_MOVE_DISTRIBUITONS_TO_POSTINITPROPS)
-	{		
-		FDistributionHelpers::RestoreDefaultConstant(Cast<UDistributionFloatConstant>(StaticFindObjectFast(UDistributionFloatConstant::StaticClass(), this, TEXT("DistributionParam1"))), TEXT("DistributionParam1"), 0.0f);
-		FDistributionHelpers::RestoreDefaultConstant(Cast<UDistributionFloatConstant>(StaticFindObjectFast(UDistributionFloatConstant::StaticClass(), this, TEXT("DistributionParam2"))), TEXT("DistributionParam2"), 0.0f);
-		FDistributionHelpers::RestoreDefaultConstant(Cast<UDistributionFloatConstant>(StaticFindObjectFast(UDistributionFloatConstant::StaticClass(), this, TEXT("DistributionParam3"))), TEXT("DistributionParam3"), 0.0f);
-		FDistributionHelpers::RestoreDefaultConstant(Cast<UDistributionFloatConstant>(StaticFindObjectFast(UDistributionFloatConstant::StaticClass(), this, TEXT("DistributionParam4"))), TEXT("DistributionParam4"), 0.0f);
-	}
-}
-
 /** Flags for optimizing update */
 enum EDynamicParameterUpdateFlags
 {
diff --git a/Engine/Source/Runtime/Engine/Private/Particles/ParticleModules_Size.cpp b/Engine/Source/Runtime/Engine/Private/Particles/ParticleModules_Size.cpp
--- a/Engine/Source/Runtime/Engine/Private/Particles/ParticleModules_Size.cpp
+++ b/Engine/Source/Runtime/Engine/Private/Particles/ParticleModules_Size.cpp
@@ -8,7 +8,6 @@
 #include "EnginePrivate.h"
 #include "Distributions/DistributionVectorConstantCurve.h"
 #include "ParticleDefinitions.h"
-#include "../DistributionHelpers.h"
 #include "Particles/Size/ParticleModuleSize_Seeded.h"
 #include "Particles/Size/ParticleModuleSizeMultiplyLife.h"
 #include "Particles/Size/ParticleModuleSizeScale.h"
@@ -59,15 +58,6 @@ void UParticleModuleSize::PostInitProperties()
 	}
 }
 
-void UParticleModuleSize::Serialize(FArchive& Ar)
-{
-	Super::Serialize(Ar);
-	if (Ar.IsLoading() && Ar.UE4Ver() < VER_UE4_MOVE_DISTRIBUITONS_TO_POSTINITPROPS)
-	{
-		FDistributionHelpers::RestoreDefaultUniform(StartSize.Distribution, TEXT("DistributionStartSize"), FVector(1.0f, 1.0f, 1.0f), FVector(1.0f, 1.0f, 1.0f));
-	}
-}
-
 void UParticleModuleSize::CompileModule( FParticleEmitterBuildInfo& EmitterInfo )
 {
 	float MinSize = 0.0f;
@@ -169,15 +159,6 @@ void UParticleModuleSizeMultiplyLife::PostInitProperties()
 	}
 }
 
-void UParticleModuleSizeMultiplyLife::Serialize(FArchive& Ar)
-{
-	Super::Serialize(Ar);
-	if (Ar.IsLoading() && Ar.UE4Ver() < VER_UE4_MOVE_DISTRIBUITONS_TO_POSTINITPROPS)
-	{
-		FDistributionHelpers::RestoreDefaultConstant(LifeMultiplier.Distribution, TEXT("DistributionLifeMultiplier"), FVector::ZeroVector);
-	}
-}
-
 void UParticleModuleSizeMultiplyLife::CompileModule( FParticleEmitterBuildInfo& EmitterInfo )
 {
 	FVector AxisScaleMask(
@@ -368,15 +349,6 @@ void UParticleModuleSizeScale::PostInitProperties()
 	}
 }
 
-void UParticleModuleSizeScale::Serialize(FArchive& Ar)
-{
-	Super::Serialize(Ar);
-	if (Ar.IsLoading() && Ar.UE4Ver() < VER_UE4_MOVE_DISTRIBUITONS_TO_POSTINITPROPS)
-	{
-		FDistributionHelpers::RestoreDefaultConstant(SizeScale.Distribution, TEXT("DistributionSizeScale"), FVector::ZeroVector);
-	}
-}
-
 void UParticleModuleSizeScale::CompileModule( FParticleEmitterBuildInfo& EmitterInfo )
 {
 	EmitterInfo.SizeScale.Initialize( SizeScale.Distribution );
diff --git a/Engine/Source/Runtime/Engine/Private/Particles/ParticleModules_Spawn.cpp b/Engine/Source/Runtime/Engine/Private/Particles/ParticleModules_Spawn.cpp
--- a/Engine/Source/Runtime/Engine/Private/Particles/ParticleModules_Spawn.cpp
+++ b/Engine/Source/Runtime/Engine/Private/Particles/ParticleModules_Spawn.cpp
@@ -5,7 +5,6 @@
 =============================================================================*/
 #include "EnginePrivate.h"
 #include "ParticleDefinitions.h"
-#include "../DistributionHelpers.h"
 #include "Particles/Spawn/ParticleModuleSpawn.h"
 #include "Particles/Spawn/ParticleModuleSpawnBase.h"
 #include "Particles/Spawn/ParticleModuleSpawnPerUnit.h"
@@ -82,23 +81,6 @@ void UParticleModuleSpawn::PostLoad()
 	}
 }
 
-void UParticleModuleSpawn::Serialize(FArchive& Ar)
-{
-	Super::Serialize(Ar);
-	if (Ar.IsLoading() && Ar.UE4Ver() < VER_UE4_MOVE_DISTRIBUITONS_TO_POSTINITPROPS)
-	{
-		FDistributionHelpers::RestoreDefaultConstant(Rate.Distribution, TEXT("RequiredDistributionSpawnRate"), 20.0f);
-		FDistributionHelpers::RestoreDefaultConstant(RateScale.Distribution, TEXT("RequiredDistributionSpawnRateScale"), 1.0f);
-		FDistributionHelpers::RestoreDefaultConstant(BurstScale.Distribution, TEXT("BurstScaleDistribution"), 1.0f);
-		if (!BurstScale.Distribution)
-		{
-			UDistributionFloatConstant* BurstScaleDistribution = NewNamedObject<UDistributionFloatConstant>(this, TEXT("BurstScaleDistribution"));
-			BurstScaleDistribution->Constant = 1.0f;
-			BurstScale.Distribution = BurstScaleDistribution;
-		}
-	}
-}
-
 #if WITH_EDITOR
 void UParticleModuleSpawn::PostEditChangeProperty(FPropertyChangedEvent& PropertyChangedEvent)
 {
@@ -299,15 +281,6 @@ void UParticleModuleSpawnPerUnit::PostInitProperties()
 	}
 }
 
-void UParticleModuleSpawnPerUnit::Serialize(FArchive& Ar)
-{
-	Super::Serialize(Ar);
-	if (Ar.IsLoading() && Ar.UE4Ver() < VER_UE4_MOVE_DISTRIBUITONS_TO_POSTINITPROPS)
-	{
-		FDistributionHelpers::RestoreDefaultConstant(SpawnPerUnit.Distribution, TEXT("RequiredDistributionSpawnPerUnit"), 0.0f);
-	}
-}
-
 void UParticleModuleSpawnPerUnit::CompileModule( FParticleEmitterBuildInfo& EmitterInfo )
 {
 	EmitterInfo.SpawnPerUnitModule = this;
diff --git a/Engine/Source/Runtime/Engine/Private/Particles/ParticleModules_VectorField.cpp b/Engine/Source/Runtime/Engine/Private/Particles/ParticleModules_VectorField.cpp
--- a/Engine/Source/Runtime/Engine/Private/Particles/ParticleModules_VectorField.cpp
+++ b/Engine/Source/Runtime/Engine/Private/Particles/ParticleModules_VectorField.cpp
@@ -7,7 +7,6 @@
 #include "EnginePrivate.h"
 #include "FXSystem.h"
 #include "ParticleDefinitions.h"
-#include "../DistributionHelpers.h"
 #include "Particles/TypeData/ParticleModuleTypeDataGpu.h"
 #include "Particles/VectorField/ParticleModuleVectorFieldLocal.h"
 #include "Particles/VectorField/ParticleModuleVectorFieldGlobal.h"
@@ -55,15 +54,6 @@ void UParticleModuleVectorFieldScale::PostInitProperties()
 	}
 }
 
-void UParticleModuleVectorFieldScale::Serialize(FArchive& Ar)
-{
-	Super::Serialize(Ar);
-	if (Ar.IsLoading() && Ar.UE4Ver() < VER_UE4_MOVE_DISTRIBUITONS_TO_POSTINITPROPS)
-	{
-		FDistributionHelpers::RestoreDefaultConstant(VectorFieldScale, TEXT("DistributionVectorFieldScale"), 1.0f);
-	}
-}
-
 void UParticleModuleVectorFieldScale::CompileModule(FParticleEmitterBuildInfo& EmitterInfo)
 {
 	EmitterInfo.VectorFieldScale.ScaleByDistribution(VectorFieldScale);
@@ -105,15 +95,6 @@ void UParticleModuleVectorFieldScaleOverLife::PostInitProperties()
 	}
 }
 
-void UParticleModuleVectorFieldScaleOverLife::Serialize(FArchive& Ar)
-{
-	Super::Serialize(Ar);
-	if (Ar.IsLoading() && Ar.UE4Ver() < VER_UE4_MOVE_DISTRIBUITONS_TO_POSTINITPROPS)
-	{
-		FDistributionHelpers::RestoreDefaultConstant(VectorFieldScaleOverLife, TEXT("DistributionVectorFieldScaleOverLife"), 1.0f);
-	}
-}
-
 void UParticleModuleVectorFieldScaleOverLife::CompileModule(FParticleEmitterBuildInfo& EmitterInfo)
 {
 	EmitterInfo.VectorFieldScaleOverLife.ScaleByDistribution(VectorFieldScaleOverLife);
diff --git a/Engine/Source/Runtime/Engine/Private/Particles/ParticleModules_Velocity.cpp b/Engine/Source/Runtime/Engine/Private/Particles/ParticleModules_Velocity.cpp
--- a/Engine/Source/Runtime/Engine/Private/Particles/ParticleModules_Velocity.cpp
+++ b/Engine/Source/Runtime/Engine/Private/Particles/ParticleModules_Velocity.cpp
@@ -10,7 +10,6 @@
 #include "Distributions/DistributionVectorConstantCurve.h"
 #include "Distributions/DistributionVectorUniform.h"
 #include "ParticleDefinitions.h"
-#include "../DistributionHelpers.h"
 #include "Particles/Lifetime/ParticleModuleLifetimeBase.h"
 #include "Particles/TypeData/ParticleModuleTypeDataMesh.h"
 #include "Particles/Velocity/ParticleModuleVelocity.h"
@@ -65,16 +64,6 @@ void UParticleModuleVelocity::PostInitProperties()
 	}
 }
 
-void UParticleModuleVelocity::Serialize(FArchive& Ar)
-{
-	Super::Serialize(Ar);
-	if (Ar.IsLoading() && Ar.UE4Ver() < VER_UE4_MOVE_DISTRIBUITONS_TO_POSTINITPROPS)
-	{
-		FDistributionHelpers::RestoreDefaultUniform(StartVelocity.Distribution, TEXT("DistributionStartVelocity"), FVector::ZeroVector, FVector::ZeroVector);
-		FDistributionHelpers::RestoreDefaultUniform(StartVelocityRadial.Distribution, TEXT("DistributionStartVelocityRadial"), 0.0f, 0.0f);
-	}
-}
-
 #if WITH_EDITOR
 void UParticleModuleVelocity::PostEditChangeProperty(FPropertyChangedEvent& PropertyChangedEvent)
 {
@@ -189,15 +178,6 @@ void UParticleModuleVelocityInheritParent::PostInitProperties()
 	}
 }
 
-void UParticleModuleVelocityInheritParent::Serialize(FArchive& Ar)
-{
-	Super::Serialize(Ar);
-	if (Ar.IsLoading() && Ar.UE4Ver() < VER_UE4_MOVE_DISTRIBUITONS_TO_POSTINITPROPS)
-	{
-		FDistributionHelpers::RestoreDefaultConstant(Scale.Distribution, TEXT("DistributionScale"), FVector(1.0f, 1.0f, 1.0f));
-	}
-}
-
 #if WITH_EDITOR
 void UParticleModuleVelocityInheritParent::PostEditChangeProperty(FPropertyChangedEvent& PropertyChangedEvent)
 {
@@ -425,16 +405,6 @@ void UParticleModuleVelocityCone::PostInitProperties()
 	}
 }
 
-void UParticleModuleVelocityCone::Serialize(FArchive& Ar)
-{
-	Super::Serialize(Ar);
-	if (Ar.IsLoading() && Ar.UE4Ver() < VER_UE4_MOVE_DISTRIBUITONS_TO_POSTINITPROPS)
-	{
-		FDistributionHelpers::RestoreDefaultUniform(Angle.Distribution, TEXT("DistributionAngle"), 0.0f, 0.0f);
-		FDistributionHelpers::RestoreDefaultUniform(Velocity.Distribution, TEXT("DistributionVelocity"), 0.0f, 0.0f);
-	}
-}
-
 #if WITH_EDITOR
 void UParticleModuleVelocityCone::PostEditChangeProperty(FPropertyChangedEvent& PropertyChangedEvent)
 {
diff --git a/Engine/Source/Runtime/Engine/Private/Particles/ParticleTrailModules.cpp b/Engine/Source/Runtime/Engine/Private/Particles/ParticleTrailModules.cpp
--- a/Engine/Source/Runtime/Engine/Private/Particles/ParticleTrailModules.cpp
+++ b/Engine/Source/Runtime/Engine/Private/Particles/ParticleTrailModules.cpp
@@ -6,7 +6,6 @@
 
 #include "EnginePrivate.h"
 #include "ParticleDefinitions.h"
-#include "../DistributionHelpers.h"
 #include "Particles/Trail/ParticleModuleTrailSource.h"
 #include "Particles/TypeData/ParticleModuleTypeDataAnimTrail.h"
 #include "Particles/TypeData/ParticleModuleTypeDataRibbon.h"
@@ -66,15 +65,6 @@ void UParticleModuleTrailSource::PostInitProperties()
 	}
 }
 
-void UParticleModuleTrailSource::Serialize(FArchive& Ar)
-{
-	Super::Serialize(Ar);
-	if (Ar.IsLoading() && Ar.UE4Ver() < VER_UE4_MOVE_DISTRIBUITONS_TO_POSTINITPROPS)
-	{
-		FDistributionHelpers::RestoreDefaultConstant(SourceStrength.Distribution, TEXT("DistributionSourceStrength"), 100.0f);
-	}
-}
-
 #if WITH_EDITOR
 void UParticleModuleTrailSource::PostEditChangeProperty(FPropertyChangedEvent& PropertyChangedEvent)
 {
diff --git a/Engine/Source/Runtime/Engine/Private/PhysicsEngine/BodyInstance.cpp b/Engine/Source/Runtime/Engine/Private/PhysicsEngine/BodyInstance.cpp
--- a/Engine/Source/Runtime/Engine/Private/PhysicsEngine/BodyInstance.cpp
+++ b/Engine/Source/Runtime/Engine/Private/PhysicsEngine/BodyInstance.cpp
@@ -202,7 +202,6 @@ FBodyInstance::FBodyInstance()
 , Scale3D(1.0f)
 , SceneIndexSync(0)
 , SceneIndexAsync(0)
-, bEnableCollision_DEPRECATED(true)
 , CollisionProfileName(UCollisionProfile::CustomCollisionProfileName)
 , CollisionEnabled(ECollisionEnabled::QueryAndPhysics)
 , ObjectType(ECC_WorldStatic)
@@ -291,19 +290,6 @@ FBodyInstance* FindParentBodyInstance(FName BodyName, USkeletalMeshComponent* Sk
 	return NULL;
 }
 
-void FBodyInstance::UpdateFromDeprecatedEnableCollision()
-{
-	//@todo should I invalidate profile name if this happens?
-	if(bEnableCollision_DEPRECATED)
-	{
-		CollisionEnabled = ECollisionEnabled::QueryAndPhysics;
-	}
-	else
-	{
-		CollisionEnabled = ECollisionEnabled::NoCollision;
-	}
-}
-
 #if WITH_PHYSX
 //Determine that the shape is associated with this subbody (or root body)
 bool ShapeBoundToBody(const PxShape * PShape, const FBodyInstance* SubBody)
@@ -3852,11 +3838,6 @@ void FBodyInstance::FixupData(class UObject* Loader)
 		}
 	}
 
-	if(UE4Version < VER_UE4_CHANGE_BENABLECOLLISION_TO_COLLISIONENABLED)
-	{
-		UpdateFromDeprecatedEnableCollision();
-	}
-
 	if (UE4Version < VER_UE4_SAVE_COLLISIONRESPONSE_PER_CHANNEL)
 	{
 		CollisionResponses.SetCollisionResponseContainer(ResponseToChannels_DEPRECATED);
diff --git a/Engine/Source/Runtime/Engine/Private/PhysicsEngine/BodySetup.cpp b/Engine/Source/Runtime/Engine/Private/PhysicsEngine/BodySetup.cpp
--- a/Engine/Source/Runtime/Engine/Private/PhysicsEngine/BodySetup.cpp
+++ b/Engine/Source/Runtime/Engine/Private/PhysicsEngine/BodySetup.cpp
@@ -718,12 +718,8 @@ void UBodySetup::Serialize(FArchive& Ar)
 		BodySetupGuid = FGuid::NewGuid();
 	}
 
-	bool bCooked = false;
-	if (Ar.UE4Ver() >= VER_UE4_ADD_COOKED_TO_BODY_SETUP)
-	{
-		bCooked = Ar.IsCooking();
-		Ar << bCooked;
-	}
+	bool bCooked = Ar.IsCooking();
+	Ar << bCooked;
 
 	if (FPlatformProperties::RequiresCookedData() && !bCooked && Ar.IsLoading())
 	{
diff --git a/Engine/Source/Runtime/Engine/Private/PhysicsEngine/PhysicsConstraintTemplate.cpp b/Engine/Source/Runtime/Engine/Private/PhysicsEngine/PhysicsConstraintTemplate.cpp
--- a/Engine/Source/Runtime/Engine/Private/PhysicsEngine/PhysicsConstraintTemplate.cpp
+++ b/Engine/Source/Runtime/Engine/Private/PhysicsEngine/PhysicsConstraintTemplate.cpp
@@ -38,57 +38,6 @@ void UPhysicsConstraintTemplate::Serialize(FArchive& Ar)
 	}
 }
 
-void UPhysicsConstraintTemplate::PostLoad()
-{
-	Super::PostLoad();
-
-	//Check the version
-	if ( GetLinkerUE4Version() < VER_UE4_ANGULAR_CONSTRAINT_OPTIONS )
-	{
-		if ( DefaultInstance.bSwingLimited_DEPRECATED )
-		{
-			if (DefaultInstance.Swing1LimitAngle < RB_MinAngleToLockDOF)
-			{
-				DefaultInstance.AngularSwing1Motion = ACM_Locked;
-			}
-			else
-			{
-				DefaultInstance.AngularSwing1Motion = ACM_Limited;
-			}
-
-			if (DefaultInstance.Swing2LimitAngle < RB_MinAngleToLockDOF)
-			{
-				DefaultInstance.AngularSwing2Motion = ACM_Locked;
-			}
-			else
-			{
-				DefaultInstance.AngularSwing2Motion = ACM_Limited;
-			}
-		}
-		else
-		{
-			DefaultInstance.AngularSwing1Motion = ACM_Free;
-			DefaultInstance.AngularSwing2Motion = ACM_Free;
-		}
-
-		if ( DefaultInstance.bTwistLimited_DEPRECATED )
-		{
-			if (DefaultInstance.TwistLimitAngle < RB_MinAngleToLockDOF)
-			{
-				DefaultInstance.AngularTwistMotion = ACM_Locked;
-			}
-			else
-			{
-				DefaultInstance.AngularTwistMotion = ACM_Limited;
-			}
-		}
-		else
-		{
-			DefaultInstance.AngularTwistMotion = ACM_Free;
-		}
-	}
-}
-
 void UPhysicsConstraintTemplate::CopySetupPropsToInstance(FConstraintInstance* Instance)
 {
 	Instance->JointName			= JointName_DEPRECATED;
@@ -114,8 +63,6 @@ void UPhysicsConstraintTemplate::CopySetupPropsToInstance(FConstraintInstance* I
 	Instance->LinearLimitDamping			= LinearLimitDamping_DEPRECATED;
 	Instance->bLinearBreakable				= bLinearBreakable_DEPRECATED;
 	Instance->LinearBreakThreshold			= LinearBreakThreshold_DEPRECATED;
-	Instance->bSwingLimited_DEPRECATED		= bSwingLimited_DEPRECATED;
-	Instance->bTwistLimited_DEPRECATED		= bTwistLimited_DEPRECATED;
 	Instance->AngularSwing1Motion			= AngularSwing1Motion_DEPRECATED;
 	Instance->AngularSwing2Motion			= AngularSwing2Motion_DEPRECATED;
 	Instance->AngularTwistMotion			= AngularTwistMotion_DEPRECATED;
diff --git a/Engine/Source/Runtime/Engine/Private/PhysicsEngine/PhysicsThruster.cpp b/Engine/Source/Runtime/Engine/Private/PhysicsEngine/PhysicsThruster.cpp
--- a/Engine/Source/Runtime/Engine/Private/PhysicsEngine/PhysicsThruster.cpp
+++ b/Engine/Source/Runtime/Engine/Private/PhysicsEngine/PhysicsThruster.cpp
@@ -31,16 +31,6 @@ void UPhysicsThrusterComponent::TickComponent(float DeltaTime, enum ELevelTick T
 	}
 }
 
-void UPhysicsThrusterComponent::PostLoad()
-{
-	Super::PostLoad();
-
-	if (GetLinkerUE4Version() < VER_UE4_CONFORM_COMPONENT_ACTIVATE_FLAG)
-	{
-		bAutoActivate = bThrustEnabled_DEPRECATED;
-	}
-}
-
 //////////////////////////////////////////////////////////////////////////
 
 APhysicsThruster::APhysicsThruster(const FObjectInitializer& ObjectInitializer)
diff --git a/Engine/Source/Runtime/Engine/Private/PhysicsEngine/RadialForceComponent.cpp b/Engine/Source/Runtime/Engine/Private/PhysicsEngine/RadialForceComponent.cpp
--- a/Engine/Source/Runtime/Engine/Private/PhysicsEngine/RadialForceComponent.cpp
+++ b/Engine/Source/Runtime/Engine/Private/PhysicsEngine/RadialForceComponent.cpp
@@ -77,11 +77,6 @@ void URadialForceComponent::PostLoad()
 {
 	Super::PostLoad();
 
-	if (GetLinkerUE4Version() < VER_UE4_CONFORM_COMPONENT_ACTIVATE_FLAG)
-	{
-		bAutoActivate = bForceEnabled_DEPRECATED;
-	}
-
 	UpdateCollisionObjectQueryParams();
 }
 
diff --git a/Engine/Source/Runtime/Engine/Private/PhysicsVolume.cpp b/Engine/Source/Runtime/Engine/Private/PhysicsVolume.cpp
--- a/Engine/Source/Runtime/Engine/Private/PhysicsVolume.cpp
+++ b/Engine/Source/Runtime/Engine/Private/PhysicsVolume.cpp
@@ -8,7 +8,6 @@
 APhysicsVolume::APhysicsVolume(const FObjectInitializer& ObjectInitializer)
 	: Super(ObjectInitializer)
 {
-	GetBrushComponent()->BodyInstance.bEnableCollision_DEPRECATED = true;
 	static FName CollisionProfileName(TEXT("OverlapAllDynamic"));
 	GetBrushComponent()->SetCollisionProfileName(CollisionProfileName);
 
diff --git a/Engine/Source/Runtime/Engine/Private/Polygon.cpp b/Engine/Source/Runtime/Engine/Private/Polygon.cpp
--- a/Engine/Source/Runtime/Engine/Private/Polygon.cpp
+++ b/Engine/Source/Runtime/Engine/Private/Polygon.cpp
@@ -43,12 +43,6 @@ FArchive& operator<<(FArchive& Ar, FLightmassPrimitiveSettings& Settings)
 	Ar << Settings.EmissiveBoost;
 	Ar << Settings.DiffuseBoost;
 
-	if( Ar.UE4Ver() < VER_UE4_REMOVED_SPECULAR_BOOST )
-	{
-		float SpecularBoost;
-		Ar << SpecularBoost;
-	}
-
 	return Ar;
 }
 
diff --git a/Engine/Source/Runtime/Engine/Private/PostProcessVolume.cpp b/Engine/Source/Runtime/Engine/Private/PostProcessVolume.cpp
--- a/Engine/Source/Runtime/Engine/Private/PostProcessVolume.cpp
+++ b/Engine/Source/Runtime/Engine/Private/PostProcessVolume.cpp
@@ -7,7 +7,6 @@
 APostProcessVolume::APostProcessVolume(const FObjectInitializer& ObjectInitializer)
 	: Super(ObjectInitializer)
 {
-	GetBrushComponent()->BodyInstance.bEnableCollision_DEPRECATED = false;
 	GetBrushComponent()->SetCollisionProfileName(UCollisionProfile::NoCollision_ProfileName);
 	// post process volume needs physics data for trace
 	GetBrushComponent()->bAlwaysCreatePhysicsState = true;
diff --git a/Engine/Source/Runtime/Engine/Private/PrecomputedLightVolume.cpp b/Engine/Source/Runtime/Engine/Private/PrecomputedLightVolume.cpp
--- a/Engine/Source/Runtime/Engine/Private/PrecomputedLightVolume.cpp
+++ b/Engine/Source/Runtime/Engine/Private/PrecomputedLightVolume.cpp
@@ -12,21 +12,7 @@ FArchive& operator<<(FArchive& Ar, FVolumeLightingSample& Sample)
 {
 	Ar << Sample.Position;
 	Ar << Sample.Radius;
-
-	if (Ar.UE4Ver() < VER_UE4_CHANGED_VOLUME_SAMPLE_FORMAT)
-	{
-		uint8 Temp = 0;
-		FColor Temp2(0, 0, 0, 0);
-
-		Ar << Temp << Temp;
-		Ar << Temp << Temp;
-		Ar << Temp2 << Temp2 << Temp2;
-		Ar << Temp;
-	}
-	else
-	{
-		Ar << Sample.Lighting;
-	}
+	Ar << Sample.Lighting;
 
 	if (Ar.UE4Ver() >= VER_UE4_SKY_BENT_NORMAL)
 	{
diff --git a/Engine/Source/Runtime/Engine/Private/PrimitiveComponent.cpp b/Engine/Source/Runtime/Engine/Private/PrimitiveComponent.cpp
--- a/Engine/Source/Runtime/Engine/Private/PrimitiveComponent.cpp
+++ b/Engine/Source/Runtime/Engine/Private/PrimitiveComponent.cpp
@@ -93,7 +93,6 @@ UPrimitiveComponent::UPrimitiveComponent(const FObjectInitializer& ObjectInitial
 	bSelectable = true;
 	AlwaysLoadOnClient = true;
 	AlwaysLoadOnServer = true;
-	BodyInstance.bEnableCollision_DEPRECATED = true;
 	SetCollisionProfileName(UCollisionProfile::BlockAll_ProfileName);
 	bAlwaysCreatePhysicsState = false;
 	bRenderInMainPass = true;
@@ -667,21 +666,6 @@ void UPrimitiveComponent::PostLoad()
 		BodyInstance.FixupData(this);
 	}
 
-	if(UE4Version < VER_UE4_VISIBILITY_FLAG_CHANGES)
-	{
-		// copy visible flag if before it wans't hidden or drawingame was true
-		bVisible = !HiddenGame_DEPRECATED || DrawInGame_DEPRECATED;
-	}
-
-	if(UE4Version < VER_UE4_DEPRECATED_BNOENCROACHCHECK)
-	{
-		if (bNoEncroachCheck_DEPRECATED)
-		{
-			bGenerateOverlapEvents = false;
-		}
-		// else use defaults
-	}
-
 	if (UE4Version < VER_UE4_RENAME_CANBECHARACTERBASE)
 	{
 		CanCharacterStepUpOn = CanBeCharacterBase_DEPRECATED;
diff --git a/Engine/Source/Runtime/Engine/Private/ReflectionCaptureComponent.cpp b/Engine/Source/Runtime/Engine/Private/ReflectionCaptureComponent.cpp
--- a/Engine/Source/Runtime/Engine/Private/ReflectionCaptureComponent.cpp
+++ b/Engine/Source/Runtime/Engine/Private/ReflectionCaptureComponent.cpp
@@ -86,7 +86,6 @@ AReflectionCapture::AReflectionCapture(const FObjectInitializer& ObjectInitializ
 		SpriteComponent->RelativeScale3D = FVector(0.5f, 0.5f, 0.5f);
 		SpriteComponent->bHiddenInGame = true;
 		SpriteComponent->bAbsoluteScale = true;
-		SpriteComponent->BodyInstance.bEnableCollision_DEPRECATED = false;	
 		SpriteComponent->SetCollisionProfileName(UCollisionProfile::NoCollision_ProfileName);
 		SpriteComponent->bIsScreenSizeScaled = true;
 	}
@@ -118,7 +117,6 @@ ASphereReflectionCapture::ASphereReflectionCapture(const FObjectInitializer& Obj
 	DrawInfluenceRadius->AttachParent = GetCaptureComponent();
 	DrawInfluenceRadius->bDrawOnlyIfSelected = true;
 	DrawInfluenceRadius->bUseEditorCompositing = true;
-	DrawInfluenceRadius->BodyInstance.bEnableCollision_DEPRECATED = false;
 	DrawInfluenceRadius->SetCollisionProfileName(UCollisionProfile::NoCollision_ProfileName);
 	SphereComponent->PreviewInfluenceRadius = DrawInfluenceRadius;
 
@@ -126,7 +124,6 @@ ASphereReflectionCapture::ASphereReflectionCapture(const FObjectInitializer& Obj
 	DrawCaptureRadius->AttachParent = GetCaptureComponent();
 	DrawCaptureRadius->bDrawOnlyIfSelected = true;
 	DrawCaptureRadius->bUseEditorCompositing = true;
-	DrawCaptureRadius->BodyInstance.bEnableCollision_DEPRECATED = false;
 	DrawCaptureRadius->SetCollisionProfileName(UCollisionProfile::NoCollision_ProfileName);
 	DrawCaptureRadius->ShapeColor = FColor(100, 90, 40);
 }
@@ -169,7 +166,6 @@ ABoxReflectionCapture::ABoxReflectionCapture(const FObjectInitializer& ObjectIni
 	DrawInfluenceBox->AttachParent = GetCaptureComponent();
 	DrawInfluenceBox->bDrawOnlyIfSelected = true;
 	DrawInfluenceBox->bUseEditorCompositing = true;
-	DrawInfluenceBox->BodyInstance.bEnableCollision_DEPRECATED = false;
 	DrawInfluenceBox->SetCollisionProfileName(UCollisionProfile::NoCollision_ProfileName);
 	DrawInfluenceBox->InitBoxExtent(FVector(1, 1, 1));
 	BoxComponent->PreviewInfluenceBox = DrawInfluenceBox;
@@ -178,7 +174,6 @@ ABoxReflectionCapture::ABoxReflectionCapture(const FObjectInitializer& ObjectIni
 	DrawCaptureBox->AttachParent = GetCaptureComponent();
 	DrawCaptureBox->bDrawOnlyIfSelected = true;
 	DrawCaptureBox->bUseEditorCompositing = true;
-	DrawCaptureBox->BodyInstance.bEnableCollision_DEPRECATED = false;
 	DrawCaptureBox->SetCollisionProfileName(UCollisionProfile::NoCollision_ProfileName);
 	DrawCaptureBox->ShapeColor = FColor(100, 90, 40);
 	DrawCaptureBox->InitBoxExtent(FVector(1, 1, 1));
@@ -202,7 +197,6 @@ APlaneReflectionCapture::APlaneReflectionCapture(const FObjectInitializer& Objec
 	DrawInfluenceRadius->bDrawOnlyIfSelected = true;
 	DrawInfluenceRadius->bAbsoluteScale = true;
 	DrawInfluenceRadius->bUseEditorCompositing = true;
-	DrawInfluenceRadius->BodyInstance.bEnableCollision_DEPRECATED = false;
 	DrawInfluenceRadius->SetCollisionProfileName(UCollisionProfile::NoCollision_ProfileName);
 	PlaneComponent->PreviewInfluenceRadius = DrawInfluenceRadius;
 
@@ -210,7 +204,6 @@ APlaneReflectionCapture::APlaneReflectionCapture(const FObjectInitializer& Objec
 	DrawCaptureBox->AttachParent = GetCaptureComponent();
 	DrawCaptureBox->bDrawOnlyIfSelected = true;
 	DrawCaptureBox->bUseEditorCompositing = true;
-	DrawCaptureBox->BodyInstance.bEnableCollision_DEPRECATED = false;
 	DrawCaptureBox->SetCollisionProfileName(UCollisionProfile::NoCollision_ProfileName);
 	DrawCaptureBox->ShapeColor = FColor(100, 90, 40);
 	DrawCaptureBox->InitBoxExtent(FVector(1, 1, 1));
diff --git a/Engine/Source/Runtime/Engine/Private/SceneCaptureComponent.cpp b/Engine/Source/Runtime/Engine/Private/SceneCaptureComponent.cpp
--- a/Engine/Source/Runtime/Engine/Private/SceneCaptureComponent.cpp
+++ b/Engine/Source/Runtime/Engine/Private/SceneCaptureComponent.cpp
@@ -17,7 +17,6 @@ ASceneCapture::ASceneCapture(const FObjectInitializer& ObjectInitializer)
 	: Super(ObjectInitializer)
 {
 	MeshComp = ObjectInitializer.CreateDefaultSubobject<UStaticMeshComponent>(this, TEXT("CamMesh0"));
-	MeshComp->BodyInstance.bEnableCollision_DEPRECATED = false;
 
 	MeshComp->SetCollisionProfileName(UCollisionProfile::NoCollision_ProfileName);
 
diff --git a/Engine/Source/Runtime/Engine/Private/SceneComponent.cpp b/Engine/Source/Runtime/Engine/Private/SceneComponent.cpp
--- a/Engine/Source/Runtime/Engine/Private/SceneComponent.cpp
+++ b/Engine/Source/Runtime/Engine/Private/SceneComponent.cpp
@@ -27,7 +27,6 @@ USceneComponent::USceneComponent(const FObjectInitializer& ObjectInitializer)
 	: Super(ObjectInitializer)
 {
 	Mobility = EComponentMobility::Movable;
-	RelativeScale_DEPRECATED = 1.0f;
 	RelativeScale3D = FVector(1.0f,1.0f,1.0f);
 	// default behavior is visible
 	bVisible = true;
