Bring minimum engine version up to 214

diff --git a/Engine/Source/Runtime/Engine/Private/SkeletalMesh.cpp b/Engine/Source/Runtime/Engine/Private/SkeletalMesh.cpp
--- a/Engine/Source/Runtime/Engine/Private/SkeletalMesh.cpp
+++ b/Engine/Source/Runtime/Engine/Private/SkeletalMesh.cpp
@@ -1178,11 +1178,8 @@ void FStaticLODModel::Serialize( FArchive& Ar, UObject* Owner, int32 Idx )
 		RawPointIndices.Serialize( Ar, Owner );
 	}
 
-	if(Ar.UE4Ver() >= VER_UE4_ADD_SKELMESH_MESHTOIMPORTVERTEXMAP)
-	{
-		Ar << MeshToImportVertexMap;
-		Ar << MaxImportVertex;
-	}
+	Ar << MeshToImportVertexMap;
+	Ar << MaxImportVertex;
 
 	if( !StripFlags.IsDataStrippedForServer() )
 	{
@@ -1199,11 +1196,6 @@ void FStaticLODModel::Serialize( FArchive& Ar, UObject* Owner, int32 Idx )
 		{
 			Ar << ColorVertexBuffer;
 		}
-		if( Ar.UE4Ver() < VER_UE4_REMOVE_EXTRA_SKELMESH_VERTEX_INFLUENCES )
-		{
-			TArray<FSkeletalMeshVertexInfluences_DEPRECATED> DummyVertexInfluences;
-			Ar << DummyVertexInfluences;
-		}
 
 		if ( !StripFlags.IsClassDataStripped( LodAdjacencyStripFlag ) )
 		{
@@ -2344,15 +2336,6 @@ void USkeletalMesh::Serialize( FArchive& Ar )
 		Ar << DummyNameIndexMap;
 	}
 
-	if (Ar.UE4Ver() < VER_UE4_REMOVE_EXTRA_SKELMESH_VERTEX_INFLUENCES)
-	{
-		TArray<FString> DummyBoneBreakNames;
-		Ar << DummyBoneBreakNames;
-
-		TArray<uint8> DummyBoneBreakOptions;
-		Ar << DummyBoneBreakOptions;
-	}
-
 	//@todo legacy
 	TArray<UObject*> DummyObjs;
 	Ar << DummyObjs;
@@ -2365,15 +2348,6 @@ void USkeletalMesh::Serialize( FArchive& Ar )
 		SkelSourceData.Serialize( Ar, this );
 	}
 
-#if WITH_EDITORONLY_DATA
-	if (BoundsPreviewAsset_DEPRECATED && BoundsPreviewAsset_DEPRECATED != PhysicsAsset)
-	{
-		PhysicsAsset = BoundsPreviewAsset_DEPRECATED;
-		BoundsPreviewAsset_DEPRECATED = NULL;
-		MarkPackageDirty();
-	}
-#endif
-
 #if WITH_EDITORONLY_DATA
 	// SourceFilePath and SourceFileTimestamp were moved into a subobject
 	if ( Ar.IsLoading() && Ar.UE4Ver() < VER_UE4_ADDED_FBX_ASSET_IMPORT_DATA )
@@ -2600,36 +2574,6 @@ void USkeletalMesh::PostLoad()
 		}
 	}
 
-#if WITH_EDITORONLY_DATA
-	if( GetLinker() && (GetLinker()->UE4Ver() < VER_UE4_FIX_REQUIRED_BONES) )
-	{
-		for (int32 LodIndex=0; LodIndex<ImportedResource->LODModels.Num(); LodIndex++)
-		{
-			FStaticLODModel& LODModel = ImportedResource->LODModels[LodIndex];
-			CalculateRequiredBones(LODModel,RefSkeleton,NULL);
-		}
-	}
-#endif
-
-	if ( MorphTargetTable_DEPRECATED.Num() )
-	{
-		TArray<FMorphTargetMap> OldMorphTargetTable = MorphTargetTable_DEPRECATED;
-		MorphTargets.Empty(OldMorphTargetTable.Num());
-
-		for ( auto TableIter = MorphTargetTable_DEPRECATED.CreateIterator(); TableIter; ++TableIter )
-		{
-			// @todo we can remove Name soon once this is all converted
-			// Name should match with object name
-			FMorphTargetMap MorphTargetMap = *TableIter;
-			if ( MorphTargetMap.MorphTarget )
-			{
-				MorphTargets.Add(MorphTargetMap.MorphTarget);
-			}
-		}
-
-		MorphTargetTable_DEPRECATED.Empty();
-	}
-
 	// Revert to using 32 bit Float UVs on hardware that doesn't support rendering with 16 bit Float UVs 
 	if( !bUseFullPrecisionUVs && !GVertexElementTypeSupport.IsSupported(VET_Half2) )
 	{
@@ -3676,7 +3620,6 @@ ASkeletalMeshActor::ASkeletalMeshActor(const FObjectInitializer& ObjectInitializ
 
 	SkeletalMeshComponent = ObjectInitializer.CreateDefaultSubobject<USkeletalMeshComponent>(this, TEXT("SkeletalMeshComponent0"));
 	SkeletalMeshComponent->MeshComponentUpdateFlag = EMeshComponentUpdateFlag::AlwaysTickPose;
-	SkeletalMeshComponent->BodyInstance.bEnableCollision_DEPRECATED = true;
 	// check BaseEngine.ini for profile setup
 	SkeletalMeshComponent->SetCollisionProfileName(UCollisionProfile::PhysicsActor_ProfileName);
 	RootComponent = SkeletalMeshComponent;
@@ -4594,9 +4537,7 @@ USkinnedMeshComponent::USkinnedMeshComponent(const FObjectInitializer& ObjectIni
 	PrimaryComponentTick.TickGroup = TG_PrePhysics;	
 	WireframeColor = FColor(221, 221, 28, 255);
 
-	bUpdateSkelWhenNotRendered_DEPRECATED = true;
 	MeshComponentUpdateFlag = EMeshComponentUpdateFlag::AlwaysTickPoseAndRefreshBones;
-	BodyInstance.bEnableCollision_DEPRECATED = false;
 
 	SetCollisionProfileName(UCollisionProfile::NoCollision_ProfileName);
 
@@ -4639,18 +4580,6 @@ void USkinnedMeshComponent::Serialize(FArchive& Ar)
 		SpaceBasesArray[1].CountBytes(Ar);
 		MasterBoneMap.CountBytes(Ar);
 	}
-
-	if (Ar.UE4Ver() < VER_UE4_CONSOLIDATE_SKINNEDMESH_UPDATE_FLAGS)
-	{
-		if (bUpdateSkelWhenNotRendered_DEPRECATED)
-		{
-			MeshComponentUpdateFlag = EMeshComponentUpdateFlag::AlwaysTickPoseAndRefreshBones;
-		}
-		else
-		{
-			MeshComponentUpdateFlag = EMeshComponentUpdateFlag::OnlyTickPoseWhenRendered;
-		}
-	}
 }
 
 void USkeletalMeshComponent::Serialize(FArchive& Ar)
@@ -4674,23 +4603,6 @@ void USkeletalMeshComponent::Serialize(FArchive& Ar)
 		RequiredBones.CountBytes(Ar);
 	}
 
-	// super will get first, and this will fix up all tick related issue
-	if (Ar.UE4Ver() < VER_UE4_CONSOLIDATE_SKINNEDMESH_UPDATE_FLAGS)
-	{
-		if (bUpdateSkelWhenNotRendered_DEPRECATED && bTickAnimationWhenNotRendered_DEPRECATED)
-		{
-			MeshComponentUpdateFlag = EMeshComponentUpdateFlag::AlwaysTickPoseAndRefreshBones;
-		}
-		else if (bTickAnimationWhenNotRendered_DEPRECATED)
-		{
-			MeshComponentUpdateFlag = EMeshComponentUpdateFlag::AlwaysTickPose;
-		}
-		else
-		{
-			MeshComponentUpdateFlag = EMeshComponentUpdateFlag::OnlyTickPoseWhenRendered;
-		}
-	}
-
 	if (Ar.UE4Ver() < VER_UE4_REMOVE_SKELETALMESH_COMPONENT_BODYSETUP_SERIALIZATION)
 	{
 		//we used to serialize bodysetup of skeletal mesh component. We no longer do this, but need to not break existing content
diff --git a/Engine/Source/Runtime/Engine/Private/SkeletalMeshComponent.cpp b/Engine/Source/Runtime/Engine/Private/SkeletalMeshComponent.cpp
--- a/Engine/Source/Runtime/Engine/Private/SkeletalMeshComponent.cpp
+++ b/Engine/Source/Runtime/Engine/Private/SkeletalMeshComponent.cpp
@@ -101,7 +101,6 @@ USkeletalMeshComponent::USkeletalMeshComponent(const FObjectInitializer& ObjectI
 	bWantsInitializeComponent = true;
 	GlobalAnimRateScale = 1.0f;
 	bNoSkeletonUpdate = false;
-	bTickAnimationWhenNotRendered_DEPRECATED = true;
 	MeshComponentUpdateFlag = EMeshComponentUpdateFlag::AlwaysTickPoseAndRefreshBones;
 	KinematicBonesUpdateType = EKinematicBonesUpdateToPhysics::SkipSimulatingBones;
 	bGenerateOverlapEvents = false;
diff --git a/Engine/Source/Runtime/Engine/Private/SoundCue.cpp b/Engine/Source/Runtime/Engine/Private/SoundCue.cpp
--- a/Engine/Source/Runtime/Engine/Private/SoundCue.cpp
+++ b/Engine/Source/Runtime/Engine/Private/SoundCue.cpp
@@ -42,12 +42,7 @@ void USoundCue::Serialize( FArchive& Ar )
 
 	Super::Serialize( Ar );
 
-	if (Ar.IsLoading() && Ar.UE4Ver() < VER_UE4_SOUND_NODE_INHERIT_FROM_ED_GRAPH_NODE)
-	{
-		Ar << EditorData_DEPRECATED;
-	}
-
-	if (!Ar.IsCooking() && Ar.UE4Ver() >= VER_UE4_SOUND_CUE_GRAPH_EDITOR)
+	if (!Ar.IsCooking())
 	{
 		Ar << SoundCueGraph;
 	}
@@ -66,47 +61,7 @@ void USoundCue::PostLoad()
 {
 	Super::PostLoad();
 
-	if (GetLinkerUE4Version() < VER_UE4_SOUND_NODE_INHERIT_FROM_ED_GRAPH_NODE)
-	{
-		CreateGraph();
-
-		for(TMap<USoundNode*,FSoundNodeEditorData>::TConstIterator It(EditorData_DEPRECATED); It; ++It)
-		{
-			USoundNode* Node = It.Key();
-			if (Node)
-			{
-				SetupSoundNode(Node);
-
-				Node->GraphNode->NodePosX = -It.Value().NodePosX - Node->GraphNode->EstimateNodeWidth();
-				Node->GraphNode->NodePosY = It.Value().NodePosY;
-
-				AllNodes.Add(Node);
-			}
-		}
-		EditorData_DEPRECATED.Empty();
-
-		LinkGraphNodesFromSoundNodes();
-	}
-	else if (GetLinkerUE4Version() < VER_UE4_SOUND_CUE_GRAPH_EDITOR)
-	{
-		CreateGraph();
-
-		for(TArray<USoundNode*>::TConstIterator It(AllNodes); It; ++It)
-		{
-			USoundNode* Node = *It;
-			if (Node)
-			{
-				SetupSoundNode(Node);
-
-				Node->GraphNode->NodePosX = -Node->NodePosX_DEPRECATED - Node->GraphNode->EstimateNodeWidth();
-				Node->GraphNode->NodePosY = Node->NodePosY_DEPRECATED;
-			}
-		}
-
-		LinkGraphNodesFromSoundNodes();
-	}
-
-	// Game doesn't care if there are NULL graph nodes
+		// Game doesn't care if there are NULL graph nodes
 	if (GIsEditor)
 	{
 		// Deal with SoundNode types being removed - iterate in reverse as nodes may be removed
@@ -144,7 +99,7 @@ template<typename T> ENGINE_API void USoundCue::RecursiveFindNode( USoundNode* N
 {
 	if( Node )
 	{
-		// Record the node if it is attenuation
+		// Record the node if it is the desired type
 		if( Node->IsA( T::StaticClass() ) )
 		{
 			OutNodes.AddUnique( static_cast<T*>( Node ) );
diff --git a/Engine/Source/Runtime/Engine/Private/SoundNode.cpp b/Engine/Source/Runtime/Engine/Private/SoundNode.cpp
--- a/Engine/Source/Runtime/Engine/Private/SoundNode.cpp
+++ b/Engine/Source/Runtime/Engine/Private/SoundNode.cpp
@@ -22,7 +22,7 @@ void USoundNode::Serialize(FArchive& Ar)
 {
 	Super::Serialize(Ar);
 
-	if (!Ar.IsCooking() && Ar.UE4Ver() >= VER_UE4_SOUND_CUE_GRAPH_EDITOR)
+	if (!Ar.IsCooking())
 	{
 		Ar << GraphNode;
 	}
diff --git a/Engine/Source/Runtime/Engine/Private/SoundNodeAttenuation.cpp b/Engine/Source/Runtime/Engine/Private/SoundNodeAttenuation.cpp
--- a/Engine/Source/Runtime/Engine/Private/SoundNodeAttenuation.cpp
+++ b/Engine/Source/Runtime/Engine/Private/SoundNodeAttenuation.cpp
@@ -12,60 +12,6 @@
 USoundNodeAttenuation::USoundNodeAttenuation(const FObjectInitializer& ObjectInitializer)
 	: Super(ObjectInitializer)
 {
-	bAttenuate_DEPRECATED = true;
-	bSpatialize_DEPRECATED = true;
-	dBAttenuationAtMax_DEPRECATED = -60;
-	DistanceAlgorithm_DEPRECATED = ATTENUATION_Linear;
-	RadiusMin_DEPRECATED = 400;
-	RadiusMax_DEPRECATED = 4000;
-	LPFRadiusMin_DEPRECATED = 3000;
-	LPFRadiusMax_DEPRECATED = 6000;
-}
-
-void USoundNodeAttenuation::Serialize(FArchive& Ar)
-{
-	Super::Serialize(Ar);
-	if (Ar.IsLoading())
-	{
-		if (Ar.UE4Ver() < VER_UE4_SOUND_NODE_ATTENUATION_SETTINGS_CHANGE)
-		{
-			bOverrideAttenuation = true;
-			AttenuationOverrides.bAttenuate			= bAttenuate_DEPRECATED;
-			AttenuationOverrides.bSpatialize		= bSpatialize_DEPRECATED;
-			AttenuationOverrides.dBAttenuationAtMax	= dBAttenuationAtMax_DEPRECATED;
-			AttenuationOverrides.DistanceAlgorithm	= DistanceAlgorithm_DEPRECATED;
-			AttenuationOverrides.bAttenuateWithLPF	= bAttenuateWithLPF_DEPRECATED;
-			AttenuationOverrides.LPFRadiusMin		= LPFRadiusMin_DEPRECATED;
-			AttenuationOverrides.LPFRadiusMax		= LPFRadiusMax_DEPRECATED;
-
-			AttenuationOverrides.RadiusMin_DEPRECATED = RadiusMin_DEPRECATED;
-			AttenuationOverrides.RadiusMax_DEPRECATED = RadiusMax_DEPRECATED;
-			AttenuationOverrides.FalloffDistance	  = RadiusMax_DEPRECATED - RadiusMin_DEPRECATED;
-
-			switch(DistanceType_DEPRECATED)
-			{
-			case SOUNDDISTANCE_Normal:
-				AttenuationOverrides.AttenuationShape = EAttenuationShape::Sphere;
-				AttenuationOverrides.AttenuationShapeExtents = FVector(RadiusMin_DEPRECATED, 0.f, 0.f);
-				break;
-
-			case SOUNDDISTANCE_InfiniteXYPlane:
-				AttenuationOverrides.AttenuationShape = EAttenuationShape::Box;
-				AttenuationOverrides.AttenuationShapeExtents = FVector(WORLD_MAX, WORLD_MAX, RadiusMin_DEPRECATED);
-				break;
-
-			case SOUNDDISTANCE_InfiniteXZPlane:
-				AttenuationOverrides.AttenuationShape = EAttenuationShape::Box;
-				AttenuationOverrides.AttenuationShapeExtents = FVector(WORLD_MAX, RadiusMin_DEPRECATED, WORLD_MAX);
-				break;
-
-			case SOUNDDISTANCE_InfiniteYZPlane:
-				AttenuationOverrides.AttenuationShape = EAttenuationShape::Box;
-				AttenuationOverrides.AttenuationShapeExtents = FVector(RadiusMin_DEPRECATED, WORLD_MAX, WORLD_MAX);
-				break;
-			}
-		}
-	}
 }
 
 float USoundNodeAttenuation::MaxAudibleDistance( float CurrentMaxDistance ) 
diff --git a/Engine/Source/Runtime/Engine/Private/SoundNodeEnveloper.cpp b/Engine/Source/Runtime/Engine/Private/SoundNodeEnveloper.cpp
--- a/Engine/Source/Runtime/Engine/Private/SoundNodeEnveloper.cpp
+++ b/Engine/Source/Runtime/Engine/Private/SoundNodeEnveloper.cpp
@@ -33,18 +33,7 @@ void USoundNodeEnveloper::Serialize(FArchive& Ar)
 	Super::Serialize(Ar);
 	if (Ar.IsLoading())
 	{
-		if (Ar.UE4Ver() < VER_UE4_MOVE_DISTRIBUITONS_TO_POSTINITPROPS)
-		{
-			if (VolumeInterpCurve_DEPRECATED != NULL && VolumeInterpCurve_DEPRECATED->ConstantCurve.Points.Num() == 0)
-			{
-				VolumeCurve.EditorCurveData.AddKey(0.0f, 1.0f);
-			}
-			if (PitchInterpCurve_DEPRECATED != NULL && PitchInterpCurve_DEPRECATED->ConstantCurve.Points.Num() == 0)
-			{
-				PitchCurve.EditorCurveData.AddKey(0.0f, 1.0f);
-			}
-		}
-		else if (Ar.UE4Ver() < VER_UE4_SOUND_NODE_ENVELOPER_CURVE_CHANGE)
+		if (Ar.UE4Ver() < VER_UE4_SOUND_NODE_ENVELOPER_CURVE_CHANGE)
 		{
 			if (VolumeInterpCurve_DEPRECATED)
 			{
diff --git a/Engine/Source/Runtime/Engine/Private/SoundNodeModulatorContinuous.cpp b/Engine/Source/Runtime/Engine/Private/SoundNodeModulatorContinuous.cpp
--- a/Engine/Source/Runtime/Engine/Private/SoundNodeModulatorContinuous.cpp
+++ b/Engine/Source/Runtime/Engine/Private/SoundNodeModulatorContinuous.cpp
@@ -3,9 +3,7 @@
 
 #include "EnginePrivate.h"
 #include "SoundDefinitions.h"
-#include "DistributionHelpers.h"
 #include "Sound/SoundNodeModulatorContinuous.h"
-#include "Distributions/DistributionFloatSoundParameter.h"
 
 float FModulatorContinuousParams::GetValue(const FActiveSound& ActiveSound) const
 {
@@ -48,46 +46,6 @@ USoundNodeModulatorContinuous::USoundNodeModulatorContinuous(const FObjectInitia
 {
 }
 
-void USoundNodeModulatorContinuous::Serialize(FArchive& Ar)
-{
-	Super::Serialize(Ar);
-	if (Ar.IsLoading())
-{
-		if (Ar.UE4Ver() < VER_UE4_MOVE_DISTRIBUITONS_TO_POSTINITPROPS)
-	{
-			FDistributionHelpers::RestoreDefaultUniform(PitchModulation_DEPRECATED.Distribution, TEXT("DistributionPitch"), 0.95f, 1.05f);
-			FDistributionHelpers::RestoreDefaultUniform(VolumeModulation_DEPRECATED.Distribution, TEXT("DistributionVolume"), 0.95f, 1.05f);
-	}
-		if (Ar.UE4Ver() < VER_UE4_MODULATOR_CONTINUOUS_NO_DISTRIBUTION)
-		{
-			// If it is a float sound parameter we convert, otherwise it will become a default 1.0 multiplier that has no effect
-			UDEPRECATED_DistributionFloatSoundParameter* Distribution = Cast<UDEPRECATED_DistributionFloatSoundParameter>(PitchModulation_DEPRECATED.Distribution);
-			if (Distribution)
-			{
-				PitchModulationParams.ParameterName = Distribution->ParameterName;
-				PitchModulationParams.Default		= Distribution->Constant;
-				PitchModulationParams.MinInput		= Distribution->MinInput;
-				PitchModulationParams.MaxInput		= Distribution->MaxInput;
-				PitchModulationParams.MinOutput		= Distribution->MinOutput;
-				PitchModulationParams.MaxOutput		= Distribution->MaxOutput;
-				PitchModulationParams.ParamMode		= (ModulationParamMode)(int32)Distribution->ParamMode;
-}
-
-			Distribution = Cast<UDEPRECATED_DistributionFloatSoundParameter>(VolumeModulation_DEPRECATED.Distribution);
-			if (Distribution)
-{
-				VolumeModulationParams.ParameterName	= Distribution->ParameterName;
-				VolumeModulationParams.Default			= Distribution->Constant;
-				VolumeModulationParams.MinInput			= Distribution->MinInput;
-				VolumeModulationParams.MaxInput			= Distribution->MaxInput;
-				VolumeModulationParams.MinOutput		= Distribution->MinOutput;
-				VolumeModulationParams.MaxOutput		= Distribution->MaxOutput;
-				VolumeModulationParams.ParamMode		= (ModulationParamMode)(int32)Distribution->ParamMode;
-			}
-		}
-	}
-}
-
 void USoundNodeModulatorContinuous::ParseNodes( FAudioDevice* AudioDevice, const UPTRINT NodeWaveInstanceHash, FActiveSound& ActiveSound, const FSoundParseParameters& ParseParams, TArray<FWaveInstance*>& WaveInstances )
 {
 	FSoundParseParameters UpdatedParams = ParseParams;
diff --git a/Engine/Source/Runtime/Engine/Private/StaticMesh.cpp b/Engine/Source/Runtime/Engine/Private/StaticMesh.cpp
--- a/Engine/Source/Runtime/Engine/Private/StaticMesh.cpp
+++ b/Engine/Source/Runtime/Engine/Private/StaticMesh.cpp
@@ -1681,22 +1681,13 @@ void UStaticMesh::Serialize(FArchive& Ar)
 	FStripDataFlags StripFlags( Ar );
 
 	bool bCooked = Ar.IsCooking();
-	if (Ar.UE4Ver() >= VER_UE4_STATIC_MESH_REFACTOR)
-	{
-		Ar << bCooked;
-	}
+	Ar << bCooked;
 
 #if WITH_EDITORONLY_DATA
 	if (Ar.IsLoading() && Ar.UE4Ver() < VER_UE4_REMOVE_ZERO_TRIANGLE_SECTIONS)
 	{
 		GStaticMeshesThatNeedMaterialFixup.Set(this);
 	}
-
-	FBoxSphereBounds LegacyBounds;
-	if (Ar.UE4Ver() < VER_UE4_STATIC_MESH_REFACTOR)
-	{
-		Ar << LegacyBounds;
-	}
 #endif // #if WITH_EDITORONLY_DATA
 
 	Ar << BodySetup;
@@ -1706,21 +1697,6 @@ void UStaticMesh::Serialize(FArchive& Ar)
 		Ar << NavCollision;
 	}
 
-#if WITH_EDITOR
-	if (Ar.UE4Ver() < VER_UE4_STATIC_MESH_REFACTOR)
-	{
-		int64 StartCycles = FPlatformTime::Cycles();
-		SerializeLegacySouceData(Ar, LegacyBounds);
-		int64 EndCycles = FPlatformTime::Cycles();
-		StaticMeshDerivedDataTimings::ConvertCycles += (EndCycles - StartCycles);
-		UE_LOG(LogStaticMesh,Verbose,
-			TEXT("Converting legacy source data for %s took %fs"),
-			*GetPathName(),
-			FPlatformTime::ToSeconds(EndCycles - StartCycles)
-			);
-	}
-#endif
-
 #if WITH_EDITORONLY_DATA
 	if( !StripFlags.IsEditorDataStripped() )
 	{
@@ -1758,40 +1734,10 @@ void UStaticMesh::Serialize(FArchive& Ar)
 	}
 
 	Ar << LightingGuid;
-
-	if (Ar.UE4Ver() < VER_UE4_STATIC_MESH_REFACTOR)
-	{
-		int32 VertexPositionVersionNumber = 0;
-		Ar << VertexPositionVersionNumber;
-	}
-
-	if (Ar.UE4Ver() < VER_UE4_REMOVE_CACHED_STATIC_MESH_STREAMING_FACTORS)
-	{
-		TArray<float> CachedStreamingTextureFactors;
-		Ar << CachedStreamingTextureFactors;
-	}
-
-#if WITH_EDITORONLY_DATA
-	if (Ar.UE4Ver() < VER_UE4_STATIC_MESH_REFACTOR)
-	{
-		bool bRemoveDegenerates_DEPRECATED;
-		Ar << bRemoveDegenerates_DEPRECATED;
-		for (int32 i = 0; i < SourceModels.Num(); ++i)
-		{
-			FStaticMeshSourceModel& SrcModel = SourceModels[i];
-			SrcModel.BuildSettings.bRemoveDegenerates = SrcModel.BuildSettings.bRemoveDegenerates && bRemoveDegenerates_DEPRECATED;
-			SrcModel.BuildSettings.bUseFullPrecisionUVs = UseFullPrecisionUVs_DEPRECATED;
-		}
-	}
-#endif // #if WITH_EDITORONLY_DATA
-
-	if (Ar.UE4Ver() >= VER_UE4_STATIC_MESH_SOCKETS)
-	{
-		Ar << Sockets;
-	}
+	Ar << Sockets;
 
 #if WITH_EDITOR
-	if (Ar.UE4Ver() >= VER_UE4_STATIC_MESH_REFACTOR && !StripFlags.IsEditorDataStripped())
+	if (!StripFlags.IsEditorDataStripped())
 	{
 		for (int32 i = 0; i < SourceModels.Num(); ++i)
 		{
diff --git a/Engine/Source/Runtime/Engine/Private/StaticMeshBuild.cpp b/Engine/Source/Runtime/Engine/Private/StaticMeshBuild.cpp
--- a/Engine/Source/Runtime/Engine/Private/StaticMeshBuild.cpp
+++ b/Engine/Source/Runtime/Engine/Private/StaticMeshBuild.cpp
@@ -419,501 +419,6 @@ struct FFragmentRange
 	}
 };
 
-class FLegacyStaticMeshElement
-{
-public:
-	class UMaterialInterface*		Material;
-	FString					Name;			
-	bool					EnableCollision,
-							OldEnableCollision,
-							bEnableShadowCasting;
-	uint32					FirstIndex,
-							NumTriangles,
-							MinVertexIndex,
-							MaxVertexIndex;
-	int32 MaterialIndex;
-	TArray<FFragmentRange> Fragments;
-
-	/** Constructor. */
-	FLegacyStaticMeshElement():
-		Material(NULL),
-		EnableCollision(false),
-		OldEnableCollision(false),
-		bEnableShadowCasting(true),
-		FirstIndex(0),
-		NumTriangles(0),
-		MinVertexIndex(0),
-		MaxVertexIndex(0),
-		MaterialIndex(0)
-	{}
-
-	friend FArchive& operator<<(FArchive& Ar,FLegacyStaticMeshElement& E)
-	{
-		Ar	<< E.Material
-			<< E.EnableCollision
-			<< E.OldEnableCollision
-			<< E.bEnableShadowCasting
-			<< E.FirstIndex
-			<< E.NumTriangles
-			<< E.MinVertexIndex
-			<< E.MaxVertexIndex
-			<< E.MaterialIndex;
-
-		Ar << E.Fragments;
-		return Ar;
-	}
-};
-
-struct FLegacyStaticMeshLODInfo
-{
-	TArray<FLegacyStaticMeshElement> Elements;
-
-	friend FArchive& operator<<(FArchive& Ar,FLegacyStaticMeshLODInfo& LODInfo)
-    {
-		// Nothing to actually serialize.
-		return Ar;
-    }
-};
-
-struct FLegacyStaticMeshRenderData
-{
-	FStaticMeshVertexBuffer VertexBuffer;
-	FPositionVertexBuffer PositionVertexBuffer;
-	FColorVertexBuffer ColorVertexBuffer;
-	uint32 NumVertices;
-	FRawStaticIndexBuffer					IndexBuffer;
-	FRawStaticIndexBuffer					ShadowIndexBuffer;
-	FRawIndexBuffer							WireframeIndexBuffer;
-	TArray<FLegacyStaticMeshElement>		Elements;
-	FStaticMeshTriangleBulkData				RawTriangles;
-	FRawStaticIndexBuffer					AdjacencyIndexBuffer;
-
-	void Serialize(FArchive& Ar, UObject* Owner, int32 Idx)
-    {
-		const uint8 AdjacencyDataStripFlag = 1;
-		FStripDataFlags StripFlags( Ar, Ar.IsCooking() && !Ar.CookingTarget()->SupportsFeature(ETargetPlatformFeatures::Tessellation) ? AdjacencyDataStripFlag : 0 );
-
-		UStaticMesh* StaticMesh = Cast<UStaticMesh>( Owner );
-		bool bNeedsCPUAccess = true;
-
-		if( !StripFlags.IsEditorDataStripped() )
-        {
-			RawTriangles.Serialize( Ar, Owner );
-        }
-
-		Ar << Elements;
-
-		if( !StripFlags.IsDataStrippedForServer() )
-        {
-			PositionVertexBuffer.Serialize( Ar, bNeedsCPUAccess );
-			VertexBuffer.Serialize( Ar, bNeedsCPUAccess );
-			ColorVertexBuffer.Serialize( Ar, bNeedsCPUAccess );
-			Ar << NumVertices;
-			IndexBuffer.Serialize( Ar, bNeedsCPUAccess );
-			if (Ar.UE4Ver() >= VER_UE4_SHADOW_ONLY_INDEX_BUFFERS)
-            {		
-				ShadowIndexBuffer.Serialize(Ar, bNeedsCPUAccess);
-            }
-			if( !StripFlags.IsEditorDataStripped() )
-            {
-				Ar << WireframeIndexBuffer;
-            }
-			if ( !StripFlags.IsClassDataStripped( AdjacencyDataStripFlag ) )
-            { 					
-				AdjacencyIndexBuffer.Serialize( Ar, bNeedsCPUAccess );
-            }
-        }
-
-		if (Ar.IsLoading())
-        {
-			if (PositionVertexBuffer.GetNumVertices() != NumVertices)
-            {
-				PositionVertexBuffer.RemoveLegacyShadowVolumeVertices(NumVertices);
-			}
-			if (VertexBuffer.GetNumVertices() != NumVertices)
-			{
-				VertexBuffer.RemoveLegacyShadowVolumeVertices(NumVertices);
-			}
-			if (VertexBuffer.GetNumVertices() != NumVertices)
-            {				
-				ColorVertexBuffer.RemoveLegacyShadowVolumeVertices(NumVertices);
-			}
-        }
-    }										
-};
-
-/**
- * Constructs a raw mesh from legacy raw triangles.
- */
-static bool BuildRawMeshFromRawTriangles(
-	FRawMesh& OutRawMesh,
-	FMeshBuildSettings& OutBuildSettings,
-	struct FStaticMeshTriangle* RawTriangles,
-	int32 NumTriangles,
-	const TCHAR* MeshName
-	)
-{
-	int32 NumWedges = NumTriangles * 3;
-	OutRawMesh.Empty();
-	OutRawMesh.FaceMaterialIndices.Empty(NumTriangles);
-	OutRawMesh.FaceSmoothingMasks.Empty(NumWedges);
-	OutRawMesh.VertexPositions.Empty(NumWedges);
-	OutRawMesh.WedgeIndices.Empty(NumWedges);
-
-	int32 NumTexCoords = 0;
-	bool bHaveNormals = false;
-	bool bHaveTangents = false;
-	bool bHaveColors = false;
-	bool bCorruptFlags = false;
-	bool bCorruptNormals = false;
-	bool bCorruptTangents = false;
-	bool bCorruptPositions = false;
-
-	for (int32 TriIndex = 0; TriIndex < NumTriangles; ++TriIndex)
-    {
-		FStaticMeshTriangle* Tri = RawTriangles + TriIndex;
-		OutRawMesh.FaceMaterialIndices.Add(Tri->MaterialIndex);
-		OutRawMesh.FaceSmoothingMasks.Add(Tri->SmoothingMask);
-
-		bCorruptFlags |= (RawTriangles[0].bExplicitNormals & 0xFFFFFFFE) != 0
-			|| (RawTriangles[0].bOverrideTangentBasis & 0xFFFFFFFE) != 0;
-
-		for (int32 CornerIndex = 0; CornerIndex < 3; ++CornerIndex)
-        {
-			FVector Position = Tri->Vertices[CornerIndex];
-			OutRawMesh.WedgeIndices.Add(OutRawMesh.VertexPositions.Add(Position));
-			NumTexCoords = FMath::Max(NumTexCoords, Tri->NumUVs);
-			bHaveNormals |= Tri->TangentZ[CornerIndex] != FVector::ZeroVector;
-			bHaveTangents |= (Tri->TangentX[CornerIndex] != FVector::ZeroVector
-				&& Tri->TangentY[CornerIndex] != FVector::ZeroVector);
-			bHaveColors |= Tri->Colors[CornerIndex] != FColor::White;
-
-			bCorruptPositions |= Position.ContainsNaN();
-			bCorruptTangents |= Tri->TangentX[CornerIndex].ContainsNaN()
-				|| Tri->TangentY[CornerIndex].ContainsNaN();
-			bCorruptNormals |= Tri->TangentZ[CornerIndex].ContainsNaN();
-        }
-    }
-
-	bool bTooManyTexCoords = NumTexCoords > 8;
-	NumTexCoords = FMath::Min(NumTexCoords, 8); // FStaticMeshTriangle has at most 8 texture coordinates.
-	NumTexCoords = FMath::Min<int32>(NumTexCoords, MAX_MESH_TEXTURE_COORDS);
-
-	for (int32 TexCoordIndex = 0; TexCoordIndex < FMath::Max(1, NumTexCoords); ++TexCoordIndex)
-	{
-		OutRawMesh.WedgeTexCoords[TexCoordIndex].AddZeroed(NumTriangles * 3);
-	}
-
-	for (int32 TriIndex = 0; TriIndex < NumTriangles; ++TriIndex)
-	{
-		FStaticMeshTriangle* Tri = RawTriangles + TriIndex;
-		for (int32 CornerIndex = 0; CornerIndex < 3; ++CornerIndex)
-		{
-			for (int32 TexCoordIndex = 0; TexCoordIndex < NumTexCoords; ++TexCoordIndex)
-			{
-				FVector2D UV = TexCoordIndex < Tri->NumUVs ? Tri->UVs[CornerIndex][TexCoordIndex] : FVector2D(0.0f,0.0f);
-				OutRawMesh.WedgeTexCoords[TexCoordIndex][TriIndex * 3 + CornerIndex] = UV;
-			}
-		}
-	}
-
-	if (bHaveNormals && !bCorruptNormals)
-	{
-		OutRawMesh.WedgeTangentZ.AddZeroed(NumTriangles * 3);
-		for (int32 TriIndex = 0; TriIndex < NumTriangles; ++TriIndex)
-		{
-			if (RawTriangles[TriIndex].bExplicitNormals || RawTriangles[TriIndex].bOverrideTangentBasis)
-			{
-				for (int32 CornerIndex = 0; CornerIndex < 3; ++CornerIndex)
-				{
-					OutRawMesh.WedgeTangentZ[TriIndex * 3 + CornerIndex] =
-						RawTriangles[TriIndex].TangentZ[CornerIndex];
-				}
-			}
-		}
-	}
-	if (bHaveTangents && !bCorruptTangents)
-	{
-		OutRawMesh.WedgeTangentX.AddZeroed(NumTriangles * 3);
-		OutRawMesh.WedgeTangentY.AddZeroed(NumTriangles * 3);
-		for (int32 TriIndex = 0; TriIndex < NumTriangles; ++TriIndex)
-	{
-			if (RawTriangles[TriIndex].bOverrideTangentBasis)
-			{
-				for (int32 CornerIndex = 0; CornerIndex < 3; ++CornerIndex)
-				{
-					OutRawMesh.WedgeTangentX[TriIndex * 3 + CornerIndex] =
-						RawTriangles[TriIndex].TangentX[CornerIndex];
-					OutRawMesh.WedgeTangentY[TriIndex * 3 + CornerIndex] =
-						RawTriangles[TriIndex].TangentY[CornerIndex];
-				}
-			}
-		}
-	}
-
-	if (bHaveColors)
-	{
-		OutRawMesh.WedgeColors.AddUninitialized(NumTriangles * 3);
-		for (int32 TriIndex = 0; TriIndex < NumTriangles; ++TriIndex)
-		{
-			for (int32 CornerIndex = 0; CornerIndex < 3; ++CornerIndex)
-			{
-				OutRawMesh.WedgeColors[TriIndex * 3 + CornerIndex] =
-					RawTriangles[TriIndex].Colors[CornerIndex];
-			}
-		}
-	}
-
-	if (bCorruptPositions || bCorruptFlags || bCorruptNormals || bCorruptTangents || bTooManyTexCoords)
-	{
-		UE_LOG(LogStaticMesh,Verbose,TEXT("Legacy raw triangles CORRUPT (%s%s%s%s%s) for %s"),
-			bCorruptPositions ? TEXT("NaN positions,") : TEXT(""),
-			bCorruptFlags ? TEXT("flags,") : TEXT(""),
-			bCorruptNormals ? TEXT("NaN normals,") : TEXT(""),
-			bCorruptTangents ? TEXT("NaN tangents,") : TEXT(""),
-			bTooManyTexCoords ? TEXT(">8 texcoords") : TEXT(""),
-			MeshName
-			);
-	}
-
-	if (bCorruptPositions)
-	{
-		OutRawMesh = FRawMesh();
-	}
-
-	OutBuildSettings.bRecomputeTangents = !bHaveTangents || bCorruptTangents;
-	OutBuildSettings.bRecomputeNormals = !bHaveNormals || bCorruptNormals;
-
-	return !(bCorruptPositions || bCorruptFlags || bTooManyTexCoords) && OutRawMesh.IsValid();
-}
-
-/**
- * Constructs a raw mesh from legacy render data.
- */
-static void BuildRawMeshFromRenderData(
-	FRawMesh& OutRawMesh,
-	struct FMeshBuildSettings& OutBuildSettings,
-	FLegacyStaticMeshRenderData& RenderData,
-	const TCHAR* MeshName
-	)
-{
-	FStaticMeshTriangle* RawTriangles = (FStaticMeshTriangle*)RenderData.RawTriangles.Lock(LOCK_READ_ONLY);
-	bool bBuiltFromRawTriangles = BuildRawMeshFromRawTriangles(
-		OutRawMesh,
-		OutBuildSettings,
-		RawTriangles,
-		RenderData.RawTriangles.GetElementCount(),
-		MeshName
-		);
-	RenderData.RawTriangles.Unlock();
-
-	if (bBuiltFromRawTriangles)
-	{
-		return;
-	}
-
-	OutRawMesh.Empty();
-	
-	FIndexArrayView Indices = RenderData.IndexBuffer.GetArrayView();
-	int32 NumVertices = RenderData.PositionVertexBuffer.GetNumVertices();
-	int32 NumTriangles = Indices.Num() / 3;
-	int32 NumWedges = NumTriangles * 3;
-
-	// Copy vertex positions.
-	{
-		OutRawMesh.VertexPositions.AddUninitialized(NumVertices);
-		for (int32 i = 0; i < NumVertices; ++i)
-		{
-			OutRawMesh.VertexPositions[i] = RenderData.PositionVertexBuffer.VertexPosition(i);
-		}
-	}
-
-	// Copy per-wedge texture coordinates.
-	for (uint32 TexCoordIndex = 0; TexCoordIndex < RenderData.VertexBuffer.GetNumTexCoords(); ++TexCoordIndex)
-	{
-		OutRawMesh.WedgeTexCoords[TexCoordIndex].AddUninitialized(NumWedges);
-		for (int32 i = 0; i < NumWedges; ++i)
-		{
-			uint32 VertIndex = Indices[i];
-			OutRawMesh.WedgeTexCoords[TexCoordIndex][i] = RenderData.VertexBuffer.GetVertexUV(VertIndex, TexCoordIndex);
-		}
-	}
-
-	// Copy per-wedge colors if they exist.
-	if (RenderData.ColorVertexBuffer.GetNumVertices() > 0)
-	{
-		OutRawMesh.WedgeColors.AddUninitialized(NumWedges);
-		for (int32 i = 0; i < NumWedges; ++i)
-		{
-			uint32 VertIndex = Indices[i];
-			OutRawMesh.WedgeColors[i] = RenderData.ColorVertexBuffer.VertexColor(VertIndex);
-		}
-	}
-
-	// Copy per-wedge tangents.
-	{
-		OutRawMesh.WedgeTangentX.AddUninitialized(NumWedges);
-		OutRawMesh.WedgeTangentY.AddUninitialized(NumWedges);
-		OutRawMesh.WedgeTangentZ.AddUninitialized(NumWedges);
-		for (int32 i = 0; i < NumWedges; ++i)
-		{
-			uint32 VertIndex = Indices[i];
-			OutRawMesh.WedgeTangentX[i] = RenderData.VertexBuffer.VertexTangentX(VertIndex);
-			OutRawMesh.WedgeTangentY[i] = RenderData.VertexBuffer.VertexTangentY(VertIndex);
-			OutRawMesh.WedgeTangentZ[i] = RenderData.VertexBuffer.VertexTangentZ(VertIndex);
-		}
-	}
-
-	// Copy per-face information.
-	{
-		OutRawMesh.FaceMaterialIndices.AddZeroed(NumTriangles);
-		OutRawMesh.FaceSmoothingMasks.AddZeroed(NumTriangles);
-		OutRawMesh.WedgeIndices.AddZeroed(NumWedges);
-		for (int32 SectionIndex = 0; SectionIndex < RenderData.Elements.Num(); ++SectionIndex)
-		{
-			const FLegacyStaticMeshElement& Section = RenderData.Elements[SectionIndex];
-			int32 FirstFace = Section.FirstIndex / 3;
-			int32 LastFace = FirstFace + Section.NumTriangles;
-			for (int32 i = FirstFace; i < LastFace; ++i)
-			{
-				OutRawMesh.FaceMaterialIndices[i] = Section.MaterialIndex;
-				// Smoothing group information has been lost but is already baked in to the tangent basis.
-				for (int32 j = 0; j < 3; ++j)
-				{
-					OutRawMesh.WedgeIndices[i * 3 + j] = Indices[i * 3 + j];
-				}
-			}
-		}
-	}
-
-	check(OutRawMesh.IsValid());
-
-	OutBuildSettings.bRecomputeNormals = false;
-	OutBuildSettings.bRecomputeTangents = false;
-	OutBuildSettings.bRemoveDegenerates = false;
-	OutBuildSettings.bUseFullPrecisionUVs = RenderData.VertexBuffer.GetUseFullPrecisionUVs();
-}
-
-void UStaticMesh::SerializeLegacySouceData(FArchive& Ar, const FBoxSphereBounds& LegacyBounds)
-{
-	int32 InternalVersion = 0;
-	bool bHaveSourceData = false;
-	TArray<FStaticMeshOptimizationSettings> OptimizationSettings;
-	bool bHasBeenSimplified = false;
-
-	Ar << InternalVersion;
-	Ar << bHaveSourceData;
-	if (bHaveSourceData)
-	{
-		FStaticMeshSourceModel& SrcModel = *new(SourceModels) FStaticMeshSourceModel();
-
-		// Serialize in the old source render data.
-		FLegacyStaticMeshRenderData TempRenderData;
-		TempRenderData.Serialize(Ar, this, INDEX_NONE);
-
-		// Convert it to a raw mesh.
-		FRawMesh RawMesh;
-		BuildRawMeshFromRenderData(RawMesh, SrcModel.BuildSettings, TempRenderData, *GetPathName());
-
-		// Store the raw mesh as our source model.
-		SrcModel.RawMeshBulkData->SaveRawMesh(RawMesh);
-	}
-	Ar << OptimizationSettings;
-	Ar << bHasBeenSimplified;
-
-	// Convert any old optimization settings to the source model's reduction settings.
-	for (int32 i = 0; i < OptimizationSettings.Num(); ++i)
-	{
-		if (!SourceModels.IsValidIndex(i))
-		{
-			new(SourceModels) FStaticMeshSourceModel();
-		}
-		FStaticMeshSourceModel& SrcModel = SourceModels[i];
-
-		FMeshReductionSettings& NewSettings = SrcModel.ReductionSettings;
-		switch (OptimizationSettings[i].ReductionMethod)
-		{
-		case OT_NumOfTriangles:
-			NewSettings.PercentTriangles = FMath::Clamp(OptimizationSettings[i].NumOfTrianglesPercentage / 100.0f, 0.0f, 1.0f);
-			NewSettings.MaxDeviation = 0.0f;
-			break;
-		case OT_MaxDeviation:
-			NewSettings.PercentTriangles = 1.0f;
-			NewSettings.MaxDeviation = OptimizationSettings[i].MaxDeviationPercentage * LegacyBounds.SphereRadius;
-			break;
-		default:
-			NewSettings.PercentTriangles = 1.0f;
-			NewSettings.MaxDeviation = 0.0f;
-			break;
-		}
-		NewSettings.WeldingThreshold = OptimizationSettings[i].WeldingThreshold;
-		NewSettings.HardAngleThreshold = OptimizationSettings[i].NormalsThreshold;
-		NewSettings.SilhouetteImportance = (EMeshFeatureImportance::Type)OptimizationSettings[i].SilhouetteImportance;
-		NewSettings.TextureImportance = (EMeshFeatureImportance::Type)OptimizationSettings[i].TextureImportance;
-		NewSettings.ShadingImportance = (EMeshFeatureImportance::Type)OptimizationSettings[i].ShadingImportance;
-	}
-
-	// Serialize in any legacy LOD models.
-	TIndirectArray<FLegacyStaticMeshRenderData> LegacyLODModels;
-	LegacyLODModels.Serialize(Ar, this);
-
-	TArray<FLegacyStaticMeshLODInfo> LegacyLODInfo;
-	Ar << LegacyLODInfo;
-
-	for (int32 LODIndex = 0; LODIndex < LegacyLODModels.Num(); ++LODIndex)
-	{
-		if (!SourceModels.IsValidIndex(LODIndex))
-		{
-			new(SourceModels) FStaticMeshSourceModel();
-		}
-		FStaticMeshSourceModel& SrcModel = SourceModels[LODIndex];
-
-		// Really we want to use LOD0 only if the mesh has /not/ been simplified.
-		// Unfortunately some data seems to be corrupted in source meshes but not LOD0, so just use LOD0.
-		//if (SrcModel.RawMeshBulkData->IsEmpty() && (!bHasBeenSimplified || LODIndex == 0))
-		{
-			// Store the raw mesh for this LOD.
-			FRawMesh RawMesh;
-			BuildRawMeshFromRenderData(RawMesh, SrcModel.BuildSettings, LegacyLODModels[LODIndex], *GetPathName());
-			SrcModel.RawMeshBulkData->SaveRawMesh(RawMesh);
-		}
-
-		// And make sure to clear reduction settings for LOD0.
-		if (LODIndex == 0)
-		{
-			FMeshReductionSettings DefaultSettings;
-			SrcModel.ReductionSettings = DefaultSettings;
-		}
-
-		// Always use a hash as the guid for legacy models.
-		SrcModel.RawMeshBulkData->UseHashAsGuid(this);
-
-		// Setup the material map for this LOD model.
-		if (LODIndex == 0)
-		{
-			for (int32 SectionIndex = 0; SectionIndex < LegacyLODModels[LODIndex].Elements.Num(); ++SectionIndex)
-			{
-				FMeshSectionInfo Info;
-				Info.MaterialIndex = Materials.Add(LegacyLODModels[LODIndex].Elements[SectionIndex].Material);
-				Info.bEnableCollision = LegacyLODModels[LODIndex].Elements[SectionIndex].EnableCollision;
-				Info.bCastShadow = LegacyLODModels[LODIndex].Elements[SectionIndex].bEnableShadowCasting;
-				SectionInfoMap.Set(LODIndex, SectionIndex, Info);
-			}
-		}
-		else
-		{
-			for (int32 SectionIndex = 0; SectionIndex < LegacyLODModels[LODIndex].Elements.Num(); ++SectionIndex)
-			{
-				FMeshSectionInfo Info;
-				Info.MaterialIndex = Materials.AddUnique(LegacyLODModels[LODIndex].Elements[SectionIndex].Material);
-				Info.bEnableCollision = LegacyLODModels[LODIndex].Elements[SectionIndex].EnableCollision;
-				Info.bCastShadow = LegacyLODModels[LODIndex].Elements[SectionIndex].bEnableShadowCasting;
-				SectionInfoMap.Set(LODIndex, SectionIndex, Info);
-			}
-		}
-	}
-}
-
 void UStaticMesh::FixupZeroTriangleSections()
 {
 	if (RenderData->MaterialIndexToImportIndex.Num() > 0 && RenderData->LODResources.Num())
diff --git a/Engine/Source/Runtime/Engine/Private/StaticMeshComponent.cpp b/Engine/Source/Runtime/Engine/Private/StaticMeshComponent.cpp
--- a/Engine/Source/Runtime/Engine/Private/StaticMeshComponent.cpp
+++ b/Engine/Source/Runtime/Engine/Private/StaticMeshComponent.cpp
@@ -138,7 +138,6 @@ UStaticMeshComponent::UStaticMeshComponent(const FObjectInitializer& ObjectIniti
 {
 	PrimaryComponentTick.bCanEverTick = false;
 
-	BodyInstance.bEnableCollision_DEPRECATED = true;
 	// check BaseEngine.ini for profile setup
 	SetCollisionProfileName(UCollisionProfile::BlockAllDynamic_ProfileName);
 
@@ -1902,11 +1901,7 @@ FArchive& operator<<(FArchive& Ar,FStaticMeshComponentLODInfo& I)
 	if( !StripFlags.IsDataStrippedForServer() )
 	{
 		Ar << I.LightMap;
-
-		if (Ar.UE4Ver() >= VER_UE4_PRECOMPUTED_SHADOW_MAPS)
-		{
-			Ar << I.ShadowMap;
-		}
+		Ar << I.ShadowMap;
 	}
 
 	if( !StripFlags.IsClassDataStripped( OverrideColorsStripFlag ) )
diff --git a/Engine/Source/Runtime/Engine/Private/Texture.cpp b/Engine/Source/Runtime/Engine/Private/Texture.cpp
--- a/Engine/Source/Runtime/Engine/Private/Texture.cpp
+++ b/Engine/Source/Runtime/Engine/Private/Texture.cpp
@@ -52,7 +52,6 @@ UTexture::UTexture(const FObjectInitializer& ObjectInitializer)
 	AdjustMaxAlpha = 1.0f;
 	MaxTextureSize = 0; // means no limitation
 	MipGenSettings = TMGS_FromTextureGroup;
-	SourceArtType_DEPRECATED = TSAT_PNGCompressed;
 	CompositeTextureMode = CTM_NormalRoughnessToAlpha;
 	CompositePower = 1.0f;
 #endif // #if WITH_EDITORONLY_DATA
@@ -226,57 +225,6 @@ void UTexture::PostEditChangeProperty(FPropertyChangedEvent& PropertyChangedEven
 }
 #endif // WITH_EDITOR
 
-#if WITH_EDITORONLY_DATA
-void UTexture::LegacySerialize(FArchive& Ar, FStripDataFlags& StripFlags)
-{
-	check(Ar.UE4Ver() < VER_UE4_TEXTURE_SOURCE_ART_REFACTOR);
-
-	// Serialize the old source art directly in to the source's bulk data.
-	if (!StripFlags.IsEditorDataStripped())
-	{
-		Source.BulkData.Serialize(Ar, this);
-	}
-
-	switch (SourceArtType_DEPRECATED)
-	{
-	case TSAT_Uncompressed:
-		Source.Format = TSF_BGRA8;
-		break;
-	case TSAT_PNGCompressed:
-		Source.Format = TSF_BGRA8;
-		Source.bPNGCompressed = true;
-		break;
-	case TSAT_DDSFile:
-#if WITH_EDITOR
-		Source.ConvertFromLegacyDDS();
-#endif
-		break;
-	}
-
-	if (Ar.UE4Ver() < VER_UE4_TEXTURE_DERIVED_DATA2)
-	{
-		// PostEditChange used to set CompressionNone for TC_EditorIcon. This is
-		// now handled by the texture derived data logic automatically. Repair
-		// this on load so that the correct textures are cached and built.
-		if (CompressionSettings == TC_EditorIcon)
-		{
-			CompressionNone = false;
-		}
-	}
-
-	// to be backwards compatible, the flag was removed
-	if(CompressionNoMipmaps_DEPRECATED)
-	{
-		MipGenSettings = TMGS_NoMipmaps;
-	}
-
-	if ( CompressionSettings >= TC_MAX )
-	{
-		CompressionSettings = TC_Default;
-	}
-}
-#endif
-
 void UTexture::Serialize(FArchive& Ar)
 {
 	Super::Serialize(Ar);
@@ -285,21 +233,6 @@ void UTexture::Serialize(FArchive& Ar)
 
 	/** Legacy serialization. */
 #if WITH_EDITORONLY_DATA
-	if (Ar.UE4Ver() < VER_UE4_TEXTURE_SOURCE_ART_REFACTOR)
-	{
-		UTexture::LegacySerialize(Ar, StripFlags);
-		return;
-	}
-
-	if (Ar.UE4Ver() < VER_UE4_TEXTURE_FORMAT_RGBA_SWIZZLE)
-	{
-		switch (Source.Format)
-		{
-		case TSF_RGBA8:	Source.Format = TSF_BGRA8; break;
-		case TSF_RGBE8:	Source.Format = TSF_BGRE8; break;
-		}
-	}
-
 	if (!StripFlags.IsEditorDataStripped())
 	{
 		Source.BulkData.Serialize(Ar, this);
@@ -866,55 +799,6 @@ int32 FTextureSource::CalcMipOffset(int32 MipIndex) const
 	return MipOffset;
 }
 
-void FTextureSource::ConvertFromLegacyDDS()
-{
-	TArray<uint8> RawDDS;
-	RawDDS.Empty(BulkData.GetBulkDataSize());
-	RawDDS.AddUninitialized(BulkData.GetBulkDataSize());
-	{
-		void* LockedDDSData = NULL;
-		BulkData.GetCopy(&LockedDDSData);
-		FMemory::Memcpy(RawDDS.GetData(), LockedDDSData, RawDDS.Num());
-		FMemory::Free(LockedDDSData);
-		BulkData.RemoveBulkData();
-	}
-
-	FDDSLoadHelper DDS(RawDDS.GetData(), RawDDS.Num());
-	if ((DDS.IsValid2DTexture() || DDS.IsValidCubemapTexture()) && DDS.ComputeSourceFormat() != TSF_Invalid)
-	{
-		Init(
-			DDS.DDSHeader->dwWidth, 
-			DDS.DDSHeader->dwHeight,
-			/*NumSlices=*/ DDS.IsValidCubemapTexture() ? 6 : 1,
-			DDS.ComputeMipMapCount(),
-			DDS.ComputeSourceFormat()
-			);
-
-		uint8* DestMipData[MAX_TEXTURE_MIP_COUNT] = {0};
-		int32 MipSize[MAX_TEXTURE_MIP_COUNT] = {0};
-		for (int32 MipIndex = 0; MipIndex < NumMips; ++MipIndex)
-		{
-			DestMipData[MipIndex] = LockMip(MipIndex);
-			MipSize[MipIndex] = CalcMipSize(MipIndex) / NumSlices;
-		}
-
-		for (int32 SliceIndex = 0; SliceIndex < NumSlices; ++SliceIndex)
-		{
-			const uint8* SrcMipData = DDS.GetDDSDataPointer((ECubeFace)SliceIndex);
-			for (int32 MipIndex = 0; MipIndex < NumMips; ++MipIndex)
-			{
-				FMemory::Memcpy(DestMipData[MipIndex] + MipSize[MipIndex] * SliceIndex, SrcMipData, MipSize[MipIndex]);
-				SrcMipData += MipSize[MipIndex];
-			}
-		}
-
-		for (int32 MipIndex = 0; MipIndex < NumMips; ++MipIndex)
-		{
-			UnlockMip(MipIndex);
-		}
-	}
-}
-
 void FTextureSource::UseHashAsGuid()
 {
 	uint32 Hash[5];
diff --git a/Engine/Source/Runtime/Engine/Private/Texture2D.cpp b/Engine/Source/Runtime/Engine/Private/Texture2D.cpp
--- a/Engine/Source/Runtime/Engine/Private/Texture2D.cpp
+++ b/Engine/Source/Runtime/Engine/Private/Texture2D.cpp
@@ -128,17 +128,14 @@ private:
 void FTexture2DMipMap::Serialize(FArchive& Ar, UObject* Owner, int32 MipIdx)
 {
 	bool bCooked = Ar.IsCooking();
-	if (Ar.UE4Ver() >= VER_UE4_TEXTURE_SOURCE_ART_REFACTOR)
-	{
-		Ar << bCooked;
-	}
+	Ar << bCooked;
 
 	BulkData.Serialize(Ar, Owner, MipIdx);
 	Ar << SizeX;
 	Ar << SizeY;
 
 #if WITH_EDITORONLY_DATA
-	if (Ar.UE4Ver() >= VER_UE4_TEXTURE_DERIVED_DATA2 && !bCooked)
+	if (!bCooked)
 	{
 		Ar << DerivedDataKey;
 	}
@@ -175,125 +172,14 @@ bool UTexture2D::GetResourceMemSettings(int32 FirstMipIdx, int32& OutSizeX, int3
 	return false;
 }
 
-#if WITH_EDITOR
-void UTexture2D::LegacySerialize(FArchive& Ar, FStripDataFlags& StripDataFlags)
-{
-	check(Ar.UE4Ver() < VER_UE4_TEXTURE_SOURCE_ART_REFACTOR);
-
-	bool bHasLegacyMips = bDisableDerivedDataCache_DEPRECATED;
-	TIndirectArray<FTexture2DMipMap> LegacyMips;
-	if (bHasLegacyMips)
-	{
-		LegacyMips.Serialize(Ar, this);
-	}
-
-	Ar << TextureFileCacheGuid_DEPRECATED;
-	
-	bool bValidGuid = TextureFileCacheGuid_DEPRECATED.IsValid();
-	if (!bValidGuid)
-	{
-		UE_LOG(LogTexture, Warning, TEXT("Texture with no guid %s, package must be resaved!"),*GetFullName());
-		MarkPackageDirty();
-		TextureFileCacheGuid_DEPRECATED = FGuid::NewGuid();
-	}
-
-	if ((bDisableDerivedDataCache_DEPRECATED || Source.BulkData.GetBulkDataSize() == 0) &&
-		LegacyMips.Num() > 0)
-	{
-		if (Format_DEPRECATED == PF_A8R8G8B8 || Format_DEPRECATED == PF_B8G8R8A8 || Format_DEPRECATED == PF_FloatRGBA || Format_DEPRECATED == PF_G8)
-		{
-			ETextureSourceFormat SourceFormat = TSF_BGRA8;
-			if (Format_DEPRECATED == PF_FloatRGBA)
-			{
-				SourceFormat = TSF_RGBA16F;
-			}
-			else if (Format_DEPRECATED == PF_G8)
-			{
-				SourceFormat = TSF_G8;
-			}
-			int32 NumSourceMips = MipGenSettings == TMGS_LeaveExistingMips ? LegacyMips.Num() : 1;
-			Source.Init(
-				LegacyMips[0].SizeX,
-				LegacyMips[0].SizeY,
-				/*NumSlices=*/ 1, 
-				/*NumMips=*/ NumSourceMips,
-				SourceFormat
-				);
-			for (int32 MipIndex = 0; MipIndex < NumSourceMips; ++MipIndex)
-			{
-				uint8* SrcMipData = NULL;
-				uint8* DestMipData = NULL;
-				LegacyMips[MipIndex].BulkData.GetCopy((void**)&SrcMipData, /*bDiscardInternalCopy=*/ true);
-				DestMipData = Source.LockMip(MipIndex);
-				FMemory::Memcpy(DestMipData, SrcMipData, LegacyMips[MipIndex].BulkData.GetBulkDataSize());
-				Source.UnlockMip(MipIndex);
-				FMemory::Free(SrcMipData);
-			}
-		}
-		else
-		{
-			UE_LOG(LogTexture,Warning,
-				TEXT("Legacy texture has no source art and compressed mip ")
-				TEXT("levels. The texture must be deleted or reimported. %s"),
-				*GetFullName()
-				);
-		}
-	}
-	else if (Source.BulkData.GetBulkDataSize() != 0 &&
-		SourceArtType_DEPRECATED != TSAT_DDSFile)
-	{
-		Source.SizeX = OriginalSizeX_DEPRECATED;
-		Source.SizeY = OriginalSizeY_DEPRECATED;
-		Source.NumSlices = 1;
-		Source.NumMips = 1;
-	}
-
-	if (bValidGuid)
-	{
-		Source.Id = TextureFileCacheGuid_DEPRECATED;
-	}
-	else
-	{
-		Source.UseHashAsGuid();
-	}
-
-	// Fix up non power-of-two textures with incorrect compression settings.
-	if (!Source.IsPowerOfTwo())
-	{
-		// Force NPT textures to have no mipmaps.
-		MipGenSettings = TMGS_NoMipmaps;
-		NeverStream = true;
-	}
-
-	// Make sure settings are correct for LUT textures.
-	if(LODGroup == TEXTUREGROUP_ColorLookupTable)
-	{
-		MipGenSettings = TMGS_NoMipmaps;
-		SRGB = false;
-	}
-}
-#endif // #if WITH_EDITOR
-
 void UTexture2D::Serialize(FArchive& Ar)
 {
 	Super::Serialize(Ar);
 
 	FStripDataFlags StripDataFlags(Ar);
 
-	bool bCooked = false;
-	if (Ar.UE4Ver() >= VER_UE4_ADD_COOKED_TO_TEXTURE2D)
-	{
-		bCooked = Ar.IsCooking();
-		Ar << bCooked;
-	}
-
-	/** Legacy serialization. */
-#if WITH_EDITOR
-	if (Ar.UE4Ver() < VER_UE4_TEXTURE_SOURCE_ART_REFACTOR)
-	{
-		UTexture2D::LegacySerialize(Ar, StripDataFlags);
-	}
-#endif // #if WITH_EDITOR
+	bool bCooked = Ar.IsCooking();
+	Ar << bCooked;
 
 	if (Ar.IsCooking() || bCooked)
 	{
diff --git a/Engine/Source/Runtime/Engine/Private/TextureCube.cpp b/Engine/Source/Runtime/Engine/Private/TextureCube.cpp
--- a/Engine/Source/Runtime/Engine/Private/TextureCube.cpp
+++ b/Engine/Source/Runtime/Engine/Private/TextureCube.cpp
@@ -19,38 +19,13 @@ void UTextureCube::Serialize(FArchive& Ar)
 {
 	Super::Serialize(Ar);
 
-	bool bCooked = false;
+	FStripDataFlags StripFlags(Ar);
+	bool bCooked = Ar.IsCooking();
+	Ar << bCooked;
 
-#if WITH_EDITOR
-	// Handle legacy serialization.
-	if (Ar.UE4Ver() < VER_UE4_TEXTURE_SOURCE_ART_REFACTOR)
+	if (bCooked || Ar.IsCooking())
 	{
-		TIndirectArray<FTexture2DMipMap> LegacyMips;
-		LegacyMips.Serialize( Ar, this );
-
-		if (Source.BulkData.GetBulkDataSize() > 0 &&
-			Source.GetNumSlices() == 1 &&
-			Source.GetFormat() == TSF_BGRA8)
-		{
-			Source.Format = TSF_BGRE8;
-		}
-		if (Source.GetFormat() == TSF_BGRE8 || Source.GetFormat() == TSF_RGBA16F)
-		{
-			CompressionSettings = TC_HDR;
-		}
-		Source.UseHashAsGuid();
-	}
-	else
-#endif // #if WITH_EDITOR
-	{
-		FStripDataFlags StripFlags(Ar);
-		bCooked = Ar.IsCooking();
-		Ar << bCooked;
-
-		if (bCooked || Ar.IsCooking())
-		{
-			SerializeCookedPlatformData(Ar);
-		}
+		SerializeCookedPlatformData(Ar);
 	}
 
 #if WITH_EDITOR
diff --git a/Engine/Source/Runtime/Engine/Private/TriggerActors.cpp b/Engine/Source/Runtime/Engine/Private/TriggerActors.cpp
--- a/Engine/Source/Runtime/Engine/Private/TriggerActors.cpp
+++ b/Engine/Source/Runtime/Engine/Private/TriggerActors.cpp
@@ -22,7 +22,6 @@ ATriggerCapsule::ATriggerCapsule(const FObjectInitializer& ObjectInitializer)
 	UCapsuleComponent* CapsuleCollisionComponent = CastChecked<UCapsuleComponent>(GetCollisionComponent());
 	CapsuleCollisionComponent->ShapeColor = TriggerBaseColor;
 	CapsuleCollisionComponent->InitCapsuleSize(+40.0f, +80.0f);
-	CapsuleCollisionComponent->BodyInstance.bEnableCollision_DEPRECATED = true;
 	static FName CollisionProfileName(TEXT("Trigger"));
 	CapsuleCollisionComponent->SetCollisionProfileName(CollisionProfileName);
 
@@ -91,7 +90,6 @@ ATriggerBox::ATriggerBox(const FObjectInitializer& ObjectInitializer)
 
 	BoxCollisionComponent->ShapeColor = TriggerBaseColor;
 	BoxCollisionComponent->InitBoxExtent(FVector(40.0f, 40.0f, 40.0f));
-	BoxCollisionComponent->BodyInstance.bEnableCollision_DEPRECATED = true;
 
 	static FName CollisionProfileName(TEXT("Trigger"));
 	BoxCollisionComponent->SetCollisionProfileName(CollisionProfileName);
@@ -135,7 +133,6 @@ ATriggerSphere::ATriggerSphere(const FObjectInitializer& ObjectInitializer)
 
 	SphereCollisionComponent->ShapeColor = TriggerBaseColor;
 	SphereCollisionComponent->InitSphereRadius(+40.0f);
-	SphereCollisionComponent->BodyInstance.bEnableCollision_DEPRECATED = true;
 	static FName CollisionProfileName(TEXT("Trigger"));
 	SphereCollisionComponent->SetCollisionProfileName(CollisionProfileName);
 
diff --git a/Engine/Source/Runtime/Engine/Private/World.cpp b/Engine/Source/Runtime/Engine/Private/World.cpp
--- a/Engine/Source/Runtime/Engine/Private/World.cpp
+++ b/Engine/Source/Runtime/Engine/Private/World.cpp
@@ -176,12 +176,6 @@ void UWorld::Serialize( FArchive& Ar )
 	Ar << ExtraReferencedObjects;
 	Ar << StreamingLevels;
 		
-	if (Ar.UE4Ver() < VER_UE4_REMOVE_CLIENTDESTROYEDACTORCONTENT)
-	{
-		TArray<class UObject*>	TempClientDestroyedActorContent;
-		Ar << TempClientDestroyedActorContent;
-	}
-
 	// Mark archive and package as containing a map if we're serializing to disk.
 	if( !HasAnyFlags( RF_ClassDefaultObject ) && Ar.IsPersistent() )
 	{
@@ -457,14 +451,6 @@ void UWorld::PostLoad()
 	Super::PostLoad();
 	CurrentLevel = PersistentLevel;
 
-	// copy StreamingLevels from WorldSettings to this if exists
-	AWorldSettings * WorldSettings = GetWorldSettings(false, false);
-	if ( WorldSettings && WorldSettings->StreamingLevels_DEPRECATED.Num() > 0 )
-	{
-		StreamingLevels = WorldSettings->StreamingLevels_DEPRECATED;
-		WorldSettings->StreamingLevels_DEPRECATED.Empty();
-	}
-	
 	// Remove null streaming level entries (could be if level was saved with transient level streaming objects)
 	StreamingLevels.Remove(nullptr);
 	
@@ -488,48 +474,6 @@ void UWorld::PostLoad()
 		}
 	}
 
-#if WITH_EDITORONLY_DATA
-	if( GIsEditor && !VisibleLayers_DEPRECATED.IsEmpty() )
-	{
-		TArray< FString > OldVisibleLayerArray; 
-		VisibleLayers_DEPRECATED.ParseIntoArray( &OldVisibleLayerArray, TEXT(","), 0 );
-
-		for( auto VisibleLayerIt = OldVisibleLayerArray.CreateConstIterator(); VisibleLayerIt; ++VisibleLayerIt )
-		{
-			FName VisibleLayerName = FName( *( *VisibleLayerIt ) );
-
-			if( VisibleLayerName == TEXT("None") )
-			{
-				continue;
-			}
-
-			bool bFoundExistingLayer = false;
-			for( auto LayerIt = Layers.CreateConstIterator(); LayerIt; ++LayerIt )
-			{
-				TWeakObjectPtr< ULayer > Layer = *LayerIt;
-				if( VisibleLayerName == Layer->LayerName )
-				{
-					bFoundExistingLayer = true;
-					break;
-				}
-			}
-
-			if( !bFoundExistingLayer )
-			{
-				ULayer* NewLayer = ConstructObject< ULayer >( ULayer::StaticClass(), this, NAME_None, RF_Transactional );
-				check( NewLayer != NULL );
-
-				NewLayer->LayerName = VisibleLayerName;
-				NewLayer->bIsVisible = true;
-
-				Layers.Add( NewLayer );
-			}
-		}
-
-		VisibleLayers_DEPRECATED.Empty();
-	}
-#endif // WITH_EDITORONLY_DATA
-
 	// Add the garbage collection callbacks
 	FLevelStreamingGCHelper::AddGarbageCollectorCallback();
 
diff --git a/Engine/Source/Runtime/Engine/Private/WorldSettings.cpp b/Engine/Source/Runtime/Engine/Private/WorldSettings.cpp
--- a/Engine/Source/Runtime/Engine/Private/WorldSettings.cpp
+++ b/Engine/Source/Runtime/Engine/Private/WorldSettings.cpp
@@ -60,7 +60,6 @@ AWorldSettings::AWorldSettings(const FObjectInitializer& ObjectInitializer)
 
 	UStaticMeshComponent* StaticMeshComponent = ObjectInitializer.CreateDefaultSubobject<UStaticMeshComponent>(this, TEXT("StaticMeshComponent0"));
 	StaticMeshComponent->bHiddenInGame = true;
-	StaticMeshComponent->BodyInstance.bEnableCollision_DEPRECATED = true;
 	StaticMeshComponent->SetCollisionProfileName(UCollisionProfile::BlockAll_ProfileName);
 	StaticMeshComponent->PostPhysicsComponentTick.bCanEverTick = false;
 	StaticMeshComponent->Mobility = EComponentMobility::Static;
diff --git a/Engine/Source/Runtime/Engine/Public/AnimEncoding.h b/Engine/Source/Runtime/Engine/Public/AnimEncoding.h
--- a/Engine/Source/Runtime/Engine/Public/AnimEncoding.h
+++ b/Engine/Source/Runtime/Engine/Public/AnimEncoding.h
@@ -162,8 +162,7 @@ public:
 	 */
 	virtual void ByteSwapIn(
 		UAnimSequence& Seq, 
-		FMemoryReader& MemoryReader,
-		int32 SourceArVersion) PURE_VIRTUAL(AnimEncoding::ByteSwapIn,);
+		FMemoryReader& MemoryReader) PURE_VIRTUAL(AnimEncoding::ByteSwapIn,);
 
 	/**
 	 * Handles Byte-swapping outgoing animation data to an array of BYTEs
@@ -298,7 +297,7 @@ public:
 		FTransform& OutAtom,
 		const UAnimSequence& Seq,
 		int32 TrackIndex,
-		float Time);
+		float Time) override;
 
 	/**
 	 * Decompress the Rotation component of a BoneAtom
@@ -362,12 +361,10 @@ public:
 	 *
 	 * @param	Seq					An Animation Sequence to contain the read data.
 	 * @param	MemoryReader		The MemoryReader object to read from.
-	 * @param	SourceArVersion		The version of the archive that the data is coming from.
 	 */
 	virtual void ByteSwapIn(
 		UAnimSequence& Seq, 
-		FMemoryReader& MemoryReader,
-		int32 SourceArVersion);
+		FMemoryReader& MemoryReader) override;
 
 	/**
 	 * Handles Byte-swapping outgoing animation data to an array of BYTEs
@@ -379,7 +376,7 @@ public:
 	virtual void ByteSwapOut(
 		UAnimSequence& Seq,
 		TArray<uint8>& SerializedData, 
-		bool ForceByteSwapping);
+		bool ForceByteSwapping) override;
 
 	/**
 	 * Handles the ByteSwap of compressed animation data on import
@@ -395,8 +392,7 @@ public:
 		UAnimSequence& Seq, 
 		FMemoryReader& MemoryReader,
 		uint8*& Stream,
-		int32 NumKeys,
-		int32 SourceArVersion) PURE_VIRTUAL(AnimEncoding::ByteSwapRotationIn,);
+		int32 NumKeys) PURE_VIRTUAL(AnimEncoding::ByteSwapRotationIn,);
 
 	/**
 	 * Handles the ByteSwap of compressed animation data on import
@@ -412,8 +408,7 @@ public:
 		UAnimSequence& Seq, 
 		FMemoryReader& MemoryReader,
 		uint8*& Stream,
-		int32 NumKeys,
-		int32 SourceArVersion) PURE_VIRTUAL(AnimEncoding::ByteSwapTranslationIn,);
+		int32 NumKeys) PURE_VIRTUAL(AnimEncoding::ByteSwapTranslationIn,);
 
 	/**
 	 * Handles the ByteSwap of compressed animation data on import
@@ -429,8 +424,7 @@ public:
 		UAnimSequence& Seq, 
 		FMemoryReader& MemoryReader,
 		uint8*& Stream,
-		int32 NumKeys,
-		int32 SourceArVersion) PURE_VIRTUAL(AnimEncoding::ByteSwapScaleIn,);
+		int32 NumKeys) PURE_VIRTUAL(AnimEncoding::ByteSwapScaleIn,);
 
 	/**
 	 * Handles the ByteSwap of compressed animation data on export
diff --git a/Engine/Source/Runtime/Engine/Public/AnimEncoding_ConstantKeyLerp.h b/Engine/Source/Runtime/Engine/Public/AnimEncoding_ConstantKeyLerp.h
--- a/Engine/Source/Runtime/Engine/Public/AnimEncoding_ConstantKeyLerp.h
+++ b/Engine/Source/Runtime/Engine/Public/AnimEncoding_ConstantKeyLerp.h
@@ -21,14 +21,12 @@ public:
 	 * @param	MemoryReader	The FMemoryReader to read from.
 	 * @param	RotTrackData	The compressed rotation data stream.
 	 * @param	NumKeysRot		The number of keys present in the stream.
-	 * @param	SourceArVersion	The version number of the source archive stream.
 	 */
-	void ByteSwapRotationIn(
+	virtual void ByteSwapRotationIn(
 		UAnimSequence& Seq, 
 		FMemoryReader& MemoryReader,
 		uint8*& RotTrackData,
-		int32 NumKeysRot,
-		int32 SourceArVersion);
+		int32 NumKeysRot) override;
 
 	/**
 	 * Handles the ByteSwap of compressed translation data on import
@@ -37,14 +35,12 @@ public:
 	 * @param	MemoryReader	The FMemoryReader to read from.
 	 * @param	TransTrackData	The compressed translation data stream.
 	 * @param	NumKeysTrans	The number of keys present in the stream.
-	 * @param	SourceArVersion	The version number of the source archive stream.
 	 */
-	void ByteSwapTranslationIn(
+	virtual void ByteSwapTranslationIn(
 		UAnimSequence& Seq, 
 		FMemoryReader& MemoryReader,
 		uint8*& TransTrackData,
-		int32 NumKeysTrans,
-		int32 SourceArVersion);
+		int32 NumKeysTrans) override;
 
 	/**
 	 * Handles the ByteSwap of compressed Scale data on import
@@ -53,14 +49,12 @@ public:
 	 * @param	MemoryReader	The FMemoryReader to read from.
 	 * @param	ScaleTrackData	The compressed Scale data stream.
 	 * @param	NumKeysScale	The number of keys present in the stream.
-	 * @param	SourceArVersion	The version number of the source archive stream.
 	 */
-	void ByteSwapScaleIn(
+	virtual void ByteSwapScaleIn(
 		UAnimSequence& Seq, 
 		FMemoryReader& MemoryReader,
 		uint8*& ScaleTrackData,
-		int32 NumKeysScale,
-		int32 SourceArVersion);
+		int32 NumKeysScale) override;
 
 
 	/**
@@ -71,11 +65,11 @@ public:
 	 * @param	RotTrackData	The compressed rotation data stream.
 	 * @param	NumKeysRot		The number of keys to write to the stream.
 	 */
-	void ByteSwapRotationOut(
+	virtual void ByteSwapRotationOut(
 		UAnimSequence& Seq, 
 		FMemoryWriter& MemoryWriter,
 		uint8*& RotTrackData,
-		int32 NumKeysRot);
+		int32 NumKeysRot) override;
 
 	/**
 	 * Handles the ByteSwap of compressed translation data on export
@@ -85,11 +79,11 @@ public:
 	 * @param	TransTrackData	The compressed translation data stream.
 	 * @param	NumKeysTrans	The number of keys to write to the stream.
 	 */
-	void ByteSwapTranslationOut(
+	virtual void ByteSwapTranslationOut(
 		UAnimSequence& Seq, 
 		FMemoryWriter& MemoryWriter,
 		uint8*& TransTrackData,
-		int32 NumKeysTrans);
+		int32 NumKeysTrans) override;
 
 	/**
 	 * Handles the ByteSwap of compressed Scale data on export
@@ -99,11 +93,11 @@ public:
 	 * @param	TransTrackData	The compressed Scale data stream.
 	 * @param	NumKeysTrans	The number of keys to write to the stream.
 	 */
-	void ByteSwapScaleOut(
+	virtual void ByteSwapScaleOut(
 		UAnimSequence& Seq, 
 		FMemoryWriter& MemoryWriter,
 		uint8*& ScaleTrackData,
-		int32 NumKeysScale);
+		int32 NumKeysScale) override;
 };
 
 template<int32 FORMAT>
diff --git a/Engine/Source/Runtime/Engine/Public/AnimEncoding_PerTrackCompression.h b/Engine/Source/Runtime/Engine/Public/AnimEncoding_PerTrackCompression.h
--- a/Engine/Source/Runtime/Engine/Public/AnimEncoding_PerTrackCompression.h
+++ b/Engine/Source/Runtime/Engine/Public/AnimEncoding_PerTrackCompression.h
@@ -21,10 +21,7 @@ public:
 	 * @param	MemoryReader		The MemoryReader object to read from.
 	 * @param	SourceArVersion		The version of the archive that the data is coming from.
 	 */
-	virtual void ByteSwapIn(
-		UAnimSequence& Seq, 
-		FMemoryReader& MemoryReader,
-		int32 SourceArVersion);
+	virtual void ByteSwapIn(UAnimSequence& Seq, FMemoryReader& MemoryReader) override;
 
 	/**
 	 * Handles Byte-swapping outgoing animation data to an array of BYTEs
@@ -36,7 +33,7 @@ public:
 	virtual void ByteSwapOut(
 		UAnimSequence& Seq,
 		TArray<uint8>& SerializedData, 
-		bool ForceByteSwapping);
+		bool ForceByteSwapping) override;
 
 	/**
 	 * Extracts a single BoneAtom from an Animation Sequence.
@@ -50,7 +47,7 @@ public:
 		FTransform& OutAtom,
 		const UAnimSequence& Seq,
 		int32 TrackIndex,
-		float Time);
+		float Time) override;
 
 #if USE_ANIMATION_CODEC_BATCH_SOLVER
 
@@ -67,7 +64,7 @@ public:
 		FTransformArray& Atoms, 
 		const BoneTrackArray& DesiredPairs,
 		const UAnimSequence& Seq,
-		float Time);
+		float Time) override;
 
 	/**
 	 * Decompress all requested translation components from an Animation Sequence
@@ -82,7 +79,7 @@ public:
 		FTransformArray& Atoms, 
 		const BoneTrackArray& DesiredPairs,
 		const UAnimSequence& Seq,
-		float Time);
+		float Time) override;
 
 	/**
 	 * Decompress all requested Scale components from an Animation Sequence
@@ -97,7 +94,7 @@ public:
 		FTransformArray& Atoms, 
 		const BoneTrackArray& DesiredPairs,
 		const UAnimSequence& Seq,
-		float Time);
+		float Time) override;
 #endif
 
 protected:
diff --git a/Engine/Source/Runtime/Engine/Public/AnimEncoding_VariableKeyLerp.h b/Engine/Source/Runtime/Engine/Public/AnimEncoding_VariableKeyLerp.h
--- a/Engine/Source/Runtime/Engine/Public/AnimEncoding_VariableKeyLerp.h
+++ b/Engine/Source/Runtime/Engine/Public/AnimEncoding_VariableKeyLerp.h
@@ -24,14 +24,12 @@ public:
 	 * @param	MemoryReader	The FMemoryReader to read from.
 	 * @param	RotTrackData	The compressed rotation data stream.
 	 * @param	NumKeysRot		The number of keys present in the stream.
-	 * @param	SourceArVersion	The version number of the source archive stream.
 	 */
-	void ByteSwapRotationIn(
+	virtual void ByteSwapRotationIn(
 		UAnimSequence& Seq, 
 		FMemoryReader& MemoryReader,
 		uint8*& RotTrackData,
-		int32 NumKeysRot,
-		int32 SourceArVersion);
+		int32 NumKeysRot) override;
 
 	/**
 	 * Handles the ByteSwap of compressed translation data on import
@@ -40,14 +38,12 @@ public:
 	 * @param	MemoryReader	The FMemoryReader to read from.
 	 * @param	TransTrackData	The compressed translation data stream.
 	 * @param	NumKeysTrans	The number of keys present in the stream.
-	 * @param	SourceArVersion	The version number of the source archive stream.
 	 */
-	void ByteSwapTranslationIn(
+	virtual void ByteSwapTranslationIn(
 		UAnimSequence& Seq, 
 		FMemoryReader& MemoryReader,
 		uint8*& TransTrackData,
-		int32 NumKeysTrans,
-		int32 SourceArVersion);
+		int32 NumKeysTransn) override;
 
 	/**
 	 * Handles the ByteSwap of compressed Scale data on import
@@ -56,14 +52,12 @@ public:
 	 * @param	MemoryReader	The FMemoryReader to read from.
 	 * @param	ScaleTrackData	The compressed Scale data stream.
 	 * @param	NumKeysScale	The number of keys present in the stream.
-	 * @param	SourceArVersion	The version number of the source archive stream.
 	 */
-	void ByteSwapScaleIn(
+	virtual void ByteSwapScaleIn(
 		UAnimSequence& Seq, 
 		FMemoryReader& MemoryReader,
 		uint8*& ScaleTrackData,
-		int32 NumKeysScale,
-		int32 SourceArVersion);
+		int32 NumKeysScale) override;
 
 	/**
 	 * Handles the ByteSwap of compressed rotation data on export
@@ -73,11 +67,11 @@ public:
 	 * @param	RotTrackData	The compressed rotation data stream.
 	 * @param	NumKeysRot		The number of keys to write to the stream.
 	 */
-	void ByteSwapRotationOut(
+	virtual void ByteSwapRotationOut(
 		UAnimSequence& Seq, 
 		FMemoryWriter& MemoryWriter,
 		uint8*& RotTrackData,
-		int32 NumKeysRot);
+		int32 NumKeysRot) override;
 
 	/**
 	 * Handles the ByteSwap of compressed translation data on export
@@ -87,11 +81,11 @@ public:
 	 * @param	TransTrackData	The compressed translation data stream.
 	 * @param	NumKeysTrans	The number of keys to write to the stream.
 	 */
-	void ByteSwapTranslationOut(
+	virtual void ByteSwapTranslationOut(
 		UAnimSequence& Seq, 
 		FMemoryWriter& MemoryWriter,
 		uint8*& TransTrackData,
-		int32 NumKeysTrans);
+		int32 NumKeysTrans) override;
 
 	/**
 	 * Handles the ByteSwap of compressed Scale data on export
@@ -101,11 +95,11 @@ public:
 	 * @param	ScaleTrackData	The compressed Scale data stream.
 	 * @param	NumKeysScale	The number of keys to write to the stream.
 	 */
-	void ByteSwapScaleOut(
+	virtual void ByteSwapScaleOut(
 		UAnimSequence& Seq, 
 		FMemoryWriter& MemoryWriter,
 		uint8*& ScaleTrackData,
-		int32 NumKeysScale);
+		int32 NumKeysScale) override;
 };
 
 template<int32 FORMAT>
@@ -122,13 +116,13 @@ public:
 	 * @param	RelativePos		Current position within the animation to solve for in the range [0.0,1.0].
 	 * @return					None. 
 	 */
-	void GetBoneAtomRotation(	
+	virtual void GetBoneAtomRotation(	
 		FTransform& OutAtom,
 		const UAnimSequence& Seq,
 		const uint8* RESTRICT Stream,
 		int32 NumKeys,
 		float Time,
-		float RelativePos);
+		float RelativePos) override;
 
 	/**
 	 * Decompress the Translation component of a BoneAtom
@@ -140,13 +134,13 @@ public:
 	 * @param	RelativePos		Current position within the animation to solve for in the range [0.0,1.0].
 	 * @return					None. 
 	 */
-	void GetBoneAtomTranslation(	
+	virtual void GetBoneAtomTranslation(	
 		FTransform& OutAtom,
 		const UAnimSequence& Seq,
 		const uint8* RESTRICT Stream,
 		int32 NumKeys,
 		float Time,
-		float RelativePos);
+		float RelativePos) override;
 
 	
 	/**
@@ -159,13 +153,13 @@ public:
 	 * @param	RelativePos		Current position within the animation to solve for in the range [0.0,1.0].
 	 * @return					None. 
 	 */
-	void GetBoneAtomScale(	
+	virtual void GetBoneAtomScale(	
 		FTransform& OutAtom,
 		const UAnimSequence& Seq,
 		const uint8* RESTRICT Stream,
 		int32 NumKeys,
 		float Time,
-		float RelativePos);
+		float RelativePos) override;
 
 #if USE_ANIMATION_CODEC_BATCH_SOLVER
 
@@ -178,11 +172,11 @@ public:
 	 * @param	Time			Current time to solve for.
 	 * @return					None. 
 	 */
-	void GetPoseRotations(	
+	virtual void GetPoseRotations(	
 		FTransformArray& Atoms, 
 		const BoneTrackArray& DesiredPairs,
 		const UAnimSequence& Seq,
-		float RelativePos);
+		float RelativePos) override;
 
 	/**
 	 * Decompress all requested translation components from an Animation Sequence
@@ -193,11 +187,11 @@ public:
 	 * @param	Time			Current time to solve for.
 	 * @return					None. 
 	 */
-	void GetPoseTranslations(	
+	virtual void GetPoseTranslations(	
 		FTransformArray& Atoms,
 		const BoneTrackArray& DesiredPairs,
 		const UAnimSequence& Seq,
-		float RelativePos);
+		float RelativePos) override;
 
 	/**
 	 * Decompress all requested Scale components from an Animation Sequence
@@ -208,11 +202,11 @@ public:
 	 * @param	Time			Current time to solve for.
 	 * @return					None. 
 	 */
-	void GetPoseScales(	
+	virtual void GetPoseScales(	
 		FTransformArray& Atoms,
 		const BoneTrackArray& DesiredPairs,
 		const UAnimSequence& Seq,
-		float RelativePos);
+		float RelativePos) override;
 #endif
 
 };
diff --git a/Engine/Source/Runtime/Engine/Public/Engine.h b/Engine/Source/Runtime/Engine/Public/Engine.h
--- a/Engine/Source/Runtime/Engine/Public/Engine.h
+++ b/Engine/Source/Runtime/Engine/Public/Engine.h
@@ -215,7 +215,6 @@
 #include "Distributions/DistributionFloatConstant.h"
 #include "Distributions/DistributionFloatParameterBase.h"
 #include "Distributions/DistributionFloatParticleParameter.h"
-#include "Distributions/DistributionFloatSoundParameter.h"
 #include "Distributions/DistributionFloatConstantCurve.h"
 #include "Distributions/DistributionFloatUniform.h"
 #include "Distributions/DistributionFloatUniformCurve.h"
diff --git a/Engine/Source/Runtime/Engine/Public/EngineUtils.h b/Engine/Source/Runtime/Engine/Public/EngineUtils.h
--- a/Engine/Source/Runtime/Engine/Public/EngineUtils.h
+++ b/Engine/Source/Runtime/Engine/Public/EngineUtils.h
@@ -713,7 +713,7 @@ public:
 	 * @param InClassFlags - User defined per class flags .
 	 * @param InVersion - Minimal strip version required to serialize strip flags
 	 */
-	FStripDataFlags( FArchive& Ar, uint8 InClassFlags = 0, int32 InVersion = VER_UE4_REMOVED_STRIP_DATA );
+	FStripDataFlags( FArchive& Ar, uint8 InClassFlags = 0, int32 InVersion = VER_UE4_OLDEST_LOADABLE_PACKAGE );
 
 	/** 
 	 * Constructor.
@@ -724,7 +724,7 @@ public:
 	 * @param InClassFlags - User defined per class flags.
 	 * @param InVersion - Minimal version required to serialize strip flags
 	 */
-	FStripDataFlags( FArchive& Ar, uint8 InGlobalFlags, uint8 InClassFlags, int32 InVersion = VER_UE4_REMOVED_STRIP_DATA );
+	FStripDataFlags( FArchive& Ar, uint8 InGlobalFlags, uint8 InClassFlags, int32 InVersion = VER_UE4_OLDEST_LOADABLE_PACKAGE );
 
 	/**
 	 * Checks if FStripDataFlags::Editor flag is set or not
diff --git a/Engine/Source/Runtime/Engine/Public/SkeletalMeshTypes.h b/Engine/Source/Runtime/Engine/Public/SkeletalMeshTypes.h
--- a/Engine/Source/Runtime/Engine/Public/SkeletalMeshTypes.h
+++ b/Engine/Source/Runtime/Engine/Public/SkeletalMeshTypes.h
@@ -1368,72 +1368,6 @@ private:
 	void AllocateData();	
 };
 
-//////////////////////////////////////////////////////////////////////////
-// DEPRECATED (can remove when min ue4 version > VER_UE4_REMOVE_EXTRA_SKELMESH_VERTEX_INFLUENCES)
-
-struct FInfluenceWeights_DEPRECATED
-{
-	uint32 InfluenceWeightsDWORD; 
-
-	friend FArchive& operator<<( FArchive& Ar, FInfluenceWeights_DEPRECATED& W )
-	{
-		return Ar << W.InfluenceWeightsDWORD;
-	}
-};
-
-struct FInfluenceBones_DEPRECATED
-{
-	uint32 InfluenceBonesDWORD; 
-
-	friend FArchive& operator<<( FArchive& Ar, FInfluenceBones_DEPRECATED& W )
-	{
-		return Ar << W.InfluenceBonesDWORD;
-	}
-
-};
-
-struct FVertexInfluence_DEPRECATED
-{
-	FInfluenceWeights_DEPRECATED Weights;
-	FInfluenceBones_DEPRECATED Bones;
-
-	friend FArchive& operator<<( FArchive& Ar, FVertexInfluence_DEPRECATED& W )
-	{
-		return Ar << W.Weights << W.Bones;
-	}
-};
-
-class FSkeletalMeshVertexInfluences_DEPRECATED : public FVertexBuffer
-{
-public:
-	TResourceArray<FVertexInfluence_DEPRECATED, VERTEXBUFFER_ALIGNMENT> Influences;
-	TMap<struct FBoneIndexPair, TArray<uint32> > VertexInfluenceMapping;
-	TArray<FSkelMeshSection> Sections;
-	TArray<FSkelMeshChunk> Chunks;
-	TArray<FBoneIndexType> RequiredBones;
-	TArray<int32> CustomLeftRightSectionMap;
-
-	FSkeletalMeshVertexInfluences_DEPRECATED() : Influences(true) {}
-
-
-	friend FArchive& operator<<( FArchive& Ar, FSkeletalMeshVertexInfluences_DEPRECATED& W )
-	{
-		Ar << W.Influences;
-
-		Ar << W.VertexInfluenceMapping;
-		Ar << W.Sections;
-		Ar << W.Chunks;
-		Ar << W.RequiredBones;
-		uint8 Usage = 0;
-		Ar << Usage;
-
-		return Ar;
-	}
-};
-
-//////////////////////////////////////////////////////////////////////////
-
-
 struct FMultiSizeIndexContainerData
 {
 	TArray<uint32> Indices;
diff --git a/Engine/Source/Runtime/GameplayAbilities/Private/Abilities/GameplayAbilityWorldReticle_ActorVisualization.cpp b/Engine/Source/Runtime/GameplayAbilities/Private/Abilities/GameplayAbilityWorldReticle_ActorVisualization.cpp
--- a/Engine/Source/Runtime/GameplayAbilities/Private/Abilities/GameplayAbilityWorldReticle_ActorVisualization.cpp
+++ b/Engine/Source/Runtime/GameplayAbilities/Private/Abilities/GameplayAbilityWorldReticle_ActorVisualization.cpp
@@ -22,7 +22,6 @@ AGameplayAbilityWorldReticle_ActorVisualization::AGameplayAbilityWorldReticle_Ac
 	CollisionComponent->bAbsoluteScale = true;
 	//CollisionComponent->AlwaysLoadOnServer = true;
 	CollisionComponent->bCanEverAffectNavigation = false;
-	CollisionComponent->BodyInstance.bEnableCollision_DEPRECATED = false;
 	CollisionComponent->SetCollisionEnabled(ECollisionEnabled::NoCollision);
 
 	//USceneComponent* SceneComponent = ObjectInitializer.CreateDefaultSubobject<USceneComponent>(this, TEXT("RootComponent0"));
diff --git a/Engine/Source/Runtime/Landscape/Classes/LandscapeComponent.h b/Engine/Source/Runtime/Landscape/Classes/LandscapeComponent.h
--- a/Engine/Source/Runtime/Landscape/Classes/LandscapeComponent.h
+++ b/Engine/Source/Runtime/Landscape/Classes/LandscapeComponent.h
@@ -72,9 +72,6 @@ struct FWeightmapLayerAllocationInfo
 {
 	GENERATED_USTRUCT_BODY()
 
-	UPROPERTY()
-	FName LayerName_DEPRECATED;
-
 	UPROPERTY()
 	ULandscapeLayerInfoObject* LayerInfo;
 
@@ -162,10 +159,6 @@ class ULandscapeComponent : public UPrimitiveComponent
 	UPROPERTY(TextExportTransient)
 	class UTexture2D* HeightmapTexture;
 
-	/** Cached bounds, created at heightmap update time */
-	UPROPERTY()
-	FBoxSphereBounds CachedBoxSphereBounds_DEPRECATED;
-
 	/** Cached local-space bounding box, created at heightmap update time */
 	UPROPERTY()
 	FBox CachedLocalBox;
diff --git a/Engine/Source/Runtime/Landscape/Classes/LandscapeHeightfieldCollisionComponent.h b/Engine/Source/Runtime/Landscape/Classes/LandscapeHeightfieldCollisionComponent.h
--- a/Engine/Source/Runtime/Landscape/Classes/LandscapeHeightfieldCollisionComponent.h
+++ b/Engine/Source/Runtime/Landscape/Classes/LandscapeHeightfieldCollisionComponent.h
@@ -22,10 +22,6 @@ class ULandscapeHeightfieldCollisionComponent : public UPrimitiveComponent
 {
 	GENERATED_UCLASS_BODY()
 
-	/** DEPRECATED List of layers painted on this component. Matches the WeightmapLayerAllocations array in the LandscapeComponent. */
-	UPROPERTY()
-	TArray<FName> ComponentLayers_DEPRECATED;
-
 	/** List of layers painted on this component. Matches the WeightmapLayerAllocations array in the LandscapeComponent. */
 	UPROPERTY()
 	TArray<ULandscapeLayerInfoObject*> ComponentLayerInfos;
@@ -53,10 +49,6 @@ class ULandscapeHeightfieldCollisionComponent : public UPrimitiveComponent
 	UPROPERTY()
 	FGuid HeightfieldGuid;
 
-	/** Cached bounds, created at heightmap update time */
-	UPROPERTY()
-	FBoxSphereBounds CachedBoxSphereBounds_DEPRECATED;
-
 	/** Cached local-space bounding box, created at heightmap update time */
 	UPROPERTY()
 	FBox CachedLocalBox;
diff --git a/Engine/Source/Runtime/Landscape/Classes/LandscapeProxy.h b/Engine/Source/Runtime/Landscape/Classes/LandscapeProxy.h
--- a/Engine/Source/Runtime/Landscape/Classes/LandscapeProxy.h
+++ b/Engine/Source/Runtime/Landscape/Classes/LandscapeProxy.h
@@ -291,9 +291,6 @@ public:
 	UPROPERTY(EditAnywhere, BlueprintReadOnly, Category=Collision, meta=(ShowOnlyInnerProperties))
 	FBodyInstance BodyInstance;
 
-	UPROPERTY()
-	TArray<struct FLandscapeLayerStruct> LayerInfoObjs_DEPRECATED;
-
 #if WITH_EDITORONLY_DATA
 	UPROPERTY()
 	TArray<ULandscapeLayerInfoObject*> EditorCachedLayerInfos_DEPRECATED;
@@ -371,7 +368,6 @@ public:
 	virtual void PostEditImport() override;
 	// End UObject Interface
 
-	FLandscapeLayerStruct* GetLayerInfo_Deprecated(FName LayerName);
 	LANDSCAPE_API static TArray<FName> GetLayersFromMaterial(UMaterialInterface* Material);
 	LANDSCAPE_API TArray<FName> GetLayersFromMaterial() const;
 	LANDSCAPE_API static ULandscapeLayerInfoObject* CreateLayerInfo(const TCHAR* LayerName, ULevel* Level);
diff --git a/Engine/Source/Runtime/Landscape/Private/Landscape.cpp b/Engine/Source/Runtime/Landscape/Private/Landscape.cpp
--- a/Engine/Source/Runtime/Landscape/Private/Landscape.cpp
+++ b/Engine/Source/Runtime/Landscape/Private/Landscape.cpp
@@ -76,7 +76,6 @@ FAutoConsoleCommand CmdPrintNumLandscapeShadows(
 ULandscapeComponent::ULandscapeComponent(const FObjectInitializer& ObjectInitializer)
 	: Super(ObjectInitializer)
 {
-	BodyInstance.bEnableCollision_DEPRECATED = true;
 	SetCollisionProfileName(UCollisionProfile::BlockAll_ProfileName);
 	bGenerateOverlapEvents = false;
 	CastShadow = true;
@@ -162,10 +161,7 @@ void ULandscapeComponent::Serialize(FArchive& Ar)
 	Super::Serialize(Ar);
 
 	Ar << LightMap;
-	if (Ar.UE4Ver() >= VER_UE4_PRECOMPUTED_SHADOW_MAPS_BSP)
-	{
-		Ar << ShadowMap;
-	}
+	Ar << ShadowMap;
 
 #if WITH_EDITOR
 	if (Ar.IsTransacting())
@@ -182,12 +178,6 @@ void ULandscapeComponent::Serialize(FArchive& Ar)
 	}
 #endif
 
-	if (Ar.UE4Ver() < VER_UE4_CHANGED_IRRELEVANT_LIGHT_GUIDS)
-	{
-		// Irrelevant lights were incorrect before VER_UE4_CHANGED_IRRELEVANT_LIGHT_GUIDS
-		IrrelevantLights.Empty();
-	}
-
 	bool bCooked = false;
 
 	if (Ar.UE4Ver() >= VER_UE4_LANDSCAPE_PLATFORMDATA_COOKING && !HasAnyFlags(RF_ClassDefaultObject))
@@ -383,25 +373,6 @@ void ULandscapeComponent::PostLoad()
 		bCastStaticShadow = LandscapeProxy->bCastStaticShadow;
 		bCastShadowAsTwoSided = LandscapeProxy->bCastShadowAsTwoSided;
 
-		// convert from deprecated layer names to direct LayerInfo references
-		static const FName DataWeightmapName = FName("__DataLayer__");
-		for (int32 i = 0; i < WeightmapLayerAllocations.Num(); i++)
-		{
-			if (WeightmapLayerAllocations[i].LayerName_DEPRECATED != NAME_None)
-			{
-				if (WeightmapLayerAllocations[i].LayerName_DEPRECATED == DataWeightmapName)
-				{
-					WeightmapLayerAllocations[i].LayerInfo = ALandscapeProxy::VisibilityLayer;
-				}
-				else
-				{
-					FLandscapeLayerStruct* Layer = LandscapeProxy->GetLayerInfo_Deprecated(WeightmapLayerAllocations[i].LayerName_DEPRECATED);
-					WeightmapLayerAllocations[i].LayerInfo = Layer ? Layer->LayerInfoObj : NULL;
-				}
-				WeightmapLayerAllocations[i].LayerName_DEPRECATED = NAME_None;
-			}
-		}
-
 		// check SectionBaseX/Y are correct
 		int32 CheckSectionBaseX = FMath::RoundToInt(RelativeLocation.X) + LandscapeProxy->LandscapeSectionOffset.X;
 		int32 CheckSectionBaseY = FMath::RoundToInt(RelativeLocation.Y) + LandscapeProxy->LandscapeSectionOffset.Y;
@@ -436,18 +407,6 @@ void ULandscapeComponent::PostLoad()
 
 	if (!HasAnyFlags(RF_ClassDefaultObject))
 	{
-		if (!CachedLocalBox.IsValid && CachedBoxSphereBounds_DEPRECATED.SphereRadius > 0)
-		{
-			USceneComponent* LandscapeRoot = LandscapeProxy->GetRootComponent();
-			check(LandscapeRoot == AttachParent);
-
-			// Component isn't attached so we can't use its LocalToWorld
-			FTransform ComponentLtWTransform = FTransform(RelativeRotation, RelativeLocation, RelativeScale3D) * FTransform(LandscapeRoot->RelativeRotation, LandscapeRoot->RelativeLocation, LandscapeRoot->RelativeScale3D);
-
-			// This is good enough. The exact box will be calculated during painting.
-			CachedLocalBox = CachedBoxSphereBounds_DEPRECATED.GetBox().InverseTransformBy(ComponentLtWTransform);
-		}
-
 		// If we're loading on a platform that doesn't require cooked data, but *only* supports OpenGL ES, preload data from the DDC
 		if (!FPlatformProperties::RequiresCookedData() && GMaxRHIFeatureLevel <= ERHIFeatureLevel::ES3_1)
 		{
diff --git a/Engine/Source/Runtime/Landscape/Private/LandscapeCollision.cpp b/Engine/Source/Runtime/Landscape/Private/LandscapeCollision.cpp
--- a/Engine/Source/Runtime/Landscape/Private/LandscapeCollision.cpp
+++ b/Engine/Source/Runtime/Landscape/Private/LandscapeCollision.cpp
@@ -1310,45 +1310,6 @@ void ULandscapeHeightfieldCollisionComponent::PostLoad()
 	if (!HasAnyFlags(RF_ClassDefaultObject))
 	{
 		bShouldSaveCookedDataToDDC = true;
-
-		if (!CachedLocalBox.IsValid && CachedBoxSphereBounds_DEPRECATED.SphereRadius > 0)
-		{
-			ALandscapeProxy* LandscapeProxy = GetLandscapeProxy();
-			USceneComponent* LandscapeRoot = LandscapeProxy->GetRootComponent();
-			check(LandscapeRoot == AttachParent);
-
-			// Component isn't attached so we can't use its LocalToWorld
-			FTransform ComponentLtWTransform = FTransform(RelativeRotation, RelativeLocation, RelativeScale3D) * FTransform(LandscapeRoot->RelativeRotation, LandscapeRoot->RelativeLocation, LandscapeRoot->RelativeScale3D);
-
-			// This is good enough. The exact box will be calculated during painting.
-			CachedLocalBox = CachedBoxSphereBounds_DEPRECATED.GetBox().InverseTransformBy(ComponentLtWTransform);
-		}
-	}
-
-	// convert from deprecated layer names to direct LayerInfo references
-	ALandscapeProxy* LandscapeProxy = GetLandscapeProxy();
-	if (ensure(LandscapeProxy) && ComponentLayerInfos.Num() == 0 && ComponentLayers_DEPRECATED.Num() > 0)
-	{
-		ComponentLayerInfos.AddUninitialized(ComponentLayers_DEPRECATED.Num());
-
-		static const FName DataWeightmapName = FName("__DataLayer__");
-		for (int32 i = 0; i < ComponentLayers_DEPRECATED.Num(); i++)
-		{
-			if (ComponentLayers_DEPRECATED[i] != NAME_None)
-			{
-				if (ComponentLayers_DEPRECATED[i] == DataWeightmapName)
-				{
-					ComponentLayerInfos[i] = ALandscapeProxy::VisibilityLayer;
-				}
-				else
-				{
-					FLandscapeLayerStruct* Layer = LandscapeProxy->GetLayerInfo_Deprecated(ComponentLayers_DEPRECATED[i]);
-					ComponentLayerInfos[i] = Layer ? Layer->LayerInfoObj : NULL;
-				}
-			}
-		}
-
-		ComponentLayers_DEPRECATED.Empty();
 	}
 #endif//WITH_EDITOR
 }
@@ -1810,7 +1771,6 @@ void ULandscapeHeightfieldCollisionComponent::SetSectionBase(FIntPoint InSection
 ULandscapeHeightfieldCollisionComponent::ULandscapeHeightfieldCollisionComponent(const FObjectInitializer& ObjectInitializer)
 	: Super(ObjectInitializer)
 {
-	BodyInstance.bEnableCollision_DEPRECATED = true;
 	SetCollisionProfileName(UCollisionProfile::BlockAll_ProfileName);
 	bGenerateOverlapEvents = false;
 	CastShadow = false;
diff --git a/Engine/Source/Runtime/Landscape/Private/LandscapeEdit.cpp b/Engine/Source/Runtime/Landscape/Private/LandscapeEdit.cpp
--- a/Engine/Source/Runtime/Landscape/Private/LandscapeEdit.cpp
+++ b/Engine/Source/Runtime/Landscape/Private/LandscapeEdit.cpp
@@ -1603,20 +1603,6 @@ struct FHeightmapInfo
 	TArray<FColor*> HeightmapTextureMipData;
 };
 
-FLandscapeLayerStruct* ALandscapeProxy::GetLayerInfo_Deprecated(FName LayerName)
-{
-	for (auto It = LayerInfoObjs_DEPRECATED.CreateIterator(); It; It++)
-	{
-		ULandscapeLayerInfoObject* LayerInfo = It->LayerInfoObj;
-		if (LayerInfo && LayerInfo->LayerName == LayerName)
-		{
-			return &*It;
-		}
-	}
-
-	return NULL;
-}
-
 TArray<FName> ALandscapeProxy::GetLayersFromMaterial(UMaterialInterface* Material)
 {
 	TArray<FName> Result;
diff --git a/Engine/Source/Runtime/Landscape/Private/LandscapeGizmoActor.cpp b/Engine/Source/Runtime/Landscape/Private/LandscapeGizmoActor.cpp
--- a/Engine/Source/Runtime/Landscape/Private/LandscapeGizmoActor.cpp
+++ b/Engine/Source/Runtime/Landscape/Private/LandscapeGizmoActor.cpp
@@ -326,7 +326,6 @@ ULandscapeGizmoRenderComponent::ULandscapeGizmoRenderComponent(const FObjectInit
 	AlwaysLoadOnClient = false;
 	AlwaysLoadOnServer = false;
 	bSelectable = false;
-	BodyInstance.bEnableCollision_DEPRECATED = true;
 	SetCollisionProfileName(UCollisionProfile::BlockAll_ProfileName);
 }
 
diff --git a/Engine/Source/Runtime/Renderer/Private/ShaderBaseClasses.cpp b/Engine/Source/Runtime/Renderer/Private/ShaderBaseClasses.cpp
--- a/Engine/Source/Runtime/Renderer/Private/ShaderBaseClasses.cpp
+++ b/Engine/Source/Runtime/Renderer/Private/ShaderBaseClasses.cpp
@@ -240,7 +240,7 @@ bool FMaterialShader::Serialize(FArchive& Ar)
 		Ar << DebugUniformExpressionSet;
 		Ar << DebugDescription;
 	}
-	else if (Ar.IsLoading() && Ar.UE4Ver() >= VER_DEBUG_MATERIALSHADER_UNIFORM_EXPRESSIONS)
+	else if (Ar.IsLoading())
 	{
 		FUniformExpressionSet TempExpressionSet;
 		TempExpressionSet.Serialize(Ar);
