MarkPendingKill objects in all sublevels as well as the persistent level

diff --git a/Engine/Source/Runtime/Engine/Private/World.cpp b/Engine/Source/Runtime/Engine/Private/World.cpp
--- a/Engine/Source/Runtime/Engine/Private/World.cpp
+++ b/Engine/Source/Runtime/Engine/Private/World.cpp
@@ -4626,11 +4626,25 @@ UWorld* FSeamlessTravelHandler::Tick()
 			GWorld = NULL;
 
 			// mark everything else contained in the world to be deleted
+			TArray<UPackage*> CurrentWorldPackages;
+			for (auto LevelIt(CurrentWorld->GetLevelIterator()); LevelIt; ++LevelIt)
+			{
+				const ULevel* Level = *LevelIt;
+				if (Level)
+				{
+					CurrentWorldPackages.Add(Level->GetOutermost());
+				}
+			}
+
 			for (TObjectIterator<UObject> It; It; ++It)
 			{
-				if (It->IsIn(CurrentWorld))
+				for (const UPackage* WorldPackage : CurrentWorldPackages)
 				{
-					It->MarkPendingKill();
+					if (It->IsIn(WorldPackage))
+					{
+						It->MarkPendingKill();
+						break;
+					}
 				}
 			}
 
