Merging development branch (CL 2365488) to //depot/UE4/... Engine only, mostly hierarchical foliage

diff --git a/Engine/Shaders/BasePassForForwardShadingPixelShader.usf b/Engine/Shaders/BasePassForForwardShadingPixelShader.usf
--- a/Engine/Shaders/BasePassForForwardShadingPixelShader.usf
+++ b/Engine/Shaders/BasePassForForwardShadingPixelShader.usf
@@ -62,7 +62,7 @@ half DirectionalLightShadowing;
 void Main( 
 	FVertexFactoryInterpolantsVSToPS Interpolants,
 	FForwardShadingBasePassInterpolantsVSToPS BasePassInterpolants,
-		in half4 SvPosition	: SV_Position,
+	in float4 SvPosition : SV_Position,
 	OPTIONAL_IsFrontFace,
 	out half4 OutColor	: SV_Target0
 	)
@@ -73,7 +73,7 @@ void Main(
 #endif
 	FMaterialPixelParameters MaterialParameters = GetMaterialPixelParameters(Interpolants, BasePassInterpolants.PixelPosition);
 
-	CalcMaterialParameters(MaterialParameters, bIsFrontFace, BasePassInterpolants.PixelPosition
+	CalcMaterialParameters(MaterialParameters, SvPosition, bIsFrontFace, BasePassInterpolants.PixelPosition
 #if USE_WORLD_POSITION_EXCLUDING_SHADER_OFFSETS
 		, BasePassInterpolants.PixelPositionExcludingWPO
 #endif
@@ -82,16 +82,6 @@ void Main(
 	//Clip if the blend mode requires it.
 	GetMaterialCoverageAndClipping(MaterialParameters);	
 
-	//@mw todo:
-	//Go over with Brian to determine the proper defines to wrap it and test cases to fix it.
-#if !(GL3_PROFILE || GL4_PROFILE || COMPILER_GLSL_ES2 || COMPILER_GLSL || METAL_PROFILE)
-	float2 PixelPos = SvPosition - View.ViewRectMin.xy;	
-	MaterialParameters.ScreenPosition.xy = ( PixelPos / View.ViewSizeAndSceneTexelSize.xy - 0.5 ) * float2(2,-2);
-	MaterialParameters.ScreenPosition.zw = SvPosition.zw;
-	
-	MaterialParameters.ScreenPosition.xyz *= MaterialParameters.ScreenPosition.w;
-#endif
-
 	// Store the results in local variables and reuse instead of calling the functions multiple times.
 	// Store the results in local variables and reuse instead of calling the functions multiple times.
 	half3 BaseColor = GetMaterialBaseColor( MaterialParameters );
diff --git a/Engine/Shaders/BasePassPixelShader.usf b/Engine/Shaders/BasePassPixelShader.usf
--- a/Engine/Shaders/BasePassPixelShader.usf
+++ b/Engine/Shaders/BasePassPixelShader.usf
@@ -49,14 +49,6 @@ float4 PointSkyBentNormal;
 
 #define EDITOR_ALPHA2COVERAGE (EDITOR_PRIMITIVE_MATERIAL && FEATURE_LEVEL >= FEATURE_LEVEL_SM5 && !COMPILER_GLSL)
 
-// For now we only use it on EDITOR_PRIMITIVE_MATERIAL (LevelGrid bias can be much lower) but later we might want to use it for all materials.
-// According to GCNPerformanceTweets.pdf Tip 38 it might have a performance cost on GCN.
-// On consoles we don't need the editor rendering features
-#if EDITOR_PRIMITIVE_MATERIAL && FEATURE_LEVEL >= FEATURE_LEVEL_SM5 && !PS4_PROFILE && !XBOXONE_PROFILE
-	// not available on mobile, if not available we have to reconstruct it from the world position interpolator which is much worse quality (noticable for depth bias)
-	#define HIGHQUALITY_PS_POSITION 1
-#endif
-
 
 #if TRANSLUCENCY_LIGHTING_SURFACE
 
@@ -534,9 +526,7 @@ uint CustomAlpha2Coverage(inout float4 InOutColor)
 void Main(
 	FVertexFactoryInterpolantsVSToPS Interpolants,
 	FBasePassInterpolantsVSToPS BasePassInterpolants,
-#if HIGHQUALITY_PS_POSITION
-   float4 InSVPosition : SV_POSITION,
-#endif
+	in float4 SvPosition : SV_Position,
 	OPTIONAL_IsFrontFace,
 	out float4 OutColor		: SV_Target0
 #if FEATURE_LEVEL >= FEATURE_LEVEL_SM4 && (MATERIALBLENDING_SOLID || MATERIALBLENDING_MASKED) && !NO_GBUFFER
@@ -561,22 +551,12 @@ void Main(
 
 	FMaterialPixelParameters MaterialParameters = GetMaterialPixelParameters(Interpolants, BasePassInterpolants.PixelPosition);
 
-	CalcMaterialParameters(MaterialParameters,bIsFrontFace,BasePassInterpolants.PixelPosition
+	CalcMaterialParameters(MaterialParameters, SvPosition, bIsFrontFace, BasePassInterpolants.PixelPosition
 #if USE_WORLD_POSITION_EXCLUDING_SHADER_OFFSETS
 		, BasePassInterpolants.PixelPositionExcludingWPO
 #endif
 	);
 
-#if HIGHQUALITY_PS_POSITION
-	// SVPosition was already computed but only at low quality.
-	// Here we override it with the higher quality.
-
-	// Pixel position relative to left top of screen (not viewport), center of pixel
-	// z:DeviceZ, .w:SceneDepth
-	MaterialParameters.SVPosition = InSVPosition;
-#endif
-
-
 #if EDITOR_PRIMITIVE_MATERIAL && (FEATURE_LEVEL >= FEATURE_LEVEL_SM4 || ES2_EMULATION)
 	const bool bEditorWeightedZBuffering = true;
 #else
diff --git a/Engine/Shaders/DepthOnlyPixelShader.usf b/Engine/Shaders/DepthOnlyPixelShader.usf
--- a/Engine/Shaders/DepthOnlyPixelShader.usf
+++ b/Engine/Shaders/DepthOnlyPixelShader.usf
@@ -12,6 +12,7 @@ void Main(
 #if !MATERIALBLENDING_SOLID
 	FVertexFactoryInterpolantsVSToPS FactoryInterpolants,
 	float4 PixelPosition : TEXCOORD6,
+	in float4 SvPosition : SV_Position,
 	OPTIONAL_IsFrontFace,
 #endif
 	out float4 OutColor : SV_Target0
@@ -20,7 +21,7 @@ void Main(
 #if !MATERIALBLENDING_SOLID
 
 	FMaterialPixelParameters MaterialParameters = GetMaterialPixelParameters(FactoryInterpolants, PixelPosition);
-	CalcMaterialParameters(MaterialParameters, bIsFrontFace, PixelPosition);
+	CalcMaterialParameters(MaterialParameters, SvPosition, bIsFrontFace, PixelPosition);
 
 	GetMaterialCoverageAndClipping(MaterialParameters);
 #endif
diff --git a/Engine/Shaders/DistortAccumulatePS.usf b/Engine/Shaders/DistortAccumulatePS.usf
--- a/Engine/Shaders/DistortAccumulatePS.usf
+++ b/Engine/Shaders/DistortAccumulatePS.usf
@@ -18,13 +18,14 @@ float4 DistortionParams;
 void Main(
 	FVertexFactoryInterpolantsVSToPS Interpolants,
 	float4	PixelPosition	: TEXCOORD6,
+	in float4 SvPosition	: SV_Position,
 	OPTIONAL_IsFrontFace,
 	out float4 OutColor		: SV_Target0
 	)
 {
 	// material parameter inputs
 	FMaterialPixelParameters MaterialParameters = GetMaterialPixelParameters(Interpolants,PixelPosition);
-	CalcMaterialParameters(MaterialParameters,bIsFrontFace,PixelPosition);
+	CalcMaterialParameters(MaterialParameters, SvPosition, bIsFrontFace, PixelPosition);
 	
 	// material distortion offset
 	half3 Normal = GetMaterialNormal(MaterialParameters);
diff --git a/Engine/Shaders/HitProxyPixelShader.usf b/Engine/Shaders/HitProxyPixelShader.usf
--- a/Engine/Shaders/HitProxyPixelShader.usf
+++ b/Engine/Shaders/HitProxyPixelShader.usf
@@ -16,12 +16,13 @@ void Main(
 	float4 InstanceHitProxyId : HIT_PROXY_ID,
 #endif
 	float4 PixelPosition	: TEXCOORD6,
+	in float4 SvPosition	: SV_Position,
 	OPTIONAL_IsFrontFace,
 	out float4 OutColor		: SV_Target0
 	)
 {
 	FMaterialPixelParameters MaterialParameters = GetMaterialPixelParameters(Interpolants, PixelPosition);
-	CalcMaterialParameters(MaterialParameters,bIsFrontFace,PixelPosition);
+	CalcMaterialParameters(MaterialParameters, SvPosition, bIsFrontFace, PixelPosition);
 
 	GetMaterialCoverageAndClipping(MaterialParameters);
 
diff --git a/Engine/Shaders/LightMapDensityShader.usf b/Engine/Shaders/LightMapDensityShader.usf
--- a/Engine/Shaders/LightMapDensityShader.usf
+++ b/Engine/Shaders/LightMapDensityShader.usf
@@ -153,12 +153,13 @@ SamplerState	GridNormalTextureSampler;
 void MainPixelShader(
 	FVertexFactoryInterpolantsVSToPS FactoryInterpolants,
 	float4 WorldPosition	: TEXCOORD6,
+	in float4 SvPosition	: SV_Position,
 	OPTIONAL_IsFrontFace,
 	out float4 OutColor		: SV_Target0
 	)
 {
 	FMaterialPixelParameters MaterialParameters = GetMaterialPixelParameters( FactoryInterpolants, WorldPosition );
-	CalcMaterialParameters(MaterialParameters, bIsFrontFace, float4(0,0,.00001f,1));
+	CalcMaterialParameters(MaterialParameters, SvPosition, bIsFrontFace, float4(0,0,.00001f,1));
 	GetMaterialCoverageAndClipping(MaterialParameters);
 
 	float2 LightMapUV,LightMapUV1;
diff --git a/Engine/Shaders/LocalVertexFactory.usf b/Engine/Shaders/LocalVertexFactory.usf
--- a/Engine/Shaders/LocalVertexFactory.usf
+++ b/Engine/Shaders/LocalVertexFactory.usf
@@ -11,22 +11,47 @@ float4x4 PreviousLocalToWorld;
 float4 LightMapCoordinateScaleBias;
 float4 ShadowMapCoordinateScaleBias;
 
+#ifndef USE_NEW_LOD_TRANSITION
+	#define USE_NEW_LOD_TRANSITION 0
+#endif
+
 #if USE_INSTANCING
+
+#if USE_NEW_LOD_TRANSITION
+    float4 InstancingViewZCompareZero;
+    float4 InstancingViewZCompareOne;
+    float4 InstancingViewZConstant;
+    float4 InstancingWorldViewOriginZero;
+    float4 InstancingWorldViewOriginOne;
+#endif
+
 float4 InstancingFadeOutParams;
 
 #ifndef USE_INSTANCING_EMULATED
-#define USE_INSTANCING_EMULATED 0
+	#define USE_INSTANCING_EMULATED 0
 #endif	// USE_INSTANCING_EMULATED
 
-	#if USE_INSTANCING_EMULATED
-// Required for CPU emulation of Instancing
-float4 CPUInstanceShadowMapBias;
-float3x4 CPUInstanceTransform;
-float3x4 CPUInstanceInverseTransform;
-	#endif	// USE_INSTANCING_EMULATED
+#if USE_INSTANCING_EMULATED
+	// Required for CPU emulation of Instancing
+	float4 CPUInstanceShadowMapBias;
+	float3x4 CPUInstanceTransform;
+	float3x4 CPUInstanceInverseTransform;
+#endif	// USE_INSTANCING_EMULATED
 
+#else // USE_INSTANCING
+	#ifndef USE_NEW_LOD_TRANSITION
+		#define USE_NEW_LOD_TRANSITION (0)
+	#endif
+	#ifndef USE_INSTANCING_EMULATED
+		#define USE_INSTANCING_EMULATED (0)
+	#endif
 #endif	// USE_INSTANCING
 
+#if !USE_INSTANCING && USE_NEW_LOD_TRANSITION
+float NonInstancedDitherLODFactor;
+#endif
+
+
 #if USE_SPLINEDEFORM
 float3 SplineStartPos;
 float3 SplineStartTangent;
@@ -435,20 +460,37 @@ FVertexFactoryIntermediates GetVertexFactoryIntermediates(FVertexFactoryInput In
 	// Swizzle vertex color.
 	Intermediates.Color = Input.Color FCOLOR_COMPONENT_SWIZZLE;
 #if USE_INSTANCING
-	#if !USE_INSTANCING_EMULATED
-	float3 InstanceLocation = TransformLocalToTranslatedWorld(float3(Input.InstanceTransform1.w, Input.InstanceTransform2.w, Input.InstanceTransform3.w)).xyz;
-	// x = per-instance random, y = per-instance fade out factor
-	Intermediates.PerInstanceParams.x = Input.InstanceInverseTransform3.w;
-	#else
-	float3 InstanceLocation = TransformLocalToTranslatedWorld(float3(CPUInstanceTransform[0].w, CPUInstanceTransform[1].w, CPUInstanceTransform[2].w)).xyz;
-	// x = per-instance random, y = per-instance fade out factor
-	Intermediates.PerInstanceParams.x = CPUInstanceInverseTransform[2].w;
-	#endif // !USE_INSTANCING_EMULATED
-	Intermediates.PerInstanceParams.y = 1.0 - saturate((length(InstanceLocation) - InstancingFadeOutParams.x) * InstancingFadeOutParams.y);
+	// x = per-instance random, y = per-instance fade out factor, z = zero or one depending of if it is shown at all. 
+
 	// PerInstanceParams.z stores a hide/show flag for this instance
 	float SelectedValue = trunc(GetInstanceShadowMapBias(Input).z * (1.0 / 256.0));
 	// InstancingFadeOutParams.z,w are RenderSelected and RenderDeselected respectively.
 	Intermediates.PerInstanceParams.z = InstancingFadeOutParams.z * SelectedValue + InstancingFadeOutParams.w * (1-SelectedValue);
+	#if !USE_NEW_LOD_TRANSITION
+		#if !USE_INSTANCING_EMULATED
+			Intermediates.PerInstanceParams.x = Input.InstanceInverseTransform3.w;
+			float3 InstanceLocation = TransformLocalToTranslatedWorld(float3(Input.InstanceTransform1.w, Input.InstanceTransform2.w, Input.InstanceTransform3.w)).xyz;
+		#else
+			Intermediates.PerInstanceParams.x = CPUInstanceInverseTransform[2].w;
+			float3 InstanceLocation = TransformLocalToTranslatedWorld(float3(CPUInstanceTransform[0].w, CPUInstanceTransform[1].w, CPUInstanceTransform[2].w)).xyz;
+		#endif // !USE_INSTANCING_EMULATED
+		Intermediates.PerInstanceParams.y = 1.0 - saturate((length(InstanceLocation) - InstancingFadeOutParams.x) * InstancingFadeOutParams.y);
+	#else
+		#if !USE_INSTANCING_EMULATED
+			Intermediates.PerInstanceParams.x = Input.InstanceInverseTransform3.w;
+			float3 InstanceLocation = TransformLocalToWorld(float3(Input.InstanceTransform1.w, Input.InstanceTransform2.w, Input.InstanceTransform3.w)).xyz;
+		#else
+			Intermediates.PerInstanceParams.x = CPUInstanceInverseTransform[2].w;
+			float3 InstanceLocation = TransformLocalToWorld(float3(CPUInstanceTransform[0].w, CPUInstanceTransform[1].w, CPUInstanceTransform[2].w)).xyz;
+		#endif // !USE_INSTANCING_EMULATED
+
+		float ViewZZero = length(InstanceLocation - InstancingWorldViewOriginZero.xyz);
+		float ViewZOne = length(InstanceLocation - InstancingWorldViewOriginOne.xyz);
+		Intermediates.PerInstanceParams.y = 
+			dot(ViewZZero > InstancingViewZCompareZero.xyz, InstancingViewZConstant.xyz) * InstancingWorldViewOriginZero.w +
+			dot(ViewZOne > InstancingViewZCompareOne.xyz, InstancingViewZConstant.xyz) * InstancingWorldViewOriginOne.w;
+		Intermediates.PerInstanceParams.z *= abs(Intermediates.PerInstanceParams.y) < .999;
+	#endif
 #endif	// USE_INSTANCING
 	return Intermediates;
 }
diff --git a/Engine/Shaders/MaterialTemplate.usf b/Engine/Shaders/MaterialTemplate.usf
--- a/Engine/Shaders/MaterialTemplate.usf
+++ b/Engine/Shaders/MaterialTemplate.usf
@@ -21,6 +21,15 @@
 	#include "AtmosphereCommon.usf"
 #endif
 
+#if USE_INSTANCING
+	#ifndef USE_NEW_LOD_TRANSITION
+		#error "USE_NEW_LOD_TRANSITION must be defined if USE_INSTANCING is"
+	#endif
+#else
+	#ifndef USE_NEW_LOD_TRANSITION
+		#define USE_NEW_LOD_TRANSITION (0)
+	#endif
+#endif
 /**
  * Parameters used by vertex and pixel shaders to access particle properties.
  */
@@ -386,7 +395,11 @@ float GetPerInstanceRandom(FMaterialPixelParameters Parameters)
 float GetPerInstanceFadeAmount(FMaterialPixelParameters Parameters)
 {
 #if USE_INSTANCING
-	return float(Parameters.PerInstanceParams.y);
+	#if !USE_NEW_LOD_TRANSITION
+		return float(Parameters.PerInstanceParams.y);
+	#else
+		return float(saturate(1.0 - abs(Parameters.PerInstanceParams.y)));
+	#endif
 #else
 	return float(1.0);
 #endif
@@ -707,6 +720,7 @@ float3 GetSpeedTreeVertexOffset(FMaterialVertexParameters Parameters, int Geomet
 
 	// compute LOD by finding screen space size
 	float LodInterp = 1.0;
+#if !USE_INSTANCING || !USE_NEW_LOD_TRANSITION
 	if (LODType == SPEEDTREE_LOD_TYPE_SMOOTH) 
 	{
 		const float ScreenMultiple = 0.5 * max(View.ViewSizeAndSceneTexelSize.x * View.ViewToClip[0][0],
@@ -716,7 +730,7 @@ float3 GetSpeedTreeVertexOffset(FMaterialVertexParameters Parameters, int Geomet
 		const float ScreenArea = PI * ScreenRadius * ScreenRadius / (View.ViewSizeAndSceneTexelSize.x * View.ViewSizeAndSceneTexelSize.y);
 		LodInterp = saturate((ScreenArea - SpeedTreeLODInfo.x) / SpeedTreeLODInfo.z);
 	}
-
+#endif
 	TreePos *= 0.001; // The only other use of the tree position is as an offset into trig functions, but big numbers don't play nice there
 
 	// SpeedTrees should only be uniformly scaled, but if necessary, it takes a few more instructions
@@ -1261,8 +1275,38 @@ float3 GetMaterialHemisphereLightTransferFull(float3 DiffuseColor,FMaterialPixel
 	       LowerLighting * LowerColor;
 }
 
+void ClipLODTransition(FMaterialPixelParameters Parameters, float DitherFactor)
+{
+	if (abs(DitherFactor) > .001)
+	{
+		float RandCos = cos(dot(floor(Parameters.SVPosition.xy), float2(347.83451793,3343.28371963)));
+		float RandomVal = frac(RandCos * 1000.0);
+		half RetVal = (DitherFactor < 0.0) ?
+			(DitherFactor + 1.0 > RandomVal) :
+			(DitherFactor < RandomVal);
+		clip(RetVal - .001);
+	}
+}
+
+#if USE_INSTANCING && USE_NEW_LOD_TRANSITION
+void ClipLODTransition(FMaterialPixelParameters Parameters)
+{
+	ClipLODTransition(Parameters, Parameters.PerInstanceParams.y);
+}
+#elif USE_NEW_LOD_TRANSITION
+void ClipLODTransition(FMaterialPixelParameters Parameters)
+{
+	ClipLODTransition(Parameters, NonInstancedDitherLODFactor);
+}
+#else
+void ClipLODTransition(FMaterialPixelParameters Parameters)
+{
+}
+#endif
+
 void GetMaterialClippingShadowDepth(FMaterialPixelParameters Parameters) 
 {
+	ClipLODTransition(Parameters);
 	#if MATERIALBLENDING_MASKED
 		clip(GetMaterialMask(Parameters)); 
 	#elif MATERIALBLENDING_TRANSLUCENT
@@ -1273,11 +1317,13 @@ void GetMaterialClippingShadowDepth(FMaterialPixelParameters Parameters)
 #if MATERIALBLENDING_MASKED
 	void GetMaterialCoverageAndClipping(FMaterialPixelParameters Parameters)
 	{
+		ClipLODTransition(Parameters);
 		clip(GetMaterialMask(Parameters));
 	}
 #else
 	void GetMaterialCoverageAndClipping(FMaterialPixelParameters Parameters)
 	{
+		ClipLODTransition(Parameters);
 	}
 #endif
 
@@ -1298,6 +1344,7 @@ void GetMaterialClippingShadowDepth(FMaterialPixelParameters Parameters)
 /** Initializes the subset of Parameters that was not set in GetMaterialPixelParameters. */
 void CalcMaterialParameters(
 	in out FMaterialPixelParameters Parameters,
+	float4 SvPosition,
 	FIsFrontFace bIsFrontFace,
 	float4 PixelPosition,
 	float4 PixelPositionExcludingShaderOffsets)
@@ -1313,23 +1360,13 @@ void CalcMaterialParameters(
 	Parameters.WorldPosition_NoOffsets = PixelPositionExcludingShaderOffsets.xyz - View.PreViewTranslation.xyz;
 #endif
 
-	// Derive ScreenPosition from WorldPosition to avoid using another interpolator
-	Parameters.ScreenPosition = mul(PixelPosition, View.TranslatedWorldToClip);
+	Parameters.SVPosition = SvPosition;
 
-	// Compute low quality version of SVPosition (reconstructed from world Pos interpolator) which
-	// might get overriden later by actual value from SV_Position.
 	{
-		float2 NDC = Parameters.ScreenPosition.xy / Parameters.ScreenPosition.w;
-
-		// Could be optimized with more view constants.
-		float2 PixelPos = (NDC * float2(0.5f, -0.5f) + float2(0.5f, 0.5f)) * View.ViewSizeAndSceneTexelSize.xy;
-		
-		// Pixel position relative to left top of screen (not viewport), center of pixel
-		Parameters.SVPosition.xy = PixelPos + View.ViewRectMin.xy;
-		// z: DeviceZ
-		Parameters.SVPosition.z = ConvertToDeviceZ(Parameters.ScreenPosition.w);
-		// w: SceneDepth
-		Parameters.SVPosition.w = Parameters.ScreenPosition.w;
+		float2 PixelPos = SvPosition - View.ViewRectMin.xy;	
+		Parameters.ScreenPosition.xy = ( PixelPos / View.ViewSizeAndSceneTexelSize.xy - 0.5 ) * float2(2,-2);
+		Parameters.ScreenPosition.zw = SvPosition.zw;
+		Parameters.ScreenPosition.xyz *= Parameters.ScreenPosition.w;
 	}
 
 	// PixelPosition is the world position translated to the camera position, which is just -CameraVector
@@ -1375,10 +1412,11 @@ void CalcMaterialParameters(
 
 void CalcMaterialParameters(
 	in out FMaterialPixelParameters Parameters,
+	float4 SvPosition,
 	FIsFrontFace bIsFrontFace,
 	float4 PixelPosition)
 {
-	CalcMaterialParameters(Parameters, bIsFrontFace, PixelPosition, PixelPosition);
+	CalcMaterialParameters(Parameters, SvPosition, bIsFrontFace, PixelPosition, PixelPosition);
 }
 
 /** Assemble the transform from tangent space into world space */
diff --git a/Engine/Shaders/PostProcessMaterialShaders.usf b/Engine/Shaders/PostProcessMaterialShaders.usf
--- a/Engine/Shaders/PostProcessMaterialShaders.usf
+++ b/Engine/Shaders/PostProcessMaterialShaders.usf
@@ -19,6 +19,7 @@ void MainVS(
 
 void MainPS(
 	in float4 ScreenPosition : TEXCOORD0,
+	in float4 SvPosition : SV_Position,
 	out float4 OutColor : SV_Target0
 	)
 {
@@ -41,7 +42,7 @@ void MainPS(
 	float3 WorldPosition = ScreenVector * SceneDepth;
 
 	// fill out other related material parameters
-	CalcMaterialParameters(Parameters, true, float4(WorldPosition, 1));
+	CalcMaterialParameters(Parameters, SvPosition, true, float4(WorldPosition, 1));
 	Parameters.ScreenPosition = ScreenPosition;
 
 	// Grab emissive colour as output
diff --git a/Engine/Shaders/ShadowDepthPixelShader.usf b/Engine/Shaders/ShadowDepthPixelShader.usf
--- a/Engine/Shaders/ShadowDepthPixelShader.usf
+++ b/Engine/Shaders/ShadowDepthPixelShader.usf
@@ -35,6 +35,7 @@ float3 GetMaterialBounceColor(FMaterialPixelParameters MaterialParameters)
 
 void Main( 
 	FShadowDepthVSToPS Inputs,
+	in float4 SvPosition : SV_Position,
 	out float4 OutColor : SV_Target0
 #if REFLECTIVE_SHADOW_MAP
 	,out float4 OutReflectiveShadowMapDiffuse : SV_Target1
@@ -50,12 +51,11 @@ void Main(
 	FMaterialPixelParameters MaterialParameters = GetMaterialPixelParameters(Inputs.FactoryInterpolants, float4(0,0,.00001f,1));
 
 	#if INTERPOLATE_POSITION
-		CalcMaterialParameters(MaterialParameters, 1, Inputs.PixelPosition);
+		CalcMaterialParameters(MaterialParameters, SvPosition, 1, Inputs.PixelPosition);
 	#else
 		// Note: Using default values for PixelPosition to reduce interpolator cost
-		CalcMaterialParameters(MaterialParameters, 1, float4(0,0,.00001f,1));
+		CalcMaterialParameters(MaterialParameters, SvPosition, 1, float4(0,0,.00001f,1));
 	#endif
-	
 	// Evaluate the mask for masked materials
 	GetMaterialClippingShadowDepth(MaterialParameters);
 #endif
diff --git a/Engine/Shaders/SlateElementPixelShader.usf b/Engine/Shaders/SlateElementPixelShader.usf
--- a/Engine/Shaders/SlateElementPixelShader.usf
+++ b/Engine/Shaders/SlateElementPixelShader.usf
@@ -51,7 +51,7 @@ struct VertexOut
 float4 GetMaterialColor( VertexOut InVertex, float2 UV )
 {
 	FMaterialPixelParameters Parameters = MakeInitializedMaterialPixelParameters();
-	CalcMaterialParameters(Parameters, true, InVertex.Position);
+	CalcMaterialParameters(Parameters, InVertex.Position, true, InVertex.Position);
 
 #if NUM_MATERIAL_TEXCOORDS == 1
 	Parameters.TexCoords[0] = UV;
diff --git a/Engine/Shaders/TranslucentShadowDepthShaders.usf b/Engine/Shaders/TranslucentShadowDepthShaders.usf
--- a/Engine/Shaders/TranslucentShadowDepthShaders.usf
+++ b/Engine/Shaders/TranslucentShadowDepthShaders.usf
@@ -72,13 +72,13 @@ float TranslInvMaxSubjectDepth;
 /** Pixel shader used to accumulate layer opacities in different channels based on the first translucent layer's depth. */
 void MainOpacityPS( 
 	FTranslucencyShadowDepthVSToPS Inputs,
-	float4 ScreenPosition: SV_POSITION,
+	in float4 SvPosition : SV_Position,
 	out float4 OutColor0 : SV_Target0,
 	out float4 OutColor1 : SV_Target1
 	)
 {
 	FMaterialPixelParameters MaterialParameters = GetMaterialPixelParameters(Inputs.FactoryInterpolants, Inputs.PixelPosition);
-	CalcMaterialParameters(MaterialParameters, 1, Inputs.PixelPosition);
+	CalcMaterialParameters(MaterialParameters, SvPosition, 1, Inputs.PixelPosition);
 	
 	// Evaluate the mask for translucent materials
 	GetMaterialClippingShadowDepth(MaterialParameters);
diff --git a/Engine/Shaders/VelocityShader.usf b/Engine/Shaders/VelocityShader.usf
--- a/Engine/Shaders/VelocityShader.usf
+++ b/Engine/Shaders/VelocityShader.usf
@@ -188,14 +188,14 @@ float4 IndividualVelocityScale;
  
 void MainPixelShader(
 	FVelocityInterpsVSToPS Input,
-	float4	PixelPosition	: SV_POSITION,
+	in float4 SvPosition : SV_Position,
 	OPTIONAL_IsFrontFace,
 	out float4 OutColor : SV_Target0
 	)
 {
 	// Manual clipping here (alpha-test, etc)
-	FMaterialPixelParameters MaterialParameters = GetMaterialPixelParameters(Input.FactoryInterpolants, PixelPosition);
-	CalcMaterialParameters(MaterialParameters, bIsFrontFace, PixelPosition);
+	FMaterialPixelParameters MaterialParameters = GetMaterialPixelParameters(Input.FactoryInterpolants, SvPosition);
+	CalcMaterialParameters(MaterialParameters, SvPosition, bIsFrontFace, SvPosition);
 	GetMaterialCoverageAndClipping(MaterialParameters);
 
 	float2 ScreenPosition = Input.PackedVelocityA.xy / Input.PackedVelocityA.w - View.ViewToClip[2].xy;
diff --git a/Engine/Shaders/VertexFactoryCommon.usf b/Engine/Shaders/VertexFactoryCommon.usf
--- a/Engine/Shaders/VertexFactoryCommon.usf
+++ b/Engine/Shaders/VertexFactoryCommon.usf
@@ -2,11 +2,15 @@
 
 #include "Common.usf"
 
-float4 TransformLocalToTranslatedWorld(float3 LocalPosition)
+float3 TransformLocalToWorld(float3 LocalPosition)
 {
 	//expand the matrix multiply to avoid a code-gen bug on PS4.
 	//float3 RotatedPosition = mul(LocalPosition.xyz, (float3x3)Primitive.LocalToWorld);
 	float3 RotatedPosition = Primitive.LocalToWorld[0].xyz * LocalPosition.xxx + Primitive.LocalToWorld[1].xyz * LocalPosition.yyy + Primitive.LocalToWorld[2].xyz * LocalPosition.zzz;
-	float3 Translation = Primitive.LocalToWorld[3].xyz + View.PreViewTranslation.xyz;
-	return float4(RotatedPosition + Translation,1);
+	return RotatedPosition + Primitive.LocalToWorld[3].xyz;
+}
+
+float4 TransformLocalToTranslatedWorld(float3 LocalPosition)
+{
+	return float4(TransformLocalToWorld(LocalPosition) + View.PreViewTranslation.xyz,1);
 }
diff --git a/Engine/Source/Editor/ComponentVisualizers/Public/ComponentVisualizers.h b/Engine/Source/Editor/ComponentVisualizers/Public/ComponentVisualizers.h
--- a/Engine/Source/Editor/ComponentVisualizers/Public/ComponentVisualizers.h
+++ b/Engine/Source/Editor/ComponentVisualizers/Public/ComponentVisualizers.h
@@ -27,10 +27,12 @@ public:
 		return true;
 	}
 
+	/** Register a visualizer for a particular componen class */
+	COMPONENTVISUALIZERS_API void RegisterComponentVisualizer(FName ComponentClassName, TSharedPtr<FComponentVisualizer> Visualizer);
+
 private:
 
-	/** Register a visualizer for a particular componen class */
-	void RegisterComponentVisualizer(FName ComponentClassName, TSharedPtr<FComponentVisualizer> Visualizer);
+	
 
 	/** Array of component class names we have registered, so we know what to unregister afterwards */
 	TArray<FName> RegisteredComponentClassNames;
diff --git a/Engine/Source/Editor/EditorStyle/Private/SlateEditorStyle.cpp b/Engine/Source/Editor/EditorStyle/Private/SlateEditorStyle.cpp
--- a/Engine/Source/Editor/EditorStyle/Private/SlateEditorStyle.cpp
+++ b/Engine/Source/Editor/EditorStyle/Private/SlateEditorStyle.cpp
@@ -4006,6 +4006,10 @@ void FSlateEditorStyle::FStyle::SetupLevelEditorStyle()
 		Set( "LevelEditor.FoliageMode.Selected", new IMAGE_BRUSH( "Icons/icon_Mode_Foliage_selected_40x", Icon40x40 ) );
 		Set( "LevelEditor.FoliageMode.Selected.Small", new IMAGE_BRUSH( "Icons/icon_Mode_Foliage_selected_40x", Icon20x20 ) );
 
+		Set("LevelEditor.EcosystemMode", new IMAGE_BRUSH("Icons/icon_Mode_Ecosystem_40x", Icon40x40));
+		Set("LevelEditor.EcosystemMode.Small", new IMAGE_BRUSH("Icons/icon_Mode_Ecosystem_40x", Icon20x20));
+
+
 		Set( "LevelEditor.BspMode", new IMAGE_BRUSH( "Icons/icon_Mode_GeoEdit_40px", Icon40x40 ) );
 		Set( "LevelEditor.BspMode.Small", new IMAGE_BRUSH( "Icons/icon_Mode_GeoEdit_40px", Icon20x20 ) );
 		Set( "LevelEditor.BspMode.Selected", new IMAGE_BRUSH( "Icons/icon_Mode_GeoEdit-a_40px", Icon40x40 ) );
diff --git a/Engine/Source/Editor/FoliageEdit/Private/FoliageEdMode.cpp b/Engine/Source/Editor/FoliageEdit/Private/FoliageEdMode.cpp
--- a/Engine/Source/Editor/FoliageEdit/Private/FoliageEdMode.cpp
+++ b/Engine/Source/Editor/FoliageEdit/Private/FoliageEdMode.cpp
@@ -20,6 +20,7 @@
 // Classes
 #include "Foliage/InstancedFoliageActor.h"
 #include "Foliage/FoliageType.h"
+#include "LandscapeProxy.h"
 #include "LandscapeComponent.h"
 #include "LandscapeHeightfieldCollisionComponent.h"
 #include "LandscapeInfo.h"
@@ -29,7 +30,9 @@
 #include "Engine/CollisionProfile.h"
 #include "Components/ModelComponent.h"
 #include "EngineUtils.h"
+#include "LandscapeDataAccess.h"
 
+#define LOCTEXT_NAMESPACE "FoliageEdMode"
 #define FOLIAGE_SNAP_TRACE (10000.f)
 
 //
@@ -119,21 +122,6 @@ void FEdModeFoliage::Enter()
 		Toolkit = MakeShareable(new FFoliageEdModeToolkit);
 		Toolkit->Init(Owner->GetToolkitHost());
 	}
-
-	// Fixup any broken clusters
-	AInstancedFoliageActor* IFA = AInstancedFoliageActor::GetInstancedFoliageActorForCurrentLevel(GetWorld());
-	for (auto& MeshPair : IFA->FoliageMeshes)
-	{
-		FFoliageMeshInfo& MeshInfo = *MeshPair.Value;
-		for (FFoliageInstanceCluster& Cluster : MeshInfo.InstanceClusters)
-		{
-			if (Cluster.ClusterComponent == nullptr)
-			{
-				MeshInfo.ReallocateClusters(IFA, MeshPair.Key);
-				break;
-			}
-		}
-	}
 }
 
 /** FEdMode: Called when the mode is exited */
@@ -455,7 +443,7 @@ struct FPotentialInstance
 		, HitWeight(InHitWeight)
 	{}
 
-	bool PlaceInstance(UFoliageType* Settings, FFoliageInstance& Inst, UWorld* InWorld, AInstancedFoliageActor* InIFA)
+	bool PlaceInstance(UFoliageType* Settings, FFoliageInstance& Inst, UWorld* InWorld, AInstancedFoliageActor* InIFA, bool bSkipCollision = false)
 	{
 		if (Settings->UniformScale)
 		{
@@ -499,7 +487,7 @@ struct FPotentialInstance
 
 		Inst.Base = HitComponent;
 
-		return CheckCollisionWithWorld(Settings, Inst, HitNormal, HitLocation, InWorld, InIFA);
+		return bSkipCollision || CheckCollisionWithWorld(Settings, Inst, HitNormal, HitLocation, InWorld, InIFA);
 	}
 };
 
@@ -1621,7 +1609,7 @@ bool FEdModeFoliage::RemoveFoliageMesh(UFoliageType* Settings)
 	FFoliageMeshInfo* MeshInfo = IFA->FindMesh(Settings);
 	if (MeshInfo != nullptr)
 	{
-		int32 InstancesNum = MeshInfo->Instances.Num() - MeshInfo->FreeInstanceIndices.Num();
+		int32 InstancesNum = MeshInfo->Instances.Num();
 
 		bool bProceed = true;
 		if (Settings->GetStaticMesh() != nullptr && InstancesNum > 0)
@@ -1656,7 +1644,7 @@ void FEdModeFoliage::BakeFoliage(UFoliageType* Settings, bool bSelectedOnly)
 		TArray<int32> InstancesToConvert;
 		if (bSelectedOnly)
 		{
-			InstancesToConvert = MeshInfo->SelectedIndices;
+			InstancesToConvert = MeshInfo->SelectedIndices.Array();
 		}
 		else
 		{
@@ -1771,7 +1759,7 @@ bool FEdModeFoliage::ReplaceStaticMesh(UFoliageType* OldSettings, UStaticMesh* N
 
 	if (OldMeshInfo != nullptr && OldSettings->GetStaticMesh() != NewStaticMesh)
 	{
-		int32 InstancesNum = OldMeshInfo->Instances.Num() - OldMeshInfo->FreeInstanceIndices.Num();
+		int32 InstancesNum = OldMeshInfo->Instances.Num();
 
 		// Look for the new mesh in the mesh list, and either create a new mesh or merge the instances.
 		UFoliageType* NewSettings = nullptr;
@@ -1815,12 +1803,9 @@ bool FEdModeFoliage::ReplaceStaticMesh(UFoliageType* OldSettings, UStaticMesh* N
 			// copy instances from old to new.
 			for (FFoliageInstance& Instance : OldMeshInfo->Instances)
 			{
-				if (Instance.ClusterIndex != -1)
-				{
 					NewMeshInfo->AddInstance(IFA, NewSettings, Instance);
 				}
 			}
-		}
 
 		// Remove the old mesh.
 		IFA->RemoveMesh(OldSettings);
@@ -1904,7 +1889,7 @@ bool FEdModeFoliage::InputKey(FEditorViewportClient* ViewportClient, FViewport*
 					FFoliageMeshInfo& Mesh = *MeshPair.Value;
 					if (Mesh.SelectedIndices.Num() > 0)
 					{
-						TArray<int32> InstancesToDelete = Mesh.SelectedIndices;
+						TArray<int32> InstancesToDelete = Mesh.SelectedIndices.Array();
 						Mesh.RemoveInstances(IFA, InstancesToDelete);
 					}
 				}
@@ -1922,10 +1907,12 @@ bool FEdModeFoliage::InputKey(FEditorViewportClient* ViewportClient, FViewport*
 				{
 					FFoliageMeshInfo& Mesh = *MeshPair.Value;
 
-					Mesh.PreMoveInstances(IFA, Mesh.SelectedIndices);
+					TArray<int32> SelectedIndices = Mesh.SelectedIndices.Array();
+
+					Mesh.PreMoveInstances(IFA, SelectedIndices);
 
 					UWorld* World = ViewportClient->GetWorld();
-					for (int32 SelectedInstanceIdx : Mesh.SelectedIndices)
+					for (int32 SelectedInstanceIdx : SelectedIndices)
 					{
 						FFoliageInstance& Instance = Mesh.Instances[SelectedInstanceIdx];
 
@@ -1959,7 +1946,7 @@ bool FEdModeFoliage::InputKey(FEditorViewportClient* ViewportClient, FViewport*
 						bMovedInstance = true;
 					}
 
-					Mesh.PostMoveInstances(IFA, Mesh.SelectedIndices);
+					Mesh.PostMoveInstances(IFA, SelectedIndices);
 				}
 
 				if (bMovedInstance)
@@ -2075,6 +2062,7 @@ bool FEdModeFoliage::HandleClick(FEditorViewportClient* InViewportClient, HHitPr
 		return true;
 	}
 
+
 	return FEdMode::HandleClick(InViewportClient, HitProxy, Click);
 }
 
@@ -2126,16 +2114,18 @@ bool FEdModeFoliage::InputDelta(FEditorViewportClient* InViewportClient, FViewpo
 		for (auto& MeshPair : IFA->FoliageMeshes)
 		{
 			FFoliageMeshInfo& MeshInfo = *MeshPair.Value;
-			bFoundSelection |= MeshInfo.SelectedIndices.Num() > 0;
+			TArray<int32> SelectedIndices = MeshInfo.SelectedIndices.Array();
+
+			bFoundSelection |= SelectedIndices.Num() > 0;
 
 			if (bAltDown && bCanAltDrag && (InViewportClient->GetCurrentWidgetAxis() & EAxisList::XYZ))
 			{
-				MeshInfo.DuplicateInstances(IFA, MeshPair.Key, MeshInfo.SelectedIndices);
+				MeshInfo.DuplicateInstances(IFA, MeshPair.Key, SelectedIndices);
 			}
 
-			MeshInfo.PreMoveInstances(IFA, MeshInfo.SelectedIndices);
+			MeshInfo.PreMoveInstances(IFA, SelectedIndices);
 
-			for (int32 SelectedInstanceIdx : MeshInfo.SelectedIndices)
+			for (int32 SelectedInstanceIdx : SelectedIndices)
 			{
 				FFoliageInstance& Instance = MeshInfo.Instances[SelectedInstanceIdx];
 				Instance.Location += InDrag;
@@ -2144,7 +2134,7 @@ bool FEdModeFoliage::InputDelta(FEditorViewportClient* InViewportClient, FViewpo
 				Instance.DrawScale3D += InScale;
 			}
 
-			MeshInfo.PostMoveInstances(IFA, MeshInfo.SelectedIndices);
+			MeshInfo.PostMoveInstances(IFA, SelectedIndices);
 		}
 
 		// Only allow alt-drag on first InputDelta
@@ -2227,3 +2217,5 @@ void FFoliageUISettings::Save()
 	GConfig->SetBool(TEXT("FoliageEdit"), TEXT("bFilterBSP"), bFilterBSP, GEditorUserSettingsIni);
 	GConfig->SetBool(TEXT("FoliageEdit"), TEXT("bFilterTranslucent"), bFilterTranslucent, GEditorUserSettingsIni);
 }
+
+#undef LOCTEXT_NAMESPACE
\ No newline at end of file
diff --git a/Engine/Source/Editor/FoliageEdit/Private/SFoliageEditMeshDisplayItem.cpp b/Engine/Source/Editor/FoliageEdit/Private/SFoliageEditMeshDisplayItem.cpp
--- a/Engine/Source/Editor/FoliageEdit/Private/SFoliageEditMeshDisplayItem.cpp
+++ b/Engine/Source/Editor/FoliageEdit/Private/SFoliageEditMeshDisplayItem.cpp
@@ -1655,14 +1655,6 @@ void SFoliageEditMeshDisplayItem::Construct(const FArguments& InArgs)
 								SNew(STextBlock)
 								.Text(this, &SFoliageEditMeshDisplayItem::GetInstanceCountString)
 							]
-
-							+ SVerticalBox::Slot()
-							.AutoHeight()
-							.VAlign(VAlign_Center)
-							[
-								SNew(STextBlock)
-								.Text(this, &SFoliageEditMeshDisplayItem::GetInstanceClusterCountString)
-							]
 						]
 
 						+ SHorizontalBox::Slot()
@@ -1951,11 +1943,6 @@ FText SFoliageEditMeshDisplayItem::GetInstanceCountString() const
 	return FText::Format(LOCTEXT("InstanceCount_Value", "Instance Count: {0}"), FText::AsNumber(FoliageMeshUIInfo->MeshInfo ? FoliageMeshUIInfo->MeshInfo->GetInstanceCount() : 0 ));
 }
 
-FText SFoliageEditMeshDisplayItem::GetInstanceClusterCountString() const
-{
-	return FText::Format(LOCTEXT("ClusterCount_Value", "Cluster Count: {0}"), FText::AsNumber(FoliageMeshUIInfo->MeshInfo ? FoliageMeshUIInfo->MeshInfo->InstanceClusters.Num() : 0));
-}
-
 EVisibility SFoliageEditMeshDisplayItem::IsReapplySettingsVisible() const
 {
 	FEdModeFoliage* Mode = (FEdModeFoliage*)GLevelEditorModeTools().GetActiveMode(FBuiltinEditorModes::EM_Foliage);
diff --git a/Engine/Source/Editor/FoliageEdit/Private/SFoliageEditMeshDisplayItem.h b/Engine/Source/Editor/FoliageEdit/Private/SFoliageEditMeshDisplayItem.h
--- a/Engine/Source/Editor/FoliageEdit/Private/SFoliageEditMeshDisplayItem.h
+++ b/Engine/Source/Editor/FoliageEdit/Private/SFoliageEditMeshDisplayItem.h
@@ -122,8 +122,6 @@ private:
 	/** Cluster Accessors */
 	FText GetInstanceCountString() const;
 
-	FText GetInstanceClusterCountString() const;
-
 	/** General functions for delegates. */
 	EVisibility IsReapplySettingsVisible() const;
 	EVisibility IsReapplySettingsVisible_Dummy() const;
@@ -147,6 +145,21 @@ private:
 
 	float GetDensityReapply() const;
 
+	int32 GetSeedsPerComponent() const;
+	void OnSeedsPerComponentChanged(int32 InValue);
+
+	int32 GetSimulationSteps() const;
+	void OnSimulationStepsChanged(int32 InValue);
+
+	float GetMaxRadius() const;
+	void OnMaxRadiusChanged(float InValue);
+
+	float GetSpreadRadius() const;
+	void OnSpreadRadiusChanged(float InValue);
+
+	float GetGrowthRadius() const;
+	void OnGrowthRadiusChanged(float InValue);
+
 	/** Radius */
 	void OnRadiusReapply(ESlateCheckBoxState::Type InState);
 
@@ -294,6 +307,19 @@ private:
 
 	float GetHeightMax() const;
 
+	float GetShadeRadius() const;
+
+	void OnShadeRadiusChanged(float InValue);
+
+	float GetShadeAffinity() const;
+
+	void OnShadeAffinityChanged(float InValue);
+
+
+	float GetSimulationPriority() const;
+
+	void OnSimulationPriorityChanged(float InValue);
+
 	/** Landscape Layer */
 	void OnLandscapeLayerReapply(ESlateCheckBoxState::Type InState);
 
diff --git a/Engine/Source/Editor/LevelEditor/Private/LevelEditor.cpp b/Engine/Source/Editor/LevelEditor/Private/LevelEditor.cpp
--- a/Engine/Source/Editor/LevelEditor/Private/LevelEditor.cpp
+++ b/Engine/Source/Editor/LevelEditor/Private/LevelEditor.cpp
@@ -583,6 +583,12 @@ void FLevelEditorModule::BindGlobalLevelEditorCommands()
 		FCanExecuteAction::CreateStatic( &FLevelEditorActionCallbacks::CanViewReferences )
 		);
 
+	const FVector* NullVector = nullptr;
+	ActionList.MapAction(
+		Commands.GoHere,
+		FExecuteAction::CreateStatic( &FLevelEditorActionCallbacks::GoHere_Clicked, NullVector )
+		);
+
 	ActionList.MapAction( 
 		Commands.SnapCameraToActor, 
 		FExecuteAction::CreateStatic( &FLevelEditorActionCallbacks::ExecuteExecCommand, FString( TEXT("CAMERA SNAP") ) )
diff --git a/Engine/Source/Editor/LevelEditor/Private/LevelEditorActions.cpp b/Engine/Source/Editor/LevelEditor/Private/LevelEditorActions.cpp
--- a/Engine/Source/Editor/LevelEditor/Private/LevelEditorActions.cpp
+++ b/Engine/Source/Editor/LevelEditor/Private/LevelEditorActions.cpp
@@ -1096,6 +1096,54 @@ void FLevelEditorActionCallbacks::CreateNewOutlinerFolder_Clicked()
 	FActorFolders::Get().CreateFolderContainingSelection(*GetWorld(), NewFolderName);
 }
 
+void FLevelEditorActionCallbacks::GoHere_Clicked( const FVector* Point )
+{
+	if( GCurrentLevelEditingViewportClient )
+	{
+		FVector ZoomToPoint;
+		if( !Point )
+		{
+			FSceneViewFamilyContext ViewFamily(FSceneViewFamily::ConstructionValues(
+				GCurrentLevelEditingViewportClient->Viewport,
+				GCurrentLevelEditingViewportClient->GetWorld()->Scene,
+				GCurrentLevelEditingViewportClient->EngineShowFlags)
+				.SetRealtimeUpdate(true));
+
+
+			FSceneView* SceneView = GCurrentLevelEditingViewportClient->CalcSceneView(&ViewFamily);
+
+			if(SceneView)
+			{
+				FIntPoint MousePosition;
+				FVector WorldOrigin;
+				FVector WorldDirection;
+				GCurrentLevelEditingViewportClient->Viewport->GetMousePos(MousePosition);
+
+				SceneView->DeprojectFVector2D(MousePosition, WorldOrigin, WorldDirection);
+
+				FHitResult HitResult;
+
+				static FName FocusOnPoint = FName(TEXT("FocusOnPoint"));
+				FCollisionQueryParams LineParams(FocusOnPoint, true);
+
+				if(GCurrentLevelEditingViewportClient->GetWorld()->LineTraceSingle(HitResult, WorldOrigin, WorldOrigin + WorldDirection * HALF_WORLD_MAX, LineParams, FCollisionObjectQueryParams(ECC_WorldStatic)))
+				{
+					ZoomToPoint = HitResult.ImpactPoint;
+				}
+			}
+		}
+		else
+		{
+			ZoomToPoint = *Point;
+		}
+
+		const float PushOutSize = 500;
+		FBox BoundingBox(ZoomToPoint-PushOutSize, ZoomToPoint+PushOutSize);
+
+		GCurrentLevelEditingViewportClient->FocusViewportOnBox(BoundingBox);
+	}
+}
+
 bool FLevelEditorActionCallbacks::LockActorMovement_IsChecked()
 {
 	return GEditor->HasLockedActors();
@@ -2612,7 +2660,10 @@ void FLevelEditorCommands::RegisterCommands()
 	UI_COMMAND( EditAsset, "Edit Asset", "Edits the asset associated with the selected actor", EUserInterfaceActionType::Button, FInputGesture( EKeys::E, EModifierKey::Control ) );
 	UI_COMMAND( EditAssetNoConfirmMultiple, "Edit Asset", "Edits the asset associated with the selected actor", EUserInterfaceActionType::Button, FInputGesture( EKeys::E, EModifierKey::Control | EModifierKey::Shift ) );
 
+	UI_COMMAND( GoHere, "Go Here", "Moves the camera to the current mouse position", EUserInterfaceActionType::Button, FInputGesture() );
+
 	UI_COMMAND( SnapCameraToActor, "Snap View to Actor", "Snaps the view to the selected actors", EUserInterfaceActionType::Button, FInputGesture() );
+
 	UI_COMMAND( GoToCodeForActor, "Go to C++ Code for Actor", "Opens a code editing IDE and navigates to the source file associated with the seleced actor", EUserInterfaceActionType::Button, FInputGesture() );
 	UI_COMMAND( GoToDocsForActor, "Go to Documentation for Actor", "Opens documentation for the Actor in the default web browser", EUserInterfaceActionType::Button, FInputGesture() );
 
diff --git a/Engine/Source/Editor/LevelEditor/Private/LevelEditorContextMenu.cpp b/Engine/Source/Editor/LevelEditor/Private/LevelEditorContextMenu.cpp
--- a/Engine/Source/Editor/LevelEditor/Private/LevelEditorContextMenu.cpp
+++ b/Engine/Source/Editor/LevelEditor/Private/LevelEditorContextMenu.cpp
@@ -233,6 +233,14 @@ void FLevelEditorContextMenu::FillMenu( FMenuBuilder& MenuBuilder, TWeakPtr<SLev
 	MenuBuilder.BeginSection( "ActorControl", LOCTEXT("ActorHeading", "Actor") );
 	{
 		MenuBuilder.AddMenuEntry( FEditorViewportCommands::Get().FocusViewportToSelection );
+
+		
+		const FVector* ClickLocation = &GEditor->ClickLocation;
+
+		FUIAction GoHereAction;
+		GoHereAction.ExecuteAction = FExecuteAction::CreateStatic( &FLevelEditorActionCallbacks::GoHere_Clicked, ClickLocation );
+
+		MenuBuilder.AddMenuEntry( FLevelEditorCommands::Get().GoHere );
 		MenuBuilder.AddMenuEntry( FLevelEditorCommands::Get().SnapCameraToActor );
 	}
 	MenuBuilder.EndSection();
diff --git a/Engine/Source/Editor/LevelEditor/Public/LevelEditorActions.h b/Engine/Source/Editor/LevelEditor/Public/LevelEditorActions.h
--- a/Engine/Source/Editor/LevelEditor/Public/LevelEditorActions.h
+++ b/Engine/Source/Editor/LevelEditor/Public/LevelEditorActions.h
@@ -110,6 +110,9 @@ public:
 	/** Edits associated asset(s) */
 	TSharedPtr< FUICommandInfo > EditAssetNoConfirmMultiple;
 
+	/** Moves the camera to the current mouse position */
+	TSharedPtr< FUICommandInfo > GoHere;
+
 	/** Snaps the camera to the selected actors. */
 	TSharedPtr< FUICommandInfo > SnapCameraToActor;
 
@@ -780,6 +783,12 @@ public:
 	/** Called when create new outliner folder is clicked */
 	static void CreateNewOutlinerFolder_Clicked();
 
+	/** Called when the go here command is clicked 
+	 * 
+	 * @param Point	- Specified point to go to.  If null, a point will be calculated from current mouse position
+	 */
+	static void GoHere_Clicked( const FVector* Point );
+
 	/** Called when 'Go to Code for Actor' is clicked */
 	static void GoToCodeForActor_Clicked();
 
diff --git a/Engine/Source/Editor/UnrealEd/UnrealEd.Build.cs b/Engine/Source/Editor/UnrealEd/UnrealEd.Build.cs
--- a/Engine/Source/Editor/UnrealEd/UnrealEd.Build.cs
+++ b/Engine/Source/Editor/UnrealEd/UnrealEd.Build.cs
@@ -188,7 +188,7 @@ public class UnrealEd : ModuleRules
 				"ReferenceViewer",
 				"EditorLiveStreaming",
 				"HotReload",
-                "IOSPlatformEditor"
+                "IOSPlatformEditor",
 			}
 		);
 
diff --git a/Engine/Source/Editor/WorldBrowser/Private/LevelCollectionModel.cpp b/Engine/Source/Editor/WorldBrowser/Private/LevelCollectionModel.cpp
--- a/Engine/Source/Editor/WorldBrowser/Private/LevelCollectionModel.cpp
+++ b/Engine/Source/Editor/WorldBrowser/Private/LevelCollectionModel.cpp
@@ -686,9 +686,9 @@ void FLevelCollectionModel::CustomizeFileMainMenu(FMenuBuilder& InMenuBuilder) c
 	}
 }
 
-FMatrix FLevelCollectionModel::GetObserverViewMatrix() const
+bool FLevelCollectionModel::GetObserverView(FVector& Location, FRotator& Rotation) const
 {
-	return FMatrix::Identity;
+	return false;
 }
 
 bool FLevelCollectionModel::CompareLevelsZOrder(TSharedPtr<FLevelModel> InA, TSharedPtr<FLevelModel> InB) const
diff --git a/Engine/Source/Editor/WorldBrowser/Private/LevelCollectionModel.h b/Engine/Source/Editor/WorldBrowser/Private/LevelCollectionModel.h
--- a/Engine/Source/Editor/WorldBrowser/Private/LevelCollectionModel.h
+++ b/Engine/Source/Editor/WorldBrowser/Private/LevelCollectionModel.h
@@ -130,8 +130,8 @@ public:
 	/**	Customize 'File' section in main menu  */
 	virtual void CustomizeFileMainMenu(FMenuBuilder& InMenuBuilder) const;
 
-	/**	@return	Observer view matrix in the world, usually camera view */
-	virtual FMatrix GetObserverViewMatrix() const;
+	/**	@return	Observer view in the world, usually camera view */
+	virtual bool GetObserverView(FVector& Location, FRotator& Rotation) const;
 
 	/**	Compares 2 levels by Z order */
 	virtual bool CompareLevelsZOrder(TSharedPtr<FLevelModel> InA, TSharedPtr<FLevelModel> InB) const;
diff --git a/Engine/Source/Editor/WorldBrowser/Private/Tiles/SWorldComposition.cpp b/Engine/Source/Editor/WorldBrowser/Private/Tiles/SWorldComposition.cpp
--- a/Engine/Source/Editor/WorldBrowser/Private/Tiles/SWorldComposition.cpp
+++ b/Engine/Source/Editor/WorldBrowser/Private/Tiles/SWorldComposition.cpp
@@ -349,16 +349,17 @@ public:
 				FEditorStyle::GetBrush(TEXT("Graph.PlayInEditor")),
 				MyClippingRect
 				);
+		}
 
-			// Draw a current camera position
+		// Draw observer location
+		{
+			FVector ObserverPosition;
+			FRotator ObserverRotation;
+			if (WorldModel->GetObserverView(ObserverPosition, ObserverRotation))
 			{
-				const FSlateBrush* CameraImage = FEditorStyle::GetBrush(TEXT("WorldBrowser.SimulationViewPositon"));
-				FMatrix ObserverViewToWorld = WorldModel->GetObserverViewMatrix().Inverse();
-				FVector ObserverPosition	= ObserverViewToWorld.GetOrigin();
-				FRotator ObserverRotation	= ObserverViewToWorld.Rotator();
-								
 				FVector2D ObserverPositionScreen = GraphCoordToPanelCoord(FVector2D(ObserverPosition.X, ObserverPosition.Y));
-
+				const FSlateBrush* CameraImage = FEditorStyle::GetBrush(TEXT("WorldBrowser.SimulationViewPositon"));
+	
 				FPaintGeometry PaintGeometry = AllottedGeometry.ToPaintGeometry(
 					ObserverPositionScreen - CameraImage->ImageSize*0.5f, 
 					CameraImage->ImageSize
@@ -371,7 +372,7 @@ public:
 					CameraImage,
 					MyClippingRect,
 					ESlateDrawEffect::None,
-					FMath::DegreesToRadians(ObserverRotation.Yaw - 90.f), // 0 degrees pointing along +Y axis, arrow on the image pointing along +X axis
+					FMath::DegreesToRadians(ObserverRotation.Yaw),
 					CameraImage->ImageSize*0.5f,
 					FSlateDrawElement::RelativeToElement
 					);
diff --git a/Engine/Source/Editor/WorldBrowser/Private/Tiles/WorldTileCollectionModel.cpp b/Engine/Source/Editor/WorldBrowser/Private/Tiles/WorldTileCollectionModel.cpp
--- a/Engine/Source/Editor/WorldBrowser/Private/Tiles/WorldTileCollectionModel.cpp
+++ b/Engine/Source/Editor/WorldBrowser/Private/Tiles/WorldTileCollectionModel.cpp
@@ -571,17 +571,43 @@ void FWorldTileCollectionModel::CustomizeFileMainMenu(FMenuBuilder& InMenuBuilde
 	InMenuBuilder.EndSection();
 }
 
-FMatrix FWorldTileCollectionModel::GetObserverViewMatrix() const
+bool FWorldTileCollectionModel::GetObserverView(FVector& Location, FRotator& Rotation) const
 {
-	UWorld*	SimulationWorld = GetSimulationWorld();
-	if (SimulationWorld && SimulationWorld->WorldComposition)
+	const UEditorEngine* EditorEngine = Editor.Get();
+
+	if (IsSimulating())
 	{
-		return SimulationWorld->WorldComposition->LastWorldToViewMatrix;
+		if (EditorEngine->bIsSimulatingInEditor && GCurrentLevelEditingViewportClient->IsSimulateInEditorViewport())
+		{
+			Rotation = GCurrentLevelEditingViewportClient->GetViewRotation();
+			Location = GCurrentLevelEditingViewportClient->GetViewLocation();
+			return true;
+		}
+		else
+		{
+			UWorld* SimulationWorld = GetSimulationWorld();
+			for (FConstPlayerControllerIterator Iterator = SimulationWorld->GetPlayerControllerIterator(); Iterator; ++Iterator)
+			{
+				APlayerController* PlayerActor = *Iterator;
+				PlayerActor->GetPlayerViewPoint(Location, Rotation);
+				return true;
+			}
+		}
 	}
 	else
 	{
-		return FMatrix::Identity;
+		for (const FLevelEditorViewportClient* ViewportClient : EditorEngine->LevelViewportClients)
+		{
+			if (ViewportClient && ViewportClient->IsPerspective())
+			{
+				Rotation = ViewportClient->GetViewRotation();
+				Location = ViewportClient->GetViewLocation();
+				return true;
+			}
+		}
 	}
+	
+	return false;
 }
 
 static double Area(FVector2D InRect)
diff --git a/Engine/Source/Editor/WorldBrowser/Private/Tiles/WorldTileCollectionModel.h b/Engine/Source/Editor/WorldBrowser/Private/Tiles/WorldTileCollectionModel.h
--- a/Engine/Source/Editor/WorldBrowser/Private/Tiles/WorldTileCollectionModel.h
+++ b/Engine/Source/Editor/WorldBrowser/Private/Tiles/WorldTileCollectionModel.h
@@ -34,7 +34,7 @@ public:
 	virtual bool PassesAllFilters(TSharedPtr<FLevelModel> InLevelModel) const override;
 	virtual void BuildHierarchyMenu(FMenuBuilder& InMenuBuilder) const override;
 	virtual void CustomizeFileMainMenu(FMenuBuilder& InMenuBuilder) const override;
-	virtual FMatrix GetObserverViewMatrix() const override;
+	virtual bool GetObserverView(FVector& Location, FRotator& Rotation) const override;
 	virtual bool CompareLevelsZOrder(TSharedPtr<FLevelModel> InA, TSharedPtr<FLevelModel> InB) const override;
 	virtual void RegisterDetailsCustomization(class FPropertyEditorModule& PropertyModule, TSharedPtr<class IDetailsView> InDetailsView)  override;
 	virtual void UnregisterDetailsCustomization(class FPropertyEditorModule& PropertyModule, TSharedPtr<class IDetailsView> InDetailsView) override;
diff --git a/Engine/Source/Runtime/Core/Public/Containers/BitArray.h b/Engine/Source/Runtime/Core/Public/Containers/BitArray.h
--- a/Engine/Source/Runtime/Core/Public/Containers/BitArray.h
+++ b/Engine/Source/Runtime/Core/Public/Containers/BitArray.h
@@ -393,9 +393,9 @@ public:
 				int32 ToIndex = BaseIndex + Index;
 				FBitReference To(GetData()[ToIndex / NumBitsPerDWORD],1 << (ToIndex & (NumBitsPerDWORD - 1)));
 
-				To = From;
+				To = (bool)From;
 #else
-				(*this)[BaseIndex + Index] = (*this)[NumBits - NumBitsToRemove + Index];
+				(*this)[BaseIndex + Index] = (bool)(*this)[NumBits - NumBitsToRemove + Index];
 #endif
 			}
 		}
diff --git a/Engine/Source/Runtime/Core/Public/Math/UnrealMatrix.h b/Engine/Source/Runtime/Core/Public/Math/UnrealMatrix.h
--- a/Engine/Source/Runtime/Core/Public/Math/UnrealMatrix.h
+++ b/Engine/Source/Runtime/Core/Public/Math/UnrealMatrix.h
@@ -602,7 +602,7 @@ inline void FMatrix::SetAxes(FVector* Axis0 /*= NULL*/, FVector* Axis1 /*= NULL*
 
 inline FVector FMatrix::GetColumn(int32 i) const
 {
-	checkSlow(i >= 0 && i <= 2);
+	checkSlow(i >= 0 && i <= 3);
 	return FVector(M[0][i], M[1][i], M[2][i]);
 }
 
diff --git a/Engine/Source/Runtime/EmptyRHI/Private/EmptyCommands.cpp b/Engine/Source/Runtime/EmptyRHI/Private/EmptyCommands.cpp
--- a/Engine/Source/Runtime/EmptyRHI/Private/EmptyCommands.cpp
+++ b/Engine/Source/Runtime/EmptyRHI/Private/EmptyCommands.cpp
@@ -304,7 +304,7 @@ void FEmptyDynamicRHI::RHIDrawPrimitiveIndirect(uint32 PrimitiveType, FVertexBuf
 }
 
 
-void FEmptyDynamicRHI::RHIDrawIndexedPrimitive(FIndexBufferRHIParamRef IndexBufferRHI, uint32 PrimitiveType, int32 BaseVertexIndex, uint32 MinIndex,
+void FEmptyDynamicRHI::RHIDrawIndexedPrimitive(FIndexBufferRHIParamRef IndexBufferRHI, uint32 PrimitiveType, int32 BaseVertexIndex, uint32 FirstInstance,
 	uint32 NumVertices, uint32 StartIndex, uint32 NumPrimitives, uint32 NumInstances)
 {
 	DYNAMIC_CAST_EMPTYRESOURCE(IndexBuffer,IndexBuffer);
diff --git a/Engine/Source/Runtime/Engine/Classes/AI/Navigation/NavigationSystem.h b/Engine/Source/Runtime/Engine/Classes/AI/Navigation/NavigationSystem.h
--- a/Engine/Source/Runtime/Engine/Classes/AI/Navigation/NavigationSystem.h
+++ b/Engine/Source/Runtime/Engine/Classes/AI/Navigation/NavigationSystem.h
@@ -515,6 +515,9 @@ public:
 	/** force updating parent node and all its children */
 	void UpdateNavOctreeParentChain(UObject* ElementOwner);
 
+	/** update component bounds in navigation octree and mark only specified area as dirty, doesn't re-export component geometry */
+	bool UpdateNavOctreeElementBounds(UActorComponent* Comp, const FBox& NewBounds, const FBox& DirtyArea);
+
 	//----------------------------------------------------------------------//
 	// Custom navigation links
 	//----------------------------------------------------------------------//
diff --git a/Engine/Source/Runtime/Engine/Classes/AI/Navigation/NavigationTypes.h b/Engine/Source/Runtime/Engine/Classes/AI/Navigation/NavigationTypes.h
--- a/Engine/Source/Runtime/Engine/Classes/AI/Navigation/NavigationTypes.h
+++ b/Engine/Source/Runtime/Engine/Classes/AI/Navigation/NavigationTypes.h
@@ -169,6 +169,11 @@ struct FNavigationDirtyElement
 	}
 };
 
+//
+// Used to gather per instance transforms in a specific area
+//
+DECLARE_DELEGATE_TwoParams(FNavDataPerInstanceTransformDelegate, const FBox&, TArray<FTransform>&);
+
 //////////////////////////////////////////////////////////////////////////
 // Path
 
diff --git a/Engine/Source/Runtime/Engine/Classes/AI/Navigation/RecastNavMesh.h b/Engine/Source/Runtime/Engine/Classes/AI/Navigation/RecastNavMesh.h
--- a/Engine/Source/Runtime/Engine/Classes/AI/Navigation/RecastNavMesh.h
+++ b/Engine/Source/Runtime/Engine/Classes/AI/Navigation/RecastNavMesh.h
@@ -40,6 +40,7 @@ class UNavigationSystem;
 class URecastNavMeshDataChunk;
 struct FAreaNavModifier;
 struct FNavigationQueryFilter;
+struct FRecastAreaNavModifierElement;
 
 UENUM()
 namespace ERecastPartitioning
@@ -752,7 +753,7 @@ public:
 	void SetDefaultForbiddenFlags(uint16 ForbiddenAreaFlags);
 
 	/** Area sort function */
-	virtual void SortAreasForGenerator(TArray<FAreaNavModifier>& Areas) const;
+	virtual void SortAreasForGenerator(TArray<FRecastAreaNavModifierElement>& Areas) const;
 
 	//----------------------------------------------------------------------//
 	// Custom navigation links
@@ -853,7 +854,7 @@ public:
 
 	/** Runs A* pathfinding on navmesh and collect data for every step */
 	int32 DebugPathfinding(const FPathFindingQuery& Query, TArray<FRecastDebugPathfindingStep>& Steps);
-	
+
 	static const FRecastQueryFilter* GetNamedFilter(ERecastNamedFilter::Type FilterType);
 	FORCEINLINE static FNavPolyFlags GetNavLinkFlag() { return NavLinkFlag; }
 	
diff --git a/Engine/Source/Runtime/Engine/Classes/Components/HierarchicalInstancedStaticMeshComponent.h b/Engine/Source/Runtime/Engine/Classes/Components/HierarchicalInstancedStaticMeshComponent.h
new file mode 100644
--- /dev/null
+++ b/Engine/Source/Runtime/Engine/Classes/Components/HierarchicalInstancedStaticMeshComponent.h
@@ -0,0 +1,108 @@
+// Copyright 1998-2014 Epic Games, Inc. All Rights Reserved.
+
+#pragma once
+
+#include "Components/InstancedStaticMeshComponent.h"
+
+#include "HierarchicalInstancedStaticMeshComponent.generated.h"
+
+class FClusterBuilder;
+
+USTRUCT()
+struct FClusterNode
+{
+	GENERATED_USTRUCT_BODY()
+
+	UPROPERTY()
+	FVector BoundMin;
+	UPROPERTY()
+	int32 FirstChild;
+	UPROPERTY()
+	FVector BoundMax;
+	UPROPERTY()
+	int32 LastChild;
+	UPROPERTY()
+	int32 FirstInstance;
+	UPROPERTY()
+	int32 LastInstance;
+
+	FClusterNode()
+		: BoundMin(MAX_flt, MAX_flt, MAX_flt)
+		, FirstChild(-1)
+		, BoundMax(MIN_flt, MIN_flt, MIN_flt)
+		, LastChild(-1)
+		, FirstInstance(-1)
+		, LastInstance(-1)
+	{
+	}
+
+	friend FArchive& operator<<(FArchive& Ar, FClusterNode& NodeData)
+	{
+		// @warning BulkSerialize: FClusterNode is serialized as memory dump
+		// See TArray::BulkSerialize for detailed description of implied limitations.
+		Ar << NodeData.BoundMin;
+		Ar << NodeData.FirstChild;
+		Ar << NodeData.BoundMax;
+		Ar << NodeData.LastChild;
+		Ar << NodeData.FirstInstance;
+		Ar << NodeData.LastInstance;
+		return Ar;
+	}
+};
+
+
+UCLASS(ClassGroup=Rendering, meta=(BlueprintSpawnableComponent))
+class ENGINE_API UHierarchicalInstancedStaticMeshComponent : public UInstancedStaticMeshComponent
+{
+	GENERATED_UCLASS_BODY()
+
+	UPROPERTY(Transient, DuplicateTransient)
+	TArray<FClusterNode> ClusterTree;
+	
+	// Table for remaping instances from cluster tree to PerInstanceSMData order
+	UPROPERTY()
+	TArray<int32> SortedInstances;
+
+	// The number of instances in the ClusterTree. Subsequent instances will always be rendered.
+	UPROPERTY()
+	int32 NumBuiltInstances;
+
+	// Bounding box of any unbuilt instances
+	UPROPERTY()
+	FBox UnbuiltInstanceBounds;
+
+	bool bIsAsyncBuilding;
+	bool bConcurrentRemoval;
+
+	// Apply the results of the async build
+	void ApplyBuildTreeAsync(ENamedThreads::Type CurrentThread, const FGraphEventRef& MyCompletionGraphEvent, TSharedRef<FClusterBuilder, ESPMode::ThreadSafe> Builder, double StartTime);
+
+public:
+	virtual void Serialize(FArchive& Ar) override;
+	virtual FPrimitiveSceneProxy* CreateSceneProxy() override;
+	virtual FBoxSphereBounds CalcBounds(const FTransform& BoundTransform) const override;
+
+	// UInstancedStaticMesh interface
+	virtual int32 AddInstance(const FTransform& InstanceTransform) override;
+	virtual bool RemoveInstance(int32 InstanceIndex) override;
+	virtual bool UpdateInstanceTransform(int32 InstanceIndex, const FTransform& NewInstanceTransform, bool bWorldSpace) override;
+	virtual void ClearInstances() override;
+
+	void BuildTree();
+	void BuildTreeAsync();
+	bool IsAsyncBuilding() const { return bIsAsyncBuilding; }
+	bool IsTreeFullyBuilt() const { return NumBuiltInstances == PerInstanceSMData.Num() && RemovedInstances.Num() == 0; }
+
+protected:
+	virtual void GetNavigationPerInstanceTransforms(const FBox& AreaBox, TArray<FTransform>& InstanceData) const override;
+	virtual void PartialNavigationUpdate(int32 InstanceIdx) override;
+	void FlushAccumulatedNavigationUpdates();
+	mutable FBox AccumulatedNavigationDirtyArea;
+
+protected:
+	friend FStaticLightingTextureMapping_InstancedStaticMesh;
+	friend FInstancedLightMap2D;
+	friend FInstancedShadowMap2D;
+	friend class FClusterBuilder;
+};
+
diff --git a/Engine/Source/Runtime/Engine/Classes/Components/InstancedStaticMeshComponent.h b/Engine/Source/Runtime/Engine/Classes/Components/InstancedStaticMeshComponent.h
--- a/Engine/Source/Runtime/Engine/Classes/Components/InstancedStaticMeshComponent.h
+++ b/Engine/Source/Runtime/Engine/Classes/Components/InstancedStaticMeshComponent.h
@@ -80,13 +80,21 @@ class ENGINE_API UInstancedStaticMeshComponent : public UStaticMeshComponent
 	UPROPERTY(EditAnywhere, BlueprintReadOnly, Category=Culling)
 	int32 InstanceEndCullDistance;
 
+	/** Mapping from PerInstanceSMData order to instance render buffer order. If empty, the PerInstanceSMData order is used. */
+	UPROPERTY()
+	TArray<int32> InstanceReorderTable;
+
+	// The render indices of any removed items we should not render.
+	UPROPERTY()
+	TArray<int32> RemovedInstances;
+
 	/** Add an instance to this component. Transform is given in local space of this component.  */
 	UFUNCTION(BlueprintCallable, Category="Components|InstancedStaticMesh")
-	void AddInstance(const FTransform& InstanceTransform);
+	virtual int32 AddInstance(const FTransform& InstanceTransform);
 
 	/** Add an instance to this component. Transform is given in world space. */
 	UFUNCTION(BlueprintCallable, Category = "Components|InstancedStaticMesh")
-	void AddInstanceWorldSpace(const FTransform& WorldTransform);
+	int32 AddInstanceWorldSpace(const FTransform& WorldTransform);
 
 	/** Get the transform for the instance specified. Instance is returned in local space of this component unless bWorldSpace is set.  Returns True on success. */
 	UFUNCTION(BlueprintCallable, Category = "Components|InstancedStaticMesh")
@@ -94,15 +102,15 @@ class ENGINE_API UInstancedStaticMeshComponent : public UStaticMeshComponent
 	
 	/** Update the transform for the instance specified. Instance is given in local space of this component unless bWorldSpace is set.  Returns True on success. */
 	UFUNCTION(BlueprintCallable, Category = "Components|InstancedStaticMesh")
-	bool UpdateInstanceTransform(int32 InstanceIndex, const FTransform& NewInstanceTransform, bool bWorldSpace=false);
+	virtual bool UpdateInstanceTransform(int32 InstanceIndex, const FTransform& NewInstanceTransform, bool bWorldSpace=false);
 
 	/** Remove the instance specified. Returns True on success. */
 	UFUNCTION(BlueprintCallable, Category = "Components|InstancedStaticMesh")
-	bool RemoveInstance(int32 InstanceIndex);
+	virtual bool RemoveInstance(int32 InstanceIndex);
 	
 	/** Clear all instances being rendered by this component */
 	UFUNCTION(BlueprintCallable, Category="Components|InstancedStaticMesh")
-	void ClearInstances();
+	virtual void ClearInstances();
 	
 	/** Get the number of instances in this component */
 	UFUNCTION(BlueprintCallable, Category = "Components|InstancedStaticMesh")
@@ -166,9 +174,6 @@ public:
 	void SelectInstance(bool bInSelected, int32 InInstanceIndex, int32 InInstanceCount = 1);
 
 private:
-	/** Initializes the body instance for the specified instance of the static mesh*/
-	void InitInstanceBody(int32 InstanceIdx, FBodyInstance* BodyInstance);
-
 	/** Creates body instances for all instances owned by this component */
 	void CreateAllInstanceBodies();
 
@@ -179,6 +184,15 @@ private:
 	void SetupNewInstanceData(FInstancedStaticMeshInstanceData& InOutNewInstanceData, int32 InInstanceIndex, const FTransform& InInstanceTransform);
 
 protected:
+	/** Request to navigation system to update only part of navmesh occupied by specified instance */
+	virtual void PartialNavigationUpdate(int32 InstanceIdx);
+	
+	/** Handles request from navigation system to gather instance transforms in a specific area box */
+	virtual void GetNavigationPerInstanceTransforms(const FBox& AreaBox, TArray<FTransform>& InstanceData) const;
+
+	/** Initializes the body instance for the specified instance of the static mesh*/
+	void InitInstanceBody(int32 InstanceIdx, FBodyInstance* BodyInstance);
+
 	/** Number of pending lightmaps still to be calculated (Apply()'d) */
 	UPROPERTY(Transient, DuplicateTransient, TextExportTransient)
 	int32 NumPendingLightmaps;
diff --git a/Engine/Source/Runtime/Engine/Classes/Engine/Engine.h b/Engine/Source/Runtime/Engine/Classes/Engine/Engine.h
--- a/Engine/Source/Runtime/Engine/Classes/Engine/Engine.h
+++ b/Engine/Source/Runtime/Engine/Classes/Engine/Engine.h
@@ -2355,6 +2355,14 @@ public:
 
 	virtual void TriggerStreamingDataRebuild() { }
 
+	/**
+	 * Updates level streaming state using active game players view and blocks until all sub-levels are loaded/ visible/ hidden
+	 * so further calls to UpdateLevelStreaming won't do any work unless state changes.
+	 *
+	 * @param InWorld Target world
+	 */
+	void BlockTillLevelStreamingCompleted(UWorld* InWorld);
+
 	/**
 	 * true if the loading movie was started during LoadMap().
 	 */
diff --git a/Engine/Source/Runtime/Engine/Classes/Engine/World.h b/Engine/Source/Runtime/Engine/Classes/Engine/World.h
--- a/Engine/Source/Runtime/Engine/Classes/Engine/World.h
+++ b/Engine/Source/Runtime/Engine/Classes/Engine/World.h
@@ -1847,14 +1847,6 @@ private:
 	 */
 	void UpdateLevelStreamingInner( UWorld* PersistentWorld, FSceneViewFamily* ViewFamily );
 
-	/**
-	 * Evaluates current world origin location against provided view information
-	 * Might issue request for world origin shifting
-	 *
-	 * @param ViewFamily	Collection of views to take into account
-	 */
-	void EvaluateWorldOriginLocation( const FSceneViewFamily& ViewFamily );
-
 public:
 	/**
 	 * Flushes level streaming in blocking fashion and returns when all levels are loaded/ visible/ hidden
diff --git a/Engine/Source/Runtime/Engine/Classes/Engine/WorldComposition.h b/Engine/Source/Runtime/Engine/Classes/Engine/WorldComposition.h
--- a/Engine/Source/Runtime/Engine/Classes/Engine/WorldComposition.h
+++ b/Engine/Source/Runtime/Engine/Classes/Engine/WorldComposition.h
@@ -68,12 +68,18 @@ class ENGINE_API UWorldComposition : public UObject
 
 	typedef TArray<FWorldCompositionTile> FTilesList;
 
-	/** Adds or removes level streaming objects to world based on distance settings from current view */
-	void UpdateStreamingState(FSceneViewFamily* InViewFamily = NULL);
+	/** Adds or removes level streaming objects to world based on distance settings from players current view */
+	void UpdateStreamingState();
 	
 	/** Adds or removes level streaming objects to world based on distance settings from current view point */
 	void UpdateStreamingState(const FVector& InLocation);
 
+	/**
+	 * Evaluates current world origin location against provided view location
+	 * Issues request for world origin rebasing in case location is far enough from current origin
+	 */
+	void EvaluateWorldOriginLocation(const FVector& ViewLocation);
+
 	/** Returns currently visible and hidden levels based on distance based streaming */
 	void GetDistanceVisibleLevels(const FVector& InLocation, TArray<FDistanceVisibleLevel>& OutVisibleLevels, TArray<FDistanceVisibleLevel>& OutHiddenLevels) const;
 
@@ -171,9 +177,6 @@ private:
 
 public:
 #if WITH_EDITOR
-	// Last view from where streaming state was updated
-	mutable FMatrix				LastWorldToViewMatrix;
-
 	// Hack for a World Browser to be able to temporally show hidden levels 
 	// regardless of current world origin and without offsetting them temporally 
 	bool						bTemporallyDisableOriginTracking;
diff --git a/Engine/Source/Runtime/Engine/Classes/Foliage/FoliageType.h b/Engine/Source/Runtime/Engine/Classes/Foliage/FoliageType.h
--- a/Engine/Source/Runtime/Engine/Classes/Foliage/FoliageType.h
+++ b/Engine/Source/Runtime/Engine/Classes/Foliage/FoliageType.h
@@ -108,17 +108,6 @@ class UFoliageType : public UObject
 	UPROPERTY(EditAnywhere, Category=Painting)
 	uint32 VertexColorMaskInvert:1;
 
-	/* The maximum number of instances that can be placed in a single cluster component and rendered in a single draw call */
-	UPROPERTY(EditAnywhere, Category=Clustering)
-	int32 MaxInstancesPerCluster;
-
-	/**
-	 * Determines the maximum size of a cluster of instances, rendered as a single component in a single draw call.
-	 * A smaller size will improve culling at the expense of more draw calls.
-	 */
-	UPROPERTY(EditAnywhere, Category = Clustering)
-	float MaxClusterRadius;
-
 	UPROPERTY()
 	float ReapplyDensityAmount;
 
diff --git a/Engine/Source/Runtime/Engine/Classes/PhysicsEngine/BodySetup.h b/Engine/Source/Runtime/Engine/Classes/PhysicsEngine/BodySetup.h
--- a/Engine/Source/Runtime/Engine/Classes/PhysicsEngine/BodySetup.h
+++ b/Engine/Source/Runtime/Engine/Classes/PhysicsEngine/BodySetup.h
@@ -217,7 +217,7 @@ public:
 	/** 
 	 *   Add the shapes defined by this body setup to the supplied PxRigidBody. 
 	 */
-	void AddShapesToRigidActor(physx::PxRigidActor* PDestActor, FVector& Scale3D, const FTransform& RelativeTM = FTransform::Identity, TArray<physx::PxShape*>* NewShapes = NULL);
+	ENGINE_API void AddShapesToRigidActor(physx::PxRigidActor* PDestActor, FVector& Scale3D, const FTransform& RelativeTM = FTransform::Identity, TArray<physx::PxShape*>* NewShapes = NULL);
 
 private:
 	void AddSpheresToRigidActor(physx::PxRigidActor* PDestActor, const FTransform& RelativeTM, float MinScale, float MinScaleAbs, TArray<physx::PxShape*>* NewShapes) const;
diff --git a/Engine/Source/Runtime/Engine/Private/AI/Navigation/NavigationModifier.cpp b/Engine/Source/Runtime/Engine/Private/AI/Navigation/NavigationModifier.cpp
--- a/Engine/Source/Runtime/Engine/Private/AI/Navigation/NavigationModifier.cpp
+++ b/Engine/Source/Runtime/Engine/Private/AI/Navigation/NavigationModifier.cpp
@@ -136,6 +136,10 @@ bool UNavLinkDefinition::HasAdjustableLinks() const
 FAreaNavModifier::FAreaNavModifier(float Radius, float Height, const FTransform& LocalToWorld, const TSubclassOf<UNavArea> InAreaClass)
 {
 	Init(InAreaClass);
+	
+	FVector Scale3D = LocalToWorld.GetScale3D().GetAbs();
+	Radius *= FMath::Max(Scale3D.X, Scale3D.Y);
+	Height *= Scale3D.Z;
 
 	Points.Init(2);
 	Points[0] = LocalToWorld.GetLocation();
@@ -586,6 +590,11 @@ uint32 FCompositeNavModifier::GetAllocatedSize() const
 	return MemUsed;
 }
 
+bool FCompositeNavModifier::HasPerInstanceTransforms() const
+{
+	return NavDataPerInstanceTransformDelegate.IsBound();
+}
+
 //----------------------------------------------------------------------//
 // UNavLinkTrivial
 //----------------------------------------------------------------------//
diff --git a/Engine/Source/Runtime/Engine/Private/AI/Navigation/NavigationOctree.cpp b/Engine/Source/Runtime/Engine/Private/AI/Navigation/NavigationOctree.cpp
--- a/Engine/Source/Runtime/Engine/Private/AI/Navigation/NavigationOctree.cpp
+++ b/Engine/Source/Runtime/Engine/Private/AI/Navigation/NavigationOctree.cpp
@@ -84,6 +84,14 @@ void FNavigationOctree::AppendToNode(const FOctreeElementId& Id, INavRelevantInt
 	AddElement(Data);
 }
 
+void FNavigationOctree::UpdateNode(const FOctreeElementId& Id, const FBox& NewBounds)
+{
+	FNavigationOctreeElement ElementCopy = GetElementById(Id);
+	RemoveElement(Id);
+	ElementCopy.Bounds = NewBounds;
+	AddElement(ElementCopy);
+}
+
 void FNavigationOctree::RemoveNode(const FOctreeElementId& Id)
 {
 	FNavigationOctreeElement& Data = GetElementById(Id);
@@ -94,6 +102,11 @@ void FNavigationOctree::RemoveNode(const FOctreeElementId& Id)
 	RemoveElement(Id);
 }
 
+bool FNavigationRelevantData::HasPerInstanceTransforms() const
+{
+	return NavDataPerInstanceTransformDelegate.IsBound();
+}
+
 bool FNavigationRelevantData::IsMatchingFilter(const FNavigationOctreeFilter& Filter) const
 {
 	return (Filter.bIncludeGeometry && HasGeometry()) ||
diff --git a/Engine/Source/Runtime/Engine/Private/AI/Navigation/NavigationSystem.cpp b/Engine/Source/Runtime/Engine/Private/AI/Navigation/NavigationSystem.cpp
--- a/Engine/Source/Runtime/Engine/Private/AI/Navigation/NavigationSystem.cpp
+++ b/Engine/Source/Runtime/Engine/Private/AI/Navigation/NavigationSystem.cpp
@@ -2349,6 +2349,30 @@ void UNavigationSystem::UpdateNavOctreeParentChain(UObject* ElementOwner)
 	}
 }
 
+bool UNavigationSystem::UpdateNavOctreeElementBounds(UActorComponent* Comp, const FBox& NewBounds, const FBox& DirtyArea)
+{
+	const FOctreeElementId* ElementId = GetObjectsNavOctreeId(Comp);
+	if (ElementId && ElementId->IsValidId())
+	{
+		NavOctree->UpdateNode(*ElementId, NewBounds);
+		
+		// Add dirty area
+		if (DirtyArea.IsValid)
+		{
+			ElementId = GetObjectsNavOctreeId(Comp);
+			if (ElementId && ElementId->IsValidId())
+			{
+				FNavigationOctreeElement& ElementData = NavOctree->GetElementById(*ElementId);
+				AddDirtyArea(DirtyArea, ElementData.Data.GetDirtyFlag());
+			}
+		}
+
+		return true;
+	}
+	
+	return false;
+}
+
 void UNavigationSystem::OnComponentRegistered(UActorComponent* Comp)
 {
 	SCOPE_CYCLE_COUNTER(STAT_DebugNavOctree);
diff --git a/Engine/Source/Runtime/Engine/Private/AI/Navigation/RecastNavMesh.cpp b/Engine/Source/Runtime/Engine/Private/AI/Navigation/RecastNavMesh.cpp
--- a/Engine/Source/Runtime/Engine/Private/AI/Navigation/RecastNavMesh.cpp
+++ b/Engine/Source/Runtime/Engine/Private/AI/Navigation/RecastNavMesh.cpp
@@ -590,16 +590,18 @@ FColor ARecastNavMesh::GetAreaIDColor(uint8 AreaID) const
 	return DefArea ? DefArea->DrawColor : NavDataConfig.Color;
 }
 
-void ARecastNavMesh::SortAreasForGenerator(TArray<FAreaNavModifier>& Areas) const
+void ARecastNavMesh::SortAreasForGenerator(TArray<FRecastAreaNavModifierElement>& Modifiers) const
 {
 	// initialize costs for sorting
 	float AreaCosts[RECAST_MAX_AREAS];
 	float AreaFixedCosts[RECAST_MAX_AREAS];
 	DefaultQueryFilter->GetAllAreaCosts(AreaCosts, AreaFixedCosts, RECAST_MAX_AREAS);
 
-	for (int32 Idx = 0; Idx < Areas.Num(); Idx++)
+	for (auto& Element : Modifiers)
 	{
-		FAreaNavModifier& AreaMod = Areas[Idx];
+		check(Element.Areas.Num() > 0);
+		
+		FAreaNavModifier& AreaMod = Element.Areas[0];
 		const int32 AreaId = GetAreaID(AreaMod.GetAreaClass());
 		if (AreaId >= 0 && AreaId < RECAST_MAX_AREAS)
 		{
@@ -610,8 +612,14 @@ void ARecastNavMesh::SortAreasForGenerator(TArray<FAreaNavModifier>& Areas) cons
 
 	struct FNavAreaSortPredicate
 	{
-		FORCEINLINE bool operator()(const FAreaNavModifier& A, const FAreaNavModifier& B) const
+		FORCEINLINE bool operator()(const FRecastAreaNavModifierElement& ElA, const FRecastAreaNavModifierElement& ElB) const
 		{
+			check(ElA.Areas.Num() > 0);
+			check(ElB.Areas.Num() > 0);
+			// assuming composite modifiers has same area type
+			const FAreaNavModifier& A = ElA.Areas[0];
+			const FAreaNavModifier& B = ElB.Areas[0];
+			
 			const bool bIsAReplacing = (A.GetAreaClassToReplace() != NULL);
 			const bool bIsBReplacing = (B.GetAreaClassToReplace() != NULL);
 			if (bIsAReplacing != bIsBReplacing)
@@ -623,7 +631,7 @@ void ARecastNavMesh::SortAreasForGenerator(TArray<FAreaNavModifier>& Areas) cons
 		}
 	};
 
-	Areas.Sort(FNavAreaSortPredicate());
+	Modifiers.Sort(FNavAreaSortPredicate());
 }
 
 void ARecastNavMesh::SerializeRecastNavMesh(FArchive& Ar, FPImplRecastNavMesh*& NavMesh)
diff --git a/Engine/Source/Runtime/Engine/Private/AI/Navigation/RecastNavMeshGenerator.cpp b/Engine/Source/Runtime/Engine/Private/AI/Navigation/RecastNavMeshGenerator.cpp
--- a/Engine/Source/Runtime/Engine/Private/AI/Navigation/RecastNavMeshGenerator.cpp
+++ b/Engine/Source/Runtime/Engine/Private/AI/Navigation/RecastNavMeshGenerator.cpp
@@ -201,6 +201,7 @@ struct FRecastGeometryExport : public FNavigableGeometryExport
 	virtual void ExportCustomMesh(const FVector* InVertices, int32 NumVerts, const int32* InIndices, int32 NumIndices, const FTransform& LocalToWorld) override;
 	virtual void ExportRigidBodySetup(UBodySetup& BodySetup, const FTransform& LocalToWorld) override;
 	virtual void AddNavModifiers(const FCompositeNavModifier& Modifiers) override;
+	virtual void SetNavDataPerInstanceTransformDelegate(const FNavDataPerInstanceTransformDelegate& InDelegate) override;
 };
 
 FRecastVoxelCache::FRecastVoxelCache(const uint8* Memory)
@@ -985,6 +986,28 @@ void FRecastGeometryExport::AddNavModifiers(const FCompositeNavModifier& Modifie
 	Data->Modifiers.Add(Modifiers);
 }
 
+void FRecastGeometryExport::SetNavDataPerInstanceTransformDelegate(const FNavDataPerInstanceTransformDelegate& InDelegate)
+{
+	Data->NavDataPerInstanceTransformDelegate = InDelegate;
+}
+
+static void PartialTransformConvexHull(FConvexNavAreaData& ConvexData, const FTransform& LocalToWorld)
+{
+	FVector ScaleXY = LocalToWorld.GetScale3D().GetAbs();
+	ScaleXY.Z = 1.f;
+
+	FVector TranslationXY = LocalToWorld.GetLocation();
+	TranslationXY.Z = 0.f;
+
+	for (FVector& Point : ConvexData.Points)
+	{
+		Point = Point*ScaleXY + TranslationXY;
+	}
+
+	ConvexData.MaxZ+= LocalToWorld.GetLocation().Z;
+	ConvexData.MinZ+= LocalToWorld.GetLocation().Z;
+}
+
 FORCEINLINE void GrowConvexHull(const float ExpandBy, const TArray<FVector>& Verts, TArray<FVector>& OutResult)
 {
 	if (Verts.Num() < 3)
@@ -1485,11 +1508,9 @@ bool FRecastTileGenerator::HasDataToBuild() const
 {
 	return
 		CompressedLayers.Num()
-		|| StaticAreas.Num()
-		|| DynamicAreas.Num()
+		|| Modifiers.Num()
 		|| OffmeshLinks.Num()
-		|| GeomCoords.Num()
-		|| GeomIndices.Num();
+		|| RawGeometry.Num();
 }
 
 void FRecastTileGenerator::DoWork()
@@ -1546,12 +1567,12 @@ void FRecastTileGenerator::GatherGeometry(const FRecastNavMeshGenerator* ParentG
 				}
 				else
 				{
-					AppendGeometry(Element.Data.CollisionData);
+					AppendGeometry(Element.Data.CollisionData, Element.Data.NavDataPerInstanceTransformDelegate);
 				}
 			}
 
 			const FCompositeNavModifier ModifierInstance = Element.GetModifierForAgent(NavDataConfig);
-			AppendModifier(ModifierInstance, bExportGeometry);
+			AppendModifier(ModifierInstance, Element.Data.NavDataPerInstanceTransformDelegate);
 		}
 	}
 }
@@ -1740,18 +1761,8 @@ void FRecastTileGenerator::AddVoxelCache(TNavStatArray<uint8>& RawVoxelCache, co
 	FMemory::Memcpy(RawVoxelCache.GetData() + NewCacheIdx + HeaderSize, CachedVoxels, VoxelsSize);
 }
 
-void FRecastTileGenerator::AppendModifier(const FCompositeNavModifier& Modifier, bool bStatic)
+void FRecastTileGenerator::AppendModifier(const FCompositeNavModifier& Modifier, const FNavDataPerInstanceTransformDelegate& InTransformsDelegate)
 {
-	// append areas to correct array
-	if (bStatic)
-	{
-		StaticAreas.Append(Modifier.GetAreas());
-	}
-	else
-	{
-		DynamicAreas.Append(Modifier.GetAreas());
-	}
-
 	// append all offmesh links (not included in compress layers)
 	OffmeshLinks.Append(Modifier.GetSimpleLinks());
 
@@ -1762,67 +1773,59 @@ void FRecastTileGenerator::AppendModifier(const FCompositeNavModifier& Modifier,
 		FSimpleLinkNavModifier SimpleLinkCollection(UNavLinkDefinition::GetLinksDefinition(LinkModifier->GetNavLinkClass()), LinkModifier->LocalToWorld);
 		OffmeshLinks.Add(SimpleLinkCollection);
 	}
-}
 
-void FRecastTileGenerator::AppendGeometry(const TNavStatArray<uint8>& RawCollisionCache)
-{
-	if (RawCollisionCache.Num() == 0)
+	if (Modifier.GetAreas().Num() == 0)
 	{
 		return;
 	}
-	QUICK_SCOPE_CYCLE_COUNTER(STAT_RecastTileGenerator_AppendGeometry);
-
-	FRecastGeometryCache CollisionCache(RawCollisionCache.GetData());
 	
-	const int32 FirstNewCoord = GeomCoords.Num();
-	const int32 FirstNewIndex = GeomIndices.Num();
-	const int32 VertBase = FirstNewCoord / 3;
-	const int32 NumCoords = CollisionCache.Header.NumVerts * 3;
-	const int32 NumIndices = CollisionCache.Header.NumFaces * 3;
-			
-	GeomCoords.Reserve(GeomCoords.Num() + NumCoords);
-	GeomCoords.AddUninitialized(NumCoords);
-	
-	GeomIndices.Reserve(GeomIndices.Num() + NumIndices);
-	GeomIndices.AddUninitialized(NumIndices);
+	FRecastAreaNavModifierElement ModifierElement;
 
-	FMemory::Memcpy(GeomCoords.GetData() + FirstNewCoord, CollisionCache.Verts, sizeof(float) * NumCoords);
-	int32* DestIndex = GeomIndices.GetData() + FirstNewIndex;
-
-	for (int32 i = 0; i < NumIndices; i++, DestIndex++)
+	// Gather per instance transforms if any
+	if (InTransformsDelegate.IsBound())
 	{
-		*DestIndex = CollisionCache.Indices[i] + VertBase;
+		InTransformsDelegate.Execute(TileBB, ModifierElement.PerInstanceTransform);
+		// skip this modifier in case there is no instances for this tile
+		if (ModifierElement.PerInstanceTransform.Num() == 0)
+		{
+			return;
+		}
 	}
+		
+	ModifierElement.Areas = Modifier.GetAreas();
+	Modifiers.Add(MoveTemp(ModifierElement));
 }
 
-void FRecastTileGenerator::AppendGeometry(const TNavStatArray<FVector>& Verts, const TNavStatArray<int32>& Faces)
-{
-	if (Faces.Num() == 0 || Verts.Num() == 0)
+void FRecastTileGenerator::AppendGeometry(const TNavStatArray<uint8>& RawCollisionCache, const FNavDataPerInstanceTransformDelegate& InTransformsDelegate)
+{	
+	if (RawCollisionCache.Num() == 0)
 	{
 		return;
 	}
-
-	const int32 FirstNewCoord = GeomCoords.Num();
-	const int32 FirstNewIndex = GeomIndices.Num();
-	const int32 VertBase = FirstNewCoord / 3;
-
-	GeomCoords.AddZeroed(Verts.Num() * 3);
-	GeomIndices.AddZeroed(Faces.Num());
-
-	float* DestCoord = GeomCoords.GetData() + FirstNewCoord;
-	int32* DestIndex = GeomIndices.GetData() + FirstNewIndex;
 	
-	for (int32 i = 0; i < Verts.Num(); i++)
+	FRecastRawGeometryElement GeometryElement;
+	FRecastGeometryCache CollisionCache(RawCollisionCache.GetData());
+	
+	// Gather per instance transforms
+	if (InTransformsDelegate.IsBound())
 	{
-		*DestCoord = Verts[i].X; DestCoord++;
-		*DestCoord = Verts[i].Y; DestCoord++;
-		*DestCoord = Verts[i].Z; DestCoord++;
+		InTransformsDelegate.Execute(TileBB, GeometryElement.PerInstanceTransform);
+		if (GeometryElement.PerInstanceTransform.Num() == 0)
+		{
+			return;
+		}
 	}
+	
+	const int32 NumCoords = CollisionCache.Header.NumVerts * 3;
+	const int32 NumIndices = CollisionCache.Header.NumFaces * 3;
+	
+	GeometryElement.GeomCoords.Init(NumCoords);
+	GeometryElement.GeomIndices.Init(NumIndices);
 
-	for (int32 i = 0; i < Faces.Num(); i++, DestIndex++)
-	{
-		*DestIndex = Faces[i] + VertBase;
-	}
+	FMemory::Memcpy(GeometryElement.GeomCoords.GetData(), CollisionCache.Verts, sizeof(float) * NumCoords);
+	FMemory::Memcpy(GeometryElement.GeomIndices.GetData(), CollisionCache.Indices, sizeof(int32) * NumIndices);
+
+	RawGeometry.Add(MoveTemp(GeometryElement));
 }
 
 bool FRecastTileGenerator::GenerateTile()
@@ -1871,6 +1874,49 @@ struct FTileRasterizationContext
 	TArray<FNavMeshTileData> Layers;
 };
 
+static void RasterizeGeometry(
+	FNavMeshBuildContext& BuildContext, const FRecastBuildConfig& TileConfig, 
+	const TArray<float>& Coords, const TArray<int32>& Indices, FTileRasterizationContext& RasterContext)
+{
+	const int32 NumFaces = Indices.Num() / 3;
+	const int32 NumVerts = Coords.Num() / 3;
+
+	TNavStatArray<uint8> TriAreas;
+	TriAreas.Reserve(NumFaces);
+	TriAreas.AddZeroed(NumFaces);
+
+	rcMarkWalkableTriangles(&BuildContext, TileConfig.walkableSlopeAngle,
+		Coords.GetData(), NumVerts, Indices.GetData(), NumFaces,
+		TriAreas.GetData());
+
+	rcRasterizeTriangles(&BuildContext,
+		Coords.GetData(), NumVerts, 
+		Indices.GetData(), TriAreas.GetData(), NumFaces,
+		*RasterContext.SolidHF, TileConfig.walkableClimb);
+}
+
+static void RasterizeGeometry(
+	FNavMeshBuildContext& BuildContext,const FRecastBuildConfig& TileConfig, 
+	const TArray<float>& Coords, const TArray<int32>& Indices, const FTransform& LocalToWorld, FTileRasterizationContext& RasterContext)
+{
+	TArray<float> WorldRecastCoords;
+	WorldRecastCoords.Init(Coords.Num());
+	
+	FMatrix LocalToRecastWorld = LocalToWorld.ToMatrixWithScale()*Unreal2RecastMatrix();
+	// Convert geometry to recast world space
+	for (int32 i = 0; i < Coords.Num(); i+=3)
+	{
+		// collision cache stores coordinates in recast space, convert them to unreal and transform to recast world space
+		FVector WorldRecastCoord = LocalToRecastWorld.TransformPosition(Recast2UnrealPoint(&Coords[i])); 
+
+		WorldRecastCoords[i+0] = WorldRecastCoord.X;
+		WorldRecastCoords[i+1] = WorldRecastCoord.Y;
+		WorldRecastCoords[i+2] = WorldRecastCoord.Z;
+	}
+
+	RasterizeGeometry(BuildContext, TileConfig, WorldRecastCoords, Indices, RasterContext);
+}
+
 bool FRecastTileGenerator::GenerateCompressedLayers(FNavMeshBuildContext& BuildContext)
 {
 	RECAST_STAT(STAT_Navigation_Async_Recast_Rasterize);
@@ -1889,9 +1935,10 @@ bool FRecastTileGenerator::GenerateCompressedLayers(FNavMeshBuildContext& BuildC
 	BuildContext.log(RC_LOG_PROGRESS, " - %d x %d cells", TileConfig.width, TileConfig.height);
 
 	FTileRasterizationContext RasterContext;
+	const bool bHasGeometry = RawGeometry.Num() > 0;
 
 	// Allocate voxel heightfield where we rasterize our input data to.
-	if (GeomCoords.Num())
+	if (bHasGeometry)
 	{
 		RasterContext.SolidHF = rcAllocHeightfield();
 		if (RasterContext.SolidHF == NULL)
@@ -1907,28 +1954,25 @@ bool FRecastTileGenerator::GenerateCompressedLayers(FNavMeshBuildContext& BuildC
 	}
 
 	// Rasterize geometry
-	if (GeomCoords.Num() && GeomIndices.Num())
+	if (bHasGeometry)
 	{
 		DECLARE_SCOPE_CYCLE_COUNTER(TEXT("Rasterization: without voxel cache"), Stat_RecastRasterNoCache, STATGROUP_Navigation);
-
-		const int32 NumFaces = GeomIndices.Num() / 3;
-		const int32 NumVerts = GeomCoords.Num() / 3;
-
-		RECAST_STAT(STAT_Navigation_RasterizeTriangles);
-		TNavStatArray<uint8> TriAreas;
-		TriAreas.Reserve(NumFaces);
-		TriAreas.AddZeroed(NumFaces);
-
-		rcMarkWalkableTriangles(&BuildContext, TileConfig.walkableSlopeAngle,
-			GeomCoords.GetData(), NumVerts, GeomIndices.GetData(), NumFaces,
-			TriAreas.GetData());
-
-		rcRasterizeTriangles(&BuildContext,
-			GeomCoords.GetData(), NumVerts, 
-			GeomIndices.GetData(), TriAreas.GetData(), NumFaces,
-			*RasterContext.SolidHF, TileConfig.walkableClimb);
+		RECAST_STAT(STAT_Navigation_RasterizeTriangles)
+		
+		for (const auto& Element : RawGeometry)
+		{
+			for (const auto& InstanceTransform : Element.PerInstanceTransform)
+			{
+				RasterizeGeometry(BuildContext, TileConfig, Element.GeomCoords, Element.GeomIndices, InstanceTransform, RasterContext);
+			}
+			
+			if (Element.PerInstanceTransform.Num() == 0)
+			{
+				RasterizeGeometry(BuildContext, TileConfig, Element.GeomCoords, Element.GeomIndices, RasterContext);
+			}
+		}
 	}
-
+	
 	if (!RasterContext.SolidHF || RasterContext.SolidHF->pools == 0)
 	{
 		BuildContext.log(RC_LOG_WARNING, "GenerateCompressedLayers: empty tile - aborting");
@@ -1994,9 +2038,6 @@ bool FRecastTileGenerator::GenerateCompressedLayers(FNavMeshBuildContext& BuildC
 		}
 	}
 
-	// (Optional) Mark areas.
-	MarkStaticAreas(BuildContext, *RasterContext.CompactHF);
-
 	// remove all low area marking at this point
 	if (TileConfig.bMarkLowHeightAreas)
 	{
@@ -2123,121 +2164,6 @@ bool FRecastTileGenerator::GenerateCompressedLayers(FNavMeshBuildContext& BuildC
 	return true;
 }
 
-void FRecastTileGenerator::MarkStaticAreas(FNavMeshBuildContext& BuildContext, rcCompactHeightfield& CompactHF)
-{
-	const int32 NumAreas = StaticAreas.Num();
-	if (NumAreas == 0)
-	{
-		return;
-	}
-
-	if (AdditionalCachedData.bUseSortFunction && AdditionalCachedData.ActorOwner && NumAreas > 1)
-	{
-		AdditionalCachedData.ActorOwner->SortAreasForGenerator(StaticAreas);
-	}
-
-	RECAST_STAT(STAT_Navigation_Async_MarkAreas);
-	const float ExpandBy = TileConfig.AgentRadius * 1.5;
-	const FAreaNavModifier* Modifier = StaticAreas.GetData();
-	
-	for (int32 ModifierIndex = 0; ModifierIndex < NumAreas; ++ModifierIndex, ++Modifier)
-	{
-		const int32* AreaID = AdditionalCachedData.AreaClassToIdMap.Find(Modifier->GetAreaClass());
-		const int32* ReplaceID = AdditionalCachedData.AreaClassToIdMap.Find(Modifier->GetAreaClassToReplace());
-		if (AreaID == NULL)
-		{
-			// happens when area is not supported by agent owning this navmesh
-			continue;
-		}
-
-		const float OffsetZ = TileConfig.ch + (Modifier->ShouldIncludeAgentHeight() ? TileConfig.AgentHeight : 0.0f);
-		switch (Modifier->GetShapeType())
-		{
-		case ENavigationShapeType::Cylinder:
-			{
-				FCylinderNavAreaData CylinderData;
-				Modifier->GetCylinder(CylinderData);
-
-				CylinderData.Height += OffsetZ;
-				CylinderData.Radius += ExpandBy;
-
-				FVector RecastPos = Unreal2RecastPoint(CylinderData.Origin);
-
-				if (ReplaceID)
-				{
-					rcReplaceCylinderArea(&BuildContext, &(RecastPos.X), CylinderData.Radius, CylinderData.Height, *AreaID, *ReplaceID, CompactHF);
-				}
-				else
-				{
-					rcMarkCylinderArea(&BuildContext, &(RecastPos.X), CylinderData.Radius, CylinderData.Height, *AreaID, CompactHF);
-				}
-			}
-			break;
-
-		case ENavigationShapeType::Box:
-			{
-				FBoxNavAreaData BoxData;
-				Modifier->GetBox(BoxData);
-
-				BoxData.Extent += FVector(ExpandBy, ExpandBy, OffsetZ);
-
-				FBox UnrealBox = FBox::BuildAABB(BoxData.Origin, BoxData.Extent);
-				FBox RecastBox = Unreal2RecastBox(UnrealBox);
-
-				if (ReplaceID)
-				{
-					rcReplaceBoxArea(&BuildContext, &(RecastBox.Min.X), &(RecastBox.Max.X), *AreaID, *ReplaceID, CompactHF);
-				}
-				else
-				{
-					rcMarkBoxArea(&BuildContext, &(RecastBox.Min.X), &(RecastBox.Max.X), *AreaID, CompactHF);
-				}
-			}
-			break;
-
-		case ENavigationShapeType::Convex:
-			{
-				FConvexNavAreaData ConvexData;
-				Modifier->GetConvex(ConvexData);
-
-				TArray<FVector> ConvexVerts;
-				GrowConvexHull(ExpandBy, ConvexData.Points, ConvexVerts);
-				ConvexData.MinZ -= OffsetZ;
-				ConvexData.MaxZ += TileConfig.ch;
-
-				if (ConvexVerts.Num())
-				{
-					TArray<float> ConvexCoords;
-					ConvexCoords.AddZeroed(ConvexVerts.Num() * 3);
-
-					float* ItCoord = ConvexCoords.GetData();
-					for (int32 i = 0; i < ConvexVerts.Num(); i++)
-					{
-						const FVector RecastV = Unreal2RecastPoint(ConvexVerts[i]);
-						*ItCoord = RecastV.X; ItCoord++;
-						*ItCoord = RecastV.Y; ItCoord++;
-						*ItCoord = RecastV.Z; ItCoord++;
-					}
-
-					if (ReplaceID)
-					{
-						rcReplaceConvexPolyArea(&BuildContext, ConvexCoords.GetData(), ConvexVerts.Num(),
-							ConvexData.MinZ, ConvexData.MaxZ, *AreaID, *ReplaceID, CompactHF);
-					}
-					else
-					{
-						rcMarkConvexPolyArea(&BuildContext, ConvexCoords.GetData(), ConvexVerts.Num(),
-							ConvexData.MinZ, ConvexData.MaxZ, *AreaID, CompactHF);
-					}
-				}
-			}
-			break;
-		
-		default: break;
-		}
-	}
-}
-
 struct FTileGenerationContext
 {
 	FTileGenerationContext(dtTileCacheAlloc* MyAllocator) :
@@ -2526,138 +2452,161 @@ bool FRecastTileGenerator::GenerateNavigationData(FNavMeshBuildContext& BuildCon
 
 void FRecastTileGenerator::MarkDynamicAreas(dtTileCacheLayer& Layer)
 {
-	if (DynamicAreas.Num() == 0)
+	if (Modifiers.Num() == 0)
 	{
 		return;
 	}
+	
+	RECAST_STAT(STAT_Navigation_Async_MarkAreas);
 
-	// combine both area types, to avoid overriding high cost static with low cost dynamic one
-	TArray<FAreaNavModifier> CombinedAreas;
-	CombinedAreas.Append(StaticAreas);
-	CombinedAreas.Append(DynamicAreas);
+	if (AdditionalCachedData.bUseSortFunction && AdditionalCachedData.ActorOwner && Modifiers.Num() > 1)
+	{
+		AdditionalCachedData.ActorOwner->SortAreasForGenerator(Modifiers);
+	}
+		
+	for (const auto& Element : Modifiers)
+	{
+		for (const auto& Area : Element.Areas)
+		{
+			for (const auto& LocalToWorld : Element.PerInstanceTransform)
+			{
+				MarkDynamicArea(Area, LocalToWorld, Layer);
+			}
+
+			if (Element.PerInstanceTransform.Num() == 0)
+			{
+				MarkDynamicArea(Area, FTransform::Identity, Layer);
+			}
+		}
+	}
+}
+
+void FRecastTileGenerator::MarkDynamicArea(const FAreaNavModifier& Modifier, const FTransform& LocalToWorld, dtTileCacheLayer& Layer)
+{
+	const int32* AreaID = AdditionalCachedData.AreaClassToIdMap.Find(Modifier.GetAreaClass());
+	const int32* ReplaceID = AdditionalCachedData.AreaClassToIdMap.Find(Modifier.GetAreaClassToReplace());
+	if (AreaID == NULL)
+	{
+		// happens when area is not supported by agent owning this navmesh
+		return;
+	}
 
-	if (AdditionalCachedData.bUseSortFunction && AdditionalCachedData.ActorOwner && CombinedAreas.Num() > 1)
+	// Check whether modifier affects this layer
+	const FBox LayerUnrealBounds = Recast2UnrealBox(Layer.header->bmin, Layer.header->bmax);
+	FBox ModifierBounds = Modifier.GetBounds().TransformBy(LocalToWorld);
+	if (Modifier.ShouldIncludeAgentHeight())
 	{
-		AdditionalCachedData.ActorOwner->SortAreasForGenerator(CombinedAreas);
+		ModifierBounds.Min.Z -= TileConfig.AgentHeight;
+	}
+
+	if (!LayerUnrealBounds.Intersect(ModifierBounds))
+	{
+		return;
 	}
 
-	RECAST_STAT(STAT_Navigation_Async_MarkAreas);
 	const float ExpandBy = TileConfig.AgentRadius;
 	const float* LayerRecastOrig = Layer.header->bmin;
-	const FBox LayerUnrealBounds = Recast2UnrealBox(Layer.header->bmin, Layer.header->bmax);
+	const float OffsetZ = TileConfig.ch + (Modifier.ShouldIncludeAgentHeight() ? TileConfig.AgentHeight : 0.0f);
 
-	const FAreaNavModifier* Modifier = CombinedAreas.GetData();
-	for (int32 ModifierIndex = 0; ModifierIndex < CombinedAreas.Num(); ++ModifierIndex, ++Modifier)
+	switch (Modifier.GetShapeType())
 	{
-		const int32* AreaID = AdditionalCachedData.AreaClassToIdMap.Find(Modifier->GetAreaClass());
-		const int32* ReplaceID = AdditionalCachedData.AreaClassToIdMap.Find(Modifier->GetAreaClassToReplace());
-		if (AreaID == NULL)
+	case ENavigationShapeType::Cylinder:
 		{
-			// happens when area is not supported by agent owning this navmesh
-			continue;
-		}
+			FCylinderNavAreaData CylinderData;
+			Modifier.GetCylinder(CylinderData);
 
-		FBox ModifierBounds = Modifier->GetBounds();
-		if (Modifier->ShouldIncludeAgentHeight())
-		{
-			ModifierBounds.Min.Z -= TileConfig.AgentHeight;
+			// Only scaling and translation
+			FVector Scale3D = LocalToWorld.GetScale3D().GetAbs();
+			CylinderData.Height *= Scale3D.Z;
+			CylinderData.Radius *= FMath::Max(Scale3D.X, Scale3D.Y);
+			CylinderData.Origin = LocalToWorld.TransformPosition(CylinderData.Origin);
+			
+			CylinderData.Origin.Z -= OffsetZ;
+			CylinderData.Height += OffsetZ*2.f;
+			CylinderData.Radius += ExpandBy;
+			
+			FVector RecastPos = Unreal2RecastPoint(CylinderData.Origin);
+
+			if (ReplaceID)
+			{
+				dtReplaceCylinderArea(Layer, LayerRecastOrig, TileConfig.cs, TileConfig.ch,
+					&(RecastPos.X), CylinderData.Radius, CylinderData.Height, *AreaID, *ReplaceID);
+			}
+			else
+			{
+				dtMarkCylinderArea(Layer, LayerRecastOrig, TileConfig.cs, TileConfig.ch,
+					&(RecastPos.X), CylinderData.Radius, CylinderData.Height, *AreaID);
+			}
 		}
+		break;
 
-		if (!LayerUnrealBounds.Intersect(ModifierBounds))
+	case ENavigationShapeType::Box:
 		{
-			continue;
+			FBoxNavAreaData BoxData;
+			Modifier.GetBox(BoxData);
+
+			FBox WorldBox = FBox::BuildAABB(BoxData.Origin, BoxData.Extent).TransformBy(LocalToWorld);
+			WorldBox = WorldBox.ExpandBy(FVector(ExpandBy, ExpandBy, OffsetZ));
+			
+			FBox RacastBox = Unreal2RecastBox(WorldBox);
+			FVector RecastPos;
+			FVector RecastExtent;
+			RacastBox.GetCenterAndExtents(RecastPos, RecastExtent);
+				
+			if (ReplaceID)
+			{
+				dtReplaceBoxArea(Layer, LayerRecastOrig, TileConfig.cs, TileConfig.ch,
+					&(RecastPos.X), &(RecastExtent.X), *AreaID, *ReplaceID);
+			}
+			else
+			{
+				dtMarkBoxArea(Layer, LayerRecastOrig, TileConfig.cs, TileConfig.ch,
+					&(RecastPos.X), &(RecastExtent.X), *AreaID);
+			}
 		}
+		break;
 
-		const float OffsetZ = TileConfig.ch + (Modifier->ShouldIncludeAgentHeight() ? TileConfig.AgentHeight : 0.0f);
-		switch (Modifier->GetShapeType())
+	case ENavigationShapeType::Convex:
 		{
-		case ENavigationShapeType::Cylinder:
-			{
-				FCylinderNavAreaData CylidnerData;
-				Modifier->GetCylinder(CylidnerData);
-				
-				CylidnerData.Height += OffsetZ;
-				CylidnerData.Radius += ExpandBy;
+			FConvexNavAreaData ConvexData;
+			Modifier.GetConvex(ConvexData);
+			// Only scaling and translation
+			PartialTransformConvexHull(ConvexData, LocalToWorld);
 
-				FVector RecastPos = Unreal2RecastPoint(CylidnerData.Origin);
+			TArray<FVector> ConvexVerts;
+			GrowConvexHull(ExpandBy, ConvexData.Points, ConvexVerts);
+			ConvexData.MinZ -= OffsetZ;
+			ConvexData.MaxZ += TileConfig.ch;
 
-				if (ReplaceID)
-				{
-					dtReplaceCylinderArea(Layer, LayerRecastOrig, TileConfig.cs, TileConfig.ch,
-						&(RecastPos.X), CylidnerData.Radius, CylidnerData.Height, *AreaID, *ReplaceID);
-				}
-				else
+			if (ConvexVerts.Num())
+			{
+				TArray<float> ConvexCoords;
+				ConvexCoords.AddZeroed(ConvexVerts.Num() * 3);
+						
+				float* ItCoord = ConvexCoords.GetData();
+				for (int32 i = 0; i < ConvexVerts.Num(); i++)
 				{
-					dtMarkCylinderArea(Layer, LayerRecastOrig, TileConfig.cs, TileConfig.ch,
-						&(RecastPos.X), CylidnerData.Radius, CylidnerData.Height, *AreaID);
+					const FVector RecastV = Unreal2RecastPoint(ConvexVerts[i]);
+					*ItCoord = RecastV.X; ItCoord++;
+					*ItCoord = RecastV.Y; ItCoord++;
+					*ItCoord = RecastV.Z; ItCoord++;
 				}
-			}
-			break;
-
-		case ENavigationShapeType::Box:
-			{
-				FBoxNavAreaData BoxData;
-				Modifier->GetBox(BoxData);
-
-				BoxData.Extent += FVector(ExpandBy, ExpandBy, OffsetZ);
-
-				FVector RecastPos = Unreal2RecastPoint(BoxData.Origin);
-				FVector RecastExtent = Unreal2RecastPoint(BoxData.Extent).GetAbs();
 
 				if (ReplaceID)
 				{
-					dtReplaceBoxArea(Layer, LayerRecastOrig, TileConfig.cs, TileConfig.ch,
-						&(RecastPos.X), &(RecastExtent.X), *AreaID, *ReplaceID);
+					dtReplaceConvexArea(Layer, LayerRecastOrig, TileConfig.cs, TileConfig.ch,
+						ConvexCoords.GetData(), ConvexVerts.Num(), ConvexData.MinZ, ConvexData.MaxZ, *AreaID, *ReplaceID);
 				}
 				else
 				{
-					dtMarkBoxArea(Layer, LayerRecastOrig, TileConfig.cs, TileConfig.ch,
-						&(RecastPos.X), &(RecastExtent.X), *AreaID);
-				}
-			}
-			break;
-
-		case ENavigationShapeType::Convex:
-			{
-				FConvexNavAreaData ConvexData;
-				Modifier->GetConvex(ConvexData);
-
-				TArray<FVector> ConvexVerts;
-				GrowConvexHull(ExpandBy, ConvexData.Points, ConvexVerts);
-
-				ConvexData.MinZ -= OffsetZ;
-				ConvexData.MaxZ += TileConfig.ch;
-
-				if (ConvexVerts.Num())
-				{
-					TArray<float> ConvexCoords;
-					ConvexCoords.AddZeroed(ConvexVerts.Num() * 3);
-
-					float* ItCoord = ConvexCoords.GetData();
-					for (int32 i = 0; i < ConvexVerts.Num(); i++)
-					{
-						const FVector RecastV = Unreal2RecastPoint(ConvexVerts[i]);
-						*ItCoord = RecastV.X; ItCoord++;
-						*ItCoord = RecastV.Y; ItCoord++;
-						*ItCoord = RecastV.Z; ItCoord++;
-					}
-
-					if (ReplaceID)
-					{
-						dtReplaceConvexArea(Layer, LayerRecastOrig, TileConfig.cs, TileConfig.ch,
-							ConvexCoords.GetData(), ConvexVerts.Num(), ConvexData.MinZ, ConvexData.MaxZ, *AreaID, *ReplaceID);
-					}
-					else
-					{
-						dtMarkConvexArea(Layer, LayerRecastOrig, TileConfig.cs, TileConfig.ch,
-							ConvexCoords.GetData(), ConvexVerts.Num(), ConvexData.MinZ, ConvexData.MaxZ, *AreaID);
-					}
+					dtMarkConvexArea(Layer, LayerRecastOrig, TileConfig.cs, TileConfig.ch,
+						ConvexCoords.GetData(), ConvexVerts.Num(), ConvexData.MinZ, ConvexData.MaxZ, *AreaID);
 				}
 			}
-			break;
-
-		default: break;
 		}
+		break;
+
+	default: break;
 	}
 }
 
@@ -2665,12 +2614,24 @@ uint32 FRecastTileGenerator::GetUsedMemCount() const
 {
 	uint32 TotalMemory = 0;
 	TotalMemory += InclusionBounds.GetAllocatedSize();
-	TotalMemory += StaticAreas.GetAllocatedSize();
-	TotalMemory += DynamicAreas.GetAllocatedSize();
+	TotalMemory += Modifiers.GetAllocatedSize();
 	TotalMemory += OffmeshLinks.GetAllocatedSize();
-	TotalMemory += GeomCoords.GetAllocatedSize();
-	TotalMemory += GeomIndices.GetAllocatedSize();
+	TotalMemory += RawGeometry.GetAllocatedSize();
+	TotalMemory += Modifiers.GetAllocatedSize();
 	
+	for (const auto& Element : RawGeometry)
+	{
+		TotalMemory += Element.GeomCoords.GetAllocatedSize();
+		TotalMemory += Element.GeomIndices.GetAllocatedSize();
+		TotalMemory += Element.PerInstanceTransform.GetAllocatedSize();
+	}
+
+	for (const auto& Element : Modifiers)
+	{
+		TotalMemory += Element.Areas.GetAllocatedSize();
+		TotalMemory += Element.PerInstanceTransform.GetAllocatedSize();
+	}
+
 	const FSimpleLinkNavModifier* SimpleLink = OffmeshLinks.GetData();
 	for (int32 Index = 0; Index < OffmeshLinks.Num(); ++Index, ++SimpleLink)
 	{
diff --git a/Engine/Source/Runtime/Engine/Private/GameEngine.cpp b/Engine/Source/Runtime/Engine/Private/GameEngine.cpp
--- a/Engine/Source/Runtime/Engine/Private/GameEngine.cpp
+++ b/Engine/Source/Runtime/Engine/Private/GameEngine.cpp
@@ -878,45 +878,9 @@ void UGameEngine::Tick( float DeltaSeconds, bool bIdleMode )
 		}
 
 		// Block on async loading if requested.
-		if( Context.World()->bRequestedBlockOnAsyncLoading )
+		if (Context.World()->bRequestedBlockOnAsyncLoading)
 		{
-			// Only perform work if there is anything to do. This ensures we are not syncronizing with the GPU
-			// and suspending the device needlessly.
-			bool bWorkToDo = IsAsyncLoading();
-			if (!bWorkToDo)
-			{
-				Context.World()->UpdateLevelStreaming();
-				bWorkToDo = Context.World()->IsVisibilityRequestPending();
-			}
-			if (bWorkToDo)
-			{
-				// tell clients to do the same so they don't fall behind
-				for( FConstPlayerControllerIterator Iterator = Context.World()->GetPlayerControllerIterator(); Iterator; ++Iterator )
-				{
-					APlayerController* PlayerController = *Iterator;
-					UNetConnection* Conn = Cast<UNetConnection>(PlayerController->Player);
-					if (Conn != NULL && Conn->GetUChildConnection() == NULL)
-					{
-						// call the event to replicate the call
-						PlayerController->ClientSetBlockOnAsyncLoading();
-						// flush the connection to make sure it gets sent immediately
-						Conn->FlushNet(true);
-					}
-				}
-
-				if( GameViewport && BeginStreamingPauseDelegate && BeginStreamingPauseDelegate->IsBound() )
-				{
-					BeginStreamingPauseDelegate->Execute( GameViewport->Viewport );
-				}	
-
-				// Flushes level streaming requests, blocking till completion.
-				Context.World()->FlushLevelStreaming();
-
-				if( EndStreamingPauseDelegate && EndStreamingPauseDelegate->IsBound() )
-				{
-					EndStreamingPauseDelegate->Execute( );
-				}	
-			}
+			BlockTillLevelStreamingCompleted(Context.World());
 			Context.World()->bRequestedBlockOnAsyncLoading = false;
 		}
 
