Merging development branch (CL 2365488) to //depot/UE4/... Engine only, mostly hierarchical foliage

diff --git a/Engine/Source/Runtime/Engine/Private/GameInstance.cpp b/Engine/Source/Runtime/Engine/Private/GameInstance.cpp
--- a/Engine/Source/Runtime/Engine/Private/GameInstance.cpp
+++ b/Engine/Source/Runtime/Engine/Private/GameInstance.cpp
@@ -121,33 +121,6 @@ bool UGameInstance::InitializePIE(bool bAnyBlueprintErrors, int32 PIEInstance)
 	return true;
 }
 
-
-/** Initializes streaming levels when PIE starts up based on the players start location */
-static void InitStreamingLevelsForPIEStartup(UGameViewportClient* GameViewportClient, UWorld* PlayWorld)
-{
-	check(PlayWorld);
-	ULocalPlayer* const Player = PlayWorld->GetFirstLocalPlayerFromController();
-	if (Player)
-	{
-		// Create a view family for the game viewport
-		FSceneViewFamilyContext ViewFamily(FSceneViewFamily::ConstructionValues(
-			GameViewportClient->Viewport,
-			PlayWorld->Scene,
-			GameViewportClient->EngineShowFlags)
-			.SetRealtimeUpdate(true));
-
-
-		// Calculate a view where the player is to update the streaming from the players start location
-		FVector ViewLocation;
-		FRotator ViewRotation;
-		Player->CalcSceneView(&ViewFamily, /*out*/ ViewLocation, /*out*/ ViewRotation, GameViewportClient->Viewport);
-
-		// Update level streaming.
-		PlayWorld->FlushLevelStreaming(&ViewFamily);
-	}
-}
-
-
 bool UGameInstance::StartPIEGameInstance(ULocalPlayer* LocalPlayer, bool bInSimulateInEditor, bool bAnyBlueprintErrors, bool bStartInSpectatorMode)
 {
 	UEditorEngine* const EditorEngine = CastChecked<UEditorEngine>(GetEngine());
@@ -209,17 +182,9 @@ bool UGameInstance::StartPIEGameInstance(ULocalPlayer* LocalPlayer, bool bInSimu
 			FMessageDialog::Open(EAppMsgType::Ok, NSLOCTEXT("UnrealEd", "Error_FailedCreateEditorPreviewWorld", "Failed to create editor preview world."));
 			return false;
 		}
-
-		// Navigate PIE world to the Editor world origin location
-		// and stream in all relevant levels around that location
-		if (PlayWorld->WorldComposition)
-		{
-			PlayWorld->NavigateTo(PlayWorld->OriginLocation);
-		}
-		else
-		{
-			PlayWorld->FlushLevelStreaming();
-		}
+		
+		// Make sure "always loaded" sub-levels are fully loaded
+		PlayWorld->FlushLevelStreaming(nullptr, EFlushLevelStreamingType::Visibility);
 
 		UNavigationSystem::InitializeForWorld(PlayWorld, LocalPlayers.Num() > 0 ? FNavigationSystem::PIEMode : FNavigationSystem::SimulationMode);
 		PlayWorld->CreateAISystem();
@@ -241,7 +206,7 @@ bool UGameInstance::StartPIEGameInstance(ULocalPlayer* LocalPlayer, bool bInSimu
 		if (GameViewport != NULL && GameViewport->Viewport != NULL)
 		{
 			// Stream any levels now that need to be loaded before the game starts
-			InitStreamingLevelsForPIEStartup(GameViewport, PlayWorld);
+			GEngine->BlockTillLevelStreamingCompleted(PlayWorld);
 		}
 		
 		if (PlayNetMode == PIE_ListenServer)
diff --git a/Engine/Source/Runtime/Engine/Private/GameMode.cpp b/Engine/Source/Runtime/Engine/Private/GameMode.cpp
--- a/Engine/Source/Runtime/Engine/Private/GameMode.cpp
+++ b/Engine/Source/Runtime/Engine/Private/GameMode.cpp
@@ -596,7 +596,7 @@ void AGameMode::HandleMatchHasStarted()
 	}
 
 	// Make sure level streaming is up to date before triggering NotifyMatchStarted
-	GetWorld()->FlushLevelStreaming();
+	GEngine->BlockTillLevelStreamingCompleted(GetWorld());
 
 	// First fire BeginPlay, if we haven't already in waiting to start match
 	GetWorldSettings()->NotifyBeginPlay();
diff --git a/Engine/Source/Runtime/Engine/Private/HierarchicalInstancedStaticMesh.cpp b/Engine/Source/Runtime/Engine/Private/HierarchicalInstancedStaticMesh.cpp
new file mode 100644
--- /dev/null
+++ b/Engine/Source/Runtime/Engine/Private/HierarchicalInstancedStaticMesh.cpp
@@ -0,0 +1,1655 @@
+// Copyright 1998-2014 Epic Games, Inc. All Rights Reserved.
+
+/*=============================================================================
+	InstancedStaticMesh.cpp: Static mesh rendering code.
+=============================================================================*/
+
+#include "EnginePrivate.h"
+#include "Components/InstancedStaticMeshComponent.h"
+#include "Components/HierarchicalInstancedStaticMeshComponent.h"
+#include "InstancedStaticMesh.h"
+#include "NavigationSystemHelpers.h"
+#include "AI/Navigation/NavCollision.h"
+#include "../../Renderer/Private/ScenePrivate.h"
+
+static TAutoConsoleVariable<int32> CVarFoliageSplitFactor(
+	TEXT("foliage.SplitFactor"),
+	16,
+	TEXT("This controls the branching factor of the foliage tree."));
+
+static TAutoConsoleVariable<int32> CVarForceLOD(
+	TEXT("foliage.ForceLOD"),
+	-1,
+	TEXT("If greater than or equal to zero, forces the foliage LOD to that level."));
+
+static TAutoConsoleVariable<int32> CVarDisableCull(
+	TEXT("foliage.DisableCull"),
+	0,
+	TEXT("If greater than zero, no culling occurs based on frustum."));
+
+static TAutoConsoleVariable<int32> CVarCullAll(
+	TEXT("foliage.CullAll"),
+	0,
+	TEXT("If greater than zero, everything is conside3red culled."));
+
+static TAutoConsoleVariable<int32> CVarMaxTrianglesToRender(
+	TEXT("foliage.MaxTrianglesToRender"),
+	100000000,
+	TEXT("This is an absolute limit on the number of foliage triangles to render in one traversal. This is used to prevent a silly LOD parameter mistake from causing the OS to kill the GPU."));
+
+static TAutoConsoleVariable<float> CVarFoliageOrthoGuardband(
+	TEXT("foliage.OrthoGuardBand"),
+	1000.0f,
+	TEXT("This controls the guardband for foliage culling in ortho and shadow views. In world units."));
+
+static TAutoConsoleVariable<float> CVarFoliageGuardband(
+	TEXT("foliage.GuardBand"),
+	5.0f,
+	TEXT("This controls the guardband for foliage culling in perspective views 5 gives a 5 degree increase in FOV."));
+
+TAutoConsoleVariable<float> CVarFoliageMinimumScreenSize(
+	TEXT("foliage.MinimumScreenSize"),
+	0.0f,
+	TEXT("This controls the screen size at which we cull foliage instances entirely."));
+
+TAutoConsoleVariable<float> CVarFoliageLODDistanceScale(
+	TEXT("foliage.LODDistanceScale"),
+	1.0f,
+	TEXT("Scale factor for the distance used in computing LOD for foliage."));
+
+
+DECLARE_CYCLE_STAT(TEXT("Traversal Time"),STAT_FoliageTraversalTime,STATGROUP_Foliage);
+DECLARE_CYCLE_STAT(TEXT("Build Time"), STAT_FoliageBuildTime, STATGROUP_Foliage);
+DECLARE_CYCLE_STAT(TEXT("Batch Time"),STAT_FoliageBatchTime,STATGROUP_Foliage);
+DECLARE_DWORD_COUNTER_STAT(TEXT("Runs"), STAT_FoliageRuns, STATGROUP_Foliage);
+DECLARE_DWORD_COUNTER_STAT(TEXT("Mesh Batches"), STAT_FoliageMeshBatches, STATGROUP_Foliage);
+DECLARE_DWORD_COUNTER_STAT(TEXT("Triangles"), STAT_FoliageTriangles, STATGROUP_Foliage);
+DECLARE_DWORD_COUNTER_STAT(TEXT("Instances"), STAT_FoliageInstances, STATGROUP_Foliage);
+DECLARE_DWORD_COUNTER_STAT(TEXT("Traversals"),STAT_FoliageTraversals,STATGROUP_Foliage);
+
+struct FClusterTree
+{
+	TArray<FClusterNode> Nodes;
+	TArray<int32> SortedInstances;
+	TArray<int32> InstanceReorderTable;
+};
+
+
+class FClusterBuilder
+{
+	int32 Num;
+	FBox InstBox;
+	int32 BranchingFactor;
+	int32 NumRoots;
+	TArray<int32> SortIndex;
+	TArray<FVector> SortPoints;
+	TArray<FMatrix> Transforms;
+
+	struct FRunPair
+	{
+		int32 Start;
+		int32 Num;
+
+		FRunPair(int32 InStart, int32 InNum)
+			: Start(InStart)
+			, Num(InNum)
+		{
+		}
+
+		bool operator< (const FRunPair& Other) const
+		{
+			return Start < Other.Start;
+		}
+	};
+	TArray<FRunPair> Clusters;
+
+	struct FSortPair
+	{
+		float d;
+		int32 Index;
+
+		bool operator< (const FSortPair& Other) const
+		{
+			return d < Other.d;
+		}
+	};
+	TArray<FSortPair> SortPairs;
+
+	void Split(int32 Num)
+	{
+		check(Num);
+		Clusters.Reset();
+		Split(0, Num - 1);
+		Clusters.Sort();
+		check(Clusters.Num());
+		int32 At = 0;
+		for (auto& Cluster : Clusters)
+		{
+			check(At == Cluster.Start);
+			At += Cluster.Num;
+		}
+		check(At == Num);
+	}
+
+	void Split(int32 Start, int32 End)
+	{
+		int32 NumRange = 1 + End - Start;
+		FBox ClusterBounds(ForceInit);
+		for (int32 Index = Start; Index <= End; Index++)
+		{
+			ClusterBounds += SortPoints[SortIndex[Index]];
+		}
+		if (NumRange <= BranchingFactor)
+		{
+			Clusters.Add(FRunPair(Start, NumRange));
+			return;
+		}
+		SortPairs.Reset();
+		int32 BestAxis = -1;
+		float BestAxisValue = -1.0f;
+		for (int32 Axis = 0; Axis < 3; Axis++)
+		{
+			float ThisAxisValue = ClusterBounds.Max[Axis] - ClusterBounds.Min[Axis];
+			if (!Axis || ThisAxisValue > BestAxisValue)
+			{
+				BestAxis = Axis;
+				BestAxisValue = ThisAxisValue;
+			}
+		}
+		for (int32 Index = Start; Index <= End; Index++)
+		{
+			FSortPair Pair;
+
+			Pair.Index = SortIndex[Index];
+			Pair.d = SortPoints[Pair.Index][BestAxis];
+			SortPairs.Add(Pair);
+		}
+		SortPairs.Sort();
+		for (int32 Index = Start; Index <= End; Index++)
+		{
+			SortIndex[Index] = SortPairs[Index - Start].Index;
+		}
+
+		int32 Half = NumRange / 2;
+
+		int32 EndLeft = Start + Half - 1;
+		int32 StartRight = 1 + End - Half;
+
+		if (NumRange & 1)
+		{
+			if (SortPairs[Half].d - SortPairs[Half - 1].d < SortPairs[Half + 1].d - SortPairs[Half].d)
+			{
+				EndLeft++;
+			}
+			else
+			{
+				StartRight--;
+			}
+		}
+		check(EndLeft + 1 == StartRight);
+		check(EndLeft >= Start);
+		check(End >= StartRight);
+
+		Split(Start, EndLeft);
+		Split(StartRight, End);
+
+	}
+public:
+	FClusterTree* Result;
+
+	FClusterBuilder(TArray<FMatrix>& InTransforms, const FBox& InInstBox)
+		: Num(InTransforms.Num())
+		, InstBox(InInstBox)
+		, Result(nullptr)
+	{
+		BranchingFactor = CVarFoliageSplitFactor.GetValueOnGameThread();
+
+		Exchange(Transforms, InTransforms);
+		SortIndex.AddUninitialized(Num);
+		SortPoints.AddUninitialized(Num);
+		for (int32 Index = 0; Index < Num; Index++)
+		{
+			SortIndex[Index] = Index;
+			SortPoints[Index] = Transforms[Index].GetOrigin();
+		}
+	}
+
+	void BuildAsync(ENamedThreads::Type CurrentThread, const FGraphEventRef& MyCompletionGraphEvent)
+	{
+		Build();
+		//FPlatformProcess::Sleep(5);
+	}
+
+	void Build()
+	{
+		Split(Num);
+
+		Result = new FClusterTree;
+		TArray<int32>& SortedInstances = Result->SortedInstances;
+		SortedInstances.AddUninitialized(Num);
+		for (int32 Index = 0; Index < Num; Index++)
+		{
+			SortedInstances[Index] = SortIndex[Index];
+		}
+		NumRoots = Clusters.Num();
+		Result->Nodes.Reserve(Clusters.Num());
+		for (int32 Index = 0; Index < NumRoots; Index++)
+		{
+			Result->Nodes.Add(FClusterNode());
+		}
+
+		for (int32 Index = 0; Index < NumRoots; Index++)
+		{
+			FClusterNode& Node = Result->Nodes[Index];
+			Node.FirstInstance = Clusters[Index].Start;
+			Node.LastInstance = Clusters[Index].Start + Clusters[Index].Num - 1;
+			FBox NodeBox(ForceInit);
+			for (int32 InstanceIndex = Node.FirstInstance; InstanceIndex <= Node.LastInstance; InstanceIndex++)
+			{
+				const FMatrix& ThisInstTrans = Transforms[SortedInstances[InstanceIndex]];
+				FBox ThisInstBox = InstBox.TransformBy(ThisInstTrans);
+				NodeBox += ThisInstBox;
+			}
+			Node.BoundMin = NodeBox.Min;
+			Node.BoundMax = NodeBox.Max;
+		}
+		TArray<int32> NodesPerLevel;
+		NodesPerLevel.Add(NumRoots);
+		int32 LOD = 0;
+		while (NumRoots > 1)
+		{
+			SortIndex.Reset();
+			SortPoints.Reset();
+			SortIndex.AddUninitialized(NumRoots);
+			SortPoints.AddUninitialized(NumRoots);
+			for (int32 Index = 0; Index < NumRoots; Index++)
+			{
+				SortIndex[Index] = Index;
+				FClusterNode& Node = Result->Nodes[Index];
+				SortPoints[Index] = (Node.BoundMin + Node.BoundMax) * 0.5f;
+			}
+			Split(NumRoots);
+
+			TArray<int32> InverseSortIndex;
+			InverseSortIndex.AddUninitialized(NumRoots);
+			for (int32 Index = 0; Index < NumRoots; Index++)
+			{
+				InverseSortIndex[SortIndex[Index]] = Index;
+			}
+
+			{
+				// rearrange the instances to match the new order of the old roots
+				TArray<int32> RemapSortIndex;
+				RemapSortIndex.AddUninitialized(Num);
+				int32 OutIndex = 0;
+				for (int32 Index = 0; Index < NumRoots; Index++)
+				{
+					FClusterNode& Node = Result->Nodes[SortIndex[Index]];
+					for (int32 InstanceIndex = Node.FirstInstance; InstanceIndex <= Node.LastInstance; InstanceIndex++)
+					{
+						RemapSortIndex[OutIndex++] = InstanceIndex;
+					}
+				}
+				TArray<int32> InverseInstanceIndex;
+				InverseInstanceIndex.AddUninitialized(Num);
+				for (int32 Index = 0; Index < Num; Index++)
+				{
+					InverseInstanceIndex[RemapSortIndex[Index]] = Index;
+				}
+				for (int32 Index = 0; Index < Result->Nodes.Num(); Index++)
+				{
+					FClusterNode& Node = Result->Nodes[Index];
+					Node.FirstInstance = InverseInstanceIndex[Node.FirstInstance];
+					Node.LastInstance = InverseInstanceIndex[Node.LastInstance];
+				}
+				TArray<int32> OldInstanceIndex;
+				Exchange(OldInstanceIndex, SortedInstances);
+				SortedInstances.AddUninitialized(Num);
+				for (int32 Index = 0; Index < Num; Index++)
+				{
+					SortedInstances[Index] = OldInstanceIndex[RemapSortIndex[Index]];
+				}
+			}
+			{
+				// rearrange the nodes to match the new order of the old roots
+				TArray<int32> RemapSortIndex;
+				int32 NewNum = Result->Nodes.Num() + Clusters.Num();
+				// RemapSortIndex[new index] == old index
+				RemapSortIndex.AddUninitialized(NewNum);
+				TArray<int32> LevelStarts;
+				LevelStarts.Add(Clusters.Num());
+				for (int32 Index = 0; Index < NodesPerLevel.Num() - 1; Index++)
+				{
+					LevelStarts.Add(LevelStarts[Index] + NodesPerLevel[Index]);
+				}
+
+				for (int32 Index = 0; Index < NumRoots; Index++)
+				{
+					FClusterNode& Node = Result->Nodes[SortIndex[Index]];
+					RemapSortIndex[LevelStarts[0]++] = SortIndex[Index];
+
+					int32 LeftIndex = Node.FirstChild;
+					int32 RightIndex = Node.LastChild;
+					int32 LevelIndex = 1;
+					while (RightIndex >= 0)
+					{
+						int32 NextLeftIndex = MAX_int32;
+						int32 NextRightIndex = -1;
+						for (int32 ChildIndex = LeftIndex; ChildIndex <= RightIndex; ChildIndex++)
+						{
+							RemapSortIndex[LevelStarts[LevelIndex]++] = ChildIndex;
+							int32 LeftChild = Result->Nodes[ChildIndex].FirstChild;
+							int32 RightChild = Result->Nodes[ChildIndex].LastChild;
+							if (LeftChild >= 0 && LeftChild <  NextLeftIndex)
+							{
+								NextLeftIndex = LeftChild;
+							}
+							if (RightChild >= 0 && RightChild >  NextRightIndex)
+							{
+								NextRightIndex = RightChild;
+							}
+						}
+						LeftIndex = NextLeftIndex;
+						RightIndex = NextRightIndex;
+						LevelIndex++;
+					}
+				}
+				check(LevelStarts[LevelStarts.Num() - 1] == NewNum);
+				TArray<int32> InverseChildIndex;
+				// InverseChildIndex[old index] == new index
+				InverseChildIndex.AddUninitialized(NewNum);
+				for (int32 Index = Clusters.Num(); Index < NewNum; Index++)
+				{
+					InverseChildIndex[RemapSortIndex[Index]] = Index;
+				}
+				for (int32 Index = 0; Index < Result->Nodes.Num(); Index++)
+				{
+					FClusterNode& Node = Result->Nodes[Index];
+					if (Node.FirstChild >= 0)
+					{
+						Node.FirstChild = InverseChildIndex[Node.FirstChild];
+						Node.LastChild = InverseChildIndex[Node.LastChild];
+					}
+				}
+				{
+					TArray<FClusterNode> OldNodes = Result->Nodes;
+					Result->Nodes.Empty(NewNum);
+					for (int32 Index = 0; Index < Clusters.Num(); Index++)
+					{
+						Result->Nodes.Add(FClusterNode());
+					}
+					Result->Nodes.AddUninitialized(OldNodes.Num());
+					for (int32 Index = 0; Index < OldNodes.Num(); Index++)
+					{
+						Result->Nodes[InverseChildIndex[Index]] = OldNodes[Index];
+					}
+				}
+				int32 OldIndex = Clusters.Num();
+				int32 InstanceTracker = 0;
+				for (int32 Index = 0; Index < Clusters.Num(); Index++)
+				{
+					FClusterNode& Node = Result->Nodes[Index];
+					Node.FirstChild = OldIndex;
+					OldIndex += Clusters[Index].Num;
+					Node.LastChild = OldIndex - 1;
+					Node.FirstInstance = Result->Nodes[Node.FirstChild].FirstInstance;
+					check(Node.FirstInstance == InstanceTracker);
+					Node.LastInstance = Result->Nodes[Node.LastChild].LastInstance;
+					InstanceTracker = Node.LastInstance + 1;
+					check(InstanceTracker <= Num);
+					FBox NodeBox(ForceInit);
+					for (int32 ChildIndex = Node.FirstChild; ChildIndex <= Node.LastChild; ChildIndex++)
+					{
+						FClusterNode& ChildNode = Result->Nodes[ChildIndex];
+						NodeBox += ChildNode.BoundMin;
+						NodeBox += ChildNode.BoundMax;
+					}
+					Node.BoundMin = NodeBox.Min;
+					Node.BoundMax = NodeBox.Max;
+				}
+				NumRoots = Clusters.Num();
+				NodesPerLevel.Insert(NumRoots, 0);
+			}
+		}
+
+		// Save inverse map
+		Result->InstanceReorderTable.AddUninitialized(Num);
+		for (int32 Index = 0; Index < Num; Index++)
+		{
+			Result->InstanceReorderTable[SortedInstances[Index]] = Index;
+		}
+	}
+
+};
+
+
+static bool PrintLevel(const FClusterTree& Tree, int32 NodeIndex, int32 Level, int32 CurrentLevel, int32 Parent)
+{
+	const FClusterNode& Node = Tree.Nodes[NodeIndex];
+	if (Level == CurrentLevel)
+	{
+		UE_LOG(LogConsoleResponse, Display, TEXT("Level %2d  Parent %3d"),
+			Level,
+			Parent
+			);
+		FVector Extent = Node.BoundMax - Node.BoundMin;
+		UE_LOG(LogConsoleResponse, Display, TEXT("    Bound (%5.1f, %5.1f, %5.1f) [(%5.1f, %5.1f, %5.1f) - (%5.1f, %5.1f, %5.1f)]"),
+			Extent.X, Extent.Y, Extent.Z, 
+			Node.BoundMin.X, Node.BoundMin.Y, Node.BoundMin.Z, 
+			Node.BoundMax.X, Node.BoundMax.Y, Node.BoundMax.Z
+			);
+		UE_LOG(LogConsoleResponse, Display, TEXT("    children %3d [%3d,%3d]   instances %3d [%3d,%3d]"),
+			(Node.FirstChild < 0) ? 0 : 1 + Node.LastChild - Node.FirstChild, Node.FirstChild, Node.LastChild,
+			1 + Node.LastInstance - Node.FirstInstance, Node.FirstInstance, Node.LastInstance
+			);
+		return true;
+	}
+	else if (Node.FirstChild < 0)
+	{
+		return false;
+	}
+	bool Ret = false;
+	for (int32 Child = Node.FirstChild; Child <= Node.LastChild; Child++)
+	{
+		Ret = PrintLevel(Tree, Child, Level, CurrentLevel + 1, NodeIndex) || Ret;
+	}
+	return Ret;
+}
+
+static void TestFoliage(const TArray<FString>& Args)
+{
+	UE_LOG(LogConsoleResponse, Display, TEXT("Running Foliage test."));
+	TArray<FInstancedStaticMeshInstanceData> Instances;
+
+	FMatrix Temp;
+	Temp.SetIdentity();
+	FRandomStream RandomStream(0x238946);
+	for (int32 i = 0; i < 1000; i++)
+	{
+		Instances.Add(FInstancedStaticMeshInstanceData());
+		Temp.SetOrigin(FVector(RandomStream.FRandRange(0.0f, 1.0f), RandomStream.FRandRange(0.0f, 1.0f), 0.0f) * 10000.0f);
+		Instances[i].Transform = Temp;
+	}
+
+	FBox TempBox(ForceInit);
+	TempBox += FVector(-100.0f, -100.0f, -100.0f);
+	TempBox += FVector(100.0f, 100.0f, 100.0f);
+
+	TArray<FMatrix> InstanceTransforms;
+	InstanceTransforms.AddUninitialized(Instances.Num());
+	for (int32 Index = 0; Index < Instances.Num(); Index++)
+	{
+		InstanceTransforms[Index] = Instances[Index].Transform;
+	}
+	FClusterBuilder Builder(InstanceTransforms, TempBox);
+	Builder.Build();
+
+	int32 Level = 0;
+
+	UE_LOG(LogConsoleResponse, Display, TEXT("-----"));
+
+	while(PrintLevel(*Builder.Result, 0, Level++, 0, -1))
+	{
+	}
+
+	delete Builder.Result;
+}
+
+static FAutoConsoleCommand TestFoliageCmd(
+	TEXT("foliage.Test"),
+	TEXT("Useful for debugging."),
+	FConsoleCommandWithArgsDelegate::CreateStatic(&TestFoliage)
+	);
+
+#if !(UE_BUILD_SHIPPING || UE_BUILD_TEST)
+	static uint32 GDebugTag = 1;
+	static uint32 GCaptureDebugRuns = 0;
+#endif
+
+static void FreezeFoliageCulling(const TArray<FString>& Args)
+{
+#if !(UE_BUILD_SHIPPING || UE_BUILD_TEST)
+	UE_LOG(LogConsoleResponse, Display, TEXT("Freezing Foliage Culling."));
+	GDebugTag++;
+	GCaptureDebugRuns = GDebugTag;
+#endif
+}
+
+static FAutoConsoleCommand FreezeFoliageCullingCmd(
+	TEXT("foliage.Freeze"),
+	TEXT("Useful for debugging. Freezes the foliage culling and LOD."),
+	FConsoleCommandWithArgsDelegate::CreateStatic(&FreezeFoliageCulling)
+	);
+
+static void UnFreezeFoliageCulling(const TArray<FString>& Args)
+{
+#if !(UE_BUILD_SHIPPING || UE_BUILD_TEST)
+	UE_LOG(LogConsoleResponse, Display, TEXT("Unfreezing Foliage Culling."));
+	GDebugTag++;
+	GCaptureDebugRuns = 0;
+#endif
+}
+
+static FAutoConsoleCommand UnFreezeFoliageCullingCmd(
+	TEXT("foliage.UnFreeze"),
+	TEXT("Useful for debugging. Freezes the foliage culling and LOD."),
+	FConsoleCommandWithArgsDelegate::CreateStatic(&UnFreezeFoliageCulling)
+	);
+
+
+
+
+struct FFoliageElementParams;
+struct FFoliageRenderInstanceParams;
+struct FFoliageCullInstanceParams;
+
+class FHierarchicalStaticMeshSceneProxy : public FInstancedStaticMeshSceneProxy
+{
+	const TArray<FClusterNode>& ClusterTree;
+
+	int32 FirstUnbuiltIndex;
+	int32 LastUnbuiltIndex;
+
+#if !(UE_BUILD_SHIPPING || UE_BUILD_TEST)
+	mutable TArray<uint32> DebugRuns[MAX_STATIC_MESH_LODS];
+	mutable int32 DebugTotalInstances[MAX_STATIC_MESH_LODS];
+	mutable int32 CaptureTag;
+#endif
+
+public:
+
+	FHierarchicalStaticMeshSceneProxy(UHierarchicalInstancedStaticMeshComponent* InComponent, ERHIFeatureLevel::Type InFeatureLevel)
+		: FInstancedStaticMeshSceneProxy(InComponent, InFeatureLevel)
+		, ClusterTree(InComponent->ClusterTree)
+		, FirstUnbuiltIndex(InComponent->NumBuiltInstances)
+		, LastUnbuiltIndex(InComponent->PerInstanceSMData.Num()+InComponent->RemovedInstances.Num()-1)
+#if !(UE_BUILD_SHIPPING || UE_BUILD_TEST)
+		, CaptureTag(0)
+#endif
+	{
+		check(InComponent->InstanceReorderTable.Num() == InComponent->PerInstanceSMData.Num());
+	}
+
+	// FPrimitiveSceneProxy interface.
+	
+	virtual FPrimitiveViewRelevance GetViewRelevance(const FSceneView* View) override
+	{
+		FPrimitiveViewRelevance Result;
+		if(View->Family->EngineShowFlags.InstancedStaticMeshes)
+		{
+			Result = FStaticMeshSceneProxy::GetViewRelevance(View);
+			Result.bDynamicRelevance = true;
+			Result.bStaticRelevance = false;
+		}
+		return Result;
+	}
+
+	virtual void GetDynamicMeshElements(const TArray<const FSceneView*>& Views, const FSceneViewFamily& ViewFamily, uint32 VisibilityMap, FMeshElementCollector& Collector) const override;
+
+	virtual void DrawStaticElements(FStaticPrimitiveDrawInterface* PDI) override
+	{
+	}
+
+
+	virtual int32 GetNumMeshBatches() const override
+	{
+		check(0);
+		return 0;
+	}
+
+	/** Sets up a shadow FMeshBatch for a specific LOD. */
+	virtual bool GetShadowMeshElement(int32 LODIndex, int32 BatchIndex, uint8 InDepthPriorityGroup, FMeshBatch& OutMeshBatch) const override
+	{
+		check(0);
+		return false;
+	}
+
+	/** Sets up a FMeshBatch for a specific LOD and element. */
+	virtual bool GetMeshElement(int32 LODIndex, int32 BatchIndex, int32 ElementIndex, uint8 InDepthPriorityGroup, const bool bUseSelectedMaterial, const bool bUseHoveredMaterial, FMeshBatch& OutMeshBatch) const override
+	{
+		check(0);
+		return false;
+	}
+
+
+	/** Sets up a wireframe FMeshBatch for a specific LOD. */
+	virtual bool GetWireframeMeshElement(int32 LODIndex, int32 BatchIndex, const FMaterialRenderProxy* WireframeRenderProxy, uint8 InDepthPriorityGroup, FMeshBatch& OutMeshBatch) const override
+	{
+		check(0);
+		return false;
+	}
+
+	void FillDynamicMeshElements(FMeshElementCollector& Collector, const FFoliageElementParams& ElementParams, const FFoliageRenderInstanceParams& Instances) const;
+
+	template<bool TUseVector>
+	void Traverse(const FFoliageCullInstanceParams& Params, int32 Index, int32 MinLOD, int32 MaxLOD, bool bFullyContained = false) const;
+};
+
+struct FFoliageRenderInstanceParams
+{
+	mutable TArray<uint32, SceneRenderingAllocator> Runs[MAX_STATIC_MESH_LODS];
+	mutable int32 TotalInstances[MAX_STATIC_MESH_LODS];
+
+	FFoliageRenderInstanceParams()
+	{
+		for (int32 Index = 0; Index < MAX_STATIC_MESH_LODS; Index++)
+		{
+			TotalInstances[Index] = 0;
+		}
+	}
+	FORCEINLINE_DEBUGGABLE void AddRun(int32 Lod, int32 FirstInstance, int32 LastInstance) const
+	{
+		TotalInstances[Lod] += 1 + LastInstance - FirstInstance;
+		if (Runs[Lod].Num() && Runs[Lod].Last() + 1 == FirstInstance)
+		{
+			Runs[Lod].Last() = (uint32)LastInstance;
+		}
+		else
+		{
+			Runs[Lod].Add((uint32)FirstInstance);
+			Runs[Lod].Add((uint32)LastInstance);
+		}
+	}
+
+	FORCEINLINE_DEBUGGABLE void AddRun(int32 Lod, const FClusterNode& Node) const
+	{
+		AddRun(Lod, Node.FirstInstance, Node.LastInstance);
+	}
+};
+
+struct FFoliageCullInstanceParams : public FFoliageRenderInstanceParams
+{
+	FConvexVolume ViewFrustumLocal;
+	FConvexVolume ViewFrustumGuardBandLocal;
+	const TArray<FClusterNode>& Tree;
+	const FSceneView* View;
+	FVector ViewOriginInLocalZero;
+	FVector ViewOriginInLocalOne;
+	int32 LODs;
+	float LODPlanesMax[MAX_STATIC_MESH_LODS];
+	float LODPlanesMin[MAX_STATIC_MESH_LODS];
+
+
+	FFoliageCullInstanceParams(const TArray<FClusterNode>& InTree)
+	:	Tree(InTree)
+	{
+	}
+};
+
+static bool GUseVectorCull = true;
+
+static void ToggleUseVectorCull(const TArray<FString>& Args)
+{
+	GUseVectorCull = !GUseVectorCull;
+}
+
+static FAutoConsoleCommand ToggleUseVectorCullCmd(
+	TEXT("foliage.ToggleVectorCull"),
+	TEXT("Useful for debugging. Toggles the optimized cull."),
+	FConsoleCommandWithArgsDelegate::CreateStatic(&ToggleUseVectorCull)
+	);
+
+const VectorRegister		VECTOR_HALF_HALF_HALF_ZERO				= DECLARE_VECTOR_REGISTER(0.5f, 0.5f, 0.5f, 0.0f);
+
+template<bool TUseVector>
+static FORCEINLINE_DEBUGGABLE bool CullNode(const FFoliageCullInstanceParams& Params, const FClusterNode& Node, bool& bOutFullyContained)
+{
+	if (TUseVector)
+	{
+		checkSlow(Params.ViewFrustumLocal.PermutedPlanes.Num() == 4);
+		checkSlow(Params.ViewFrustumGuardBandLocal.PermutedPlanes.Num() == 4);
+
+		//@todo, once we have more than one mesh per tree, these should be aligned
+		VectorRegister BoxMin = VectorLoad(&Node.BoundMin);
+		VectorRegister BoxMax = VectorLoad(&Node.BoundMax);
+
+		VectorRegister BoxDiff = VectorSubtract(BoxMax,BoxMin);
+		VectorRegister BoxSum = VectorAdd(BoxMax,BoxMin);
+
+		// Load the origin & extent
+		VectorRegister Orig = VectorMultiply(VECTOR_HALF_HALF_HALF_ZERO, BoxSum);
+		VectorRegister Ext = VectorMultiply(VECTOR_HALF_HALF_HALF_ZERO, BoxDiff);
+		// Splat origin into 3 vectors
+		VectorRegister OrigX = VectorReplicate(Orig, 0);
+		VectorRegister OrigY = VectorReplicate(Orig, 1);
+		VectorRegister OrigZ = VectorReplicate(Orig, 2);
+		// Splat the abs for the pushout calculation
+		VectorRegister AbsExtentX = VectorReplicate(Ext, 0);
+		VectorRegister AbsExtentY = VectorReplicate(Ext, 1);
+		VectorRegister AbsExtentZ = VectorReplicate(Ext, 2);
+		// Since we are moving straight through get a pointer to the data
+		const FPlane* RESTRICT PermutedPlanePtr = (FPlane*)Params.ViewFrustumLocal.PermutedPlanes.GetData();
+		// Process four planes at a time until we have < 4 left
+		// Load 4 planes that are already all Xs, Ys, ...
+		VectorRegister PlanesX = VectorLoadAligned(PermutedPlanePtr);
+		PermutedPlanePtr++;
+		VectorRegister PlanesY = VectorLoadAligned(PermutedPlanePtr);
+		PermutedPlanePtr++;
+		VectorRegister PlanesZ = VectorLoadAligned(PermutedPlanePtr);
+		PermutedPlanePtr++;
+		VectorRegister PlanesW = VectorLoadAligned(PermutedPlanePtr);
+		PermutedPlanePtr++;
+		// Calculate the distance (x * x) + (y * y) + (z * z) - w
+		VectorRegister DistX = VectorMultiply(OrigX,PlanesX);
+		VectorRegister DistY = VectorMultiplyAdd(OrigY,PlanesY,DistX);
+		VectorRegister DistZ = VectorMultiplyAdd(OrigZ,PlanesZ,DistY);
+		VectorRegister Distance = VectorSubtract(DistZ,PlanesW);
+		// Now do the push out FMath::Abs(x * x) + FMath::Abs(y * y) + FMath::Abs(z * z)
+		VectorRegister PushX = VectorMultiply(AbsExtentX,VectorAbs(PlanesX));
+		VectorRegister PushY = VectorMultiplyAdd(AbsExtentY,VectorAbs(PlanesY),PushX);
+		VectorRegister PushOut = VectorMultiplyAdd(AbsExtentZ,VectorAbs(PlanesZ),PushY);
+
+		const FPlane* RESTRICT PermutedGuardPlanePtr = (FPlane*)Params.ViewFrustumGuardBandLocal.PermutedPlanes.GetData();
+		PlanesX = VectorLoadAligned(PermutedGuardPlanePtr);
+		PermutedGuardPlanePtr++;
+		PlanesY = VectorLoadAligned(PermutedGuardPlanePtr);
+		PermutedGuardPlanePtr++;
+		PlanesZ = VectorLoadAligned(PermutedGuardPlanePtr);
+		PermutedGuardPlanePtr++;
+		PlanesW = VectorLoadAligned(PermutedGuardPlanePtr);
+		PermutedGuardPlanePtr++;
+
+		// Calculate the distance (x * x) + (y * y) + (z * z) - w
+		DistX = VectorMultiply(OrigX,PlanesX);
+		DistY = VectorMultiplyAdd(OrigY,PlanesY,DistX);
+		DistZ = VectorMultiplyAdd(OrigZ,PlanesZ,DistY);
+		VectorRegister DistanceGuard = VectorSubtract(DistZ,PlanesW);
+		// Now do the push out FMath::Abs(x * x) + FMath::Abs(y * y) + FMath::Abs(z * z)
+		PushX = VectorMultiply(AbsExtentX,VectorAbs(PlanesX));
+		PushY = VectorMultiplyAdd(AbsExtentY,VectorAbs(PlanesY),PushX);
+		VectorRegister PushOutGuard = VectorMultiplyAdd(AbsExtentZ,VectorAbs(PlanesZ),PushY);
+		VectorRegister PushOutNegative = VectorNegate(PushOutGuard);
+
+		// Definitely inside frustums, but check to see if it's fully contained in the guard band
+		bOutFullyContained = !VectorAnyGreaterThan(DistanceGuard,PushOutNegative);
+		// Check for completely outside
+		return !!VectorAnyGreaterThan(Distance,PushOut);
+	}
+	FVector Center = (Node.BoundMin + Node.BoundMax) * 0.5f;
+	FVector Extent = (Node.BoundMax - Node.BoundMin) * 0.5f;
+	if (!Params.ViewFrustumLocal.IntersectBox(Center, Extent) ||
+		!Params.ViewFrustumGuardBandLocal.IntersectBox(Center, Extent, bOutFullyContained)) // only roundoff error can cause a cull here, but whatever
+	{
+		return true;
+	}
+	return false;
+}
+
+template<bool TUseVector>
+void FHierarchicalStaticMeshSceneProxy::Traverse(const FFoliageCullInstanceParams& Params, int32 Index, int32 MinLOD, int32 MaxLOD, bool bFullyContained) const
+{
+	const FClusterNode& Node = Params.Tree[Index];
+	if (!bFullyContained)
+	{
+		if (CullNode<TUseVector>(Params, Node, bFullyContained))
+		{
+			return;
+		}
+	}
+
+	if (MinLOD != MaxLOD)
+	{
+		FVector Center = (Node.BoundMax + Node.BoundMin) * 0.5f;
+		float DistCenterZero = FVector::Dist(Center, Params.ViewOriginInLocalZero);
+		float DistCenterOne = FVector::Dist(Center, Params.ViewOriginInLocalOne);
+		float HalfWidth = FVector::Dist(Node.BoundMax, Node.BoundMin) * 0.5f;
+		float NearDot = FMath::Min(DistCenterZero, DistCenterOne) - HalfWidth;
+		float FarDot = FMath::Max(DistCenterZero, DistCenterOne) + HalfWidth;
+
+		while (MaxLOD > MinLOD &&  NearDot > Params.LODPlanesMax[MinLOD])
+		{
+			MinLOD++;
+		}
+		if (MinLOD >= Params.LODs)
+		{
+			return;
+		}
+		while (MaxLOD > MinLOD && FarDot < Params.LODPlanesMin[MaxLOD - 1])
+		{
+			MaxLOD--;
+		}
+	}
+
+	bool bSplit = (!bFullyContained || MinLOD < MaxLOD) && Node.FirstChild >= 0;
+
+	if (!bSplit)
+	{
+		MaxLOD = FMath::Min(MaxLOD, Params.LODs - 1);
+		for (int32 LODIndex = MinLOD; LODIndex <= MaxLOD; LODIndex++)
+		{
+			Params.AddRun(LODIndex, Node);
+		}
+		return;
+	}
+	for (int32 ChildIndex = Node.FirstChild; ChildIndex <= Node.LastChild; ChildIndex++)
+	{
+		Traverse<TUseVector>(Params, ChildIndex, MinLOD, MaxLOD, bFullyContained);
+	}
+}
+
+struct FFoliageElementParams
+{
+	const FInstancingUserData* PassUserData[2];
+	int32 NumSelectionGroups;
+	const FSceneView* View;
+	int32 ViewIndex;
+	bool bSelectionRenderEnabled;
+	bool BatchRenderSelection[2];
+	bool bIsWireframe;
+	bool bUseHoveredMaterial;
+	bool bInstanced;
+	bool bBlendLODs;
+};
+
+void FHierarchicalStaticMeshSceneProxy::FillDynamicMeshElements(FMeshElementCollector& Collector, const FFoliageElementParams& ElementParams, const FFoliageRenderInstanceParams& Params) const
+{
+	SCOPE_CYCLE_COUNTER(STAT_FoliageBatchTime);
+	int64 TotalTriangles = 0;
+	for (int32 LODIndex = 0; LODIndex < InstancedRenderData.VertexFactories.Num(); LODIndex++)
+	{
+		if (!Params.Runs[LODIndex].Num())
+		{
+			continue;
+		}
+		const FStaticMeshLODResources& LODModel = StaticMesh->RenderData->LODResources[LODIndex];
+
+		for (int32 SelectionGroupIndex = 0; SelectionGroupIndex < ElementParams.NumSelectionGroups; SelectionGroupIndex++)
+		{
+			for (int32 SectionIndex = 0; SectionIndex < LODModel.Sections.Num(); SectionIndex++)
+			{
+
+				int32 NumBatches = 1;
+				int32 CurrentRun = 0;
+				int32 CurrentInstance = 0;
+				int32 RemainingInstances = Params.TotalInstances[LODIndex];
+
+				if (!ElementParams.bInstanced)
+				{
+					NumBatches = FMath::DivideAndRoundUp(RemainingInstances, (int32)FInstancedStaticMeshVertexFactory::NumBitsForVisibilityMask());
+					if (NumBatches)
+					{
+						check(Params.Runs[LODIndex].Num());
+						CurrentInstance = Params.Runs[LODIndex][CurrentRun];
+					}
+				}
+
+#if STATS
+				INC_DWORD_STAT_BY(STAT_FoliageInstances, RemainingInstances);
+				if (!ElementParams.bInstanced)
+				{
+					INC_DWORD_STAT_BY(STAT_FoliageRuns, NumBatches);
+				}
+#endif
+				bool bDidStats = false;
+				for (int32 BatchIndex = 0; BatchIndex < NumBatches; BatchIndex++)
+				{
+					FMeshBatch& MeshElement = Collector.AllocateMesh();
+					INC_DWORD_STAT(STAT_FoliageMeshBatches);
+
+					if (!FStaticMeshSceneProxy::GetMeshElement(LODIndex, 0, SectionIndex, GetDepthPriorityGroup(ElementParams.View), ElementParams.BatchRenderSelection[SelectionGroupIndex], ElementParams.bUseHoveredMaterial, MeshElement))
+					{
+						continue;
+					}
+					checkSlow(MeshElement.GetNumPrimitives() > 0);
+
+					MeshElement.VertexFactory = &InstancedRenderData.VertexFactories[LODIndex];
+					FMeshBatchElement& BatchElement0 = MeshElement.Elements[0];
+					BatchElement0.UserData = (void*)&UserData_AllInstances;
+
+					BatchElement0.UserData = ElementParams.PassUserData[SelectionGroupIndex];
+					BatchElement0.MaxScreenSize = 1.0;
+					BatchElement0.MinScreenSize = 0.0;
+					BatchElement0.InstancedLODIndex = LODIndex;
+					BatchElement0.InstancedLODRange = ElementParams.bBlendLODs ? 1 : 0;
+					MeshElement.bCanApplyViewModeOverrides = true;
+					MeshElement.bUseSelectionOutline = ElementParams.BatchRenderSelection[SelectionGroupIndex];
+					MeshElement.bUseWireframeSelectionColoring = ElementParams.BatchRenderSelection[SelectionGroupIndex];
+					if (!bDidStats)
+					{
+						bDidStats = true;
+						int64 Tris = int64(RemainingInstances) * int64(BatchElement0.NumPrimitives);
+						TotalTriangles += Tris;
+					}
+					if (ElementParams.bInstanced)
+					{
+						BatchElement0.NumInstances = Params.Runs[LODIndex].Num() / 2;
+						BatchElement0.InstanceRuns = &Params.Runs[LODIndex][0];
+#if STATS
+						INC_DWORD_STAT_BY(STAT_FoliageRuns, BatchElement0.NumInstances);
+#endif
+					}
+					else
+					{
+						uint32 NumInstancesThisBatch = FMath::Min(RemainingInstances, (int32)FInstancedStaticMeshVertexFactory::NumBitsForVisibilityMask());
+
+						MeshElement.Elements.Reserve(NumInstancesThisBatch);
+						check(NumInstancesThisBatch);
+
+						for (uint32 Instance = 0; Instance < NumInstancesThisBatch; ++Instance)
+						{
+							FMeshBatchElement* NewBatchElement; 
+
+							if (Instance == 0)
+							{
+								NewBatchElement = &MeshElement.Elements[0];
+							}
+							else
+							{
+								NewBatchElement = new(MeshElement.Elements) FMeshBatchElement();
+								*NewBatchElement = MeshElement.Elements[0];
+							}
+
+							NewBatchElement->UserIndex = CurrentInstance;
+							if (--RemainingInstances)
+							{
+								if ((uint32)CurrentInstance >= Params.Runs[LODIndex][CurrentRun + 1])
+								{
+									CurrentRun += 2;
+									check(CurrentRun + 1 < Params.Runs[LODIndex].Num());
+									CurrentInstance = Params.Runs[LODIndex][CurrentRun];
+								}
+								else
+								{
+									CurrentInstance++;
+								}
+							}
+						}
+					}
+					if (TotalTriangles < (int64)CVarMaxTrianglesToRender.GetValueOnRenderThread())
+					{
+						Collector.AddMesh(ElementParams.ViewIndex, MeshElement);
+					}
+				}
+			}
+		}
+	}
+#if STATS
+	TotalTriangles = FMath::Min<int64>(TotalTriangles, MAX_int32);
+	INC_DWORD_STAT_BY(STAT_FoliageTriangles, (uint32)TotalTriangles);
+	INC_DWORD_STAT_BY(STAT_StaticMeshTriangles, (uint32)TotalTriangles);
+#endif
+}
+
+
+
+void FHierarchicalStaticMeshSceneProxy::GetDynamicMeshElements(const TArray<const FSceneView*>& Views, const FSceneViewFamily& ViewFamily, uint32 VisibilityMap, FMeshElementCollector& Collector) const
+{
+	QUICK_SCOPE_CYCLE_COUNTER(STAT_HierarchicalInstancedStaticMeshSceneProxy_GetMeshElements);
+
+	for (int32 ViewIndex = 0; ViewIndex < Views.Num(); ViewIndex++)
+	{
+		if (VisibilityMap & (1 << ViewIndex))
+		{
+			const FSceneView* View = Views[ViewIndex];
+
+			FFoliageElementParams ElementParams;
+			ElementParams.bSelectionRenderEnabled = GIsEditor && ViewFamily.EngineShowFlags.Selection;
+			ElementParams.NumSelectionGroups = (ElementParams.bSelectionRenderEnabled && bHasSelectedInstances) ? 2 : 1;
+			ElementParams.PassUserData[0] = bHasSelectedInstances && ElementParams.bSelectionRenderEnabled ? &UserData_SelectedInstances : &UserData_AllInstances;
+			ElementParams.PassUserData[1] = &UserData_DeselectedInstances;
+			ElementParams.BatchRenderSelection[0] = ElementParams.bSelectionRenderEnabled && IsSelected();
+			ElementParams.BatchRenderSelection[1] = false;
+			ElementParams.bIsWireframe = ViewFamily.EngineShowFlags.Wireframe;
+			ElementParams.bUseHoveredMaterial = IsHovered();
+			ElementParams.bInstanced = RHISupportsInstancing(GetFeatureLevelShaderPlatform(InstancedRenderData.FeatureLevel));
+			ElementParams.ViewIndex = ViewIndex;
+			ElementParams.View = View;
+
+			// Render built instances
+			{
+				FFoliageCullInstanceParams InstanceParams(ClusterTree);
+				InstanceParams.LODs = RenderData->LODResources.Num();
+
+				InstanceParams.View = View;
+
+				FMatrix WorldToLocal = GetLocalToWorld().Inverse();
+				bool bUseVectorCull = GUseVectorCull;
+
+				bool bDisableCull = !!CVarDisableCull.GetValueOnRenderThread();
+				if (View->ViewMatrices.GetDynamicMeshElementsShadowCullFrustum)
+				{
+					float GuardBand = CVarFoliageOrthoGuardband.GetValueOnRenderThread();
+					for (int32 Index = 0; Index < View->ViewMatrices.GetDynamicMeshElementsShadowCullFrustum->Planes.Num(); Index++)
+					{
+						FPlane Src = View->ViewMatrices.GetDynamicMeshElementsShadowCullFrustum->Planes[Index];
+						FPlane Norm = Src / Src.Size();
+						// remove world space preview translation
+						Norm.W -= (FVector(Norm) | View->ViewMatrices.PreShadowTranslation);
+						FPlane Local = Norm.TransformBy(WorldToLocal);
+						FPlane LocalNorm = Local / Local.Size();
+						InstanceParams.ViewFrustumLocal.Planes.Add(LocalNorm);
+						LocalNorm.W += GuardBand;
+						InstanceParams.ViewFrustumGuardBandLocal.Planes.Add(LocalNorm);
+					}
+					bUseVectorCull = InstanceParams.ViewFrustumLocal.Planes.Num() == 4;
+				}
+				else
+				{
+					FMatrix LocalViewProjForCulling = GetLocalToWorld() * View->ViewMatrices.GetViewProjMatrix();
+
+					GetViewFrustumBounds(InstanceParams.ViewFrustumLocal, LocalViewProjForCulling, false);
+
+
+					if (View->ViewMatrices.IsPerspectiveProjection())
+					{
+						InstanceParams.ViewFrustumLocal.Planes.Pop(false); // we don't want the far plane either
+						FMatrix ThreePlanes;
+						ThreePlanes.SetIdentity();
+						ThreePlanes.SetAxes(&InstanceParams.ViewFrustumLocal.Planes[0], &InstanceParams.ViewFrustumLocal.Planes[1], &InstanceParams.ViewFrustumLocal.Planes[2]);
+						FVector ProjectionOrigin = ThreePlanes.Inverse().GetTransposed().TransformVector(FVector(InstanceParams.ViewFrustumLocal.Planes[0].W, InstanceParams.ViewFrustumLocal.Planes[1].W, InstanceParams.ViewFrustumLocal.Planes[2].W));
+
+						float GuardBand = FMath::Tan(FMath::DegreesToRadians(CVarFoliageGuardband.GetValueOnRenderThread()));
+						FVector Forward = LocalViewProjForCulling.GetColumn(3).GetSafeNormal();
+						for (int32 Index = 0; Index < InstanceParams.ViewFrustumLocal.Planes.Num(); Index++)
+						{
+							FPlane Src = InstanceParams.ViewFrustumLocal.Planes[Index];
+							FVector Normal = Src.GetSafeNormal();
+							InstanceParams.ViewFrustumLocal.Planes[Index] = FPlane(Normal, Normal | ProjectionOrigin);
+							FVector Perp = ((Normal ^ Forward).GetSafeNormal() ^ Normal).GetSafeNormal();
+							if ((Perp | Forward) > 0.0f)
+							{
+								Perp *= -1.0f;
+							}
+
+							float BlendNormal = FMath::Sqrt(1.0f - GuardBand * GuardBand);
+							 
+							Normal = ((Normal * BlendNormal) + (Perp * GuardBand)).GetSafeNormal();
+
+							InstanceParams.ViewFrustumGuardBandLocal.Planes.Add(FPlane(Normal, Normal | ProjectionOrigin));
+						}
+					}
+					else
+					{
+						bUseVectorCull = false;
+						float GuardBand = CVarFoliageOrthoGuardband.GetValueOnRenderThread();
+						for (int32 Index = 0; Index < InstanceParams.ViewFrustumLocal.Planes.Num(); Index++)
+						{
+							FPlane Src = InstanceParams.ViewFrustumLocal.Planes[Index];
+							Src.W += GuardBand;
+							InstanceParams.ViewFrustumGuardBandLocal.Planes.Add(Src);
+						}
+					}
+				}
+				if (!InstanceParams.ViewFrustumLocal.Planes.Num() || !InstanceParams.ViewFrustumGuardBandLocal.Planes.Num())
+				{
+					bDisableCull = true;
+				}
+				else
+				{
+					InstanceParams.ViewFrustumLocal.Init();
+					InstanceParams.ViewFrustumGuardBandLocal.Init();
+				}
+
+				InstanceParams.ViewOriginInLocalZero = WorldToLocal.TransformPosition(View->GetTemporalLODOrigin(0));
+				InstanceParams.ViewOriginInLocalOne = WorldToLocal.TransformPosition(View->GetTemporalLODOrigin(1));
+
+				ElementParams.bBlendLODs = View->GetTemporalLODActive();
+				float MinSize = CVarFoliageMinimumScreenSize.GetValueOnRenderThread();
+				float LODScale = CVarFoliageLODDistanceScale.GetValueOnRenderThread();
+				float MaxDrawDistanceScale = GetCachedScalabilityCVars().ViewDistanceScale;
+				float SphereRadius = RenderData->Bounds.SphereRadius;
+
+				for (int32 LODIndex = 0; LODIndex < InstanceParams.LODs; LODIndex++)
+				{
+					InstanceParams.LODPlanesMax[LODIndex] = MIN_flt;
+					InstanceParams.LODPlanesMin[LODIndex] = MAX_flt;
+				}
+				for (int32 SampleIndex = 0; SampleIndex < 2; SampleIndex++)
+				{
+					// we aren't going to accurately figure out LOD planes for both samples and try all 4 possibilities, rather we just consider the worst case LOD range; the shaders will render correctly
+					float Fac = View->GetTemporalLODDistanceFactor(SampleIndex) * SphereRadius * SphereRadius * LODScale * LODScale;
+
+					float FinalCull = MAX_flt;
+
+					if (MinSize > 0.0)
+					{
+						FinalCull = FMath::Sqrt(Fac / MinSize);
+					}
+					if (UserData_AllInstances.EndCullDistance > 0.0f && UserData_AllInstances.EndCullDistance < FinalCull)
+					{
+						FinalCull = UserData_AllInstances.EndCullDistance;
+					}
+					FinalCull *= MaxDrawDistanceScale;
+
+					for (int32 LODIndex = 1; LODIndex < InstanceParams.LODs; LODIndex++)
+					{
+						float Val = FMath::Min(FMath::Sqrt(Fac / RenderData->ScreenSize[LODIndex]), FinalCull);
+						InstanceParams.LODPlanesMin[LODIndex - 1] = FMath::Min(InstanceParams.LODPlanesMin[LODIndex - 1], Val);
+						InstanceParams.LODPlanesMax[LODIndex - 1] = FMath::Max(InstanceParams.LODPlanesMax[LODIndex - 1], Val);
+					}
+					InstanceParams.LODPlanesMin[InstanceParams.LODs - 1] = FMath::Min(InstanceParams.LODPlanesMin[InstanceParams.LODs - 1], FinalCull);
+					InstanceParams.LODPlanesMax[InstanceParams.LODs - 1] = FMath::Max(InstanceParams.LODPlanesMax[InstanceParams.LODs - 1], FinalCull);
+				}
+
+				INC_DWORD_STAT(STAT_FoliageTraversals);
+#if !(UE_BUILD_SHIPPING || UE_BUILD_TEST)
+				if (GCaptureDebugRuns == GDebugTag && CaptureTag == GDebugTag)
+				{
+					for (int32 LODIndex = 0; LODIndex < InstanceParams.LODs; LODIndex++)
+					{
+						for (int32 Run = 0; Run < DebugRuns[LODIndex].Num(); Run++)
+						{
+							InstanceParams.Runs[LODIndex].Add(DebugRuns[LODIndex][Run]);
+						}
+						InstanceParams.TotalInstances[LODIndex] = DebugTotalInstances[LODIndex];
+					}
+				}
+				else
+#endif
+				{
+					SCOPE_CYCLE_COUNTER(STAT_FoliageTraversalTime);
+
+					// validate that the bounding box is layed out correctly in memory
+					check(&((const FVector4*)&ClusterTree[0].BoundMin)[1] == (const FVector4*)&ClusterTree[0].BoundMax);
+					//check(UPTRINT(&ClusterTree[0].BoundMin) % 16 == 0);
+					//check(UPTRINT(&ClusterTree[0].BoundMax) % 16 == 0);
+
+					int32 MinLOD = 0;
+					int32 MaxLOD = InstanceParams.LODs;
+
+					int32 Force = CVarForceLOD.GetValueOnRenderThread();
+					if (Force >= 0)
+					{
+						MinLOD = Force;
+						MaxLOD = Force;
+					}
+					if (CVarCullAll.GetValueOnRenderThread() < 1)
+					{
+						if (bUseVectorCull)
+						{
+							Traverse<true>(InstanceParams, 0, MinLOD, MaxLOD, bDisableCull);
+						}
+						else
+						{
+							Traverse<false>(InstanceParams, 0, MinLOD, MaxLOD, bDisableCull);
+						}
+					}
+#if !(UE_BUILD_SHIPPING || UE_BUILD_TEST)
+					if (GCaptureDebugRuns == GDebugTag && CaptureTag != GDebugTag)
+					{
+						CaptureTag = GDebugTag;
+						for (int32 LODIndex = 0; LODIndex < InstanceParams.LODs; LODIndex++)
+						{
+							DebugRuns[LODIndex].Empty();
+							DebugTotalInstances[LODIndex] = InstanceParams.TotalInstances[LODIndex];
+							for (int32 Run = 0; Run < InstanceParams.Runs[LODIndex].Num(); Run++)
+							{
+								DebugRuns[LODIndex].Add(InstanceParams.Runs[LODIndex][Run]);
+							}
+						}
+					}
+#endif
+				}
+
+				FillDynamicMeshElements(Collector, ElementParams, InstanceParams);
+			}
+
+			// Render unbuilt instances
+			if (FirstUnbuiltIndex <= LastUnbuiltIndex)
+			{
+				FFoliageRenderInstanceParams InstanceParams;
+
+				// disable LOD blending for unbuilt instances as we haven't calculated the correct LOD.
+				ElementParams.bBlendLODs = false;
+
+				if (LastUnbuiltIndex - FirstUnbuiltIndex < 1000)
+				{
+					// less than 1000 instances, just add them all at LOD 0
+					InstanceParams.AddRun(0, FirstUnbuiltIndex, LastUnbuiltIndex);
+				}
+				else
+				{
+					// more than 1000, render them all at Max LOD
+					InstanceParams.AddRun(RenderData->LODResources.Num() - 1, FirstUnbuiltIndex, FMath::Min<int32>(LastUnbuiltIndex - FirstUnbuiltIndex + 10000, LastUnbuiltIndex));
+				}
+				FillDynamicMeshElements(Collector, ElementParams, InstanceParams);
+			}
+		}
+	}
+}
+
+FBoxSphereBounds UHierarchicalInstancedStaticMeshComponent::CalcBounds(const FTransform& BoundTransform) const
+{
+	const bool bMeshIsValid =
+		// make sure we have instances
+		PerInstanceSMData.Num() > 0 &&
+		// make sure we have an actual staticmesh
+		StaticMesh &&
+		StaticMesh->HasValidRenderData() &&
+		// You really can't use hardware instancing on the consoles with multiple elements because they share the same index buffer. 
+		// @todo: Level error or something to let LDs know this
+		1;//StaticMesh->LODModels(0).Elements.Num() == 1;
+
+	if (ClusterTree.Num() > 0 && bMeshIsValid)
+	{
+		FBoxSphereBounds Result = FBox(ClusterTree[0].BoundMin, ClusterTree[0].BoundMax);
+		for (int32 Index = 1; Index < ClusterTree.Num(); Index++)
+		{
+			Result = Result + FBox(ClusterTree[Index].BoundMin, ClusterTree[Index].BoundMax);
+		}
+
+		// Include any unbuilt instances
+		if (UnbuiltInstanceBounds.IsValid)
+		{
+			Result = Result + UnbuiltInstanceBounds;
+		}
+
+		return Result.TransformBy(BoundTransform);
+	}
+	else
+	{
+		return Super::CalcBounds(BoundTransform);
+	}
+}
+
+UHierarchicalInstancedStaticMeshComponent::UHierarchicalInstancedStaticMeshComponent(const FObjectInitializer& ObjectInitializer)
+	: Super(ObjectInitializer)
+	, NumBuiltInstances(0)
+	, UnbuiltInstanceBounds(0)
+	, bIsAsyncBuilding(false)
+	, bConcurrentRemoval(false)
+	, AccumulatedNavigationDirtyArea(0)
+{
+}
+
+void UHierarchicalInstancedStaticMeshComponent::Serialize(FArchive& Ar)
+{
+	Super::Serialize(Ar);
+
+	ClusterTree.BulkSerialize(Ar);
+}
+
+bool UHierarchicalInstancedStaticMeshComponent::RemoveInstance(int32 InstanceIndex)
+{
+	if (!PerInstanceSMData.IsValidIndex(InstanceIndex))
+	{
+		return false;
+	}
+
+	PartialNavigationUpdate(InstanceIndex);
+
+	// Save the render index
+	RemovedInstances.Add(InstanceReorderTable[InstanceIndex]);
+	
+	// Remove the instance
+	PerInstanceSMData.RemoveAtSwap(InstanceIndex);
+	InstanceReorderTable.RemoveAtSwap(InstanceIndex);
+#if WITH_EDITOR
+	if (SelectedInstances.Num())
+	{
+		SelectedInstances.RemoveAtSwap(InstanceIndex);
+	}
+#endif
+
+	if (IsAsyncBuilding())
+	{
+		// invalidate the results of the current async build as it's too slow to fix up deletes
+		bConcurrentRemoval = true;
+	}
+	else
+	{
+		BuildTreeAsync();
+	}
+
+
+	// update the physics state
+	if (bPhysicsStateCreated)
+	{
+		int32 OldLastIndex = PerInstanceSMData.Num();
+
+		// always shut down physics for the last instance
+		InstanceBodies[OldLastIndex]->TermBody();
+
+		// recreate physics for the instance we swapped in the removed item's place
+		if (InstanceIndex != PerInstanceSMData.Num() && PerInstanceSMData.Num())
+		{
+			InstanceBodies[InstanceIndex]->TermBody();
+			InitInstanceBody(InstanceIndex, InstanceBodies[InstanceIndex]);
+		}
+	}
+
+	MarkRenderStateDirty();
+	
+	return true;
+}
+
+bool UHierarchicalInstancedStaticMeshComponent::UpdateInstanceTransform(int32 InstanceIndex, const FTransform& NewInstanceTransform, bool bWorldSpace)
+{
+	if (!PerInstanceSMData.IsValidIndex(InstanceIndex))
+	{
+		return false;
+	}
+
+	if (IsAsyncBuilding())
+	{
+		// invalidate the results of the current async build we need to modify the tree
+		bConcurrentRemoval = true;
+	}
+
+	// Treat the old instance render data like a removal.
+	RemovedInstances.Add(InstanceReorderTable[InstanceIndex]);
+
+	// Allocate a new instance render order ID, rendered last
+	InstanceReorderTable[InstanceIndex] = PerInstanceSMData.Num()-1 + RemovedInstances.Num();
+
+	bool Result = Super::UpdateInstanceTransform(InstanceIndex, NewInstanceTransform, bWorldSpace);
+
+	UnbuiltInstanceBounds += StaticMesh->GetBounds().GetBox().TransformBy(NewInstanceTransform);
+
+	if (!IsAsyncBuilding())
+	{
+		BuildTreeAsync();
+	}
+
+	return Result;
+}
+
+int32 UHierarchicalInstancedStaticMeshComponent::AddInstance(const FTransform& InstanceTransform)
+{
+	int32 InstanceIndex = UInstancedStaticMeshComponent::AddInstance(InstanceTransform);
+
+	// Need to offset the newly added instance's RenderIndex by the amount that will be adjusted at the end of the frame
+	InstanceReorderTable.Add(InstanceIndex + RemovedInstances.Num());
+
+	if (PerInstanceSMData.Num() == 1)
+	{
+		BuildTree();
+	}
+	else
+	if (!IsAsyncBuilding())
+	{
+		BuildTreeAsync();
+	}
+
+	if (StaticMesh)
+	{
+		UnbuiltInstanceBounds += StaticMesh->GetBounds().GetBox().TransformBy(InstanceTransform);
+	}
+
+	return InstanceIndex;
+}
+
+void UHierarchicalInstancedStaticMeshComponent::ClearInstances()
+{
+	if (IsAsyncBuilding())
+	{
+		bConcurrentRemoval = true;
+	}
+
+	ClusterTree.Empty();
+	NumBuiltInstances = 0;
+	UnbuiltInstanceBounds.Init();
+
+	Super::ClearInstances();
+}
+
+void UHierarchicalInstancedStaticMeshComponent::BuildTree()
+{
+	// Verify that the mesh is valid before using it.
+	const bool bMeshIsValid = 
+		// make sure we have instances
+		PerInstanceSMData.Num() > 0 &&
+		// make sure we have an actual staticmesh
+		StaticMesh &&
+		StaticMesh->HasValidRenderData() &&
+		// You really can't use hardware instancing on the consoles with multiple elements because they share the same index buffer. 
+		// @todo: Level error or something to let LDs know this
+		1;//StaticMesh->LODModels(0).Elements.Num() == 1;
+
+	if(bMeshIsValid)
+	{
+		double StartTime = FPlatformTime::Seconds();
+		// If we don't have a random seed for this instanced static mesh component yet, then go ahead and
+		// generate one now.  This will be saved with the static mesh component and used for future generation
+		// of random numbers for this component's instances. (Used by the PerInstanceRandom material expression)
+		while( InstancingRandomSeed == 0 )
+		{
+			InstancingRandomSeed = FMath::Rand();
+		}
+
+		TArray<FMatrix> InstanceTransforms;
+		InstanceTransforms.AddUninitialized(PerInstanceSMData.Num());
+		for (int32 Index = 0; Index < PerInstanceSMData.Num(); Index++)
+		{
+			InstanceTransforms[Index] = PerInstanceSMData[Index].Transform;
+		}
+
+		TUniquePtr<FClusterBuilder> Builder(new FClusterBuilder(InstanceTransforms, StaticMesh->GetBounds().GetBox()));
+		Builder->Build();
+
+		NumBuiltInstances = Builder->Result->InstanceReorderTable.Num();
+		UnbuiltInstanceBounds.Init();
+		RemovedInstances.Empty();
+
+		Exchange(ClusterTree, Builder->Result->Nodes);
+		Exchange(InstanceReorderTable, Builder->Result->InstanceReorderTable);
+		Exchange(SortedInstances, Builder->Result->SortedInstances);
+		
+		FlushAccumulatedNavigationUpdates();
+				
+		UE_LOG(LogStaticMesh, Display, TEXT("Built a foliage hierarchy with %d elements in %.1fs."), PerInstanceSMData.Num(), float(FPlatformTime::Seconds() - StartTime));
+	}
+}
+
+void UHierarchicalInstancedStaticMeshComponent::ApplyBuildTreeAsync(ENamedThreads::Type CurrentThread, const FGraphEventRef& MyCompletionGraphEvent, TSharedRef<FClusterBuilder, ESPMode::ThreadSafe> Builder, double StartTime)
+{
+	bIsAsyncBuilding = false;
+
+	if (bConcurrentRemoval)
+	{
+		bConcurrentRemoval = false;
+
+		UE_LOG(LogStaticMesh, Display, TEXT("Discarded foliage hierarchy of %d elements build due to concurrent removal (%.1fs)"), Builder->Result->InstanceReorderTable.Num(), float(FPlatformTime::Seconds() - StartTime));
+
+		// There were removes while we were building, it's too slow to fix up the result now, so build async again.
+		BuildTreeAsync();
+	}
+	else
+	{
+		NumBuiltInstances = Builder->Result->InstanceReorderTable.Num();
+
+	    if (NumBuiltInstances < PerInstanceSMData.Num())
+	    {
+		    // Add remap entries for unbuilt instances
+			Builder->Result->InstanceReorderTable.AddUninitialized(PerInstanceSMData.Num() - NumBuiltInstances);
+		    for (int32 Index = NumBuiltInstances; Index < PerInstanceSMData.Num(); Index++)
+		    {
+				Builder->Result->InstanceReorderTable[Index] = Index;
+		    }
+	    }
+    
+	    Exchange(ClusterTree, Builder->Result->Nodes);
+		Exchange(InstanceReorderTable, Builder->Result->InstanceReorderTable);
+		Exchange(SortedInstances, Builder->Result->SortedInstances);
+		RemovedInstances.Empty();
+
+	    UE_LOG(LogStaticMesh, Display, TEXT("Built a foliage hierarchy with %d of %d elements in %.1fs."), NumBuiltInstances, PerInstanceSMData.Num(), float(FPlatformTime::Seconds() - StartTime));
+    
+	    if (NumBuiltInstances < PerInstanceSMData.Num())
+	    {
+		    // There are new outstanding instances, build again!
+		    BuildTreeAsync();
+	    }
+		else
+		{
+			UnbuiltInstanceBounds.Init();
+			FlushAccumulatedNavigationUpdates();
+		}
+
+	    MarkRenderStateDirty();
+	}
+}
+
+void UHierarchicalInstancedStaticMeshComponent::BuildTreeAsync()
+{
+	check(!bIsAsyncBuilding);
+
+	// Verify that the mesh is valid before using it.
+	const bool bMeshIsValid =
+		// make sure we have instances
+		PerInstanceSMData.Num() > 0 &&
+		// make sure we have an actual staticmesh
+		StaticMesh &&
+		StaticMesh->HasValidRenderData() &&
+		// You really can't use hardware instancing on the consoles with multiple elements because they share the same index buffer. 
+		// @todo: Level error or something to let LDs know this
+		1;//StaticMesh->LODModels(0).Elements.Num() == 1;
+
+	if (bMeshIsValid)
+	{
+		double StartTime = FPlatformTime::Seconds();
+		// If we don't have a random seed for this instanced static mesh component yet, then go ahead and
+		// generate one now.  This will be saved with the static mesh component and used for future generation
+		// of random numbers for this component's instances. (Used by the PerInstanceRandom material expression)
+		while (InstancingRandomSeed == 0)
+		{
+			InstancingRandomSeed = FMath::Rand();
+		}
+
+		int32 Num = PerInstanceSMData.Num();
+		TArray<FMatrix> InstanceTransforms;
+		InstanceTransforms.AddUninitialized(Num);
+		for (int32 Index = 0; Index < Num; Index++)
+		{
+			InstanceTransforms[Index] = PerInstanceSMData[Index].Transform;
+		}
+
+		UE_LOG(LogStaticMesh, Display, TEXT("Copied %d transforms in %.3fs."), Num, float(FPlatformTime::Seconds() - StartTime));
+
+		TSharedRef<FClusterBuilder, ESPMode::ThreadSafe> Builder(new FClusterBuilder(InstanceTransforms, StaticMesh->GetBounds().GetBox()));
+
+		bIsAsyncBuilding = true;
+
+		FGraphEventRef BuildTreeAsyncResult(
+			FDelegateGraphTask::CreateAndDispatchWhenReady(FDelegateGraphTask::FDelegate::CreateRaw(&Builder.Get(), &FClusterBuilder::BuildAsync), GET_STATID(STAT_FoliageBuildTime), NULL, ENamedThreads::GameThread, ENamedThreads::AnyThread));
+
+		// add a dependent task to run on the main thread when build is complete
+		FGraphEventRef UnusedAsyncResult(
+			FDelegateGraphTask::CreateAndDispatchWhenReady(
+			FDelegateGraphTask::FDelegate::CreateUObject(this, &UHierarchicalInstancedStaticMeshComponent::ApplyBuildTreeAsync, Builder, StartTime), GET_STATID(STAT_FoliageBuildTime),
+			BuildTreeAsyncResult, ENamedThreads::GameThread, ENamedThreads::GameThread
+			)
+			);
+	}
+}
+
+FPrimitiveSceneProxy* UHierarchicalInstancedStaticMeshComponent::CreateSceneProxy()
+{
+	// Verify that the mesh is valid before using it.
+	const bool bMeshIsValid = 
+		// make sure we have instances
+		PerInstanceSMData.Num() > 0 &&
+		// make sure we have an actual staticmesh
+		StaticMesh &&
+		StaticMesh->HasValidRenderData() &&
+		// You really can't use hardware instancing on the consoles with multiple elements because they share the same index buffer. 
+		// @todo: Level error or something to let LDs know this
+		1;//StaticMesh->LODModels(0).Elements.Num() == 1;
+
+	if(bMeshIsValid)
+	{
+		// If we don't have a random seed for this instanced static mesh component yet, then go ahead and
+		// generate one now.  This will be saved with the static mesh component and used for future generation
+		// of random numbers for this component's instances. (Used by the PerInstanceRandom material expression)
+		while( InstancingRandomSeed == 0 )
+		{
+			InstancingRandomSeed = FMath::Rand();
+		}
+
+		return ::new FHierarchicalStaticMeshSceneProxy(this, GetWorld()->FeatureLevel);
+	}
+	return NULL;
+}
+
+static void GatherInstanceTransformsInArea(const UHierarchicalInstancedStaticMeshComponent& Component, const FBox& AreaBox, int32 Child, TArray<FTransform>& InstanceData)
+{
+	const FClusterNode& ChildNode = Component.ClusterTree[Child];
+	const FVector WorldMinMax[] = {
+		Component.ComponentToWorld.TransformPosition(ChildNode.BoundMin),
+		Component.ComponentToWorld.TransformPosition(ChildNode.BoundMax)
+	};
+	const FBox WorldNodeBox(WorldMinMax, 2);
+
+	if (AreaBox.Intersect(WorldNodeBox))
+	{
+		if (ChildNode.FirstChild < 0 || AreaBox.IsInside(WorldNodeBox))
+		{
+			// Unfortunately ordering of PerInstanceSMData does not match ordering of cluster tree, so we have to use remaping
+			const bool bUseRemaping = Component.SortedInstances.Num() > 0;
+			
+			// In case there no more subdivision or node is completely encapsulated by a area box
+			// add all instances to the result
+			for (int32 i = ChildNode.FirstInstance; i <= ChildNode.LastInstance; ++i)
+			{
+				int32 SortedIdx = bUseRemaping ? Component.SortedInstances[i] : i;
+				FTransform InstanceToComponent(Component.PerInstanceSMData[SortedIdx].Transform);
+				if (!InstanceToComponent.GetScale3D().IsZero())
+				{
+					InstanceData.Add(InstanceToComponent*Component.ComponentToWorld);
+				}
+			}
+		}
+		else
+		{
+			for (int32 i = ChildNode.FirstChild; i <= ChildNode.LastChild; ++i)
+			{
+				GatherInstanceTransformsInArea(Component, AreaBox, i, InstanceData);
+			}
+		}
+	}
+}
+
+void UHierarchicalInstancedStaticMeshComponent::GetNavigationPerInstanceTransforms(const FBox& AreaBox, TArray<FTransform>& InstanceData) const
+{
+	if (IsTreeFullyBuilt())
+	{
+		if (ClusterTree.Num())
+		{
+			GatherInstanceTransformsInArea(*this, AreaBox, 0, InstanceData);
+		}
+	}
+	else
+	{
+		// This area should be processed again by navigation system when cluster tree is available
+		// Store smaller tile box in accumulated dirty area, so we will not unintentionally mark as dirty neighbor tiles 
+		const FBox SmallTileBox = AreaBox.ExpandBy(-AreaBox.GetExtent()/2.f);
+		AccumulatedNavigationDirtyArea+= SmallTileBox;
+	}
+}
+
+void UHierarchicalInstancedStaticMeshComponent::PartialNavigationUpdate(int32 InstanceIdx)
+{
+	if (InstanceIdx == INDEX_NONE)
+	{
+		AccumulatedNavigationDirtyArea.Init();
+		UNavigationSystem::UpdateNavOctree(this);
+	}
+	else if (StaticMesh)
+	{
+		// Accumulate dirty areas and send them to navigation system once cluster tree is rebuilt
+		UNavigationSystem* NavSys = UNavigationSystem::GetCurrent(GetWorld());
+		// Check if this component is registered in navigation system
+		if (NavSys && NavSys->GetObjectsNavOctreeId(this))
+		{
+			FTransform InstanceTransform(PerInstanceSMData[InstanceIdx].Transform);
+			FBox InstanceBox = StaticMesh->GetBounds().TransformBy(InstanceTransform*ComponentToWorld).GetBox(); // in world space
+			AccumulatedNavigationDirtyArea+= InstanceBox;
+		}
+	}
+}
+
+void UHierarchicalInstancedStaticMeshComponent::FlushAccumulatedNavigationUpdates()
+{
+	if (AccumulatedNavigationDirtyArea.IsValid)
+	{
+		UNavigationSystem* NavSys = UNavigationSystem::GetCurrent(GetWorld());
+		// Check if this component is registered in navigation system
+		if (ClusterTree.Num() && NavSys && NavSys->GetObjectsNavOctreeId(this))
+		{
+			FBox NewBounds = FBox(ClusterTree[0].BoundMin, ClusterTree[0].BoundMax).TransformBy(ComponentToWorld);
+			NavSys->UpdateNavOctreeElementBounds(this, NewBounds, AccumulatedNavigationDirtyArea);
+		}
+			
+		AccumulatedNavigationDirtyArea.Init();
+	}
+}
