Merging development branch (CL 2365488) to //depot/UE4/... Engine only, mostly hierarchical foliage

diff --git a/Engine/Source/Runtime/Engine/Private/InstancedFoliage.cpp b/Engine/Source/Runtime/Engine/Private/InstancedFoliage.cpp
--- a/Engine/Source/Runtime/Engine/Private/InstancedFoliage.cpp
+++ b/Engine/Source/Runtime/Engine/Private/InstancedFoliage.cpp
@@ -12,6 +12,8 @@
 #include "Foliage/InstancedFoliageActor.h"
 #include "Foliage/FoliageType_InstancedStaticMesh.h"
 #include "Components/ModelComponent.h"
+#include "Components/HierarchicalInstancedStaticMeshComponent.h"
+#include "Serialization/CustomVersion.h"
 
 #define LOCTEXT_NAMESPACE "InstancedFoliage"
 
@@ -23,6 +25,34 @@ DEFINE_LOG_CATEGORY_STATIC(LogInstancedFoliage, Log, All);
 
 IMPLEMENT_HIT_PROXY(HInstancedStaticMeshInstance, HHitProxy);
 
+// Custom serialization version for all packages containing Instance Foliage
+struct FFoliageCustomVersion
+{
+	enum Type
+	{
+		// Before any version changes were made in the plugin
+		BeforeCustomVersionWasAdded = 0,
+		// Converted to use HierarchicalInstancedStaticMeshComponent
+		FoliageUsingHierarchicalISMC = 1,
+		// Changed Component to not RF_Transactional
+		HierarchicalISMCNonTransactional = 2,
+		// -----<new versions can be added above this line>-------------------------------------------------
+		VersionPlusOne,
+		LatestVersion = VersionPlusOne - 1
+	};
+
+	// The GUID for this custom version number
+	const static FGuid GUID;
+
+private:
+	FFoliageCustomVersion() {}
+};
+
+const FGuid FFoliageCustomVersion::GUID(0x430C4D19, 0x71544970, 0x87699B69, 0xDF90B0E5);
+// Register the custom version with core
+FCustomVersionRegistration GRegisterFoliageCustomVersion(FFoliageCustomVersion::GUID, FFoliageCustomVersion::LatestVersion, TEXT("FoliageVer"));
+
+
 //
 // Serializers for struct data
 //
@@ -33,36 +63,63 @@ FArchive& operator<<(FArchive& Ar, FFoliageInstance& Instance)
 	Ar << Instance.Location;
 	Ar << Instance.Rotation;
 	Ar << Instance.DrawScale3D;
-	Ar << Instance.ClusterIndex;
+	if (Ar.CustomVer(FFoliageCustomVersion::GUID) < FFoliageCustomVersion::FoliageUsingHierarchicalISMC)
+	{
+		int32 OldClusterIndex;
+		Ar << OldClusterIndex;
+	}
 	Ar << Instance.PreAlignRotation;
 	Ar << Instance.Flags;
 	Ar << Instance.ZOffset;
 	return Ar;
 }
 
-FArchive& operator<<(FArchive& Ar, FFoliageInstanceCluster& Cluster)
+/**
+*	FFoliageInstanceCluster_Deprecated
+*/
+struct FFoliageInstanceCluster_Deprecated
 {
-	Ar << Cluster.Bounds;
-	Ar << Cluster.ClusterComponent;
+	UInstancedStaticMeshComponent* ClusterComponent;
+	FBoxSphereBounds Bounds;
 
 #if WITH_EDITORONLY_DATA
-	if (!Ar.ArIsFilterEditorOnly ||
-		Ar.UE4Ver() < VER_UE4_FOLIAGE_SETTINGS_TYPE)
+	TArray<int32> InstanceIndices;	// index into editor editor Instances array
+#endif
+
+	friend FArchive& operator<<(FArchive& Ar, FFoliageInstanceCluster_Deprecated& OldCluster)
 	{
-		Ar << Cluster.InstanceIndices;
-	}
+		check(Ar.CustomVer(FFoliageCustomVersion::GUID) < FFoliageCustomVersion::FoliageUsingHierarchicalISMC);
+
+		Ar << OldCluster.Bounds;
+		Ar << OldCluster.ClusterComponent;
+
+#if WITH_EDITORONLY_DATA
+		if (!Ar.ArIsFilterEditorOnly ||
+			Ar.UE4Ver() < VER_UE4_FOLIAGE_SETTINGS_TYPE)
+		{
+			Ar << OldCluster.InstanceIndices;
+		}
 #endif
 
-	return Ar;
-}
+		return Ar;
+	}
+};
 
 FArchive& operator<<(FArchive& Ar, FFoliageMeshInfo& MeshInfo)
 {
-	Ar << MeshInfo.InstanceClusters;
-
+	if (Ar.CustomVer(FFoliageCustomVersion::GUID) >= FFoliageCustomVersion::FoliageUsingHierarchicalISMC)
+	{
+		Ar << MeshInfo.Component;
+	}
+	else
+	{
+		TArray<FFoliageInstanceCluster_Deprecated> OldInstanceClusters;
+		Ar << OldInstanceClusters;
+	}
+	
 #if WITH_EDITORONLY_DATA
-	if (!Ar.ArIsFilterEditorOnly ||
-		Ar.UE4Ver() < VER_UE4_FOLIAGE_SETTINGS_TYPE)
+	if ((!Ar.ArIsFilterEditorOnly || Ar.UE4Ver() < VER_UE4_FOLIAGE_SETTINGS_TYPE) && 
+		(!(Ar.GetPortFlags() & PPF_DuplicateForPIE)))
 	{
 		Ar << MeshInfo.Instances;
 	}
@@ -72,7 +129,6 @@ FArchive& operator<<(FArchive& Ar, FFoliageMeshInfo& MeshInfo)
 	{
 		Ar << *MeshInfo.InstanceHash;
 		Ar << MeshInfo.ComponentHash;
-		Ar << MeshInfo.FreeInstanceIndices;
 		Ar << MeshInfo.SelectedIndices;
 	}
 #endif
@@ -111,8 +167,6 @@ UFoliageType::UFoliageType(const FObjectInitializer& ObjectInitializer)
 	ZOffsetMin = 0.0f;
 	ZOffsetMax = 0.0f;
 	LandscapeLayer = NAME_None;
-	MaxInstancesPerCluster = 100;
-	MaxClusterRadius = 10000.0f;
 	DisplayOrder = 0;
 	IsSelected = false;
 	ShowNothing = false;
@@ -128,6 +182,7 @@ UFoliageType::UFoliageType(const FObjectInitializer& ObjectInitializer)
 	bCastDynamicShadow = true;
 	bCastStaticShadow = true;
 	bAffectDynamicIndirectLighting = false;
+	bAffectDistanceFieldLighting = true;
 	bCastShadowAsTwoSided = false;
 	bReceivesDecals = false;
 
@@ -157,8 +212,9 @@ void UFoliageType::PostEditChangeProperty(struct FPropertyChangedEvent& Property
 // FFoliageMeshInfo
 //
 FFoliageMeshInfo::FFoliageMeshInfo()
+	: Component(nullptr)
 #if WITH_EDITOR
-	: InstanceHash(GIsEditor ? new FFoliageInstanceHash() : nullptr)
+	, InstanceHash(GIsEditor ? new FFoliageInstanceHash() : nullptr)
 #endif
 { }
 
@@ -239,10 +295,47 @@ void FFoliageMeshInfo::CheckValid()
 
 void FFoliageMeshInfo::AddInstance(AInstancedFoliageActor* InIFA, UFoliageType* InSettings, const FFoliageInstance& InNewInstance)
 {
+	// Want to make sure the component exists before we start a transaction, so it will be there on undo.
+	if (Component == nullptr)
+	{
+		Component = ConstructObject<UHierarchicalInstancedStaticMeshComponent>(UHierarchicalInstancedStaticMeshComponent::StaticClass(), InIFA);
+
+		Component->Mobility = EComponentMobility::Static;
+
+		Component->StaticMesh = InSettings->GetStaticMesh();
+		Component->bSelectable = true;
+		Component->bHasPerInstanceHitProxies = true;
+		Component->InstancingRandomSeed = FMath::Rand();
+		Component->InstanceStartCullDistance = InSettings->StartCullDistance;
+		Component->InstanceEndCullDistance = InSettings->EndCullDistance;
+
+		Component->CastShadow = InSettings->CastShadow;
+		Component->bCastDynamicShadow = InSettings->bCastDynamicShadow;
+		Component->bCastStaticShadow = InSettings->bCastStaticShadow;
+		Component->bAffectDynamicIndirectLighting = InSettings->bAffectDynamicIndirectLighting;
+		Component->bAffectDistanceFieldLighting = InSettings->bAffectDistanceFieldLighting;
+		Component->bCastShadowAsTwoSided = InSettings->bCastShadowAsTwoSided;
+		Component->bReceivesDecals = InSettings->bReceivesDecals;
+		Component->bOverrideLightMapRes = InSettings->bOverrideLightMapRes;
+		Component->OverriddenLightMapRes = InSettings->OverriddenLightMapRes;
+
+		Component->BodyInstance.CopyBodyInstancePropertiesFrom(&InSettings->BodyInstance);
+
+		Component->AttachTo(InIFA->GetRootComponent());
+		Component->RegisterComponent();
+
+		// Use only instance translation as a component transform
+		Component->SetWorldTransform(InIFA->GetRootComponent()->ComponentToWorld);
+	}
+	else
+	{
+		Component->InvalidateLightingCache();
+	}
+
 	InIFA->Modify();
 
 	// Add the instance taking either a free slot or adding a new item.
-	int32 InstanceIndex = FreeInstanceIndices.Num() > 0 ?  FreeInstanceIndices.Pop() : Instances.AddUninitialized();
+	int32 InstanceIndex = Instances.AddUninitialized();
 
 	FFoliageInstance& AddedInstance = Instances[InstanceIndex];
 	AddedInstance = InNewInstance;
@@ -256,94 +349,17 @@ void FFoliageMeshInfo::AddInstance(AInstancedFoliageActor* InIFA, UFoliageType*
 	}
 	ComponentHashInfo->Instances.Add(InstanceIndex);
 
-	// Find the best cluster to allocate the instance to.
-	FFoliageInstanceCluster* BestCluster = nullptr;
-	int32 BestClusterIndex = INDEX_NONE;
-	float BestClusterDistSq = FLT_MAX;
-
-	int32 MaxInstancesPerCluster = InSettings->MaxInstancesPerCluster;
-	float MaxClusterRadiusSq = FMath::Square(InSettings->MaxClusterRadius);
-
-	for (int32 ClusterIdx = 0; ClusterIdx < InstanceClusters.Num(); ClusterIdx++)
-	{
-		FFoliageInstanceCluster& Cluster = InstanceClusters[ClusterIdx];
-		if (Cluster.InstanceIndices.Num() < MaxInstancesPerCluster)
-		{
-			float DistSq = (Cluster.Bounds.Origin - InNewInstance.Location).SizeSquared();
-			if (DistSq < BestClusterDistSq && DistSq < MaxClusterRadiusSq)
-			{
-				BestCluster = &Cluster;
-				BestClusterIndex = ClusterIdx;
-				BestClusterDistSq = DistSq;
-			}
-		}
-	}
-
 	// Calculate transform for the instance
 	FTransform InstanceToWorld = InNewInstance.GetInstanceWorldTransform();
 
-	// Bounds of the new instance
-	FBoxSphereBounds InstanceWorldBounds = InSettings->GetStaticMesh()->GetBounds().TransformBy(InstanceToWorld);
-
-	if (BestCluster == nullptr)
-	{
-		BestClusterIndex = InstanceClusters.Num();
-		BestCluster = new(InstanceClusters)FFoliageInstanceCluster(
-			ConstructObject<UInstancedStaticMeshComponent>(UInstancedStaticMeshComponent::StaticClass(), InIFA, NAME_None, RF_Transactional),
-			InstanceWorldBounds
-			);
-
-		BestCluster->ClusterComponent->Mobility = EComponentMobility::Static;
-
-		BestCluster->ClusterComponent->StaticMesh = InSettings->GetStaticMesh();
-		BestCluster->ClusterComponent->bSelectable = true;
-		BestCluster->ClusterComponent->bHasPerInstanceHitProxies = true;
-		BestCluster->ClusterComponent->InstancingRandomSeed = FMath::Rand();
-		BestCluster->ClusterComponent->InstanceStartCullDistance = InSettings->StartCullDistance;
-		BestCluster->ClusterComponent->InstanceEndCullDistance   = InSettings->EndCullDistance;
-
-		BestCluster->ClusterComponent->CastShadow                     = InSettings->CastShadow;
-		BestCluster->ClusterComponent->bCastDynamicShadow             = InSettings->bCastDynamicShadow;
-		BestCluster->ClusterComponent->bCastStaticShadow              = InSettings->bCastStaticShadow;
-		BestCluster->ClusterComponent->bAffectDynamicIndirectLighting = InSettings->bAffectDynamicIndirectLighting;
-		BestCluster->ClusterComponent->bCastShadowAsTwoSided          = InSettings->bCastShadowAsTwoSided;
-		BestCluster->ClusterComponent->bReceivesDecals                = InSettings->bReceivesDecals;
-		BestCluster->ClusterComponent->bOverrideLightMapRes           = InSettings->bOverrideLightMapRes;
-		BestCluster->ClusterComponent->OverriddenLightMapRes          = InSettings->OverriddenLightMapRes;
-
-		BestCluster->ClusterComponent->BodyInstance.CopyBodyInstancePropertiesFrom(&InSettings->BodyInstance);
-
-		BestCluster->ClusterComponent->AttachTo(InIFA->GetRootComponent());
-		BestCluster->ClusterComponent->RegisterComponent();
-		
-		// Use only instance translation as a component transform
-		BestCluster->ClusterComponent->SetWorldTransform(FTransform(InstanceToWorld.GetTranslation()));
-	}
-	else
-	{
-		BestCluster->ClusterComponent->Modify();
-		BestCluster->ClusterComponent->InvalidateLightingCache();
-		BestCluster->Bounds = BestCluster->Bounds + InstanceWorldBounds;
-	}
-
-	BestCluster->InstanceIndices.Add(InstanceIndex);
-
-	// Save the cluster index
-	AddedInstance.ClusterIndex = BestClusterIndex;
-
 	// Add the instance to the component
-	BestCluster->ClusterComponent->AddInstanceWorldSpace(InstanceToWorld);
-
-	if (BestCluster->ClusterComponent->SelectedInstances.Num() > 0)
-	{
-		BestCluster->ClusterComponent->SelectedInstances.Add(false);
-	}
+	Component->AddInstanceWorldSpace(InstanceToWorld);
 
 	// Update PrimitiveComponent's culling distance taking into account the radius of the bounds, as
 	// it is based on the center of the component's bounds.
-	float CullDistance = InSettings->EndCullDistance > 0 ? (float)InSettings->EndCullDistance + BestCluster->Bounds.SphereRadius : 0.f;
-	BestCluster->ClusterComponent->LDMaxDrawDistance = CullDistance;
-	BestCluster->ClusterComponent->CachedMaxDrawDistance = CullDistance;
+//	float CullDistance = InSettings->EndCullDistance > 0 ? (float)InSettings->EndCullDistance + BestCluster->Bounds.SphereRadius : 0.f;
+//	Component->LDMaxDrawDistance = CullDistance;
+//	Component->CachedMaxDrawDistance = CullDistance;
 
 	CheckValid();
 }
@@ -352,126 +368,80 @@ void FFoliageMeshInfo::RemoveInstances(AInstancedFoliageActor* InIFA, const TArr
 {
 	bool bRemoved = false;
 
-	InIFA->Modify();
-
-	// Used to store a list of instances to remove sorted by cluster
-	TMap<int32, TSet<int32> > ClusterInstanceSetMap;
-
-	// Remove instances from the hash
-	for (TArray<int32>::TConstIterator It(InInstancesToRemove); It; ++It)
-	{
-		int32 InstanceIndex = *It;
-		const FFoliageInstance& Instance = Instances[InstanceIndex];
-		InstanceHash->RemoveInstance(Instance.Location, InstanceIndex);
-		FFoliageComponentHashInfo* ComponentHashInfo = ComponentHash.Find(Instance.Base);
-		if (ComponentHashInfo)
-		{
-			ComponentHashInfo->Instances.Remove(InstanceIndex);
-			if (ComponentHashInfo->Instances.Num() == 0)
-			{
-				// Remove the component from the component hash if this is the last instance.
-				ComponentHash.Remove(Instance.Base);
-			}
-		}
-
-		// Store in ClusterInstanceSetMap
-		int32 ClusterIndex = Instances[InstanceIndex].ClusterIndex;
-		TSet<int32>& ClusterInstances = ClusterInstanceSetMap.FindOrAdd(ClusterIndex);
-		ClusterInstances.Add(InstanceIndex);
-	}
+	TArray<int32> InstancesToRemove = InInstancesToRemove;
 
-	// Process each cluster in turn
-	for (TMap<int32, TSet<int32> >::TIterator It(ClusterInstanceSetMap); It; ++It)
+	if (InstancesToRemove.Num())
 	{
-		int32 CurrentClusterIndex = It.Key();
-		FFoliageInstanceCluster& Cluster = InstanceClusters[CurrentClusterIndex];
-		UInstancedStaticMeshComponent* ClusterComponent = Cluster.ClusterComponent;
-		bool bRemovedFromComponent = false;
-
-		TSet<int32>& ClusterInstancesToRemove = It.Value();
+		check(Component);
+		InIFA->Modify();
 
-		// Look through all indices in this cluster
-		for (int32 ClusterInstanceIdx = 0; ClusterInstanceIdx < Cluster.InstanceIndices.Num(); ClusterInstanceIdx++)
+		// Remove instances
+		for (int32 Idx = 0; Idx < InstancesToRemove.Num(); Idx++)
 		{
-			int32 InstanceIndex = Cluster.InstanceIndices[ClusterInstanceIdx];
-
-			// Check if this instance is one we need to remove
-			if (ClusterInstancesToRemove.Contains(InstanceIndex))
+			int32 InstanceIndex = InstancesToRemove[Idx];
+
+			FFoliageInstance& Instance = Instances[InstanceIndex];
+
+			// remove from hash
+			InstanceHash->RemoveInstance(Instance.Location, InstanceIndex);
+		    FFoliageComponentHashInfo* ComponentHashInfo = ComponentHash.Find(Instance.Base);
+		    if (ComponentHashInfo)
+		    {
+			    ComponentHashInfo->Instances.Remove(InstanceIndex);
+			    if (ComponentHashInfo->Instances.Num() == 0)
+			    {
+				    // Remove the component from the component hash if this is the last instance.
+				    ComponentHash.Remove(Instance.Base);
+			    }
+		    }
+
+			// remove from the component
+			Component->RemoveInstance(InstanceIndex);
+
+			// Remove it from the selection.
+			SelectedIndices.Remove(InstanceIndex);
+
+			// remove from instances array
+			Instances.RemoveAtSwap(InstanceIndex);
+
+			// update hashes for swapped instance
+			if (InstanceIndex != Instances.Num() && Instances.Num())
 			{
-				// Remove the instance data and the index array entry
-				if (ClusterComponent->IsRegistered())
-				{
-					ClusterComponent->UnregisterComponent();
-				}
-				ClusterComponent->Modify();
-				ClusterComponent->PerInstanceSMData.RemoveAtSwap(ClusterInstanceIdx);
-				if (ClusterInstanceIdx < ClusterComponent->SelectedInstances.Num())
-				{
-					ClusterComponent->SelectedInstances.RemoveAtSwap(ClusterInstanceIdx);
-				}
+				// Instance hash
+				FFoliageInstance& SwappedInstance = Instances[InstanceIndex];
+				InstanceHash->RemoveInstance(SwappedInstance.Location, Instances.Num());
+				InstanceHash->InsertInstance(SwappedInstance.Location, InstanceIndex);
 
-				if (ClusterInstanceIdx < ClusterComponent->InstanceBodies.Num())
+				// Component hash
+				if (SwappedInstance.Base)
 				{
-					delete ClusterComponent->InstanceBodies[ClusterInstanceIdx];
-					ClusterComponent->InstanceBodies.RemoveAtSwap(ClusterInstanceIdx);
+					FFoliageComponentHashInfo* ComponentHashInfo = ComponentHash.Find(SwappedInstance.Base);
+					check(ComponentHashInfo);
+					ComponentHashInfo->Instances.Remove(Instances.Num());
+					ComponentHashInfo->Instances.Add(InstanceIndex);
 				}
 
-				Cluster.InstanceIndices.RemoveAtSwap(ClusterInstanceIdx);
-				ClusterInstanceIdx--;
-				bRemovedFromComponent = true;
-
-				// Invalidate this instance's editor data so we reuse the slot later
-				FFoliageInstance& InstanceEditorData = Instances[InstanceIndex];
-				InstanceEditorData.ClusterIndex = INDEX_NONE;
-				InstanceEditorData.Base = nullptr;
-
-				// And remember the slot to reuse it.
-				FreeInstanceIndices.Add(InstanceIndex);
-
-				// Remove it from the selection.
-				SelectedIndices.Remove(InstanceIndex);
+				// Selection
+				if (SelectedIndices.Contains(Instances.Num()))
+				{
+					SelectedIndices.Remove(Instances.Num());
+					SelectedIndices.Add(InstanceIndex);
+				}				
 			}
-		}
-
-		if (bRemovedFromComponent)
-		{
-			//!! need to update bounds for this cluster
-			ClusterComponent->RegisterComponent();
-			bRemoved = true;
-		}
-	}
 
-	if (bRemoved)
-	{
-		// See if we can delete any clusters.
-		for (int32 ClusterIdx = 0; ClusterIdx < InstanceClusters.Num(); ClusterIdx++)
-		{
-			if (InstanceClusters[ClusterIdx].InstanceIndices.Num() == 0)
+			// See if any item we're about to remove was swapped into this slot.
+			for (int32 NextIdx = Idx + 1; NextIdx < InstancesToRemove.Num(); NextIdx++)
 			{
-				// Unregister the component to remove it from the scene
-				UInstancedStaticMeshComponent* Component = InstanceClusters[ClusterIdx].ClusterComponent;
-
-				// Do not reregister this component if register all components gets called again
-				Component->bAutoRegister = false;
-
-				// Remove cluster
-				InstanceClusters.RemoveAt(ClusterIdx);
-
-				// update the ClusterIndex for remaining foliage instances
-				for (int32 InstanceIdx = 0; InstanceIdx < Instances.Num(); InstanceIdx++)
+				if (InstancesToRemove[NextIdx] == Instances.Num())
 				{
-					if (Instances[InstanceIdx].ClusterIndex > ClusterIdx)
-					{
-						Instances[InstanceIdx].ClusterIndex--;
-					}
+					InstancesToRemove[NextIdx] = InstanceIndex;
+					break;
 				}
-
-				ClusterIdx--;
 			}
-
-			InIFA->CheckSelection();
 		}
 
+		InIFA->CheckSelection();
+
 		CheckValid();
 	}
 }
@@ -490,36 +460,28 @@ void FFoliageMeshInfo::PreMoveInstances(AInstancedFoliageActor* InIFA, const TAr
 
 void FFoliageMeshInfo::PostUpdateInstances(AInstancedFoliageActor* InIFA, const TArray<int32>& InInstancesUpdated, bool bReAddToHash)
 {
-	for (TArray<int32>::TConstIterator It(InInstancesUpdated); It; ++It)
+	if (InInstancesUpdated.Num())
 	{
-		int32 InstanceIndex = *It;
-		const FFoliageInstance& Instance = Instances[InstanceIndex];
+		check(Component);
 
-		// Update this instances' transform in the UInstancedStaticMeshComponent
-		if (InstanceClusters.IsValidIndex(Instance.ClusterIndex))
+		for (TArray<int32>::TConstIterator It(InInstancesUpdated); It; ++It)
 		{
-			FFoliageInstanceCluster& Cluster = InstanceClusters[Instance.ClusterIndex];
-			Cluster.ClusterComponent->Modify();
-			int32 ClusterInstanceDataIndex = Cluster.InstanceIndices.Find(InstanceIndex);
-			check(ClusterInstanceDataIndex != INDEX_NONE);
-			check(Cluster.ClusterComponent->PerInstanceSMData.IsValidIndex(ClusterInstanceDataIndex));
-			Cluster.ClusterComponent->MarkRenderStateDirty();
-
-			// Update bounds
+			int32 InstanceIndex = *It;
+			const FFoliageInstance& Instance = Instances[InstanceIndex];
+
 			FTransform InstanceToWorld = Instance.GetInstanceWorldTransform();
-			Cluster.Bounds = Cluster.Bounds + Cluster.ClusterComponent->StaticMesh->GetBounds().TransformBy(InstanceToWorld);
-			// Update transform in InstancedStaticMeshComponent (in component local space)
-			// Transform from world space to local space
-			FTransform InstanceToLocal = InstanceToWorld.GetRelativeTransform(Cluster.ClusterComponent->ComponentToWorld);
-			Cluster.ClusterComponent->PerInstanceSMData[ClusterInstanceDataIndex].Transform = InstanceToLocal.ToMatrixWithScale();
-			Cluster.ClusterComponent->InvalidateLightingCache();
-		}
 
-		if (bReAddToHash)
-		{
-			// Re-add instance to the hash
-			InstanceHash->InsertInstance(Instance.Location, InstanceIndex);
+			Component->UpdateInstanceTransform(InstanceIndex, InstanceToWorld, true);
+
+			// Re-add instance to the hash if requested
+			if (bReAddToHash)
+			{
+				InstanceHash->InsertInstance(Instance.Location, InstanceIndex);
+			}
 		}
+
+		Component->InvalidateLightingCache();
+		Component->MarkRenderStateDirty();
 	}
 }
 
@@ -540,15 +502,7 @@ void FFoliageMeshInfo::DuplicateInstances(AInstancedFoliageActor* InIFA, UFoliag
 /* Get the number of placed instances */
 int32 FFoliageMeshInfo::GetInstanceCount() const
 {
-	int32 Count = 0;
-	for (int32 Idx = 0; Idx < Instances.Num(); Idx++)
-	{
-		if (Instances[Idx].ClusterIndex != -1)
-		{
-			Count++;
-		}
-	}
-	return Count;
+	return Instances.Num();
 }
 
 // Destroy existing clusters and reassign all instances to new clusters
@@ -557,45 +511,64 @@ void FFoliageMeshInfo::ReallocateClusters(AInstancedFoliageActor* InIFA, UFoliag
 	// Detach all components
 	InIFA->UnregisterAllComponents();
 
-	// Delete any existing clusters from the components array
-	for (FFoliageInstanceCluster& Cluster : InstanceClusters)
+	if (Component != nullptr)
 	{
-		UInstancedStaticMeshComponent* Component = Cluster.ClusterComponent;
-		if (Component != nullptr)
-		{
-			Component->PerInstanceSMData.Empty();
-			Component->bAutoRegister = false;
-		}
+		Component->bAutoRegister = false;
+		Component = nullptr;
 	}
 
-	TArray<FFoliageInstance> OldInstances = Instances;
-	OldInstances.RemoveAllSwap([](FFoliageInstance& Instance){ return Instance.ClusterIndex == INDEX_NONE; });
 
 	// Remove everything
-	InstanceClusters.Empty();
-	Instances.Empty();
+	TArray<FFoliageInstance> OldInstances;
+	Exchange(Instances, OldInstances);
 	InstanceHash->Empty();
 	ComponentHash.Empty();
-	FreeInstanceIndices.Empty();
 	SelectedIndices.Empty();
 
 	// Re-add
 	for (FFoliageInstance& Instance : OldInstances)
 	{
-		Instance.ClusterIndex = INDEX_NONE;
 		AddInstance(InIFA, InSettings, Instance);
 	}
 
 	InIFA->RegisterAllComponents();
 }
 
+void FFoliageMeshInfo::ReapplyInstancesToComponent()
+{
+	if (Component)
+	{
+		Component->UnregisterComponent();
+		Component->ClearInstances();
+
+		for (auto& Instance : Instances)
+		{
+			Component->AddInstanceWorldSpace(Instance.GetInstanceWorldTransform());
+		}
+
+		if (SelectedIndices.Num())
+		{
+			if (Component->SelectedInstances.Num() != Component->PerInstanceSMData.Num())
+			{
+				Component->SelectedInstances.Init(false, Component->PerInstanceSMData.Num());
+			}
+			for (int32 i : SelectedIndices)
+			{
+				Component->SelectedInstances[i] = true;
+			}
+		}
+
+		Component->RegisterComponent();
+	}
+}
+
+
 // Update settings in the clusters based on the current settings (eg culling distance, collision, ...)
 //void FFoliageMeshInfo::UpdateClusterSettings(AInstancedFoliageActor* InIFA)
 //{
 //	for (FFoliageInstanceCluster& Cluster : InstanceClusters)
 //	{
 //		UInstancedStaticMeshComponent* ClusterComponent = Cluster.ClusterComponent;
-//		ClusterComponent->Modify();
 //		ClusterComponent->MarkRenderStateDirty();
 //
 //		// Copy settings
@@ -656,57 +629,37 @@ bool FFoliageMeshInfo::CheckForOverlappingInstanceExcluding(int32 TestInstanceId
 
 void FFoliageMeshInfo::SelectInstances(AInstancedFoliageActor* InIFA, bool bSelect, TArray<int32>& InInstances)
 {
-	InIFA->Modify();
-
-	if (bSelect)
+	if (InInstances.Num())
 	{
-		for (int32 i : InInstances)
-		{
-			SelectedIndices.AddUnique(i);
-		}
+		check(Component);
+		InIFA->Modify();
 
-		for (FFoliageInstanceCluster& Cluster : InstanceClusters)
+		if (bSelect)
 		{
-			// Apply any selections in the component
-			Cluster.ClusterComponent->Modify();
-			Cluster.ClusterComponent->MarkRenderStateDirty();
-			if (Cluster.ClusterComponent->SelectedInstances.Num() != Cluster.InstanceIndices.Num())
+			// Apply selections to the component
+			Component->MarkRenderStateDirty();
+
+			if (Component->SelectedInstances.Num() != Component->PerInstanceSMData.Num())
 			{
-				Cluster.ClusterComponent->SelectedInstances.Init(false, Cluster.InstanceIndices.Num());
+				Component->SelectedInstances.Init(false, Component->PerInstanceSMData.Num());
 			}
 
-			for (int32 ClusterInstanceIdx = 0; ClusterInstanceIdx < Cluster.InstanceIndices.Num(); ClusterInstanceIdx++)
+			for (int32 i : InInstances)
 			{
-				int32 InstanceIdx = Cluster.InstanceIndices[ClusterInstanceIdx];
-				if (InInstances.Find(InstanceIdx) != INDEX_NONE)
-				{
-					Cluster.ClusterComponent->SelectedInstances[ClusterInstanceIdx] = true;
-				}
+				SelectedIndices.Add(i);
+				Component->SelectedInstances[i] = true;
 			}
 		}
-	}
-	else
-	{
-		for (int32 i : InInstances)
-		{
-			SelectedIndices.RemoveSingleSwap(i);
-		}
-
-		for (FFoliageInstanceCluster& Cluster : InstanceClusters)
+		else
 		{
-			if (Cluster.ClusterComponent->SelectedInstances.Num() != 0)
+			if (Component->SelectedInstances.Num())
 			{
-				// Remove any selections from the component
-				Cluster.ClusterComponent->Modify();
-				Cluster.ClusterComponent->MarkRenderStateDirty();
+				Component->MarkRenderStateDirty();
 
-				for (int32 ClusterInstanceIdx = 0; ClusterInstanceIdx < Cluster.InstanceIndices.Num(); ClusterInstanceIdx++)
+				for (int32 i : InInstances)
 				{
-					int32 InstanceIdx = Cluster.InstanceIndices[ClusterInstanceIdx];
-					if (InInstances.Find(InstanceIdx) != INDEX_NONE)
-					{
-						Cluster.ClusterComponent->SelectedInstances[ClusterInstanceIdx] = false;
-					}
+					SelectedIndices.Remove(i);
+					Component->SelectedInstances[i] = false;
 				}
 			}
 		}
@@ -811,24 +764,9 @@ void AInstancedFoliageActor::MoveInstancesForMovedComponent(UActorComponent* InC
 				Instance.Location = DeltaTransform.GetOrigin();
 				Instance.Rotation = DeltaTransform.Rotator();
 
-				// Update this instances' transform in the UInstancedStaticMeshComponent
-				if (MeshInfo.InstanceClusters.IsValidIndex(Instance.ClusterIndex))
-				{
-					FFoliageInstanceCluster& Cluster = MeshInfo.InstanceClusters[Instance.ClusterIndex];
-					Cluster.ClusterComponent->Modify();
-					int32 ClusterInstanceDataIndex = Cluster.InstanceIndices.Find(InstanceIndex);
-					check(ClusterInstanceDataIndex != INDEX_NONE);
-					check(Cluster.ClusterComponent->PerInstanceSMData.IsValidIndex(ClusterInstanceDataIndex));
-					Cluster.ClusterComponent->MarkRenderStateDirty();
-					bUpdatedInstances = true;
-					// Update bounds
-					FTransform InstanceToWorld = Instance.GetInstanceWorldTransform();
-					Cluster.Bounds = Cluster.Bounds + Cluster.ClusterComponent->StaticMesh->GetBounds().TransformBy(InstanceToWorld);
-					// Update transform in InstancedStaticMeshComponent
-					FTransform InstanceToLocal = InstanceToWorld.GetRelativeTransform(Cluster.ClusterComponent->ComponentToWorld);
-					Cluster.ClusterComponent->PerInstanceSMData[ClusterInstanceDataIndex].Transform = InstanceToLocal.ToMatrixWithScale();
-					Cluster.ClusterComponent->InvalidateLightingCache();
-				}
+				// Apply render data
+				check(MeshInfo.Component);
+				MeshInfo.Component->UpdateInstanceTransform(InstanceIndex, Instance.GetInstanceWorldTransform(), true);
 
 				// Re-add the new instance location to the hash
 				MeshInfo.InstanceHash->InsertInstance(Instance.Location, InstanceIndex);
@@ -930,7 +868,6 @@ void AInstancedFoliageActor::MoveInstancesToNewComponent(UPrimitiveComponent* In
 				{
 					FFoliageInstance NewInstance = MeshInfo.Instances[InstanceIndex];
 					NewInstance.Base = InNewComponent;
-					NewInstance.ClusterIndex = INDEX_NONE;
 					NewMeshInfo->AddInstance(NewIFA, FoliageType, NewInstance);
 				}
 
@@ -1073,16 +1010,9 @@ void AInstancedFoliageActor::RemoveMesh(UFoliageType* InSettings)
 
 	// Remove all components for this mesh from the Components array.
 	FFoliageMeshInfo* MeshInfo = FindMesh(InSettings);
-	if (MeshInfo)
+	if (MeshInfo && MeshInfo->Component)
 	{
-		for (FFoliageInstanceCluster& Cluster : MeshInfo->InstanceClusters)
-		{
-			UInstancedStaticMeshComponent* Component = Cluster.ClusterComponent;
-			if (Component != nullptr)
-			{
-				Component->bAutoRegister = false;
-			}
-		}
+		MeshInfo->Component->bAutoRegister = false;
 	}
 
 	FoliageMeshes.Remove(InSettings);
@@ -1113,10 +1043,8 @@ UFoliageType* AInstancedFoliageActor::GetSettingsForMesh(UStaticMesh* InMesh, FF
 	return Type;
 }
 
-void AInstancedFoliageActor::SelectInstance(UInstancedStaticMeshComponent* InComponent, int32 InComponentInstanceIndex, bool bToggle)
+void AInstancedFoliageActor::SelectInstance(UInstancedStaticMeshComponent* InComponent, int32 InInstanceIndex, bool bToggle)
 {
-	bool bNeedsUpdate = false;
-
 	Modify();
 
 	// If we're not toggling, we need to first deselect everything else
@@ -1128,17 +1056,12 @@ void AInstancedFoliageActor::SelectInstance(UInstancedStaticMeshComponent* InCom
 
 			if (MeshInfo.SelectedIndices.Num() > 0)
 			{
-				for (int32 ClusterIdx = 0; ClusterIdx < MeshInfo.InstanceClusters.Num(); ClusterIdx++)
+				check(MeshInfo.Component);
+				if (MeshInfo.Component->SelectedInstances.Num() > 0)
 				{
-					FFoliageInstanceCluster& Cluster = MeshInfo.InstanceClusters[ClusterIdx];
-					if (Cluster.ClusterComponent->SelectedInstances.Num() > 0)
-					{
-						Cluster.ClusterComponent->Modify();
-						Cluster.ClusterComponent->SelectedInstances.Empty();
-						Cluster.ClusterComponent->MarkRenderStateDirty();
-					}
+					MeshInfo.Component->SelectedInstances.Empty();
+					MeshInfo.Component->MarkRenderStateDirty();
 				}
-				bNeedsUpdate = true;
 			}
 
 			MeshInfo.SelectedIndices.Empty();
@@ -1154,27 +1077,23 @@ void AInstancedFoliageActor::SelectInstance(UInstancedStaticMeshComponent* InCom
 
 		if (MeshInfo)
 		{
-			FFoliageInstanceCluster* Cluster = MeshInfo->InstanceClusters.FindByPredicate([InComponent](FFoliageInstanceCluster& Cluster){ return Cluster.ClusterComponent == InComponent; });
-			if (Cluster)
+			if (InComponent == MeshInfo->Component)
 			{
-				InComponent->Modify();
-				int32 InstanceIndex = Cluster->InstanceIndices[InComponentInstanceIndex];
-				int32 SelectedIndex = MeshInfo->SelectedIndices.Find(InstanceIndex);
-
-				bNeedsUpdate = true;
+				bool bIsSelected = MeshInfo->SelectedIndices.Contains(InInstanceIndex);
 
 				// Deselect if it's already selected.
-				if (InComponentInstanceIndex < InComponent->SelectedInstances.Num())
+				if (InInstanceIndex < InComponent->SelectedInstances.Num())
 				{
-					InComponent->SelectedInstances[InComponentInstanceIndex] = false;
+					InComponent->SelectedInstances[InInstanceIndex] = false;
 					InComponent->MarkRenderStateDirty();
 				}
-				if (SelectedIndex != INDEX_NONE)
+
+				if (bIsSelected)
 				{
-					MeshInfo->SelectedIndices.RemoveAt(SelectedIndex);
+					MeshInfo->SelectedIndices.Remove(InInstanceIndex);
 				}
 
-				if (bToggle && SelectedIndex != INDEX_NONE)
+				if (bToggle && bIsSelected)
 				{
 					if (SelectedMesh == Type && MeshInfo->SelectedIndices.Num() == 0)
 					{
@@ -1186,25 +1105,32 @@ void AInstancedFoliageActor::SelectInstance(UInstancedStaticMeshComponent* InCom
 					// Add the selection
 					if (InComponent->SelectedInstances.Num() < InComponent->PerInstanceSMData.Num())
 					{
-						InComponent->SelectedInstances.Init(false, Cluster->InstanceIndices.Num());
+						InComponent->SelectedInstances.Init(false, InComponent->PerInstanceSMData.Num());
 					}
-					InComponent->SelectedInstances[InComponentInstanceIndex] = true;
+					InComponent->SelectedInstances[InInstanceIndex] = true;
 					InComponent->MarkRenderStateDirty();
 
 					SelectedMesh = Type;
-					MeshInfo->SelectedIndices.Insert(InstanceIndex, 0);
+					MeshInfo->SelectedIndices.Add(InInstanceIndex);
 				}
 			}
 		}
 	}
 
-	if (bNeedsUpdate)
+	CheckSelection();
+
+}
+
+void AInstancedFoliageActor::PostEditUndo()
+{
+	Super::PostEditUndo();
+
+	FlushRenderingCommands();
+	for (auto& MeshPair : FoliageMeshes)
 	{
-		// Update selection
-		MarkComponentsRenderStateDirty();
+		FFoliageMeshInfo& MeshInfo = *MeshPair.Value;
+		MeshInfo.ReapplyInstancesToComponent();
 	}
-
-	CheckSelection();
 }
 
 void AInstancedFoliageActor::ApplySelectionToComponents(bool bApply)
@@ -1219,44 +1145,27 @@ void AInstancedFoliageActor::ApplySelectionToComponents(bool bApply)
 		{
 			if (MeshInfo.SelectedIndices.Num() > 0)
 			{
-				for (FFoliageInstanceCluster& Cluster : MeshInfo.InstanceClusters)
-				{
-					// Apply any selections in the component
-					Cluster.ClusterComponent->SelectedInstances.Init(false, Cluster.InstanceIndices.Num());
-					Cluster.ClusterComponent->MarkRenderStateDirty();
-					bNeedsUpdate = true;
+				check(MeshInfo.Component);
 
-					for (int32 ClusterInstanceIdx = 0; ClusterInstanceIdx < Cluster.InstanceIndices.Num(); ClusterInstanceIdx++)
-					{
-						int32 InstanceIdx = Cluster.InstanceIndices[ClusterInstanceIdx];
-						if (MeshInfo.SelectedIndices.Find(InstanceIdx) != INDEX_NONE)
-						{
-							Cluster.ClusterComponent->SelectedInstances[ClusterInstanceIdx] = true;
-						}
-					}
+				// Apply any selections in the component
+				MeshInfo.Component->SelectedInstances.Init(false, MeshInfo.Component->PerInstanceSMData.Num());
+				MeshInfo.Component->MarkRenderStateDirty();
+				for (int32 i : MeshInfo.SelectedIndices)
+				{
+					MeshInfo.Component->SelectedInstances[i] = true;
 				}
 			}
 		}
-		else
+		else		
 		{
-			for (FFoliageInstanceCluster& Cluster : MeshInfo.InstanceClusters)
+			if (MeshInfo.Component && MeshInfo.Component->SelectedInstances.Num() > 0)
 			{
-				if (Cluster.ClusterComponent->SelectedInstances.Num() > 0)
-				{
-					// remove any selections in the component
-					Cluster.ClusterComponent->SelectedInstances.Empty();
-					Cluster.ClusterComponent->MarkRenderStateDirty();
-					bNeedsUpdate = true;
-				}
+				// remove any selections in the component
+				MeshInfo.Component->SelectedInstances.Empty();
+				MeshInfo.Component->MarkRenderStateDirty();
 			}
 		}
 	}
-
-	if (bNeedsUpdate)
-	{
-		// Update selection
-		MarkComponentsRenderStateDirty();
-	}
 }
 
 void AInstancedFoliageActor::CheckSelection()
@@ -1294,7 +1203,7 @@ FVector AInstancedFoliageActor::GetSelectionLocation()
 		FFoliageMeshInfo* MeshInfo = FindMesh(SelectedMesh);
 		if (MeshInfo && MeshInfo->SelectedIndices.Num() > 0)
 		{
-			Result = MeshInfo->Instances[MeshInfo->SelectedIndices[0]].Location;
+			Result = MeshInfo->Instances[MeshInfo->SelectedIndices.Array()[0]].Location;
 		}
 	}
 
@@ -1338,7 +1247,6 @@ void AInstancedFoliageActor::MapRebuild()
 					FFoliageInstance const& Instance = MeshInfo.Instances[InstanceIdx];
 
 					FFoliageInstance NewInstance = Instance;
-					NewInstance.ClusterIndex = INDEX_NONE;
 
 					FTransform InstanceToWorld = Instance.GetInstanceWorldTransform();
 					FVector Down(-FVector::UpVector);
@@ -1406,12 +1314,6 @@ void AInstancedFoliageActor::ApplyLevelTransform(const FTransform& LevelTransfor
 				FFoliageComponentHashInfo& Info = It.Value();
 				Info.UpdateLocationFromActor(It.Key());
 			}
-
-			// Recalc cluster bounds
-			for (FFoliageInstanceCluster& Cluster : MeshInfo.InstanceClusters)
-			{
-				Cluster.Bounds = Cluster.ClusterComponent->CalcBounds(Cluster.ClusterComponent->ComponentToWorld);
-			}
 		}
 	}
 }
@@ -1420,12 +1322,8 @@ void AInstancedFoliageActor::ApplyLevelTransform(const FTransform& LevelTransfor
 
 struct FFoliageMeshInfo_Old
 {
-	TArray<FFoliageInstanceCluster> InstanceClusters;
+	TArray<FFoliageInstanceCluster_Deprecated> InstanceClusters;
 	TArray<FFoliageInstance> Instances;
-	//FFoliageInstanceHash* InstanceHash;
-	//TMap<UActorComponent*, FFoliageComponentHashInfo> ComponentHash;
-	//TArray<int32> FreeInstanceIndices;
-	//TArray<int32> SelectedIndices;
 	UFoliageType_InstancedStaticMesh* Settings; // Type remapped via +ActiveClassRedirects
 };
 FArchive& operator<<(FArchive& Ar, FFoliageMeshInfo_Old& MeshInfo)
@@ -1440,6 +1338,9 @@ FArchive& operator<<(FArchive& Ar, FFoliageMeshInfo_Old& MeshInfo)
 void AInstancedFoliageActor::Serialize(FArchive& Ar)
 {
 	Super::Serialize(Ar);
+
+	Ar.UsingCustomVersion(FFoliageCustomVersion::GUID);
+
 	if (Ar.UE4Ver() < VER_UE4_FOLIAGE_SETTINGS_TYPE)
 	{
 		TMap<UStaticMesh*, FFoliageMeshInfo_Old> OldFoliageMeshes;
@@ -1447,7 +1348,7 @@ void AInstancedFoliageActor::Serialize(FArchive& Ar)
 		for (auto& OldMeshInfo : OldFoliageMeshes)
 		{
 			FFoliageMeshInfo NewMeshInfo;
-			NewMeshInfo.InstanceClusters = MoveTemp(OldMeshInfo.Value.InstanceClusters);
+
 #if WITH_EDITORONLY_DATA
 			NewMeshInfo.Instances = MoveTemp(OldMeshInfo.Value.Instances);
 #endif
@@ -1478,32 +1379,14 @@ void AInstancedFoliageActor::Serialize(FArchive& Ar)
 		Ar << FoliageMeshes;
 	}
 
-	if (Ar.UE4Ver() < VER_UE4_FOLIAGE_STATIC_LIGHTING_SUPPORT)
-	{
-		for (auto& MeshInfo : FoliageMeshes)
-		{
-			for (const FFoliageInstanceCluster& Cluster : MeshInfo.Value->InstanceClusters)
-			{
-				if (Ar.UE4Ver() < VER_UE4_FOLIAGE_STATIC_LIGHTING_SUPPORT)
-				{
-					Cluster.ClusterComponent->SetMobility(EComponentMobility::Static);
-				}
-				if (Ar.UE4Ver() < VER_UE4_FOLIAGE_COLLISION)
-				{
-					Cluster.ClusterComponent->SetCollisionEnabled(ECollisionEnabled::NoCollision);
-				}
-			}
-		}
-	}
-
-	// Attach cluster components to a root component
-	if (Ar.UE4Ver() < VER_UE4_ADD_ROOTCOMPONENT_TO_FOLIAGEACTOR && GIsEditor && Ar.IsLoading())
+	// Clean up any old cluster components and convert to hierarchical instanced foliage.
+	if (Ar.CustomVer(FFoliageCustomVersion::GUID) < FFoliageCustomVersion::FoliageUsingHierarchicalISMC)
 	{
 		TArray<UInstancedStaticMeshComponent*> ClusterComponents;
 		GetComponents(ClusterComponents);
 		for (UInstancedStaticMeshComponent* Component : ClusterComponents)
 		{
-			Component->AttachTo(GetRootComponent(), NAME_None, EAttachLocation::KeepWorldPosition);
+			Component->bAutoRegister = false;
 		}
 	}
 }
@@ -1553,17 +1436,10 @@ void AInstancedFoliageActor::PostLoad()
 			// Find the per-mesh info matching the mesh.
 			FFoliageMeshInfo& MeshInfo = *MeshPair.Value;
 
-			// Update the FreeInstanceIndices list and hash.
+			// Update the hash.
 			for (int32 InstanceIdx = 0; InstanceIdx < MeshInfo.Instances.Num(); InstanceIdx++)
 			{
 				FFoliageInstance& Instance = MeshInfo.Instances[InstanceIdx];
-				if (Instance.ClusterIndex == INDEX_NONE)
-				{
-					// Add invalid instances to the FreeInstanceIndices list.
-					MeshInfo.FreeInstanceIndices.Add(InstanceIdx);
-					Instance.Base = nullptr;
-				}
-				else
 				{
 					// Add valid instances to the hash.
 					MeshInfo.InstanceHash->InsertInstance(Instance.Location, InstanceIdx);
@@ -1576,24 +1452,18 @@ void AInstancedFoliageActor::PostLoad()
 				}
 			}
 
-			// Check if any of the clusters are nullptr
-			int32 NumMissingComponents = 0;
-			for (FFoliageInstanceCluster& Cluster : MeshInfo.InstanceClusters)
+			// Convert to Heirarchical foliage
+			if (GetLinkerCustomVersion(FFoliageCustomVersion::GUID) < FFoliageCustomVersion::FoliageUsingHierarchicalISMC)
 			{
-				if (Cluster.ClusterComponent == nullptr)
-				{
-					NumMissingComponents++;
-				}
+				MeshInfo.ReallocateClusters(this, MeshPair.Key);
 			}
-			if (NumMissingComponents > 0)
+
+			if (GetLinkerCustomVersion(FFoliageCustomVersion::GUID) < FFoliageCustomVersion::HierarchicalISMCNonTransactional)
 			{
-				FFormatNamedArguments Arguments;
-				Arguments.Add(TEXT("MissingCount"), NumMissingComponents);
-				Arguments.Add(TEXT("MeshName"), FText::FromString(MeshPair.Key->GetName()));
-				FMessageLog("MapCheck").Warning()
-					->AddToken(FUObjectToken::Create(this))
-					->AddToken(FTextToken::Create(FText::Format(LOCTEXT("MapCheck_Message_FoliageMissingClusterComponent", "Foliage in this map is missing {MissingCount} cluster component(s) for static mesh {MeshName}. Opening the Foliage tool will this problem."), Arguments)))
-					->AddToken(FMapErrorToken::Create(FMapErrors::FoliageMissingClusterComponent));
+				if (MeshInfo.Component)
+				{
+					MeshInfo.Component->ClearFlags(RF_Transactional);
+				}
 			}
 		}
 	}
@@ -1626,9 +1496,9 @@ void AInstancedFoliageActor::AddReferencedObjects(UObject* InThis, FReferenceCol
 		}
 #endif
 
-		for (FFoliageInstanceCluster& Cluster : MeshInfo.InstanceClusters)
+		if (MeshInfo.Component)
 		{
-			Collector.AddReferencedObject(Cluster.ClusterComponent, This);
+			Collector.AddReferencedObject(MeshInfo.Component, This);
 		}
 	}
 	Super::AddReferencedObjects(This, Collector);
@@ -1650,22 +1520,14 @@ void AInstancedFoliageActor::ApplyWorldOffset(const FVector& InOffset, bool bWor
 		{
 			FFoliageMeshInfo& MeshInfo = *MeshPair.Value;
 
-			for (FFoliageInstanceCluster& Cluster : MeshInfo.InstanceClusters)
-			{
-				Cluster.Bounds.Origin += InOffset;
-			}
-
 #if WITH_EDITORONLY_DATA
 			MeshInfo.InstanceHash->Empty();
 			for (int32 InstanceIdx = 0; InstanceIdx < MeshInfo.Instances.Num(); InstanceIdx++)
 			{
 				FFoliageInstance& Instance = MeshInfo.Instances[InstanceIdx];
-				if (Instance.ClusterIndex != INDEX_NONE)
-				{
-					Instance.Location += InOffset;
-					// Rehash instance location
-					MeshInfo.InstanceHash->InsertInstance(Instance.Location, InstanceIdx);
-				}
+				Instance.Location += InOffset;
+				// Rehash instance location
+				MeshInfo.InstanceHash->InsertInstance(Instance.Location, InstanceIdx);
 			}
 
 			for (auto It = MeshInfo.ComponentHash.CreateIterator(); It; ++It)
diff --git a/Engine/Source/Runtime/Engine/Private/InstancedStaticMesh.cpp b/Engine/Source/Runtime/Engine/Private/InstancedStaticMesh.cpp
--- a/Engine/Source/Runtime/Engine/Private/InstancedStaticMesh.cpp
+++ b/Engine/Source/Runtime/Engine/Private/InstancedStaticMesh.cpp
@@ -5,184 +5,16 @@
 =============================================================================*/
 
 #include "EnginePrivate.h"
-#include "PhysicsPublic.h"
-#include "ShaderParameters.h"
-#include "ShaderParameterUtils.h"
-
-#include "Misc/UObjectToken.h"
-#include "Components/InteractiveFoliageComponent.h"
-#include "Components/SplineMeshComponent.h"
-#include "Components/ModelComponent.h"
-#include "Components/NiagaraComponent.h"
-#include "Components/ShapeComponent.h"
-#include "Components/BoxComponent.h"
-#include "Components/CapsuleComponent.h"
-#include "Components/SphereComponent.h"
-#include "Components/DrawSphereComponent.h"
-#include "Components/TextRenderComponent.h"
-#include "Components/VectorFieldComponent.h"
-#include "PhysicsEngine/RadialForceComponent.h"
-#include "GameFramework/SpringArmComponent.h"
-#include "Components/WindDirectionalSourceComponent.h"
-#include "Components/TimelineComponent.h"
-#include "SlateBasics.h"
-#include "NavDataGenerator.h"
-#include "OnlineSubsystemUtils.h"
-#include "AI/Navigation/RecastHelpers.h"
-
-#include "StaticMeshResources.h"
-#include "StaticMeshLight.h"
-#include "SpeedTreeWind.h"
-#include "ComponentInstanceDataCache.h"
-#include "InstancedFoliage.h"
-#include "VertexFactory.h"
-#include "LocalVertexFactory.h"
-
-#if WITH_PHYSX
-#include "PhysicsEngine/PhysXSupport.h"
-#include "Collision/PhysXCollision.h"
-#endif
-
-#if WITH_EDITOR
-#include "LightMap.h"
-#include "ShadowMap.h"
-#include "Logging/MessageLog.h"
-#endif
 #include "NavigationSystemHelpers.h"
 #include "AI/Navigation/NavCollision.h"
 #include "Components/InstancedStaticMeshComponent.h"
+#include "InstancedStaticMesh.h"
+#include "../../Renderer/Private/ScenePrivate.h"
 
-// This must match the maximum a user could specify in the material (see 
-// FHLSLMaterialTranslator::TextureCoordinate), otherwise the material will attempt 
-// to look up a texture coordinate we didn't provide an element for.
-static const int32 InstancedStaticMeshMaxTexCoord = 8;
-
-/*-----------------------------------------------------------------------------
-	FStaticMeshInstanceData
------------------------------------------------------------------------------*/
-
-/** The implementation of the static mesh instance data storage type. */
-class FStaticMeshInstanceData :
-	public FStaticMeshVertexDataInterface,
-	public TResourceArray<FVector4,VERTEXBUFFER_ALIGNMENT>
-{
-public:
-
-	typedef TResourceArray<FVector4,VERTEXBUFFER_ALIGNMENT> ArrayType;
-
-	/**
-	 * Constructor
-	 * @param InNeedsCPUAccess - true if resource array data should be CPU accessible
-	 */
-	FStaticMeshInstanceData(bool InNeedsCPUAccess=false)
-		:	TResourceArray<FVector4,VERTEXBUFFER_ALIGNMENT>(InNeedsCPUAccess)
-	{
-	}
-
-	/**
-	 * Resizes the vertex data buffer, discarding any data which no longer fits.
-	 * @param NumVertices - The number of vertices to allocate the buffer for.
-	 */
-	virtual void ResizeBuffer(uint32 NumInstances)
-	{
-		checkf(0, TEXT("ArrayType::Add is not supported on all platforms"));
-	}
-
-	virtual uint32 GetStride() const
-	{
-		const uint32 VectorsPerInstance = 7;
-		return sizeof(FVector4) * VectorsPerInstance;
-	}
-	virtual uint8* GetDataPointer()
-	{
-		return (uint8*)&(*this)[0];
-	}
-	virtual FResourceArrayInterface* GetResourceArray()
-	{
-		return this;
-	}
-	virtual void Serialize(FArchive& Ar)
-	{
-		TResourceArray<FVector4,VERTEXBUFFER_ALIGNMENT>::BulkSerialize(Ar);
-	}
-
-	void Set(const TArray<FVector4>& RawData)
-	{
-		*((ArrayType*)this) = TArray<FVector4,TAlignedHeapAllocator<VERTEXBUFFER_ALIGNMENT> >(RawData);
-	}
-};
-
-
-/*-----------------------------------------------------------------------------
-	FStaticMeshInstanceBuffer
------------------------------------------------------------------------------*/
-
-
-/** A vertex buffer of positions. */
-class FStaticMeshInstanceBuffer : public FVertexBuffer
-{
-public:
-
-	/** Default constructor. */
-	FStaticMeshInstanceBuffer(ERHIFeatureLevel::Type InFeatureLevel);
-
-	/** Destructor. */
-	~FStaticMeshInstanceBuffer();
-
-	/** Delete existing resources */
-	void CleanUp();
-
-	/**
-	 * Initializes the buffer with the component's data.
-	 * @param InComponent - The owning component
-	 * @param InHitProxies - Array of hit proxies for each instance, if desired.
-	 */
-	void Init(UInstancedStaticMeshComponent* InComponent, const TArray<TRefCountPtr<HHitProxy> >& InHitProxies);
-
-	/** Serializer. */
-	friend FArchive& operator<<(FArchive& Ar, FStaticMeshInstanceBuffer& VertexBuffer);
-
-	/**
-	 * Specialized assignment operator, only used when importing LOD's. 
-	 */
-	void operator=(const FStaticMeshInstanceBuffer &Other);
-
-	// Other accessors.
-	FORCEINLINE uint32 GetStride() const
-	{
-		return Stride;
-	}
-	FORCEINLINE uint32 GetNumInstances() const
-	{
-		return NumInstances;
-	}
-
-	const void* GetRawData() const
-	{
-		return InstanceData->GetDataPointer();
-	}
-
-	// FRenderResource interface.
-	virtual void InitRHI() override;
-	virtual FString GetFriendlyName() const { return TEXT("Static-mesh instances"); }
-
-private:
+const int32 InstancedStaticMeshMaxTexCoord = 8;
 
-	/** The vertex data storage type */
-	FStaticMeshInstanceData* InstanceData;
 
-	/** The cached vertex stride. */
-	uint32 Stride;
-
-	/** The cached number of instances. */
-	uint32 NumInstances;
-
-	/** Allocates the vertex data storage type. */
-	void AllocateData();
-};
-
-
-FStaticMeshInstanceBuffer::FStaticMeshInstanceBuffer(ERHIFeatureLevel::Type InFeatureLevel) : 
+FStaticMeshInstanceBuffer::FStaticMeshInstanceBuffer(ERHIFeatureLevel::Type InFeatureLevel):
 	InstanceData(NULL)
 {
 	SetFeatureLevel(InFeatureLevel);
@@ -203,22 +35,45 @@ void FStaticMeshInstanceBuffer::CleanUp()
 	}
 }
 
+
+static void FillInstanceRenderData(const FInstancedStaticMeshInstanceData& Instance, FVector4* RenderData, float RandomInstanceID, float Z = 0, float W = 0)
+{
+	RenderData[0] = FVector4(Instance.ShadowmapUVBias.X, Instance.ShadowmapUVBias.Y, Z, W);
+
+	// Instance -> local matrix.  Every mesh instance has it's own transformation into
+	// the actor's coordinate space.
+	RenderData[1] = FVector4(Instance.Transform.M[0][0], Instance.Transform.M[1][0], Instance.Transform.M[2][0], Instance.Transform.M[3][0]);
+	RenderData[2] = FVector4(Instance.Transform.M[0][1], Instance.Transform.M[1][1], Instance.Transform.M[2][1], Instance.Transform.M[3][1]);
+	RenderData[3] = FVector4(Instance.Transform.M[0][2], Instance.Transform.M[1][2], Instance.Transform.M[2][2], Instance.Transform.M[3][2]);
+
+	// Instance -> local rotation matrix (3x3)
+	// hide the offset (bias) of the lightmap and the per-instance random id in the matrix's w
+	const FMatrix Inverse = Instance.Transform;
+
+	RenderData[4] = FVector4(Inverse.M[0][0], Inverse.M[1][0], Inverse.M[2][0], Instance.LightmapUVBias.X);
+	RenderData[5] = FVector4(Inverse.M[0][1], Inverse.M[1][1], Inverse.M[2][1], Instance.LightmapUVBias.Y);
+	RenderData[6] = FVector4(Inverse.M[0][2], Inverse.M[1][2], Inverse.M[2][2], RandomInstanceID);
+}
+
 /**
  * Initializes the buffer with the component's data.
  * @param InComponent - The owning component
  */
 void FStaticMeshInstanceBuffer::Init(UInstancedStaticMeshComponent* InComponent, const TArray<TRefCountPtr<HHitProxy> >& InHitProxies)
 {
+	bool bUseRemapTable = InComponent->PerInstanceSMData.Num() == InComponent->InstanceReorderTable.Num();
+
 	NumInstances = InComponent->PerInstanceSMData.Num();
+	int32 NumRemoved = InComponent->RemovedInstances.Num();
 
 	// Allocate the vertex data storage type.
 	AllocateData();
 
 	// We cannot write directly to the data on all platforms,
 	// so we make a TArray of the right type, then assign it
-	TArray<FVector4> RawData;
 	check( GetStride() % sizeof(FVector4) == 0 );
-	RawData.Empty(NumInstances * GetStride() / sizeof(FVector4));
+	InstanceData->Empty((NumInstances + NumRemoved) * GetStride() / sizeof(FVector4));
+	InstanceData->AddUninitialized((NumInstances + NumRemoved) * GetStride() / sizeof(FVector4));
 
 	// @todo: Make LD-customizable per component?
 	const float RandomInstanceIDBase = 0.0f;
@@ -228,17 +83,17 @@ void FStaticMeshInstanceBuffer::Init(UInstancedStaticMeshComponent* InComponent,
 	// given instance index between reattaches
 	FRandomStream RandomStream( InComponent->InstancingRandomSeed );
 
-	FMatrix LocalToWorld = InComponent->GetComponentToWorld().ToMatrixWithScale();
-
-	for (uint32 InstanceIndex = 0; InstanceIndex < NumInstances; InstanceIndex++)
+	for (int32 InstanceIndex = 0; InstanceIndex < (int32)NumInstances; InstanceIndex++)
 	{
 		const FInstancedStaticMeshInstanceData& Instance = InComponent->PerInstanceSMData[InstanceIndex];
 
+		FVector4* InstanceRenderData = InstanceData->GetData() + (bUseRemapTable ? InComponent->InstanceReorderTable[InstanceIndex] : InstanceIndex) * (GetStride() / sizeof(FVector4));
+
 		// X, Y	: Shadow map UV bias
-		// Z, W : Encoded HitProxy ID.
+		// Z, W : Encoded HitProxy ID and selection mask.
 		float Z = 0.f;
 		float W = 0.f;
-		if( InHitProxies.Num() == NumInstances )
+		if (InHitProxies.Num() == NumInstances)
 		{
 			FColor HitProxyColor = InHitProxies[InstanceIndex]->Id.GetColor();
 			Z = (float)HitProxyColor.R;
@@ -246,37 +101,39 @@ void FStaticMeshInstanceBuffer::Init(UInstancedStaticMeshComponent* InComponent,
 		}
 #if WITH_EDITOR
 		// Record if the instance is selected
-		if( InstanceIndex < (uint32)InComponent->SelectedInstances.Num() && InComponent->SelectedInstances[InstanceIndex] )
+		if (InstanceIndex < InComponent->SelectedInstances.Num() && InComponent->SelectedInstances[InstanceIndex])
 		{
 			Z += 256.f;
 		}
 #endif
-		RawData.Add( FVector4( Instance.ShadowmapUVBias.X, Instance.ShadowmapUVBias.Y, Z, W ) );
 
-		// Instance -> local matrix.  Every mesh instance has it's own transformation into
-		// the actor's coordinate space.
-		{
-			const FMatrix Transpose = Instance.Transform.GetTransposed();
+		const float RandomInstanceID = RandomInstanceIDBase + RandomStream.GetFraction() * RandomInstanceIDRange;
 						
-			RawData.Add( FVector4(Transpose.M[0][0], Transpose.M[0][1], Transpose.M[0][2], Transpose.M[0][3]) );
-			RawData.Add( FVector4(Transpose.M[1][0], Transpose.M[1][1], Transpose.M[1][2], Transpose.M[1][3]) );
-			RawData.Add( FVector4(Transpose.M[2][0], Transpose.M[2][1], Transpose.M[2][2], Transpose.M[2][3]) );
+		FillInstanceRenderData(Instance, InstanceRenderData, RandomInstanceID, Z, W);
 		}
 
-		// Instance -> local rotation matrix (3x3)
+	// Hide any removed instances
+	if (NumRemoved)
+	{
+		check(bUseRemapTable);
+
+		for (int32 InstanceIndex = 0; InstanceIndex < NumRemoved; InstanceIndex++)
 		{
-			const float RandomInstanceID = RandomInstanceIDBase + RandomStream.GetFraction() * RandomInstanceIDRange;
-			// hide the offset (bias) of the lightmap and the per-instance random id in the matrix's w
-			const FMatrix Transpose = Instance.Transform.Inverse().GetTransposed();
+			FVector4* RenderData = InstanceData->GetData() + InComponent->RemovedInstances[InstanceIndex] * (GetStride() / sizeof(FVector4));
+
+			RenderData[0] = FVector4(0, 0, 0, 0);
+
+			// transform is 0 matrix
+			RenderData[1] = FVector4(0, 0, 0, 0);
+			RenderData[2] = FVector4(0, 0, 0, 0);
+			RenderData[3] = FVector4(0, 0, 0, 0);
 			
-			RawData.Add( FVector4(Transpose.M[0][0], Transpose.M[0][1], Transpose.M[0][2], Instance.LightmapUVBias.X) );
-			RawData.Add( FVector4(Transpose.M[1][0], Transpose.M[1][1], Transpose.M[1][2], Instance.LightmapUVBias.Y) );
-			RawData.Add( FVector4(Transpose.M[2][0], Transpose.M[2][1], Transpose.M[2][2], RandomInstanceID) );
+			// inverse is identity
+			RenderData[4] = FVector4(1, 0, 0, 0);
+			RenderData[5] = FVector4(0, 1, 0, 0);
+			RenderData[6] = FVector4(0, 0, 1, 0);
 		}
 	}
-
-	// Allocate the vertex data buffer.
-	InstanceData->Set(RawData);
 }
 
 /** Serializer. */
@@ -334,101 +191,6 @@ void FStaticMeshInstanceBuffer::AllocateData()
 
 
 
-/*-----------------------------------------------------------------------------
-	FInstancedStaticMeshVertexFactory
------------------------------------------------------------------------------*/
-
-struct FInstancingUserData
-{
-	struct FInstanceStream
-	{
-		FVector4 InstanceShadowmapUVBias;
-		FVector4 InstanceTransform[3];
-		FVector4 InstanceInverseTransform[3];
-	};
-	class FInstancedStaticMeshRenderData* RenderData;
-
-	int32 StartCullDistance;
-	int32 EndCullDistance;
-
-	bool bRenderSelected;
-	bool bRenderUnselected;
-};
-
-/**
- * A vertex factory for instanced static meshes
- */
-struct FInstancedStaticMeshVertexFactory : public FLocalVertexFactory
-{
-	DECLARE_VERTEX_FACTORY_TYPE(FInstancedStaticMeshVertexFactory);
-public:
-	struct DataType : public FLocalVertexFactory::DataType
-	{
-		/** The stream to read shadow map bias (and random instance ID) from. */
-		FVertexStreamComponent InstancedShadowMapBiasComponent;
-
-		/** The stream to read the mesh transform from. */
-		FVertexStreamComponent InstancedTransformComponent[3];
-
-		/** The stream to read the inverse transform, as well as the Lightmap Bias in 0/1 */
-		FVertexStreamComponent InstancedInverseTransformComponent[3];
-	};
-
-	/**
-	 * Should we cache the material's shadertype on this platform with this vertex factory? 
-	 */
-	static bool ShouldCache(EShaderPlatform Platform, const class FMaterial* Material, const class FShaderType* ShaderType);
-
-	/**
-	 * Modify compile environment to enable instancing
-	 * @param OutEnvironment - shader compile environment to modify
-	 */
-	static void ModifyCompilationEnvironment(EShaderPlatform Platform, const FMaterial* Material, FShaderCompilerEnvironment& OutEnvironment)
-	{
-		OutEnvironment.SetDefine(TEXT("USE_INSTANCING"),TEXT("1"));
-		const bool bInstanced = RHISupportsInstancing(Platform);
-		OutEnvironment.SetDefine(TEXT("USE_INSTANCING_EMULATED"), bInstanced ? TEXT("0") : TEXT("1"));
-	}
-
-	/**
-	 * An implementation of the interface used by TSynchronizedResource to update the resource with new data from the game thread.
-	 */
-	void SetData(const DataType& InData)
-	{
-		Data = InData;
-		UpdateRHI();
-	}
-
-	/**
-	 * Copy the data from another vertex factory
-	 * @param Other - factory to copy from
-	 */
-	void Copy(const FInstancedStaticMeshVertexFactory& Other);
-
-	// FRenderResource interface.
-	virtual void InitRHI() override;
-
-	static FVertexFactoryShaderParameters* ConstructShaderParameters(EShaderFrequency ShaderFrequency);
-
-	/** Make sure we account for changes in the signature of GetStaticBatchElementVisibility() */
-	static CONSTEXPR uint32 NumBitsForVisibilityMask()
-	{		
-		return 8 * sizeof(decltype(((FInstancedStaticMeshVertexFactory*)nullptr)->GetStaticBatchElementVisibility(FSceneView(FSceneViewInitOptions()), nullptr)));
-	}
-
-	/**
-	* Get a bitmask representing the visibility of each FMeshBatch element.
-	*/
-	virtual uint64 GetStaticBatchElementVisibility(const class FSceneView& View, const struct FMeshBatch* Batch) const override
-	{
-		uint32 NumElements = FMath::Min((uint32)Batch->Elements.Num(), NumBitsForVisibilityMask());
-		return (1ULL << (uint64)NumElements) - 1ULL;
-	}
-
-private:
-	DataType Data;
-};
-
 
 
 /**
@@ -552,38 +314,6 @@ void FInstancedStaticMeshVertexFactory::InitRHI()
 	InitDeclaration(Elements,Data);
 }
 
-class FInstancedStaticMeshVertexFactoryShaderParameters : public FLocalVertexFactoryShaderParameters
-{
-	virtual void Bind(const FShaderParameterMap& ParameterMap) override
-	{
-		FLocalVertexFactoryShaderParameters::Bind(ParameterMap);
-
-		InstancingFadeOutParamsParameter.Bind(ParameterMap, TEXT("InstancingFadeOutParams"));
-		CPUInstanceShadowMapBias.Bind(ParameterMap, TEXT("CPUInstanceShadowMapBias"));
-		CPUInstanceTransform.Bind(ParameterMap, TEXT("CPUInstanceTransform"));
-		CPUInstanceInverseTransform.Bind(ParameterMap, TEXT("CPUInstanceInverseTransform"));
-	}
-
-	virtual void SetMesh(FRHICommandList& RHICmdList, FShader* VertexShader,const class FVertexFactory* VertexFactory,const class FSceneView& View,const struct FMeshBatchElement& BatchElement,uint32 DataFlags) const override;
-
-	void Serialize(FArchive& Ar)
-	{
-		FLocalVertexFactoryShaderParameters::Serialize(Ar);
-		Ar << InstancingFadeOutParamsParameter;
-		Ar << CPUInstanceShadowMapBias;
-		Ar << CPUInstanceTransform;
-		Ar << CPUInstanceInverseTransform;
-	}
-
-	virtual uint32 GetSize() const { return sizeof(*this); }
-
-private:
-	FShaderParameter InstancingFadeOutParamsParameter;
-
-	FShaderParameter CPUInstanceShadowMapBias;
-	FShaderParameter CPUInstanceTransform;
-	FShaderParameter CPUInstanceInverseTransform;
-};
 
 FVertexFactoryShaderParameters* FInstancedStaticMeshVertexFactory::ConstructShaderParameters(EShaderFrequency ShaderFrequency)
 {
@@ -595,115 +325,6 @@ IMPLEMENT_VERTEX_FACTORY_TYPE(FInstancedStaticMeshVertexFactory,"LocalVertexFact
 
 
 
-/*-----------------------------------------------------------------------------
-	FInstancedStaticMeshRenderData
------------------------------------------------------------------------------*/
-
-class FInstancedStaticMeshRenderData
-{
-public:
-
-	FInstancedStaticMeshRenderData(UInstancedStaticMeshComponent* InComponent, ERHIFeatureLevel::Type InFeatureLevel)
-	  : Component(InComponent)
-	  , InstanceBuffer(InFeatureLevel)
-	  , LODModels(Component->StaticMesh->RenderData->LODResources)
-	  , FeatureLevel(InFeatureLevel)
-	{
-		// Allocate the vertex factories for each LOD
-		for( int32 LODIndex=0;LODIndex<LODModels.Num();LODIndex++ )
-		{
-			FInstancedStaticMeshVertexFactory* VertexFactory = new(VertexFactories)FInstancedStaticMeshVertexFactory;
-			VertexFactory->SetFeatureLevel(InFeatureLevel);
-		}
-
-		// Create hit proxies for each instance if the component wants
-		if( GIsEditor && InComponent->bHasPerInstanceHitProxies )
-		{
-			for( int32 InstanceIdx=0;InstanceIdx<Component->PerInstanceSMData.Num();InstanceIdx++ )
-			{
-				HitProxies.Add( new HInstancedStaticMeshInstance(InComponent, InstanceIdx) );
-			}
-		}
-
-		// initialize the instance buffer from the component's instances
-		InstanceBuffer.Init(Component, HitProxies);
-		InitResources();
-	}
-
-	~FInstancedStaticMeshRenderData()
-	{
-	}
-
-	void InitResources()
-	{
-		BeginInitResource(&InstanceBuffer);
-
-		// Initialize the static mesh's vertex factory.
-		ENQUEUE_UNIQUE_RENDER_COMMAND_THREEPARAMETER(
-			CallInitStaticMeshVertexFactory,
-			TArray<FInstancedStaticMeshVertexFactory>*,VertexFactories,&VertexFactories,
-			FInstancedStaticMeshRenderData*,InstancedRenderData,this,
-			UStaticMesh*,Parent,Component->StaticMesh,
-		{
-			InitStaticMeshVertexFactories( VertexFactories, InstancedRenderData, Parent );
-		});
-
-		for( int32 LODIndex=0;LODIndex<VertexFactories.Num();LODIndex++ )
-		{
-			BeginInitResource(&VertexFactories[LODIndex]);
-		}
-
-		// register SpeedTree wind with the scene
-		if (Component->StaticMesh->SpeedTreeWind.IsValid())
-		{
-			for (int32 LODIndex = 0; LODIndex < LODModels.Num(); LODIndex++)
-			{
-				Component->GetScene()->AddSpeedTreeWind(&VertexFactories[LODIndex], Component->StaticMesh);
-			}
-		}
-	}
-
-	void ReleaseResources(FSceneInterface* Scene, const UStaticMesh* StaticMesh)
-	{
-		// unregister SpeedTree wind with the scene
-		if (Scene && StaticMesh && StaticMesh->SpeedTreeWind.IsValid())
-		{
-			for (int32 LODIndex = 0; LODIndex < VertexFactories.Num(); LODIndex++)
-			{
-				Scene->RemoveSpeedTreeWind_RenderThread(&VertexFactories[LODIndex], StaticMesh);
-			}
-		}
-
-		InstanceBuffer.ReleaseResource();
-		for( int32 LODIndex=0;LODIndex<VertexFactories.Num();LODIndex++ )
-		{
-			VertexFactories[LODIndex].ReleaseResource();
-		}
-	}
-
-	static void InitStaticMeshVertexFactories(
-		TArray<FInstancedStaticMeshVertexFactory>* VertexFactories,
-		FInstancedStaticMeshRenderData* InstancedRenderData,
-		UStaticMesh* Parent);
-
-	/** Source component */
-	UInstancedStaticMeshComponent* Component;
-
-	/** Instance buffer */
-	FStaticMeshInstanceBuffer InstanceBuffer;
-
-	/** Vertex factory */
-	TArray<FInstancedStaticMeshVertexFactory> VertexFactories;
-
-	/** LOD render data from the static mesh. */
-	TIndirectArray<FStaticMeshLODResources>& LODModels;
-
-	/** Hit proxies for the instances */
-	TArray<TRefCountPtr<HHitProxy> > HitProxies;
-
-	/** Feature level used when creating instance data */
-	ERHIFeatureLevel::Type FeatureLevel;
-};
 
 void FInstancedStaticMeshRenderData::InitStaticMeshVertexFactories(
 		TArray<FInstancedStaticMeshVertexFactory>* VertexFactories,
@@ -862,152 +483,6 @@ void FInstancedStaticMeshRenderData::InitStaticMeshVertexFactories(
 }
 
 
-
-/*-----------------------------------------------------------------------------
-	FInstancedStaticMeshSceneProxy
------------------------------------------------------------------------------*/
-
-class FInstancedStaticMeshSceneProxy : public FStaticMeshSceneProxy
-{
-public:
-
-	FInstancedStaticMeshSceneProxy(UInstancedStaticMeshComponent* InComponent, ERHIFeatureLevel::Type InFeatureLevel)
-	:	FStaticMeshSceneProxy(InComponent)
-	, InstancedRenderData(InComponent, InFeatureLevel)
-#if WITH_EDITOR
-	,	bHasSelectedInstances(InComponent->SelectedInstances.Num() > 0)
-#endif
-	{
-#if WITH_EDITOR
-		if( bHasSelectedInstances )
-		{
-			// if we have selected indices, mark scene proxy as selected.
-			SetSelection_GameThread(true);
-		}
-#endif
-		// Make sure all the materials are okay to be rendered as an instanced mesh.
-		for (int32 LODIndex = 0; LODIndex < LODs.Num(); LODIndex++)
-		{
-			FStaticMeshSceneProxy::FLODInfo& LODInfo = LODs[LODIndex];
-			for (int32 SectionIndex = 0; SectionIndex < LODInfo.Sections.Num(); SectionIndex++)
-			{
-				FStaticMeshSceneProxy::FLODInfo::FSectionInfo& Section = LODInfo.Sections[SectionIndex];
-				if (!Section.Material->CheckMaterialUsage_Concurrent(MATUSAGE_InstancedStaticMeshes))
-				{
-					Section.Material = UMaterial::GetDefaultMaterial(MD_Surface);
-				}
-			}
-		}
-
-		check(InstancedRenderData.InstanceBuffer.GetStride() == sizeof(FInstancingUserData::FInstanceStream));
-
-		const bool bInstanced = RHISupportsInstancing(GetFeatureLevelShaderPlatform(InFeatureLevel));
-
-		// Copy the parameters for LOD - all instances
-		UserData_AllInstances.StartCullDistance = InComponent->InstanceStartCullDistance;
-		UserData_AllInstances.EndCullDistance = InComponent->InstanceEndCullDistance;
-		UserData_AllInstances.bRenderSelected = true;
-		UserData_AllInstances.bRenderUnselected = true;
-		UserData_AllInstances.RenderData = bInstanced ? nullptr : &InstancedRenderData;
-
-		// selected only
-		UserData_SelectedInstances = UserData_AllInstances;
-		UserData_SelectedInstances.bRenderUnselected = false;
-
-		// unselected only
-		UserData_DeselectedInstances = UserData_AllInstances;
-		UserData_DeselectedInstances.bRenderSelected = false;
-	}
-
-	~FInstancedStaticMeshSceneProxy()
-	{
-		InstancedRenderData.ReleaseResources(&GetScene( ), StaticMesh);
-	}
-
-	// FPrimitiveSceneProxy interface.
-	
-	virtual FPrimitiveViewRelevance GetViewRelevance(const FSceneView* View) override
-	{
-		FPrimitiveViewRelevance Result;
-		if(View->Family->EngineShowFlags.InstancedStaticMeshes)
-		{
-			Result = FStaticMeshSceneProxy::GetViewRelevance(View);
-#if WITH_EDITOR
-			// use dynamic path to render selected indices
-			if( bHasSelectedInstances )
-			{
-				Result.bDynamicRelevance = true;
-			}
-#endif
-		}
-		return Result;
-	}
-
-	virtual void GetDynamicMeshElements(const TArray<const FSceneView*>& Views, const FSceneViewFamily& ViewFamily, uint32 VisibilityMap, FMeshElementCollector& Collector) const override;
-
-	virtual void GetDistancefieldAtlasData(FBox& LocalVolumeBounds, FIntVector& OutBlockMin, FIntVector& OutBlockSize, bool& bOutBuiltAsIfTwoSided, bool& bMeshWasPlane, TArray<FMatrix>& ObjectLocalToWorldTransforms) const override;
-
-	virtual int32 GetNumMeshBatches() const override;
-
-	/** Sets up a shadow FMeshBatch for a specific LOD. */
-	virtual bool GetShadowMeshElement(int32 LODIndex, int32 BatchIndex, uint8 InDepthPriorityGroup, FMeshBatch& OutMeshBatch) const override;
-
-	/** Sets up a FMeshBatch for a specific LOD and element. */
-	virtual bool GetMeshElement(int32 LODIndex, int32 BatchIndex, int32 ElementIndex, uint8 InDepthPriorityGroup, const bool bUseSelectedMaterial, const bool bUseHoveredMaterial, FMeshBatch& OutMeshBatch) const override;
-
-	/** Sets up a wireframe FMeshBatch for a specific LOD. */
-	virtual bool GetWireframeMeshElement(int32 LODIndex, int32 BatchIndex, const FMaterialRenderProxy* WireframeRenderProxy, uint8 InDepthPriorityGroup, FMeshBatch& OutMeshBatch) const override;
-
-	/**
-	 * Creates the hit proxies are used when DrawDynamicElements is called.
-	 * Called in the game thread.
-	 * @param OutHitProxies - Hit proxes which are created should be added to this array.
-	 * @return The hit proxy to use by default for elements drawn by DrawDynamicElements.
-	 */
-	virtual HHitProxy* CreateHitProxies(UPrimitiveComponent* Component,TArray<TRefCountPtr<HHitProxy> >& OutHitProxies) override
-	{
-		if( InstancedRenderData.HitProxies.Num() )
-		{
-			// Add any per-instance hit proxies.
-			OutHitProxies += InstancedRenderData.HitProxies;
-
-			// No default hit proxy.
-			return NULL;
-		}
-		else
-		{
-			return FStaticMeshSceneProxy::CreateHitProxies(Component, OutHitProxies);
-		}
-	}
-
-	virtual bool IsDetailMesh() const override
-	{
-		return true;
-	}
-
-private:
-	/** Per component render data */
-	FInstancedStaticMeshRenderData InstancedRenderData;
-
-#if WITH_EDITOR
-	/* If we we have any selected instances */
-	bool bHasSelectedInstances;
-#else
-	static const bool bHasSelectedInstances = false;
-#endif
-
-	/** LOD transition info. */
-	FInstancingUserData UserData_AllInstances;
-	FInstancingUserData UserData_SelectedInstances;
-	FInstancingUserData UserData_DeselectedInstances;
-
-	/** Common path for the Get*MeshElement functions */
-	void SetupInstancedMeshBatch2(int32 LODIndex, TArray<FMeshBatch, TInlineAllocator<1>>& OutMeshBatches) const;
-
-	/** Common path for the Get*MeshElement functions */
-	void SetupInstancedMeshBatch(int32 LODIndex, int32 BatchIndex, FMeshBatch& OutMeshBatch) const;
-};
-
 void FInstancedStaticMeshSceneProxy::GetDynamicMeshElements(const TArray<const FSceneView*>& Views, const FSceneViewFamily& ViewFamily, uint32 VisibilityMap, FMeshElementCollector& Collector) const
 {
 	QUICK_SCOPE_CYCLE_COUNTER(STAT_InstancedStaticMeshSceneProxy_GetMeshElements);
@@ -1094,6 +569,7 @@ void FInstancedStaticMeshSceneProxy::SetupInstancedMeshBatch(int32 LODIndex, int
 	const uint32 NumInstances = InstancedRenderData.InstanceBuffer.GetNumInstances();
 	FMeshBatchElement& BatchElement0 = OutMeshBatch.Elements[0];
 	BatchElement0.UserData = (void*)&UserData_AllInstances;
+	BatchElement0.InstancedLODIndex = LODIndex;
 	BatchElement0.UserIndex = 0;
 
 	if (bInstanced)
@@ -1501,94 +977,6 @@ FBoxSphereBounds UInstancedStaticMeshComponent::CalcBounds(const FTransform& Bou
 }
 
 #if WITH_EDITOR
-/*-----------------------------------------------------------------------------
-	FInstancedStaticMeshStaticLightingMesh
------------------------------------------------------------------------------*/
-
-/**
- * A static lighting mesh class that transforms the points by the per-instance transform of an 
- * InstancedStaticMeshComponent
- */
-class FStaticLightingMesh_InstancedStaticMesh : public FStaticMeshStaticLightingMesh
-{
-public:
-
-	/** Initialization constructor. */
-	FStaticLightingMesh_InstancedStaticMesh(const UInstancedStaticMeshComponent* InPrimitive, int32 LODIndex, int32 InstanceIndex, const TArray<ULightComponent*>& InRelevantLights)
-		: FStaticMeshStaticLightingMesh(InPrimitive, LODIndex, InRelevantLights)
-	{
-		// override the local to world to combine the per instance transform with the component's standard transform
-		SetLocalToWorld(InPrimitive->PerInstanceSMData[InstanceIndex].Transform * InPrimitive->ComponentToWorld.ToMatrixWithScale());
-	}
-};
-
-/*-----------------------------------------------------------------------------
-	FInstancedStaticMeshStaticLightingTextureMapping
------------------------------------------------------------------------------*/
-
-
-/** Represents a static mesh primitive with texture mapped static lighting. */
-class FStaticLightingTextureMapping_InstancedStaticMesh : public FStaticMeshStaticLightingTextureMapping
-{
-public:
-	/** Initialization constructor. */
-	FStaticLightingTextureMapping_InstancedStaticMesh(UInstancedStaticMeshComponent* InPrimitive, int32 LODIndex, int32 InInstanceIndex, FStaticLightingMesh* InMesh, int32 InSizeX, int32 InSizeY, int32 InTextureCoordinateIndex, bool bPerformFullQualityRebuild)
-		: FStaticMeshStaticLightingTextureMapping(InPrimitive, LODIndex, InMesh, InSizeX, InSizeY, InTextureCoordinateIndex, bPerformFullQualityRebuild)
-		, InstanceIndex(InInstanceIndex)
-		, QuantizedData(nullptr)
-		, ShadowMapData()
-		, bComplete(false)
-	{
-	}
-
-	// FStaticLightingTextureMapping interface
-	virtual void Apply(FQuantizedLightmapData* InQuantizedData, const TMap<ULightComponent*, FShadowMapData2D*>& InShadowMapData) override
-	{
-		check(bComplete == false);
-
-		UInstancedStaticMeshComponent* InstancedComponent = Cast<UInstancedStaticMeshComponent>(Primitive.Get());
-
-		if (InstancedComponent)
-		{
-			// Save the static lighting until all of the component's static lighting has been built.
-			QuantizedData = TUniquePtr<FQuantizedLightmapData>(InQuantizedData);
-			ShadowMapData.Empty(InShadowMapData.Num());
-			for (auto& ShadowDataPair : InShadowMapData)
-			{
-				ShadowMapData.Add(ShadowDataPair.Key, TUniquePtr<FShadowMapData2D>(ShadowDataPair.Value));
-			}
-
-			InstancedComponent->ApplyLightMapping(this);
-		}
-
-		bComplete = true;
-	}
-
-	virtual bool DebugThisMapping() const override
-	{
-		return false;
-	}
-
-	virtual FString GetDescription() const override
-	{
-		return FString(TEXT("InstancedSMLightingMapping"));
-	}
-
-private:
-	friend class UInstancedStaticMeshComponent;
-
-	/** The instance of the primitive this mapping represents. */
-	const int32 InstanceIndex;
-
-	// Light/shadow map data stored until all instances for this component are processed
-	// so we can apply them all into one light/shadowmap
-	TUniquePtr<FQuantizedLightmapData> QuantizedData;
-	TMap<ULightComponent*, TUniquePtr<FShadowMapData2D>> ShadowMapData;
-
-	/** Has this mapping already been completed? */
-	bool bComplete;
-};
-
 void UInstancedStaticMeshComponent::GetStaticLightingInfo(FStaticLightingPrimitiveInfo& OutPrimitiveInfo, const TArray<ULightComponent*>& InRelevantLights, const FLightingBuildOptions& Options)
 {
 	if (HasValidSettingsForStaticLighting())
@@ -1725,42 +1113,6 @@ void UInstancedStaticMeshComponent::ApplyLightMapping(FStaticLightingTextureMapp
 }
 #endif
 
-/**
- * Structure that maps a component to it's lighting/instancing specific data which must be the same
- * between all instances that are bound to that component.
- */
-struct FComponentInstanceSharingData
-{
-	/** The component that is associated (owns) this data */
-	UInstancedStaticMeshComponent* Component;
-
-	/** Light map texture */
-	UTexture* LightMapTexture;
-
-	/** Shadow map texture (or NULL if no shadow map) */
-	UTexture* ShadowMapTexture;
-
-
-	FComponentInstanceSharingData()
-		: Component( NULL ),
-		  LightMapTexture( NULL ),
-		  ShadowMapTexture( NULL )
-	{
-	}
-};
-
-
-/**
- * Helper struct to hold information about what components use what lightmap textures
- */
-struct FComponentInstancedLightmapData
-{
-	/** List of all original components and their original instances containing */
-	TMap<UInstancedStaticMeshComponent*, TArray<FInstancedStaticMeshInstanceData> > ComponentInstances;
-
-	/** List of new components */
-	TArray< FComponentInstanceSharingData > SharingData;
-};
 
 void UInstancedStaticMeshComponent::GetLightAndShadowMapMemoryUsage( int32& LightMapMemoryUsage, int32& ShadowMapMemoryUsage ) const
 {
@@ -1788,21 +1140,67 @@ void UInstancedStaticMeshComponent::Serialize(FArchive& Ar)
 #endif
 }
 
-void UInstancedStaticMeshComponent::AddInstance(const FTransform& InstanceTransform)
+int32 UInstancedStaticMeshComponent::AddInstance(const FTransform& InstanceTransform)
 {
+	int InstanceIdx = PerInstanceSMData.Num();
+
 	FInstancedStaticMeshInstanceData* NewInstanceData = new(PerInstanceSMData) FInstancedStaticMeshInstanceData();
-	SetupNewInstanceData(*NewInstanceData, PerInstanceSMData.Num() - 1, InstanceTransform);
+	SetupNewInstanceData(*NewInstanceData, InstanceIdx, InstanceTransform);
+
+#if WITH_EDITOR
+	if (SelectedInstances.Num())
+	{
+		SelectedInstances.Add(false);
+	}
+#endif
 
 	MarkRenderStateDirty();
 
-	UNavigationSystem::UpdateNavOctree(this);
+	PartialNavigationUpdate(InstanceIdx);
+
+	return InstanceIdx;
 }
 
-void UInstancedStaticMeshComponent::AddInstanceWorldSpace(const FTransform& WorldTransform)
+int32 UInstancedStaticMeshComponent::AddInstanceWorldSpace(const FTransform& WorldTransform)
  {
 	// Transform from world space to local space
 	FTransform RelativeTM = WorldTransform.GetRelativeTransform(ComponentToWorld);
-	AddInstance(RelativeTM);
+	return AddInstance(RelativeTM);
+}
+
+bool UInstancedStaticMeshComponent::RemoveInstance(int32 InstanceIndex)
+{
+	if (!PerInstanceSMData.IsValidIndex(InstanceIndex))
+	{
+		return false;
+	}
+
+	// Request navigation update
+	PartialNavigationUpdate(InstanceIndex);
+
+	// remove instance
+	PerInstanceSMData.RemoveAt(InstanceIndex);
+
+#if WITH_EDITOR
+	// remove selection flag if array is filled in
+	if (SelectedInstances.IsValidIndex(InstanceIndex))
+	{
+		SelectedInstances.RemoveAt(InstanceIndex);
+	}
+#endif
+
+	// update the physics state
+	if (bPhysicsStateCreated)
+	{
+		// TODO: it may be possible to instead just update the BodyInstanceIndex for all bodies after the removed instance. 
+		ClearAllInstanceBodies();
+		CreateAllInstanceBodies();
+	}
+
+	// Indicate we need to update render state to reflect changes
+	MarkRenderStateDirty();
+
+	return true;
 }
 
 bool UInstancedStaticMeshComponent::GetInstanceTransform(int32 InstanceIndex, FTransform& OutInstanceTransform, bool bWorldSpace) const
@@ -1830,6 +1228,9 @@ bool UInstancedStaticMeshComponent::UpdateInstanceTransform(int32 InstanceIndex,
 		return false;
 	}
 
+	// Request navigation update
+	PartialNavigationUpdate(InstanceIndex);
+
 	FInstancedStaticMeshInstanceData& InstanceData = PerInstanceSMData[InstanceIndex];
 
 	// Render data uses local transform of the instance
@@ -1847,6 +1248,9 @@ bool UInstancedStaticMeshComponent::UpdateInstanceTransform(int32 InstanceIndex,
 #endif //WITH_PHYSX
 	}
 
+	// Request navigation update
+	PartialNavigationUpdate(InstanceIndex);
+
 	MarkRenderStateDirty();
 
 	return true;
@@ -1857,44 +1261,13 @@ bool UInstancedStaticMeshComponent::ShouldCreatePhysicsState() const
 	return IsRegistered() && (bAlwaysCreatePhysicsState || IsCollisionEnabled());
 }
 
-bool UInstancedStaticMeshComponent::RemoveInstance(int32 InstanceIndex)
-{
-	if (!PerInstanceSMData.IsValidIndex(InstanceIndex))
-	{
-		return false;
-	}
-
-	// remove instance
-	PerInstanceSMData.RemoveAt(InstanceIndex);
-
-#if WITH_EDITOR
-	// remove selection flag if array is filled in
-	if (SelectedInstances.IsValidIndex(InstanceIndex))
-	{
-		SelectedInstances.RemoveAt(InstanceIndex);
-	}
-#endif
-
-	// update the physics state
-	if (bPhysicsStateCreated)
-	{
-		// TODO: it may be possible to instead just update the BodyInstanceIndex for all bodies after the removed instance. 
-		ClearAllInstanceBodies();
-		CreateAllInstanceBodies();
-	}
-
-	// Indicate we need to update render state to reflect changes
-	MarkRenderStateDirty();
-
-	UNavigationSystem::UpdateNavOctree(this);
-
-	return true;
-}
-
 void UInstancedStaticMeshComponent::ClearInstances()
 {
 	// Clear all the per-instance data
 	PerInstanceSMData.Empty();
+	InstanceReorderTable.Empty();
+	RemovedInstances.Empty();
+
 	// Release any physics representations
 	ClearAllInstanceBodies();
 
@@ -1951,65 +1324,61 @@ void UInstancedStaticMeshComponent::SetupNewInstanceData(FInstancedStaticMeshIns
 	}
 }
 
+void UInstancedStaticMeshComponent::PartialNavigationUpdate(int32 InstanceIdx)
+{
+	// Just update everything
+	UNavigationSystem::UpdateNavOctree(this);
+}
+
 bool UInstancedStaticMeshComponent::DoCustomNavigableGeometryExport(struct FNavigableGeometryExport* GeomExport) const
 {
-	if (StaticMesh != NULL)
+	if (StaticMesh && StaticMesh->NavCollision)
 	{
 		UNavCollision* NavCollision = StaticMesh->NavCollision;
-		if (NavCollision && NavCollision->bIsDynamicObstacle)
-		{
-			for (const FInstancedStaticMeshInstanceData& InstanceData : PerInstanceSMData)
-			{
-				const FVector Scale3D = InstanceData.Transform.GetScaleVector();
-				// if any of scales is 0 there's no point in exporting it
-				if (!Scale3D.IsZero())
+		if (NavCollision->bIsDynamicObstacle)
 				{
 					FCompositeNavModifier Modifiers;
-					NavCollision->GetNavigationModifier(Modifiers, FTransform(InstanceData.Transform) * ComponentToWorld);
+			NavCollision->GetNavigationModifier(Modifiers, FTransform::Identity);
 					GeomExport->AddNavModifiers(Modifiers);
 				}
-			}
-		}
-		else if (NavCollision && NavCollision->bHasConvexGeometry)
-		{
-			for (const FInstancedStaticMeshInstanceData& InstanceData : PerInstanceSMData)
-			{
-				const FVector Scale3D = InstanceData.Transform.GetScaleVector();
-				// if any of scales is 0 there's no point in exporting it
-				if (!Scale3D.IsZero())
+		else if (NavCollision->bHasConvexGeometry)
 				{
 					GeomExport->ExportCustomMesh(NavCollision->ConvexCollision.VertexBuffer.GetData(), NavCollision->ConvexCollision.VertexBuffer.Num(),
-						NavCollision->ConvexCollision.IndexBuffer.GetData(), NavCollision->ConvexCollision.IndexBuffer.Num(), FTransform(InstanceData.Transform) * ComponentToWorld);
+				NavCollision->ConvexCollision.IndexBuffer.GetData(), NavCollision->ConvexCollision.IndexBuffer.Num(), FTransform::Identity);
 
 					GeomExport->ExportCustomMesh(NavCollision->TriMeshCollision.VertexBuffer.GetData(), NavCollision->TriMeshCollision.VertexBuffer.Num(),
-						NavCollision->TriMeshCollision.IndexBuffer.GetData(), NavCollision->TriMeshCollision.IndexBuffer.Num(), FTransform(InstanceData.Transform) * ComponentToWorld);
-				}
-			}
+				NavCollision->TriMeshCollision.IndexBuffer.GetData(), NavCollision->TriMeshCollision.IndexBuffer.Num(), FTransform::Identity);
 		}
 		else
 		{
 			UBodySetup* BodySetup = StaticMesh->BodySetup;
 			if (BodySetup)
 			{
-				for (const FInstancedStaticMeshInstanceData& InstanceData : PerInstanceSMData)
-				{
-					const FVector Scale3D = InstanceData.Transform.GetScaleVector();
-					// if any of scales is 0 there's no point in exporting it
-					if (!Scale3D.IsZero())
-					{
-						GeomExport->ExportRigidBodySetup(*BodySetup, FTransform(InstanceData.Transform) * ComponentToWorld);
+				GeomExport->ExportRigidBodySetup(*BodySetup, FTransform::Identity);
 					}
 				}
 
-				//GeomExport->SlopeOverride = BodySetup->WalkableSlopeOverride;
-			}
-		}
+		// Hook per instance transform delegate
+		GeomExport->SetNavDataPerInstanceTransformDelegate(FNavDataPerInstanceTransformDelegate::CreateUObject(this, &UInstancedStaticMeshComponent::GetNavigationPerInstanceTransforms));
 	}
 
 	// we don't want "regular" collision export for this component
 	return false;
 }
 
+void UInstancedStaticMeshComponent::GetNavigationPerInstanceTransforms(const FBox& AreaBox, TArray<FTransform>& InstanceData) const
+{
+	for (const auto& InstancedData : PerInstanceSMData)
+	{
+		//TODO: Is it worth doing per instance bounds check here ?
+		const FTransform InstanceToComponent(InstancedData.Transform);
+		if (!InstanceToComponent.GetScale3D().IsZero())
+		{
+			InstanceData.Add(InstanceToComponent*ComponentToWorld);
+		}
+	}
+}
+
 SIZE_T UInstancedStaticMeshComponent::GetResourceSize( EResourceSizeMode::Type Mode )
 {
 	SIZE_T ResSize = 0;
@@ -2089,13 +1458,107 @@ void FInstancedStaticMeshVertexFactoryShaderParameters::SetMesh( FRHICommandList
 	FLocalVertexFactoryShaderParameters::SetMesh(RHICmdList, VertexShader, VertexFactory, View, BatchElement, DataFlags);
 
 	FRHIVertexShader* VS = VertexShader->GetVertexShader();
+	const FInstancingUserData* InstancingUserData = (const FInstancingUserData*)BatchElement.UserData;
+#if USE_NEW_LOD_TRANSITION_FOR_FOLIAGE
+	if( InstancingWorldViewOriginOneParameter.IsBound() )
+	{
+		FVector4 InstancingViewZCompareZero(MIN_flt, MIN_flt, MAX_flt, MAX_flt);
+		FVector4 InstancingViewZCompareOne(MIN_flt, MIN_flt, MAX_flt, MAX_flt);
+		FVector4 InstancingViewZConstant(ForceInit);
+		FVector4 InstancingWorldViewOriginZero(ForceInit);
+		FVector4 InstancingWorldViewOriginOne(ForceInit);
+		InstancingWorldViewOriginOne.W = 1.0f;
+		if (InstancingUserData && BatchElement.InstancedLODRange && View.GetTemporalLODActive())
+		{
+			float SphereRadius = InstancingUserData->MeshRenderData->Bounds.SphereRadius;
+			float MinSize = CVarFoliageMinimumScreenSize.GetValueOnRenderThread();
+			float LODScale = CVarFoliageLODDistanceScale.GetValueOnRenderThread();
+			float MaxDrawDistanceScale = GetCachedScalabilityCVars().ViewDistanceScale;
+
+			if (BatchElement.InstancedLODIndex)
+			{
+				InstancingViewZConstant.X = -1.0f;
+			}
+			else
+			{
+				// this is the first LOD, so we don't have a fade-in region
+				InstancingViewZConstant.X = 0.0f;
+			}
+			InstancingViewZConstant.Y = 0.0f;
+			InstancingViewZConstant.Z = 1.0f;
+
+			// now we subtract off the lower segments, since they will be incorporated 
+			InstancingViewZConstant.Y -= InstancingViewZConstant.X;
+			InstancingViewZConstant.Z -= InstancingViewZConstant.X + InstancingViewZConstant.Y;
+			//not using W InstancingViewZConstant.W -= InstancingViewZConstant.X + InstancingViewZConstant.Y + InstancingViewZConstant.Z;
+
+			for (int32 SampleIndex = 0; SampleIndex < 2; SampleIndex++)
+			{
+				FVector4& InstancingViewZCompare(SampleIndex ? InstancingViewZCompareOne : InstancingViewZCompareZero);
+				float Fac = View.GetTemporalLODDistanceFactor(SampleIndex) * SphereRadius * SphereRadius * LODScale * LODScale;
+
+				float FinalCull = MAX_flt;
+				if (MinSize > 0.0)
+				{
+					FinalCull = FMath::Sqrt(Fac / MinSize);
+				}
+				if (InstancingUserData->EndCullDistance > 0.0f && InstancingUserData->EndCullDistance < FinalCull)
+				{
+					FinalCull = InstancingUserData->EndCullDistance;
+				}
+				FinalCull *= MaxDrawDistanceScale;
+
+				InstancingViewZCompare.Z = FinalCull;
+				if (BatchElement.InstancedLODIndex < InstancingUserData->MeshRenderData->LODResources.Num() - 1)
+				{
+					float NextCut = InstancingUserData->MeshRenderData->ScreenSize[BatchElement.InstancedLODIndex + 1];
+					InstancingViewZCompare.Z = FMath::Min(FMath::Sqrt(Fac / NextCut), FinalCull);
+				}
+
+				InstancingViewZCompare.X = MIN_flt;
+				if (BatchElement.InstancedLODIndex)
+				{
+					float CurCut = FMath::Sqrt(Fac / InstancingUserData->MeshRenderData->ScreenSize[BatchElement.InstancedLODIndex]);
+					if (CurCut < FinalCull)
+					{
+						InstancingViewZCompare.Y = CurCut;
+
+					}
+					else
+					{
+						// this LOD is completely removed by one of the other two factors
+						InstancingViewZCompare.Y = MIN_flt;
+						InstancingViewZCompare.Z = MIN_flt;
+					}
+				}
+				else
+				{
+					// this is the first LOD, so we don't have a fade-in region
+					InstancingViewZCompare.Y = MIN_flt;
+				}
+			}
+
+
+			InstancingWorldViewOriginZero = View.GetTemporalLODOrigin(0);
+			InstancingWorldViewOriginOne = View.GetTemporalLODOrigin(1);
+
+			float Alpha = View.GetTemporalLODTransition();
+			InstancingWorldViewOriginZero.W = 1.0f - Alpha;
+			InstancingWorldViewOriginOne.W = Alpha;
+		}
+		SetShaderValue(RHICmdList, VS, InstancingViewZCompareZeroParameter, InstancingViewZCompareZero );
+		SetShaderValue(RHICmdList, VS, InstancingViewZCompareOneParameter, InstancingViewZCompareOne );
+		SetShaderValue(RHICmdList, VS, InstancingViewZConstantParameter, InstancingViewZConstant );
+		SetShaderValue(RHICmdList, VS, InstancingWorldViewOriginZeroParameter, InstancingWorldViewOriginZero );
+		SetShaderValue(RHICmdList, VS, InstancingWorldViewOriginOneParameter, InstancingWorldViewOriginOne );
+	}
+#endif
 	if( InstancingFadeOutParamsParameter.IsBound() )
 	{
-		FVector4 InstancingFadeOutParams(0.f,0.f,1.f,1.f);
-
-		const FInstancingUserData* InstancingUserData = (const FInstancingUserData*)BatchElement.UserData;
+		FVector4 InstancingFadeOutParams(MAX_flt,0.f,1.f,1.f);
 		if (InstancingUserData)
 		{
+#if !USE_NEW_LOD_TRANSITION_FOR_FOLIAGE
 			InstancingFadeOutParams.X = InstancingUserData->StartCullDistance;
 			if( InstancingUserData->EndCullDistance > 0 )
 			{
@@ -2112,7 +1575,7 @@ void FInstancedStaticMeshVertexFactoryShaderParameters::SetMesh( FRHICommandList
 			{
 				InstancingFadeOutParams.Y = 0.f;
 			}
-
+#endif
 			InstancingFadeOutParams.Z = InstancingUserData->bRenderSelected ? 1.f : 0.f;
 			InstancingFadeOutParams.W = InstancingUserData->bRenderUnselected ? 1.f : 0.f;
 		}
