Merging development branch (CL 2365488) to //depot/UE4/... Engine only, mostly hierarchical foliage

diff --git a/Engine/Source/Runtime/Engine/Private/InstancedStaticMesh.h b/Engine/Source/Runtime/Engine/Private/InstancedStaticMesh.h
new file mode 100644
--- /dev/null
+++ b/Engine/Source/Runtime/Engine/Private/InstancedStaticMesh.h
@@ -0,0 +1,732 @@
+// Copyright 1998-2014 Epic Games, Inc. All Rights Reserved.
+
+/*=============================================================================
+	InstancedStaticMesh.h: Instanced static mesh header
+=============================================================================*/
+
+#pragma once
+
+#include "PhysicsPublic.h"
+#include "ShaderParameters.h"
+#include "ShaderParameterUtils.h"
+
+#include "Misc/UObjectToken.h"
+#include "Components/InteractiveFoliageComponent.h"
+#include "Components/SplineMeshComponent.h"
+#include "Components/ModelComponent.h"
+#include "Components/NiagaraComponent.h"
+#include "Components/ShapeComponent.h"
+#include "Components/BoxComponent.h"
+#include "Components/CapsuleComponent.h"
+#include "Components/SphereComponent.h"
+#include "Components/DrawSphereComponent.h"
+#include "Components/TextRenderComponent.h"
+#include "Components/VectorFieldComponent.h"
+#include "PhysicsEngine/RadialForceComponent.h"
+#include "GameFramework/SpringArmComponent.h"
+#include "Components/WindDirectionalSourceComponent.h"
+#include "Components/TimelineComponent.h"
+#include "SlateBasics.h"
+#include "NavDataGenerator.h"
+#include "OnlineSubsystemUtils.h"
+#include "AI/Navigation/RecastHelpers.h"
+
+#include "StaticMeshResources.h"
+#include "StaticMeshLight.h"
+#include "SpeedTreeWind.h"
+#include "ComponentInstanceDataCache.h"
+#include "InstancedFoliage.h"
+#include "VertexFactory.h"
+#include "LocalVertexFactory.h"
+
+#if WITH_PHYSX
+#include "PhysicsEngine/PhysXSupport.h"
+#include "Collision/PhysXCollision.h"
+#endif
+
+#if WITH_EDITOR
+#include "LightMap.h"
+#include "ShadowMap.h"
+#include "Logging/MessageLog.h"
+#endif
+
+#include "NavigationSystemHelpers.h"
+#include "AI/Navigation/NavCollision.h"
+#include "Components/InstancedStaticMeshComponent.h"
+
+DECLARE_STATS_GROUP(TEXT("Foliage"), STATGROUP_Foliage, STATCAT_Advanced);
+
+#define USE_NEW_LOD_TRANSITION_FOR_FOLIAGE (1)
+
+extern TAutoConsoleVariable<float> CVarFoliageMinimumScreenSize;
+extern TAutoConsoleVariable<float> CVarFoliageLODDistanceScale;
+
+
+// This must match the maximum a user could specify in the material (see 
+// FHLSLMaterialTranslator::TextureCoordinate), otherwise the material will attempt 
+// to look up a texture coordinate we didn't provide an element for.
+extern const int32 InstancedStaticMeshMaxTexCoord;
+
+/*-----------------------------------------------------------------------------
+	FStaticMeshInstanceData
+-----------------------------------------------------------------------------*/
+
+/** The implementation of the static mesh instance data storage type. */
+class FStaticMeshInstanceData :
+	public FStaticMeshVertexDataInterface,
+	public TResourceArray<FVector4,VERTEXBUFFER_ALIGNMENT>
+{
+public:
+
+	typedef TResourceArray<FVector4,VERTEXBUFFER_ALIGNMENT> ArrayType;
+
+	/**
+	 * Constructor
+	 * @param InNeedsCPUAccess - true if resource array data should be CPU accessible
+	 */
+	FStaticMeshInstanceData(bool InNeedsCPUAccess=false)
+		:	TResourceArray<FVector4,VERTEXBUFFER_ALIGNMENT>(InNeedsCPUAccess)
+	{
+	}
+
+	/**
+	 * Resizes the vertex data buffer, discarding any data which no longer fits.
+	 * @param NumVertices - The number of vertices to allocate the buffer for.
+	 */
+	virtual void ResizeBuffer(uint32 NumInstances)
+	{
+		checkf(0, TEXT("ArrayType::Add is not supported on all platforms"));
+	}
+
+	virtual uint32 GetStride() const
+	{
+		const uint32 VectorsPerInstance = 7;
+		return sizeof(FVector4) * VectorsPerInstance;
+	}
+	virtual uint8* GetDataPointer()
+	{
+		return (uint8*)&(*this)[0];
+	}
+	virtual FResourceArrayInterface* GetResourceArray()
+	{
+		return this;
+	}
+	virtual void Serialize(FArchive& Ar)
+	{
+		TResourceArray<FVector4,VERTEXBUFFER_ALIGNMENT>::BulkSerialize(Ar);
+	}
+
+	void Set(const TArray<FVector4>& RawData)
+	{
+		*((ArrayType*)this) = TArray<FVector4,TAlignedHeapAllocator<VERTEXBUFFER_ALIGNMENT> >(RawData);
+	}
+};
+
+
+/*-----------------------------------------------------------------------------
+	FStaticMeshInstanceBuffer
+-----------------------------------------------------------------------------*/
+
+
+/** A vertex buffer of positions. */
+class FStaticMeshInstanceBuffer : public FVertexBuffer
+{
+public:
+
+	/** Default constructor. */
+	FStaticMeshInstanceBuffer(ERHIFeatureLevel::Type InFeatureLevel);
+
+	/** Destructor. */
+	~FStaticMeshInstanceBuffer();
+
+	/** Delete existing resources */
+	void CleanUp();
+
+	/**
+	 * Initializes the buffer with the component's data.
+	 * @param InComponent - The owning component
+	 * @param InHitProxies - Array of hit proxies for each instance, if desired.
+	 */
+	void Init(UInstancedStaticMeshComponent* InComponent, const TArray<TRefCountPtr<HHitProxy> >& InHitProxies);
+
+	/** Serializer. */
+	friend FArchive& operator<<(FArchive& Ar, FStaticMeshInstanceBuffer& VertexBuffer);
+
+	/**
+	 * Specialized assignment operator, only used when importing LOD's. 
+	 */
+	void operator=(const FStaticMeshInstanceBuffer &Other);
+
+	// Other accessors.
+	FORCEINLINE uint32 GetStride() const
+	{
+		return Stride;
+	}
+	FORCEINLINE uint32 GetNumInstances() const
+	{
+		return NumInstances;
+	}
+
+	const void* GetRawData() const
+	{
+		return InstanceData->GetDataPointer();
+	}
+
+	// FRenderResource interface.
+	virtual void InitRHI() override;
+	virtual FString GetFriendlyName() const { return TEXT("Static-mesh instances"); }
+
+private:
+
+	/** The vertex data storage type */
+	FStaticMeshInstanceData* InstanceData;
+
+	/** The cached vertex stride. */
+	uint32 Stride;
+
+	/** The cached number of instances. */
+	uint32 NumInstances;
+
+	/** Allocates the vertex data storage type. */
+	void AllocateData();
+};
+
+/*-----------------------------------------------------------------------------
+	FInstancedStaticMeshVertexFactory
+-----------------------------------------------------------------------------*/
+
+struct FInstancingUserData
+{
+	struct FInstanceStream
+	{
+		FVector4 InstanceShadowmapUVBias;
+		FVector4 InstanceTransform[3];
+		FVector4 InstanceInverseTransform[3];
+	};
+
+	class FInstancedStaticMeshRenderData* RenderData;
+	class FStaticMeshRenderData* MeshRenderData;
+
+	int32 StartCullDistance;
+	int32 EndCullDistance;
+
+	bool bRenderSelected;
+	bool bRenderUnselected;
+};
+
+/**
+ * A vertex factory for instanced static meshes
+ */
+struct FInstancedStaticMeshVertexFactory : public FLocalVertexFactory
+{
+	DECLARE_VERTEX_FACTORY_TYPE(FInstancedStaticMeshVertexFactory);
+public:
+	struct DataType : public FLocalVertexFactory::DataType
+	{
+		/** The stream to read shadow map bias (and random instance ID) from. */
+		FVertexStreamComponent InstancedShadowMapBiasComponent;
+
+		/** The stream to read the mesh transform from. */
+		FVertexStreamComponent InstancedTransformComponent[3];
+
+		/** The stream to read the inverse transform, as well as the Lightmap Bias in 0/1 */
+		FVertexStreamComponent InstancedInverseTransformComponent[3];
+	};
+
+	/**
+	 * Should we cache the material's shadertype on this platform with this vertex factory? 
+	 */
+	static bool ShouldCache(EShaderPlatform Platform, const class FMaterial* Material, const class FShaderType* ShaderType);
+
+	/**
+	 * Modify compile environment to enable instancing
+	 * @param OutEnvironment - shader compile environment to modify
+	 */
+	static void ModifyCompilationEnvironment(EShaderPlatform Platform, const FMaterial* Material, FShaderCompilerEnvironment& OutEnvironment)
+	{
+		OutEnvironment.SetDefine(TEXT("USE_INSTANCING"),TEXT("1"));
+		const bool bInstanced = RHISupportsInstancing(Platform);
+		OutEnvironment.SetDefine(TEXT("USE_INSTANCING_EMULATED"), bInstanced ? TEXT("0") : TEXT("1"));
+		OutEnvironment.SetDefine(TEXT("USE_NEW_LOD_TRANSITION"), USE_NEW_LOD_TRANSITION_FOR_FOLIAGE ? TEXT("1") : TEXT("0"));
+	}
+
+	/**
+	 * An implementation of the interface used by TSynchronizedResource to update the resource with new data from the game thread.
+	 */
+	void SetData(const DataType& InData)
+	{
+		Data = InData;
+		UpdateRHI();
+	}
+
+	/**
+	 * Copy the data from another vertex factory
+	 * @param Other - factory to copy from
+	 */
+	void Copy(const FInstancedStaticMeshVertexFactory& Other);
+
+	// FRenderResource interface.
+	virtual void InitRHI() override;
+
+	static FVertexFactoryShaderParameters* ConstructShaderParameters(EShaderFrequency ShaderFrequency);
+
+	/** Make sure we account for changes in the signature of GetStaticBatchElementVisibility() */
+	static CONSTEXPR uint32 NumBitsForVisibilityMask()
+	{		
+		return 8 * sizeof(decltype(((FInstancedStaticMeshVertexFactory*)nullptr)->GetStaticBatchElementVisibility(FSceneView(FSceneViewInitOptions()), nullptr)));
+	}
+
+	/**
+	* Get a bitmask representing the visibility of each FMeshBatch element.
+	*/
+	virtual uint64 GetStaticBatchElementVisibility(const class FSceneView& View, const struct FMeshBatch* Batch) const override
+	{
+		uint32 NumElements = FMath::Min((uint32)Batch->Elements.Num(), NumBitsForVisibilityMask());
+		return (1ULL << (uint64)NumElements) - 1ULL;
+	}
+
+private:
+	DataType Data;
+};
+
+
+
+
+class FInstancedStaticMeshVertexFactoryShaderParameters : public FLocalVertexFactoryShaderParameters
+{
+	virtual void Bind(const FShaderParameterMap& ParameterMap) override
+	{
+		FLocalVertexFactoryShaderParameters::Bind(ParameterMap);
+
+		InstancingFadeOutParamsParameter.Bind(ParameterMap, TEXT("InstancingFadeOutParams"));
+#if USE_NEW_LOD_TRANSITION_FOR_FOLIAGE
+		InstancingViewZCompareZeroParameter.Bind(ParameterMap, TEXT("InstancingViewZCompareZero"));
+		InstancingViewZCompareOneParameter.Bind(ParameterMap, TEXT("InstancingViewZCompareOne"));
+		InstancingViewZConstantParameter.Bind(ParameterMap, TEXT("InstancingViewZConstant"));
+		InstancingWorldViewOriginZeroParameter.Bind(ParameterMap, TEXT("InstancingWorldViewOriginZero"));
+		InstancingWorldViewOriginOneParameter.Bind(ParameterMap, TEXT("InstancingWorldViewOriginOne"));
+#endif
+		CPUInstanceShadowMapBias.Bind(ParameterMap, TEXT("CPUInstanceShadowMapBias"));
+		CPUInstanceTransform.Bind(ParameterMap, TEXT("CPUInstanceTransform"));
+		CPUInstanceInverseTransform.Bind(ParameterMap, TEXT("CPUInstanceInverseTransform"));
+	}
+
+	virtual void SetMesh(FRHICommandList& RHICmdList, FShader* VertexShader,const class FVertexFactory* VertexFactory,const class FSceneView& View,const struct FMeshBatchElement& BatchElement,uint32 DataFlags) const override;
+
+	void Serialize(FArchive& Ar)
+	{
+		FLocalVertexFactoryShaderParameters::Serialize(Ar);
+		Ar << InstancingFadeOutParamsParameter;
+#if USE_NEW_LOD_TRANSITION_FOR_FOLIAGE
+		Ar << InstancingViewZCompareZeroParameter;
+		Ar << InstancingViewZCompareOneParameter;
+		Ar << InstancingViewZConstantParameter;
+		Ar << InstancingWorldViewOriginZeroParameter;
+		Ar << InstancingWorldViewOriginOneParameter;
+#endif
+		Ar << CPUInstanceShadowMapBias;
+		Ar << CPUInstanceTransform;
+		Ar << CPUInstanceInverseTransform;
+	}
+
+	virtual uint32 GetSize() const { return sizeof(*this); }
+
+private:
+	FShaderParameter InstancingFadeOutParamsParameter;
+#if USE_NEW_LOD_TRANSITION_FOR_FOLIAGE
+	FShaderParameter InstancingViewZCompareZeroParameter;
+	FShaderParameter InstancingViewZCompareOneParameter;
+	FShaderParameter InstancingViewZConstantParameter;
+	FShaderParameter InstancingWorldViewOriginZeroParameter;
+	FShaderParameter InstancingWorldViewOriginOneParameter;
+#endif
+
+	FShaderParameter CPUInstanceShadowMapBias;
+	FShaderParameter CPUInstanceTransform;
+	FShaderParameter CPUInstanceInverseTransform;
+};
+
+
+/*-----------------------------------------------------------------------------
+	FInstancedStaticMeshRenderData
+-----------------------------------------------------------------------------*/
+
+class FInstancedStaticMeshRenderData
+{
+public:
+
+	FInstancedStaticMeshRenderData(UInstancedStaticMeshComponent* InComponent, ERHIFeatureLevel::Type InFeatureLevel)
+	  : Component(InComponent)
+	  , InstanceBuffer(InFeatureLevel)
+	  , LODModels(Component->StaticMesh->RenderData->LODResources)
+	  , FeatureLevel(InFeatureLevel)
+	{
+		// Allocate the vertex factories for each LOD
+		for( int32 LODIndex=0;LODIndex<LODModels.Num();LODIndex++ )
+		{
+			FInstancedStaticMeshVertexFactory* VertexFactory = new(VertexFactories)FInstancedStaticMeshVertexFactory;
+			VertexFactory->SetFeatureLevel(InFeatureLevel);
+		}
+
+		// Create hit proxies for each instance if the component wants
+		if( GIsEditor && InComponent->bHasPerInstanceHitProxies )
+		{
+			HitProxies.Empty(Component->PerInstanceSMData.Num());
+			for( int32 InstanceIdx=0;InstanceIdx<Component->PerInstanceSMData.Num();InstanceIdx++ )
+			{
+				HitProxies.Add(new HInstancedStaticMeshInstance(InComponent, InstanceIdx));
+			}
+		}
+
+		// initialize the instance buffer from the component's instances
+		InstanceBuffer.Init(Component, HitProxies);
+		InitResources();
+	}
+
+	~FInstancedStaticMeshRenderData()
+	{
+	}
+
+	void InitResources()
+	{
+		BeginInitResource(&InstanceBuffer);
+
+		// Initialize the static mesh's vertex factory.
+		ENQUEUE_UNIQUE_RENDER_COMMAND_THREEPARAMETER(
+			CallInitStaticMeshVertexFactory,
+			TArray<FInstancedStaticMeshVertexFactory>*,VertexFactories,&VertexFactories,
+			FInstancedStaticMeshRenderData*,InstancedRenderData,this,
+			UStaticMesh*,Parent,Component->StaticMesh,
+		{
+			InitStaticMeshVertexFactories( VertexFactories, InstancedRenderData, Parent );
+		});
+
+		for( int32 LODIndex=0;LODIndex<VertexFactories.Num();LODIndex++ )
+		{
+			BeginInitResource(&VertexFactories[LODIndex]);
+		}
+
+		// register SpeedTree wind with the scene
+		if (Component->StaticMesh->SpeedTreeWind.IsValid())
+		{
+			for (int32 LODIndex = 0; LODIndex < LODModels.Num(); LODIndex++)
+			{
+				Component->GetScene()->AddSpeedTreeWind(&VertexFactories[LODIndex], Component->StaticMesh);
+			}
+		}
+	}
+
+	void ReleaseResources(FSceneInterface* Scene, const UStaticMesh* StaticMesh)
+	{
+		// unregister SpeedTree wind with the scene
+		if (Scene && StaticMesh && StaticMesh->SpeedTreeWind.IsValid())
+		{
+			for (int32 LODIndex = 0; LODIndex < VertexFactories.Num(); LODIndex++)
+			{
+				Scene->RemoveSpeedTreeWind_RenderThread(&VertexFactories[LODIndex], StaticMesh);
+			}
+		}
+
+		InstanceBuffer.ReleaseResource();
+		for( int32 LODIndex=0;LODIndex<VertexFactories.Num();LODIndex++ )
+		{
+			VertexFactories[LODIndex].ReleaseResource();
+		}
+	}
+
+	static void InitStaticMeshVertexFactories(
+		TArray<FInstancedStaticMeshVertexFactory>* VertexFactories,
+		FInstancedStaticMeshRenderData* InstancedRenderData,
+		UStaticMesh* Parent);
+
+	/** Source component */
+	UInstancedStaticMeshComponent* Component;
+
+	/** Instance buffer */
+	FStaticMeshInstanceBuffer InstanceBuffer;
+
+	/** Vertex factory */
+	TArray<FInstancedStaticMeshVertexFactory> VertexFactories;
+
+	/** LOD render data from the static mesh. */
+	TIndirectArray<FStaticMeshLODResources>& LODModels;
+
+	/** Hit proxies for the instances */
+	TArray<TRefCountPtr<HHitProxy> > HitProxies;
+
+	/** Feature level used when creating instance data */
+	ERHIFeatureLevel::Type FeatureLevel;
+};
+
+
+/*-----------------------------------------------------------------------------
+	FInstancedStaticMeshSceneProxy
+-----------------------------------------------------------------------------*/
+
+class FInstancedStaticMeshSceneProxy : public FStaticMeshSceneProxy
+{
+public:
+
+	FInstancedStaticMeshSceneProxy(UInstancedStaticMeshComponent* InComponent, ERHIFeatureLevel::Type InFeatureLevel)
+	:	FStaticMeshSceneProxy(InComponent)
+	,	InstancedRenderData(InComponent, InFeatureLevel)
+#if WITH_EDITOR
+	,	bHasSelectedInstances(InComponent->SelectedInstances.Num() > 0)
+#endif
+	{
+#if WITH_EDITOR
+		if( bHasSelectedInstances )
+		{
+			// if we have selected indices, mark scene proxy as selected.
+			SetSelection_GameThread(true);
+		}
+#endif
+		// Make sure all the materials are okay to be rendered as an instanced mesh.
+		for (int32 LODIndex = 0; LODIndex < LODs.Num(); LODIndex++)
+		{
+			FStaticMeshSceneProxy::FLODInfo& LODInfo = LODs[LODIndex];
+			for (int32 SectionIndex = 0; SectionIndex < LODInfo.Sections.Num(); SectionIndex++)
+			{
+				FStaticMeshSceneProxy::FLODInfo::FSectionInfo& Section = LODInfo.Sections[SectionIndex];
+				if (!Section.Material->CheckMaterialUsage_Concurrent(MATUSAGE_InstancedStaticMeshes))
+				{
+					Section.Material = UMaterial::GetDefaultMaterial(MD_Surface);
+				}
+			}
+		}
+
+		check(InstancedRenderData.InstanceBuffer.GetStride() == sizeof(FInstancingUserData::FInstanceStream));
+
+		const bool bInstanced = RHISupportsInstancing(GetFeatureLevelShaderPlatform(InFeatureLevel));
+
+		// Copy the parameters for LOD - all instances
+		UserData_AllInstances.MeshRenderData = InComponent->StaticMesh->RenderData;
+		UserData_AllInstances.StartCullDistance = InComponent->InstanceStartCullDistance;
+		UserData_AllInstances.EndCullDistance = InComponent->InstanceEndCullDistance;
+		UserData_AllInstances.bRenderSelected = true;
+		UserData_AllInstances.bRenderUnselected = true;
+		UserData_AllInstances.RenderData = bInstanced ? nullptr : &InstancedRenderData;
+
+		// selected only
+		UserData_SelectedInstances = UserData_AllInstances;
+		UserData_SelectedInstances.bRenderUnselected = false;
+
+		// unselected only
+		UserData_DeselectedInstances = UserData_AllInstances;
+		UserData_DeselectedInstances.bRenderSelected = false;
+	}
+
+	~FInstancedStaticMeshSceneProxy()
+	{
+		InstancedRenderData.ReleaseResources(&GetScene( ), StaticMesh);
+	}
+
+	// FPrimitiveSceneProxy interface.
+	
+	virtual FPrimitiveViewRelevance GetViewRelevance(const FSceneView* View) override
+	{
+		FPrimitiveViewRelevance Result;
+		if(View->Family->EngineShowFlags.InstancedStaticMeshes)
+		{
+			Result = FStaticMeshSceneProxy::GetViewRelevance(View);
+#if WITH_EDITOR
+			// use dynamic path to render selected indices
+			if( bHasSelectedInstances )
+			{
+				Result.bDynamicRelevance = true;
+			}
+#endif
+		}
+		return Result;
+	}
+
+	virtual void GetDynamicMeshElements(const TArray<const FSceneView*>& Views, const FSceneViewFamily& ViewFamily, uint32 VisibilityMap, FMeshElementCollector& Collector) const override;
+
+	virtual int32 GetNumMeshBatches() const override;
+
+	/** Sets up a shadow FMeshBatch for a specific LOD. */
+	virtual bool GetShadowMeshElement(int32 LODIndex, int32 BatchIndex, uint8 InDepthPriorityGroup, FMeshBatch& OutMeshBatch) const override;
+
+	/** Sets up a FMeshBatch for a specific LOD and element. */
+	virtual bool GetMeshElement(int32 LODIndex, int32 BatchIndex, int32 ElementIndex, uint8 InDepthPriorityGroup, const bool bUseSelectedMaterial, const bool bUseHoveredMaterial, FMeshBatch& OutMeshBatch) const override;
+
+	/** Sets up a wireframe FMeshBatch for a specific LOD. */
+	virtual bool GetWireframeMeshElement(int32 LODIndex, int32 BatchIndex, const FMaterialRenderProxy* WireframeRenderProxy, uint8 InDepthPriorityGroup, FMeshBatch& OutMeshBatch) const override;
+
+	virtual void GetDistancefieldAtlasData(FBox& LocalVolumeBounds, FIntVector& OutBlockMin, FIntVector& OutBlockSize, bool& bOutBuiltAsIfTwoSided, bool& bMeshWasPlane, TArray<FMatrix>& ObjectLocalToWorldTransforms) const override;
+
+	/**
+	 * Creates the hit proxies are used when DrawDynamicElements is called.
+	 * Called in the game thread.
+	 * @param OutHitProxies - Hit proxes which are created should be added to this array.
+	 * @return The hit proxy to use by default for elements drawn by DrawDynamicElements.
+	 */
+	virtual HHitProxy* CreateHitProxies(UPrimitiveComponent* Component,TArray<TRefCountPtr<HHitProxy> >& OutHitProxies) override
+	{
+		if( InstancedRenderData.HitProxies.Num() )
+		{
+			// Add any per-instance hit proxies.
+			OutHitProxies += InstancedRenderData.HitProxies;
+
+			// No default hit proxy.
+			return NULL;
+		}
+		else
+		{
+			return FStaticMeshSceneProxy::CreateHitProxies(Component, OutHitProxies);
+		}
+	}
+
+	virtual bool IsDetailMesh() const override
+	{
+		return true;
+	}
+
+protected:
+	/** Per component render data */
+	FInstancedStaticMeshRenderData InstancedRenderData;
+
+#if WITH_EDITOR
+	/* If we we have any selected instances */
+	bool bHasSelectedInstances;
+#else
+	static const bool bHasSelectedInstances = false;
+#endif
+
+	/** LOD transition info. */
+	FInstancingUserData UserData_AllInstances;
+	FInstancingUserData UserData_SelectedInstances;
+	FInstancingUserData UserData_DeselectedInstances;
+
+	/** Common path for the Get*MeshElement functions */
+	void SetupInstancedMeshBatch(int32 LODIndex, int32 BatchIndex, FMeshBatch& OutMeshBatch) const;
+};
+
+#if WITH_EDITOR
+/*-----------------------------------------------------------------------------
+	FInstancedStaticMeshStaticLightingMesh
+-----------------------------------------------------------------------------*/
+
+/**
+ * A static lighting mesh class that transforms the points by the per-instance transform of an 
+ * InstancedStaticMeshComponent
+ */
+class FStaticLightingMesh_InstancedStaticMesh : public FStaticMeshStaticLightingMesh
+{
+public:
+
+	/** Initialization constructor. */
+	FStaticLightingMesh_InstancedStaticMesh(const UInstancedStaticMeshComponent* InPrimitive, int32 LODIndex, int32 InstanceIndex, const TArray<ULightComponent*>& InRelevantLights)
+		: FStaticMeshStaticLightingMesh(InPrimitive, LODIndex, InRelevantLights)
+	{
+		// override the local to world to combine the per instance transform with the component's standard transform
+		SetLocalToWorld(InPrimitive->PerInstanceSMData[InstanceIndex].Transform * InPrimitive->ComponentToWorld.ToMatrixWithScale());
+	}
+};
+
+/*-----------------------------------------------------------------------------
+	FInstancedStaticMeshStaticLightingTextureMapping
+-----------------------------------------------------------------------------*/
+
+
+/** Represents a static mesh primitive with texture mapped static lighting. */
+class FStaticLightingTextureMapping_InstancedStaticMesh : public FStaticMeshStaticLightingTextureMapping
+{
+public:
+	/** Initialization constructor. */
+	FStaticLightingTextureMapping_InstancedStaticMesh(UInstancedStaticMeshComponent* InPrimitive, int32 LODIndex, int32 InInstanceIndex, FStaticLightingMesh* InMesh, int32 InSizeX, int32 InSizeY, int32 InTextureCoordinateIndex, bool bPerformFullQualityRebuild)
+		: FStaticMeshStaticLightingTextureMapping(InPrimitive, LODIndex, InMesh, InSizeX, InSizeY, InTextureCoordinateIndex, bPerformFullQualityRebuild)
+		, InstanceIndex(InInstanceIndex)
+		, QuantizedData(nullptr)
+		, ShadowMapData()
+		, bComplete(false)
+	{
+	}
+
+	// FStaticLightingTextureMapping interface
+	virtual void Apply(FQuantizedLightmapData* InQuantizedData, const TMap<ULightComponent*, FShadowMapData2D*>& InShadowMapData) override
+	{
+		check(bComplete == false);
+
+		UInstancedStaticMeshComponent* InstancedComponent = Cast<UInstancedStaticMeshComponent>(Primitive.Get());
+
+		if (InstancedComponent)
+		{
+			// Save the static lighting until all of the component's static lighting has been built.
+			QuantizedData = TUniquePtr<FQuantizedLightmapData>(InQuantizedData);
+			ShadowMapData.Empty(InShadowMapData.Num());
+			for (auto& ShadowDataPair : InShadowMapData)
+			{
+				ShadowMapData.Add(ShadowDataPair.Key, TUniquePtr<FShadowMapData2D>(ShadowDataPair.Value));
+			}
+
+			InstancedComponent->ApplyLightMapping(this);
+		}
+
+		bComplete = true;
+	}
+
+	virtual bool DebugThisMapping() const override
+	{
+		return false;
+	}
+
+	virtual FString GetDescription() const override
+	{
+		return FString(TEXT("InstancedSMLightingMapping"));
+	}
+
+private:
+	friend class UInstancedStaticMeshComponent;
+
+	/** The instance of the primitive this mapping represents. */
+	const int32 InstanceIndex;
+
+	// Light/shadow map data stored until all instances for this component are processed
+	// so we can apply them all into one light/shadowmap
+	TUniquePtr<FQuantizedLightmapData> QuantizedData;
+	TMap<ULightComponent*, TUniquePtr<FShadowMapData2D>> ShadowMapData;
+
+	/** Has this mapping already been completed? */
+	bool bComplete;
+};
+
+#endif
+
+/**
+ * Structure that maps a component to it's lighting/instancing specific data which must be the same
+ * between all instances that are bound to that component.
+ */
+struct FComponentInstanceSharingData
+{
+	/** The component that is associated (owns) this data */
+	UInstancedStaticMeshComponent* Component;
+
+	/** Light map texture */
+	UTexture* LightMapTexture;
+
+	/** Shadow map texture (or NULL if no shadow map) */
+	UTexture* ShadowMapTexture;
+
+
+	FComponentInstanceSharingData()
+		: Component( NULL ),
+		  LightMapTexture( NULL ),
+		  ShadowMapTexture( NULL )
+	{
+	}
+};
+
+
+/**
+ * Helper struct to hold information about what components use what lightmap textures
+ */
+struct FComponentInstancedLightmapData
+{
+	/** List of all original components and their original instances containing */
+	TMap<UInstancedStaticMeshComponent*, TArray<FInstancedStaticMeshInstanceData> > ComponentInstances;
+
+	/** List of new components */
+	TArray< FComponentInstanceSharingData > SharingData;
+};
+
diff --git a/Engine/Source/Runtime/Engine/Private/LevelTick.cpp b/Engine/Source/Runtime/Engine/Private/LevelTick.cpp
--- a/Engine/Source/Runtime/Engine/Private/LevelTick.cpp
+++ b/Engine/Source/Runtime/Engine/Private/LevelTick.cpp
@@ -8,6 +8,7 @@
 #include "Components/SceneCaptureComponent2D.h"
 #include "Components/SceneCaptureComponentCube.h"
 #include "Engine/LevelStreamingVolume.h"
+#include "Engine/WorldComposition.h"
 #include "Net/UnrealNetwork.h"
 #include "Collision.h"
 #include "PhysicsPublic.h"
@@ -1200,10 +1201,15 @@ void UWorld::Tick( ELevelTick TickType, float DeltaSeconds )
 		if( !bIsPaused )
 		{
 			// Issues level streaming load/unload requests based on local players being inside/outside level streaming volumes.
-			if( IsGameWorld() && GetNetMode() != NM_Client)
+			if (IsGameWorld() && GetNetMode() != NM_Client)
 			{
 				ProcessLevelStreamingVolumes();
 			}
+
+			if (IsGameWorld() && WorldComposition)
+			{
+				WorldComposition->UpdateStreamingState();
+			}
 		}
 	}
 
diff --git a/Engine/Source/Runtime/Engine/Private/PlayerController.cpp b/Engine/Source/Runtime/Engine/Private/PlayerController.cpp
--- a/Engine/Source/Runtime/Engine/Private/PlayerController.cpp
+++ b/Engine/Source/Runtime/Engine/Private/PlayerController.cpp
@@ -863,7 +863,8 @@ void APlayerController::CalcCamera(float DeltaTime, FMinimalViewInfo& OutResult)
 
 void APlayerController::GetPlayerViewPoint( FVector& out_Location, FRotator& out_Rotation ) const
 {
-	if (PlayerCameraManager != NULL)
+	if (PlayerCameraManager != NULL && 
+		PlayerCameraManager->CameraCache.TimeStamp > 0.f) // Whether camera was updated at least once)
 	{
 		PlayerCameraManager->GetCameraViewPoint(out_Location, out_Rotation);
 	}
diff --git a/Engine/Source/Runtime/Engine/Private/SceneManagement.cpp b/Engine/Source/Runtime/Engine/Private/SceneManagement.cpp
--- a/Engine/Source/Runtime/Engine/Private/SceneManagement.cpp
+++ b/Engine/Source/Runtime/Engine/Private/SceneManagement.cpp
@@ -4,6 +4,50 @@
 #include "StaticMeshResources.h"
 #include "../../Renderer/Private/ScenePrivate.h"
 
+
+static TAutoConsoleVariable<float> CVarLODTemporalLag(
+	TEXT("lod.TemporalLag"),
+	0.5f,
+	TEXT("This controls the the time lag for temporal LOD, in seconds."));
+
+void FTemporalLODState::UpdateTemporalLODTransition(const FViewInfo& View, float LastRenderTime)
+{
+	bool bOk = false;
+	if (!View.bDisableDistanceBasedFadeTransitions)
+	{
+		bOk = true;
+		TemporalLODLag = CVarLODTemporalLag.GetValueOnRenderThread();
+		if (TemporalLODTime[1] < LastRenderTime - TemporalLODLag)
+		{
+			if (TemporalLODTime[0] < TemporalLODTime[1])
+			{
+				TemporalLODViewOrigin[0] = TemporalLODViewOrigin[1];
+				TemporalDistanceFactor[0] = TemporalDistanceFactor[1];
+				TemporalLODTime[0] = TemporalLODTime[1];
+			}
+			TemporalLODViewOrigin[1] = View.ViewMatrices.ViewOrigin;
+			TemporalDistanceFactor[1] = View.GetLODDistanceFactor();
+			TemporalLODTime[1] = LastRenderTime;
+			if (TemporalLODTime[1] <= TemporalLODTime[0])
+			{
+				bOk = false; // we are paused or something or otherwise didn't get a good sample
+			}
+		}
+	}
+	if (!bOk)
+	{
+		TemporalLODViewOrigin[0] = View.ViewMatrices.ViewOrigin;
+		TemporalLODViewOrigin[1] = View.ViewMatrices.ViewOrigin;
+		TemporalDistanceFactor[0] = View.GetLODDistanceFactor();
+		TemporalDistanceFactor[1] = TemporalDistanceFactor[0];
+		TemporalLODTime[0] = LastRenderTime;
+		TemporalLODTime[1] = LastRenderTime;
+		TemporalLODLag = 0.0f;
+	}
+}
+
+
+
 FSimpleElementCollector::~FSimpleElementCollector()
 {
 	// Cleanup the dynamic resources.
diff --git a/Engine/Source/Runtime/Engine/Private/SceneView.cpp b/Engine/Source/Runtime/Engine/Private/SceneView.cpp
--- a/Engine/Source/Runtime/Engine/Private/SceneView.cpp
+++ b/Engine/Source/Runtime/Engine/Private/SceneView.cpp
@@ -435,6 +435,64 @@ FSceneView::FSceneView(const FSceneViewInitOptions& InitOptions)
 #endif
 }
 
+bool FSceneView::GetTemporalLODActive() const
+{
+	if (State)
+	{
+		const FTemporalLODState& LODState = State->GetTemporalLODState();
+		if (LODState.TemporalLODLag != 0.0f)
+		{
+			return true;
+		}
+	}
+	return false;
+}
+
+float FSceneView::GetLODDistanceFactor() const
+{
+	const float ScreenMultiple = FMath::Max(ViewRect.Width() / 2.0f * ViewMatrices.ProjMatrix.M[0][0],
+		ViewRect.Height() / 2.0f * ViewMatrices.ProjMatrix.M[1][1]);
+	float Fac = PI * ScreenMultiple * ScreenMultiple / ViewRect.Area();
+	return Fac;
+}
+
+float FSceneView::GetTemporalLODDistanceFactor(int32 Index) const
+{
+	if (State)
+	{
+		const FTemporalLODState& LODState = State->GetTemporalLODState();
+		if (LODState.TemporalLODLag != 0.0f)
+		{
+			return LODState.TemporalDistanceFactor[Index];
+		}
+	}
+	return GetLODDistanceFactor();
+}
+
+
+FVector FSceneView::GetTemporalLODOrigin(int32 Index) const
+{
+	if (State)
+	{
+		const FTemporalLODState& LODState = State->GetTemporalLODState();
+		if (LODState.TemporalLODLag != 0.0f)
+		{
+			return LODState.TemporalLODViewOrigin[Index];
+		}
+	}
+	return ViewMatrices.ViewOrigin;
+}
+
+float FSceneView::GetTemporalLODTransition() const
+{
+	if (State)
+	{
+		return State->GetTemporalLODTransition();
+	}
+	return 0.0f;
+}
+
+
 void FSceneView::UpdateViewMatrix()
 {
 	FVector StereoViewLocation = ViewLocation;
diff --git a/Engine/Source/Runtime/Engine/Private/Slate/SceneViewport.cpp b/Engine/Source/Runtime/Engine/Private/Slate/SceneViewport.cpp
--- a/Engine/Source/Runtime/Engine/Private/Slate/SceneViewport.cpp
+++ b/Engine/Source/Runtime/Engine/Private/Slate/SceneViewport.cpp
@@ -462,39 +462,37 @@ FReply FSceneViewport::OnMouseMove( const FGeometry& InGeometry, const FPointerE
 	// Start a new reply state
 	CurrentReplyState = FReply::Handled();
 
-	if( !InMouseEvent.GetCursorDelta().IsZero() )
+	UpdateCachedMousePos(InGeometry, InMouseEvent);
+	UpdateCachedGeometry(InGeometry);
+
+	const bool bViewportHasCapture = ViewportWidget.IsValid() && ViewportWidget.Pin()->HasMouseCapture();
+	if ( ViewportClient && GetSizeXY() != FIntPoint::ZeroValue )
 	{
-		UpdateCachedMousePos( InGeometry, InMouseEvent );
-		UpdateCachedGeometry(InGeometry);
+		// Switch to the viewport clients world before processing input
+		FScopedConditionalWorldSwitcher WorldSwitcher(ViewportClient);
 
-		const bool bViewportHasCapture = ViewportWidget.IsValid() && ViewportWidget.Pin()->HasMouseCapture();
-		if( ViewportClient && GetSizeXY() != FIntPoint::ZeroValue )
+		if ( bViewportHasCapture )
 		{
-			// Switch to the viewport clients world before processing input
-			FScopedConditionalWorldSwitcher WorldSwitcher( ViewportClient );
+			ViewportClient->CapturedMouseMove(this, GetMouseX(), GetMouseY());
+		}
+		else
+		{
+			ViewportClient->MouseMove(this, GetMouseX(), GetMouseY());
+		}
 
-			if( bViewportHasCapture )
-			{
-				ViewportClient->CapturedMouseMove( this, GetMouseX(), GetMouseY() );
-			}
-			else
-			{
-				ViewportClient->MouseMove( this, GetMouseX(), GetMouseY() );
-			}
-		
-			if( bViewportHasCapture )
-			{
-				// Accumulate delta changes to mouse movment.  Depending on the sample frequency of a mouse we may get many per frame.
-				//@todo Slate: In directinput, number of samples in x/y could be different...
-				const FVector2D CursorDelta = InMouseEvent.GetCursorDelta();
-				MouseDelta.X += CursorDelta.X;
-				++NumMouseSamplesX;
-
-				MouseDelta.Y -= CursorDelta.Y;
-				++NumMouseSamplesY;
-			}
+		if ( bViewportHasCapture )
+		{
+			// Accumulate delta changes to mouse movement.  Depending on the sample frequency of a mouse we may get many per frame.
+			//@todo Slate: In directinput, number of samples in x/y could be different...
+			const FVector2D CursorDelta = InMouseEvent.GetCursorDelta();
+			MouseDelta.X += CursorDelta.X;
+			++NumMouseSamplesX;
+
+			MouseDelta.Y -= CursorDelta.Y;
+			++NumMouseSamplesY;
 		}
 	}
+
 	return CurrentReplyState;
 }
 
diff --git a/Engine/Source/Runtime/Engine/Private/StaticMeshComponent.cpp b/Engine/Source/Runtime/Engine/Private/StaticMeshComponent.cpp
--- a/Engine/Source/Runtime/Engine/Private/StaticMeshComponent.cpp
+++ b/Engine/Source/Runtime/Engine/Private/StaticMeshComponent.cpp
@@ -458,7 +458,7 @@ void UStaticMeshComponent::OnRegister()
 		}
 	}
 
-	if (StaticMesh != NULL && StaticMesh->SpeedTreeWind.IsValid())
+	if (StaticMesh != NULL && StaticMesh->SpeedTreeWind.IsValid() && GetScene())
 	{
 		for (int32 LODIndex = 0; LODIndex < StaticMesh->RenderData->LODResources.Num(); ++LODIndex)
 		{
@@ -471,7 +471,7 @@ void UStaticMeshComponent::OnRegister()
 
 void UStaticMeshComponent::OnUnregister()
 {
-	if (StaticMesh != NULL && StaticMesh->SpeedTreeWind.IsValid())
+	if (StaticMesh != NULL && StaticMesh->SpeedTreeWind.IsValid() && GetScene())
 	{
 		for (int32 LODIndex = 0; LODIndex < StaticMesh->RenderData->LODResources.Num(); ++LODIndex)
 		{
diff --git a/Engine/Source/Runtime/Engine/Private/UnrealEngine.cpp b/Engine/Source/Runtime/Engine/Private/UnrealEngine.cpp
--- a/Engine/Source/Runtime/Engine/Private/UnrealEngine.cpp
+++ b/Engine/Source/Runtime/Engine/Private/UnrealEngine.cpp
@@ -50,6 +50,7 @@
 #include "GameFramework/Character.h"
 #include "GameFramework/GameMode.h"
 #include "Engine/LevelStreamingVolume.h"
+#include "Engine/WorldComposition.h"
 #include "Engine/LevelScriptActor.h"
 #include "Vehicles/TireType.h"
 
@@ -8947,16 +8948,8 @@ bool UEngine::LoadMap( FWorldContext& WorldContext, FURL URL, class UPendingNetG
 	check(WorldContext.World()->PersistentLevel);
 	LoadPackagesFully(WorldContext.World(), FULLYLOAD_Map, WorldContext.World()->PersistentLevel->GetOutermost()->GetName());
 
-	// Make sure all relevant streaming levels are fully loaded
-	if (WorldContext.World()->WorldComposition)
-	{
-		// Set initial world origin and stream in levels
-		WorldContext.World()->NavigateTo(FIntVector::ZeroValue);
-	}
-	else
-	{
-		WorldContext.World()->FlushLevelStreaming();
-	}
+	// Make sure "always loaded" sub-levels are fully loaded
+	WorldContext.World()->FlushLevelStreaming(nullptr, EFlushLevelStreamingType::Visibility);
 	
 	UNavigationSystem::InitializeForWorld(WorldContext.World(), FNavigationSystem::GameMode);
 	
@@ -9029,6 +9022,62 @@ bool UEngine::LoadMap( FWorldContext& WorldContext, FURL URL, class UPendingNetG
 	return true;
 }
 
+void UEngine::BlockTillLevelStreamingCompleted(UWorld* InWorld)
+{
+	check(InWorld);
+	
+	// Update streaming levels state using streaming volumes
+	if (InWorld->GetNetMode() != NM_Client)
+	{
+		InWorld->ProcessLevelStreamingVolumes();
+	}
+
+	if (InWorld->WorldComposition)
+	{
+		InWorld->WorldComposition->UpdateStreamingState();
+	}
+
+	// Create view collection
+	int32 NumPlayers = GetNumGamePlayers(InWorld);
+	TUniquePtr<FSceneViewFamilyContext> ViewFamily;
+	if (NumPlayers && GameViewport)
+	{
+		ViewFamily = MakeUnique<FSceneViewFamilyContext>(
+			FSceneViewFamily::ConstructionValues(GameViewport->Viewport, InWorld->Scene, GameViewport->EngineShowFlags).SetRealtimeUpdate(true));
+
+		for (int32 PlayerIndex = 0; PlayerIndex < NumPlayers; ++PlayerIndex)
+		{
+			ULocalPlayer* Player = GetGamePlayer(InWorld, PlayerIndex);
+			if (Player && Player->ViewportClient)
+			{
+				FVector ViewLocation;
+				FRotator ViewRotation;
+				Player->CalcSceneView(ViewFamily.Get(), /*out*/ ViewLocation, /*out*/ ViewRotation, Player->ViewportClient->Viewport);
+			}
+		}
+	}
+
+	// Probe if we have anything to do
+	InWorld->UpdateLevelStreaming(ViewFamily.Get());
+	bool bWorkToDo = (InWorld->IsVisibilityRequestPending() || IsAsyncLoading());
+	
+	if (bWorkToDo)
+	{
+		if( GameViewport && GEngine->BeginStreamingPauseDelegate && GEngine->BeginStreamingPauseDelegate->IsBound() )
+		{
+			GEngine->BeginStreamingPauseDelegate->Execute( GameViewport->Viewport );
+		}	
+
+		// Flush level streaming requests, blocking till completion.
+		InWorld->FlushLevelStreaming(ViewFamily.Get(), EFlushLevelStreamingType::Full);
+
+		if( GEngine->EndStreamingPauseDelegate && GEngine->EndStreamingPauseDelegate->IsBound() )
+		{
+			GEngine->EndStreamingPauseDelegate->Execute( );
+		}	
+	}
+}
+
 void UEngine::CleanupPackagesToFullyLoad(FWorldContext &Context, EFullyLoadPackageType FullyLoadType, const FString& Tag)
 {
 	//UE_LOG(LogEngine, Log, TEXT("------------------ CleanupPackagesToFullyLoad: %i, %s"),(int32)FullyLoadType, *Tag);
diff --git a/Engine/Source/Runtime/Engine/Private/World.cpp b/Engine/Source/Runtime/Engine/Private/World.cpp
--- a/Engine/Source/Runtime/Engine/Private/World.cpp
+++ b/Engine/Source/Runtime/Engine/Private/World.cpp
@@ -2276,17 +2276,6 @@ void UWorld::UpdateLevelStreaming( FSceneViewFamily* ViewFamily )
 		return;
 	}
 
-	if (WorldComposition)
-	{
-		// May issue the world origin shift request depending on view location
-		if (ViewFamily)
-		{
-			EvaluateWorldOriginLocation(*ViewFamily);
-		}
-		// May add/remove streaming objects to persistent world depending on view location
-		WorldComposition->UpdateStreamingState(ViewFamily);
-	}
-
 	// Store current count of pending unload levels
 	const int32 NumLevelsPendingPurge = FLevelStreamingGCHelper::GetNumLevelsPendingPurge();
 	// Store current count of async loading packages
@@ -2310,28 +2299,6 @@ void UWorld::UpdateLevelStreaming( FSceneViewFamily* ViewFamily )
 	}
 }
 
-void UWorld::EvaluateWorldOriginLocation(const FSceneViewFamily& ViewFamily)
-{
-	if (IsGameWorld() && GetWorldSettings()->bEnableWorldOriginRebasing)
-	{
-		FVector CentroidLocation(0,0,0);
-		for (int32 ViewIndex = 0; ViewIndex < ViewFamily.Views.Num(); ViewIndex++)
-		{
-			CentroidLocation+= ViewFamily.Views[ViewIndex]->ViewMatrices.ViewOrigin;
-		}
-
-		CentroidLocation/= ViewFamily.Views.Num();
-		// Consider only XY plane
-		CentroidLocation.Z = 0.f;
-	
-		// Request to shift world in case current view is quite far from current origin
-		if (CentroidLocation.Size() > HALF_WORLD_MAX1*0.5f)
-		{
-			RequestNewWorldOrigin(FIntVector(CentroidLocation.X, CentroidLocation.Y, CentroidLocation.Z) + OriginLocation);
-		}
-	}
-}
-
 void UWorld::FlushLevelStreaming(FSceneViewFamily* ViewFamily, EFlushLevelStreamingType FlushType, FName ExcludeType)
 {
 	AWorldSettings* WorldSettings = GetWorldSettings();
diff --git a/Engine/Source/Runtime/Engine/Private/WorldComposition.cpp b/Engine/Source/Runtime/Engine/Private/WorldComposition.cpp
--- a/Engine/Source/Runtime/Engine/Private/WorldComposition.cpp
+++ b/Engine/Source/Runtime/Engine/Private/WorldComposition.cpp
@@ -557,27 +557,74 @@ void UWorldComposition::UpdateStreamingState(const FVector& InLocation)
 	}
 }
 
-void UWorldComposition::UpdateStreamingState(FSceneViewFamily* ViewFamily)
+void UWorldComposition::UpdateStreamingState()
 {
-	// Calculate centroid location
-	FVector ViewLocation(0,0,0);
-	if (ViewFamily &&
-		ViewFamily->Views.Num() > 0)
+	UWorld* PlayWorld = GetWorld();
+	check(PlayWorld);
+
+	// Dedicated server does not use distance based streaming and just loads everything
+	if (PlayWorld->GetNetMode() == NM_DedicatedServer)
+	{
+		UpdateStreamingState(FVector::ZeroVector);
+		return;
+	}
+	
+	int32 NumPlayers = GEngine->GetNumGamePlayers(PlayWorld);
+	if (NumPlayers == 0 || GEngine->GameViewport == nullptr)
+	{
+		return;
+	}
+
+	// calculate centroid location using local players views
+	FSceneViewFamilyContext ViewFamily(
+		FSceneViewFamily::ConstructionValues(GEngine->GameViewport->Viewport, PlayWorld->Scene, GEngine->GameViewport->EngineShowFlags).SetRealtimeUpdate(true));
+
+	int32 NumViews = 0;
+	FVector CentroidLocation = FVector::ZeroVector;
+	
+	for (int32 PlayerIndex = 0; PlayerIndex < NumPlayers; ++PlayerIndex)
 	{
-		// Iterate over all views in collection.
-		for (int32 ViewIndex=0; ViewIndex < ViewFamily->Views.Num(); ViewIndex++)
+		ULocalPlayer* Player = GEngine->GetGamePlayer(PlayWorld, PlayerIndex);
+		if (Player && 
+			Player->ViewportClient && 
+			Player->PlayerController)
 		{
-			ViewLocation+= ViewFamily->Views[ViewIndex]->ViewMatrices.ViewOrigin;
+			FVector ViewLocation;
+			FRotator ViewRotation;
+			if (Player->CalcSceneView(&ViewFamily, /*out*/ ViewLocation, /*out*/ ViewRotation, Player->ViewportClient->Viewport))
+			{
+				CentroidLocation+= ViewLocation;
+				NumViews++;
+			}
 		}
+	}
 
-		ViewLocation/= ViewFamily->Views.Num();
-
-#if WITH_EDITOR
-	LastWorldToViewMatrix = ViewFamily->Views[0]->ViewMatrices.ViewMatrix;
-#endif //WITH_EDITOR
+	// In case there is no valid views don't bother updating level streaming state
+	if (NumViews > 0)
+	{
+		CentroidLocation/= NumViews;
+		if (PlayWorld->GetWorldSettings()->bEnableWorldOriginRebasing)
+		{
+			EvaluateWorldOriginLocation(CentroidLocation);
+		}
+		
+		UpdateStreamingState(CentroidLocation);
 	}
+}
+
+void UWorldComposition::EvaluateWorldOriginLocation(const FVector& ViewLocation)
+{
+	UWorld* OwningWorld = GetWorld();
 	
-	UpdateStreamingState(ViewLocation);
+	FVector Location = ViewLocation;
+	// Consider only XY plane
+	Location.Z = 0.f;
+		
+	// Request to shift world in case current view is quite far from current origin
+	if (Location.Size() > HALF_WORLD_MAX1*0.5f)
+	{
+		OwningWorld->RequestNewWorldOrigin(FIntVector(Location.X, Location.Y, Location.Z) + OwningWorld->OriginLocation);
+	}
 }
 
 bool UWorldComposition::IsDistanceDependentLevel(FName PackageName) const
@@ -623,8 +670,13 @@ void UWorldComposition::CommitTileStreamingState(UWorld* PersistenWorld, int32 T
 	StreamingLevel->bShouldBeLoaded		= bShouldBeLoaded;
 	StreamingLevel->bShouldBeVisible	= bShouldBeVisible;
 	StreamingLevel->LevelLODIndex		= LODIdx;
-}
 
+	// dedicated server always block on load
+	if (PersistenWorld->GetNetMode() == NM_DedicatedServer && bShouldBeLoaded) 
+	{
+		StreamingLevel->bShouldBlockOnLoad = true;
+	}
+}
 
 void UWorldComposition::OnLevelAddedToWorld(ULevel* InLevel)
 {
diff --git a/Engine/Source/Runtime/Engine/Public/AI/Navigation/RecastNavMeshGenerator.h b/Engine/Source/Runtime/Engine/Public/AI/Navigation/RecastNavMeshGenerator.h
--- a/Engine/Source/Runtime/Engine/Public/AI/Navigation/RecastNavMeshGenerator.h
+++ b/Engine/Source/Runtime/Engine/Public/AI/Navigation/RecastNavMeshGenerator.h
@@ -105,6 +105,26 @@ struct FRecastGeometryCache
 	FRecastGeometryCache(const uint8* Memory);
 };
 
+struct FRecastRawGeometryElement
+{
+	// Instance geometry
+	TArray<float>		GeomCoords;
+	TArray<int32>		GeomIndices;
+	
+	// Per instance transformations in unreal coords
+	// When empty geometry is in world space
+	TArray<FTransform>	PerInstanceTransform;
+};
+
+struct FRecastAreaNavModifierElement
+{
+	TArray<FAreaNavModifier> Areas;
+	
+	// Per instance transformations in unreal coords
+	// When empty areas are in world space
+	TArray<FTransform>	PerInstanceTransform;
+};
+
 /**
  * Class handling generation of a single tile, caching data that can speed up subsequent tile generations
  */
@@ -156,16 +176,13 @@ private:
 
 	void ApplyVoxelFilter(struct rcHeightfield* SolidHF, float WalkableRadius);
 
-	/** apply areas from StaticAreas to heightfield */
-	void MarkStaticAreas(FNavMeshBuildContext& BuildContext, rcCompactHeightfield& CompactHF);
-
 	/** apply areas from DynamicAreas to layer */
 	void MarkDynamicAreas(dtTileCacheLayer& Layer);
+	void MarkDynamicArea(const FAreaNavModifier& Modifier, const FTransform& LocalToWorld, dtTileCacheLayer& Layer);
 	
-	void AppendModifier(const FCompositeNavModifier& Modifier, bool bStatic);
+	void AppendModifier(const FCompositeNavModifier& Modifier, const FNavDataPerInstanceTransformDelegate& InTransformsDelegate);
 	/** Appends specified geometry to tile's geometry */
-	void AppendGeometry(const TNavStatArray<uint8>& RawCollisionCache);
-	void AppendGeometry(const TNavStatArray<FVector>& Verts, const TNavStatArray<int32>& Faces);
+	void AppendGeometry(const TNavStatArray<uint8>& RawCollisionCache, const FNavDataPerInstanceTransformDelegate& InTransformsDelegate);
 	void AppendVoxels(rcSpanCache* SpanData, int32 NumSpans);
 	
 	/** prepare voxel cache from collision data */
@@ -185,7 +202,7 @@ protected:
 	FBox TileBB;
 	
 	/** Layers dirty flags */
-	TBitArray<>  DirtyLayers;
+	TBitArray<> DirtyLayers;
 	
 	/** Parameters defining navmesh tiles */
 	FRecastBuildConfig TileConfig;
@@ -201,12 +218,9 @@ protected:
 	TArray<FNavMeshTileData> NavigationData;
 	
 	// tile's geometry: without voxel cache
-	TArray<float> GeomCoords;
-	TArray<int32> GeomIndices;
-	// areas used for creating compressed layers: static zones
-	TArray<FAreaNavModifier> StaticAreas;
+	TArray<FRecastRawGeometryElement> RawGeometry;
 	// areas used for creating navigation data: obstacles
-	TArray<FAreaNavModifier> DynamicAreas;
+	TArray<FRecastAreaNavModifierElement> Modifiers;
 	// navigation links
 	TArray<FSimpleLinkNavModifier> OffmeshLinks;
 };
diff --git a/Engine/Source/Runtime/Engine/Public/AI/NavigationModifier.h b/Engine/Source/Runtime/Engine/Public/AI/NavigationModifier.h
--- a/Engine/Source/Runtime/Engine/Public/AI/NavigationModifier.h
+++ b/Engine/Source/Runtime/Engine/Public/AI/NavigationModifier.h
@@ -92,7 +92,7 @@ protected:
 	TSubclassOf<UNavArea> AreaClass;
 	TSubclassOf<UNavArea> ReplaceAreaClass;
 	FBox Bounds;
-
+	
 	TArray<FVector> Points;
 	TEnumAsByte<ENavigationShapeType::Type> ShapeType;
 
@@ -252,6 +252,14 @@ struct ENGINE_API FCompositeNavModifier : public FNavigationModifier
 	FCompositeNavModifier GetInstantiatedMetaModifier(const struct FNavAgentProperties* NavAgent, TWeakObjectPtr<UObject> WeakOwnerPtr) const;
 	uint32 GetAllocatedSize() const;
 
+	bool HasPerInstanceTransforms() const;
+	// Should be called only on game thread
+	void GetPerInstanceTransforms(const FBox& AreaBox, TArray<FTransform>& PerInstanceTransforms) const;
+
+public:
+	// Gathers per instance data for navigation area modifiers in a specified area box
+	FNavDataPerInstanceTransformDelegate NavDataPerInstanceTransformDelegate;
+
 private:
 	TArray<FAreaNavModifier> Areas;
 	TArray<FSimpleLinkNavModifier> SimpleLinks;
diff --git a/Engine/Source/Runtime/Engine/Public/AI/NavigationOctree.h b/Engine/Source/Runtime/Engine/Public/AI/NavigationOctree.h
--- a/Engine/Source/Runtime/Engine/Public/AI/NavigationOctree.h
+++ b/Engine/Source/Runtime/Engine/Public/AI/NavigationOctree.h
@@ -44,6 +44,9 @@ struct ENGINE_API FNavigationRelevantData
 	/** additional modifiers: areas and external links */
 	FCompositeNavModifier Modifiers;
 
+	// Gathers per instance data for navigation geometry in a specified area box
+	FNavDataPerInstanceTransformDelegate NavDataPerInstanceTransformDelegate;
+
 	FORCEINLINE bool HasGeometry() const { return VoxelData.Num() || CollisionData.Num(); }
 	FORCEINLINE bool HasModifiers() const { return !Modifiers.IsEmpty(); }
 	FORCEINLINE bool IsEmpty() const { return !HasGeometry() && !HasModifiers(); }
@@ -54,7 +57,8 @@ struct ENGINE_API FNavigationRelevantData
 			(HasModifiers() ? ENavigationDirtyFlag::DynamicModifier : 0) |
 			(Modifiers.HasAgentHeightAdjust() ? ENavigationDirtyFlag::UseAgentHeight : 0);
 	}
-
+	
+	bool HasPerInstanceTransforms() const;
 	bool IsMatchingFilter(const FNavigationOctreeFilter& Filter) const;
 	void Shrink();
 };
@@ -147,6 +151,9 @@ public:
 	/** Append new data to existing node */
 	void AppendToNode(const FOctreeElementId& Id, INavRelevantInterface* NavElement, const FBox& Bounds, FNavigationOctreeElement& Data);
 
+	/** Updates element bounds remove/add operation */
+	void UpdateNode(const FOctreeElementId& Id, const FBox& NewBounds);
+
 	/** Remove node */
 	void RemoveNode(const FOctreeElementId& Id);
 
diff --git a/Engine/Source/Runtime/Engine/Public/AI/NavigationSystemHelpers.h b/Engine/Source/Runtime/Engine/Public/AI/NavigationSystemHelpers.h
--- a/Engine/Source/Runtime/Engine/Public/AI/NavigationSystemHelpers.h
+++ b/Engine/Source/Runtime/Engine/Public/AI/NavigationSystemHelpers.h
@@ -33,6 +33,9 @@ struct FNavigableGeometryExport
 	virtual void ExportCustomMesh(const FVector* VertexBuffer, int32 NumVerts, const int32* IndexBuffer, int32 NumIndices, const FTransform& LocalToWorld) = 0;
 	
 	virtual void AddNavModifiers(const FCompositeNavModifier& Modifiers) = 0;
+
+	// Optional delegate for geometry per instance transforms
+	virtual void SetNavDataPerInstanceTransformDelegate(const FNavDataPerInstanceTransformDelegate& InDelegate) = 0;
 };
 
 
diff --git a/Engine/Source/Runtime/Engine/Public/InstancedFoliage.h b/Engine/Source/Runtime/Engine/Public/InstancedFoliage.h
--- a/Engine/Source/Runtime/Engine/Public/InstancedFoliage.h
+++ b/Engine/Source/Runtime/Engine/Public/InstancedFoliage.h
@@ -10,6 +10,7 @@
 // Forward declarations.
 //
 class UInstancedStaticMeshComponent;
+class UHierarchicalInstancedStaticMeshComponent;
 class AInstancedFoliageActor;
 class UFoliageType;
 struct FFoliageInstanceHash;
@@ -52,11 +53,9 @@ struct FFoliageInstancePlacementInfo
 struct FFoliageInstance : public FFoliageInstancePlacementInfo
 {
 	UPrimitiveComponent* Base;
-	int32 ClusterIndex;	// -1/INDEX_NONE if this instance is invalid and its array slot can be reused.
 
 	FFoliageInstance()
-		: Base(NULL)
-		, ClusterIndex(INDEX_NONE)
+	: Base(NULL)
 	{}
 
 
@@ -96,27 +95,6 @@ struct FFoliageInstance : public FFoliageInstancePlacementInfo
 	}
 };
 
-/**
- *	FFoliageInstanceCluster - render info for a cluster of meshes
- */
-struct FFoliageInstanceCluster
-{
-	UInstancedStaticMeshComponent* ClusterComponent;
-	FBoxSphereBounds Bounds;
-
-#if WITH_EDITORONLY_DATA
-	TArray<int32> InstanceIndices;	// index into editor editor Instances array
-#endif
-
-	FFoliageInstanceCluster()
-	{}
-	FFoliageInstanceCluster(UInstancedStaticMeshComponent* InClusterComponent, const FBoxSphereBounds& InBounds)
-		: ClusterComponent(InClusterComponent)
-		, Bounds(InBounds)
-	{}
-	friend FArchive& operator<<(FArchive& Ar, FFoliageInstanceCluster& Cluster);
-};
-
 /**
  * FFoliageComponentHashInfo
  * Cached instance list and component location info stored in the ComponentHash.
@@ -172,8 +150,7 @@ struct FFoliageComponentHashInfo
  */
 struct FFoliageMeshInfo
 {
-	// Cluster allocation data and pointers to components
-	TArray<FFoliageInstanceCluster> InstanceClusters;
+	UHierarchicalInstancedStaticMeshComponent* Component;
 
 #if WITH_EDITORONLY_DATA
 	// Editor-only placed instances
@@ -185,11 +162,8 @@ struct FFoliageMeshInfo
 	// Transient, editor-only set of instances per component
 	TMap<UActorComponent*, FFoliageComponentHashInfo> ComponentHash;
 
-	// Transient, editor-only list of free instance slots.
-	TArray<int32> FreeInstanceIndices;
-
 	// Transient, editor-only list of selected instances.
-	TArray<int32> SelectedIndices;
+	TSet<int32> SelectedIndices;
 #endif
 
 	ENGINE_API FFoliageMeshInfo();
@@ -199,12 +173,11 @@ struct FFoliageMeshInfo
 
 	FFoliageMeshInfo(FFoliageMeshInfo&& Other)
 		// even VC++2013 doesn't support "=default" on move constructors
-		: InstanceClusters(MoveTemp(Other.InstanceClusters))
+		: Component(Other.Component)
 #if WITH_EDITORONLY_DATA
 		, Instances(MoveTemp(Other.Instances))
 		, InstanceHash(MoveTemp(Other.InstanceHash))
 		, ComponentHash(MoveTemp(Other.ComponentHash))
-		, FreeInstanceIndices(MoveTemp(Other.FreeInstanceIndices))
 		, SelectedIndices(MoveTemp(Other.SelectedIndices))
 #endif
 	{ }
@@ -212,12 +185,11 @@ struct FFoliageMeshInfo
 	FFoliageMeshInfo& operator=(FFoliageMeshInfo&& Other)
 		// even VC++2013 doesn't support "=default" on move assignment
 	{
-		InstanceClusters = MoveTemp(Other.InstanceClusters);
+		Component = Other.Component;
 #if WITH_EDITORONLY_DATA
 		Instances = MoveTemp(Other.Instances);
 		InstanceHash = MoveTemp(Other.InstanceHash);
 		ComponentHash = MoveTemp(Other.ComponentHash);
-		FreeInstanceIndices = MoveTemp(Other.FreeInstanceIndices);
 		SelectedIndices = MoveTemp(Other.SelectedIndices);
 #endif
 
@@ -237,8 +209,9 @@ struct FFoliageMeshInfo
 
 	// Destroy existing clusters and reassign all instances to new clusters
 	ENGINE_API void ReallocateClusters(AInstancedFoliageActor* InIFA, UFoliageType* InSettings);
-	// Update settings in the clusters based on the current settings (eg culling distance, collision, ...)
-	//ENGINE_API void UpdateClusterSettings(AInstancedFoliageActor* InIFA);
+
+	ENGINE_API void ReapplyInstancesToComponent();
+
 	ENGINE_API void SelectInstances(AInstancedFoliageActor* InIFA, bool bSelect, TArray<int32>& Instances);
 
 	// Get the number of placed instances
diff --git a/Engine/Source/Runtime/Engine/Public/MeshBatch.h b/Engine/Source/Runtime/Engine/Public/MeshBatch.h
--- a/Engine/Source/Runtime/Engine/Public/MeshBatch.h
+++ b/Engine/Source/Runtime/Engine/Public/MeshBatch.h
@@ -18,32 +18,38 @@ struct FMeshBatchElement
 	TUniformBufferRef<FPrimitiveUniformShaderParameters> PrimitiveUniformBuffer;
 
 	const FIndexBuffer* IndexBuffer;
+	uint32* InstanceRuns;
+	const void* UserData;
+	/** 
+	 *	DynamicIndexData - pointer to user memory containing the index data.
+	 *	Used for rendering dynamic data directly.
+	 */
+	const void* DynamicIndexData;
 	uint32 FirstIndex;
 	uint32 NumPrimitives;
 	uint32 NumInstances;
 	uint32 MinVertexIndex;
 	uint32 MaxVertexIndex;
 	int32 UserIndex;
-	const void* UserData;
 	float MinScreenSize;
 	float MaxScreenSize;
 
-	/** 
-	 *	DynamicIndexData - pointer to user memory containing the index data.
-	 *	Used for rendering dynamic data directly.
-	 */
-	const void* DynamicIndexData;
 	uint16 DynamicIndexStride;
+	uint8 InstancedLODIndex;
+	uint8 InstancedLODRange;
 	
 	FMeshBatchElement()
-	:	PrimitiveUniformBufferResource(NULL)
-	,	IndexBuffer(NULL)
+	:	PrimitiveUniformBufferResource(nullptr)
+	,	IndexBuffer(nullptr)
+	,	InstanceRuns(nullptr)
+	,	UserData(nullptr)
+	,	DynamicIndexData(nullptr)
 	,	NumInstances(1)
 	,	UserIndex(-1)
-	,	UserData(NULL)
 	,	MinScreenSize(0.0f)
 	,	MaxScreenSize(1.0f)
-	,	DynamicIndexData(NULL)
+	,	InstancedLODIndex(0)
+	,	InstancedLODRange(0)
 	{
 	}
 };
@@ -149,7 +155,17 @@ struct FMeshBatch
 		int32 Count=0;
 		for( int32 ElementIdx=0;ElementIdx<Elements.Num();ElementIdx++ )
 		{
-			Count += Elements[ElementIdx].NumPrimitives * Elements[ElementIdx].NumInstances;
+			if (Elements[ElementIdx].InstanceRuns)
+			{
+				for (uint32 Run = 0; Run < Elements[ElementIdx].NumInstances; Run++)
+				{
+					Count += Elements[ElementIdx].NumPrimitives * (Elements[ElementIdx].InstanceRuns[Run * 2 + 1] - Elements[ElementIdx].InstanceRuns[Run * 2] + 1);
+				}
+			}
+			else
+			{
+				Count += Elements[ElementIdx].NumPrimitives * Elements[ElementIdx].NumInstances;
+			}
 		}
 		return Count;
 	}
diff --git a/Engine/Source/Runtime/Engine/Public/SceneManagement.h b/Engine/Source/Runtime/Engine/Public/SceneManagement.h
--- a/Engine/Source/Runtime/Engine/Public/SceneManagement.h
+++ b/Engine/Source/Runtime/Engine/Public/SceneManagement.h
@@ -31,6 +31,39 @@ DECLARE_LOG_CATEGORY_EXTERN(LogBufferVisualization, Log, All);
 // -----------------------------------------------------------------------------
 
 
+/**
+ * struct to hold the temporal LOD state within a view state
+ */
+struct ENGINE_API FTemporalLODState
+{
+	/** The last two camera origin samples collected for stateless temporal LOD transitions */
+	FVector	TemporalLODViewOrigin[2];
+	/** The last two fov-like parameters from the projection matrix for stateless temporal LOD transitions */
+	float	TemporalDistanceFactor[2];
+	/** The last two time samples collected for stateless temporal LOD transitions */
+	float	TemporalLODTime[2];
+	/** If non-zero, then we are doing temporal LOD smoothing, this is the time interval. */
+	float	TemporalLODLag;
+
+	FTemporalLODState()
+		: TemporalLODLag(0.0f) // nothing else is used if this is zero
+	{
+
+	}
+	/** 
+	 * Returns the blend factor between the last two LOD samples
+	 */
+	float GetTemporalLODTransition(float LastRenderTime) const
+	{
+		if (TemporalLODLag == 0.0)
+		{
+			return 0.0f; // no fade
+		}
+		return FMath::Clamp((LastRenderTime - TemporalLODLag - TemporalLODTime[0]) / (TemporalLODTime[1] - TemporalLODTime[0]), 0.0f, 1.0f);
+	}
+
+	void UpdateTemporalLODTransition(const class FViewInfo& View, float LastRenderTime);
+};
 
 /**
  * The scene manager's persistent view state.
@@ -103,6 +136,13 @@ public:
 	virtual void OnStartPostProcessing(FSceneView& CurrentView) = 0;
 	/** Allows MIDs being created and released during view rendering without the overhead of creating and relasing objects */
 	virtual UMaterialInstanceDynamic* GetReusableMID(class UMaterialInterface* ParentMaterial) = 0;
+	/** Returns the temporal LOD struct from the viewstate */
+	virtual FTemporalLODState& GetTemporalLODState() = 0;
+	virtual const FTemporalLODState& GetTemporalLODState() const = 0;
+	/** 
+	 * Returns the blend factor between the last two LOD samples
+	 */
+	virtual float GetTemporalLODTransition() const = 0;
 protected:
 	// Don't allow direct deletion of the view state, Destroy should be called instead.
 	virtual ~FSceneViewStateInterface() {}
diff --git a/Engine/Source/Runtime/Engine/Public/SceneView.h b/Engine/Source/Runtime/Engine/Public/SceneView.h
--- a/Engine/Source/Runtime/Engine/Public/SceneView.h
+++ b/Engine/Source/Runtime/Engine/Public/SceneView.h
@@ -136,6 +136,8 @@ struct FViewMatrices
 		ViewMatrix.SetIdentity();
 		TranslatedViewProjectionMatrix.SetIdentity();
 		InvTranslatedViewProjectionMatrix.SetIdentity();
+		GetDynamicMeshElementsShadowCullFrustum = nullptr;
+		PreShadowTranslation = FVector::ZeroVector;
 		PreViewTranslation = FVector::ZeroVector;
 		ViewOrigin = FVector::ZeroVector;
 	}
@@ -148,6 +150,10 @@ struct FViewMatrices
 	FMatrix		TranslatedViewProjectionMatrix;
 	/** The inverse view-projection transform, ending with world-space points translated by -ViewOrigin. */
 	FMatrix		InvTranslatedViewProjectionMatrix;
+	/** During GetDynamicMeshElements this will be the correct cull volume for shadow stuff */
+	const FConvexVolume* GetDynamicMeshElementsShadowCullFrustum;
+	/** If the above is non-null, a translation that is applied to world-space before transforming by one of the shadow matrices. */
+	FVector PreShadowTranslation;
 	/** The translation to apply to the world before TranslatedViewProjectionMatrix. Usually it is -ViewOrigin but with rereflections this can differ */
 	FVector		PreViewTranslation;
 	/** To support ortho and other modes this is redundant, in world space */
@@ -534,6 +540,31 @@ public:
 	/** @return true:perspective, false:orthographic */
 	inline bool IsPerspectiveProjection() const { return ViewMatrices.IsPerspectiveProjection(); }
 
+	/** Returns true if temporal LOD is active **/
+	bool GetTemporalLODActive() const;
+
+	/** Returns the location used as the origin for LOD computations
+	 * @param Index, 0 or 1, which LOD origin to return
+	 * @return LOD origin
+	 */
+	FVector GetTemporalLODOrigin(int32 Index) const;
+
+	/** Get LOD distance factor: Sqrt(GetLODDistanceFactor()*SphereRadius*SphereRadius / ScreenPercentage) = distance to this LOD transition
+	 * @return distance factor
+	 */
+	float GetLODDistanceFactor() const;
+
+	/** Get LOD distance factor for temporal LOD: Sqrt(GetTemporalLODDistanceFactor(?)*SphereRadius*SphereRadius / ScreenPercentage) = distance to this LOD transition
+	 * @param Index, 0 or 1, which temporal sample to return
+	 * @return distance factor
+	 */
+	float GetTemporalLODDistanceFactor(int32 Index) const;
+
+	/** 
+	 * Returns the blend factor between the last two LOD samples
+	 */
+	float GetTemporalLODTransition() const;
+
 	/** Allow things like HMD displays to update the view matrix at the last minute, to minimize perceived latency */
 	void UpdateViewMatrix();
 
diff --git a/Engine/Source/Runtime/IOS/MetalRHI/Private/MetalCommands.cpp b/Engine/Source/Runtime/IOS/MetalRHI/Private/MetalCommands.cpp
--- a/Engine/Source/Runtime/IOS/MetalRHI/Private/MetalCommands.cpp
+++ b/Engine/Source/Runtime/IOS/MetalRHI/Private/MetalCommands.cpp
@@ -563,12 +563,13 @@ void FMetalDynamicRHI::RHIDrawPrimitiveIndirect(uint32 PrimitiveType, FVertexBuf
 }
 
 
-void FMetalDynamicRHI::RHIDrawIndexedPrimitive(FIndexBufferRHIParamRef IndexBufferRHI, uint32 PrimitiveType, int32 BaseVertexIndex, uint32 MinIndex,
+void FMetalDynamicRHI::RHIDrawIndexedPrimitive(FIndexBufferRHIParamRef IndexBufferRHI, uint32 PrimitiveType, int32 BaseVertexIndex, uint32 FirstInstance,
 	uint32 NumVertices, uint32 StartIndex, uint32 NumPrimitives, uint32 NumInstances)
 {
 	SCOPE_CYCLE_COUNTER(STAT_MetalDrawCallTime);
 	//checkf(NumInstances == 1, TEXT("Currently only 1 instance is supported"));
 	checkf(BaseVertexIndex  == 0, TEXT("BaseVertexIndex must be 0, see GRHISupportsBaseVertexIndex"));
+	checkf(FirstInstance  == 0, TEXT("FirstInstance is currently unsupported on this RHI"));
 
 	RHI_DRAW_CALL_STATS(PrimitiveType,NumInstances*NumPrimitives);
 
diff --git a/Engine/Source/Runtime/Landscape/Classes/LandscapeHeightfieldCollisionComponent.h b/Engine/Source/Runtime/Landscape/Classes/LandscapeHeightfieldCollisionComponent.h
--- a/Engine/Source/Runtime/Landscape/Classes/LandscapeHeightfieldCollisionComponent.h
+++ b/Engine/Source/Runtime/Landscape/Classes/LandscapeHeightfieldCollisionComponent.h
@@ -104,8 +104,9 @@ class ULandscapeHeightfieldCollisionComponent : public UPrimitiveComponent
 	/** 
 	 *	Flag to indicate that the next time we cook data, we should save it to the DDC.
 	 *	Used to ensure DDC is populated when loading content for the first time. 
+	 *  For editor and full version of collision objects
 	 */
-	mutable bool								bShouldSaveCookedDataToDDC;
+	mutable bool								bShouldSaveCookedDataToDDC[2];
 #endif //WITH_EDITORONLY_DATA
 
 	/** 
diff --git a/Engine/Source/Runtime/Landscape/Private/LandscapeCollision.cpp b/Engine/Source/Runtime/Landscape/Private/LandscapeCollision.cpp
--- a/Engine/Source/Runtime/Landscape/Private/LandscapeCollision.cpp
+++ b/Engine/Source/Runtime/Landscape/Private/LandscapeCollision.cpp
@@ -21,6 +21,7 @@ LandscapeCollision.cpp: Landscape collision
 #include "InstancedFoliage.h"
 #include "Foliage/FoliageType.h"
 #include "Engine/StaticMesh.h"
+#include "Components/HierarchicalInstancedStaticMeshComponent.h"
 #include "Interfaces/Interface_CollisionDataProvider.h"
 #include "AI/NavigationSystemHelpers.h"
 #include "AI/Navigation/NavigationSystem.h"
@@ -337,6 +338,8 @@ void ULandscapeHeightfieldCollisionComponent::CreateCollisionObject()
 bool ULandscapeHeightfieldCollisionComponent::CookCollsionData(const FName& Format, bool bUseDefMaterial, bool bCheckDDC, TArray<uint8>& OutCookedData, TArray<UPhysicalMaterial*>& OutMaterials) const
 {
 #if WITH_PHYSX
+	// we have 2 versions of collision objects
+	const int32 CookedDataIndex = bUseDefMaterial ? 0 : 1;
 
 	if (bCheckDDC)
 	{
@@ -345,7 +348,7 @@ bool ULandscapeHeightfieldCollisionComponent::CookCollsionData(const FName& Form
 		{
 			if (GetDerivedDataCacheRef().GetSynchronous(*GetHFDDCKeyString(Format, bUseDefMaterial, HeightfieldGuid), OutCookedData))
 			{
-				bShouldSaveCookedDataToDDC = false;
+				bShouldSaveCookedDataToDDC[CookedDataIndex] = false;
 				return true;
 			}
 		}
@@ -451,10 +454,10 @@ bool ULandscapeHeightfieldCollisionComponent::CookCollsionData(const FName& Form
 		OutCookedData.Init(OutData.Num());
 		FMemory::Memcpy(OutCookedData.GetData(), OutData.GetData(), OutData.Num());
 
-		if (bShouldSaveCookedDataToDDC)
+		if (bShouldSaveCookedDataToDDC[CookedDataIndex])
 		{
 			GetDerivedDataCacheRef().Put(*GetHFDDCKeyString(Format, bUseDefMaterial, HeightfieldGuid), OutCookedData);
-			bShouldSaveCookedDataToDDC = false;
+			bShouldSaveCookedDataToDDC[CookedDataIndex] = false;
 		}
 	}
 	else
@@ -473,6 +476,9 @@ bool ULandscapeMeshCollisionComponent::CookCollsionData(const FName& Format, boo
 {
 #if WITH_PHYSX
 	
+	// we have 2 versions of collision objects
+	const int32 CookedDataIndex = bUseDefMaterial ? 0 : 1;
+
 	if (bCheckDDC)
 	{
 		// Ensure that content was saved with physical materials before using DDC data
@@ -480,7 +486,7 @@ bool ULandscapeMeshCollisionComponent::CookCollsionData(const FName& Format, boo
 		{
 			if (GetDerivedDataCacheRef().GetSynchronous(*GetHFDDCKeyString(Format, bUseDefMaterial, HeightfieldGuid), OutCookedData))
 			{
-				bShouldSaveCookedDataToDDC = false;
+				bShouldSaveCookedDataToDDC[CookedDataIndex] = false;
 				return true;
 			}
 		}
@@ -620,10 +626,10 @@ bool ULandscapeMeshCollisionComponent::CookCollsionData(const FName& Format, boo
 		OutCookedData.Init(OutData.Num());
 		FMemory::Memcpy(OutCookedData.GetData(), OutData.GetData(), OutData.Num());
 
-		if (bShouldSaveCookedDataToDDC)
+		if (bShouldSaveCookedDataToDDC[CookedDataIndex])
 		{
 			GetDerivedDataCacheRef().Put(*GetHFDDCKeyString(Format, bUseDefMaterial, HeightfieldGuid), OutCookedData);
-			bShouldSaveCookedDataToDDC = false;
+			bShouldSaveCookedDataToDDC[CookedDataIndex] = false;
 		}
 	}
 	else
@@ -1077,28 +1083,10 @@ void ULandscapeHeightfieldCollisionComponent::SnapFoliageInstances(AInstancedFol
 
 								// Todo: add do validation with other parameters such as max/min height etc.
 
-								// Update this instances' transform in the UInstancedStaticMeshComponent
-								if (MeshInfo.InstanceClusters.IsValidIndex(Instance.ClusterIndex))
-								{
-									FFoliageInstanceCluster& Cluster = MeshInfo.InstanceClusters[Instance.ClusterIndex];
-
-									Cluster.ClusterComponent->Modify();
-
-									int32 ClusterInstanceDataIndex = Cluster.InstanceIndices.Find(InstanceIndex);
-
-									check(ClusterInstanceDataIndex != INDEX_NONE);
-									check(Cluster.ClusterComponent->PerInstanceSMData.IsValidIndex(ClusterInstanceDataIndex));
-									Cluster.ClusterComponent->MarkRenderStateDirty();
-
-									// Update bounds
-									FTransform InstanceToWorld = Instance.GetInstanceWorldTransform();
-									Cluster.Bounds = Cluster.Bounds + Cluster.ClusterComponent->StaticMesh->GetBounds().TransformBy(InstanceToWorld);
-
-									// Update transform in InstancedStaticMeshComponent
-									FTransform InstanceToLocal = InstanceToWorld.GetRelativeTransform(Cluster.ClusterComponent->ComponentToWorld);
-									Cluster.ClusterComponent->PerInstanceSMData[ClusterInstanceDataIndex].Transform = InstanceToLocal.ToMatrixWithScale();
-									Cluster.ClusterComponent->InvalidateLightingCache();
-								}
+								check(MeshInfo.Component);
+								MeshInfo.Component->Modify();
+								MeshInfo.Component->UpdateInstanceTransform(InstanceIndex, Instance.GetInstanceWorldTransform(), true);
+								MeshInfo.Component->InvalidateLightingCache();
 
 								// Re-add the new instance location to the hash
 								MeshInfo.InstanceHash->InsertInstance(Instance.Location, InstanceIndex);
@@ -1175,8 +1163,12 @@ void ULandscapeHeightfieldCollisionComponent::Serialize(FArchive& Ar)
 		else
 		{
 #if WITH_EDITORONLY_DATA			
+			// do not need this data in PIE, so don't bother duplicating it
+			if (!(Ar.GetPortFlags() & PPF_DuplicateForPIE)) 
+			{
 			CollisionHeightData.Serialize(Ar, this);
 			DominantLayerData.Serialize(Ar, this);
+			}
 #endif//WITH_EDITORONLY_DATA
 		}
 	}
@@ -1309,7 +1301,8 @@ void ULandscapeHeightfieldCollisionComponent::PostLoad()
 #if WITH_EDITOR
 	if (!HasAnyFlags(RF_ClassDefaultObject))
 	{
-		bShouldSaveCookedDataToDDC = true;
+		bShouldSaveCookedDataToDDC[0] = true;
+		bShouldSaveCookedDataToDDC[1] = true;
 	}
 #endif//WITH_EDITOR
 }
diff --git a/Engine/Source/Runtime/OpenGLDrv/Private/OpenGLCommands.cpp b/Engine/Source/Runtime/OpenGLDrv/Private/OpenGLCommands.cpp
--- a/Engine/Source/Runtime/OpenGLDrv/Private/OpenGLCommands.cpp
+++ b/Engine/Source/Runtime/OpenGLDrv/Private/OpenGLCommands.cpp
@@ -2537,7 +2537,7 @@ void FOpenGLDynamicRHI::RHIDrawIndexedIndirect(FIndexBufferRHIParamRef IndexBuff
 	}
 }
 
-void FOpenGLDynamicRHI::RHIDrawIndexedPrimitive(FIndexBufferRHIParamRef IndexBufferRHI,uint32 PrimitiveType,int32 BaseVertexIndex,uint32 MinIndex,uint32 NumVertices,uint32 StartIndex,uint32 NumPrimitives,uint32 NumInstances)
+void FOpenGLDynamicRHI::RHIDrawIndexedPrimitive(FIndexBufferRHIParamRef IndexBufferRHI,uint32 PrimitiveType,int32 BaseVertexIndex,uint32 FirstInstance,uint32 NumVertices,uint32 StartIndex,uint32 NumPrimitives,uint32 NumInstances)
 {
 	SCOPE_CYCLE_COUNTER_DETAILED(STAT_OpenGLDrawPrimitiveTime);
 	VERIFY_GL_SCOPE();
@@ -2578,6 +2578,7 @@ void FOpenGLDynamicRHI::RHIDrawIndexedPrimitive(FIndexBufferRHIParamRef IndexBuf
 	{
 		SCOPE_CYCLE_COUNTER_DETAILED(STAT_OpenGLDrawPrimitiveDriverTime);
 		check( FOpenGL::SupportsInstancing() );
+		checkf(FirstInstance  == 0, TEXT("FirstInstance is currently unsupported on this RHI"));
 		FOpenGL::DrawElementsInstanced(DrawMode, NumElements, IndexType, INDEX_TO_VOID(StartIndex), NumInstances);
 		REPORT_GL_DRAW_ELEMENTS_INSTANCED_EVENT_FOR_FRAME_DUMP(DrawMode, NumElements, IndexType, (void *)StartIndex, NumInstances);
 	}
@@ -2586,7 +2587,7 @@ void FOpenGLDynamicRHI::RHIDrawIndexedPrimitive(FIndexBufferRHIParamRef IndexBuf
 		SCOPE_CYCLE_COUNTER_DETAILED(STAT_OpenGLDrawPrimitiveDriverTime);
 		if ( FOpenGL::SupportsDrawIndexOffset() )
 		{
-			FOpenGL::DrawRangeElements(DrawMode, MinIndex, MinIndex + NumVertices, NumElements, IndexType, INDEX_TO_VOID(StartIndex));
+			FOpenGL::DrawRangeElements(DrawMode, 0, NumVertices, NumElements, IndexType, INDEX_TO_VOID(StartIndex));
 		}
 		else
 		{
diff --git a/Engine/Source/Runtime/RHI/Private/RHI.cpp b/Engine/Source/Runtime/RHI/Private/RHI.cpp
--- a/Engine/Source/Runtime/RHI/Private/RHI.cpp
+++ b/Engine/Source/Runtime/RHI/Private/RHI.cpp
@@ -186,6 +186,7 @@ bool GTriggerGPUProfile = false;
 bool GRHISupportsTextureStreaming = false;
 bool GSupportsDepthBoundsTest = false;
 bool GRHISupportsBaseVertexIndex = true;
+bool GRHISupportsFirstInstance = false;
 bool GRHIRequiresEarlyBackBufferRenderTarget = true;
 
 /** Whether we are profiling GPU hitches. */
diff --git a/Engine/Source/Runtime/RHI/Public/RHI.h b/Engine/Source/Runtime/RHI/Public/RHI.h
--- a/Engine/Source/Runtime/RHI/Public/RHI.h
+++ b/Engine/Source/Runtime/RHI/Public/RHI.h
@@ -185,6 +185,9 @@ extern RHI_API int32 GNumPrimitivesDrawnRHI;
 /** Whether or not the RHI can handle a non-zero BaseVertexIndex - extra SetStreamSource calls will be needed if this is false */
 extern RHI_API bool GRHISupportsBaseVertexIndex;
 
+/** Whether or not the RHI can handle a non-zero FirstInstance - extra SetStreamSource calls will be needed if this is false */
+extern RHI_API bool GRHISupportsFirstInstance;
+
 /** Whether or not the engine should set the BackBuffer as a render target early in the frame. */
 extern RHI_API bool GRHIRequiresEarlyBackBufferRenderTarget;
 
diff --git a/Engine/Source/Runtime/RHI/Public/RHICommandList.h b/Engine/Source/Runtime/RHI/Public/RHICommandList.h
--- a/Engine/Source/Runtime/RHI/Public/RHICommandList.h
+++ b/Engine/Source/Runtime/RHI/Public/RHICommandList.h
@@ -409,16 +409,16 @@ struct FRHICommandDrawIndexedPrimitive : public FRHICommand<FRHICommandDrawIndex
 	FIndexBufferRHIParamRef IndexBuffer;
 	uint32 PrimitiveType;
 	int32 BaseVertexIndex;
-	uint32 MinIndex;
+	uint32 FirstInstance;
 	uint32 NumVertices;
 	uint32 StartIndex;
 	uint32 NumPrimitives;
 	uint32 NumInstances;
-	FORCEINLINE_DEBUGGABLE FRHICommandDrawIndexedPrimitive(FIndexBufferRHIParamRef InIndexBuffer, uint32 InPrimitiveType, int32 InBaseVertexIndex, uint32 InMinIndex, uint32 InNumVertices, uint32 InStartIndex, uint32 InNumPrimitives, uint32 InNumInstances)
+	FORCEINLINE_DEBUGGABLE FRHICommandDrawIndexedPrimitive(FIndexBufferRHIParamRef InIndexBuffer, uint32 InPrimitiveType, int32 InBaseVertexIndex, uint32 InFirstInstance, uint32 InNumVertices, uint32 InStartIndex, uint32 InNumPrimitives, uint32 InNumInstances)
 		: IndexBuffer(InIndexBuffer)
 		, PrimitiveType(InPrimitiveType)
 		, BaseVertexIndex(InBaseVertexIndex)
-		, MinIndex(InMinIndex)
+		, FirstInstance(InFirstInstance)
 		, NumVertices(InNumVertices)
 		, StartIndex(InStartIndex)
 		, NumPrimitives(InNumPrimitives)
@@ -1440,14 +1440,14 @@ public:
 		new (AllocCommand<FRHICommandDrawPrimitive>()) FRHICommandDrawPrimitive(PrimitiveType, BaseVertexIndex, NumPrimitives, NumInstances);
 	}
 
-	FORCEINLINE_DEBUGGABLE void DrawIndexedPrimitive(FIndexBufferRHIParamRef IndexBuffer, uint32 PrimitiveType, int32 BaseVertexIndex, uint32 MinIndex, uint32 NumVertices, uint32 StartIndex, uint32 NumPrimitives, uint32 NumInstances)
+	FORCEINLINE_DEBUGGABLE void DrawIndexedPrimitive(FIndexBufferRHIParamRef IndexBuffer, uint32 PrimitiveType, int32 BaseVertexIndex, uint32 FirstInstance, uint32 NumVertices, uint32 StartIndex, uint32 NumPrimitives, uint32 NumInstances)
 	{
 		if (Bypass())
 		{
-			CMD_CONTEXT(DrawIndexedPrimitive)(IndexBuffer, PrimitiveType, BaseVertexIndex, MinIndex, NumVertices, StartIndex, NumPrimitives, NumInstances);
+			CMD_CONTEXT(DrawIndexedPrimitive)(IndexBuffer, PrimitiveType, BaseVertexIndex, FirstInstance, NumVertices, StartIndex, NumPrimitives, NumInstances);
 			return;
 		}
-		new (AllocCommand<FRHICommandDrawIndexedPrimitive>()) FRHICommandDrawIndexedPrimitive(IndexBuffer, PrimitiveType, BaseVertexIndex, MinIndex, NumVertices, StartIndex, NumPrimitives, NumInstances);
+		new (AllocCommand<FRHICommandDrawIndexedPrimitive>()) FRHICommandDrawIndexedPrimitive(IndexBuffer, PrimitiveType, BaseVertexIndex, FirstInstance, NumVertices, StartIndex, NumPrimitives, NumInstances);
 	}
 
 	FORCEINLINE_DEBUGGABLE void SetStreamSource(uint32 StreamIndex, FVertexBufferRHIParamRef VertexBuffer, uint32 Stride, uint32 Offset)
diff --git a/Engine/Source/Runtime/RHI/Public/RHICommandListCommandExecutes.inl b/Engine/Source/Runtime/RHI/Public/RHICommandListCommandExecutes.inl
--- a/Engine/Source/Runtime/RHI/Public/RHICommandListCommandExecutes.inl
+++ b/Engine/Source/Runtime/RHI/Public/RHICommandListCommandExecutes.inl
@@ -124,7 +124,7 @@ void FRHICommandDrawPrimitive::Execute(FRHICommandListBase& CmdList)
 void FRHICommandDrawIndexedPrimitive::Execute(FRHICommandListBase& CmdList)
 {
 	RHISTAT(DrawIndexedPrimitive);
-	INTERNAL_DECORATOR(DrawIndexedPrimitive)(IndexBuffer, PrimitiveType, BaseVertexIndex, MinIndex, NumVertices, StartIndex, NumPrimitives, NumInstances);
+	INTERNAL_DECORATOR(DrawIndexedPrimitive)(IndexBuffer, PrimitiveType, BaseVertexIndex, FirstInstance, NumVertices, StartIndex, NumPrimitives, NumInstances);
 }
 
 void FRHICommandSetBoundShaderState::Execute(FRHICommandListBase& CmdList)
diff --git a/Engine/Source/Runtime/RHI/Public/RHIMethods.h b/Engine/Source/Runtime/RHI/Public/RHIMethods.h
--- a/Engine/Source/Runtime/RHI/Public/RHIMethods.h
+++ b/Engine/Source/Runtime/RHI/Public/RHIMethods.h
@@ -1938,7 +1938,7 @@ DEFINE_RHIMETHOD_CMDLIST_8(
 	FIndexBufferRHIParamRef,IndexBuffer,
 	uint32,PrimitiveType,
 	int32,BaseVertexIndex,
-	uint32,MinIndex,
+	uint32,FirstInstance,
 	uint32,NumVertices,
 	uint32,StartIndex,
 	uint32,NumPrimitives,
diff --git a/Engine/Source/Runtime/Renderer/Private/DepthRendering.cpp b/Engine/Source/Runtime/Renderer/Private/DepthRendering.cpp
--- a/Engine/Source/Runtime/Renderer/Private/DepthRendering.cpp
+++ b/Engine/Source/Runtime/Renderer/Private/DepthRendering.cpp
@@ -254,6 +254,7 @@ FPositionOnlyDepthDrawingPolicy::FPositionOnlyDepthDrawingPolicy(
 	FMeshDrawingPolicy(InVertexFactory,InMaterialRenderProxy,InMaterialResource,false,bIsTwoSided,bIsWireframe)
 {
 	VertexShader = InMaterialResource.GetShader<TDepthOnlyVS<true> >(InVertexFactory->GetType());
+	bUsePositionOnlyVS = true;
 }
 
 void FPositionOnlyDepthDrawingPolicy::SetSharedState(FRHICommandList& RHICmdList, const FSceneView* View, const ContextDataType PolicyContext) const
diff --git a/Engine/Source/Runtime/Renderer/Private/DrawingPolicy.cpp b/Engine/Source/Runtime/Renderer/Private/DrawingPolicy.cpp
--- a/Engine/Source/Runtime/Renderer/Private/DrawingPolicy.cpp
+++ b/Engine/Source/Runtime/Renderer/Private/DrawingPolicy.cpp
@@ -40,6 +40,8 @@ FMeshDrawingPolicy::FMeshDrawingPolicy(
 	bNeedsBackfacePass = !bInTwoSidedOverride && bMaterialResourceIsTwoSided
 		//@todo - hook up bTwoSidedSeparatePass here if we re-add it
 		&& false;
+	
+	bUsePositionOnlyVS = false;
 }
 
 void FMeshDrawingPolicy::SetMeshRenderState(
@@ -103,16 +105,75 @@ void FMeshDrawingPolicy::DrawMesh(FRHICommandList& RHICmdList, const FMeshBatch&
 		if(BatchElement.IndexBuffer)
 		{
 			check(BatchElement.IndexBuffer->IsInitialized());
-			RHICmdList.DrawIndexedPrimitive(
+			if (BatchElement.InstanceRuns)
+			{
+				if (!GRHISupportsFirstInstance)
+				{
+					if (bUsePositionOnlyVS)
+					{
+						for (uint32 Run = 0; Run < BatchElement.NumInstances; Run++)
+						{
+							VertexFactory->OffsetPositionInstanceStreams(RHICmdList, BatchElement.InstanceRuns[Run * 2]); 
+							RHICmdList.DrawIndexedPrimitive(
+								BatchElement.IndexBuffer->IndexBufferRHI,
+								Mesh.Type,
+								0,
+								0,
+								BatchElement.MaxVertexIndex - BatchElement.MinVertexIndex + 1,
+								BatchElement.FirstIndex,
+								BatchElement.NumPrimitives,
+								1 + BatchElement.InstanceRuns[Run * 2 + 1] - BatchElement.InstanceRuns[Run * 2]
+							);
+						}
+					}
+					else
+					{
+						for (uint32 Run = 0; Run < BatchElement.NumInstances; Run++)
+						{
+							VertexFactory->OffsetInstanceStreams(RHICmdList, BatchElement.InstanceRuns[Run * 2]); 
+							RHICmdList.DrawIndexedPrimitive(
+								BatchElement.IndexBuffer->IndexBufferRHI,
+								Mesh.Type,
+								0,
+								0,
+								BatchElement.MaxVertexIndex - BatchElement.MinVertexIndex + 1,
+								BatchElement.FirstIndex,
+								BatchElement.NumPrimitives,
+								1 + BatchElement.InstanceRuns[Run * 2 + 1] - BatchElement.InstanceRuns[Run * 2]
+							);
+						}
+					}
+				}
+				else
+				{
+					for (uint32 Run = 0; Run < BatchElement.NumInstances; Run++)
+					{
+						RHICmdList.DrawIndexedPrimitive(
+							BatchElement.IndexBuffer->IndexBufferRHI,
+							Mesh.Type,
+							0,
+							BatchElement.InstanceRuns[Run * 2],
+							BatchElement.MaxVertexIndex - BatchElement.MinVertexIndex + 1,
+							BatchElement.FirstIndex,
+							BatchElement.NumPrimitives,
+							1 + BatchElement.InstanceRuns[Run * 2 + 1] - BatchElement.InstanceRuns[Run * 2]
+						);
+					}
+				}
+			}
+			else
+			{
+				RHICmdList.DrawIndexedPrimitive(
 					BatchElement.IndexBuffer->IndexBufferRHI,
 					Mesh.Type,
 					0,
-					BatchElement.MinVertexIndex,
+					0,
 					BatchElement.MaxVertexIndex - BatchElement.MinVertexIndex + 1,
 					BatchElement.FirstIndex,
 					BatchElement.NumPrimitives,
 					BatchElement.NumInstances
 					);
+			}
 		}
 		else
 		{
diff --git a/Engine/Source/Runtime/Renderer/Private/ScenePrivate.h b/Engine/Source/Runtime/Renderer/Private/ScenePrivate.h
--- a/Engine/Source/Runtime/Renderer/Private/ScenePrivate.h
+++ b/Engine/Source/Runtime/Renderer/Private/ScenePrivate.h
@@ -552,6 +552,8 @@ public:
 	bool bBokehDOFHistory;
 	bool bBokehDOFHistory2;
 
+	FTemporalLODState TemporalLODState;
+
 	// call after SetupTemporalAA()
 	uint32 GetCurrentTemporalAASampleIndex() const
 	{
@@ -625,6 +627,14 @@ public:
 		}
 	}
 
+	/** 
+	 * Called every frame after UpdateLastRenderTime, sets up the information for the lagged temporal LOD transition
+	 */
+	void UpdateTemporalLODTransition(const FViewInfo& View)
+	{
+		TemporalLODState.UpdateTemporalLODTransition(View, LastRenderTime);
+	}
+
 	/** 
 	 * Returns an array of visibility data for the given view, or NULL if none exists. 
 	 * The data bits are indexed by VisibilityId of each primitive in the scene.
@@ -778,6 +788,24 @@ public:
 		return Ret;
 	}
 
+	virtual FTemporalLODState& GetTemporalLODState() override
+	{
+		return TemporalLODState;
+	}
+	virtual const FTemporalLODState& GetTemporalLODState() const override
+	{
+		return TemporalLODState;
+	}
+	/** 
+	 * Returns the blend factor between the last two LOD samples
+	 */
+	float GetTemporalLODTransition() const override
+	{
+		return TemporalLODState.GetTemporalLODTransition(LastRenderTime);
+	}
+
+
+
 	// FDeferredCleanupInterface
 	virtual void FinishCleanup()
 	{
diff --git a/Engine/Source/Runtime/Renderer/Private/SceneVisibility.cpp b/Engine/Source/Runtime/Renderer/Private/SceneVisibility.cpp
--- a/Engine/Source/Runtime/Renderer/Private/SceneVisibility.cpp
+++ b/Engine/Source/Runtime/Renderer/Private/SceneVisibility.cpp
@@ -1777,6 +1777,8 @@ void FSceneRenderer::PreVisibilityFrameSetup(FRHICommandListImmediate& RHICmdLis
 
 			// This finishes the update of view state
 			ViewState->UpdateLastRenderTime(*View.Family);
+
+			ViewState->UpdateTemporalLODTransition(View);
 		}
 	}
 }
diff --git a/Engine/Source/Runtime/Renderer/Private/ShadowRendering.cpp b/Engine/Source/Runtime/Renderer/Private/ShadowRendering.cpp
--- a/Engine/Source/Runtime/Renderer/Private/ShadowRendering.cpp
+++ b/Engine/Source/Runtime/Renderer/Private/ShadowRendering.cpp
@@ -1606,6 +1606,31 @@ void FProjectedShadowInfo::RenderDepth(FRHICommandList& RHICmdList, FSceneRender
 	// Backup properties of the view that we will override
 	TUniformBufferRef<FViewUniformShaderParameters> OriginalUniformBuffer = FoundView->UniformBuffer;
 	FMatrix OriginalViewMatrix = FoundView->ViewMatrices.ViewMatrix;
+	FIntRect OriginalViewRect = FoundView->ViewRect;
+	FoundView->ViewRect.Min.X = 0;
+	FoundView->ViewRect.Min.Y = 0;
+	FoundView->ViewRect.Max.X = ResolutionX;
+	FoundView->ViewRect.Max.Y =  ResolutionY;
+
+	float JitterX = FoundView->ViewMatrices.ProjMatrix.M[2][0];
+	float JitterY = FoundView->ViewMatrices.ProjMatrix.M[2][1];
+
+	FoundView->ViewMatrices.ProjMatrix.M[2][0] = 0.0f;
+	FoundView->ViewMatrices.ProjMatrix.M[2][1] = 0.0f;
+
+	{
+		// Compute the view projection matrix and its inverse.
+		FoundView->ViewProjectionMatrix = FoundView->ViewMatrices.ViewMatrix * FoundView->ViewMatrices.ProjMatrix;
+		FoundView->InvViewProjectionMatrix = FoundView->ViewMatrices.GetInvProjMatrix() * FoundView->InvViewMatrix;
+
+		/** The view transform, starting from world-space points translated by -ViewOrigin. */
+		FMatrix TranslatedViewMatrix = FTranslationMatrix(-FoundView->ViewMatrices.PreViewTranslation) * FoundView->ViewMatrices.ViewMatrix;
+
+		// Compute a transform from view origin centered world-space to clip space.
+		FoundView->ViewMatrices.TranslatedViewProjectionMatrix = TranslatedViewMatrix * FoundView->ViewMatrices.ProjMatrix;
+		FoundView->ViewMatrices.InvTranslatedViewProjectionMatrix = FoundView->ViewMatrices.TranslatedViewProjectionMatrix.Inverse();
+	}
+
 
 	// Override the view matrix so that billboarding primitives will be aligned to the light
 	//@todo - creating a new uniform buffer is expensive, only do this when the vertex factory needs an accurate view matrix (particle sprites)
@@ -1617,6 +1642,9 @@ void FProjectedShadowInfo::RenderDepth(FRHICommandList& RHICmdList, FSceneRender
 		VolumeBounds,
 		TVC_MAX);
 
+	// we are going to set this back now because we only want the correct view rect for the uniform buffer. For LOD calculations, we want the rendering viewrect and proj matrix.
+	FoundView->ViewRect = OriginalViewRect;
+
 	// Prevent materials from getting overridden during shadow casting in viewmodes like lighting only
 	// Lighting only should only affect the material used with direct lighting, not the indirect lighting
 	FoundView->bForceShowMaterials = true;
@@ -1626,6 +1654,22 @@ void FProjectedShadowInfo::RenderDepth(FRHICommandList& RHICmdList, FSceneRender
 	FoundView->bForceShowMaterials = false;
 	FoundView->UniformBuffer = OriginalUniformBuffer;
 	FoundView->ViewMatrices.ViewMatrix = OriginalViewMatrix;
+
+	FoundView->ViewMatrices.ProjMatrix.M[2][0] = JitterX;
+	FoundView->ViewMatrices.ProjMatrix.M[2][1] = JitterY;
+
+	{
+		// Compute the view projection matrix and its inverse.
+		FoundView->ViewProjectionMatrix = FoundView->ViewMatrices.ViewMatrix * FoundView->ViewMatrices.ProjMatrix;
+		FoundView->InvViewProjectionMatrix = FoundView->ViewMatrices.GetInvProjMatrix() * FoundView->InvViewMatrix;
+
+		/** The view transform, starting from world-space points translated by -ViewOrigin. */
+		FMatrix TranslatedViewMatrix = FTranslationMatrix(-FoundView->ViewMatrices.PreViewTranslation) * FoundView->ViewMatrices.ViewMatrix;
+
+		// Compute a transform from view origin centered world-space to clip space.
+		FoundView->ViewMatrices.TranslatedViewProjectionMatrix = TranslatedViewMatrix * FoundView->ViewMatrices.ProjMatrix;
+		FoundView->ViewMatrices.InvTranslatedViewProjectionMatrix = FoundView->ViewMatrices.TranslatedViewProjectionMatrix.Inverse();
+	}
 }
 
 void StencilingGeometry::DrawSphere(FRHICommandList& RHICmdList)
diff --git a/Engine/Source/Runtime/ShaderCore/Private/VertexFactory.cpp b/Engine/Source/Runtime/ShaderCore/Private/VertexFactory.cpp
--- a/Engine/Source/Runtime/ShaderCore/Private/VertexFactory.cpp
+++ b/Engine/Source/Runtime/ShaderCore/Private/VertexFactory.cpp
@@ -178,6 +178,18 @@ void FVertexFactory::Set(FRHICommandList& RHICmdList) const
 	}
 }
 
+void FVertexFactory::OffsetInstanceStreams(FRHICommandList& RHICmdList, uint32 FirstVertex) const
+{
+	for(int32 StreamIndex = 0;StreamIndex < Streams.Num();StreamIndex++)
+	{
+		const FVertexStream& Stream = Streams[StreamIndex];
+		if (Stream.bUseInstanceIndex)
+		{
+			RHICmdList.SetStreamSource( StreamIndex, Stream.VertexBuffer->VertexBufferRHI, Stream.Stride, Stream.Offset + Stream.Stride * FirstVertex);
+		}
+	}
+}
+
 void FVertexFactory::SetPositionStream(FRHICommandList& RHICmdList) const
 {
 	check(IsInitialized());
@@ -190,6 +202,18 @@ void FVertexFactory::SetPositionStream(FRHICommandList& RHICmdList) const
 	}
 }
 
+void FVertexFactory::OffsetPositionInstanceStreams(FRHICommandList& RHICmdList, uint32 FirstVertex) const
+{
+	for(int32 StreamIndex = 0;StreamIndex < PositionStream.Num();StreamIndex++)
+	{
+		const FVertexStream& Stream = PositionStream[StreamIndex];
+		if (Stream.bUseInstanceIndex)
+		{
+			RHICmdList.SetStreamSource( StreamIndex, Stream.VertexBuffer->VertexBufferRHI, Stream.Stride, Stream.Offset + Stream.Stride * FirstVertex);
+		}
+	}
+}
+
 void FVertexFactory::ReleaseRHI()
 {
 	Declaration.SafeRelease();
@@ -244,6 +268,7 @@ FVertexElement FVertexFactory::AccessStreamComponent(const FVertexStreamComponen
 	VertexStream.VertexBuffer = Component.VertexBuffer;
 	VertexStream.Stride = Component.Stride;
 	VertexStream.Offset = 0;
+	VertexStream.bUseInstanceIndex = Component.bUseInstanceIndex;
 
 	return FVertexElement(Streams.AddUnique(VertexStream),Component.Offset,Component.Type,AttributeIndex,VertexStream.Stride,Component.bUseInstanceIndex);
 }
@@ -254,6 +279,7 @@ FVertexElement FVertexFactory::AccessPositionStreamComponent(const FVertexStream
 	VertexStream.VertexBuffer = Component.VertexBuffer;
 	VertexStream.Stride = Component.Stride;
 	VertexStream.Offset = 0;
+	VertexStream.bUseInstanceIndex = Component.bUseInstanceIndex;
 
 	return FVertexElement(PositionStream.AddUnique(VertexStream),Component.Offset,Component.Type,AttributeIndex,VertexStream.Stride,Component.bUseInstanceIndex);
 }
diff --git a/Engine/Source/Runtime/ShaderCore/Public/VertexFactory.h b/Engine/Source/Runtime/ShaderCore/Public/VertexFactory.h
--- a/Engine/Source/Runtime/ShaderCore/Public/VertexFactory.h
+++ b/Engine/Source/Runtime/ShaderCore/Public/VertexFactory.h
@@ -354,11 +354,21 @@ public:
 	 */
 	void Set(FRHICommandList& RHICmdList) const;
 
+	/**
+	 * Call SetStreamSource on instance streams to offset the read pointer
+	 */
+	void OffsetInstanceStreams(FRHICommandList& RHICmdList, uint32 FirstVertex) const;
+
 	/**
 	* Sets the position stream as the current stream source.
 	*/
 	void SetPositionStream(FRHICommandList& RHICmdList) const;
 
+	/**
+	* Call SetStreamSource on instance streams to offset the read pointer
+	 */
+	void OffsetPositionInstanceStreams(FRHICommandList& RHICmdList, uint32 FirstVertex) const;
+
 	/**
 	* Can be overridden by FVertexFactory subclasses to modify their compile environment just before compilation occurs.
 	*/
@@ -441,6 +451,7 @@ protected:
 		const FVertexBuffer* VertexBuffer;
 		uint32 Stride;
 		uint32 Offset;
+		bool bUseInstanceIndex;
 
 		friend bool operator==(const FVertexStream& A,const FVertexStream& B)
 		{
diff --git a/Engine/Source/Runtime/Windows/D3D11RHI/Private/D3D11Commands.cpp b/Engine/Source/Runtime/Windows/D3D11RHI/Private/D3D11Commands.cpp
--- a/Engine/Source/Runtime/Windows/D3D11RHI/Private/D3D11Commands.cpp
+++ b/Engine/Source/Runtime/Windows/D3D11RHI/Private/D3D11Commands.cpp
@@ -1297,7 +1297,7 @@ void FD3D11DynamicRHI::RHIDrawIndexedIndirect(FIndexBufferRHIParamRef IndexBuffe
 	}
 }
 
-void FD3D11DynamicRHI::RHIDrawIndexedPrimitive(FIndexBufferRHIParamRef IndexBufferRHI,uint32 PrimitiveType,int32 BaseVertexIndex,uint32 MinIndex,uint32 NumVertices,uint32 StartIndex,uint32 NumPrimitives,uint32 NumInstances)
+void FD3D11DynamicRHI::RHIDrawIndexedPrimitive(FIndexBufferRHIParamRef IndexBufferRHI,uint32 PrimitiveType,int32 BaseVertexIndex,uint32 FirstInstance,uint32 NumVertices,uint32 StartIndex,uint32 NumPrimitives,uint32 NumInstances)
 {
 	DYNAMIC_CAST_D3D11RESOURCE(IndexBuffer,IndexBuffer);
 
@@ -1327,7 +1327,7 @@ void FD3D11DynamicRHI::RHIDrawIndexedPrimitive(FIndexBufferRHIParamRef IndexBuff
 
 	if(NumInstances > 1)
 	{
-		Direct3DDeviceIMContext->DrawIndexedInstanced(IndexCount,NumInstances,StartIndex,BaseVertexIndex,0);
+		Direct3DDeviceIMContext->DrawIndexedInstanced(IndexCount, NumInstances, StartIndex, BaseVertexIndex, FirstInstance);
 	}
 	else
 	{
diff --git a/Engine/Source/Runtime/Windows/D3D11RHI/Private/Windows/WindowsD3D11Device.cpp b/Engine/Source/Runtime/Windows/D3D11RHI/Private/Windows/WindowsD3D11Device.cpp
--- a/Engine/Source/Runtime/Windows/D3D11RHI/Private/Windows/WindowsD3D11Device.cpp
+++ b/Engine/Source/Runtime/Windows/D3D11RHI/Private/Windows/WindowsD3D11Device.cpp
@@ -663,6 +663,7 @@ void FD3D11DynamicRHI::InitD3DDevice()
 		FHardwareInfo::RegisterHardwareInfo( NAME_RHI, TEXT( "D3D11" ) );
 
 		GRHISupportsTextureStreaming = true;
+		GRHISupportsFirstInstance = true;
 
 		// Set the RHI initialized flag.
 		GIsRHIInitialized = true;
