Fixing issue where LoadMap thought map name was invalid when running PIE

diff --git a/Engine/Source/Runtime/Engine/Private/UnrealEngine.cpp b/Engine/Source/Runtime/Engine/Private/UnrealEngine.cpp
--- a/Engine/Source/Runtime/Engine/Private/UnrealEngine.cpp
+++ b/Engine/Source/Runtime/Engine/Private/UnrealEngine.cpp
@@ -8304,7 +8304,7 @@ bool UEngine::TickWorldTravel(FWorldContext& Context, float DeltaSeconds)
 		}
 		else if( Context.PendingNetGame && Context.PendingNetGame->bSuccessfullyConnected && !Context.PendingNetGame->bSentJoinRequest )
 		{
-			if (!MakeSureMapNameIsValid(Context.PendingNetGame->URL.Map))
+			if (Context.WorldType != EWorldType::PIE && !MakeSureMapNameIsValid(Context.PendingNetGame->URL.Map))
 			{
 				BrowseToDefaultMap(Context);
 				BroadcastTravelFailure(Context.World(), ETravelFailure::PackageMissing, Context.PendingNetGame->URL.RedirectURL);
