Compartmentalize End Overlap code for when an Actor is being removed from the world (Destroyed and now EndPlay with RemoveFromWorld as reason)

diff --git a/Engine/Source/Runtime/Engine/Classes/GameFramework/Actor.h b/Engine/Source/Runtime/Engine/Classes/GameFramework/Actor.h
--- a/Engine/Source/Runtime/Engine/Classes/GameFramework/Actor.h
+++ b/Engine/Source/Runtime/Engine/Classes/GameFramework/Actor.h
@@ -1561,6 +1561,12 @@ public:
 	//--------------------------------------------------------------------------------------
 	// Actor overlap tracking
 	
+	/**
+	 * Dispatch all EndOverlap for all of the Actor's PrimitiveComponents. 
+	 * Generally used when removing the Actor from the world.
+	 */
+	void ClearComponentOverlaps();
+
 	/** 
 	 * Queries world and updates overlap detection state for this actor.
 	 * @param bDoNotifies		True to dispatch being/end overlap notifications when these events occur.
diff --git a/Engine/Source/Runtime/Engine/Private/Actor.cpp b/Engine/Source/Runtime/Engine/Private/Actor.cpp
--- a/Engine/Source/Runtime/Engine/Private/Actor.cpp
+++ b/Engine/Source/Runtime/Engine/Private/Actor.cpp
@@ -897,6 +897,34 @@ void AActor::SetTickGroup(ETickingGroup NewTickGroup)
 	PrimaryActorTick.TickGroup = NewTickGroup;
 }
 
+void AActor::ClearComponentOverlaps()
+{
+	TArray<UPrimitiveComponent*> PrimitiveComponents;
+	GetComponents(PrimitiveComponents);
+
+	// Remove owned components from overlap tracking
+	// We don't traverse the RootComponent attachment tree since that might contain
+	// components owned by other actors.
+	for (UPrimitiveComponent* const PrimComp : PrimitiveComponents)
+	{
+		TArray<UPrimitiveComponent*> OverlappingComponents;
+		PrimComp->GetOverlappingComponents(OverlappingComponents);
+
+		for (UPrimitiveComponent* const OverlapComp : OverlappingComponents)
+		{
+			if (OverlapComp)
+			{
+				PrimComp->EndComponentOverlap(OverlapComp, true, true);
+
+				if (IsPendingKill())
+				{
+					return;
+				}
+			}
+		}
+	}
+}
+
 void AActor::UpdateOverlaps(bool bDoNotifies)
 {
 	// just update the root component, which will cascade down to the children
@@ -1506,6 +1534,8 @@ void AActor::EndPlay(const EEndPlayReason::Type EndPlayReason)
 	// Behaviors specific to an actor being unloaded due to a streaming level removal
 	if (EndPlayReason == EEndPlayReason::RemovedFromWorld)
 	{
+		ClearComponentOverlaps();
+
 		bActorInitialized = false;
 		GetWorld()->RemoveNetworkActor(this);
 
diff --git a/Engine/Source/Runtime/Engine/Private/LevelActor.cpp b/Engine/Source/Runtime/Engine/Private/LevelActor.cpp
--- a/Engine/Source/Runtime/Engine/Private/LevelActor.cpp
+++ b/Engine/Source/Runtime/Engine/Private/LevelActor.cpp
@@ -502,31 +502,11 @@ bool UWorld::DestroyActor( AActor* ThisActor, bool bNetForce, bool bShouldModify
 		return true;
 	}
 
-	// Remove owned components from overlap tracking
-	// We don't traverse the RootComponent attachment tree since that might contain
-	// components owned by other actors.
-	for (int32 Index = 0; Index < SceneComponents.Num(); ++Index)
-	{
-		UPrimitiveComponent* const PrimComp = Cast<UPrimitiveComponent>(SceneComponents[Index]);
-		if (PrimComp)
-		{
-			TArray<UPrimitiveComponent*> OverlappingComponents;
-			PrimComp->GetOverlappingComponents(OverlappingComponents);
-
-			for (auto OverlapIt = OverlappingComponents.CreateIterator(); OverlapIt; ++OverlapIt)
-			{
-				UPrimitiveComponent* const OverlapComp = *OverlapIt;
-				if (OverlapComp)
-				{
-					PrimComp->EndComponentOverlap(OverlapComp, true, true);
+	ThisActor->ClearComponentOverlaps();
 
-					if (ThisActor->IsPendingKill())
-					{
-						return true;
-					}
-				}
-			}
-		}
+	if (ThisActor->IsPendingKill())
+	{
+		return true;
 	}
 
 	// If this actor has an owner, notify it that it has lost a child.
