Engine: Replace some color literals with FColor / FColorList entries

diff --git a/Engine/Plugins/2D/Paper2D/Source/Paper2DEditor/Private/Atlasing/AtlasAssetTypeActions.cpp b/Engine/Plugins/2D/Paper2D/Source/Paper2DEditor/Private/Atlasing/AtlasAssetTypeActions.cpp
--- a/Engine/Plugins/2D/Paper2D/Source/Paper2DEditor/Private/Atlasing/AtlasAssetTypeActions.cpp
+++ b/Engine/Plugins/2D/Paper2D/Source/Paper2DEditor/Private/Atlasing/AtlasAssetTypeActions.cpp
@@ -17,7 +17,7 @@ FText FAtlasAssetTypeActions::GetName() const
 
 FColor FAtlasAssetTypeActions::GetTypeColor() const
 {
-	return FColor(0, 255, 255);
+	return FColor::Cyan;
 }
 
 UClass* FAtlasAssetTypeActions::GetSupportedClass() const
diff --git a/Engine/Plugins/2D/Paper2D/Source/Paper2DEditor/Private/SpriteAssetTypeActions.cpp b/Engine/Plugins/2D/Paper2D/Source/Paper2DEditor/Private/SpriteAssetTypeActions.cpp
--- a/Engine/Plugins/2D/Paper2D/Source/Paper2DEditor/Private/SpriteAssetTypeActions.cpp
+++ b/Engine/Plugins/2D/Paper2D/Source/Paper2DEditor/Private/SpriteAssetTypeActions.cpp
@@ -20,7 +20,7 @@ FText FSpriteAssetTypeActions::GetName() const
 
 FColor FSpriteAssetTypeActions::GetTypeColor() const
 {
-	return FColor(0, 255, 255);
+	return FColor::Cyan;
 }
 
 UClass* FSpriteAssetTypeActions::GetSupportedClass() const
diff --git a/Engine/Plugins/2D/Paper2D/Source/Paper2DEditor/Private/TileMapEditing/TileMapAssetTypeActions.cpp b/Engine/Plugins/2D/Paper2D/Source/Paper2DEditor/Private/TileMapEditing/TileMapAssetTypeActions.cpp
--- a/Engine/Plugins/2D/Paper2D/Source/Paper2DEditor/Private/TileMapEditing/TileMapAssetTypeActions.cpp
+++ b/Engine/Plugins/2D/Paper2D/Source/Paper2DEditor/Private/TileMapEditing/TileMapAssetTypeActions.cpp
@@ -17,7 +17,7 @@ FText FTileMapAssetTypeActions::GetName() const
 
 FColor FTileMapAssetTypeActions::GetTypeColor() const
 {
-	return FColor(0, 255, 255);
+	return FColor::Cyan;
 }
 
 UClass* FTileMapAssetTypeActions::GetSupportedClass() const
diff --git a/Engine/Plugins/2D/Paper2D/Source/Paper2DEditor/Private/TileSetAssetTypeActions.cpp b/Engine/Plugins/2D/Paper2D/Source/Paper2DEditor/Private/TileSetAssetTypeActions.cpp
--- a/Engine/Plugins/2D/Paper2D/Source/Paper2DEditor/Private/TileSetAssetTypeActions.cpp
+++ b/Engine/Plugins/2D/Paper2D/Source/Paper2DEditor/Private/TileSetAssetTypeActions.cpp
@@ -21,7 +21,7 @@ FText FTileSetAssetTypeActions::GetName() const
 
 FColor FTileSetAssetTypeActions::GetTypeColor() const
 {
-	return FColor(0, 255, 255);
+	return FColor::Cyan;
 }
 
 UClass* FTileSetAssetTypeActions::GetSupportedClass() const
diff --git a/Engine/Plugins/2D/Paper2D/Source/PaperJsonImporter/Private/PaperSpriteSheetAssetTypeActions.cpp b/Engine/Plugins/2D/Paper2D/Source/PaperJsonImporter/Private/PaperSpriteSheetAssetTypeActions.cpp
--- a/Engine/Plugins/2D/Paper2D/Source/PaperJsonImporter/Private/PaperSpriteSheetAssetTypeActions.cpp
+++ b/Engine/Plugins/2D/Paper2D/Source/PaperJsonImporter/Private/PaperSpriteSheetAssetTypeActions.cpp
@@ -23,7 +23,7 @@ FText FPaperSpriteSheetAssetTypeActions::GetName() const
 
 FColor FPaperSpriteSheetAssetTypeActions::GetTypeColor() const
 {
-	return FColor(0, 255, 255);
+	return FColor::Cyan;
 }
 
 UClass* FPaperSpriteSheetAssetTypeActions::GetSupportedClass() const
diff --git a/Engine/Plugins/Media/MediaPlayerEditor/Source/MediaPlayerEditor/Private/AssetTools/MediaPlayerActions.cpp b/Engine/Plugins/Media/MediaPlayerEditor/Source/MediaPlayerEditor/Private/AssetTools/MediaPlayerActions.cpp
--- a/Engine/Plugins/Media/MediaPlayerEditor/Source/MediaPlayerEditor/Private/AssetTools/MediaPlayerActions.cpp
+++ b/Engine/Plugins/Media/MediaPlayerEditor/Source/MediaPlayerEditor/Private/AssetTools/MediaPlayerActions.cpp
@@ -96,7 +96,7 @@ UClass* FMediaPlayerActions::GetSupportedClass() const
 
 FColor FMediaPlayerActions::GetTypeColor() const
 {
-	return FColor(255, 0, 0);
+	return FColor::Red;
 }
 
 
diff --git a/Engine/Plugins/Media/MediaPlayerEditor/Source/MediaPlayerEditor/Private/AssetTools/MediaSoundWaveActions.cpp b/Engine/Plugins/Media/MediaPlayerEditor/Source/MediaPlayerEditor/Private/AssetTools/MediaSoundWaveActions.cpp
--- a/Engine/Plugins/Media/MediaPlayerEditor/Source/MediaPlayerEditor/Private/AssetTools/MediaSoundWaveActions.cpp
+++ b/Engine/Plugins/Media/MediaPlayerEditor/Source/MediaPlayerEditor/Private/AssetTools/MediaSoundWaveActions.cpp
@@ -35,7 +35,7 @@ UClass* FMediaSoundWaveActions::GetSupportedClass() const
 
 FColor FMediaSoundWaveActions::GetTypeColor() const
 {
-	return FColor(0, 0, 255);
+	return FColor::Blue;
 }
 
 
diff --git a/Engine/Plugins/Media/MediaPlayerEditor/Source/MediaPlayerEditor/Private/AssetTools/MediaTextureActions.cpp b/Engine/Plugins/Media/MediaPlayerEditor/Source/MediaPlayerEditor/Private/AssetTools/MediaTextureActions.cpp
--- a/Engine/Plugins/Media/MediaPlayerEditor/Source/MediaPlayerEditor/Private/AssetTools/MediaTextureActions.cpp
+++ b/Engine/Plugins/Media/MediaPlayerEditor/Source/MediaPlayerEditor/Private/AssetTools/MediaTextureActions.cpp
@@ -79,7 +79,7 @@ UClass* FMediaTextureActions::GetSupportedClass() const
 
 FColor FMediaTextureActions::GetTypeColor() const
 {
-	return FColor(255, 0, 0);
+	return FColor::Red;
 }
 
 
diff --git a/Engine/Source/Developer/CollisionAnalyzer/Private/CollisionAnalyzer.cpp b/Engine/Source/Developer/CollisionAnalyzer/Private/CollisionAnalyzer.cpp
--- a/Engine/Source/Developer/CollisionAnalyzer/Private/CollisionAnalyzer.cpp
+++ b/Engine/Source/Developer/CollisionAnalyzer/Private/CollisionAnalyzer.cpp
@@ -107,7 +107,7 @@ void FCollisionAnalyzer::TickAnalyzer(UWorld* World)
 	// Draw debug box if desired
 	if(DrawBox.IsValid)
 	{
-		DrawDebugBox(World, DrawBox.GetCenter(), DrawBox.GetExtent(), FColor(255,255,255));
+		DrawDebugBox(World, DrawBox.GetCenter(), DrawBox.GetExtent(), FColor::White);
 	}
 }
 
diff --git a/Engine/Source/Editor/DistCurveEditor/Private/CurveEditorSharedData.cpp b/Engine/Source/Editor/DistCurveEditor/Private/CurveEditorSharedData.cpp
--- a/Engine/Source/Editor/DistCurveEditor/Private/CurveEditorSharedData.cpp
+++ b/Engine/Source/Editor/DistCurveEditor/Private/CurveEditorSharedData.cpp
@@ -25,7 +25,7 @@ FCurveEditorSharedData::FCurveEditorSharedData(UInterpCurveEdSetup* InEdSetup)
 
 	bShowPositionMarker = false;
 	MarkerPosition = 0.f;
-	MarkerColor = FColor(255,255,255);
+	MarkerColor = FColor::White;
 
 	bShowEndMarker = false;
 	EndMarkerPosition = 0.f;
diff --git a/Engine/Source/Editor/DistCurveEditor/Private/CurveEditorViewportClient.cpp b/Engine/Source/Editor/DistCurveEditor/Private/CurveEditorViewportClient.cpp
--- a/Engine/Source/Editor/DistCurveEditor/Private/CurveEditorViewportClient.cpp
+++ b/Engine/Source/Editor/DistCurveEditor/Private/CurveEditorViewportClient.cpp
@@ -1135,7 +1135,7 @@ void FCurveEditorViewportClient::DrawEntry(FViewport* Viewport, FCanvas* Canvas,
 					
 
 					if(Canvas->IsHitTesting()) Canvas->SetHitProxy(new HCurveEditorKeyHandleProxy(CurveIndex, SubIdx, KeyIdx, true));
-					Canvas->DrawTile(HandlePos.X-2, HandlePos.Y-2, 5, 5, 0.f, 0.f, 1.f, 1.f, FColor(255,255,255));
+					Canvas->DrawTile(HandlePos.X - 2, HandlePos.Y - 2, 5, 5, 0.f, 0.f, 1.f, 1.f, FColor::White);
 					if(Canvas->IsHitTesting()) Canvas->SetHitProxy(NULL);
 				}
 
@@ -1152,7 +1152,7 @@ void FCurveEditorViewportClient::DrawEntry(FViewport* Viewport, FCanvas* Canvas,
 					LineItem.Draw( Canvas, NewKeyPos, HandlePos );
 					
 					if(Canvas->IsHitTesting()) Canvas->SetHitProxy(new HCurveEditorKeyHandleProxy(CurveIndex, SubIdx, KeyIdx, false));
-					Canvas->DrawTile(HandlePos.X-2, HandlePos.Y-2, 5, 5, 0.f, 0.f, 1.f, 1.f, FColor(255,255,255));
+					Canvas->DrawTile(HandlePos.X - 2, HandlePos.Y - 2, 5, 5, 0.f, 0.f, 1.f, 1.f, FColor::White);
 					if(Canvas->IsHitTesting()) Canvas->SetHitProxy(NULL);
 				}
 			}
diff --git a/Engine/Source/Editor/FoliageEdit/Private/FoliageEdMode.cpp b/Engine/Source/Editor/FoliageEdit/Private/FoliageEdMode.cpp
--- a/Engine/Source/Editor/FoliageEdit/Private/FoliageEdMode.cpp
+++ b/Engine/Source/Editor/FoliageEdit/Private/FoliageEdMode.cpp
@@ -416,7 +416,7 @@ bool CheckCollisionWithWorld(UFoliageType* Settings, const FFoliageInstance& Ins
 			MeshBox.Min.Z = FMath::Min(MeshBox.Max.Z, LocalHit.Z + Settings->MeshBounds.BoxExtent.Z * 0.05f);
 			FBoxSphereBounds ShrinkBound(MeshBox);
 			FBoxSphereBounds WorldBound = ShrinkBound.TransformBy(InstTransform);
-			//::DrawDebugBox(InWorld, WorldBound.Origin, WorldBound.BoxExtent, FColor(255, 0, 0), true, 10.f);
+			//::DrawDebugBox(InWorld, WorldBound.Origin, WorldBound.BoxExtent, FColor::Red, true, 10.f);
 			static FName NAME_FoliageCollisionWithWorld = FName(TEXT("FoliageCollisionWithWorld"));
 			if (InWorld->OverlapTest(WorldBound.Origin, FQuat(Inst.Rotation), FCollisionShape::MakeBox(ShrinkBound.BoxExtent * Inst.DrawScale3D * Settings->CollisionScale), FCollisionQueryParams(NAME_FoliageCollisionWithWorld, false), FCollisionObjectQueryParams(ECC_WorldStatic)))
 			{
diff --git a/Engine/Source/Editor/GeometryMode/Private/GeometryEdMode.cpp b/Engine/Source/Editor/GeometryMode/Private/GeometryEdMode.cpp
--- a/Engine/Source/Editor/GeometryMode/Private/GeometryEdMode.cpp
+++ b/Engine/Source/Editor/GeometryMode/Private/GeometryEdMode.cpp
@@ -527,12 +527,12 @@ void FEdModeGeometry::RenderPoly( const FSceneView* View, FViewport* Viewport, F
 				}
 
 				// Draw Polygon Triangles
-				const int32 VertexOffset = MeshBuilder.AddVertex(Verts[0], FVector2D::ZeroVector, FVector(1,0,0), FVector(0,1,0), FVector(0,0,1), FColor(255,255,255));
-				MeshBuilder.AddVertex(Verts[1], FVector2D::ZeroVector, FVector(1,0,0), FVector(0,1,0), FVector(0,0,1), FColor(255,255,255));
+				const int32 VertexOffset = MeshBuilder.AddVertex(Verts[0], FVector2D::ZeroVector, FVector(1,0,0), FVector(0,1,0), FVector(0,0,1), FColor::White);
+				MeshBuilder.AddVertex(Verts[1], FVector2D::ZeroVector, FVector(1,0,0), FVector(0,1,0), FVector(0,0,1), FColor::White);
 
 				for( int32 VertIdx = 2 ; VertIdx < Verts.Num() ; ++VertIdx )
 				{
-					MeshBuilder.AddVertex(Verts[VertIdx], FVector2D::ZeroVector, FVector(1,0,0), FVector(0,1,0), FVector(0,0,1), FColor(255,255,255));
+					MeshBuilder.AddVertex(Verts[VertIdx], FVector2D::ZeroVector, FVector(1,0,0), FVector(0,1,0), FVector(0,0,1), FColor::White);
 					MeshBuilder.AddTriangle( VertexOffset + VertIdx - 1, VertexOffset, VertexOffset + VertIdx);
 				}
 
diff --git a/Engine/Source/Editor/LandscapeEditor/Private/LandscapeEdMode.cpp b/Engine/Source/Editor/LandscapeEditor/Private/LandscapeEdMode.cpp
--- a/Engine/Source/Editor/LandscapeEditor/Private/LandscapeEdMode.cpp
+++ b/Engine/Source/Editor/LandscapeEditor/Private/LandscapeEdMode.cpp
@@ -2026,7 +2026,7 @@ void FEdModeLandscape::Render(const FSceneView* View, FViewport* Viewport, FPrim
 
 		for (int32 i = 0; i < 8; ++i)
 		{
-			MeshBuilder.AddVertex(CurrentGizmoActor->FrustumVerts[i], FVector2D(0, 0), FVector(1, 0, 0), FVector(0, 1, 0), FVector(0, 0, 1), FColor(255, 255, 255));
+			MeshBuilder.AddVertex(CurrentGizmoActor->FrustumVerts[i], FVector2D(0, 0), FVector(1, 0, 0), FVector(0, 1, 0), FVector(0, 0, 1), FColor::White);
 		}
 
 		// Upper box.
diff --git a/Engine/Source/Editor/LandscapeEditor/Private/LandscapeEdModeSplineTools.cpp b/Engine/Source/Editor/LandscapeEditor/Private/LandscapeEdModeSplineTools.cpp
--- a/Engine/Source/Editor/LandscapeEditor/Private/LandscapeEdModeSplineTools.cpp
+++ b/Engine/Source/Editor/LandscapeEditor/Private/LandscapeEdModeSplineTools.cpp
@@ -1282,7 +1282,7 @@ public:
 					PDI->DrawLine(StartPos, HandlePos, FColor::White, SDPG_Foreground);
 
 					if (PDI->IsHitTesting()) PDI->SetHitProxy(new HLandscapeSplineProxy_Tangent(Connection.Segment, Connection.End));
-					PDI->DrawPoint(HandlePos, FColor(255, 255, 255), 10.f, SDPG_Foreground);
+					PDI->DrawPoint(HandlePos, FColor::White, 10.f, SDPG_Foreground);
 					if (PDI->IsHitTesting()) PDI->SetHitProxy(NULL);
 				}
 			}
@@ -1305,7 +1305,7 @@ public:
 
 					PDI->DrawLine(EndPos, EndHandlePos, FColor::White, SDPG_Foreground);
 					if (PDI->IsHitTesting()) PDI->SetHitProxy(new HLandscapeSplineProxy_Tangent(Segment, !!End));
-					PDI->DrawPoint(EndHandlePos, FColor(255, 255, 255), 10.f, SDPG_Foreground);
+					PDI->DrawPoint(EndHandlePos, FColor::White, 10.f, SDPG_Foreground);
 					if (PDI->IsHitTesting()) PDI->SetHitProxy(NULL);
 				}
 			}
diff --git a/Engine/Source/Editor/Matinee/Private/Matinee.cpp b/Engine/Source/Editor/Matinee/Private/Matinee.cpp
--- a/Engine/Source/Editor/Matinee/Private/Matinee.cpp
+++ b/Engine/Source/Editor/Matinee/Private/Matinee.cpp
@@ -54,8 +54,8 @@ DEFINE_LOG_CATEGORY(LogSlateMatinee);
 
 #define LOCTEXT_NAMESPACE "Matinee"
 
-const FColor FMatinee::ActiveCamColor = FColor(255, 255, 0);
-const FColor FMatinee::SelectedCurveColor = FColor(255, 255, 0);
+const FColor FMatinee::ActiveCamColor = FColor::Yellow;
+const FColor FMatinee::SelectedCurveColor = FColor::Yellow;
 const int32	FMatinee::DuplicateKeyOffset = 10;
 const int32	FMatinee::KeySnapPixels = 5;
 
@@ -63,7 +63,7 @@ const float FMatinee::InterpEditor_ZoomIncrement = 1.2f;
 
 const FColor FMatinee::PositionMarkerLineColor = FColor(255, 222, 206);
 const FColor FMatinee::LoopRegionFillColor = FColor(80,255,80,24);
-const FColor FMatinee::Track3DSelectedColor = FColor(255,255,0);
+const FColor FMatinee::Track3DSelectedColor = FColor::Yellow;
 
 const float FMatinee::InterpEdSnapSizes[5] = { 0.01f, 0.05f, 0.1f, 0.5f, 1.0f };
 const float FMatinee::InterpEdFPSSnapSizes[9] =
diff --git a/Engine/Source/Editor/Matinee/Private/MatineeDraw.cpp b/Engine/Source/Editor/Matinee/Private/MatineeDraw.cpp
--- a/Engine/Source/Editor/Matinee/Private/MatineeDraw.cpp
+++ b/Engine/Source/Editor/Matinee/Private/MatineeDraw.cpp
@@ -734,7 +734,7 @@ void FMatineeViewportClient::DrawTimeline(FViewport* Viewport, FCanvas* Canvas)
 		FString KeyTitle = FString::Printf( TEXT("%s%d"), ( rSelKey.Track ? *rSelKey.Track->TrackTitle : TEXT( "?" ) ), rSelKey.KeyIndex );
 		FString AdjustString = FText::Format( NSLOCTEXT("UnrealEd", "Key_F", "KEY {0}"), FText::FromString(KeyTitle) ).ToString();
 
-		Canvas->DrawNGon(FIntPoint(HeadTitleMargin + 5, ViewY - 1.1 * YL - 2 * HeadTitleMargin), FColor(255, 0, 0), 12, 5);
+		Canvas->DrawNGon(FIntPoint(HeadTitleMargin + 5, ViewY - 1.1 * YL - 2 * HeadTitleMargin), FColor::Red, 12, 5);
 		TextItem.SetColor( FLinearColor::Red );
 		TextItem.Text = FText::FromString( AdjustString );
 		Canvas->DrawItem(TextItem, 2 * HeadTitleMargin + 10, (int32)( ViewY - 1.6 * YL - 2 * HeadTitleMargin ));
@@ -775,7 +775,7 @@ void FMatineeViewportClient::DrawTimeline(FViewport* Viewport, FCanvas* Canvas)
 			AdjustString += TEXT( "] " );
 		}
 
-		Canvas->DrawNGon(FIntPoint(HeadTitleMargin + 5, ViewY - 1.1*YL - 2 * HeadTitleMargin), FColor(255, 0, 0), 12, 5);
+		Canvas->DrawNGon(FIntPoint(HeadTitleMargin + 5, ViewY - 1.1*YL - 2 * HeadTitleMargin), FColor::Red, 12, 5);
 		TextItem.SetColor( FLinearColor::Red );
 		TextItem.Text = FText::FromString( AdjustString );
 		Canvas->DrawItem(TextItem, 2 * HeadTitleMargin + 10, (int32)(ViewY - 1.6*YL - 2 * HeadTitleMargin));
diff --git a/Engine/Source/Editor/MeshPaint/Private/MeshPaintEdMode.cpp b/Engine/Source/Editor/MeshPaint/Private/MeshPaintEdMode.cpp
--- a/Engine/Source/Editor/MeshPaint/Private/MeshPaintEdMode.cpp
+++ b/Engine/Source/Editor/MeshPaint/Private/MeshPaintEdMode.cpp
@@ -1444,7 +1444,7 @@ void FEdModeMeshPaint::PaintMeshVertices(
 						bool bConvertSRGB = false;
 						FColor FillColor = Params.BrushColor.ToFColor(bConvertSRGB);
 						// Original mesh didn't have any colors, so just use a default color
-						InstanceMeshLODInfo->OverrideVertexColors->InitFromSingleColor(FColor(255, 255, 255), LODModel.GetNumVertices());
+						InstanceMeshLODInfo->OverrideVertexColors->InitFromSingleColor(FColor::White, LODModel.GetNumVertices());
 					}
 
 				}
@@ -3713,8 +3713,8 @@ void FEdModeMeshPaint::CopyInstanceVertexColors()
 								// mismatched or empty color buffer - just use white
 								for( uint32 VertexIndex = 0; VertexIndex < NumPosVertices; VertexIndex++ )
 								{
-									LodColorData.ColorsByIndex.Add( FColor(255,255,255,255) );
-									LodColorData.ColorsByPosition.Add( PosBuffer->VertexPosition( VertexIndex ), FColor(255,255,255,255) );
+									LodColorData.ColorsByIndex.Add( FColor::White );
+									LodColorData.ColorsByPosition.Add( PosBuffer->VertexPosition( VertexIndex ), FColor::White );
 								}
 							}
 						}
@@ -3796,7 +3796,7 @@ void FEdModeMeshPaint::PasteInstanceVertexColors()
 
 								for (int32 TargetVertIdx = 0; TargetVertIdx < LodRenderData.GetNumVertices(); TargetVertIdx++)
 								{
-									ReOrderedColors[TargetVertIdx] = FColor(255,255,255,255);
+									ReOrderedColors[TargetVertIdx] = FColor::White;
 								}
 							}
 							else if (LodRenderData.GetNumVertices() == FoundColors->PerLODVertexColorData[CurLODIndex].ColorsByIndex.Num())
@@ -3824,7 +3824,7 @@ void FEdModeMeshPaint::PasteInstanceVertexColors()
 									else
 									{
 										// A matching color for this vertex could not be found. Make this vertex white
-										ReOrderedColors[TargetVertIdx] = FColor(255,255,255,255);
+										ReOrderedColors[TargetVertIdx] = FColor::White;
 									}
 								}
 							}
diff --git a/Engine/Source/Editor/Persona/Private/SAnimationBlendSpace.cpp b/Engine/Source/Editor/Persona/Private/SAnimationBlendSpace.cpp
--- a/Engine/Source/Editor/Persona/Private/SAnimationBlendSpace.cpp
+++ b/Engine/Source/Editor/Persona/Private/SAnimationBlendSpace.cpp
@@ -967,7 +967,7 @@ int32 SBlendSpaceGridWidget::OnPaint( const FPaintArgs& Args, const FGeometry& A
 				}
 			}
 			FText ToolTipMessage = GetToolTipText(GridPos, Seqs, Weights);
-			DrawToolTip( MousePos.GetValue(), ToolTipMessage, AllottedGeometry, MyClippingRect, FColor(255, 255, 255), OutDrawElements, TooltipLayer);
+			DrawToolTip( MousePos.GetValue(), ToolTipMessage, AllottedGeometry, MyClippingRect, FColor::White, OutDrawElements, TooltipLayer);
 		}
 		else if ( Elements.Num() > 0 )
 		{
@@ -1002,7 +1002,7 @@ int32 SBlendSpaceGridWidget::OnPaint( const FPaintArgs& Args, const FGeometry& A
 					ToolTipMessage = GetToolTipText(LastValidMouseEditorPoint, Seqs, Weights);
 				}
 
-				DrawToolTip( MousePos.GetValue(), ToolTipMessage, AllottedGeometry, MyClippingRect, FColor(255, 255, 255), OutDrawElements, TooltipLayer);
+				DrawToolTip(MousePos.GetValue(), ToolTipMessage, AllottedGeometry, MyClippingRect, FColor::White, OutDrawElements, TooltipLayer);
 			}
 
 			if (bPreviewOn)
@@ -1020,7 +1020,7 @@ int32 SBlendSpaceGridWidget::OnPaint( const FPaintArgs& Args, const FGeometry& A
 			UAnimSequence * Seq = CachedSamples[HighlightedSampleIndex].Animation;
 			const FText Sequence = (Seq)? FText::FromString( Seq->GetName() ) : LOCTEXT("None", "None");
 			const FText ToolTipMessage = FText::Format( LOCTEXT("Sequence", "{0}\nSequence ({1})"), GetInputText(LastValidMouseEditorPoint), Sequence );
-			DrawToolTip( MousePos.GetValue(), ToolTipMessage, AllottedGeometry, MyClippingRect, FColor(255, 255, 255), OutDrawElements, TooltipLayer);
+			DrawToolTip(MousePos.GetValue(), ToolTipMessage, AllottedGeometry, MyClippingRect, FColor::White, OutDrawElements, TooltipLayer);
 		}
 	}
 
diff --git a/Engine/Source/Editor/Persona/Private/SAnimationBlendSpace1D.cpp b/Engine/Source/Editor/Persona/Private/SAnimationBlendSpace1D.cpp
--- a/Engine/Source/Editor/Persona/Private/SAnimationBlendSpace1D.cpp
+++ b/Engine/Source/Editor/Persona/Private/SAnimationBlendSpace1D.cpp
@@ -382,7 +382,7 @@ int SBlendSpace1DWidget::OnPaint( const FPaintArgs& Args, const FGeometry& Allot
 				{
 					ToolTipMessage = GetToolTipText(LastValidMouseEditorPoint, Seqs, Weights);
 				}
-				DrawToolTip( MousePos.GetValue(), ToolTipMessage, AllottedGeometry, MyClippingRect, FColor(255, 255, 255), OutDrawElements, TooltipLayer);
+				DrawToolTip(MousePos.GetValue(), ToolTipMessage, AllottedGeometry, MyClippingRect, FColor::White, OutDrawElements, TooltipLayer);
 			}
 		}
 		else if (bIsOnValidSample) // If I'm on top of a sample
@@ -390,7 +390,7 @@ int SBlendSpace1DWidget::OnPaint( const FPaintArgs& Args, const FGeometry& Allot
 			UAnimSequence * Seq = CachedSamples[HighlightedSampleIndex].Animation;
 			const FText Sequence = (Seq)? FText::FromString( Seq->GetName() ): LOCTEXT("None", "None");
 			const FText ToolTipMessage = FText::Format( LOCTEXT("Sequence", "{0}\nSequence ({1})"), GetInputText(LastValidMouseEditorPoint), Sequence );
-			DrawToolTip( MousePos.GetValue(), ToolTipMessage, AllottedGeometry, MyClippingRect, FColor(255, 255, 255), OutDrawElements, TooltipLayer);
+			DrawToolTip(MousePos.GetValue(), ToolTipMessage, AllottedGeometry, MyClippingRect, FColor::White, OutDrawElements, TooltipLayer);
 		}
 	}
 
diff --git a/Engine/Source/Editor/UnrealEd/Private/Animation/DebugSkelMeshComponent.cpp b/Engine/Source/Editor/UnrealEd/Private/Animation/DebugSkelMeshComponent.cpp
--- a/Engine/Source/Editor/UnrealEd/Private/Animation/DebugSkelMeshComponent.cpp
+++ b/Engine/Source/Editor/UnrealEd/Private/Animation/DebugSkelMeshComponent.cpp
@@ -25,7 +25,7 @@ public:
 	* Constructor. 
 	* @param	Component - skeletal mesh primitive being added
 	*/
-	FDebugSkelMeshSceneProxy(const UDebugSkelMeshComponent* InComponent, FSkeletalMeshResource* InSkelMeshResource, const FColor& InWireframeOverlayColor = FColor(255, 255, 255)) :
+	FDebugSkelMeshSceneProxy(const UDebugSkelMeshComponent* InComponent, FSkeletalMeshResource* InSkelMeshResource, const FColor& InWireframeOverlayColor = FColor::White) :
 		FSkeletalMeshSceneProxy( InComponent, InSkelMeshResource )
 	{
 		SkeletalMeshComponent = InComponent;
diff --git a/Engine/Source/Editor/UnrealEd/Private/EditorServer.cpp b/Engine/Source/Editor/UnrealEd/Private/EditorServer.cpp
--- a/Engine/Source/Editor/UnrealEd/Private/EditorServer.cpp
+++ b/Engine/Source/Editor/UnrealEd/Private/EditorServer.cpp
@@ -3240,7 +3240,7 @@ bool UEditorEngine::GetPropertyColorationColor(UObject* Object, FColor& OutColor
 				if ( PropertyValue == GPropertyColorationValue )
 				{
 					bResult  = true;
-					OutColor = FColor(255,0,0);
+					OutColor = FColor::Red;
 
 					// Collect actor references.
 					if ( GPropertyColorationActorCollector && Owner )
diff --git a/Engine/Source/Editor/UnrealEd/Private/Fbx/FbxSkeletalMeshImport.cpp b/Engine/Source/Editor/UnrealEd/Private/Fbx/FbxSkeletalMeshImport.cpp
--- a/Engine/Source/Editor/UnrealEd/Private/Fbx/FbxSkeletalMeshImport.cpp
+++ b/Engine/Source/Editor/UnrealEd/Private/Fbx/FbxSkeletalMeshImport.cpp
@@ -2161,7 +2161,7 @@ bool UnFbx::FFbxImporter::FillSkelMeshImporterFromFbx( FSkeletalMeshImportData&
 			TmpWedges[UnrealVertexIndex].MatIndex = Triangle.MatIndex;
 			TmpWedges[UnrealVertexIndex].VertexIndex = ExistPointNum + Mesh->GetPolygonVertex(LocalIndex,VertexIndex);
 			// Initialize all colors to white.
-			TmpWedges[UnrealVertexIndex].Color = FColor(255,255,255);
+			TmpWedges[UnrealVertexIndex].Color = FColor::White;
 		}
 
 		//
diff --git a/Engine/Source/Editor/UnrealEd/Private/Settings/SettingsClasses.cpp b/Engine/Source/Editor/UnrealEd/Private/Settings/SettingsClasses.cpp
--- a/Engine/Source/Editor/UnrealEd/Private/Settings/SettingsClasses.cpp
+++ b/Engine/Source/Editor/UnrealEd/Private/Settings/SettingsClasses.cpp
@@ -40,10 +40,10 @@ UDestructableMeshEditorSettings::UDestructableMeshEditorSettings( const FObjectI
 	: Super(ObjectInitializer)
 {
 	AnimPreviewLightingDirection = FRotator(-45.0f, 45.0f, 0);
-	AnimPreviewSkyColor = FColor(0, 0, 255);
+	AnimPreviewSkyColor = FColor::Blue;
 	AnimPreviewFloorColor = FColor(51, 51, 51);
 	AnimPreviewSkyBrightness = 0.2f * PI;
-	AnimPreviewDirectionalColor = FColor(255, 255, 255);
+	AnimPreviewDirectionalColor = FColor::White;
 	AnimPreviewLightBrightness = 1.0f * PI;
 }
 
diff --git a/Engine/Source/Editor/UnrealEd/Private/StaticLightingSystem/StaticLightingDebug.cpp b/Engine/Source/Editor/UnrealEd/Private/StaticLightingSystem/StaticLightingDebug.cpp
--- a/Engine/Source/Editor/UnrealEd/Private/StaticLightingSystem/StaticLightingDebug.cpp
+++ b/Engine/Source/Editor/UnrealEd/Private/StaticLightingSystem/StaticLightingDebug.cpp
@@ -470,7 +470,7 @@ void DrawStaticLightingDebugInfo(const FSceneView* View,FPrimitiveDrawInterface*
 		for (int32 RayIndex = 0; RayIndex < GDebugStaticLightingInfo.ShadowRays.Num(); RayIndex++)
 		{
 			const FDebugStaticLightingRay& CurrentRay = GDebugStaticLightingInfo.ShadowRays[RayIndex];
-			PDI->DrawLine(CurrentRay.Start, CurrentRay.End, CurrentRay.bHit ? FColor(255,0,0) : FColor(0,255,0), SDPG_World);
+			PDI->DrawLine(CurrentRay.Start, CurrentRay.End, CurrentRay.bHit ? FColor::Red : FColor::Green, SDPG_World);
 		}
 		
 		for (int32 RayIndex = 0; RayIndex < GDebugStaticLightingInfo.PathRays.Num(); RayIndex++)
@@ -556,7 +556,7 @@ void DrawStaticLightingDebugInfo(const FSceneView* View,FPrimitiveDrawInterface*
 		for (int32 RayIndex = 0; RayIndex < GDebugStaticLightingInfo.IndirectPhotonPaths.Num(); RayIndex++)
 		{
 			const FDebugStaticLightingRay& CurrentRay = GDebugStaticLightingInfo.IndirectPhotonPaths[RayIndex];
-			PDI->DrawLine(CurrentRay.Start, CurrentRay.End, FColor(255,255,255), SDPG_World);
+			PDI->DrawLine(CurrentRay.Start, CurrentRay.End, FColor::White, SDPG_World);
 		}
 
 		for (int32 SampleIndex = 0; SampleIndex < GDebugStaticLightingInfo.VolumeLightingSamples.Num(); SampleIndex++)
diff --git a/Engine/Source/Editor/UnrealEd/Private/Tests/EditorAutomationTests.cpp b/Engine/Source/Editor/UnrealEd/Private/Tests/EditorAutomationTests.cpp
--- a/Engine/Source/Editor/UnrealEd/Private/Tests/EditorAutomationTests.cpp
+++ b/Engine/Source/Editor/UnrealEd/Private/Tests/EditorAutomationTests.cpp
@@ -49,7 +49,7 @@ struct PointLightParameters
 		, LightBrightness(5000.0f)
 		, LightRadius(1000.0f)
 		, LightLocation(FVector(0.0f, 0.0f, 0.0f))
-		, LightColor(FColor(255, 255, 255))
+		, LightColor(FColor::White)
 	{
 	}
 };
@@ -1460,7 +1460,7 @@ bool FLightPlacement::RunTest(const FString& Parameters)
 	const FTransform Transform(FVector(0.0f, 0.0f, 400.0f));
 	APointLight* PointLight = Cast<APointLight>(GEditor->AddActor(World->GetCurrentLevel(), APointLight::StaticClass(), Transform));
 	LightParameters.PointLight = PointLight;
-	LightParameters.LightColor = FColor(255, 0, 0);
+	LightParameters.LightColor = FColor::Red;
 	LightParameters.LightLocation = FVector(0.0f, 0.0f, 400.0f);
 	ADD_LATENT_AUTOMATION_COMMAND(PointLightUpdateCommand(LightParameters));
 
@@ -1495,7 +1495,7 @@ bool FLightPlacement::RunTest(const FString& Parameters)
 	//Update the original point light actor.
 	LightParameters.LightRadius = 500.0f;
 	LightParameters.LightLocation = FVector(500.0f, 300.0f, 500.0f);
-	LightParameters.LightColor = FColor(255, 255, 255);
+	LightParameters.LightColor = FColor::White;
 	ADD_LATENT_AUTOMATION_COMMAND(PointLightUpdateCommand(LightParameters));
 
 	//Wait
diff --git a/Engine/Source/Editor/UnrealEd/Private/UnrealWidget.cpp b/Engine/Source/Editor/UnrealEd/Private/UnrealWidget.cpp
--- a/Engine/Source/Editor/UnrealEd/Private/UnrealWidget.cpp
+++ b/Engine/Source/Editor/UnrealEd/Private/UnrealWidget.cpp
@@ -24,9 +24,9 @@ FWidget::FWidget()
 	AxisColorX = FLinearColor(0.594f,0.0197f,0.0f);
 	AxisColorY = FLinearColor(0.1349f,0.3959f,0.0f);
 	AxisColorZ = FLinearColor(0.0251f,0.207f,0.85f);
-	PlaneColorXY = FColor(255,255,0);
+	PlaneColorXY = FColor::Yellow;
 	ScreenSpaceColor  = FColor(196, 196, 196);
-	CurrentColor = FColor(255,255,0);
+	CurrentColor = FColor::Yellow;
 
 	UMaterial* AxisMaterialBase = GEngine->ArrowMaterial;
 
diff --git a/Engine/Source/Editor/UnrealEd/Private/Utils.cpp b/Engine/Source/Editor/UnrealEd/Private/Utils.cpp
--- a/Engine/Source/Editor/UnrealEd/Private/Utils.cpp
+++ b/Engine/Source/Editor/UnrealEd/Private/Utils.cpp
@@ -91,16 +91,22 @@ void FUnrealEdUtils::DrawWidget(const FSceneView* View,FPrimitiveDrawInterface*
 	const float WidgetRadius = View->Project(WidgetOrigin).W * (UnrealEd_WidgetSize / ZoomFactor);
 
 	// Choose its color. Highlight manipulated axis in yellow.
-	FColor XColor(255, 0, 0);
-	FColor YColor(0, 255, 0);
-	FColor ZColor(0, 0, 255);
+	FColor XColor(FColor::Red);
+	FColor YColor(FColor::Green);
+	FColor ZColor(FColor::Blue);
 
-	if(HighlightAxis == EAxisList::X)
-		XColor = FColor(255, 255, 0);
-	else if(HighlightAxis == EAxisList::Y)
-		YColor = FColor(255, 255, 0);
-	else if(HighlightAxis == EAxisList::Z)
-		ZColor = FColor(255, 255, 0);
+	if (HighlightAxis == EAxisList::X)
+	{
+		XColor = FColor::Yellow;
+	}
+	else if (HighlightAxis == EAxisList::Y)
+	{
+		YColor = FColor::Yellow;
+	}
+	else if (HighlightAxis == EAxisList::Z)
+	{
+		ZColor = FColor::Yellow;
+	}
 
 	const FVector XAxis = WidgetMatrix.GetScaledAxis( EAxis::X ); 
 	const FVector YAxis = WidgetMatrix.GetScaledAxis( EAxis::Y ); 
@@ -146,9 +152,9 @@ void FUnrealEdUtils::DrawWidget(const FSceneView* View,FPrimitiveDrawInterface*
 			FVector AlongY = WidgetOrigin + (YAxis * WidgetRadius * 0.3f);
 			FVector AlongZ = WidgetOrigin + (ZAxis * WidgetRadius * 0.3f);
 
-			PDI->DrawLine(AlongX, AlongY, FColor(255,255,255), SDPG_Foreground);
-			PDI->DrawLine(AlongY, AlongZ, FColor(255,255,255), SDPG_Foreground);
-			PDI->DrawLine(AlongZ, AlongX, FColor(255,255,255), SDPG_Foreground);
+			PDI->DrawLine(AlongX, AlongY, FColor::White, SDPG_Foreground);
+			PDI->DrawLine(AlongY, AlongZ, FColor::White, SDPG_Foreground);
+			PDI->DrawLine(AlongZ, AlongX, FColor::White, SDPG_Foreground);
 		}
 	}
 }
diff --git a/Engine/Source/Runtime/AppFramework/Private/Widgets/Testing/STestSuite.cpp b/Engine/Source/Runtime/AppFramework/Private/Widgets/Testing/STestSuite.cpp
--- a/Engine/Source/Runtime/AppFramework/Private/Widgets/Testing/STestSuite.cpp
+++ b/Engine/Source/Runtime/AppFramework/Private/Widgets/Testing/STestSuite.cpp
@@ -548,10 +548,10 @@ private:
 	{
 		TArray<FSlateGradientStop> GradientStops;
 
-		GradientStops.Add( FSlateGradientStop( FVector2D(InParams.Geometry.Size.X*.1f,0), FColor(255,255,0) ) );
-		GradientStops.Add( FSlateGradientStop( FVector2D(InParams.Geometry.Size.X*.25f,0), FColor(255,0,255) ) );
-		GradientStops.Add( FSlateGradientStop( FVector2D(InParams.Geometry.Size.X*.75f,0), FColor(0,0,255) ) );
-		GradientStops.Add( FSlateGradientStop( FVector2D(InParams.Geometry.Size.X*0.9f,0), FColor(0,255,0) ) );
+		GradientStops.Add( FSlateGradientStop(FVector2D(InParams.Geometry.Size.X*.1f, 0), FColor::Yellow) );
+		GradientStops.Add( FSlateGradientStop( FVector2D(InParams.Geometry.Size.X*.25f,0), FColor::Magenta ) );
+		GradientStops.Add( FSlateGradientStop( FVector2D(InParams.Geometry.Size.X*.75f,0), FColor::Blue ) );
+		GradientStops.Add( FSlateGradientStop( FVector2D(InParams.Geometry.Size.X*0.9f,0), FColor::Green) );
 
 		FSlateDrawElement::MakeGradient(
 			InParams.OutDrawElements,
@@ -582,7 +582,7 @@ private:
 			InParams.ClippingRect,
 			4.0f,
 			InParams.bEnabled ? ESlateDrawEffect::None : ESlateDrawEffect::DisabledEffect,
-			FColor(255,255,255)
+			FColor::White
 		);
 
 	
@@ -602,7 +602,7 @@ private:
 			LinePoints,
 			InParams.ClippingRect,
 			InParams.bEnabled ? ESlateDrawEffect::None : ESlateDrawEffect::DisabledEffect,
-			FColor(255,0,255)
+			FColor::Magenta
 			);
 			
 
diff --git a/Engine/Source/Runtime/Engine/Private/AI/Navigation/NavAreas/NavArea.cpp b/Engine/Source/Runtime/Engine/Private/AI/Navigation/NavAreas/NavArea.cpp
--- a/Engine/Source/Runtime/Engine/Private/AI/Navigation/NavAreas/NavArea.cpp
+++ b/Engine/Source/Runtime/Engine/Private/AI/Navigation/NavAreas/NavArea.cpp
@@ -7,7 +7,7 @@ UNavArea::UNavArea(const FObjectInitializer& ObjectInitializer) : Super(ObjectIn
 {
 	DefaultCost = 1.f;
 	FixedAreaEnteringCost = 0.f;
-	DrawColor = FColor(255,0,255);
+	DrawColor = FColor::Magenta;
 	SupportedAgentsBits = 0xffffffff;
 	// NOTE! AreaFlags == 0 means UNWALKABLE!
 	AreaFlags = 1;  
diff --git a/Engine/Source/Runtime/Engine/Private/Atmosphere/Atmosphere.cpp b/Engine/Source/Runtime/Engine/Private/Atmosphere/Atmosphere.cpp
--- a/Engine/Source/Runtime/Engine/Private/Atmosphere/Atmosphere.cpp
+++ b/Engine/Source/Runtime/Engine/Private/Atmosphere/Atmosphere.cpp
@@ -100,7 +100,7 @@ UAtmosphericFogComponent::UAtmosphericFogComponent(const FObjectInitializer& Obj
 
 	// Default lighting
 	DefaultBrightness = 50.f;
-	DefaultLightColor = FColor(255, 255, 255);
+	DefaultLightColor = FColor::White;
 
 	bDisableSunDisk = false;
 	bDisableGroundScattering = false;
diff --git a/Engine/Source/Runtime/Engine/Private/BillboardComponent.cpp b/Engine/Source/Runtime/Engine/Private/BillboardComponent.cpp
--- a/Engine/Source/Runtime/Engine/Private/BillboardComponent.cpp
+++ b/Engine/Source/Runtime/Engine/Private/BillboardComponent.cpp
@@ -160,7 +160,7 @@ public:
 					// Sprites of locked actors draw in red.
 					if (bIsActorLocked)
 					{
-						ColorToUse = FColor(255,0,0);
+						ColorToUse = FColor::Red;
 					}
 					FLinearColor LevelColorToUse = IsSelected() ? ColorToUse : (FLinearColor)LevelColor;
 					FLinearColor PropertyColorToUse = PropertyColor;
diff --git a/Engine/Source/Runtime/Engine/Private/CharacterMovementComponent.cpp b/Engine/Source/Runtime/Engine/Private/CharacterMovementComponent.cpp
--- a/Engine/Source/Runtime/Engine/Private/CharacterMovementComponent.cpp
+++ b/Engine/Source/Runtime/Engine/Private/CharacterMovementComponent.cpp
@@ -2336,7 +2336,7 @@ void UCharacterMovementComponent::CalcAvoidanceVelocity(float DeltaTime)
 #if !(UE_BUILD_SHIPPING || UE_BUILD_TEST)
 			if (bShowDebug)
 			{
-				DrawDebugLine(GetWorld(), GetActorFeetLocation(), GetActorFeetLocation() + Velocity, FColor(0,0,255), true, 0.5f, SDPG_MAX);
+				DrawDebugLine(GetWorld(), GetActorFeetLocation(), GetActorFeetLocation() + Velocity, FColor::Blue, true, 0.5f, SDPG_MAX);
 			}
 #endif
 		}
@@ -2356,7 +2356,7 @@ void UCharacterMovementComponent::CalcAvoidanceVelocity(float DeltaTime)
 #if !(UE_BUILD_SHIPPING || UE_BUILD_TEST)
 				if (bShowDebug)
 				{
-					DrawDebugLine(GetWorld(), GetActorFeetLocation(), GetActorFeetLocation() + Velocity, FColor(255,0,0), true, 0.05f, SDPG_MAX, 10.0f);
+					DrawDebugLine(GetWorld(), GetActorFeetLocation(), GetActorFeetLocation() + Velocity, FColor::Red, true, 0.05f, SDPG_MAX, 10.0f);
 				}
 #endif
 			}
@@ -2367,7 +2367,7 @@ void UCharacterMovementComponent::CalcAvoidanceVelocity(float DeltaTime)
 #if !(UE_BUILD_SHIPPING || UE_BUILD_TEST)
 				if (bShowDebug)
 				{
-					//DrawDebugLine(GetWorld(), GetActorLocation(), GetActorLocation() + Velocity, FColor(0,255,0), true, 0.05f, SDPG_MAX, 10.0f);
+					//DrawDebugLine(GetWorld(), GetActorLocation(), GetActorLocation() + Velocity, FColor::Green, true, 0.05f, SDPG_MAX, 10.0f);
 				}
 #endif
 			}
@@ -2380,13 +2380,13 @@ void UCharacterMovementComponent::CalcAvoidanceVelocity(float DeltaTime)
 #if !(UE_BUILD_SHIPPING || UE_BUILD_TEST)
 	else if (bShowDebug)
 	{
-		DrawDebugLine(GetWorld(), GetActorFeetLocation(), GetActorFeetLocation() + Velocity, FColor(255,255,0), true, 0.05f, SDPG_MAX);
+		DrawDebugLine(GetWorld(), GetActorFeetLocation(), GetActorFeetLocation() + Velocity, FColor::Yellow, true, 0.05f, SDPG_MAX);
 	}
 
 	if (bShowDebug)
 	{
 		FVector UpLine(0,0,500);
-		DrawDebugLine(GetWorld(), GetActorFeetLocation(), GetActorFeetLocation() + UpLine, (AvoidanceLockTimer > 0.01f) ? FColor(255,0,0) : FColor(0,0,255), true, 0.05f, SDPG_MAX, 5.0f);
+		DrawDebugLine(GetWorld(), GetActorFeetLocation(), GetActorFeetLocation() + UpLine, (AvoidanceLockTimer > 0.01f) ? FColor::Red : FColor::Blue, true, 0.05f, SDPG_MAX, 5.0f);
 	}
 #endif
 }
diff --git a/Engine/Source/Runtime/Engine/Private/CheatManager.cpp b/Engine/Source/Runtime/Engine/Private/CheatManager.cpp
--- a/Engine/Source/Runtime/Engine/Private/CheatManager.cpp
+++ b/Engine/Source/Runtime/Engine/Private/CheatManager.cpp
@@ -683,7 +683,7 @@ void UCheatManager::TickCollisionDebug()
 	for (int32 TraceIdx=0; TraceIdx < DebugTraceInfoList.Num(); ++TraceIdx)
 	{
 		FDebugTraceInfo & TraceInfo = DebugTraceInfoList[TraceIdx];
-		DrawDebugDirectionalArrow(GetWorld(), TraceInfo.LineTraceStart, TraceInfo.LineTraceEnd, 10.f, FColor(255,255,255), SDPG_World);
+		DrawDebugDirectionalArrow(GetWorld(), TraceInfo.LineTraceStart, TraceInfo.LineTraceEnd, 10.f, FColor::White, SDPG_World);
 		// if it's current trace index, use highlight color
 		if (CurrentTraceIndex == TraceIdx)
 		{
diff --git a/Engine/Source/Runtime/Engine/Private/Distributions.cpp b/Engine/Source/Runtime/Engine/Private/Distributions.cpp
--- a/Engine/Source/Runtime/Engine/Private/Distributions.cpp
+++ b/Engine/Source/Runtime/Engine/Private/Distributions.cpp
@@ -1612,7 +1612,7 @@ FColor UDistributionFloatConstant::GetKeyColor(int32 SubIndex, int32 KeyIndex, c
 	check( KeyIndex == 0 );
 
 	// Always return RED since there is only one key
-	return FColor(255,0,0);
+	return FColor::Red;
 }
 
 void UDistributionFloatConstant::GetInRange(float& MinIn, float& MaxIn) const
@@ -1725,7 +1725,7 @@ FColor UDistributionFloatConstantCurve::GetKeyColor(int32 SubIndex, int32 KeyInd
 	check( KeyIndex >= 0 && KeyIndex < ConstantCurve.Points.Num() );
 
 	// Always return RED since there is only one sub-curve.
-	return FColor(255,0,0);
+	return FColor::Red;
 }
 
 void UDistributionFloatConstantCurve::GetInRange(float& MinIn, float& MaxIn) const
@@ -1918,11 +1918,11 @@ FColor UDistributionFloatUniform::GetSubCurveButtonColor(int32 SubCurveIndex, bo
 	{
 	case 0:
 		// Red
-		ButtonColor = bIsSubCurveHidden ? FColor(32, 0,  0) : FColor(255, 0, 0);
+		ButtonColor = bIsSubCurveHidden ? FColor(32, 0, 0) : FColor::Red;
 		break;
 	case 1:
 		// Green
-		ButtonColor = bIsSubCurveHidden ? FColor(0, 32,  0) : FColor(0, 255, 0);
+		ButtonColor = bIsSubCurveHidden ? FColor(0, 32, 0) : FColor::Green;
 		break;
 	default:
 		// A bad sub-curve index was given. 
@@ -1957,13 +1957,11 @@ FColor UDistributionFloatUniform::GetKeyColor(int32 SubIndex, int32 KeyIndex, co
 
 	if( 0 == SubIndex )
 	{
-		// RED
-		KeyColor = FColor(255,0,0);
+		KeyColor = FColor::Red;
 	} 
 	else
 	{
-		// GREEN
-		KeyColor = FColor(0,255,0);
+		KeyColor = FColor::Green;
 	}
 
 	return KeyColor;
@@ -2109,11 +2107,11 @@ FColor UDistributionFloatUniformCurve::GetSubCurveButtonColor(int32 SubCurveInde
 	{
 	case 0:
 		// Red
-		ButtonColor = bIsSubCurveHidden ? FColor(32, 0,  0) : FColor(255, 0, 0);
+		ButtonColor = bIsSubCurveHidden ? FColor(32, 0, 0) : FColor::Red;
 		break;
 	case 1:
 		// Green
-		ButtonColor = bIsSubCurveHidden ? FColor(0, 32,  0) : FColor(0, 255, 0);
+		ButtonColor = bIsSubCurveHidden ? FColor(0, 32, 0) : FColor::Green;
 		break;
 	default:
 		// A bad sub-curve index was given. 
@@ -2152,11 +2150,11 @@ FColor UDistributionFloatUniformCurve::GetKeyColor(int32 SubIndex, int32 KeyInde
 
 	if (SubIndex == 0)
 	{
-		return FColor(255,0,0);
+		return FColor::Red;
 	}
 	else
 	{
-		return FColor(0,255,0);
+		return FColor::Green;
 	}
 }
 
@@ -2388,15 +2386,15 @@ FColor UDistributionVectorConstant::GetSubCurveButtonColor(int32 SubCurveIndex,
 	{
 	case 0:
 		// Red
-		ButtonColor = bIsSubCurveHidden ? FColor(32, 0,  0) : FColor(255, 0, 0);
+		ButtonColor = bIsSubCurveHidden ? FColor(32, 0, 0) : FColor::Red;
 		break;
 	case 1:
 		// Green
-		ButtonColor = bIsSubCurveHidden ? FColor(0, 32,  0) : FColor(0, 255, 0);
+		ButtonColor = bIsSubCurveHidden ? FColor(0, 32, 0) : FColor::Green;
 		break;
 	case 2:
 		// Blue
-		ButtonColor = bIsSubCurveHidden ? FColor(0, 0, 32) : FColor(0, 0, 255);
+		ButtonColor = bIsSubCurveHidden ? FColor(0, 0, 32) : FColor::Blue;
 		break;
 	default:
 		// A bad sub-curve index was given. 
@@ -2457,12 +2455,18 @@ FColor UDistributionVectorConstant::GetKeyColor(int32 SubIndex, int32 KeyIndex,
 	check( SubIndex >= 0 && SubIndex < 3);
 	check( KeyIndex == 0 );
 
-	if(SubIndex == 0)
-		return FColor(255,0,0);
-	else if(SubIndex == 1)
-		return FColor(0,255,0);
+	if (SubIndex == 0)
+	{
+		return FColor::Red;
+	}
+	else if (SubIndex == 1)
+	{
+		return FColor::Green;
+	}
 	else
-		return FColor(0,0,255);
+	{
+		return FColor::Blue;
+	}
 }
 
 void UDistributionVectorConstant::GetInRange(float& MinIn, float& MaxIn) const
@@ -2622,15 +2626,15 @@ FColor UDistributionVectorConstantCurve::GetSubCurveButtonColor(int32 SubCurveIn
 	{
 	case 0:
 		// Red
-		ButtonColor = bIsSubCurveHidden ? FColor(32, 0,  0) : FColor(255, 0, 0);
+		ButtonColor = bIsSubCurveHidden ? FColor(32, 0, 0) : FColor::Red;
 		break;
 	case 1:
 		// Green
-		ButtonColor = bIsSubCurveHidden ? FColor(0, 32,  0) : FColor(0, 255, 0);
+		ButtonColor = bIsSubCurveHidden ? FColor(0, 32, 0) : FColor::Green;
 		break;
 	case 2:
 		// Blue
-		ButtonColor = bIsSubCurveHidden ? FColor(0, 0, 32) : FColor(0, 0, 255);
+		ButtonColor = bIsSubCurveHidden ? FColor(0, 0, 32) : FColor::Blue;
 		break;
 	default:
 		// A bad sub-curve index was given. 
@@ -2687,12 +2691,18 @@ FColor UDistributionVectorConstantCurve::GetKeyColor(int32 SubIndex, int32 KeyIn
 	check( SubIndex >= 0 && SubIndex < 3);
 	check( KeyIndex >= 0 && KeyIndex < ConstantCurve.Points.Num() );
 
-	if(SubIndex == 0)
-		return FColor(255,0,0);
-	else if(SubIndex == 1)
-		return FColor(0,255,0);
+	if (SubIndex == 0)
+	{
+		return FColor::Red;
+	}
+	else if (SubIndex == 1)
+	{
+		return FColor::Green;
+	}
 	else
-		return FColor(0,0,255);
+	{
+		return FColor::Blue;
+	}
 }
 
 void UDistributionVectorConstantCurve::GetInRange(float& MinIn, float& MaxIn) const
@@ -3255,7 +3265,7 @@ FColor UDistributionVectorUniform::GetSubCurveButtonColor(int32 SubCurveIndex, b
 	{
 	case 0:
 		// Red
-		ButtonColor = bIsSubCurveHidden ? FColor(32, 0, 0) : FColor(255, 0, 0);
+		ButtonColor = bIsSubCurveHidden ? FColor(32, 0, 0) : FColor::Red;
 		break;
 	case 1:
 		if (bShouldGroupMinAndMax)
@@ -3266,19 +3276,19 @@ FColor UDistributionVectorUniform::GetSubCurveButtonColor(int32 SubCurveIndex, b
 		else
 		{
 			// Green
-			ButtonColor = bIsSubCurveHidden ? FColor(0, 32,  0) : FColor(0, 255, 0);
+			ButtonColor = bIsSubCurveHidden ? FColor(0, 32,  0) : FColor::Green;
 		}
 		break;
 	case 2:
 		if (bShouldGroupMinAndMax)
 		{
 			// Green
-			ButtonColor = bIsSubCurveHidden ? FColor(0, 32, 0) : FColor(0, 255, 0);
+			ButtonColor = bIsSubCurveHidden ? FColor(0, 32, 0) : FColor::Green;
 		}
 		else
 		{
 			// Blue
-			ButtonColor = bIsSubCurveHidden ? FColor(0, 0, 32) : FColor(0, 0, 255);
+			ButtonColor = bIsSubCurveHidden ? FColor(0, 0, 32) : FColor::Blue;
 		}
 		break;
 	case 3:
@@ -3287,7 +3297,7 @@ FColor UDistributionVectorUniform::GetSubCurveButtonColor(int32 SubCurveIndex, b
 		break;
 	case 4:
 		// Blue
-		ButtonColor = bIsSubCurveHidden ? FColor(0, 0, 32) : FColor(0, 0, 255);
+		ButtonColor = bIsSubCurveHidden ? FColor(0, 0, 32) : FColor::Blue;
 		break;
 	case 5:
 		// Dark blue
@@ -3361,18 +3371,30 @@ FColor UDistributionVectorUniform::GetKeyColor(int32 SubIndex, int32 KeyIndex, c
 	check( SubIndex >= 0 && SubIndex < 6);
 	check( KeyIndex == 0 );
 
-	if(SubIndex == 0)
-		return FColor(128,0,0);
-	else if(SubIndex == 1)
-		return FColor(255,0,0);
-	else if(SubIndex == 2)
-		return FColor(0,128,0);
-	else if(SubIndex == 3)
-		return FColor(0,255,0);
-	else if(SubIndex == 4)
-		return FColor(0,0,128);
+	if (SubIndex == 0)
+	{
+		return FColor(128, 0, 0);
+	}
+	else if (SubIndex == 1)
+	{
+		return FColor::Red;
+	}
+	else if (SubIndex == 2)
+	{
+		return FColor(0, 128, 0);
+	}
+	else if (SubIndex == 3)
+	{
+		return FColor::Green;
+	}
+	else if (SubIndex == 4)
+	{
+		return FColor(0, 0, 128);
+	}
 	else
-		return FColor(0,0,255);
+	{
+		return FColor::Blue;
+	}
 }
 
 void UDistributionVectorUniform::GetInRange(float& MinIn, float& MaxIn) const
@@ -3651,7 +3673,7 @@ FColor UDistributionVectorUniformCurve::GetSubCurveButtonColor(int32 SubCurveInd
 	{
 	case 0:
 		// Red
-		ButtonColor = bIsSubCurveHidden ? FColor(32, 0,  0) : FColor(255, 0, 0);
+		ButtonColor = bIsSubCurveHidden ? FColor(32, 0, 0) : FColor::Red;
 		break;
 	case 1:
 		if (bShouldGroupMinAndMax)
@@ -3662,19 +3684,19 @@ FColor UDistributionVectorUniformCurve::GetSubCurveButtonColor(int32 SubCurveInd
 		else
 		{
 			// Green
-			ButtonColor = bIsSubCurveHidden ? FColor(0, 32, 0) : FColor(0, 255, 0);
+			ButtonColor = bIsSubCurveHidden ? FColor(0, 32, 0) : FColor::Green;
 		}
 		break;
 	case 2:
 		if (bShouldGroupMinAndMax)
 		{
 			// Green
-			ButtonColor = bIsSubCurveHidden ? FColor(0, 32, 0) : FColor(0, 255, 0);
+			ButtonColor = bIsSubCurveHidden ? FColor(0, 32, 0) : FColor::Green;
 		}
 		else
 		{
 			// Blue
-			ButtonColor = bIsSubCurveHidden ? FColor(0, 0, 32) : FColor(0, 0, 255);
+			ButtonColor = bIsSubCurveHidden ? FColor(0, 0, 32) : FColor::Blue;
 		}
 		break;
 	case 3:
@@ -3683,7 +3705,7 @@ FColor UDistributionVectorUniformCurve::GetSubCurveButtonColor(int32 SubCurveInd
 		break;
 	case 4:
 		// Blue
-		ButtonColor = bIsSubCurveHidden ? FColor(  0,  0, 32) : FColor(0, 0, 255);
+		ButtonColor = bIsSubCurveHidden ? FColor(0, 0, 32) : FColor::Blue;
 		break;
 	case 5:
 		// Dark blue
@@ -3729,27 +3751,23 @@ FColor UDistributionVectorUniformCurve::GetKeyColor(int32 SubIndex, int32 KeyInd
 
 	if (SubIndex == 0)
 	{
-		return FColor(255,0,0);
+		return FColor::Red;
 	}
-	else 
-	if (SubIndex == 1)
+	else if (SubIndex == 1)
 	{
 		return FColor(128,0,0);
 	}
-	else
-	if (SubIndex == 2)
+	else if (SubIndex == 2)
 	{
-		return FColor(0,255,0);
+		return FColor::Green;
 	}
-	else
-	if (SubIndex == 3)
+	else if (SubIndex == 3)
 	{
 		return FColor(0,128,0);
 	}
-	else
-	if (SubIndex == 4)
+	else if (SubIndex == 4)
 	{
-		return FColor(0,0,255);
+		return FColor::Blue;
 	}
 	else
 	{
diff --git a/Engine/Source/Runtime/Engine/Private/InterpolationCurveEd.cpp b/Engine/Source/Runtime/Engine/Private/InterpolationCurveEd.cpp
--- a/Engine/Source/Runtime/Engine/Private/InterpolationCurveEd.cpp
+++ b/Engine/Source/Runtime/Engine/Private/InterpolationCurveEd.cpp
@@ -84,15 +84,15 @@ FColor UInterpTrackMove::GetSubCurveButtonColor(int32 SubCurveIndex, bool bIsSub
 	{
 	case 0:
 		// Red
-		ButtonColor = bIsSubCurveHidden ? FColor(32, 0,  0) : FColor(255, 0, 0);
+		ButtonColor = bIsSubCurveHidden ? FColor(32, 0, 0) : FColor::Red;
 		break;
 	case 1:
 		// Green
-		ButtonColor = bIsSubCurveHidden ? FColor(0, 32,  0) : FColor(0, 255, 0);
+		ButtonColor = bIsSubCurveHidden ? FColor(0, 32, 0) : FColor::Green;
 		break;
 	case 2:
 		// Blue
-		ButtonColor = bIsSubCurveHidden ? FColor(0, 0, 32) : FColor(0, 0, 255);
+		ButtonColor = bIsSubCurveHidden ? FColor(0, 0, 32) : FColor::Blue;
 		break;
 	case 3:
 		// Dark red
@@ -208,11 +208,11 @@ FColor UInterpTrackMove::GetKeyColor(int32 SubIndex, int32 KeyIndex, const FColo
 	check( KeyIndex >= 0 && KeyIndex < PosTrack.Points.Num() );
 
 	if(SubIndex == CalcSubIndex(true,0))
-		return FColor(255,0,0);
+		return FColor::Red;
 	else if(SubIndex == CalcSubIndex(true,1))
-		return FColor(0,255,0);
+		return FColor::Green;
 	else if(SubIndex == CalcSubIndex(true,2))
-		return FColor(0,0,255);
+		return FColor::Blue;
 	else if(SubIndex == CalcSubIndex(false,0))
 		return FColor(255,128,128);
 	else if(SubIndex == CalcSubIndex(false,1))
@@ -222,7 +222,7 @@ FColor UInterpTrackMove::GetKeyColor(int32 SubIndex, int32 KeyIndex, const FColo
 	else
 	{
 		check(0);
-		return FColor(0,0,0);
+		return FColor::Black;
 	}
 }
 
@@ -452,15 +452,15 @@ FColor UInterpTrackMoveAxis::GetSubCurveButtonColor( int32 SubCurveIndex, bool b
 	{
 	case AXIS_TranslationX:
 	case AXIS_RotationX:
-		ButtonColor = bIsSubCurveHidden ? FColor(32,0,0) : FColor(255,0,0);
+		ButtonColor = bIsSubCurveHidden ? FColor(32,0,0) : FColor::Red;
 		break;
 	case AXIS_TranslationY:
 	case AXIS_RotationY:
-		ButtonColor = bIsSubCurveHidden ? FColor(0,32,0) : FColor(0,255,0);
+		ButtonColor = bIsSubCurveHidden ? FColor(0,32,0) : FColor::Green;
 		break;
 	case AXIS_TranslationZ:
 	case AXIS_RotationZ:
-		ButtonColor = bIsSubCurveHidden ? FColor(0,0,32) : FColor(0,0,255);
+		ButtonColor = bIsSubCurveHidden ? FColor(0,0,32) : FColor::Blue;
 		break;
 	default:
 		checkf(false, TEXT("Invalid axis") );
@@ -480,13 +480,13 @@ FColor UInterpTrackMoveAxis::GetKeyColor( int32 SubIndex, int32 KeyIndex, const
 	switch( MoveAxis )
 	{
 	case AXIS_TranslationX:
-		KeyColor = FColor(255,0,0);
+		KeyColor = FColor::Red;
 		break;
 	case AXIS_TranslationY:
-		KeyColor = FColor(0,255,0);
+		KeyColor = FColor::Green;
 		break;
 	case AXIS_TranslationZ:
-		KeyColor = FColor(0,0,255);
+		KeyColor = FColor::Blue;
 		break;
 	case AXIS_RotationX:
 		KeyColor = FColor(255,128,128);
@@ -685,15 +685,15 @@ FColor UInterpTrackVectorBase::GetSubCurveButtonColor(int32 SubCurveIndex, bool
 	{
 	case 0:
 		// Red
-		ButtonColor = bIsSubCurveHidden ? FColor(32, 0,  0) : FColor(255, 0, 0);
+		ButtonColor = bIsSubCurveHidden ? FColor(32, 0, 0) : FColor::Red;
 		break;
 	case 1:
 		// Green
-		ButtonColor = bIsSubCurveHidden ? FColor(0, 32,  0) : FColor(0, 255, 0);
+		ButtonColor = bIsSubCurveHidden ? FColor(0, 32,  0) : FColor::Green;
 		break;
 	case 2:
 		// Blue
-		ButtonColor = bIsSubCurveHidden ? FColor(0, 0, 32) : FColor(0, 0, 255);
+		ButtonColor = bIsSubCurveHidden ? FColor(0, 0, 32) : FColor::Blue;
 		break;
 	default:
 		// A bad sub-curve index was given. 
@@ -751,12 +751,18 @@ FColor UInterpTrackVectorBase::GetKeyColor(int32 SubIndex, int32 KeyIndex, const
 	check( SubIndex >= 0 && SubIndex < 3);
 	check( KeyIndex >= 0 && KeyIndex < VectorTrack.Points.Num() );
 
-	if(SubIndex == 0)
-		return FColor(255,0,0);
-	else if(SubIndex == 1)
-		return FColor(0,255,0);
+	if (SubIndex == 0)
+	{
+		return FColor::Red;
+	}
+	else if (SubIndex == 1)
+	{
+		return FColor::Green;
+	}
 	else
-		return FColor(0,0,255);
+	{
+		return FColor::Blue;
+	}
 }
 
 EInterpCurveMode UInterpTrackVectorBase::GetKeyInterpMode(int32 KeyIndex) const
@@ -897,19 +903,18 @@ FColor UInterpTrackLinearColorBase::GetSubCurveButtonColor(int32 SubCurveIndex,
 	{
 	case 0:
 		// Red
-		ButtonColor = bIsSubCurveHidden ? FColor(32, 0,  0) : FColor(255, 0, 0);
+		ButtonColor = bIsSubCurveHidden ? FColor(32, 0,  0) : FColor::Red;
 		break;
 	case 1:
 		// Green
-		ButtonColor = bIsSubCurveHidden ? FColor(0, 32,  0) : FColor(0, 255, 0);
+		ButtonColor = bIsSubCurveHidden ? FColor(0, 32,  0) : FColor::Green;
 		break;
 	case 2:
 		// Blue
-		ButtonColor = bIsSubCurveHidden ? FColor(0, 0, 32) : FColor(0, 0, 255);
+		ButtonColor = bIsSubCurveHidden ? FColor(0, 0, 32) : FColor::Blue;
 		break;
 	case 3:
-		// Dark red
-		ButtonColor = bIsSubCurveHidden ? FColor(0, 0, 0) : FColor(255, 255, 255);
+		ButtonColor = bIsSubCurveHidden ? FColor::Black : FColor::White;
 		break;
 	default:
 		// A bad sub-curve index was given. 
@@ -969,14 +974,22 @@ FColor UInterpTrackLinearColorBase::GetKeyColor(int32 SubIndex, int32 KeyIndex,
 	check( SubIndex >= 0 && SubIndex < 4);
 	check( KeyIndex >= 0 && KeyIndex < LinearColorTrack.Points.Num() );
 
-	if(SubIndex == 0)
-		return FColor(255,0,0);
-	else if(SubIndex == 1)
-		return FColor(0,255,0);
-	else if(SubIndex == 2)
-		return FColor(0,0,255);
+	if (SubIndex == 0)
+	{
+		return FColor::Red;
+	}
+	else if (SubIndex == 1)
+	{
+		return FColor::Green;
+	}
+	else if (SubIndex == 2)
+	{
+		return FColor::Blue;
+	}
 	else
-		return FColor(255,255,255);
+	{
+		return FColor::White;
+	}
 }
 
 EInterpCurveMode UInterpTrackLinearColorBase::GetKeyInterpMode(int32 KeyIndex) const
diff --git a/Engine/Source/Runtime/Engine/Private/LevelTick.cpp b/Engine/Source/Runtime/Engine/Private/LevelTick.cpp
--- a/Engine/Source/Runtime/Engine/Private/LevelTick.cpp
+++ b/Engine/Source/Runtime/Engine/Private/LevelTick.cpp
@@ -1343,11 +1343,11 @@ void UWorld::Tick( ELevelTick TickType, float DeltaSeconds )
 
 		if( IsGameWorld() && GAreScreenMessagesEnabled && ThisFramePawnSpawns.Num() > GEngine->NumPawnsAllowedToBeSpawnedInAFrame )
 		{
-			GEngine->AddOnScreenDebugMessage((uint64)((PTRINT)this), 5.f, FColor(255,0,0), *WarningMessage);
+			GEngine->AddOnScreenDebugMessage((uint64)((PTRINT)this), 5.f, FColor::Red, *WarningMessage);
 
 			for(int32 i=0; i<ThisFramePawnSpawns.Num(); i++)
 			{
-				GEngine->AddOnScreenDebugMessage((uint64)((PTRINT)+i), 5.f, FColor(255,0,0), *ThisFramePawnSpawns[i]);
+				GEngine->AddOnScreenDebugMessage((uint64)((PTRINT)+i), 5.f, FColor::Red, *ThisFramePawnSpawns[i]);
 			}
 		}
 	}
diff --git a/Engine/Source/Runtime/Engine/Private/LightComponent.cpp b/Engine/Source/Runtime/Engine/Private/LightComponent.cpp
--- a/Engine/Source/Runtime/Engine/Private/LightComponent.cpp
+++ b/Engine/Source/Runtime/Engine/Private/LightComponent.cpp
@@ -293,7 +293,7 @@ ULightComponentBase::ULightComponentBase(const FObjectInitializer& ObjectInitial
 {
 	Brightness_DEPRECATED = 3.1415926535897932f;
 	Intensity = 3.1415926535897932f;
-	LightColor = FColor(255, 255, 255);
+	LightColor = FColor::White;
 	bAffectsWorld = true;
 	CastShadows = true;
 	CastStaticShadows = true;
diff --git a/Engine/Source/Runtime/Engine/Private/LightMap.cpp b/Engine/Source/Runtime/Engine/Private/LightMap.cpp
--- a/Engine/Source/Runtime/Engine/Private/LightMap.cpp
+++ b/Engine/Source/Runtime/Engine/Private/LightMap.cpp
@@ -999,11 +999,11 @@ void FLightMapPendingTexture::StartEncoding()
 							X == Allocation->MappedRect.Min.X+1 || Y == Allocation->MappedRect.Min.Y+1 ||
 							X == Allocation->MappedRect.Max.X-2 || Y == Allocation->MappedRect.Max.Y-2 )
 						{
-							DestColor = FColor(255,0,0);
+							DestColor = FColor::Red;
 						}
 						else
 						{
-							DestColor = FColor(0,255,0);
+							DestColor = FColor::Green;
 						}
 #else
 						DestColor.R = SourceCoefficients.Coefficients[CoefficientIndex][0];
diff --git a/Engine/Source/Runtime/Engine/Private/Particles/ParticleModules.cpp b/Engine/Source/Runtime/Engine/Private/Particles/ParticleModules.cpp
--- a/Engine/Source/Runtime/Engine/Private/Particles/ParticleModules.cpp
+++ b/Engine/Source/Runtime/Engine/Private/Particles/ParticleModules.cpp
@@ -2817,7 +2817,7 @@ void UParticleModuleLight::Render3DPreview(FParticleEmitterInstance* Owner, cons
 				const FVector Size = Scale * Particle.Size;
 				const float LightRadius = LightPayload->RadiusScale * (Size.X + Size.Y) / 2.0f;
 
-				DrawWireSphere(PDI, LightPosition, FColor(255, 255, 255), LightRadius, 18, SDPG_World);
+				DrawWireSphere(PDI, LightPosition, FColor::White, LightRadius, 18, SDPG_World);
 			}
 		}
 	}
diff --git a/Engine/Source/Runtime/Engine/Private/Particles/ParticleSystemRender.cpp b/Engine/Source/Runtime/Engine/Private/Particles/ParticleSystemRender.cpp
--- a/Engine/Source/Runtime/Engine/Private/Particles/ParticleSystemRender.cpp
+++ b/Engine/Source/Runtime/Engine/Private/Particles/ParticleSystemRender.cpp
@@ -2346,9 +2346,9 @@ void FDynamicBeam2EmitterData::RenderDirectLine(const FParticleSystemSceneProxy*
 			continue;
 		}
 
-		DrawWireStar(PDI, BeamPayloadData->SourcePoint, 20.0f, FColor(0,255,0),Proxy->GetDepthPriorityGroup(View));
-		DrawWireStar(PDI, BeamPayloadData->TargetPoint, 20.0f, FColor(255,0,0),Proxy->GetDepthPriorityGroup(View));
-		PDI->DrawLine(BeamPayloadData->SourcePoint, BeamPayloadData->TargetPoint, FColor(255,255,0),Proxy->GetDepthPriorityGroup(View));
+		DrawWireStar(PDI, BeamPayloadData->SourcePoint, 20.0f, FColor::Green, Proxy->GetDepthPriorityGroup(View));
+		DrawWireStar(PDI, BeamPayloadData->TargetPoint, 20.0f, FColor::Red, Proxy->GetDepthPriorityGroup(View));
+		PDI->DrawLine(BeamPayloadData->SourcePoint, BeamPayloadData->TargetPoint, FColor::Yellow, Proxy->GetDepthPriorityGroup(View));
 	}
 }
 
@@ -2518,7 +2518,7 @@ void FDynamicBeam2EmitterData::RenderLines(const FParticleSystemSceneProxy* Prox
 			EndPoint	= CurrDrawPosition;
 
 			// 'Lead' edge
-			DrawWireStar(PDI, Location, 15.0f, FColor(0,255,0), Proxy->GetDepthPriorityGroup(View));
+			DrawWireStar(PDI, Location, 15.0f, FColor::Green, Proxy->GetDepthPriorityGroup(View));
 
 			for (int32 StepIndex = 0; StepIndex < BeamPayloadData->Steps; StepIndex++)
 			{
@@ -2618,12 +2618,12 @@ void FDynamicBeam2EmitterData::RenderLines(const FParticleSystemSceneProxy* Prox
 					FColor StarColor(255,0,255);
 					if (TessIndex == 0)
 					{
-						StarColor = FColor(0,0,255);
+						StarColor = FColor::Blue;
 					}
 					else
 					if (TessIndex == (TessFactor - 1))
 					{
-						StarColor = FColor(255,255,0);
+						StarColor = FColor::Yellow;
 					}
 
 					// Generate the vertex
@@ -2690,12 +2690,11 @@ void FDynamicBeam2EmitterData::RenderLines(const FParticleSystemSceneProxy* Prox
 					FColor StarColor(255,0,255);
 					if (TessIndex == 0)
 					{
-						StarColor = FColor(255,255,255);
+						StarColor = FColor::White;
 					}
-					else
-					if (TessIndex == (TessFactor - 1))
+					else if (TessIndex == (TessFactor - 1))
 					{
-						StarColor = FColor(255,255,0);
+						StarColor = FColor::Yellow;
 					}
 
 					// Generate the vertex
@@ -2728,9 +2727,9 @@ void FDynamicBeam2EmitterData::RenderLines(const FParticleSystemSceneProxy* Prox
 				FVector EndPoint	= Particle->Location;
 				FVector Location	= BeamPayloadData->SourcePoint;
 
-				DrawWireStar(PDI, Location, 15.0f, FColor(255,0,0), Proxy->GetDepthPriorityGroup(View));
-				DrawWireStar(PDI, EndPoint, 15.0f, FColor(255,0,0), Proxy->GetDepthPriorityGroup(View));
-				PDI->DrawLine(Location, EndPoint, FColor(255,255,0), Proxy->GetDepthPriorityGroup(View));
+				DrawWireStar(PDI, Location, 15.0f, FColor::Red, Proxy->GetDepthPriorityGroup(View));
+				DrawWireStar(PDI, EndPoint, 15.0f, FColor::Red, Proxy->GetDepthPriorityGroup(View));
+				PDI->DrawLine(Location, EndPoint, FColor::Yellow, Proxy->GetDepthPriorityGroup(View));
 			}
 		}
 		else
@@ -2760,12 +2759,12 @@ void FDynamicBeam2EmitterData::RenderLines(const FParticleSystemSceneProxy* Prox
 				Location	= BeamPayloadData->SourcePoint;
 				EndPoint	= InterpolatedPoints[0];
 
-				DrawWireStar(PDI, Location, 15.0f, FColor(255,0,0), Proxy->GetDepthPriorityGroup(View));
+				DrawWireStar(PDI, Location, 15.0f, FColor::Red, Proxy->GetDepthPriorityGroup(View));
 				for (int32 StepIndex = 0; StepIndex < BeamPayloadData->InterpolationSteps; StepIndex++)
 				{
 					EndPoint = InterpolatedPoints[StepIndex];
-					DrawWireStar(PDI, EndPoint, 15.0f, FColor(255,0,0), Proxy->GetDepthPriorityGroup(View));
-					PDI->DrawLine(Location, EndPoint, FColor(255,255,0), Proxy->GetDepthPriorityGroup(View));
+					DrawWireStar(PDI, EndPoint, 15.0f, FColor::Red, Proxy->GetDepthPriorityGroup(View));
+					PDI->DrawLine(Location, EndPoint, FColor::Yellow, Proxy->GetDepthPriorityGroup(View));
 					Location = EndPoint;
 				}
 			}
@@ -5258,11 +5257,11 @@ void FDynamicRibbonEmitterData::RenderDebug(const FParticleSystemSceneProxy* Pro
 					{
 						if (TrailPayload->bInterpolatedSpawn == false)
 						{
-							DrawWireStar(PDI, DrawPosition, DrawSize, FColor(255,0,0), Proxy->GetDepthPriorityGroup(View));
+							DrawWireStar(PDI, DrawPosition, DrawSize, FColor::Red, Proxy->GetDepthPriorityGroup(View));
 						}
 						else
 						{
-							DrawWireStar(PDI, DrawPosition, DrawSize, FColor(0,255,0), Proxy->GetDepthPriorityGroup(View));
+							DrawWireStar(PDI, DrawPosition, DrawSize, FColor::Green, Proxy->GetDepthPriorityGroup(View));
 						}
 
 						//
@@ -5272,7 +5271,7 @@ void FDynamicRibbonEmitterData::RenderDebug(const FParticleSystemSceneProxy* Pro
 							{
 								// Draw a straight line between the particles
 								// This will allow us to visualize the tessellation difference
-								PDI->DrawLine(DrawPosition, PrevParticle->Location, FColor(0,0,255), Proxy->GetDepthPriorityGroup(View));
+								PDI->DrawLine(DrawPosition, PrevParticle->Location, FColor::Blue, Proxy->GetDepthPriorityGroup(View));
 								int32 InterpCount = TrailPayload->RenderingInterpCount;
 								// Interpolate between current and next...
 								FVector LineStart = DrawPosition;
@@ -5949,21 +5948,21 @@ void FDynamicAnimTrailEmitterData::RenderDebug(const FParticleSystemSceneProxy*
 				{
 					if (TRAIL_EMITTER_IS_START(RenderData.Payload->Flags))
 					{
-						DrawWireStar(PDI, DrawPosition, DrawSize, FColor(0, 255, 0), Proxy->GetDepthPriorityGroup(View));
-						DrawWireStar(PDI, DrawFirstEdgePosition, DrawSize, FColor(0, 255, 0), Proxy->GetDepthPriorityGroup(View));
-						DrawWireStar(PDI, DrawSecondEdgePosition, DrawSize, FColor(0, 255, 0), Proxy->GetDepthPriorityGroup(View));
+						DrawWireStar(PDI, DrawPosition, DrawSize, FColor::Green, Proxy->GetDepthPriorityGroup(View));
+						DrawWireStar(PDI, DrawFirstEdgePosition, DrawSize, FColor::Green, Proxy->GetDepthPriorityGroup(View));
+						DrawWireStar(PDI, DrawSecondEdgePosition, DrawSize, FColor::Green, Proxy->GetDepthPriorityGroup(View));
 					}
 					else if (TRAIL_EMITTER_IS_DEADTRAIL(RenderData.Payload->Flags))
 					{
-						DrawWireStar(PDI, DrawPosition, DrawSize, FColor(255, 0, 0), Proxy->GetDepthPriorityGroup(View));
-						DrawWireStar(PDI, DrawFirstEdgePosition, DrawSize, FColor(255, 0, 0), Proxy->GetDepthPriorityGroup(View));
-						DrawWireStar(PDI, DrawSecondEdgePosition, DrawSize, FColor(255, 0, 0), Proxy->GetDepthPriorityGroup(View));
+						DrawWireStar(PDI, DrawPosition, DrawSize, FColor::Red, Proxy->GetDepthPriorityGroup(View));
+						DrawWireStar(PDI, DrawFirstEdgePosition, DrawSize, FColor::Red, Proxy->GetDepthPriorityGroup(View));
+						DrawWireStar(PDI, DrawSecondEdgePosition, DrawSize, FColor::Red, Proxy->GetDepthPriorityGroup(View));
 					}
 					else if (TRAIL_EMITTER_IS_END(RenderData.Payload->Flags))
 					{
-						DrawWireStar(PDI, DrawPosition, DrawSize, FColor(255, 255, 255), Proxy->GetDepthPriorityGroup(View));
-						DrawWireStar(PDI, DrawFirstEdgePosition, DrawSize, FColor(255, 255, 255), Proxy->GetDepthPriorityGroup(View));
-						DrawWireStar(PDI, DrawSecondEdgePosition, DrawSize, FColor(255, 255, 255), Proxy->GetDepthPriorityGroup(View));
+						DrawWireStar(PDI, DrawPosition, DrawSize, FColor::White, Proxy->GetDepthPriorityGroup(View));
+						DrawWireStar(PDI, DrawFirstEdgePosition, DrawSize, FColor::White, Proxy->GetDepthPriorityGroup(View));
+						DrawWireStar(PDI, DrawSecondEdgePosition, DrawSize, FColor::White, Proxy->GetDepthPriorityGroup(View));
 					}
 					else
 					{
@@ -5984,9 +5983,9 @@ void FDynamicAnimTrailEmitterData::RenderDebug(const FParticleSystemSceneProxy*
 
 						// Draw a straight line between the particles
 						// This will allow us to visualize the tessellation difference
-						PDI->DrawLine(DrawPosition, PrevDrawPosition, FColor(0,0,255), Proxy->GetDepthPriorityGroup(View));
-						PDI->DrawLine(DrawFirstEdgePosition, PrevDrawFirstEdgePosition, FColor(0,0,255), Proxy->GetDepthPriorityGroup(View));
-						PDI->DrawLine(DrawSecondEdgePosition, PrevDrawSecondEdgePosition, FColor(0,0,255), Proxy->GetDepthPriorityGroup(View));
+						PDI->DrawLine(DrawPosition, PrevDrawPosition, FColor::Blue, Proxy->GetDepthPriorityGroup(View));
+						PDI->DrawLine(DrawFirstEdgePosition, PrevDrawFirstEdgePosition, FColor::Blue, Proxy->GetDepthPriorityGroup(View));
+						PDI->DrawLine(DrawSecondEdgePosition, PrevDrawSecondEdgePosition, FColor::Blue, Proxy->GetDepthPriorityGroup(View));
 
 						int32 InterpCount = RenderData.Payload->RenderingInterpCount;
 						// Interpolate between prev and current...
diff --git a/Engine/Source/Runtime/Engine/Private/PhysicsEngine/PhysDrawing.cpp b/Engine/Source/Runtime/Engine/Private/PhysicsEngine/PhysDrawing.cpp
--- a/Engine/Source/Runtime/Engine/Private/PhysicsEngine/PhysDrawing.cpp
+++ b/Engine/Source/Runtime/Engine/Private/PhysicsEngine/PhysDrawing.cpp
@@ -18,12 +18,12 @@ static const float JointRenderSize = 10.f;
 static const float LimitRenderSize = 16.0f;
 
 static const FColor JointUnselectedColor(255, 0, 255);
-static const FColor JointRed(255, 0, 0);
-static const FColor JointGreen(0, 255, 0);
-static const FColor JointBlue(0, 0, 255);
+static const FColor JointRed(FColor::Red);
+static const FColor JointGreen(FColor::Green);
+static const FColor JointBlue(FColor::Blue);
 
-static const FColor	JointLimitColor(0,255,0);
-static const FColor	JointRefColor(255,255,0);
+static const FColor	JointLimitColor(FColor::Green);
+static const FColor	JointRefColor(FColor::Yellow);
 static const FColor JointLockedColor(255,128,10);
 
 /////////////////////////////////////////////////////////////////////////////////////
@@ -590,7 +590,7 @@ void FKAggregateGeom::GetAggGeom(const FTransform& Transform, const FColor Color
 				for(int32 i=0; i<ConvexElems.Num(); i++)
 				{
 					// Get vertices/triangles from this hull.
-					ConvexElems[i].AddCachedSolidConvexGeom(ThisGeom.RenderInfo->VertexBuffer->Vertices, ThisGeom.RenderInfo->IndexBuffer->Indices, FColor(255,255,255));
+					ConvexElems[i].AddCachedSolidConvexGeom(ThisGeom.RenderInfo->VertexBuffer->Vertices, ThisGeom.RenderInfo->IndexBuffer->Indices, FColor::White);
 				}
 
 				// Only continue if we actually got some valid geometry
diff --git a/Engine/Source/Runtime/Engine/Private/PreviewScene.cpp b/Engine/Source/Runtime/Engine/Private/PreviewScene.cpp
--- a/Engine/Source/Runtime/Engine/Private/PreviewScene.cpp
+++ b/Engine/Source/Runtime/Engine/Private/PreviewScene.cpp
@@ -36,7 +36,7 @@ FPreviewScene::FPreviewScene(FPreviewScene::ConstructionValues CVS)
 
 	DirectionalLight = ConstructObject<UDirectionalLightComponent>(UDirectionalLightComponent::StaticClass());
 	DirectionalLight->Intensity = CVS.LightBrightness;
-	DirectionalLight->LightColor = FColor(255,255,255);
+	DirectionalLight->LightColor = FColor::White;
 	AddComponent(DirectionalLight, FTransform(CVS.LightRotation));
 
 	LineBatcher = ConstructObject<ULineBatchComponent>(ULineBatchComponent::StaticClass());
diff --git a/Engine/Source/Runtime/Engine/Private/PrimitiveSceneProxy.cpp b/Engine/Source/Runtime/Engine/Private/PrimitiveSceneProxy.cpp
--- a/Engine/Source/Runtime/Engine/Private/PrimitiveSceneProxy.cpp
+++ b/Engine/Source/Runtime/Engine/Private/PrimitiveSceneProxy.cpp
@@ -464,9 +464,9 @@ void FPrimitiveSceneProxy::RenderBounds(
 	{
 		// Draw the static mesh's bounding box and sphere.
 		DrawWireBox(PDI,Bounds.GetBox(), FColor(72,72,255),DrawBoundsDPG);
-		DrawCircle(PDI,Bounds.Origin,FVector(1,0,0),FVector(0,1,0),FColor(255,255,0),Bounds.SphereRadius,32,DrawBoundsDPG);
-		DrawCircle(PDI,Bounds.Origin,FVector(1,0,0),FVector(0,0,1),FColor(255,255,0),Bounds.SphereRadius,32,DrawBoundsDPG);
-		DrawCircle(PDI,Bounds.Origin,FVector(0,1,0),FVector(0,0,1),FColor(255,255,0),Bounds.SphereRadius,32,DrawBoundsDPG);
+		DrawCircle(PDI, Bounds.Origin, FVector(1, 0, 0), FVector(0, 1, 0), FColor::Yellow, Bounds.SphereRadius, 32, DrawBoundsDPG);
+		DrawCircle(PDI, Bounds.Origin, FVector(1, 0, 0), FVector(0, 0, 1), FColor::Yellow, Bounds.SphereRadius, 32, DrawBoundsDPG);
+		DrawCircle(PDI, Bounds.Origin, FVector(0, 1, 0), FVector(0, 0, 1), FColor::Yellow, Bounds.SphereRadius, 32, DrawBoundsDPG);
 	}
 }
 
diff --git a/Engine/Source/Runtime/Engine/Private/TextRenderComponent.cpp b/Engine/Source/Runtime/Engine/Private/TextRenderComponent.cpp
--- a/Engine/Source/Runtime/Engine/Private/TextRenderComponent.cpp
+++ b/Engine/Source/Runtime/Engine/Private/TextRenderComponent.cpp
@@ -704,7 +704,7 @@ UTextRenderComponent::UTextRenderComponent(const FObjectInitializer& ObjectIniti
 	TextMaterial = ConstructorStatics.TextMaterial.Get();
 
 	SetCollisionProfileName(UCollisionProfile::NoCollision_ProfileName);
-	TextRenderColor = FColor(255, 255, 255);
+	TextRenderColor = FColor::White;
 	XScale = 1;
 	YScale = 1;
 	HorizSpacingAdjust = 0;
diff --git a/Engine/Source/Runtime/Engine/Private/UnrealClient.cpp b/Engine/Source/Runtime/Engine/Private/UnrealClient.cpp
--- a/Engine/Source/Runtime/Engine/Private/UnrealClient.cpp
+++ b/Engine/Source/Runtime/Engine/Private/UnrealClient.cpp
@@ -370,8 +370,8 @@ int32 FStatUnitData::DrawStat(FViewport* InViewport, FCanvas* InCanvas, int32 In
 	const bool bShowUnitTimeGraph = InViewport->GetClient() ? InViewport->GetClient()->IsStatEnabled(TEXT("UnitGraph")) : false;
 
 	// 0-34 ms: Green, 34-50 ms: Yellow, 50+ ms: Red
-	Color = FrameTime < 34.0f ? FColor(0, 255, 0) : (FrameTime < 50.0f ? FColor(255, 255, 0) : FColor(255, 0, 0));
-	InCanvas->DrawShadowedString(X1, InY, TEXT("Frame:"), Font, bShowUnitTimeGraph ? FColor(100, 255, 100) : FColor(255, 255, 255));
+	Color = FrameTime < 34.0f ? FColor::Green : (FrameTime < 50.0f ? FColor::Yellow : FColor::Red);
+	InCanvas->DrawShadowedString(X1, InY, TEXT("Frame:"), Font, bShowUnitTimeGraph ? FColor(100, 255, 100) : FColor::White);
 	InCanvas->DrawShadowedString(X2, InY, *FString::Printf(TEXT("%3.2f ms"), FrameTime), Font, Color);
 	if (bShowUnitMaxTimes)
 	{
@@ -379,8 +379,8 @@ int32 FStatUnitData::DrawStat(FViewport* InViewport, FCanvas* InCanvas, int32 In
 	}
 	InY += RowHeight;
 
-	Color = GameThreadTime < 34.0f ? FColor(0, 255, 0) : (GameThreadTime < 50.0f ? FColor(255, 255, 0) : FColor(255, 0, 0));
-	InCanvas->DrawShadowedString(X1, InY, TEXT("Game:"), Font, bShowUnitTimeGraph ? FColor(255, 100, 100) : FColor(255, 255, 255));
+	Color = GameThreadTime < 34.0f ? FColor::Green : (GameThreadTime < 50.0f ? FColor::Yellow : FColor::Red);
+	InCanvas->DrawShadowedString(X1, InY, TEXT("Game:"), Font, bShowUnitTimeGraph ? FColor(255, 100, 100) : FColor::White);
 	InCanvas->DrawShadowedString(X2, InY, *FString::Printf(TEXT("%3.2f ms"), GameThreadTime), Font, Color);
 	if (bShowUnitMaxTimes)
 	{
@@ -388,8 +388,8 @@ int32 FStatUnitData::DrawStat(FViewport* InViewport, FCanvas* InCanvas, int32 In
 	}
 	InY += RowHeight;
 
-	Color = RenderThreadTime < 34.0f ? FColor(0, 255, 0) : (RenderThreadTime < 50.0f ? FColor(255, 255, 0) : FColor(255, 0, 0));
-	InCanvas->DrawShadowedString(X1, InY, TEXT("Draw:"), Font, bShowUnitTimeGraph ? FColor(100, 100, 255) : FColor(255, 255, 255));
+	Color = RenderThreadTime < 34.0f ? FColor::Green : (RenderThreadTime < 50.0f ? FColor::Yellow : FColor::Red);
+	InCanvas->DrawShadowedString(X1, InY, TEXT("Draw:"), Font, bShowUnitTimeGraph ? FColor(100, 100, 255) : FColor::White);
 	InCanvas->DrawShadowedString(X2, InY, *FString::Printf(TEXT("%3.2f ms"), RenderThreadTime), Font, Color);
 	if (bShowUnitMaxTimes)
 	{
@@ -400,8 +400,8 @@ int32 FStatUnitData::DrawStat(FViewport* InViewport, FCanvas* InCanvas, int32 In
 	const bool bHaveGPUData = GPUCycles > 0;
 	if (bHaveGPUData)
 	{
-		Color = GPUFrameTime < 34.0f ? FColor(0, 255, 0) : (GPUFrameTime < 50.0f ? FColor(255, 255, 0) : FColor(255, 0, 0));
-		InCanvas->DrawShadowedString(X1, InY, TEXT("GPU:"), Font, bShowUnitTimeGraph ? FColor(255, 255, 100) : FColor(255, 255, 255));
+		Color = GPUFrameTime < 34.0f ? FColor::Green : (GPUFrameTime < 50.0f ? FColor::Yellow : FColor::Red);
+		InCanvas->DrawShadowedString(X1, InY, TEXT("GPU:"), Font, bShowUnitTimeGraph ? FColor(255, 255, 100) : FColor::White);
 		InCanvas->DrawShadowedString(X2, InY, *FString::Printf(TEXT("%3.2f ms"), GPUFrameTime), Font, Color);
 		if (bShowUnitMaxTimes)
 		{
@@ -629,9 +629,9 @@ int32 FStatHitchesData::DrawStat(FViewport* InViewport, FCanvas* InCanvas, int32
 		{
 			if (When[i] > 0 && When[i] <= CurrentTime && When[i] >= CurrentTime - TravelTime)
 			{
-				FColor MyColor = FColor(0, 255, 0);
-				if (Hitches[i] > 0.2f) MyColor = FColor(255, 255, 0);
-				if (Hitches[i] > 0.3f) MyColor = FColor(255, 0, 0);
+				FColor MyColor = FColor::Green;
+				if (Hitches[i] > 0.2f) MyColor = FColor::Yellow;
+				if (Hitches[i] > 0.3f) MyColor = FColor::Red;
 				int32 MyY = InY + int32(float(MaxY - InY) * float((CurrentTime - When[i]) / TravelTime));
 				FString Hitch = FString::Printf(TEXT("%5d"), int32(Hitches[i] * 1000.0f));
 				InCanvas->DrawShadowedString(InX, MyY, *Hitch, GEngine->GetSmallFont(), MyColor);
diff --git a/Engine/Source/Runtime/Engine/Private/UnrealEngine.cpp b/Engine/Source/Runtime/Engine/Private/UnrealEngine.cpp
--- a/Engine/Source/Runtime/Engine/Private/UnrealEngine.cpp
+++ b/Engine/Source/Runtime/Engine/Private/UnrealEngine.cpp
@@ -10595,29 +10595,29 @@ void UEngine::HandleScreenshotCaptured(int32 Width, int32 Height, const TArray<F
 /** Utility that gets a color for a particular level status */
 FColor GetColorForLevelStatus(int32 Status)
 {
-	FColor Color = FColor(255, 255, 255);
+	FColor Color = FColor::White;
 	switch (Status)
 	{
 	case LEVEL_Visible:
-		Color = FColor(255, 0, 0);	// red  loaded and visible
+		Color = FColor::Red;		// red  loaded and visible
 		break;
 	case LEVEL_MakingVisible:
-		Color = FColor(255, 128, 0);	// orange, in process of being made visible
+		Color = FColor(255, 128, 0);// orange, in process of being made visible
 		break;
 	case LEVEL_Loading:
-		Color = FColor(255, 0, 255);	// purple, in process of being loaded
+		Color = FColor::Magenta;	// purple, in process of being loaded
 		break;
 	case LEVEL_Loaded:
-		Color = FColor(255, 255, 0);	// yellow loaded but not visible
+		Color = FColor::Yellow;		// yellow loaded but not visible
 		break;
 	case LEVEL_UnloadedButStillAround:
-		Color = FColor(0, 0, 255);	// blue  (GC needs to occur to remove this)
+		Color = FColor::Blue;		// blue  (GC needs to occur to remove this)
 		break;
 	case LEVEL_Unloaded:
-		Color = FColor(0, 255, 0);	// green
+		Color = FColor::Green;		// green
 		break;
 	case LEVEL_Preloading:
-		Color = FColor(255, 0, 255);	// purple (preloading)
+		Color = FColor::Magenta;	// purple (preloading)
 		break;
 	default:
 		break;
@@ -10759,7 +10759,7 @@ int32 UEngine::RenderStatFPS(UWorld* World, FViewport* Viewport, FCanvas* Canvas
 	UFont* Font = FPlatformProperties::SupportsWindowedMode() ? GetSmallFont() : GetMediumFont();
 
 	// Choose the counter color based on the average framerate.
-	FColor FPSColor = GAverageFPS < 20.0f ? FColor(255, 0, 0) : (GAverageFPS < 29.5f ? FColor(255, 255, 0) : FColor(0, 255, 0));
+	FColor FPSColor = GAverageFPS < 20.0f ? FColor::Red : (GAverageFPS < 29.5f ? FColor::Yellow : FColor::Green);
 
 	// Start drawing the various counters.
 	const int32 RowHeight = FMath::TruncToInt(Font->GetMaxCharHeight() * 1.1f);
@@ -11290,7 +11290,7 @@ int32 UEngine::RenderStatSoundMixes(UWorld* World, FViewport* Viewport, FCanvas*
 	FAudioDevice* AudioDevice = GetAudioDevice();
 	if (AudioDevice)
 	{
-		Canvas->DrawShadowedString(X, Y, TEXT("Active Sound Mixes:"), GetSmallFont(), FColor(0, 255, 0));
+		Canvas->DrawShadowedString(X, Y, TEXT("Active Sound Mixes:"), GetSmallFont(), FColor::Green);
 		Y += 12;
 
 		if (AudioDevice->SoundMixModifiers.Num() > 0)
@@ -11302,7 +11302,7 @@ int32 UEngine::RenderStatSoundMixes(UWorld* World, FViewport* Viewport, FCanvas*
 				uint32 TotalRefCount = It.Value().ActiveRefCount + It.Value().PassiveRefCount;
 				FString TheString = FString::Printf(TEXT("%s - Fade Proportion: %1.2f - Total Ref Count: %i"), *It.Key()->GetName(), It.Value().InterpValue, TotalRefCount);
 
-				FColor TextColour = FColor(255, 255, 255);
+				FColor TextColour = FColor::White;
 				if (It.Key() == CurrentEQMix)
 				{
 					TextColour = FColor(255, 255, 0);
@@ -11315,7 +11315,7 @@ int32 UEngine::RenderStatSoundMixes(UWorld* World, FViewport* Viewport, FCanvas*
 		}
 		else
 		{
-			Canvas->DrawShadowedString(X + 12, Y, TEXT("None"), GetSmallFont(), FColor(255, 255, 255));
+			Canvas->DrawShadowedString(X + 12, Y, TEXT("None"), GetSmallFont(), FColor::White);
 			Y += 12;
 		}
 	}
@@ -11352,7 +11352,7 @@ int32 UEngine::RenderStatSoundWaves(UWorld* World, FViewport* Viewport, FCanvas*
 				SoundOwner ? *SoundOwner->GetName() : TEXT("None"),
 				SoundClass ? *SoundClass->GetName() : TEXT("None"));
 
-			Canvas->DrawShadowedString(X, Y, *TheString, GetSmallFont(), FColor(255, 255, 255));
+			Canvas->DrawShadowedString(X, Y, *TheString, GetSmallFont(), FColor::White);
 			Y += 12;
 		}
 
@@ -11397,7 +11397,7 @@ int32 UEngine::RenderStatSoundCues(UWorld* World, FViewport* Viewport, FCanvas*
 		}
 	}
 
-	Canvas->DrawShadowedString(X, Y, TEXT("Active Sound Cues:"), GetSmallFont(), FColor(0, 255, 0));
+	Canvas->DrawShadowedString(X, Y, TEXT("Active Sound Cues:"), GetSmallFont(), FColor::Green);
 	Y += 12;
 
 	int32 ActiveSoundCount = 0;
@@ -11405,11 +11405,11 @@ int32 UEngine::RenderStatSoundCues(UWorld* World, FViewport* Viewport, FCanvas*
 	{
 		USoundClass* SoundClass = ActiveSound->GetSoundClass();
 		const FString TheString = FString::Printf(TEXT("%4i. %s %s"), ActiveSoundCount++, *ActiveSound->Sound->GetPathName(), (SoundClass ? *SoundClass->GetName() : TEXT("None")));
-		Canvas->DrawShadowedString(X, Y, *TheString, GetSmallFont(), FColor(255, 255, 255));
+		Canvas->DrawShadowedString(X, Y, *TheString, GetSmallFont(), FColor::White);
 		Y += 12;
 	}
 
-	Canvas->DrawShadowedString(X, Y, *FString::Printf(TEXT("Total: %i"), ActiveSounds.Num()), GetSmallFont(), FColor(0, 255, 0));
+	Canvas->DrawShadowedString(X, Y, *FString::Printf(TEXT("Total: %i"), ActiveSounds.Num()), GetSmallFont(), FColor::Green);
 	Y += 12;
 	return Y;
 }
@@ -11572,14 +11572,14 @@ int32 UEngine::RenderStatSounds(UWorld* World, FViewport* Viewport, FCanvas* Can
 		}
 
 
-		Canvas->DrawShadowedString(X, Y, TEXT("Active Sounds:"), GetSmallFont(), FColor(0, 255, 0));
+		Canvas->DrawShadowedString(X, Y, TEXT("Active Sounds:"), GetSmallFont(), FColor::Green);
 		Y += 12;
 
 		const FString InfoText = FString::Printf(TEXT(" Sorting: %s Debug: %s"), *SortingName, bDebug ? TEXT("enabled") : TEXT("disabled"));
 		Canvas->DrawShadowedString(X, Y, *InfoText, GetSmallFont(), FColor(128, 255, 128));
 		Y += 12;
 
-		Canvas->DrawShadowedString(X, Y, TEXT("Index Path (Class) Distance"), GetSmallFont(), FColor(0, 255, 0));
+		Canvas->DrawShadowedString(X, Y, TEXT("Index Path (Class) Distance"), GetSmallFont(), FColor::Green);
 		Y += 12;
 
 		int32 TotalSoundWavesNum = 0;
@@ -11593,7 +11593,7 @@ int32 UEngine::RenderStatSounds(UWorld* World, FViewport* Viewport, FCanvas* Can
 			{
 				{
 					const FString TheString = FString::Printf(TEXT("%4i. %s (%s) %6.2f"), SoundIndex, *SoundInfo.PathName, *SoundInfo.ClassName.ToString(), SoundInfo.Distance);
-					Canvas->DrawShadowedString(X, Y, *TheString, GetSmallFont(), FColor(255, 255, 255));
+					Canvas->DrawShadowedString(X, Y, *TheString, GetSmallFont(), FColor::White);
 					Y += 12;
 				}
 
@@ -11615,10 +11615,10 @@ int32 UEngine::RenderStatSounds(UWorld* World, FViewport* Viewport, FCanvas* Can
 			}
 		}
 
-		Canvas->DrawShadowedString(X, Y, *FString::Printf(TEXT("Total sounds: %i, sound waves: %i"), SoundIndex, TotalSoundWavesNum), GetSmallFont(), FColor(0, 255, 0));
+		Canvas->DrawShadowedString(X, Y, *FString::Printf(TEXT("Total sounds: %i, sound waves: %i"), SoundIndex, TotalSoundWavesNum), GetSmallFont(), FColor::Green);
 		Y += 12;
 
-		Canvas->DrawShadowedString(X, Y, *FString::Printf(TEXT("Listener position: %s"), *ListenerPosition.ToString()), GetSmallFont(), FColor(0, 255, 0));
+		Canvas->DrawShadowedString(X, Y, *FString::Printf(TEXT("Listener position: %s"), *ListenerPosition.ToString()), GetSmallFont(), FColor::Green);
 		Y += 12;
 
 		// Draw sound cue's sphere.
@@ -11746,7 +11746,7 @@ int32 UEngine::RenderStatAI(UWorld* World, FViewport* Viewport, FCanvas* Canvas,
 
 #define MAXDUDES 20
 #define BADAMTOFDUDES 12
-	FColor TotalColor = FColor(0, 255, 0);
+	FColor TotalColor = FColor::Green;
 	if (NumAI > BADAMTOFDUDES)
 	{
 		float Scalar = 1.0f - FMath::Clamp<float>((float)NumAI / (float)MAXDUDES, 0.f, 1.f);
@@ -11754,7 +11754,7 @@ int32 UEngine::RenderStatAI(UWorld* World, FViewport* Viewport, FCanvas* Canvas,
 		TotalColor = FColor::MakeRedToGreenColorFromScalar(Scalar);
 	}
 
-	FColor RenderedColor = FColor(0, 255, 0);
+	FColor RenderedColor = FColor::Green;
 	if (NumAIRendered > BADAMTOFDUDES)
 	{
 		float Scalar = 1.0f - FMath::Clamp<float>((float)NumAIRendered / (float)MAXDUDES, 0.f, 1.f);
@@ -11803,7 +11803,7 @@ int32 UEngine::RenderStatSlateBatches(UWorld* World, FViewport* Viewport, FCanva
 		Y,
 		TEXT("Slate Batches:"),
 		Font, 
-		FColor(0,255,0) );
+		FColor::Green );
 	
 	Y+=RowHeight;
 	
@@ -11821,7 +11821,7 @@ int32 UEngine::RenderStatSlateBatches(UWorld* World, FViewport* Viewport, FCanva
 			Y,
 			*FString::Printf(TEXT("Layer: %d, Elements: %d, Vertices: %d"), Stat.Layer, Stat.NumElementsInBatch, Stat.NumVertices ), 
 			Font,
-			FColor(0,255,0) );
+			FColor::Green );
 		Y += RowHeight;
 	}*/
 	return Y;
diff --git a/Engine/Source/Runtime/Engine/Private/Vehicles/WheeledVehicleMovementComponent.cpp b/Engine/Source/Runtime/Engine/Private/Vehicles/WheeledVehicleMovementComponent.cpp
--- a/Engine/Source/Runtime/Engine/Private/Vehicles/WheeledVehicleMovementComponent.cpp
+++ b/Engine/Source/Runtime/Engine/Private/Vehicles/WheeledVehicleMovementComponent.cpp
@@ -1405,9 +1405,9 @@ void UWheeledVehicleMovementComponent::DrawDebugLines()
 	const PxTransform T = GlobalT.transform( PActor->getCMassLocalPose() );
 	const PxVec3 ChassisExtent = PActor->getWorldBounds().getExtents();
 	const float ChassisSize = ChassisExtent.magnitude();
-	DrawDebugLine( World, P2UVector(T.p), P2UVector( T.p + T.rotate( PxVec3( ChassisSize, 0, 0 ) ) ), FColor(255,0,0) );
-	DrawDebugLine( World, P2UVector(T.p), P2UVector( T.p + T.rotate( PxVec3( 0, ChassisSize, 0 ) ) ), FColor(0,255,0) );
-	DrawDebugLine( World, P2UVector(T.p), P2UVector( T.p + T.rotate( PxVec3( 0, 0, ChassisSize ) ) ), FColor(0,0,255) );
+	DrawDebugLine(World, P2UVector(T.p), P2UVector(T.p + T.rotate(PxVec3(ChassisSize, 0, 0))), FColor::Red);
+	DrawDebugLine(World, P2UVector(T.p), P2UVector(T.p + T.rotate(PxVec3(0, ChassisSize, 0))), FColor::Green);
+	DrawDebugLine(World, P2UVector(T.p), P2UVector(T.p + T.rotate(PxVec3(0, 0, ChassisSize))), FColor::Blue);
 
 	PxVehicleTelemetryData* TelemetryData = MyVehicleManager->GetTelemetryData();
 	
@@ -1525,7 +1525,7 @@ void UWheeledVehicleMovementComponent::CalculateAvoidanceVelocity(float DeltaTim
 #if !(UE_BUILD_SHIPPING || UE_BUILD_TEST)
 			if (bShowDebug)
 			{
-				DrawDebugLine(GetWorld(), GetRVOAvoidanceOrigin(), GetRVOAvoidanceOrigin() + AvoidanceVelocity, FColor(0, 0, 255), true, 0.5f, SDPG_MAX);
+				DrawDebugLine(GetWorld(), GetRVOAvoidanceOrigin(), GetRVOAvoidanceOrigin() + AvoidanceVelocity, FColor::Blue, true, 0.5f, SDPG_MAX);
 			}
 #endif
 		}
@@ -1540,7 +1540,7 @@ void UWheeledVehicleMovementComponent::CalculateAvoidanceVelocity(float DeltaTim
 #if !(UE_BUILD_SHIPPING || UE_BUILD_TEST)
 				if (bShowDebug)
 				{
-					DrawDebugLine(GetWorld(), GetRVOAvoidanceOrigin(), GetRVOAvoidanceOrigin() + AvoidanceVelocity, FColor(255, 0, 0), true, 20.0f, SDPG_MAX, 10.0f);
+					DrawDebugLine(GetWorld(), GetRVOAvoidanceOrigin(), GetRVOAvoidanceOrigin() + AvoidanceVelocity, FColor::Red, true, 20.0f, SDPG_MAX, 10.0f);
 				}
 #endif
 			}
@@ -1557,13 +1557,13 @@ void UWheeledVehicleMovementComponent::CalculateAvoidanceVelocity(float DeltaTim
 #if !(UE_BUILD_SHIPPING || UE_BUILD_TEST)
 	else if (bShowDebug)
 	{
-		DrawDebugLine(GetWorld(), GetRVOAvoidanceOrigin(), GetRVOAvoidanceOrigin() + GetVelocityForRVOConsideration(), FColor(255, 255, 0), true, 0.05f, SDPG_MAX);
+		DrawDebugLine(GetWorld(), GetRVOAvoidanceOrigin(), GetRVOAvoidanceOrigin() + GetVelocityForRVOConsideration(), FColor::Yellow, true, 0.05f, SDPG_MAX);
 	}
 
 	if (bShowDebug)
 	{
 		FVector UpLine(0, 0, 500);
-		DrawDebugLine(GetWorld(), GetRVOAvoidanceOrigin(), GetRVOAvoidanceOrigin() + UpLine, (AvoidanceLockTimer > 0.01f) ? FColor(255, 0, 0) : FColor(0, 0, 255), true, 0.05f, SDPG_MAX, 5.0f);
+		DrawDebugLine(GetWorld(), GetRVOAvoidanceOrigin(), GetRVOAvoidanceOrigin() + UpLine, (AvoidanceLockTimer > 0.01f) ? FColor::Red : FColor::Blue, true, 0.05f, SDPG_MAX, 5.0f);
 	}
 #endif
 }
diff --git a/Engine/Source/Runtime/GameplayAbilities/Private/AbilitySystemComponent.cpp b/Engine/Source/Runtime/GameplayAbilities/Private/AbilitySystemComponent.cpp
--- a/Engine/Source/Runtime/GameplayAbilities/Private/AbilitySystemComponent.cpp
+++ b/Engine/Source/Runtime/GameplayAbilities/Private/AbilitySystemComponent.cpp
@@ -685,7 +685,7 @@ void UAbilitySystemComponent::InvokeGameplayCueEvent(const FGameplayEffectSpec &
 
 		if (DebugGameplayCues && Spec.GetContext().GetHitResult())
 		{
-			DrawDebugSphere(GetWorld(), Spec.GetContext().GetHitResult()->Location, 30.f, 32, FColor(255, 0, 0), true, 30.f);
+			DrawDebugSphere(GetWorld(), Spec.GetContext().GetHitResult()->Location, 30.f, 32, FColor::Red, true, 30.f);
 			ABILITY_LOG(Warning, TEXT("   %s"), *CueInfo.GameplayCueTags.ToString());
 		}
 	}
diff --git a/Engine/Source/Runtime/Landscape/Private/LandscapeGizmoActor.cpp b/Engine/Source/Runtime/Landscape/Private/LandscapeGizmoActor.cpp
--- a/Engine/Source/Runtime/Landscape/Private/LandscapeGizmoActor.cpp
+++ b/Engine/Source/Runtime/Landscape/Private/LandscapeGizmoActor.cpp
@@ -193,7 +193,7 @@ public:
 
 						SampledPositions.Add(SampledPos);
 						SampledNormals.Add(SampledNormal);
-						//MeshBuilder.AddVertex(SampledPos, FVector2D((float)X / (Gizmo->SampleSizeX), (float)Y / (Gizmo->SampleSizeY)), TangentX, SampledNormal^TangentX, SampledNormal, FColor(255, 255, 255) );
+						//MeshBuilder.AddVertex(SampledPos, FVector2D((float)X / (Gizmo->SampleSizeX), (float)Y / (Gizmo->SampleSizeY)), TangentX, SampledNormal^TangentX, SampledNormal, FColor::White );
 					}
 				}
 			}
@@ -227,35 +227,36 @@ public:
 					{
 						FDynamicMeshBuilder MeshBuilder;
 
-						MeshBuilder.AddVertex(FrustumVerts[0], FVector2D(0, 0), FVector(1,0,0), FVector(0,1,0), FVector(0,0,1), FColor(255,255,255));
-						MeshBuilder.AddVertex(FrustumVerts[1], FVector2D(1, 0), FVector(1,0,0), FVector(0,1,0), FVector(0,0,1), FColor(255,255,255));
-						MeshBuilder.AddVertex(FrustumVerts[2], FVector2D(1, 1), FVector(1,0,0), FVector(0,1,0), FVector(0,0,1), FColor(255,255,255));
-						MeshBuilder.AddVertex(FrustumVerts[3], FVector2D(0, 1), FVector(1,0,0), FVector(0,1,0), FVector(0,0,1), FColor(255,255,255));
-
-						MeshBuilder.AddVertex(FrustumVerts[4], FVector2D(0, 0), FVector(1,0,0), FVector(0,1,0), FVector(0,0,1), FColor(255,255,255));
-						MeshBuilder.AddVertex(FrustumVerts[5], FVector2D(1, 0), FVector(1,0,0), FVector(0,1,0), FVector(0,0,1), FColor(255,255,255));
-						MeshBuilder.AddVertex(FrustumVerts[6], FVector2D(1, 1), FVector(1,0,0), FVector(0,1,0), FVector(0,0,1), FColor(255,255,255));
-						MeshBuilder.AddVertex(FrustumVerts[7], FVector2D(0, 1), FVector(1,0,0), FVector(0,1,0), FVector(0,0,1), FColor(255,255,255));
-
-						MeshBuilder.AddVertex(FrustumVerts[1], FVector2D(0, 0), FVector(1,0,0), FVector(0,1,0), FVector(0,0,1), FColor(255,255,255));
-						MeshBuilder.AddVertex(FrustumVerts[0], FVector2D(1, 0), FVector(1,0,0), FVector(0,1,0), FVector(0,0,1), FColor(255,255,255));
-						MeshBuilder.AddVertex(FrustumVerts[4], FVector2D(1, 1), FVector(1,0,0), FVector(0,1,0), FVector(0,0,1), FColor(255,255,255));
-						MeshBuilder.AddVertex(FrustumVerts[5], FVector2D(0, 1), FVector(1,0,0), FVector(0,1,0), FVector(0,0,1), FColor(255,255,255));
-
-						MeshBuilder.AddVertex(FrustumVerts[3], FVector2D(0, 0), FVector(1,0,0), FVector(0,1,0), FVector(0,0,1), FColor(255,255,255));
-						MeshBuilder.AddVertex(FrustumVerts[2], FVector2D(1, 0), FVector(1,0,0), FVector(0,1,0), FVector(0,0,1), FColor(255,255,255));
-						MeshBuilder.AddVertex(FrustumVerts[6], FVector2D(1, 1), FVector(1,0,0), FVector(0,1,0), FVector(0,0,1), FColor(255,255,255));
-						MeshBuilder.AddVertex(FrustumVerts[7], FVector2D(0, 1), FVector(1,0,0), FVector(0,1,0), FVector(0,0,1), FColor(255,255,255));
-
-						MeshBuilder.AddVertex(FrustumVerts[2], FVector2D(0, 0), FVector(1,0,0), FVector(0,1,0), FVector(0,0,1), FColor(255,255,255));
-						MeshBuilder.AddVertex(FrustumVerts[1], FVector2D(1, 0), FVector(1,0,0), FVector(0,1,0), FVector(0,0,1), FColor(255,255,255));
-						MeshBuilder.AddVertex(FrustumVerts[5], FVector2D(1, 1), FVector(1,0,0), FVector(0,1,0), FVector(0,0,1), FColor(255,255,255));
-						MeshBuilder.AddVertex(FrustumVerts[6], FVector2D(0, 1), FVector(1,0,0), FVector(0,1,0), FVector(0,0,1), FColor(255,255,255));
-
-						MeshBuilder.AddVertex(FrustumVerts[0], FVector2D(0, 0), FVector(1,0,0), FVector(0,1,0), FVector(0,0,1), FColor(255,255,255));
-						MeshBuilder.AddVertex(FrustumVerts[3], FVector2D(1, 0), FVector(1,0,0), FVector(0,1,0), FVector(0,0,1), FColor(255,255,255));
-						MeshBuilder.AddVertex(FrustumVerts[7], FVector2D(1, 1), FVector(1,0,0), FVector(0,1,0), FVector(0,0,1), FColor(255,255,255));
-						MeshBuilder.AddVertex(FrustumVerts[4], FVector2D(0, 1), FVector(1,0,0), FVector(0,1,0), FVector(0,0,1), FColor(255,255,255));
+						const FColor GizmoColor = FColor::White;
+						MeshBuilder.AddVertex(FrustumVerts[0], FVector2D(0, 0), FVector(1,0,0), FVector(0,1,0), FVector(0,0,1), GizmoColor);
+						MeshBuilder.AddVertex(FrustumVerts[1], FVector2D(1, 0), FVector(1,0,0), FVector(0,1,0), FVector(0,0,1), GizmoColor);
+						MeshBuilder.AddVertex(FrustumVerts[2], FVector2D(1, 1), FVector(1,0,0), FVector(0,1,0), FVector(0,0,1), GizmoColor);
+						MeshBuilder.AddVertex(FrustumVerts[3], FVector2D(0, 1), FVector(1,0,0), FVector(0,1,0), FVector(0,0,1), GizmoColor);
+
+						MeshBuilder.AddVertex(FrustumVerts[4], FVector2D(0, 0), FVector(1,0,0), FVector(0,1,0), FVector(0,0,1), GizmoColor);
+						MeshBuilder.AddVertex(FrustumVerts[5], FVector2D(1, 0), FVector(1,0,0), FVector(0,1,0), FVector(0,0,1), GizmoColor);
+						MeshBuilder.AddVertex(FrustumVerts[6], FVector2D(1, 1), FVector(1,0,0), FVector(0,1,0), FVector(0,0,1), GizmoColor);
+						MeshBuilder.AddVertex(FrustumVerts[7], FVector2D(0, 1), FVector(1,0,0), FVector(0,1,0), FVector(0,0,1), GizmoColor);
+
+						MeshBuilder.AddVertex(FrustumVerts[1], FVector2D(0, 0), FVector(1,0,0), FVector(0,1,0), FVector(0,0,1), GizmoColor);
+						MeshBuilder.AddVertex(FrustumVerts[0], FVector2D(1, 0), FVector(1,0,0), FVector(0,1,0), FVector(0,0,1), GizmoColor);
+						MeshBuilder.AddVertex(FrustumVerts[4], FVector2D(1, 1), FVector(1,0,0), FVector(0,1,0), FVector(0,0,1), GizmoColor);
+						MeshBuilder.AddVertex(FrustumVerts[5], FVector2D(0, 1), FVector(1,0,0), FVector(0,1,0), FVector(0,0,1), GizmoColor);
+
+						MeshBuilder.AddVertex(FrustumVerts[3], FVector2D(0, 0), FVector(1,0,0), FVector(0,1,0), FVector(0,0,1), GizmoColor);
+						MeshBuilder.AddVertex(FrustumVerts[2], FVector2D(1, 0), FVector(1,0,0), FVector(0,1,0), FVector(0,0,1), GizmoColor);
+						MeshBuilder.AddVertex(FrustumVerts[6], FVector2D(1, 1), FVector(1,0,0), FVector(0,1,0), FVector(0,0,1), GizmoColor);
+						MeshBuilder.AddVertex(FrustumVerts[7], FVector2D(0, 1), FVector(1,0,0), FVector(0,1,0), FVector(0,0,1), GizmoColor);
+
+						MeshBuilder.AddVertex(FrustumVerts[2], FVector2D(0, 0), FVector(1,0,0), FVector(0,1,0), FVector(0,0,1), GizmoColor);
+						MeshBuilder.AddVertex(FrustumVerts[1], FVector2D(1, 0), FVector(1,0,0), FVector(0,1,0), FVector(0,0,1), GizmoColor);
+						MeshBuilder.AddVertex(FrustumVerts[5], FVector2D(1, 1), FVector(1,0,0), FVector(0,1,0), FVector(0,0,1), GizmoColor);
+						MeshBuilder.AddVertex(FrustumVerts[6], FVector2D(0, 1), FVector(1,0,0), FVector(0,1,0), FVector(0,0,1), GizmoColor);
+
+						MeshBuilder.AddVertex(FrustumVerts[0], FVector2D(0, 0), FVector(1,0,0), FVector(0,1,0), FVector(0,0,1), GizmoColor);
+						MeshBuilder.AddVertex(FrustumVerts[3], FVector2D(1, 0), FVector(1,0,0), FVector(0,1,0), FVector(0,0,1), GizmoColor);
+						MeshBuilder.AddVertex(FrustumVerts[7], FVector2D(1, 1), FVector(1,0,0), FVector(0,1,0), FVector(0,0,1), GizmoColor);
+						MeshBuilder.AddVertex(FrustumVerts[4], FVector2D(0, 1), FVector(1,0,0), FVector(0,1,0), FVector(0,0,1), GizmoColor);
 
 						for (int32 i = 0; i < 6; ++i)
 						{
@@ -279,7 +280,7 @@ public:
 								FVector TangentX(SampledNormal.Z, 0, -SampledNormal.X);
 								TangentX = TangentX.GetSafeNormal();
 
-								MeshBuilder.AddVertex(SampledPositions[X + Y * SampleSizeX], FVector2D((float)X / (SampleSizeX), (float)Y / (SampleSizeY)), TangentX, SampledNormal^TangentX, SampledNormal, FColor(255, 255, 255) );
+								MeshBuilder.AddVertex(SampledPositions[X + Y * SampleSizeX], FVector2D((float)X / (SampleSizeX), (float)Y / (SampleSizeY)), TangentX, SampledNormal^TangentX, SampledNormal, FColor::White);
 							}
 						}
 
diff --git a/Engine/Source/Runtime/RenderCore/Private/RenderUtils.cpp b/Engine/Source/Runtime/RenderCore/Private/RenderUtils.cpp
--- a/Engine/Source/Runtime/RenderCore/Private/RenderUtils.cpp
+++ b/Engine/Source/Runtime/RenderCore/Private/RenderUtils.cpp
@@ -461,7 +461,7 @@ private:
 class FWhiteTextureCube : public FSolidColorTextureCube
 {
 public:
-	FWhiteTextureCube(): FSolidColorTextureCube(FColor(255,255,255)) {}
+	FWhiteTextureCube() : FSolidColorTextureCube(FColor::White) {}
 };
 FTexture* GWhiteTextureCube = new TGlobalResource<FWhiteTextureCube>;
 
@@ -469,7 +469,7 @@ FTexture* GWhiteTextureCube = new TGlobalResource<FWhiteTextureCube>;
 class FBlackTextureCube : public FSolidColorTextureCube
 {
 public:
-	FBlackTextureCube(): FSolidColorTextureCube(FColor(0,0,0)) {}
+	FBlackTextureCube(): FSolidColorTextureCube(FColor::Black) {}
 };
 FTexture* GBlackTextureCube = new TGlobalResource<FBlackTextureCube>;
 
diff --git a/Engine/Source/Runtime/Renderer/Private/ShadowRendering.cpp b/Engine/Source/Runtime/Renderer/Private/ShadowRendering.cpp
--- a/Engine/Source/Runtime/Renderer/Private/ShadowRendering.cpp
+++ b/Engine/Source/Runtime/Renderer/Private/ShadowRendering.cpp
@@ -1165,7 +1165,7 @@ void FProjectedShadowInfo::ClearDepth(FRHICommandList& RHICmdList, FDeferredShad
 			);
 
 		// Clear depth only.
-		RHICmdList.Clear(false, FColor(255, 255, 255), true, 1.0f, false, 0, FIntRect());
+		RHICmdList.Clear(false, FColor::White, true, 1.0f, false, 0, FIntRect());
 	}
 	else
 	{
@@ -1198,7 +1198,7 @@ void FProjectedShadowInfo::ClearDepth(FRHICommandList& RHICmdList, FDeferredShad
 				);
 
 			// Clear depth only.
-			RHICmdList.Clear(false, FColor(255, 255, 255), true, 1.0f, false, 0, FIntRect());
+			RHICmdList.Clear(false, FColor::White, true, 1.0f, false, 0, FIntRect());
 		}
 	}
 }
diff --git a/Engine/Source/Runtime/SlateCore/Private/Rendering/ElementBatcher.cpp b/Engine/Source/Runtime/SlateCore/Private/Rendering/ElementBatcher.cpp
--- a/Engine/Source/Runtime/SlateCore/Private/Rendering/ElementBatcher.cpp
+++ b/Engine/Source/Runtime/SlateCore/Private/Rendering/ElementBatcher.cpp
@@ -1011,7 +1011,7 @@ void FSlateElementBatcher::AddLineElement( const FSlateDrawElement& DrawElement
 					IntersectUpper = IntersectionPoint;
 
 					// visualizes the intersection
-					//AddQuadElement( IntersectUpper-FVector2D(1,1), FVector2D(2,2), 1, InClippingRect, Layer+1, FColor(255,0,0));
+					//AddQuadElement( IntersectUpper-FVector2D(1,1), FVector2D(2,2), 1, InClippingRect, Layer+1, FColor::Red);
 
 				}
 
@@ -1021,7 +1021,7 @@ void FSlateElementBatcher::AddLineElement( const FSlateDrawElement& DrawElement
 					IntersectLower = IntersectionPoint;
 
 					// visualizes the intersection
-					//AddQuadElement( IntersectLower-FVector2D(1,1), FVector2D(2,2), 1, InClippingRect, Layer+1, FColor(255,255,0));
+					//AddQuadElement( IntersectLower-FVector2D(1,1), FVector2D(2,2), 1, InClippingRect, Layer+1, FColor::Yellow);
 				}
 				// the midpoint of the intersection.  Used as the new end to the line segment (not adjusted for anti-aliasing)
 				IntersectCenter = (IntersectUpper+IntersectLower) * .5f;
