Don't send keep-alive on connections that we've closed

diff --git a/Engine/Source/Runtime/Engine/Private/NetConnection.cpp b/Engine/Source/Runtime/Engine/Private/NetConnection.cpp
--- a/Engine/Source/Runtime/Engine/Private/NetConnection.cpp
+++ b/Engine/Source/Runtime/Engine/Private/NetConnection.cpp
@@ -477,7 +477,7 @@ void UNetConnection::FlushNet(bool bIgnoreSimulation)
 	TimeSensitive = 0;
 
 	// If there is any pending data to send, send it.
-	if ( SendBuffer.GetNumBits() || ( Driver->Time-LastSendTime > Driver->KeepAliveTime && !InternalAck ) )
+	if ( SendBuffer.GetNumBits() || ( Driver->Time-LastSendTime > Driver->KeepAliveTime && !InternalAck && State != USOCK_Closed ) )
 	{
 		// If sending keepalive packet, still write the packet id
 		if ( SendBuffer.GetNumBits() == 0 )
