Engine: Expose 2D dot and cross product methods to Blueprints

diff --git a/Engine/Source/Runtime/Engine/Classes/Kismet/KismetMathLibrary.h b/Engine/Source/Runtime/Engine/Classes/Kismet/KismetMathLibrary.h
--- a/Engine/Source/Runtime/Engine/Classes/Kismet/KismetMathLibrary.h
+++ b/Engine/Source/Runtime/Engine/Classes/Kismet/KismetMathLibrary.h
@@ -537,14 +537,22 @@ class UKismetMathLibrary : public UBlueprintFunctionLibrary
 	UFUNCTION(BlueprintPure, meta=(FriendlyName = "Not Equal (vector)", CompactNodeTitle = "!=", Keywords = "!= not equal"), Category="Math|Vector")
 	static bool NotEqual_VectorVector(FVector A, FVector B, float ErrorTolerance = 1.e-4f);
 
-	/* Returns the dot product of two vectors */
+	/* Returns the dot product of two 3d vectors */
 	UFUNCTION(BlueprintPure, meta=(FriendlyName = "Dot Product", CompactNodeTitle = "."), Category="Math|Vector" )
 	static float Dot_VectorVector(FVector A, FVector B);
 
-	/* Returns the cross product of two vectors */
-	UFUNCTION(BlueprintPure, meta=(FriendlyName = "Cross Product", CompactNodeTitle = "x"), Category="Math|Vector" )
+	/* Returns the cross product of two 3d vectors */
+	UFUNCTION(BlueprintPure, meta=(FriendlyName = "Cross Product", CompactNodeTitle = ""), Category="Math|Vector" )
 	static FVector Cross_VectorVector(FVector A, FVector B);
-	
+
+	/* Returns the dot product of two 2d vectors */
+	UFUNCTION(BlueprintPure, meta=(FriendlyName="Dot Product (2D)", CompactNodeTitle="."), Category="Math|Vector")
+	static float DotProduct2D(FVector2D A, FVector2D B);
+
+	/* Returns the cross product of two 2d vectors */
+	UFUNCTION(BlueprintPure, meta=(FriendlyName="Cross Product (2D)", CompactNodeTitle="x"), Category="Math|Vector")
+	static float CrossProduct2D(FVector2D A, FVector2D B);
+
 	/* Returns the length of the FVector */
 	UFUNCTION(BlueprintPure, meta=(FriendlyName = "VectorLength"), Category="Math|Vector")
 	static float VSize(FVector A);
diff --git a/Engine/Source/Runtime/Engine/Private/KismetMathLibrary.cpp b/Engine/Source/Runtime/Engine/Private/KismetMathLibrary.cpp
--- a/Engine/Source/Runtime/Engine/Private/KismetMathLibrary.cpp
+++ b/Engine/Source/Runtime/Engine/Private/KismetMathLibrary.cpp
@@ -705,13 +705,23 @@ bool UKismetMathLibrary::NotEqual_VectorVector(FVector A, FVector B, float Error
 
 float UKismetMathLibrary::Dot_VectorVector(FVector A, FVector B)
 {
-	return A | B;
+	return FVector::DotProduct(A, B);
 }	
 
 FVector UKismetMathLibrary::Cross_VectorVector(FVector A, FVector B)
 {
-	return A ^ B;
-}	
+	return FVector::CrossProduct(A, B);
+}
+
+float UKismetMathLibrary::DotProduct2D(FVector2D A, FVector2D B)
+{
+	return FVector2D::DotProduct(A, B);
+}
+
+float UKismetMathLibrary::CrossProduct2D(FVector2D A, FVector2D B)
+{
+	return FVector2D::CrossProduct(A, B);
+}
 
 float UKismetMathLibrary::VSize(FVector A)
 {
