Fix engine analytics asserting for no reason if the engine was shutdown before analytics was initialized

diff --git a/Engine/Source/Runtime/Engine/Private/EngineAnalytics.cpp b/Engine/Source/Runtime/Engine/Private/EngineAnalytics.cpp
--- a/Engine/Source/Runtime/Engine/Private/EngineAnalytics.cpp
+++ b/Engine/Source/Runtime/Engine/Private/EngineAnalytics.cpp
@@ -214,7 +214,6 @@ void FEngineAnalytics::Initialize()
 
 void FEngineAnalytics::Shutdown()
 {
-	checkf(bIsInitialized, TEXT("FEngineAnalytics::Shutdown called outside of Initialize."));
 	Analytics.Reset();
 	bIsInitialized = false;
 }
