#ue4 - Fix CharacterMovement incorrectly changing the location of a simulated network Character's mesh that is simulating physics. When the mesh simulates physics it is detached, so SetRelativeLocation() starts working in world coordinates.

Jira: UE-6128

diff --git a/Engine/Source/Runtime/Engine/Private/CharacterMovementComponent.cpp b/Engine/Source/Runtime/Engine/Private/CharacterMovementComponent.cpp
--- a/Engine/Source/Runtime/Engine/Private/CharacterMovementComponent.cpp
+++ b/Engine/Source/Runtime/Engine/Private/CharacterMovementComponent.cpp
@@ -5149,10 +5149,10 @@ void UCharacterMovementComponent::SmoothClientPosition(float DeltaSeconds)
 	}
 
 	FNetworkPredictionData_Client_Character* ClientData = GetPredictionData_Client_Character();
-	if (ClientData && ClientData->bSmoothNetUpdates)
+	if (ClientData && ClientData->bSmoothNetUpdates && CharacterOwner->GetMesh() && !CharacterOwner->GetMesh()->IsSimulatingPhysics())
 	{
-		// smooth interpolation of mesh translation to avoid popping of other client pawns, unless driving or ragdoll or low tick rate
-		if ((DeltaSeconds < ClientData->SmoothNetUpdateTime) && CharacterOwner->GetMesh() && !CharacterOwner->GetMesh()->IsSimulatingPhysics())
+		// smooth interpolation of mesh translation to avoid popping of other client pawns unless under a low tick rate
+		if (DeltaSeconds < ClientData->SmoothNetUpdateTime)
 		{
 			ClientData->MeshTranslationOffset = (ClientData->MeshTranslationOffset * (1.f - DeltaSeconds / ClientData->SmoothNetUpdateTime));
 		}
@@ -5167,11 +5167,8 @@ void UCharacterMovementComponent::SmoothClientPosition(float DeltaSeconds)
 			ClientData->MeshTranslationOffset.Z = 0;
 		}
 
-		if (CharacterOwner->GetMesh())
-		{
-			const FVector NewRelTranslation = CharacterOwner->ActorToWorld().InverseTransformVectorNoScale(ClientData->MeshTranslationOffset + CharacterOwner->GetBaseTranslationOffset());
-			CharacterOwner->GetMesh()->SetRelativeLocation(NewRelTranslation);
-		}
+		const FVector NewRelTranslation = CharacterOwner->ActorToWorld().InverseTransformVectorNoScale(ClientData->MeshTranslationOffset + CharacterOwner->GetBaseTranslationOffset());
+		CharacterOwner->GetMesh()->SetRelativeLocation(NewRelTranslation);
 	}
 }
 
