Fix crash calling UGameplayStatics::SetGamePaused on a client

diff --git a/Engine/Source/Runtime/Engine/Private/GameplayStatics.cpp b/Engine/Source/Runtime/Engine/Private/GameplayStatics.cpp
--- a/Engine/Source/Runtime/Engine/Private/GameplayStatics.cpp
+++ b/Engine/Source/Runtime/Engine/Private/GameplayStatics.cpp
@@ -129,18 +129,10 @@ void UGameplayStatics::SetGlobalTimeDilation(UObject* WorldContextObject, float
 bool UGameplayStatics::SetGamePaused(UObject* WorldContextObject, bool bPaused)
 {
 	UWorld* const World = GEngine->GetWorldFromContextObject( WorldContextObject );
-	if (bPaused)
-	{
-		APlayerController* const PC = World->GetFirstPlayerController();
-		check(PC); // Gathering some information for TTP #303973
+	APlayerController* const PC = World->GetFirstPlayerController();
+	check(PC); // Gathering some information for TTP #303973
 
-		return World->GetAuthGameMode()->SetPause(PC);
-	}
-	else
-	{
-		World->GetAuthGameMode()->ClearPause();
-		return true;
-	}
+	return PC->SetPause(bPaused);
 }
 
 bool UGameplayStatics::IsGamePaused(UObject* WorldContextObject)
