#ue4 - Expose CharacterMovementComponent::DisableMovement() to Blueprints.

diff --git a/Engine/Source/Runtime/Engine/Classes/GameFramework/CharacterMovementComponent.h b/Engine/Source/Runtime/Engine/Classes/GameFramework/CharacterMovementComponent.h
--- a/Engine/Source/Runtime/Engine/Classes/GameFramework/CharacterMovementComponent.h
+++ b/Engine/Source/Runtime/Engine/Classes/GameFramework/CharacterMovementComponent.h
@@ -507,7 +507,7 @@ public:
 	float MaxSimulationTimeStep;
 
 	/**
-	 * Max number of iterations used for each discrete simulation step
+	 * Max number of iterations used for each discrete simulation step.
 	 * Used primarily in the the more advanced movement modes that break up larger time steps (usually those applying gravity such as falling and walking).
 	 * Increasing this value can address issues with fast-moving objects or complex collision scenarios, at the cost of performance.
 	 *
@@ -789,6 +789,7 @@ public:
 #endif // WITH_EDITOR
 
 	/** Make movement impossible (sets movement mode to MOVE_None). */
+	UFUNCTION(BlueprintCallable, Category="Pawn|Components|CharacterMovement")
 	virtual void DisableMovement();
 
 	/** Return true if we have a valid CharacterOwner and UpdatedComponent. */
