#ue4 - MovementComponent.UpdatedComponent is now DuplicateTransient.

diff --git a/Engine/Source/Runtime/Engine/Classes/GameFramework/MovementComponent.h b/Engine/Source/Runtime/Engine/Classes/GameFramework/MovementComponent.h
--- a/Engine/Source/Runtime/Engine/Classes/GameFramework/MovementComponent.h
+++ b/Engine/Source/Runtime/Engine/Classes/GameFramework/MovementComponent.h
@@ -52,7 +52,7 @@ class ENGINE_API UMovementComponent : public UActorComponent
 	 * If this is null at startup and bAutoRegisterUpdatedComponent is true, the owning Actor's root component will automatically be set as our UpdatedComponent at startup.
 	 * @see bAutoRegisterUpdatedComponent, SetUpdatedComponent()
 	 */
-	UPROPERTY(BlueprintReadOnly, Category=MovementComponent)
+	UPROPERTY(BlueprintReadOnly, DuplicateTransient, Category=MovementComponent)
 	UPrimitiveComponent* UpdatedComponent;
 
 	/**
