Fixes crash when undoing Blueprint actions that add the Blueprint to the transaction buffer.

On Undo, the Blueprint would Serialize and gather Find-in-Blueprint data while the asset was in an invalid state. Gathering this data during a transaction makes little sense as the established method of updating the searchable data is either saving the Blueprint or compiling it.

diff --git a/Engine/Source/Runtime/Engine/Private/Blueprint.cpp b/Engine/Source/Runtime/Engine/Private/Blueprint.cpp
--- a/Engine/Source/Runtime/Engine/Private/Blueprint.cpp
+++ b/Engine/Source/Runtime/Engine/Private/Blueprint.cpp
@@ -351,7 +351,7 @@ void UBlueprint::Serialize(FArchive& Ar)
 		}
 	}
 
-	if(Ar.IsSaving())
+	if(Ar.IsSaving() && !Ar.IsTransacting())
 	{
 		// Cache the BP for use
 		FFindInBlueprintSearchManager::Get().AddOrUpdateBlueprintSearchMetadata(this);
