Engine: Replace color literal with FColorList entry for LEVEL_MakingVisible

diff --git a/Engine/Source/Runtime/Engine/Private/UnrealEngine.cpp b/Engine/Source/Runtime/Engine/Private/UnrealEngine.cpp
--- a/Engine/Source/Runtime/Engine/Private/UnrealEngine.cpp
+++ b/Engine/Source/Runtime/Engine/Private/UnrealEngine.cpp
@@ -10602,7 +10602,7 @@ FColor GetColorForLevelStatus(int32 Status)
 		Color = FColor::Red;		// red  loaded and visible
 		break;
 	case LEVEL_MakingVisible:
-		Color = FColor(255, 128, 0);// orange, in process of being made visible
+		Color = FColorList::Orange;	// orange, in process of being made visible
 		break;
 	case LEVEL_Loading:
 		Color = FColor::Magenta;	// purple, in process of being loaded
