#ue4 - Removed dummy StaticMeshComponent that was created by WorldSettings, that was used for character networking movement on BSPs a while back before networking was switched to use component references.

diff --git a/Engine/Source/Runtime/Engine/Private/WorldSettings.cpp b/Engine/Source/Runtime/Engine/Private/WorldSettings.cpp
--- a/Engine/Source/Runtime/Engine/Private/WorldSettings.cpp
+++ b/Engine/Source/Runtime/Engine/Private/WorldSettings.cpp
@@ -58,14 +58,6 @@ AWorldSettings::AWorldSettings(const FObjectInitializer& ObjectInitializer)
 	VisibilityAggressiveness = VIS_LeastAggressive;
 	LevelLightingQuality = Quality_MAX;
 
-	UStaticMeshComponent* StaticMeshComponent = ObjectInitializer.CreateDefaultSubobject<UStaticMeshComponent>(this, TEXT("StaticMeshComponent0"));
-	StaticMeshComponent->bHiddenInGame = true;
-	StaticMeshComponent->SetCollisionProfileName(UCollisionProfile::BlockAll_ProfileName);
-	StaticMeshComponent->PostPhysicsComponentTick.bCanEverTick = false;
-	StaticMeshComponent->Mobility = EComponentMobility::Static;
-
-	RootComponent = StaticMeshComponent;
-
 #if WITH_EDITORONLY_DATA
 	bActorLabelEditable = false;
 #endif // WITH_EDITORONLY_DATA
