#ue4 - PlayerController::ClientSetHUD() is now BlueprintCallable.

diff --git a/Engine/Source/Runtime/Engine/Classes/GameFramework/PlayerController.h b/Engine/Source/Runtime/Engine/Classes/GameFramework/PlayerController.h
--- a/Engine/Source/Runtime/Engine/Classes/GameFramework/PlayerController.h
+++ b/Engine/Source/Runtime/Engine/Classes/GameFramework/PlayerController.h
@@ -767,8 +767,8 @@ public:
 	UFUNCTION(Reliable, Client)
 	void ClientSetForceMipLevelsToBeResident(class UMaterialInterface* Material, float ForceDuration, int32 CinematicTextureGroups = 0);
 
-	/** Set the client's class of HUD */
-	UFUNCTION(Reliable, Client)
+	/** Set the client's class of HUD and spawns a new instance of it. If there was already a HUD active, it is destroyed. */
+	UFUNCTION(BlueprintCallable, Category="HUD", Reliable, Client)
 	void ClientSetHUD(TSubclassOf<class AHUD> NewHUDClass);
 
 	/** Helper to get the size of the HUD canvas for this player controller.  Returns 0 if there is no HUD */
