Fixed NavigationSystem skipping the very first navigation rebuilding even if some navigation data needs to be rebuild due to failed serialization #UE4

- also, unhacked it a bit

diff --git a/Engine/Source/Runtime/Engine/Classes/AI/Navigation/NavigationData.h b/Engine/Source/Runtime/Engine/Classes/AI/Navigation/NavigationData.h
--- a/Engine/Source/Runtime/Engine/Classes/AI/Navigation/NavigationData.h
+++ b/Engine/Source/Runtime/Engine/Classes/AI/Navigation/NavigationData.h
@@ -393,9 +393,14 @@ class ENGINE_API ANavigationData : public AActor
 	//----------------------------------------------------------------------//
 
 	/** If true, the NavMesh can be dynamically rebuilt at runtime. */
-	UPROPERTY(EditAnywhere, Category=Runtime, config)
+	UPROPERTY(EditAnywhere, Category = Runtime, config)
 	uint32 bRebuildAtRuntime:1;
 
+	/** By default navigation will skip the first update after being successfully loaded
+	 *  setting bForceRebuildOnLoad to false can override this behavior */
+	UPROPERTY(config, EditAnywhere, Category = Runtime)
+	uint32 bForceRebuildOnLoad : 1;
+
 	/** all observed paths will be processed every ObservedPathsTickInterval seconds */
 	UPROPERTY(EditAnywhere, Category = Runtime, config)
 	float ObservedPathsTickInterval;
@@ -448,6 +453,8 @@ public:
 
 	virtual bool DoesSupportAgent(const FNavAgentProperties& AgentProps) const;
 
+	FORCEINLINE void MarkAsNeedingUpdate() { bWantsUpdate = true; }
+
 protected:
 	virtual void FillConfig(FNavDataConfig& Dest) { Dest = NavDataConfig; }
 
@@ -780,6 +787,10 @@ protected:
 	/** was it generated for default agent (SupportedAgents[0]) */
 	uint32 bSupportsDefaultAgent:1;
 
+	/** controls whether this NavigationData will accept navigation building requests resulting from dirty areas
+	 *	true by default, set to false on successful serialization to skip initial rebuild */
+	uint32 bWantsUpdate:1;
+
 private:
 	uint16 NavDataUniqueID;
 
diff --git a/Engine/Source/Runtime/Engine/Classes/AI/Navigation/NavigationSystem.h b/Engine/Source/Runtime/Engine/Classes/AI/Navigation/NavigationSystem.h
--- a/Engine/Source/Runtime/Engine/Classes/AI/Navigation/NavigationSystem.h
+++ b/Engine/Source/Runtime/Engine/Classes/AI/Navigation/NavigationSystem.h
@@ -171,11 +171,6 @@ class ENGINE_API UNavigationSystem : public UBlueprintFunctionLibrary
 	UPROPERTY(config, EditAnywhere, Category=NavigationSystem)
 	uint32 bSkipAgentHeightCheckWhenPickingNavData:1;
 
-	/** If set to true, navigation data will be always rebuilt from scratch after loading
-	 *  otherwise navigation system will reuse data saved with map */
-	UPROPERTY(config, EditAnywhere, Category = NavigationSystem)
-	uint32 bForceRebuildOnLoad : 1;
-
 	UPROPERTY(config, EditAnywhere, Category = Agents)
 	TArray<FNavDataConfig> SupportedAgents;
 	
@@ -700,7 +695,6 @@ protected:
 	uint32 bInitialBuildingLockActive:1;
 	uint32 bInitialSetupHasBeenPerformed:1;
 	uint32 bInitialLevelsAdded:1;
-	uint32 bSkipDirtyAreasOnce:1;
 	uint32 bAsyncBuildPaused:1;
 
 	/** cached navigable world bounding box*/
diff --git a/Engine/Source/Runtime/Engine/Private/AI/Navigation/NavigationData.cpp b/Engine/Source/Runtime/Engine/Private/AI/Navigation/NavigationData.cpp
--- a/Engine/Source/Runtime/Engine/Private/AI/Navigation/NavigationData.cpp
+++ b/Engine/Source/Runtime/Engine/Private/AI/Navigation/NavigationData.cpp
@@ -100,10 +100,12 @@ ANavigationData::ANavigationData(const FObjectInitializer& ObjectInitializer)
 	: Super(ObjectInitializer)
 	, bEnableDrawing(false)
 	, bRebuildAtRuntime(false)
+	, bForceRebuildOnLoad(false)
 	, DataVersion(NAVMESHVER_LATEST)
 	, FindPathImplementation(NULL)
 	, FindHierarchicalPathImplementation(NULL)
 	, bRegistered(false)
+	, bWantsUpdate(true)
 	, NavDataUniqueID(GetNextUniqueID())
 {
 	PrimaryActorTick.bCanEverTick = true;
@@ -425,10 +427,14 @@ void ANavigationData::TickAsyncBuild(float DeltaSeconds)
 
 void ANavigationData::RebuildDirtyAreas(const TArray<FNavigationDirtyArea>& DirtyAreas)
 {
-	if (NavDataGenerator)
+	// the 'bWantsUpdate' mechanics allows us to skip first requested update after data is loaded
+	// Can be also used to manually control navigation rebuilding, by for example forcing bWantsUpdate 
+	//	to false in Tick function effectively supressing rebuilding
+	if (bWantsUpdate == true && NavDataGenerator)
 	{
 		NavDataGenerator->RebuildDirtyAreas(DirtyAreas);
 	}
+	bWantsUpdate = true;
 }
 
 TArray<FBox> ANavigationData::GetNavigableBounds() const
diff --git a/Engine/Source/Runtime/Engine/Private/AI/Navigation/NavigationSystem.cpp b/Engine/Source/Runtime/Engine/Private/AI/Navigation/NavigationSystem.cpp
--- a/Engine/Source/Runtime/Engine/Private/AI/Navigation/NavigationSystem.cpp
+++ b/Engine/Source/Runtime/Engine/Private/AI/Navigation/NavigationSystem.cpp
@@ -207,7 +207,6 @@ UNavigationSystem::UNavigationSystem(const FObjectInitializer& ObjectInitializer
 	, bWholeWorldNavigable(false)
 	, bAddPlayersToGenerationSeeds(true)
 	, bSkipAgentHeightCheckWhenPickingNavData(false)
-	, bForceRebuildOnLoad(false)
 	, DirtyAreasUpdateFreq(60)
 	, OperationMode(FNavigationSystem::InvalidMode)
 	, NavOctree(NULL)
@@ -454,9 +453,6 @@ void UNavigationSystem::OnWorldInitDone(FNavigationSystem::EMode Mode)
 			bInitialBuildingLockActive = false;
 		}
 
-		// don't mark dirty areas after loading a map if navigation system doesn't want to rebuilt from scratch
-		bSkipDirtyAreasOnce = !bForceRebuildOnLoad;
-
 		if (bAutoCreateNavigationData == true)
 		{
 			SpawnMissingNavigationData();
@@ -583,12 +579,6 @@ void UNavigationSystem::Tick(float DeltaSeconds)
 
 		if (DirtyAreas.Num() > 0 && bCanRebuildNow)
 		{
-			if (bSkipDirtyAreasOnce)
-			{
-				bSkipDirtyAreasOnce = false;
-				DirtyAreas.Reset();
-			}
-
 			for (int32 NavDataIndex = 0; NavDataIndex < NavDataSet.Num(); ++NavDataIndex)
 			{
 				ANavigationData* NavData = NavDataSet[NavDataIndex];
diff --git a/Engine/Source/Runtime/Engine/Private/AI/Navigation/RecastNavMesh.cpp b/Engine/Source/Runtime/Engine/Private/AI/Navigation/RecastNavMesh.cpp
--- a/Engine/Source/Runtime/Engine/Private/AI/Navigation/RecastNavMesh.cpp
+++ b/Engine/Source/Runtime/Engine/Private/AI/Navigation/RecastNavMesh.cpp
@@ -682,6 +682,7 @@ void ARecastNavMesh::Serialize( FArchive& Ar )
 		else if (RecastNavMeshSizeBytes > 4)
 		{
 			SerializeRecastNavMesh(Ar, RecastNavMeshImpl);
+			bWantsUpdate = bForceRebuildOnLoad == true || HasValidNavmesh() == false;
 #if !(UE_BUILD_SHIPPING)
 			RequestDrawingUpdate();
 #endif //!(UE_BUILD_SHIPPING)
diff --git a/Engine/Source/Runtime/Engine/Private/AI/Navigation/RecastNavMeshGenerator.cpp b/Engine/Source/Runtime/Engine/Private/AI/Navigation/RecastNavMeshGenerator.cpp
--- a/Engine/Source/Runtime/Engine/Private/AI/Navigation/RecastNavMeshGenerator.cpp
+++ b/Engine/Source/Runtime/Engine/Private/AI/Navigation/RecastNavMeshGenerator.cpp
@@ -2708,6 +2708,7 @@ FRecastNavMeshGenerator::FRecastNavMeshGenerator(ARecastNavMesh* InDestNavMesh)
 	if (bRecreateNavmesh)
 	{
 		ConstructTiledNavMesh();
+		InDestNavMesh->MarkAsNeedingUpdate();
 	}
 	else
 	{
