UE4 - tweaks to hierarchical foliage. Fix a few bugs in the handling of temporal LOD. Better legacy support. Opaque surfaces working correctly.

diff --git a/Engine/Shaders/LocalVertexFactory.usf b/Engine/Shaders/LocalVertexFactory.usf
--- a/Engine/Shaders/LocalVertexFactory.usf
+++ b/Engine/Shaders/LocalVertexFactory.usf
@@ -11,13 +11,13 @@ float4x4 PreviousLocalToWorld;
 float4 LightMapCoordinateScaleBias;
 float4 ShadowMapCoordinateScaleBias;
 
-#ifndef USE_NEW_LOD_TRANSITION
-	#define USE_NEW_LOD_TRANSITION 0
+#ifndef USE_DITHERED_LOD_TRANSITION
+	#error "USE_DITHERED_LOD_TRANSITION should have been defined"
 #endif
 
 #if USE_INSTANCING
 
-#if USE_NEW_LOD_TRANSITION
+#if USE_DITHERED_LOD_TRANSITION
     float4 InstancingViewZCompareZero;
     float4 InstancingViewZCompareOne;
     float4 InstancingViewZConstant;
@@ -39,18 +39,14 @@ float4 InstancingFadeOutParams;
 #endif	// USE_INSTANCING_EMULATED
 
 #else // USE_INSTANCING
-	#ifndef USE_NEW_LOD_TRANSITION
-		#define USE_NEW_LOD_TRANSITION (0)
-	#endif
 	#ifndef USE_INSTANCING_EMULATED
 		#define USE_INSTANCING_EMULATED (0)
 	#endif
+	#if USE_DITHERED_LOD_TRANSITION
+		float NonInstancedDitherLODFactor;
+	#endif
 #endif	// USE_INSTANCING
 
-#if !USE_INSTANCING && USE_NEW_LOD_TRANSITION
-float NonInstancedDitherLODFactor;
-#endif
-
 
 #if USE_SPLINEDEFORM
 float3 SplineStartPos;
@@ -151,7 +147,8 @@ struct FVertexFactoryIntermediates
 
 	half4 Color;
 #if USE_INSTANCING
-	float3 PerInstanceParams;
+	// x = per-instance random, y = per-instance fade out amount, z = hide/show flag, w dither fade cutoff
+	float4 PerInstanceParams;
 #endif	// USE_INSTANCING
 };
 
@@ -460,36 +457,29 @@ FVertexFactoryIntermediates GetVertexFactoryIntermediates(FVertexFactoryInput In
 	// Swizzle vertex color.
 	Intermediates.Color = Input.Color FCOLOR_COMPONENT_SWIZZLE;
 #if USE_INSTANCING
-	// x = per-instance random, y = per-instance fade out factor, z = zero or one depending of if it is shown at all. 
+	// x = per-instance random, y = per-instance fade out factor, z = zero or one depending of if it is shown at all, w is dither cutoff 
 
 	// PerInstanceParams.z stores a hide/show flag for this instance
 	float SelectedValue = trunc(GetInstanceShadowMapBias(Input).z * (1.0 / 256.0));
 	// InstancingFadeOutParams.z,w are RenderSelected and RenderDeselected respectively.
-	Intermediates.PerInstanceParams.z = InstancingFadeOutParams.z * SelectedValue + InstancingFadeOutParams.w * (1-SelectedValue);
-	#if !USE_NEW_LOD_TRANSITION
-		#if !USE_INSTANCING_EMULATED
-			Intermediates.PerInstanceParams.x = Input.InstanceInverseTransform3.w;
-			float3 InstanceLocation = TransformLocalToTranslatedWorld(float3(Input.InstanceTransform1.w, Input.InstanceTransform2.w, Input.InstanceTransform3.w)).xyz;
-		#else
-			Intermediates.PerInstanceParams.x = CPUInstanceInverseTransform[2].w;
-			float3 InstanceLocation = TransformLocalToTranslatedWorld(float3(CPUInstanceTransform[0].w, CPUInstanceTransform[1].w, CPUInstanceTransform[2].w)).xyz;
-		#endif // !USE_INSTANCING_EMULATED
-		Intermediates.PerInstanceParams.y = 1.0 - saturate((length(InstanceLocation) - InstancingFadeOutParams.x) * InstancingFadeOutParams.y);
+	#if !USE_INSTANCING_EMULATED
+		Intermediates.PerInstanceParams.x = Input.InstanceInverseTransform3.w;
+		float3 InstanceLocation = TransformLocalToWorld(float3(Input.InstanceTransform1.w, Input.InstanceTransform2.w, Input.InstanceTransform3.w)).xyz;
 	#else
-		#if !USE_INSTANCING_EMULATED
-			Intermediates.PerInstanceParams.x = Input.InstanceInverseTransform3.w;
-			float3 InstanceLocation = TransformLocalToWorld(float3(Input.InstanceTransform1.w, Input.InstanceTransform2.w, Input.InstanceTransform3.w)).xyz;
-		#else
-			Intermediates.PerInstanceParams.x = CPUInstanceInverseTransform[2].w;
-			float3 InstanceLocation = TransformLocalToWorld(float3(CPUInstanceTransform[0].w, CPUInstanceTransform[1].w, CPUInstanceTransform[2].w)).xyz;
-		#endif // !USE_INSTANCING_EMULATED
-
+		Intermediates.PerInstanceParams.x = CPUInstanceInverseTransform[2].w;
+		float3 InstanceLocation = TransformLocalToWorld(float3(CPUInstanceTransform[0].w, CPUInstanceTransform[1].w, CPUInstanceTransform[2].w)).xyz;
+	#endif // !USE_INSTANCING_EMULATED
+	Intermediates.PerInstanceParams.y = 1.0 - saturate((length(InstanceLocation + View.PreViewTranslation.xyz) - InstancingFadeOutParams.x) * InstancingFadeOutParams.y);
+	Intermediates.PerInstanceParams.z = InstancingFadeOutParams.z * SelectedValue + InstancingFadeOutParams.w * (1-SelectedValue);
+	#if USE_DITHERED_LOD_TRANSITION
 		float ViewZZero = length(InstanceLocation - InstancingWorldViewOriginZero.xyz);
 		float ViewZOne = length(InstanceLocation - InstancingWorldViewOriginOne.xyz);
-		Intermediates.PerInstanceParams.y = 
+		Intermediates.PerInstanceParams.w = 
 			dot(ViewZZero > InstancingViewZCompareZero.xyz, InstancingViewZConstant.xyz) * InstancingWorldViewOriginZero.w +
 			dot(ViewZOne > InstancingViewZCompareOne.xyz, InstancingViewZConstant.xyz) * InstancingWorldViewOriginOne.w;
-		Intermediates.PerInstanceParams.z *= abs(Intermediates.PerInstanceParams.y) < .999;
+		Intermediates.PerInstanceParams.z *= abs(Intermediates.PerInstanceParams.w) < .999;
+	#else
+		Intermediates.PerInstanceParams.w = 0;
 	#endif
 #endif	// USE_INSTANCING
 	return Intermediates;
diff --git a/Engine/Shaders/LocalVertexFactoryCommon.usf b/Engine/Shaders/LocalVertexFactoryCommon.usf
--- a/Engine/Shaders/LocalVertexFactoryCommon.usf
+++ b/Engine/Shaders/LocalVertexFactoryCommon.usf
@@ -17,8 +17,8 @@ struct FVertexFactoryInterpolantsVSToPS
 #endif
 
 #if USE_INSTANCING
-	// x = per-instance random, y = per-instance fade out amount, z = hide/show flag
-	float3  PerInstanceParams : COLOR1;
+	// x = per-instance random, y = per-instance fade out amount, z = hide/show flag, w dither fade cutoff
+	float4  PerInstanceParams : COLOR1;
 #endif
 
 #if NUM_MATERIAL_TEXCOORDS
@@ -142,8 +142,8 @@ struct FVertexFactoryInterpolantsVSToPS
 #endif
 
 #if USE_INSTANCING
-	// x = per-instance random, y = per-instance fade out amount
-	float3  PerInstanceParams : COLOR1;
+	// x = per-instance random, y = per-instance fade out amount, z = hide/show flag, w dither fade cutoff
+	float4  PerInstanceParams : COLOR1;
 #endif
 };
 
diff --git a/Engine/Shaders/MaterialTemplate.usf b/Engine/Shaders/MaterialTemplate.usf
--- a/Engine/Shaders/MaterialTemplate.usf
+++ b/Engine/Shaders/MaterialTemplate.usf
@@ -21,15 +21,10 @@
 	#include "AtmosphereCommon.usf"
 #endif
 
-#if USE_INSTANCING
-	#ifndef USE_NEW_LOD_TRANSITION
-		#error "USE_NEW_LOD_TRANSITION must be defined if USE_INSTANCING is"
-	#endif
-#else
-	#ifndef USE_NEW_LOD_TRANSITION
-		#define USE_NEW_LOD_TRANSITION (0)
-	#endif
+#ifndef USE_DITHERED_LOD_TRANSITION
+	#define USE_DITHERED_LOD_TRANSITION 0
 #endif
+
 /**
  * Parameters used by vertex and pixel shaders to access particle properties.
  */
@@ -141,7 +136,7 @@ struct FMaterialPixelParameters
 #endif
 
 #if USE_INSTANCING
-	half3 PerInstanceParams;
+	half4 PerInstanceParams;
 #endif
 
 	/** Per-particle properties. Only valid for particle vertex factories. */
@@ -238,7 +233,7 @@ struct FMaterialVertexParameters
 	float4x4 InstanceLocalToWorld;
 	float3x3 InstanceWorldToLocal;
 	float3 InstanceLocalPosition;
-	float3 PerInstanceParams;
+	float4 PerInstanceParams;
 #elif PARTICLE_MESH_FACTORY
 	/** Per-particle properties. */
 	float4x4 InstanceLocalToWorld;
@@ -395,11 +390,7 @@ float GetPerInstanceRandom(FMaterialPixelParameters Parameters)
 float GetPerInstanceFadeAmount(FMaterialPixelParameters Parameters)
 {
 #if USE_INSTANCING
-	#if !USE_NEW_LOD_TRANSITION
-		return float(Parameters.PerInstanceParams.y);
-	#else
-		return float(saturate(1.0 - abs(Parameters.PerInstanceParams.y)));
-	#endif
+	return float(Parameters.PerInstanceParams.y);
 #else
 	return float(1.0);
 #endif
@@ -722,7 +713,7 @@ float3 GetSpeedTreeVertexOffset(FMaterialVertexParameters Parameters, int Geomet
 
 	// compute LOD by finding screen space size
 	float LodInterp = 1.0;
-#if !USE_INSTANCING || !USE_NEW_LOD_TRANSITION
+#if !USE_INSTANCING || !USE_DITHERED_LOD_TRANSITION
 	if (LODType == SPEEDTREE_LOD_TYPE_SMOOTH) 
 	{
 		const float ScreenMultiple = 0.5 * max(View.ViewSizeAndSceneTexelSize.x * View.ViewToClip[0][0],
@@ -1277,11 +1268,11 @@ float3 GetMaterialHemisphereLightTransferFull(float3 DiffuseColor,FMaterialPixel
 	       LowerLighting * LowerColor;
 }
 
-void ClipLODTransition(FMaterialPixelParameters Parameters, float DitherFactor)
+void ClipLODTransition(float2 SVPosition, float DitherFactor)
 {
 	if (abs(DitherFactor) > .001)
 	{
-		float RandCos = cos(dot(floor(Parameters.SVPosition.xy), float2(347.83451793,3343.28371963)));
+		float RandCos = cos(dot(floor(SVPosition.xy), float2(347.83451793,3343.28371963)));
 		float RandomVal = frac(RandCos * 1000.0);
 		half RetVal = (DitherFactor < 0.0) ?
 			(DitherFactor + 1.0 > RandomVal) :
@@ -1290,20 +1281,35 @@ void ClipLODTransition(FMaterialPixelParameters Parameters, float DitherFactor)
 	}
 }
 
-#if USE_INSTANCING && USE_NEW_LOD_TRANSITION
+void ClipLODTransition(FMaterialPixelParameters Parameters, float DitherFactor)
+{
+	ClipLODTransition(Parameters.SVPosition.xy, DitherFactor);
+}
+
+
+#define REQUIRES_VF_ATTRIBUTES_FOR_CLIPPING (USE_INSTANCING && USE_DITHERED_LOD_TRANSITION)
+
+#if USE_INSTANCING && USE_DITHERED_LOD_TRANSITION
 void ClipLODTransition(FMaterialPixelParameters Parameters)
 {
-	ClipLODTransition(Parameters, Parameters.PerInstanceParams.y);
+	ClipLODTransition(Parameters, Parameters.PerInstanceParams.w);
 }
-#elif USE_NEW_LOD_TRANSITION
+#elif USE_DITHERED_LOD_TRANSITION
 void ClipLODTransition(FMaterialPixelParameters Parameters)
 {
 	ClipLODTransition(Parameters, NonInstancedDitherLODFactor);
 }
+void ClipLODTransition(float2 SVPosition)
+{
+	ClipLODTransition(SVPosition, NonInstancedDitherLODFactor);
+}
 #else
 void ClipLODTransition(FMaterialPixelParameters Parameters)
 {
 }
+void ClipLODTransition(float2 SVPosition)
+{
+}
 #endif
 
 void GetMaterialClippingShadowDepth(FMaterialPixelParameters Parameters) 
diff --git a/Engine/Shaders/ShadowDepthCommon.usf b/Engine/Shaders/ShadowDepthCommon.usf
--- a/Engine/Shaders/ShadowDepthCommon.usf
+++ b/Engine/Shaders/ShadowDepthCommon.usf
@@ -12,7 +12,7 @@
 #define REFLECTIVE_SHADOW_MAP 0
 #endif
 
-#define INTERPOLATE_VF_ATTRIBUTES (!MATERIALBLENDING_SOLID || REFLECTIVE_SHADOW_MAP || MATERIALBLENDING_MASKED)
+#define INTERPOLATE_VF_ATTRIBUTES (!MATERIALBLENDING_SOLID || REFLECTIVE_SHADOW_MAP || MATERIALBLENDING_MASKED || (REQUIRES_VF_ATTRIBUTES_FOR_CLIPPING && !POSITION_ONLY))
 #define INTERPOLATE_POSITION ((REFLECTIVE_SHADOW_MAP && !ONEPASS_POINTLIGHT_SHADOW) || MATERIALBLENDING_MASKED)
 
 // x: Tweakable depth bias, y:InvMaxSubjectDepthUsed to normalize the outputted depth
diff --git a/Engine/Shaders/ShadowDepthPixelShader.usf b/Engine/Shaders/ShadowDepthPixelShader.usf
--- a/Engine/Shaders/ShadowDepthPixelShader.usf
+++ b/Engine/Shaders/ShadowDepthPixelShader.usf
@@ -39,7 +39,6 @@ void Main(
 	out float4 OutColor : SV_Target0
 #if REFLECTIVE_SHADOW_MAP
 	,out float4 OutReflectiveShadowMapDiffuse : SV_Target1
-	,bool IsFrontFace : SV_IsFrontFace
 #endif
 #if PERSPECTIVE_CORRECT_DEPTH
 	,out float OutDepth : SV_DEPTH
@@ -58,6 +57,8 @@ void Main(
 	#endif
 	// Evaluate the mask for masked materials
 	GetMaterialClippingShadowDepth(MaterialParameters);
+#else
+	ClipLODTransition(SvPosition.xy);
 #endif
 
 #if PERSPECTIVE_CORRECT_DEPTH
diff --git a/Engine/Source/Runtime/Engine/Private/HierarchicalInstancedStaticMesh.cpp b/Engine/Source/Runtime/Engine/Private/HierarchicalInstancedStaticMesh.cpp
--- a/Engine/Source/Runtime/Engine/Private/HierarchicalInstancedStaticMesh.cpp
+++ b/Engine/Source/Runtime/Engine/Private/HierarchicalInstancedStaticMesh.cpp
@@ -32,6 +32,11 @@ static TAutoConsoleVariable<int32> CVarCullAll(
 	0,
 	TEXT("If greater than zero, everything is conside3red culled."));
 
+static TAutoConsoleVariable<int32> CVarDitheredLOD(
+	TEXT("foliage.DitheredLOD"),
+	1,
+	TEXT("If greater than zero, dithered LOD is used, otherwise popping LOD is used."));
+
 static TAutoConsoleVariable<int32> CVarMaxTrianglesToRender(
 	TEXT("foliage.MaxTrianglesToRender"),
 	100000000,
@@ -551,8 +556,10 @@ class FHierarchicalStaticMeshSceneProxy : public FInstancedStaticMeshSceneProxy
 	int32 LastUnbuiltIndex;
 
 #if !(UE_BUILD_SHIPPING || UE_BUILD_TEST)
-	mutable TArray<uint32> DebugRuns[MAX_STATIC_MESH_LODS];
-	mutable int32 DebugTotalInstances[MAX_STATIC_MESH_LODS];
+	mutable TArray<uint32> SingleDebugRuns[MAX_STATIC_MESH_LODS];
+	mutable int32 SingleDebugTotalInstances[MAX_STATIC_MESH_LODS];
+	mutable TArray<uint32> MultipleDebugRuns[MAX_STATIC_MESH_LODS];
+	mutable int32 MultipleDebugTotalInstances[MAX_STATIC_MESH_LODS];
 	mutable int32 CaptureTag;
 #endif
 
@@ -627,33 +634,55 @@ public:
 
 struct FFoliageRenderInstanceParams
 {
-	mutable TArray<uint32, SceneRenderingAllocator> Runs[MAX_STATIC_MESH_LODS];
-	mutable int32 TotalInstances[MAX_STATIC_MESH_LODS];
-
-	FFoliageRenderInstanceParams()
+	bool bNeedsSingleLODRuns;
+	bool bNeedsMultipleLODRuns;
+	mutable TArray<uint32, SceneRenderingAllocator> MultipleLODRuns[MAX_STATIC_MESH_LODS];
+	mutable TArray<uint32, SceneRenderingAllocator> SingleLODRuns[MAX_STATIC_MESH_LODS];
+	mutable int32 TotalSingleLODInstances[MAX_STATIC_MESH_LODS];
+	mutable int32 TotalMultipleLODInstances[MAX_STATIC_MESH_LODS];
+
+	FFoliageRenderInstanceParams(bool InbNeedsSingleLODRuns, bool InbNeedsMultipleLODRuns)
+		: bNeedsSingleLODRuns(InbNeedsSingleLODRuns)
+		, bNeedsMultipleLODRuns(InbNeedsMultipleLODRuns)
 	{
 		for (int32 Index = 0; Index < MAX_STATIC_MESH_LODS; Index++)
 		{
-			TotalInstances[Index] = 0;
+			TotalSingleLODInstances[Index] = 0;
+			TotalMultipleLODInstances[Index] = 0;
 		}
 	}
-	FORCEINLINE_DEBUGGABLE void AddRun(int32 Lod, int32 FirstInstance, int32 LastInstance) const
+	static FORCEINLINE_DEBUGGABLE void AddRun(TArray<uint32, SceneRenderingAllocator>& Array, int32 FirstInstance, int32 LastInstance)
 	{
-		TotalInstances[Lod] += 1 + LastInstance - FirstInstance;
-		if (Runs[Lod].Num() && Runs[Lod].Last() + 1 == FirstInstance)
+		if (Array.Num() && Array.Last() + 1 == FirstInstance)
 		{
-			Runs[Lod].Last() = (uint32)LastInstance;
+			Array.Last() = (uint32)LastInstance;
 		}
 		else
 		{
-			Runs[Lod].Add((uint32)FirstInstance);
-			Runs[Lod].Add((uint32)LastInstance);
+			Array.Add((uint32)FirstInstance);
+			Array.Add((uint32)LastInstance);
+		}
+	}
+	FORCEINLINE_DEBUGGABLE void AddRun(int32 MinLod, int32 MaxLod, int32 FirstInstance, int32 LastInstance) const
+	{
+		if (bNeedsSingleLODRuns)
+		{
+			AddRun(SingleLODRuns[MinLod], FirstInstance, LastInstance);
+			TotalSingleLODInstances[MinLod] += 1 + LastInstance - FirstInstance;
+		}
+		if (bNeedsMultipleLODRuns)
+		{
+			for (int32 Lod = MinLod; Lod <= MaxLod; Lod++)
+			{
+				TotalMultipleLODInstances[Lod] += 1 + LastInstance - FirstInstance;
+				AddRun(MultipleLODRuns[Lod], FirstInstance, LastInstance);
+			}
 		}
 	}
 
-	FORCEINLINE_DEBUGGABLE void AddRun(int32 Lod, const FClusterNode& Node) const
+	FORCEINLINE_DEBUGGABLE void AddRun(int32 MinLod, int32 MaxLod, const FClusterNode& Node) const
 	{
-		AddRun(Lod, Node.FirstInstance, Node.LastInstance);
+		AddRun(MinLod, MaxLod, Node.FirstInstance, Node.LastInstance);
 	}
 };
 
@@ -670,8 +699,9 @@ struct FFoliageCullInstanceParams : public FFoliageRenderInstanceParams
 	float LODPlanesMin[MAX_STATIC_MESH_LODS];
 
 
-	FFoliageCullInstanceParams(const TArray<FClusterNode>& InTree)
-	:	Tree(InTree)
+	FFoliageCullInstanceParams(bool InbNeedsSingleLODRuns, bool InbNeedsMultipleLODRuns, const TArray<FClusterNode>& InTree)
+	:	FFoliageRenderInstanceParams(InbNeedsSingleLODRuns, InbNeedsMultipleLODRuns)
+	,	Tree(InTree)
 	{
 	}
 };
@@ -815,10 +845,7 @@ void FHierarchicalStaticMeshSceneProxy::Traverse(const FFoliageCullInstanceParam
 	if (!bSplit)
 	{
 		MaxLOD = FMath::Min(MaxLOD, Params.LODs - 1);
-		for (int32 LODIndex = MinLOD; LODIndex <= MaxLOD; LODIndex++)
-		{
-			Params.AddRun(LODIndex, Node);
-		}
+		Params.AddRun(MinLOD, MaxLOD, Node);
 		return;
 	}
 	for (int32 ChildIndex = Node.FirstChild; ChildIndex <= Node.LastChild; ChildIndex++)
@@ -847,29 +874,35 @@ void FHierarchicalStaticMeshSceneProxy::FillDynamicMeshElements(FMeshElementColl
 	int64 TotalTriangles = 0;
 	for (int32 LODIndex = 0; LODIndex < InstancedRenderData.VertexFactories.Num(); LODIndex++)
 	{
-		if (!Params.Runs[LODIndex].Num())
-		{
-			continue;
-		}
 		const FStaticMeshLODResources& LODModel = StaticMesh->RenderData->LODResources[LODIndex];
 
 		for (int32 SelectionGroupIndex = 0; SelectionGroupIndex < ElementParams.NumSelectionGroups; SelectionGroupIndex++)
 		{
 			for (int32 SectionIndex = 0; SectionIndex < LODModel.Sections.Num(); SectionIndex++)
 			{
+				const FLODInfo& ProxyLODInfo = LODs[LODIndex];
+				UMaterialInterface* Material = ProxyLODInfo.Sections[SectionIndex].Material;
+				const bool bDitherLODEnabled = ElementParams.bBlendLODs;
+
+				TArray<uint32, SceneRenderingAllocator>& RunArray = bDitherLODEnabled ? Params.MultipleLODRuns[LODIndex] : Params.SingleLODRuns[LODIndex];
+
+				if (!RunArray.Num())
+				{
+					continue;
+				}
 
 				int32 NumBatches = 1;
 				int32 CurrentRun = 0;
 				int32 CurrentInstance = 0;
-				int32 RemainingInstances = Params.TotalInstances[LODIndex];
+				int32 RemainingInstances = bDitherLODEnabled ? Params.TotalMultipleLODInstances[LODIndex] : Params.TotalSingleLODInstances[LODIndex];
 
 				if (!ElementParams.bInstanced)
 				{
 					NumBatches = FMath::DivideAndRoundUp(RemainingInstances, (int32)FInstancedStaticMeshVertexFactory::NumBitsForVisibilityMask());
 					if (NumBatches)
 					{
-						check(Params.Runs[LODIndex].Num());
-						CurrentInstance = Params.Runs[LODIndex][CurrentRun];
+						check(RunArray.Num());
+						CurrentInstance = RunArray[CurrentRun];
 					}
 				}
 
@@ -900,10 +933,12 @@ void FHierarchicalStaticMeshSceneProxy::FillDynamicMeshElements(FMeshElementColl
 					BatchElement0.MaxScreenSize = 1.0;
 					BatchElement0.MinScreenSize = 0.0;
 					BatchElement0.InstancedLODIndex = LODIndex;
-					BatchElement0.InstancedLODRange = ElementParams.bBlendLODs ? 1 : 0;
+					BatchElement0.InstancedLODRange = bDitherLODEnabled ? 1 : 0;
 					MeshElement.bCanApplyViewModeOverrides = true;
 					MeshElement.bUseSelectionOutline = ElementParams.BatchRenderSelection[SelectionGroupIndex];
 					MeshElement.bUseWireframeSelectionColoring = ElementParams.BatchRenderSelection[SelectionGroupIndex];
+					MeshElement.bUseAsOccluder = false;
+
 					if (!bDidStats)
 					{
 						bDidStats = true;
@@ -912,8 +947,8 @@ void FHierarchicalStaticMeshSceneProxy::FillDynamicMeshElements(FMeshElementColl
 					}
 					if (ElementParams.bInstanced)
 					{
-						BatchElement0.NumInstances = Params.Runs[LODIndex].Num() / 2;
-						BatchElement0.InstanceRuns = &Params.Runs[LODIndex][0];
+						BatchElement0.NumInstances = RunArray.Num() / 2;
+						BatchElement0.InstanceRuns = &RunArray[0];
 #if STATS
 						INC_DWORD_STAT_BY(STAT_FoliageRuns, BatchElement0.NumInstances);
 #endif
@@ -942,11 +977,11 @@ void FHierarchicalStaticMeshSceneProxy::FillDynamicMeshElements(FMeshElementColl
 							NewBatchElement->UserIndex = CurrentInstance;
 							if (--RemainingInstances)
 							{
-								if ((uint32)CurrentInstance >= Params.Runs[LODIndex][CurrentRun + 1])
+								if ((uint32)CurrentInstance >= RunArray[CurrentRun + 1])
 								{
 									CurrentRun += 2;
-									check(CurrentRun + 1 < Params.Runs[LODIndex].Num());
-									CurrentInstance = Params.Runs[LODIndex][CurrentRun];
+									check(CurrentRun + 1 < RunArray.Num());
+									CurrentInstance = RunArray[CurrentRun];
 								}
 								else
 								{
@@ -976,6 +1011,9 @@ void FHierarchicalStaticMeshSceneProxy::GetDynamicMeshElements(const TArray<cons
 {
 	QUICK_SCOPE_CYCLE_COUNTER(STAT_HierarchicalInstancedStaticMeshSceneProxy_GetMeshElements);
 
+	bool bMultipleSections = ALLOW_DITHERED_LOD_FOR_INSTANCED_STATIC_MESHES && CVarDitheredLOD.GetValueOnRenderThread() > 0;
+	bool bSingleSections = !bMultipleSections;
+
 	for (int32 ViewIndex = 0; ViewIndex < Views.Num(); ViewIndex++)
 	{
 		if (VisibilityMap & (1 << ViewIndex))
@@ -997,7 +1035,8 @@ void FHierarchicalStaticMeshSceneProxy::GetDynamicMeshElements(const TArray<cons
 
 			// Render built instances
 			{
-				FFoliageCullInstanceParams InstanceParams(ClusterTree);
+
+				FFoliageCullInstanceParams InstanceParams(bSingleSections, bMultipleSections, ClusterTree);
 				InstanceParams.LODs = RenderData->LODResources.Num();
 
 				InstanceParams.View = View;
@@ -1080,10 +1119,11 @@ void FHierarchicalStaticMeshSceneProxy::GetDynamicMeshElements(const TArray<cons
 					InstanceParams.ViewFrustumGuardBandLocal.Init();
 				}
 
-				InstanceParams.ViewOriginInLocalZero = WorldToLocal.TransformPosition(View->GetTemporalLODOrigin(0));
-				InstanceParams.ViewOriginInLocalOne = WorldToLocal.TransformPosition(View->GetTemporalLODOrigin(1));
+				ElementParams.bBlendLODs = bMultipleSections;
+
+				InstanceParams.ViewOriginInLocalZero = WorldToLocal.TransformPosition(View->GetTemporalLODOrigin(0, bMultipleSections));
+				InstanceParams.ViewOriginInLocalOne = WorldToLocal.TransformPosition(View->GetTemporalLODOrigin(1, bMultipleSections));
 
-				ElementParams.bBlendLODs = View->GetTemporalLODActive();
 				float MinSize = CVarFoliageMinimumScreenSize.GetValueOnRenderThread();
 				float LODScale = CVarFoliageLODDistanceScale.GetValueOnRenderThread();
 				float MaxDrawDistanceScale = GetCachedScalabilityCVars().ViewDistanceScale;
@@ -1097,7 +1137,7 @@ void FHierarchicalStaticMeshSceneProxy::GetDynamicMeshElements(const TArray<cons
 				for (int32 SampleIndex = 0; SampleIndex < 2; SampleIndex++)
 				{
 					// we aren't going to accurately figure out LOD planes for both samples and try all 4 possibilities, rather we just consider the worst case LOD range; the shaders will render correctly
-					float Fac = View->GetTemporalLODDistanceFactor(SampleIndex) * SphereRadius * SphereRadius * LODScale * LODScale;
+					float Fac = View->GetTemporalLODDistanceFactor(SampleIndex, bMultipleSections) * SphereRadius * SphereRadius * LODScale * LODScale;
 
 					float FinalCull = MAX_flt;
 
@@ -1127,11 +1167,16 @@ void FHierarchicalStaticMeshSceneProxy::GetDynamicMeshElements(const TArray<cons
 				{
 					for (int32 LODIndex = 0; LODIndex < InstanceParams.LODs; LODIndex++)
 					{
-						for (int32 Run = 0; Run < DebugRuns[LODIndex].Num(); Run++)
+						for (int32 Run = 0; Run < SingleDebugRuns[LODIndex].Num(); Run++)
+						{
+							InstanceParams.SingleLODRuns[LODIndex].Add(SingleDebugRuns[LODIndex][Run]);
+						}
+						InstanceParams.TotalSingleLODInstances[LODIndex] = SingleDebugTotalInstances[LODIndex];
+						for (int32 Run = 0; Run < MultipleDebugRuns[LODIndex].Num(); Run++)
 						{
-							InstanceParams.Runs[LODIndex].Add(DebugRuns[LODIndex][Run]);
+							InstanceParams.MultipleLODRuns[LODIndex].Add(MultipleDebugRuns[LODIndex][Run]);
 						}
-						InstanceParams.TotalInstances[LODIndex] = DebugTotalInstances[LODIndex];
+						InstanceParams.TotalMultipleLODInstances[LODIndex] = MultipleDebugTotalInstances[LODIndex];
 					}
 				}
 				else
@@ -1170,11 +1215,17 @@ void FHierarchicalStaticMeshSceneProxy::GetDynamicMeshElements(const TArray<cons
 						CaptureTag = GDebugTag;
 						for (int32 LODIndex = 0; LODIndex < InstanceParams.LODs; LODIndex++)
 						{
-							DebugRuns[LODIndex].Empty();
-							DebugTotalInstances[LODIndex] = InstanceParams.TotalInstances[LODIndex];
-							for (int32 Run = 0; Run < InstanceParams.Runs[LODIndex].Num(); Run++)
+							SingleDebugRuns[LODIndex].Empty();
+							SingleDebugTotalInstances[LODIndex] = InstanceParams.TotalSingleLODInstances[LODIndex];
+							for (int32 Run = 0; Run < InstanceParams.SingleLODRuns[LODIndex].Num(); Run++)
+							{
+								SingleDebugRuns[LODIndex].Add(InstanceParams.SingleLODRuns[LODIndex][Run]);
+							}
+							MultipleDebugRuns[LODIndex].Empty();
+							MultipleDebugTotalInstances[LODIndex] = InstanceParams.TotalMultipleLODInstances[LODIndex];
+							for (int32 Run = 0; Run < InstanceParams.MultipleLODRuns[LODIndex].Num(); Run++)
 							{
-								DebugRuns[LODIndex].Add(InstanceParams.Runs[LODIndex][Run]);
+								MultipleDebugRuns[LODIndex].Add(InstanceParams.MultipleLODRuns[LODIndex][Run]);
 							}
 						}
 					}
@@ -1187,7 +1238,7 @@ void FHierarchicalStaticMeshSceneProxy::GetDynamicMeshElements(const TArray<cons
 			// Render unbuilt instances
 			if (FirstUnbuiltIndex <= LastUnbuiltIndex)
 			{
-				FFoliageRenderInstanceParams InstanceParams;
+				FFoliageRenderInstanceParams InstanceParams(true, false);
 
 				// disable LOD blending for unbuilt instances as we haven't calculated the correct LOD.
 				ElementParams.bBlendLODs = false;
@@ -1195,12 +1246,12 @@ void FHierarchicalStaticMeshSceneProxy::GetDynamicMeshElements(const TArray<cons
 				if (LastUnbuiltIndex - FirstUnbuiltIndex < 1000)
 				{
 					// less than 1000 instances, just add them all at LOD 0
-					InstanceParams.AddRun(0, FirstUnbuiltIndex, LastUnbuiltIndex);
+					InstanceParams.AddRun(0, 0, FirstUnbuiltIndex, LastUnbuiltIndex);
 				}
 				else
 				{
 					// more than 1000, render them all at Max LOD
-					InstanceParams.AddRun(RenderData->LODResources.Num() - 1, FirstUnbuiltIndex, FMath::Min<int32>(LastUnbuiltIndex - FirstUnbuiltIndex + 10000, LastUnbuiltIndex));
+					InstanceParams.AddRun(RenderData->LODResources.Num() - 1, RenderData->LODResources.Num() - 1, FirstUnbuiltIndex, FMath::Min<int32>(LastUnbuiltIndex - FirstUnbuiltIndex + 10000, LastUnbuiltIndex));
 				}
 				FillDynamicMeshElements(Collector, ElementParams, InstanceParams);
 			}
diff --git a/Engine/Source/Runtime/Engine/Private/InstancedStaticMesh.cpp b/Engine/Source/Runtime/Engine/Private/InstancedStaticMesh.cpp
--- a/Engine/Source/Runtime/Engine/Private/InstancedStaticMesh.cpp
+++ b/Engine/Source/Runtime/Engine/Private/InstancedStaticMesh.cpp
@@ -224,6 +224,7 @@ void FInstancedStaticMeshVertexFactory::InitRHI()
 	check(HasValidFeatureLevel());
 	const bool bInstanced = RHISupportsInstancing(GetFeatureLevelShaderPlatform(GetFeatureLevel()));
 
+#if !ALLOW_DITHERED_LOD_FOR_INSTANCED_STATIC_MESHES // position only shaders cannot work with dithered LOD
 	// If the vertex buffer containing position is not the same vertex buffer containing the rest of the data,
 	// then initialize PositionStream and PositionDeclaration.
 	if(Data.PositionComponent.VertexBuffer != Data.TangentBasisComponents[0].VertexBuffer)
@@ -240,6 +241,7 @@ void FInstancedStaticMeshVertexFactory::InitRHI()
 		}
 		InitPositionDeclaration(PositionOnlyStreamElements);
 	}
+#endif
 
 	FVertexDeclarationElementList Elements;
 	if(Data.PositionComponent.VertexBuffer != NULL)
@@ -1459,7 +1461,6 @@ void FInstancedStaticMeshVertexFactoryShaderParameters::SetMesh( FRHICommandList
 
 	FRHIVertexShader* VS = VertexShader->GetVertexShader();
 	const FInstancingUserData* InstancingUserData = (const FInstancingUserData*)BatchElement.UserData;
-#if USE_NEW_LOD_TRANSITION_FOR_FOLIAGE
 	if( InstancingWorldViewOriginOneParameter.IsBound() )
 	{
 		FVector4 InstancingViewZCompareZero(MIN_flt, MIN_flt, MAX_flt, MAX_flt);
@@ -1468,7 +1469,7 @@ void FInstancedStaticMeshVertexFactoryShaderParameters::SetMesh( FRHICommandList
 		FVector4 InstancingWorldViewOriginZero(ForceInit);
 		FVector4 InstancingWorldViewOriginOne(ForceInit);
 		InstancingWorldViewOriginOne.W = 1.0f;
-		if (InstancingUserData && BatchElement.InstancedLODRange && View.GetTemporalLODActive())
+		if (InstancingUserData && BatchElement.InstancedLODRange)
 		{
 			float SphereRadius = InstancingUserData->MeshRenderData->Bounds.SphereRadius;
 			float MinSize = CVarFoliageMinimumScreenSize.GetValueOnRenderThread();
@@ -1552,13 +1553,11 @@ void FInstancedStaticMeshVertexFactoryShaderParameters::SetMesh( FRHICommandList
 		SetShaderValue(RHICmdList, VS, InstancingWorldViewOriginZeroParameter, InstancingWorldViewOriginZero );
 		SetShaderValue(RHICmdList, VS, InstancingWorldViewOriginOneParameter, InstancingWorldViewOriginOne );
 	}
-#endif
 	if( InstancingFadeOutParamsParameter.IsBound() )
 	{
 		FVector4 InstancingFadeOutParams(MAX_flt,0.f,1.f,1.f);
 		if (InstancingUserData)
 		{
-#if !USE_NEW_LOD_TRANSITION_FOR_FOLIAGE
 			InstancingFadeOutParams.X = InstancingUserData->StartCullDistance;
 			if( InstancingUserData->EndCullDistance > 0 )
 			{
@@ -1575,7 +1574,6 @@ void FInstancedStaticMeshVertexFactoryShaderParameters::SetMesh( FRHICommandList
 			{
 				InstancingFadeOutParams.Y = 0.f;
 			}
-#endif
 			InstancingFadeOutParams.Z = InstancingUserData->bRenderSelected ? 1.f : 0.f;
 			InstancingFadeOutParams.W = InstancingUserData->bRenderUnselected ? 1.f : 0.f;
 		}
diff --git a/Engine/Source/Runtime/Engine/Private/InstancedStaticMesh.h b/Engine/Source/Runtime/Engine/Private/InstancedStaticMesh.h
--- a/Engine/Source/Runtime/Engine/Private/InstancedStaticMesh.h
+++ b/Engine/Source/Runtime/Engine/Private/InstancedStaticMesh.h
@@ -56,8 +56,6 @@
 
 DECLARE_STATS_GROUP(TEXT("Foliage"), STATGROUP_Foliage, STATCAT_Advanced);
 
-#define USE_NEW_LOD_TRANSITION_FOR_FOLIAGE (1)
-
 extern TAutoConsoleVariable<float> CVarFoliageMinimumScreenSize;
 extern TAutoConsoleVariable<float> CVarFoliageLODDistanceScale;
 
@@ -247,7 +245,7 @@ public:
 		OutEnvironment.SetDefine(TEXT("USE_INSTANCING"),TEXT("1"));
 		const bool bInstanced = RHISupportsInstancing(Platform);
 		OutEnvironment.SetDefine(TEXT("USE_INSTANCING_EMULATED"), bInstanced ? TEXT("0") : TEXT("1"));
-		OutEnvironment.SetDefine(TEXT("USE_NEW_LOD_TRANSITION"), USE_NEW_LOD_TRANSITION_FOR_FOLIAGE ? TEXT("1") : TEXT("0"));
+		OutEnvironment.SetDefine(TEXT("USE_DITHERED_LOD_TRANSITION"), ALLOW_DITHERED_LOD_FOR_INSTANCED_STATIC_MESHES ? TEXT("1") : TEXT("0"));
 	}
 
 	/**
@@ -284,6 +282,9 @@ public:
 		uint32 NumElements = FMath::Min((uint32)Batch->Elements.Num(), NumBitsForVisibilityMask());
 		return (1ULL << (uint64)NumElements) - 1ULL;
 	}
+#if ALLOW_DITHERED_LOD_FOR_INSTANCED_STATIC_MESHES
+	virtual bool SupportsNullPixelShader() const override { return false; }
+#endif
 
 private:
 	DataType Data;
@@ -299,13 +300,11 @@ class FInstancedStaticMeshVertexFactoryShaderParameters : public FLocalVertexFac
 		FLocalVertexFactoryShaderParameters::Bind(ParameterMap);
 
 		InstancingFadeOutParamsParameter.Bind(ParameterMap, TEXT("InstancingFadeOutParams"));
-#if USE_NEW_LOD_TRANSITION_FOR_FOLIAGE
 		InstancingViewZCompareZeroParameter.Bind(ParameterMap, TEXT("InstancingViewZCompareZero"));
 		InstancingViewZCompareOneParameter.Bind(ParameterMap, TEXT("InstancingViewZCompareOne"));
 		InstancingViewZConstantParameter.Bind(ParameterMap, TEXT("InstancingViewZConstant"));
 		InstancingWorldViewOriginZeroParameter.Bind(ParameterMap, TEXT("InstancingWorldViewOriginZero"));
 		InstancingWorldViewOriginOneParameter.Bind(ParameterMap, TEXT("InstancingWorldViewOriginOne"));
-#endif
 		CPUInstanceShadowMapBias.Bind(ParameterMap, TEXT("CPUInstanceShadowMapBias"));
 		CPUInstanceTransform.Bind(ParameterMap, TEXT("CPUInstanceTransform"));
 		CPUInstanceInverseTransform.Bind(ParameterMap, TEXT("CPUInstanceInverseTransform"));
@@ -317,13 +316,11 @@ class FInstancedStaticMeshVertexFactoryShaderParameters : public FLocalVertexFac
 	{
 		FLocalVertexFactoryShaderParameters::Serialize(Ar);
 		Ar << InstancingFadeOutParamsParameter;
-#if USE_NEW_LOD_TRANSITION_FOR_FOLIAGE
 		Ar << InstancingViewZCompareZeroParameter;
 		Ar << InstancingViewZCompareOneParameter;
 		Ar << InstancingViewZConstantParameter;
 		Ar << InstancingWorldViewOriginZeroParameter;
 		Ar << InstancingWorldViewOriginOneParameter;
-#endif
 		Ar << CPUInstanceShadowMapBias;
 		Ar << CPUInstanceTransform;
 		Ar << CPUInstanceInverseTransform;
@@ -333,13 +330,11 @@ class FInstancedStaticMeshVertexFactoryShaderParameters : public FLocalVertexFac
 
 private:
 	FShaderParameter InstancingFadeOutParamsParameter;
-#if USE_NEW_LOD_TRANSITION_FOR_FOLIAGE
 	FShaderParameter InstancingViewZCompareZeroParameter;
 	FShaderParameter InstancingViewZCompareOneParameter;
 	FShaderParameter InstancingViewZConstantParameter;
 	FShaderParameter InstancingWorldViewOriginZeroParameter;
 	FShaderParameter InstancingWorldViewOriginOneParameter;
-#endif
 
 	FShaderParameter CPUInstanceShadowMapBias;
 	FShaderParameter CPUInstanceTransform;
diff --git a/Engine/Source/Runtime/Engine/Private/SceneView.cpp b/Engine/Source/Runtime/Engine/Private/SceneView.cpp
--- a/Engine/Source/Runtime/Engine/Private/SceneView.cpp
+++ b/Engine/Source/Runtime/Engine/Private/SceneView.cpp
@@ -435,19 +435,6 @@ FSceneView::FSceneView(const FSceneViewInitOptions& InitOptions)
 #endif
 }
 
-bool FSceneView::GetTemporalLODActive() const
-{
-	if (State)
-	{
-		const FTemporalLODState& LODState = State->GetTemporalLODState();
-		if (LODState.TemporalLODLag != 0.0f)
-		{
-			return true;
-		}
-	}
-	return false;
-}
-
 float FSceneView::GetLODDistanceFactor() const
 {
 	const float ScreenMultiple = FMath::Max(ViewRect.Width() / 2.0f * ViewMatrices.ProjMatrix.M[0][0],
@@ -456,9 +443,9 @@ float FSceneView::GetLODDistanceFactor() const
 	return Fac;
 }
 
-float FSceneView::GetTemporalLODDistanceFactor(int32 Index) const
+float FSceneView::GetTemporalLODDistanceFactor(int32 Index, bool bUseLaggedLODTransition) const
 {
-	if (State)
+	if (bUseLaggedLODTransition && State)
 	{
 		const FTemporalLODState& LODState = State->GetTemporalLODState();
 		if (LODState.TemporalLODLag != 0.0f)
@@ -470,9 +457,9 @@ float FSceneView::GetTemporalLODDistanceFactor(int32 Index) const
 }
 
 
-FVector FSceneView::GetTemporalLODOrigin(int32 Index) const
+FVector FSceneView::GetTemporalLODOrigin(int32 Index, bool bUseLaggedLODTransition) const
 {
-	if (State)
+	if (bUseLaggedLODTransition && State)
 	{
 		const FTemporalLODState& LODState = State->GetTemporalLODState();
 		if (LODState.TemporalLODLag != 0.0f)
diff --git a/Engine/Source/Runtime/Engine/Public/MaterialShared.h b/Engine/Source/Runtime/Engine/Public/MaterialShared.h
--- a/Engine/Source/Runtime/Engine/Public/MaterialShared.h
+++ b/Engine/Source/Runtime/Engine/Public/MaterialShared.h
@@ -42,9 +42,10 @@ struct FShaderCompilerEnvironment;
 #define ME_STD_LABEL_PAD		16
 #define ME_STD_TAB_HEIGHT		21
 
-DECLARE_LOG_CATEGORY_EXTERN(LogMaterial,Log,Verbose);
 
+#define ALLOW_DITHERED_LOD_FOR_INSTANCED_STATIC_MESHES (1)
 
+DECLARE_LOG_CATEGORY_EXTERN(LogMaterial,Log,Verbose);
 
 /** Creates a string that represents the given quality level. */
 extern void GetMaterialQualityLevelName(EMaterialQualityLevel::Type InMaterialQualityLevel, FString& OutName);
diff --git a/Engine/Source/Runtime/Engine/Public/SceneView.h b/Engine/Source/Runtime/Engine/Public/SceneView.h
--- a/Engine/Source/Runtime/Engine/Public/SceneView.h
+++ b/Engine/Source/Runtime/Engine/Public/SceneView.h
@@ -540,14 +540,11 @@ public:
 	/** @return true:perspective, false:orthographic */
 	inline bool IsPerspectiveProjection() const { return ViewMatrices.IsPerspectiveProjection(); }
 
-	/** Returns true if temporal LOD is active **/
-	bool GetTemporalLODActive() const;
-
 	/** Returns the location used as the origin for LOD computations
 	 * @param Index, 0 or 1, which LOD origin to return
 	 * @return LOD origin
 	 */
-	FVector GetTemporalLODOrigin(int32 Index) const;
+	FVector GetTemporalLODOrigin(int32 Index, bool bUseLaggedLODTransition = true) const;
 
 	/** Get LOD distance factor: Sqrt(GetLODDistanceFactor()*SphereRadius*SphereRadius / ScreenPercentage) = distance to this LOD transition
 	 * @return distance factor
@@ -558,7 +555,7 @@ public:
 	 * @param Index, 0 or 1, which temporal sample to return
 	 * @return distance factor
 	 */
-	float GetTemporalLODDistanceFactor(int32 Index) const;
+	float GetTemporalLODDistanceFactor(int32 Index, bool bUseLaggedLODTransition = true) const;
 
 	/** 
 	 * Returns the blend factor between the last two LOD samples
diff --git a/Engine/Source/Runtime/Renderer/Private/ShadowRendering.cpp b/Engine/Source/Runtime/Renderer/Private/ShadowRendering.cpp
--- a/Engine/Source/Runtime/Renderer/Private/ShadowRendering.cpp
+++ b/Engine/Source/Runtime/Renderer/Private/ShadowRendering.cpp
@@ -839,7 +839,7 @@ FShadowDepthDrawingPolicy<bRenderingReflectiveShadowMaps>::FShadowDepthDrawingPo
 	}
 
 	// Pixel shaders
-	if (!MaterialResource->IsMasked() && !bUsePerspectiveCorrectShadowDepths && !bRenderingReflectiveShadowMaps)
+	if (!MaterialResource->IsMasked() && !bUsePerspectiveCorrectShadowDepths && !bRenderingReflectiveShadowMaps && VertexFactory->SupportsNullPixelShader())
 	{
 		// No pixel shader necessary.
 		PixelShader = NULL;
diff --git a/Engine/Source/Runtime/ShaderCore/Public/VertexFactory.h b/Engine/Source/Runtime/ShaderCore/Public/VertexFactory.h
--- a/Engine/Source/Runtime/ShaderCore/Public/VertexFactory.h
+++ b/Engine/Source/Runtime/ShaderCore/Public/VertexFactory.h
@@ -394,6 +394,9 @@ public:
 	/** Indicates whether the vertex factory supports a position-only stream. */
 	bool SupportsPositionOnlyStream() const { return !!PositionStream.Num(); }
 
+	/** Indicates whether the vertex factory supports a null pixel shader. */
+	virtual bool SupportsNullPixelShader() const { return true; }
+
 	/**
 	* Fill in array of strides from this factory's vertex streams without shadow/light maps
 	* @param OutStreamStrides - output array of # MaxVertexElementCount stream strides to fill
