Fixed a crash in clothing asset release
 - Removed ClothingAssetWrapper class

diff --git a/Engine/Source/Editor/Persona/Private/PersonaMeshDetails.cpp b/Engine/Source/Editor/Persona/Private/PersonaMeshDetails.cpp
--- a/Engine/Source/Editor/Persona/Private/PersonaMeshDetails.cpp
+++ b/Engine/Source/Editor/Persona/Private/PersonaMeshDetails.cpp
@@ -1383,7 +1383,7 @@ void FPersonaMeshDetails::OnGenerateElementForClothingAsset( TSharedRef<IPropert
 	USkeletalMesh* SkelMesh = PersonaPtr->GetMesh();
 	FClothingAssetData& AssetData = SkelMesh->ClothingAssets[ElementIndex];
 
-	ApexClothingUtils::GetPhysicsPropertiesFromApexAsset(AssetData.ApexClothingAsset->GetAsset(), AssetData.PhysicsProperties);
+	ApexClothingUtils::GetPhysicsPropertiesFromApexAsset(AssetData.ApexClothingAsset, AssetData.PhysicsProperties);
 
 	// cloth physics properties
 	TSharedRef<IPropertyHandle> ClothPhysicsProperties = StructProperty->GetChildHandle(FName("PhysicsProperties")).ToSharedRef();
@@ -1422,7 +1422,7 @@ TSharedRef<SUniformGridPanel> FPersonaMeshDetails::MakeApexDetailsWidget(int32 A
 		return Grid;
 	}
 
-	int32 NumLODs = ApexClothingUtils::GetNumLODs(Asset.ApexClothingAsset->GetAsset());
+	int32 NumLODs = ApexClothingUtils::GetNumLODs(Asset.ApexClothingAsset);
 	int32 RowNumber = 0;
 
 	for(int32 LODIndex=0; LODIndex < NumLODs; LODIndex++)
@@ -1438,7 +1438,7 @@ TSharedRef<SUniformGridPanel> FPersonaMeshDetails::MakeApexDetailsWidget(int32 A
 		RowNumber++;
 
 		TArray<FSubmeshInfo> SubmeshInfos;
-		if(ApexClothingUtils::GetSubmeshInfoFromApexAsset(Asset.ApexClothingAsset->GetAsset(), LODIndex, SubmeshInfos))
+		if(ApexClothingUtils::GetSubmeshInfoFromApexAsset(Asset.ApexClothingAsset, LODIndex, SubmeshInfos))
 		{
 			// content names
 			Grid->AddSlot(0, RowNumber) // x, y
@@ -1577,7 +1577,7 @@ FReply FPersonaMeshDetails::OnReimportApexFileClicked(int32 AssetIndex, IDetailL
 		if (EAppReturnType::Yes == Ret)
 		{
 			bNeedToLeaveProperties = true;
-			ApexClothingUtils::GetPhysicsPropertiesFromApexAsset(SkelMesh->ClothingAssets[AssetIndex].ApexClothingAsset->GetAsset(), CurClothPhysicsProperties);
+			ApexClothingUtils::GetPhysicsPropertiesFromApexAsset(SkelMesh->ClothingAssets[AssetIndex].ApexClothingAsset, CurClothPhysicsProperties);
 		}
 		else if (EAppReturnType::Cancel == Ret)
 		{
@@ -1593,7 +1593,7 @@ FReply FPersonaMeshDetails::OnReimportApexFileClicked(int32 AssetIndex, IDetailL
 		{
 			// overwrites changed values instead of original values
 			SkelMesh->ClothingAssets[AssetIndex].bClothPropertiesChanged = true;
-			ApexClothingUtils::SetPhysicsPropertiesToApexAsset(SkelMesh->ClothingAssets[AssetIndex].ApexClothingAsset->GetAsset(), CurClothPhysicsProperties);
+			ApexClothingUtils::SetPhysicsPropertiesToApexAsset(SkelMesh->ClothingAssets[AssetIndex].ApexClothingAsset, CurClothPhysicsProperties);
 		}
 		else
 		{
@@ -1671,7 +1671,7 @@ FReply FPersonaMeshDetails::OnOpenClothingFileClicked(IDetailLayoutBuilder* Deta
 				{
 					int32 AssetIndex = SkelMesh->ClothingAssets.Num()-1;
 					FClothingAssetData& AssetData = SkelMesh->ClothingAssets[AssetIndex];
-					int32 NumLODs = ApexClothingUtils::GetNumLODs(AssetData.ApexClothingAsset->GetAsset());
+					int32 NumLODs = ApexClothingUtils::GetNumLODs(AssetData.ApexClothingAsset);
 
 					uint32 MaxClothVertices = ApexClothingUtils::GetMaxClothSimulVertices(GMaxRHIFeatureLevel);
 
@@ -1682,7 +1682,7 @@ FReply FPersonaMeshDetails::OnOpenClothingFileClicked(IDetailLayoutBuilder* Deta
 					for(int32 LODIndex=0; LODIndex < NumLODs; LODIndex++)
 					{
 						TArray<FSubmeshInfo> SubmeshInfos;
-						if(ApexClothingUtils::GetSubmeshInfoFromApexAsset(AssetData.ApexClothingAsset->GetAsset(), LODIndex, SubmeshInfos))
+						if(ApexClothingUtils::GetSubmeshInfoFromApexAsset(AssetData.ApexClothingAsset, LODIndex, SubmeshInfos))
 						{
 							for(int32 SubIndex=0; SubIndex < SubmeshInfos.Num(); SubIndex++)
 							{
@@ -1817,7 +1817,7 @@ void FPersonaMeshDetails::UpdateComboBoxStrings()
 
 			TArray<FSubmeshInfo> SubmeshInfos;
 			// if failed to get sub-mesh info, then skip
-			if (!ApexClothingUtils::GetSubmeshInfoFromApexAsset(ClothingAssetData.ApexClothingAsset->GetAsset(), LODIdx, SubmeshInfos))
+			if (!ApexClothingUtils::GetSubmeshInfoFromApexAsset(ClothingAssetData.ApexClothingAsset, LODIdx, SubmeshInfos))
 			{
 				continue;
 			}
@@ -2021,7 +2021,7 @@ void FPersonaMeshDetails::UpdateClothPhysicsProperties(int32 AssetIndex)
 	USkeletalMesh* SkelMesh = PersonaPtr->GetMesh();
 	FClothingAssetData& Asset = SkelMesh->ClothingAssets[AssetIndex];
 
-	ApexClothingUtils::SetPhysicsPropertiesToApexAsset(Asset.ApexClothingAsset->GetAsset(), Asset.PhysicsProperties);
+	ApexClothingUtils::SetPhysicsPropertiesToApexAsset(Asset.ApexClothingAsset, Asset.PhysicsProperties);
 	Asset.bClothPropertiesChanged = true;
 }
 
diff --git a/Engine/Source/Runtime/Engine/Classes/Engine/SkeletalMesh.h b/Engine/Source/Runtime/Engine/Classes/Engine/SkeletalMesh.h
--- a/Engine/Source/Runtime/Engine/Classes/Engine/SkeletalMesh.h
+++ b/Engine/Source/Runtime/Engine/Classes/Engine/SkeletalMesh.h
@@ -351,43 +351,6 @@ struct FMorphTargetMap
 	}
 };
 
-#if WITH_APEX_CLOTHING
-class FClothingAssetWrapper
-{
-public:
-
-	FClothingAssetWrapper(physx::apex::NxClothingAsset* InApexClothingAsset)
-	:	ApexClothingAsset(InApexClothingAsset),bValid(true)
-	{}
-
-	ENGINE_API ~FClothingAssetWrapper();
-
-	physx::apex::NxClothingAsset* GetAsset()
-	{
-		return ApexClothingAsset;
-	}
-
-	//returns bone name converted to fbx style
-	FName GetConvertedBoneName(int32 BoneIndex);
-
-
-	bool	IsValid()
-	{ 
-		return bValid;
-	}
-
-	void	SetValid(bool valid)
-	{
-		bValid = valid;
-	}
-
-private:
-	physx::apex::NxClothingAsset* ApexClothingAsset;
-	bool						  bValid;
-};
-#endif // #if WITH_APEX_CLOTHING
-
-
 /** 
  * constrain Coefficients - max distance, collisionSphere radius, collision sphere distance 
  */
@@ -497,7 +460,7 @@ struct FClothingAssetData
 	FClothPhysicsProperties PhysicsProperties;
 
 #if WITH_APEX_CLOTHING
-	TSharedPtr<FClothingAssetWrapper> ApexClothingAsset;
+	physx::apex::NxClothingAsset* ApexClothingAsset;
 
 	/** Collision volume data for showing to the users whether collision shape is correct or not */
 	TArray<FApexClothCollisionVolumeData> ClothCollisionVolumes;
@@ -513,6 +476,11 @@ struct FClothingAssetData
 	TArray<FClothVisualizationInfo> ClothVisualizationInfos;
 	/** currently mapped morph target name */
 	FName PreparedMorphTargetName;
+
+	FClothingAssetData()
+		:ApexClothingAsset(NULL)
+	{
+	}
 #endif// #if WITH_APEX_CLOTHING
 
 	// serialization
diff --git a/Engine/Source/Runtime/Engine/Private/SkeletalMesh.cpp b/Engine/Source/Runtime/Engine/Private/SkeletalMesh.cpp
--- a/Engine/Source/Runtime/Engine/Private/SkeletalMesh.cpp
+++ b/Engine/Source/Runtime/Engine/Private/SkeletalMesh.cpp
@@ -1781,23 +1781,6 @@ void FSkelMeshChunk::CalcMaxBoneInfluences()
 	}
 }
 
-#if WITH_APEX_CLOTHING
-/*-----------------------------------------------------------------------------
-	FClothingAssetWrapper
------------------------------------------------------------------------------*/
-
-FClothingAssetWrapper::~FClothingAssetWrapper()
-{
-	check(ApexClothingAsset);
-	GPhysCommandHandler->DeferredRelease(ApexClothingAsset);
-}
-
-FName FClothingAssetWrapper::GetConvertedBoneName(int32 BoneIndex)
-{
-	return *FString(ApexClothingAsset->getBoneName(BoneIndex)).Replace(TEXT(" "),TEXT("-"));
-}
-
-#endif// #if WITH_APEX_CLOTHING
 /*-----------------------------------------------------------------------------
 	FClothingAssetData
 -----------------------------------------------------------------------------*/
@@ -1817,7 +1800,7 @@ FArchive& operator<<(FArchive& Ar, FClothingAssetData& A)
 			Buffer.AddUninitialized( AssetSize );
 			Ar.Serialize( Buffer.GetData(), AssetSize );
 #if WITH_APEX_CLOTHING
-			A.ApexClothingAsset = MakeShareable( new FClothingAssetWrapper(LoadApexClothingAssetFromBlob(Buffer)) );
+			A.ApexClothingAsset = LoadApexClothingAssetFromBlob(Buffer);
 #endif //#if WITH_APEX_CLOTHING
 		}
 	}
@@ -1825,10 +1808,10 @@ FArchive& operator<<(FArchive& Ar, FClothingAssetData& A)
 	if( Ar.IsSaving() )
 	{
 #if WITH_APEX_CLOTHING
-		if( A.ApexClothingAsset->GetAsset() )
+		if (A.ApexClothingAsset)
 		{
 			TArray<uint8> Buffer;
-			SaveApexClothingAssetToBlob(A.ApexClothingAsset->GetAsset(), Buffer);
+			SaveApexClothingAssetToBlob(A.ApexClothingAsset, Buffer);
 			uint32 AssetSize = Buffer.Num();
 			Ar << AssetSize;
 			Ar.Serialize(Buffer.GetData(), AssetSize);
@@ -2276,7 +2259,11 @@ void USkeletalMesh::BeginDestroy()
 	// release clothing assets
 	for (FClothingAssetData& Data : ClothingAssets)
 	{
-		Data.ApexClothingAsset.Reset();
+		if (Data.ApexClothingAsset)
+		{
+			GPhysCommandHandler->DeferredRelease(Data.ApexClothingAsset);
+			Data.ApexClothingAsset = NULL;
+		}
 	}
 #endif // #if WITH_APEX_CLOTHING
 }
@@ -2614,9 +2601,9 @@ void USkeletalMesh::PostLoad()
 	// load clothing section collision
 	for( int32 AssetIdx=0; AssetIdx<ClothingAssets.Num();AssetIdx++ )
 	{
-		if( ClothingAssets[AssetIdx].ApexClothingAsset->GetAsset() )
+		if( ClothingAssets[AssetIdx].ApexClothingAsset )
 		{
-			LoadClothCollisionVolumes(AssetIdx, ClothingAssets[AssetIdx].ApexClothingAsset->GetAsset());
+			LoadClothCollisionVolumes(AssetIdx, ClothingAssets[AssetIdx].ApexClothingAsset);
 		}
 #if WITH_EDITOR
 		// Remove any clothing sections that have invalid APEX data.
diff --git a/Engine/Source/Runtime/Engine/Private/SkeletalMeshComponentPhysics.cpp b/Engine/Source/Runtime/Engine/Private/SkeletalMeshComponentPhysics.cpp
--- a/Engine/Source/Runtime/Engine/Private/SkeletalMeshComponentPhysics.cpp
+++ b/Engine/Source/Runtime/Engine/Private/SkeletalMeshComponentPhysics.cpp
@@ -1894,6 +1894,13 @@ bool USkeletalMeshComponent::ComponentOverlapMulti(TArray<struct FOverlapResult>
 
 
 #if WITH_APEX_CLOTHING
+
+// convert a bone name from APEX stype to FBX style
+static FName GetConvertedBoneName(NxClothingAsset* ApexClothingAsset, int32 BoneIndex)
+{
+	return *FString(ApexClothingAsset->getBoneName(BoneIndex)).Replace(TEXT(" "), TEXT("-"));
+}
+
 void USkeletalMeshComponent::AddClothingBounds(FBoxSphereBounds& InOutBounds) const
 {
 	int32 NumAssets = ClothingActors.Num();
@@ -1994,7 +2001,7 @@ void USkeletalMeshComponent::ValidateClothingActors()
  * APEX clothing actor is created from APEX clothing asset for cloth simulation 
  * If this is invalid, re-create actor , but if valid ,just skip to create
 */
-bool USkeletalMeshComponent::CreateClothingActor(int32 AssetIndex, TSharedPtr<FClothingAssetWrapper> ClothingAssetWrapper, TArray<FVector>* BlendedDelta)
+bool USkeletalMeshComponent::CreateClothingActor(int32 AssetIndex, physx::apex::NxClothingAsset* ClothingAsset, TArray<FVector>* BlendedDelta)
 {	
 	int32 NumActors = ClothingActors.Num();
 	int32 ActorIndex = -1;
@@ -2022,7 +2029,6 @@ bool USkeletalMeshComponent::CreateClothingActor(int32 AssetIndex, TSharedPtr<FC
 		ActorIndex = ClothingActors.AddZeroed();
 	}
 	 
-	NxClothingAsset* ClothingAsset = ClothingAssetWrapper->GetAsset();
 	// Get the (singleton!) default actor descriptor.
 	NxParameterized::Interface* ActorDesc = ClothingAsset->getDefaultActorDesc();
 	PX_ASSERT(ActorDesc != NULL);
@@ -2115,7 +2121,7 @@ bool USkeletalMeshComponent::CreateClothingActor(int32 AssetIndex, TSharedPtr<FC
 	}
 
 	//set parent pointer to verify later whether became invalid or not
-	ClothingActors[ActorIndex].ParentClothingAsset = ClothingAssetWrapper;
+	ClothingActors[ActorIndex].ParentClothingAsset = ClothingAsset;
 
 	// budget is millisecond units
 	ScenePtr->setLODResourceBudget(100); // for temporary, 100ms
@@ -2164,7 +2170,7 @@ void USkeletalMeshComponent::SetClothingLOD(int32 LODIndex)
 				}
 			}
 
-			int32 NumClothLODs = Actor.ParentClothingAsset->GetAsset()->getNumGraphicalLodLevels();
+			int32 NumClothLODs = Actor.ParentClothingAsset->getNumGraphicalLodLevels();
 
 			// decide whether should enable or disable
 			if (!IsMappedClothLOD || (LODIndex >= NumClothLODs))
@@ -2587,7 +2593,7 @@ void USkeletalMeshComponent::FindClothCollisions(TArray<FApexClothCollisionVolum
 				continue;
 			}
 
-			FName BoneName = Asset.ApexClothingAsset->GetConvertedBoneName(Collision.BoneIndex);
+			FName BoneName = GetConvertedBoneName(Asset.ApexClothingAsset, Collision.BoneIndex);
 
 			int32 BoneIndex = GetBoneIndex(BoneName);
 
@@ -3208,7 +3214,7 @@ void USkeletalMeshComponent::ChangeClothMorphTargetMapping(FClothMorphTargetData
 				ClothOriginalPosArray[Index] = U2PVector(MorphData.OriginPos[Index]);
 			}
 
-			NxClothingAsset* ClothingAsset = Asset.ApexClothingAsset->GetAsset();
+			NxClothingAsset* ClothingAsset = Asset.ApexClothingAsset;
 			float Epsilon = 0.0f;
 			uint32 NumMapped = ClothingAsset->prepareMorphTargetMapping(ClothOriginalPosArray.GetData(), NumOriginPos, Epsilon);
 
@@ -3317,7 +3323,7 @@ void USkeletalMeshComponent::PrepareClothMorphTargets()
 			{
 				if (ClothOriginalPosArray[AssetIdx].Num() > 0)
 				{
-					NxClothingAsset* ClothingAsset = SkeletalMesh->ClothingAssets[AssetIdx].ApexClothingAsset->GetAsset();
+					NxClothingAsset* ClothingAsset = SkeletalMesh->ClothingAssets[AssetIdx].ApexClothingAsset;
 					float Epsilon = 0.0f;
 					uint32 NumMapped = ClothingAsset->prepareMorphTargetMapping(ClothOriginalPosArray[AssetIdx].GetData(), ClothOriginalPosArray[AssetIdx].Num(), Epsilon);
 
@@ -3485,7 +3491,7 @@ void USkeletalMeshComponent::UpdateClothState(float DeltaTime)
 
 		TArray<physx::PxMat44> BoneMatrices;
 
-		NxClothingAsset* ClothingAsset = ClothingActors[ActorIdx].ParentClothingAsset->GetAsset();
+		NxClothingAsset* ClothingAsset = ClothingActors[ActorIdx].ParentClothingAsset;
 
 		uint32 NumUsedBones = ClothingAsset->getNumUsedBones();
 
@@ -3494,7 +3500,7 @@ void USkeletalMeshComponent::UpdateClothState(float DeltaTime)
 
 		for(uint32 Index=0; Index < NumUsedBones; Index++)
 		{
-		   FName BoneName = ClothingActors[ActorIdx].ParentClothingAsset->GetConvertedBoneName(Index);
+			FName BoneName = GetConvertedBoneName(ClothingActors[ActorIdx].ParentClothingAsset, Index);
 
 		   int32 BoneIndex = GetBoneIndex(BoneName);
 
@@ -3546,15 +3552,15 @@ void USkeletalMeshComponent::GetClothRootBoneMatrix(int32 AssetIndex, FMatrix& O
 {
 	if (IsValidClothingActor(AssetIndex))
 	{
-		TSharedPtr<FClothingAssetWrapper> Asset = ClothingActors[AssetIndex].ParentClothingAsset;
+		NxClothingAsset* Asset = ClothingActors[AssetIndex].ParentClothingAsset;
 
-		check(Asset.IsValid());
+		check(Asset);
 
-		const NxParameterized::Interface* AssetParams = Asset->GetAsset()->getAssetNxParameterized();
+		const NxParameterized::Interface* AssetParams = Asset->getAssetNxParameterized();
 		uint32 InternalRootBoneIndex;
 		verify(NxParameterized::getParamU32(*AssetParams, "rootBoneIndex", InternalRootBoneIndex));
 		check(InternalRootBoneIndex >= 0);
-		FName BoneName = Asset->GetConvertedBoneName(InternalRootBoneIndex);
+		FName BoneName = GetConvertedBoneName(Asset, InternalRootBoneIndex);
 		int32 BoneIndex = GetBoneIndex(BoneName);
 		check(BoneIndex >= 0);
 		OutRootBoneMatrix = GetBoneMatrix(BoneIndex);
@@ -3726,6 +3732,11 @@ bool USkeletalMeshComponent::IsValidClothingActor(int32 ActorIndex) const
 {
 #if WITH_APEX_CLOTHING
 
+	if (!SkeletalMesh)
+	{
+		return false;
+	}
+
 	//false if ActorIndex is out-range
 	if(ActorIndex >= SkeletalMesh->ClothingAssets.Num()
 	|| ActorIndex >= ClothingActors.Num())
@@ -3734,7 +3745,7 @@ bool USkeletalMeshComponent::IsValidClothingActor(int32 ActorIndex) const
 	}
 
 	if(ClothingActors[ActorIndex].ApexClothingActor
-	&& SkeletalMesh->ClothingAssets[ActorIndex].ApexClothingAsset->IsValid())
+	&& ClothingActors[ActorIndex].ParentClothingAsset == SkeletalMesh->ClothingAssets[ActorIndex].ApexClothingAsset)
 	{
 		return true;
 	}
@@ -4091,7 +4102,7 @@ void USkeletalMeshComponent::DrawClothingCollisionVolumes(FPrimitiveDrawInterfac
 				continue;
 			}
 
-			FName BoneName = SkeletalMesh->ClothingAssets[AssetIdx].ApexClothingAsset->GetConvertedBoneName(Collision.BoneIndex);
+			FName BoneName = GetConvertedBoneName(SkeletalMesh->ClothingAssets[AssetIdx].ApexClothingAsset, Collision.BoneIndex);
 			
 			int32 BoneIndex = GetBoneIndex(BoneName);
 
@@ -4150,7 +4161,7 @@ void USkeletalMeshComponent::DrawClothingCollisionVolumes(FPrimitiveDrawInterfac
 
 		for(int32 PlaneIdx=0; PlaneIdx < NumPlanes; PlaneIdx++)
 		{
-			FName BoneName = SkeletalMesh->ClothingAssets[AssetIdx].ApexClothingAsset->GetConvertedBoneName(BonePlanes[PlaneIdx].BoneIndex);
+			FName BoneName = GetConvertedBoneName(SkeletalMesh->ClothingAssets[AssetIdx].ApexClothingAsset, BonePlanes[PlaneIdx].BoneIndex);
 
 			int32 BoneIndex = GetBoneIndex(BoneName);
 
@@ -4188,7 +4199,7 @@ void USkeletalMeshComponent::DrawClothingCollisionVolumes(FPrimitiveDrawInterfac
 				continue;
 			}
 
-			FName BoneName = SkeletalMesh->ClothingAssets[AssetIdx].ApexClothingAsset->GetConvertedBoneName(Spheres[i].BoneIndex);
+			FName BoneName = GetConvertedBoneName(SkeletalMesh->ClothingAssets[AssetIdx].ApexClothingAsset, Spheres[i].BoneIndex);
 
 			int32 BoneIndex = GetBoneIndex(BoneName);
 
@@ -4365,7 +4376,7 @@ void USkeletalMeshComponent::LoadClothingVisualizationInfo(int32 AssetIndex)
 
 
 	FClothingAssetData& AssetData = SkeletalMesh->ClothingAssets[AssetIndex];
-	NxClothingAsset* ApexClothingAsset = AssetData.ApexClothingAsset->GetAsset();
+	NxClothingAsset* ApexClothingAsset = AssetData.ApexClothingAsset;
 	const NxParameterized::Interface* AssetParams = ApexClothingAsset->getAssetNxParameterized();
 
 	int32 NumPhysicalLODs;
@@ -4747,7 +4758,7 @@ void USkeletalMeshComponent::DrawClothingPhysicalMeshWire(FPrimitiveDrawInterfac
 						{
 							uint16 ApexBoneIndex = VisualInfo.ClothPhysicalMeshBoneWeightsInfo[FixedVertIdx].Indices[BoneWeightIdx];
 
-							FName BoneName = SkeletalMesh->ClothingAssets[AssetIdx].ApexClothingAsset->GetConvertedBoneName(ApexBoneIndex);
+							FName BoneName = GetConvertedBoneName(SkeletalMesh->ClothingAssets[AssetIdx].ApexClothingAsset, ApexBoneIndex);
 
 							int32 BoneIndex = GetBoneIndex(BoneName);
 
