a missing file when cheking in CL# 2380026.

diff --git a/Engine/Source/Runtime/Engine/Classes/Components/SkeletalMeshComponent.h b/Engine/Source/Runtime/Engine/Classes/Components/SkeletalMeshComponent.h
--- a/Engine/Source/Runtime/Engine/Classes/Components/SkeletalMeshComponent.h
+++ b/Engine/Source/Runtime/Engine/Classes/Components/SkeletalMeshComponent.h
@@ -53,6 +53,7 @@ namespace physx
 { 
 	namespace apex 
 	{
+		class NxClothingAsset;
 		class NxClothingActor;
 		class NxClothingCollision;
 	}
@@ -88,7 +89,7 @@ public:
 	 * to check whether this actor is valid or not 
 	 * because clothing asset can be changed by editing 
 	 */
-	TSharedPtr<class FClothingAssetWrapper>	ParentClothingAsset;
+	physx::apex::NxClothingAsset*	ParentClothingAsset;
 	/** APEX clothing actor is created from APEX clothing asset for cloth simulation */
 	physx::apex::NxClothingActor*		ApexClothingActor;
 	FPhysScene * PhysScene;
@@ -1052,7 +1053,7 @@ public:
 	* create only if became invalid
 	* BlendedData : added for cloth morph target but not used commonly
 	*/
-	bool CreateClothingActor(int32 AssetIndex, TSharedPtr<FClothingAssetWrapper> ClothingAsset, TArray<FVector>* BlendedDelta=NULL);
+	bool CreateClothingActor(int32 AssetIndex, physx::apex::NxClothingAsset* ClothingAsset, TArray<FVector>* BlendedDelta = NULL);
 	/** should call this method if occurred any changes in clothing assets */
 	void ValidateClothingActors();
 	/** add bounding box for cloth */
