Fix non-unity compile errors.

diff --git a/Engine/Source/Runtime/Engine/Public/AI/NavigationModifier.h b/Engine/Source/Runtime/Engine/Public/AI/NavigationModifier.h
--- a/Engine/Source/Runtime/Engine/Public/AI/NavigationModifier.h
+++ b/Engine/Source/Runtime/Engine/Public/AI/NavigationModifier.h
@@ -2,6 +2,7 @@
 
 #pragma once 
 #include "AI/Navigation/NavLinkDefinition.h"
+#include "AI/Navigation/NavigationTypes.h"
 
 class UNavArea;
 class UNavLinkDefinition;
