Strip the PIE prefix when testing map name is valid
Don't skip testing validity of map names when in PIE

diff --git a/Engine/Source/Runtime/Engine/Private/UnrealEngine.cpp b/Engine/Source/Runtime/Engine/Private/UnrealEngine.cpp
--- a/Engine/Source/Runtime/Engine/Private/UnrealEngine.cpp
+++ b/Engine/Source/Runtime/Engine/Private/UnrealEngine.cpp
@@ -7969,18 +7969,20 @@ bool UEngine::HandleReconnectCommand( const TCHAR* Cmd, FOutputDevice& Ar, UWorl
 
 bool UEngine::MakeSureMapNameIsValid(FString& InOutMapName)
 {
+	const FString TestMapName = UWorld::RemovePIEPrefix(InOutMapName);
+
 	// Check if the map name is long package name and if it actually exists.
 	// Short package names are only supported in non-shipping builds.
-	bool bIsValid = !FPackageName::IsShortPackageName(InOutMapName);
+	bool bIsValid = !FPackageName::IsShortPackageName(TestMapName);
 	if (bIsValid)
 	{
-		bIsValid = FPackageName::DoesPackageExist(InOutMapName);
+		bIsValid = FPackageName::DoesPackageExist(TestMapName);
 	}
 	else
 	{
 		// Look up on disk. Slow!
 		FString LongPackageName;
-		bIsValid = FPackageName::SearchForPackageOnDisk(InOutMapName, &LongPackageName);
+		bIsValid = FPackageName::SearchForPackageOnDisk(TestMapName, &LongPackageName);
 		if (bIsValid)
 		{
 			InOutMapName = LongPackageName;
@@ -8304,7 +8306,7 @@ bool UEngine::TickWorldTravel(FWorldContext& Context, float DeltaSeconds)
 		}
 		else if( Context.PendingNetGame && Context.PendingNetGame->bSuccessfullyConnected && !Context.PendingNetGame->bSentJoinRequest )
 		{
-			if (Context.WorldType != EWorldType::PIE && !MakeSureMapNameIsValid(Context.PendingNetGame->URL.Map))
+			if (!MakeSureMapNameIsValid(Context.PendingNetGame->URL.Map))
 			{
 				BrowseToDefaultMap(Context);
 				BroadcastTravelFailure(Context.World(), ETravelFailure::PackageMissing, Context.PendingNetGame->URL.RedirectURL);
