Parallel Evaluation corruption fix

diff --git a/Engine/Source/Runtime/Engine/Public/AnimEncoding.h b/Engine/Source/Runtime/Engine/Public/AnimEncoding.h
--- a/Engine/Source/Runtime/Engine/Public/AnimEncoding.h
+++ b/Engine/Source/Runtime/Engine/Public/AnimEncoding.h
@@ -494,14 +494,8 @@ FORCEINLINE float AnimEncoding::TimeToIndex(
 	int32 &PosIndex0Out,
 	int32 &PosIndex1Out)
 {
-	static int32		NumKeysCache = 0; // this value is guaranteed to not be used for valid data
-	static float	TimeCache;
-	static float	SequenceLengthCache;
-	static int32		PosIndex0CacheOut; 
-	static int32		PosIndex1CacheOut; 
-	static float	AlphaCacheOut;
-
 	const float SequenceLength= Seq.SequenceLength;
+	float Alpha;
 
 	if (NumKeys < 2)
 	{
@@ -510,48 +504,36 @@ FORCEINLINE float AnimEncoding::TimeToIndex(
 		PosIndex1Out = 0;
 		return 0.0f;
 	}
-	if (
-		NumKeysCache		!= NumKeys ||
-		SequenceLengthCache != SequenceLength ||
-		TimeCache			!= RelativePos
-		)
+	// Check for before-first-frame case.
+	if( RelativePos <= 0.f )
+	{
+		PosIndex0Out = 0;
+		PosIndex1Out = 0;
+		Alpha = 0.0f;
+	}
+	else
 	{
-		NumKeysCache		= NumKeys;
-		SequenceLengthCache = SequenceLength;
-		TimeCache			= RelativePos;
-		// Check for before-first-frame case.
-		if( RelativePos <= 0.f )
+		NumKeys -= 1; // never used without the minus one in this case
+		// Check for after-last-frame case.
+		if( RelativePos >= 1.0f )
 		{
-			PosIndex0CacheOut = 0;
-			PosIndex1CacheOut = 0;
-			AlphaCacheOut = 0.0f;
+			// If we're not looping, key n-1 is the final key.
+			PosIndex0Out = NumKeys;
+			PosIndex1Out = NumKeys;
+			Alpha = 0.0f;
 		}
 		else
 		{
-			NumKeys -= 1; // never used without the minus one in this case
-			// Check for after-last-frame case.
-			if( RelativePos >= 1.0f )
-			{
-				// If we're not looping, key n-1 is the final key.
-				PosIndex0CacheOut = NumKeys;
-				PosIndex1CacheOut = NumKeys;
-				AlphaCacheOut = 0.0f;
-			}
-			else
-			{
-				// For non-looping animation, the last frame is the ending frame, and has no duration.
-				const float KeyPos = RelativePos * float(NumKeys);
-				checkSlow(KeyPos >= 0.0f);
-				const float KeyPosFloor = floorf(KeyPos);
-				PosIndex0CacheOut = FMath::Min( FMath::TruncToInt(KeyPosFloor), NumKeys );
-				AlphaCacheOut = KeyPos - KeyPosFloor;
-				PosIndex1CacheOut = FMath::Min( PosIndex0CacheOut + 1, NumKeys );
-			}
+			// For non-looping animation, the last frame is the ending frame, and has no duration.
+			const float KeyPos = RelativePos * float(NumKeys);
+			checkSlow(KeyPos >= 0.0f);
+			const float KeyPosFloor = floorf(KeyPos);
+			PosIndex0Out = FMath::Min( FMath::TruncToInt(KeyPosFloor), NumKeys );
+			Alpha = KeyPos - KeyPosFloor;
+			PosIndex1Out = FMath::Min( PosIndex0Out + 1, NumKeys );
 		}
 	}
-	PosIndex0Out = PosIndex0CacheOut;
-	PosIndex1Out = PosIndex1CacheOut;
-	return AlphaCacheOut;
+	return Alpha;
 }
 
 /**
