Only mark objects in the world as pending kill, not everything in the package.

diff --git a/Engine/Source/Runtime/Engine/Private/World.cpp b/Engine/Source/Runtime/Engine/Private/World.cpp
--- a/Engine/Source/Runtime/Engine/Private/World.cpp
+++ b/Engine/Source/Runtime/Engine/Private/World.cpp
@@ -4593,21 +4593,21 @@ UWorld* FSeamlessTravelHandler::Tick()
 			GWorld = NULL;
 
 			// mark everything else contained in the world to be deleted
-			TArray<UPackage*> CurrentWorldPackages;
+			TArray<UWorld*> CurrentWorlds;
 			for (auto LevelIt(CurrentWorld->GetLevelIterator()); LevelIt; ++LevelIt)
 			{
 				const ULevel* Level = *LevelIt;
 				if (Level)
 				{
-					CurrentWorldPackages.Add(Level->GetOutermost());
+					CurrentWorlds.Add(CastChecked<UWorld>(Level->GetOuter()));
 				}
 			}
 
 			for (TObjectIterator<UObject> It; It; ++It)
 			{
-				for (const UPackage* WorldPackage : CurrentWorldPackages)
+				for (const UWorld* World : CurrentWorlds)
 				{
-					if (It->IsIn(WorldPackage))
+					if (It->IsIn(World))
 					{
 						It->MarkPendingKill();
 						break;
