Fixed montage sections collapsing to zero after engine version change
[UE-6444]

diff --git a/Engine/Source/Runtime/Engine/Private/Animation/AnimMontage.cpp b/Engine/Source/Runtime/Engine/Private/Animation/AnimMontage.cpp
--- a/Engine/Source/Runtime/Engine/Private/Animation/AnimMontage.cpp
+++ b/Engine/Source/Runtime/Engine/Private/Animation/AnimMontage.cpp
@@ -351,9 +351,17 @@ void UAnimMontage::PostLoad()
 
 	int32 Ver = GetLinker()->UE4Ver();
 
-	for (auto& Composite : CompositeSections)
+	for(auto& Composite : CompositeSections)
 	{
-		Composite.LinkMontage(this, Composite.GetTime());
+		if(Composite.StartTime_DEPRECATED != 0.0f)
+		{
+			Composite.Clear();
+			Composite.LinkMontage(this, Composite.StartTime_DEPRECATED);
+		}
+		else
+		{
+			Composite.LinkMontage(this, Composite.GetTime());
+		}
 	}
 
 	bool bRootMotionEnabled = bEnableRootMotionTranslation || bEnableRootMotionRotation;
