Fixed crash with retargeting

diff --git a/Engine/Source/Runtime/Engine/Private/Animation/AnimSequence.cpp b/Engine/Source/Runtime/Engine/Private/Animation/AnimSequence.cpp
--- a/Engine/Source/Runtime/Engine/Private/Animation/AnimSequence.cpp
+++ b/Engine/Source/Runtime/Engine/Private/Animation/AnimSequence.cpp
@@ -1777,7 +1777,7 @@ void FillUpTransformBasedOnRig(USkeleton* Skeleton, TArray<FTransform>& NodeSpac
 						{
 							// I think translation has to include rotation, otherwise it won't work
 							Translations[Index] = SpaceBases[BoneIndex].GetRelativeTransform(SpaceBases[ParentBoneIndex]);
-							TranslationParentFlags[BoneIndex] = true;
+							TranslationParentFlags[Index] = true;
 						}
 					}
 				}
