Fixed notifies collapsing to time 0 in anim sequences when loading an older sequence after recent merge conflict removed time conversion code
[UE-6457]

diff --git a/Engine/Source/Runtime/Engine/Private/Animation/AnimSequence.cpp b/Engine/Source/Runtime/Engine/Private/Animation/AnimSequence.cpp
--- a/Engine/Source/Runtime/Engine/Private/Animation/AnimSequence.cpp
+++ b/Engine/Source/Runtime/Engine/Private/Animation/AnimSequence.cpp
@@ -424,6 +424,24 @@ void UAnimSequence::PostLoad()
 		INC_DWORD_STAT_BY( STAT_AnimationMemory, GetResourceSize(EResourceSizeMode::Exclusive) );
 	}
 
+	for(FAnimNotifyEvent& Notify : Notifies)
+	{
+		if(Notify.DisplayTime_DEPRECATED != 0.0f)
+		{
+			Notify.Clear();
+			Notify.LinkSequence(this, Notify.DisplayTime_DEPRECATED);
+		}
+		else
+		{
+			Notify.LinkSequence(this, Notify.GetTime());
+		}
+	
+		if(Notify.Duration != 0.0f)
+		{
+			Notify.EndLink.LinkSequence(this, Notify.GetTime() + Notify.Duration);
+		}
+	}
+
 #if WITH_EDITOR
 	if (USkeleton * CurrentSkeleton = GetSkeleton())
 	{
