Crash fix for root motion

diff --git a/Engine/Source/Runtime/Engine/Private/Animation/AnimCompositeBase.cpp b/Engine/Source/Runtime/Engine/Private/Animation/AnimCompositeBase.cpp
--- a/Engine/Source/Runtime/Engine/Private/Animation/AnimCompositeBase.cpp
+++ b/Engine/Source/Runtime/Engine/Private/Animation/AnimCompositeBase.cpp
@@ -191,7 +191,7 @@ bool FAnimTrack::HasRootMotion() const
 {
 	for (const FAnimSegment& AnimSegment : AnimSegments)
 	{
-		if (AnimSegment.AnimReference->HasRootMotion())
+		if (AnimSegment.AnimReference && AnimSegment.AnimReference->HasRootMotion())
 		{
 			return true;
 		}
