Prevent duplicate simulataneous shader compilation jobs when applying a material to multiple MaterialInstances that share the same StaticParameterSet
In FMaterial::CacheShaders, check if we're compiling the current shadermap for another FMaterial and if so, associate ourselves with that compilation job (UE-2418)

diff --git a/Engine/Source/Runtime/Engine/Private/Materials/MaterialShader.cpp b/Engine/Source/Runtime/Engine/Private/Materials/MaterialShader.cpp
--- a/Engine/Source/Runtime/Engine/Private/Materials/MaterialShader.cpp
+++ b/Engine/Source/Runtime/Engine/Private/Materials/MaterialShader.cpp
@@ -1315,6 +1315,20 @@ bool FMaterialShaderMap::ProcessCompilationResults(const TArray<FShaderCompileJo
 	return false;
 }
 
+bool FMaterialShaderMap::TryToAddToExistingCompilationTask(FMaterial* Material)
+{
+	check(NumRefs > 0);
+	TArray<FMaterial*>* CorrespondingMaterials = FMaterialShaderMap::ShaderMapsBeingCompiled.Find(this);
+
+	if (CorrespondingMaterials)
+	{
+		CorrespondingMaterials->AddUnique(Material);
+		return true;
+	}
+
+	return false;
+}
+
 bool FMaterialShaderMap::IsComplete(const FMaterial* Material, bool bSilent)
 {
 	bool bIsComplete = true;
diff --git a/Engine/Source/Runtime/Engine/Private/Materials/MaterialShared.cpp b/Engine/Source/Runtime/Engine/Private/Materials/MaterialShared.cpp
--- a/Engine/Source/Runtime/Engine/Private/Materials/MaterialShared.cpp
+++ b/Engine/Source/Runtime/Engine/Private/Materials/MaterialShared.cpp
@@ -1272,7 +1272,13 @@ bool FMaterial::CacheShaders(const FMaterialShaderMapId& ShaderMapId, EShaderPla
 		}
 	}
 
-	bool bRequiredRecompile = false;
+	if (GameThreadShaderMap && GameThreadShaderMap->TryToAddToExistingCompilationTask(this))
+	{
+		// Reset the shader map so the default material will be used until the compile finishes.
+		GameThreadShaderMap = NULL;
+		bSucceeded = true;
+	}
+	else
 	if (!GameThreadShaderMap || !GameThreadShaderMap->IsComplete(this, true))
 	{
 		if (bContainsInlineShaders || FPlatformProperties::RequiresCookedData())
@@ -1306,7 +1312,6 @@ bool FMaterial::CacheShaders(const FMaterialShaderMapId& ShaderMapId, EShaderPla
 			// If there's no cached shader map for this material, compile a new one.
 			// This is just kicking off the async compile, GameThreadShaderMap will not be complete yet
 			bSucceeded = BeginCompileShaderMap(ShaderMapId, Platform, GameThreadShaderMap, bApplyCompletedShaderMapForRendering);
-			bRequiredRecompile = true;
 
 			if (!bSucceeded)
 			{
diff --git a/Engine/Source/Runtime/Engine/Public/MaterialShared.h b/Engine/Source/Runtime/Engine/Public/MaterialShared.h
--- a/Engine/Source/Runtime/Engine/Public/MaterialShared.h
+++ b/Engine/Source/Runtime/Engine/Public/MaterialShared.h
@@ -556,6 +556,14 @@ public:
 	/** Attempts to load missing shaders from memory. */
 	void LoadMissingShadersFromMemory(const FMaterial* Material);
 
+	/**
+	 * Checks to see if the shader map is already being compiled for another material, and if so
+	 * adds the specified material to the list to be applied to once the compile finishes.
+	 * @param Material - The material we also wish to apply the compiled shader map to.
+	 * @return True if the shader map was being compiled and we added Material to the list to be applied.
+	 */
+	bool TryToAddToSimultaneousCompilation(FMaterial* Material);
+
 	/** Builds a list of the shaders in a shader map. */
 	ENGINE_API void GetShaderList(TMap<FShaderId,FShader*>& OutShaders) const;
 
