Display a slow task dialog when switching preview rendering level (UE-5476)

diff --git a/Engine/Source/Runtime/Engine/Private/World.cpp b/Engine/Source/Runtime/Engine/Private/World.cpp
--- a/Engine/Source/Runtime/Engine/Private/World.cpp
+++ b/Engine/Source/Runtime/Engine/Private/World.cpp
@@ -5284,53 +5284,67 @@ void UWorld::ChangeFeatureLevel(ERHIFeatureLevel::Type InFeatureLevel)
 {
 	if (InFeatureLevel != FeatureLevel)
 	{
-		FlushRenderingCommands();
-
-		// Give all scene components the opportunity to prepare for pending feature level change.
-		for (TObjectIterator<USceneComponent> It; It; ++It)
-		{
-			USceneComponent* SceneComponent = *It;
-			if (SceneComponent->GetWorld() == this)
-			{
-				SceneComponent->PreFeatureLevelChange(InFeatureLevel);
-			}
-		}
-
-		FGlobalComponentReregisterContext RecreateComponents;
-
-		// Decrement refcount on old feature level
-		UMaterialInterface::SetGlobalRequiredFeatureLevel(InFeatureLevel, true);
-
-		UMaterial::AllMaterialsCacheResourceShadersForRendering();
-		UMaterialInstance::AllMaterialsCacheResourceShadersForRendering();
-		GetGlobalShaderMap(InFeatureLevel, false);
-		GShaderCompilingManager->ProcessAsyncResults(false, true);
-
-		//invalidate global bound shader states so they will be created with the new shaders the next time they are set (in SetGlobalBoundShaderState)
-		for (TLinkedList<FGlobalBoundShaderStateResource*>::TIterator It(FGlobalBoundShaderStateResource::GetGlobalBoundShaderStateList()); It; It.Next())
-		{
-			BeginUpdateResourceRHI(*It);
-		}
-
-		FeatureLevel = InFeatureLevel;
-
-		if (Scene)
-		{
-			PersistentLevel->ReleaseRenderingResources();
-			Scene->Release();
-			GetRendererModule().RemoveScene(Scene);
-
-			GetRendererModule().AllocateScene(this, bRequiresHitProxies, FXSystem != nullptr, InFeatureLevel );
-
-			PersistentLevel->InitializeRenderingResources();
-			PersistentLevel->PrecomputedVisibilityHandler.UpdateScene(Scene);
-			PersistentLevel->PrecomputedVolumeDistanceField.UpdateScene(Scene);
-		}
-
-		TriggerStreamingDataRebuild();
-
-		FOutputDeviceNull Ar;
-		RecompileShaders(TEXT("CHANGED"), Ar);
+        FScopedSlowTask SlowTask(100.f, NSLOCTEXT("Engine", "ChangingPreviewRenderingLevelMessage", "Changing Preview Rendering Level"));
+        SlowTask.MakeDialog();
+        {
+            FlushRenderingCommands();
+
+            SlowTask.EnterProgressFrame(10.0f);
+            // Give all scene components the opportunity to prepare for pending feature level change.
+            for (TObjectIterator<USceneComponent> It; It; ++It)
+            {
+                USceneComponent* SceneComponent = *It;
+                if (SceneComponent->GetWorld() == this)
+                {
+                    SceneComponent->PreFeatureLevelChange(InFeatureLevel);
+                }
+            }
+
+            SlowTask.EnterProgressFrame(10.0f);
+            FGlobalComponentReregisterContext RecreateComponents;
+
+            // Decrement refcount on old feature level
+            UMaterialInterface::SetGlobalRequiredFeatureLevel(InFeatureLevel, true);
+
+            SlowTask.EnterProgressFrame(10.0f);
+            UMaterial::AllMaterialsCacheResourceShadersForRendering();
+            SlowTask.EnterProgressFrame(10.0f);
+            UMaterialInstance::AllMaterialsCacheResourceShadersForRendering();
+            SlowTask.EnterProgressFrame(10.0f);
+            GetGlobalShaderMap(InFeatureLevel, false);
+            SlowTask.EnterProgressFrame(10.0f);
+            GShaderCompilingManager->ProcessAsyncResults(false, true);
+
+            SlowTask.EnterProgressFrame(10.0f);
+            //invalidate global bound shader states so they will be created with the new shaders the next time they are set (in SetGlobalBoundShaderState)
+            for (TLinkedList<FGlobalBoundShaderStateResource*>::TIterator It(FGlobalBoundShaderStateResource::GetGlobalBoundShaderStateList()); It; It.Next())
+            {
+                BeginUpdateResourceRHI(*It);
+            }
+
+            FeatureLevel = InFeatureLevel;
+
+            SlowTask.EnterProgressFrame(10.0f);
+            if (Scene)
+            {
+                PersistentLevel->ReleaseRenderingResources();
+                Scene->Release();
+                GetRendererModule().RemoveScene(Scene);
+
+                GetRendererModule().AllocateScene(this, bRequiresHitProxies, FXSystem != nullptr, InFeatureLevel );
+
+                PersistentLevel->InitializeRenderingResources();
+                PersistentLevel->PrecomputedVisibilityHandler.UpdateScene(Scene);
+                PersistentLevel->PrecomputedVolumeDistanceField.UpdateScene(Scene);
+            }
+
+            SlowTask.EnterProgressFrame(10.0f);
+            TriggerStreamingDataRebuild();
+
+            SlowTask.EnterProgressFrame(10.0f);
+            FOutputDeviceNull Ar;
+            RecompileShaders(TEXT("CHANGED"), Ar);
+        }
 	}
 }
 #endif // WITH_EDITOR
