Don't crash when default touch interface is none and a touch interface is activated

diff --git a/Engine/Source/Runtime/Engine/Private/GameFramework/TouchInterface.cpp b/Engine/Source/Runtime/Engine/Private/GameFramework/TouchInterface.cpp
--- a/Engine/Source/Runtime/Engine/Private/GameFramework/TouchInterface.cpp
+++ b/Engine/Source/Runtime/Engine/Private/GameFramework/TouchInterface.cpp
@@ -21,30 +21,33 @@ UTouchInterface::UTouchInterface(const FObjectInitializer& ObjectInitializer)
 
 void UTouchInterface::Activate(TSharedPtr<SVirtualJoystick> VirtualJoystick)
 {
-	VirtualJoystick->SetGlobalParameters(ActiveOpacity, InactiveOpacity, TimeUntilDeactive, TimeUntilReset, ActivationDelay, bPreventRecenter, StartupDelay);
+	if (VirtualJoystick.IsValid())
+	{
+		VirtualJoystick->SetGlobalParameters(ActiveOpacity, InactiveOpacity, TimeUntilDeactive, TimeUntilReset, ActivationDelay, bPreventRecenter, StartupDelay);
 
-	// convert from the UStructs to the slate structs
-	TArray<SVirtualJoystick::FControlInfo> SlateControls;
+		// convert from the UStructs to the slate structs
+		TArray<SVirtualJoystick::FControlInfo> SlateControls;
 
-	for (int32 ControlIndex = 0; ControlIndex < Controls.Num(); ControlIndex++)
-	{
-		FTouchInputControl Control = Controls[ControlIndex];
-		SVirtualJoystick::FControlInfo* SlateControl = new(SlateControls)SVirtualJoystick::FControlInfo;
-
-		SlateControl->Image1 = FCoreStyle::GetDynamicImageBrush("Engine.Joystick.Image1", Control.Image1, Control.Image1->GetFName());
-		SlateControl->Image2 = FCoreStyle::GetDynamicImageBrush("Engine.Joystick.Image2", Control.Image2, Control.Image2->GetFName());
-		SlateControl->Center = Control.Center;
-		SlateControl->VisualSize = Control.VisualSize;
-		SlateControl->ThumbSize = Control.ThumbSize;
-		if (Control.InputScale.Size() > DELTA)
+		for (int32 ControlIndex = 0; ControlIndex < Controls.Num(); ControlIndex++)
 		{
-			SlateControl->InputScale = Control.InputScale;
+			FTouchInputControl Control = Controls[ControlIndex];
+			SVirtualJoystick::FControlInfo* SlateControl = new(SlateControls)SVirtualJoystick::FControlInfo;
+
+			SlateControl->Image1 = FCoreStyle::GetDynamicImageBrush("Engine.Joystick.Image1", Control.Image1, Control.Image1->GetFName());
+			SlateControl->Image2 = FCoreStyle::GetDynamicImageBrush("Engine.Joystick.Image2", Control.Image2, Control.Image2->GetFName());
+			SlateControl->Center = Control.Center;
+			SlateControl->VisualSize = Control.VisualSize;
+			SlateControl->ThumbSize = Control.ThumbSize;
+			if (Control.InputScale.Size() > DELTA)
+			{
+				SlateControl->InputScale = Control.InputScale;
+			}
+			SlateControl->InteractionSize = Control.InteractionSize;
+			SlateControl->MainInputKey = Control.MainInputKey;
+			SlateControl->AltInputKey = Control.AltInputKey;
 		}
-		SlateControl->InteractionSize = Control.InteractionSize;
-		SlateControl->MainInputKey = Control.MainInputKey;
-		SlateControl->AltInputKey = Control.AltInputKey;
-	}
 
-	// set them as active!
-	VirtualJoystick->SetControls(SlateControls);
+		// set them as active!
+		VirtualJoystick->SetControls(SlateControls);
+	}
 }
diff --git a/Engine/Source/Runtime/Engine/Private/PlayerController.cpp b/Engine/Source/Runtime/Engine/Private/PlayerController.cpp
--- a/Engine/Source/Runtime/Engine/Private/PlayerController.cpp
+++ b/Engine/Source/Runtime/Engine/Private/PlayerController.cpp
@@ -1039,10 +1039,19 @@ void APlayerController::CreateTouchInterface()
 			Cast<ULocalPlayer>(Player)->ViewportClient->RemoveViewportWidgetContent(VirtualJoystick.ToSharedRef());
 		}
 
-		// load what the game wants to show at startup
-		FStringAssetReference DefaultTouchInterfaceName = GetDefault<UInputSettings>()->DefaultTouchInterface;
+		if (CurrentTouchInterface == nullptr)
+		{
+			// load what the game wants to show at startup
+			FStringAssetReference DefaultTouchInterfaceName = GetDefault<UInputSettings>()->DefaultTouchInterface;
+
+			if (DefaultTouchInterfaceName.IsValid())
+			{
+				// activate this interface if we have it
+				CurrentTouchInterface = LoadObject<UTouchInterface>(NULL, *DefaultTouchInterfaceName.ToString());
+			}
+		}
 
-		if (DefaultTouchInterfaceName.IsValid())
+		if (CurrentTouchInterface)
 		{
 			// create the joystick 
 			VirtualJoystick = SNew(SVirtualJoystick);
@@ -1050,12 +1059,7 @@ void APlayerController::CreateTouchInterface()
 			// add it to the player's viewport
 			LocalPlayer->ViewportClient->AddViewportWidgetContent(VirtualJoystick.ToSharedRef());
 
-			// activate this interface if we have it
-			UTouchInterface* DefaultTouchInterface = LoadObject<UTouchInterface>(NULL, *DefaultTouchInterfaceName.ToString());
-			if (DefaultTouchInterface != NULL)
-			{
-				ActivateTouchInterface(DefaultTouchInterface);
-			}
+			ActivateTouchInterface(CurrentTouchInterface);
 		}
 	}
 }
@@ -1064,10 +1068,9 @@ void APlayerController::CleanupGameViewport()
 {
 	ULocalPlayer* LocalPlayer = Cast<ULocalPlayer>(Player);
 
-	if (LocalPlayer && LocalPlayer->ViewportClient && VirtualJoystick.IsValid())
+	if (VirtualJoystick.IsValid())
 	{
-		LocalPlayer->ViewportClient->RemoveViewportWidgetContent(VirtualJoystick.ToSharedRef());
-		VirtualJoystick = NULL;
+		ActivateTouchInterface(nullptr);
 	}
 }
 
@@ -1333,7 +1336,7 @@ void APlayerController::EndPlay(const EEndPlayReason::Type EndPlayReason)
 	{
 		if (VirtualJoystick.IsValid())
 		{
-			LocalPlayer->ViewportClient->RemoveViewportWidgetContent(VirtualJoystick.ToSharedRef());
+			ActivateTouchInterface(nullptr);
 		}
 
 		// Stop any force feedback effects that may be active
@@ -4300,11 +4303,27 @@ void APlayerController::DisableInput(class APlayerController* PlayerController)
 
 void APlayerController::ActivateTouchInterface(UTouchInterface* NewTouchInterface)
 {
+	CurrentTouchInterface = NewTouchInterface;
 	if(NewTouchInterface)
 	{
-		NewTouchInterface->Activate(VirtualJoystick);
+		if (!VirtualJoystick.IsValid())
+		{
+			CreateTouchInterface();
+		}
+		else
+		{
+			NewTouchInterface->Activate(VirtualJoystick);
+		}
+	}
+	else if (VirtualJoystick.IsValid())
+	{
+		ULocalPlayer* LocalPlayer = Cast<ULocalPlayer>(Player);
+		if (LocalPlayer && LocalPlayer->ViewportClient)
+		{
+			LocalPlayer->ViewportClient->RemoveViewportWidgetContent(VirtualJoystick.ToSharedRef());
+		}
+		VirtualJoystick = NULL;
 	}
-	CurrentTouchInterface = NewTouchInterface;
 }
 
 void APlayerController::SetVirtualJoystickVisibility(bool bVisible)
