Fix for physics assets being collapsed when added to a world.

diff --git a/Engine/Source/Runtime/Engine/Private/SkeletalMeshComponent.cpp b/Engine/Source/Runtime/Engine/Private/SkeletalMeshComponent.cpp
--- a/Engine/Source/Runtime/Engine/Private/SkeletalMeshComponent.cpp
+++ b/Engine/Source/Runtime/Engine/Private/SkeletalMeshComponent.cpp
@@ -232,6 +232,7 @@ void USkeletalMeshComponent::InitAnim(bool bForceReinit)
         TickAnimation(0.f); 
 
 		RefreshBoneTransforms();
+		FlipEditableSpaceBases();
 		UpdateComponentToWorld();
 	}
 }
