#UE4 - fixed net drivers switching mode from NM_Client to NM_Standalone (server) after a call to HandleDisconnect  (UE-6092)

diff --git a/Engine/Source/Runtime/Engine/Private/UnrealEngine.cpp b/Engine/Source/Runtime/Engine/Private/UnrealEngine.cpp
--- a/Engine/Source/Runtime/Engine/Private/UnrealEngine.cpp
+++ b/Engine/Source/Runtime/Engine/Private/UnrealEngine.cpp
@@ -7938,12 +7938,12 @@ void UEngine::HandleDisconnect( UWorld *InWorld, UNetDriver *NetDriver )
 		WorldContext->LastURL.RemoveOption( TEXT("Listen") );
 		WorldContext->LastURL.RemoveOption( TEXT("LAN") );
 
+		// Net driver destruction will occur during LoadMap (prevents GetNetMode from changing output for the remainder of the frame)
 		SetClientTravel( InWorld, TEXT("?closed"), TRAVEL_Absolute );
 	}
-
-	// Shut down any existing game connections
-	if (NetDriver)
+	else if (NetDriver)
 	{
+		// Shut down any existing game connections
 		if (InWorld)
 		{
 			// Call this to remove the NetDriver from the world context's ActiveNetDriver list
