Move highres screenshot material aquisition to app initialisation.

diff --git a/Engine/Source/Runtime/Engine/Private/HighResScreenshot.cpp b/Engine/Source/Runtime/Engine/Private/HighResScreenshot.cpp
--- a/Engine/Source/Runtime/Engine/Private/HighResScreenshot.cpp
+++ b/Engine/Source/Runtime/Engine/Private/HighResScreenshot.cpp
@@ -31,6 +31,23 @@ void FHighResScreenshotConfig::Init()
 		ImageCompressorLDR = ImageWrapperModule->CreateImageWrapper(EImageFormat::PNG);
 		ImageCompressorHDR = ImageWrapperModule->CreateImageWrapper(EImageFormat::EXR);
 	}
+
+	HighResScreenshotMaterial = LoadObject<UMaterial>(NULL, TEXT("/Engine/EngineMaterials/HighResScreenshot.HighResScreenshot"));
+	HighResScreenshotMaskMaterial = LoadObject<UMaterial>(NULL, TEXT("/Engine/EngineMaterials/HighResScreenshotMask.HighResScreenshotMask"));
+	HighResScreenshotCaptureRegionMaterial = LoadObject<UMaterial>(NULL, TEXT("/Engine/EngineMaterials/HighResScreenshotCaptureRegion.HighResScreenshotCaptureRegion"));
+
+	if (HighResScreenshotMaterial)
+	{
+		HighResScreenshotMaterial->AddToRoot();
+	}
+	if (HighResScreenshotMaskMaterial)
+	{
+		HighResScreenshotMaskMaterial->AddToRoot();
+	}
+	if (HighResScreenshotCaptureRegionMaterial)
+	{
+		HighResScreenshotCaptureRegionMaterial->AddToRoot();
+	}
 }
 
 void FHighResScreenshotConfig::ChangeViewport(TWeakPtr<FSceneViewport> InViewport)
diff --git a/Engine/Source/Runtime/Engine/Private/SceneView.cpp b/Engine/Source/Runtime/Engine/Private/SceneView.cpp
--- a/Engine/Source/Runtime/Engine/Private/SceneView.cpp
+++ b/Engine/Source/Runtime/Engine/Private/SceneView.cpp
@@ -1211,41 +1211,13 @@ void FSceneView::EndFinalPostprocessSettings(const FSceneViewInitOptions& ViewIn
 	}
 
 #if WITH_EDITOR
-	// Access the materials for the hires screenshot system
-	static UMaterial* HighResScreenshotMaterial = NULL;
-	static UMaterial* HighResScreenshotMaskMaterial = NULL;
-	static UMaterial* HighResScreenshotCaptureRegionMaterial = NULL;
-	static bool bHighResScreenshotMaterialsConfigured = false;
-
-	if (!bHighResScreenshotMaterialsConfigured)
-	{
-		HighResScreenshotMaterial = LoadObject<UMaterial>(NULL, TEXT("/Engine/EngineMaterials/HighResScreenshot.HighResScreenshot"));
-		HighResScreenshotMaskMaterial = LoadObject<UMaterial>(NULL, TEXT("/Engine/EngineMaterials/HighResScreenshotMask.HighResScreenshotMask"));
-		HighResScreenshotCaptureRegionMaterial = LoadObject<UMaterial>(NULL, TEXT("/Engine/EngineMaterials/HighResScreenshotCaptureRegion.HighResScreenshotCaptureRegion"));
-
-		if (HighResScreenshotMaterial)
-		{
-			HighResScreenshotMaterial->AddToRoot();
-		}
-		if (HighResScreenshotMaskMaterial)
-		{
-			HighResScreenshotMaskMaterial->AddToRoot(); 
-		}
-		if (HighResScreenshotCaptureRegionMaterial)
-		{
-			HighResScreenshotCaptureRegionMaterial->AddToRoot();
-		}
-
-		bHighResScreenshotMaterialsConfigured = true;
-	}
+	FHighResScreenshotConfig& Config = GetHighResScreenshotConfig();
 
 	// Pass highres screenshot materials through post process settings
-	FinalPostProcessSettings.HighResScreenshotMaterial = HighResScreenshotMaterial;
-	FinalPostProcessSettings.HighResScreenshotMaskMaterial = HighResScreenshotMaskMaterial;
+	FinalPostProcessSettings.HighResScreenshotMaterial = Config.HighResScreenshotMaterial;
+	FinalPostProcessSettings.HighResScreenshotMaskMaterial = Config.HighResScreenshotMaskMaterial;
 	FinalPostProcessSettings.HighResScreenshotCaptureRegionMaterial = NULL;
 
-	FHighResScreenshotConfig& Config = GetHighResScreenshotConfig();
-
 	// If the highres screenshot UI is open and we're not taking a highres screenshot this frame
 	if (Config.bDisplayCaptureRegion && !GIsHighResScreenshot)
 	{
@@ -1266,7 +1238,7 @@ void FSceneView::EndFinalPostprocessSettings(const FSceneViewInitOptions& ViewIn
 				NormalizedCaptureRegion.A = (float)Config.UnscaledCaptureRegion.Max.Y / (float)ViewRect.Height();
 
 				// Get a MID for drawing this frame and push the capture region into the shader parameter
-				FinalPostProcessSettings.HighResScreenshotCaptureRegionMaterial = State->GetReusableMID(HighResScreenshotCaptureRegionMaterial);
+				FinalPostProcessSettings.HighResScreenshotCaptureRegionMaterial = State->GetReusableMID(Config.HighResScreenshotCaptureRegionMaterial);
 				FinalPostProcessSettings.HighResScreenshotCaptureRegionMaterial->SetVectorParameterValue(ParamName, NormalizedCaptureRegion);
 			}
 		}
diff --git a/Engine/Source/Runtime/Engine/Public/HighResScreenshot.h b/Engine/Source/Runtime/Engine/Public/HighResScreenshot.h
--- a/Engine/Source/Runtime/Engine/Public/HighResScreenshot.h
+++ b/Engine/Source/Runtime/Engine/Public/HighResScreenshot.h
@@ -19,6 +19,11 @@ struct ENGINE_API FHighResScreenshotConfig
 	TSharedPtr<class IImageWrapper> ImageCompressorLDR;
 	TSharedPtr<class IImageWrapper> ImageCompressorHDR;
 
+	// Materials used in the editor to help with the capture of highres screenshots
+	UMaterial* HighResScreenshotMaterial;
+	UMaterial* HighResScreenshotMaskMaterial;
+	UMaterial* HighResScreenshotCaptureRegionMaterial;
+
 	FHighResScreenshotConfig();
 
 	/** Initialize the image wrapper modules (required for SaveImage) **/
