Sort lods ActiveBoneIndices in the data so that we dont have to do it every frame at runtime, saves roughly 30% of the cost of UpdateRefToLocalMatrices

diff --git a/Engine/Source/Editor/UnrealEd/Private/Factories/SkeletalMeshImport.cpp b/Engine/Source/Editor/UnrealEd/Private/Factories/SkeletalMeshImport.cpp
--- a/Engine/Source/Editor/UnrealEd/Private/Factories/SkeletalMeshImport.cpp
+++ b/Engine/Source/Editor/UnrealEd/Private/Factories/SkeletalMeshImport.cpp
@@ -846,6 +846,7 @@ void RestoreExistingSkelMeshData(ExistingSkelMeshData* MeshData, USkeletalMesh*
 
 				// Sort ascending for parent child relationship
 				LODModel.RequiredBones.Sort();
+				LODModel.ActiveBoneIndices.Sort();
 
 				// Fix the chunks' BoneMaps.
 				for(int32 ChunkIndex = 0;ChunkIndex < LODModel.Chunks.Num();ChunkIndex++)
diff --git a/Engine/Source/Editor/UnrealEd/Private/Fbx/FbxSkeletalMeshImport.cpp b/Engine/Source/Editor/UnrealEd/Private/Fbx/FbxSkeletalMeshImport.cpp
--- a/Engine/Source/Editor/UnrealEd/Private/Fbx/FbxSkeletalMeshImport.cpp
+++ b/Engine/Source/Editor/UnrealEd/Private/Fbx/FbxSkeletalMeshImport.cpp
@@ -2575,6 +2575,7 @@ bool UnFbx::FFbxImporter::ImportSkeletalMeshLOD(USkeletalMesh* InSkeletalMesh, U
 
 	// Also sort the RequiredBones array to be strictly increasing.
 	NewLODModel.RequiredBones.Sort();
+	NewLODModel.ActiveBoneIndices.Sort();
 
 	// To be extra-nice, we apply the difference between the root transform of the meshes to the verts.
 	FMatrix LODToBaseTransform = InSkeletalMesh->GetRefPoseMatrix(0).InverseFast() * BaseSkeletalMesh->GetRefPoseMatrix(0);
diff --git a/Engine/Source/Runtime/Core/Public/UObject/ObjectVersion.h b/Engine/Source/Runtime/Core/Public/UObject/ObjectVersion.h
--- a/Engine/Source/Runtime/Core/Public/UObject/ObjectVersion.h
+++ b/Engine/Source/Runtime/Core/Public/UObject/ObjectVersion.h
@@ -440,6 +440,8 @@ enum EUnrealEngineObjectUE4Version
 	VER_UE4_PAWN_AUTO_POSSESS_AI,
 	// Added serialization of timezone to FTextHistory for AsDate operations.
 	VER_UE4_FTEXT_HISTORY_DATE_TIMEZONE,
+	// Sort ActiveBoneIndices on lods so that we can avoid doing it at run time
+	VER_UE4_SORT_ACTIVE_BONE_INDICES,
 	// -----<new versions can be added before this line>-------------------------------------------------
 	// - this needs to be the last line (see note below)
 	VER_UE4_AUTOMATIC_VERSION_PLUS_ONE,
diff --git a/Engine/Source/Runtime/Engine/Private/SkeletalMesh.cpp b/Engine/Source/Runtime/Engine/Private/SkeletalMesh.cpp
--- a/Engine/Source/Runtime/Engine/Private/SkeletalMesh.cpp
+++ b/Engine/Source/Runtime/Engine/Private/SkeletalMesh.cpp
@@ -2666,6 +2666,15 @@ void USkeletalMesh::PostLoad()
 		RefSkeleton.RebuildNameToIndexMap();
 	}
 
+	if (GetLinkerUE4Version() < VER_UE4_SORT_ACTIVE_BONE_INDICES)
+	{
+		for (int32 LodIndex = 0; LodIndex < LODInfo.Num(); LodIndex++)
+		{
+			FStaticLODModel & ThisLODModel = ImportedResource->LODModels[LodIndex];
+			ThisLODModel.ActiveBoneIndices.Sort();
+		}
+	}
+
 #if WITH_EDITORONLY_DATA
 	if (RetargetBasePose.Num() == 0)
 	{
diff --git a/Engine/Source/Runtime/Engine/Private/SkeletalMeshMerge.cpp b/Engine/Source/Runtime/Engine/Private/SkeletalMeshMerge.cpp
--- a/Engine/Source/Runtime/Engine/Private/SkeletalMeshMerge.cpp
+++ b/Engine/Source/Runtime/Engine/Private/SkeletalMeshMerge.cpp
@@ -567,6 +567,7 @@ void FSkeletalMeshMerge::GenerateLODModel( int32 LODIdx )
 
 	// sort required bone array in strictly increasing order
 	MergeLODModel.RequiredBones.Sort();
+	MergeLODModel.ActiveBoneIndices.Sort();
 
 	// copy the new vertices and indices to the vertex buffer for the new model
 	MergeLODModel.VertexBufferGPUSkin.SetUseFullPrecisionUVs(MergeMesh->bUseFullPrecisionUVs);
diff --git a/Engine/Source/Runtime/Engine/Private/SkeletalRender.cpp b/Engine/Source/Runtime/Engine/Private/SkeletalRender.cpp
--- a/Engine/Source/Runtime/Engine/Private/SkeletalRender.cpp
+++ b/Engine/Source/Runtime/Engine/Private/SkeletalRender.cpp
@@ -207,10 +207,7 @@ void UpdateRefToLocalMatrices( TArray<FMatrix>& ReferenceToLocal, const USkinned
 
 	const bool bIsMasterCompValid = MasterComp && InMeshComponent->MasterBoneMap.Num() == ThisMesh->RefSkeleton.GetNum();
 
-	TArray<FBoneIndexType> ActiveBoneIndices(LOD.ActiveBoneIndices);
-	ActiveBoneIndices.Sort();
-
-	const TArray<FBoneIndexType>* RequiredBoneSets[3] = { &ActiveBoneIndices, ExtraRequiredBoneIndices, NULL };
+	const TArray<FBoneIndexType>* RequiredBoneSets[3] = { &LOD.ActiveBoneIndices, ExtraRequiredBoneIndices, NULL };
 
 	const bool bBoneVisibilityStatesValid = InMeshComponent->BoneVisibilityStates.Num() == InMeshComponent->GetNumSpaceBases();
 
