Fixed crash in multiplayer PIE game viewport clients when 'show collision' is enabled.

diff --git a/Engine/Source/Runtime/Engine/Private/GameViewportClient.cpp b/Engine/Source/Runtime/Engine/Private/GameViewportClient.cpp
--- a/Engine/Source/Runtime/Engine/Private/GameViewportClient.cpp
+++ b/Engine/Source/Runtime/Engine/Private/GameViewportClient.cpp
@@ -2203,12 +2203,13 @@ void UGameViewportClient::ToggleShowCollision()
 			UPrimitiveComponent* PrimitiveComponent = *It;
 			if (!PrimitiveComponent->IsVisible() && PrimitiveComponent->IsCollisionEnabled() && PrimitiveComponent->GetScene() == GetWorld()->Scene)
 			{
-				check(PrimitiveComponent->GetOwner() && PrimitiveComponent->GetOwner()->GetWorld() && PrimitiveComponent->GetOwner()->GetWorld()->IsGameWorld());
-
-				// Save state before modifying the collision visibility
-				Mapping.Add(PrimitiveComponent, CollVisibilityState(PrimitiveComponent->bHiddenInGame, PrimitiveComponent->bVisible));
-				PrimitiveComponent->SetHiddenInGame(false);
-				PrimitiveComponent->SetVisibility(true);
+				if (PrimitiveComponent->GetOwner() && PrimitiveComponent->GetOwner()->GetWorld() && PrimitiveComponent->GetOwner()->GetWorld()->IsGameWorld())
+				{
+					// Save state before modifying the collision visibility
+					Mapping.Add(PrimitiveComponent, CollVisibilityState(PrimitiveComponent->bHiddenInGame, PrimitiveComponent->bVisible));
+					PrimitiveComponent->SetHiddenInGame(false);
+					PrimitiveComponent->SetVisibility(true);
+				}
 			}
 		}
 	}
