#ANIM: editing - fixed so that removing scale doesn't reset to 0.

https://jira.ol.epicgames.net/browse/UE-6462

diff --git a/Engine/Source/Runtime/Engine/Classes/Animation/AnimSequenceBase.h b/Engine/Source/Runtime/Engine/Classes/Animation/AnimSequenceBase.h
--- a/Engine/Source/Runtime/Engine/Classes/Animation/AnimSequenceBase.h
+++ b/Engine/Source/Runtime/Engine/Classes/Animation/AnimSequenceBase.h
@@ -352,6 +352,7 @@ struct FVectorCurve : public FAnimCurveBase
 	void CopyCurve(FVectorCurve& SourceCurve);
 	FVector Evaluate(float CurrentTime, float BlendWeight) const;
 	void UpdateOrAddKey(const FVector& NewKey, float CurrentTime);
+	bool DoesContainKey() const { return (FloatCurves[0].GetNumKeys() > 0 || FloatCurves[1].GetNumKeys() > 0 || FloatCurves[2].GetNumKeys() > 0);}
 };
 
 USTRUCT()
diff --git a/Engine/Source/Runtime/Engine/Private/Animation/AnimSequenceBase.cpp b/Engine/Source/Runtime/Engine/Private/Animation/AnimSequenceBase.cpp
--- a/Engine/Source/Runtime/Engine/Private/Animation/AnimSequenceBase.cpp
+++ b/Engine/Source/Runtime/Engine/Private/Animation/AnimSequenceBase.cpp
@@ -180,7 +180,14 @@ FTransform FTransformCurve::Evaluate(float CurrentTime, float BlendWeight) const
 {
 	FTransform Value;
 	Value.SetTranslation(TranslationCurve.Evaluate(CurrentTime, BlendWeight));
-	Value.SetScale3D(ScaleCurve.Evaluate(CurrentTime, BlendWeight));
+	if (ScaleCurve.DoesContainKey())
+	{
+		Value.SetScale3D(ScaleCurve.Evaluate(CurrentTime, BlendWeight));
+	}
+	else
+	{
+		Value.SetScale3D(FVector(1.f));
+	}
 
 	// blend rotation float curve
 	FVector RotationAsVector = RotationCurve.Evaluate(CurrentTime, BlendWeight);
