Packagemap network logging improvements

diff --git a/Engine/Source/Runtime/Engine/Classes/Engine/PackageMapClient.h b/Engine/Source/Runtime/Engine/Classes/Engine/PackageMapClient.h
--- a/Engine/Source/Runtime/Engine/Classes/Engine/PackageMapClient.h
+++ b/Engine/Source/Runtime/Engine/Classes/Engine/PackageMapClient.h
@@ -67,6 +67,8 @@ public:
 	bool			IsGUIDRegistered( const FNetworkGUID& NetGUID ) const;
 	bool			IsGUIDLoaded( const FNetworkGUID& NetGUID ) const;
 	bool			IsGUIDBroken( const FNetworkGUID& NetGUID, const bool bMustBeRegistered ) const;
+	FString			FullNetGUIDPath( const FNetworkGUID& NetGUID ) const;
+	void			GenerateFullNetGUIDPath_r( const FNetworkGUID& NetGUID, FString& FullPath ) const;
 
 	void			AsyncPackageCallback( const FName& PackageName, UPackage * Package );
 	
diff --git a/Engine/Source/Runtime/Engine/Private/PackageMapClient.cpp b/Engine/Source/Runtime/Engine/Private/PackageMapClient.cpp
--- a/Engine/Source/Runtime/Engine/Private/PackageMapClient.cpp
+++ b/Engine/Source/Runtime/Engine/Private/PackageMapClient.cpp
@@ -187,7 +187,9 @@ bool UPackageMapClient::SerializeNewActor(FArchive& Ar, class UActorChannel *Cha
 
 	Actor = Cast<AActor>(NewObj);
 
-	bool SpawnedActor = false;
+	// When we return an actor, we don't necessarily always spawn it (we might have found it already in memory)
+	// The calling code may want to know, so this is why we distinguish
+	bool bActorWasSpawned = false;
 
 	if ( Ar.AtEnd() && NetGUID.IsDynamic() )
 	{
@@ -205,7 +207,7 @@ bool UPackageMapClient::SerializeNewActor(FArchive& Ar, class UActorChannel *Cha
 		return false;		// This doesn't mean an error. This just simply means we didn't spawn an actor.
 	}
 
-	if (NetGUID.IsDynamic())
+	if ( NetGUID.IsDynamic() )
 	{
 		UObject* Archetype = NULL;
 		FVector_NetQuantize10 Location;
@@ -279,7 +281,7 @@ bool UPackageMapClient::SerializeNewActor(FArchive& Ar, class UActorChannel *Cha
 					Actor->PostNetReceiveVelocity(Velocity);
 
 					GuidCache->RegisterNetGUID_Client( NetGUID, Actor );
-					SpawnedActor = true;
+					bActorWasSpawned = true;
 				}
 				else
 				{
@@ -288,10 +290,14 @@ bool UPackageMapClient::SerializeNewActor(FArchive& Ar, class UActorChannel *Cha
 			}
 		}
 	}
+	else if ( Ar.IsLoading() && Actor == NULL )
+	{
+		UE_LOG( LogNetPackageMap, Warning, TEXT( "SerializeNewActor: Static actor failed to load: FullNetGuidPath: %s, Channel: %d" ), *GuidCache->FullNetGUIDPath( NetGUID ), Channel->ChIndex );
+	}
 
-	UE_LOG(LogNetPackageMap, Log, TEXT("SerializeNewActor END: Finished Serializing Actor %s <%s> on Channel %d"), Actor ? *Actor->GetName() : TEXT("NULL"), *NetGUID.ToString(), Channel->ChIndex );
+	UE_LOG( LogNetPackageMap, Log, TEXT( "SerializeNewActor END: Finished Serializing. Actor: %s, FullNetGUIDPath: %s, Channel: %d, IsLoading: %i, IsDynamic: %i" ), Actor ? *Actor->GetName() : TEXT("NULL"), *GuidCache->FullNetGUIDPath( NetGUID ), Channel->ChIndex, (int)Ar.IsLoading(), (int)NetGUID.IsDynamic() );
 
-	return SpawnedActor;
+	return bActorWasSpawned;
 }
 
 //--------------------------------------------------------------------
@@ -671,16 +677,7 @@ FNetworkGUID UPackageMapClient::InternalLoadObject( FArchive & Ar, UObject *& Ob
 	}
 	else if ( Object == NULL && !GuidCache->ShouldIgnoreWhenMissing( NetGUID ) )
 	{
-		const FNetGuidCacheObject* CacheObject = GuidCache->ObjectLookup.Find( NetGUID );
-
-		if ( CacheObject == NULL )
-		{
-			UE_LOG( LogNetPackageMap, Warning, TEXT( "InternalLoadObject: Unable to resolve object (not in cache). NetGUID: %s" ), *NetGUID.ToString() );
-		}
-		else
-		{
-			UE_LOG( LogNetPackageMap, Warning, TEXT( "InternalLoadObject: Unable to resolve object. Cached Path: %s, NetGUID: %s" ), *CacheObject->PathName.ToString(), *NetGUID.ToString() );
-		}
+		UE_LOG( LogNetPackageMap, Warning, TEXT( "InternalLoadObject: Unable to resolve object. FullNetGUIDPath: %s" ), *GuidCache->FullNetGUIDPath( NetGUID ) );
 	}
 
 	return NetGUID;
@@ -1517,7 +1514,7 @@ void FNetGUIDCache::RegisterNetGUID_Client( const FNetworkGUID& NetGUID, const U
 	check( !NetGUID.IsDefault() );
 	check( NetGUID.IsDynamic() );	// Clients should only assign dynamic guids through here (static guids go through RegisterNetGUIDFromPath_Client)
 
-	UE_LOG( LogNetPackageMap, Log, TEXT( "RegisterNetGUID_Client: Assigning %s NetGUID <%s>" ), Object ? *Object->GetName() : TEXT( "NULL" ), *NetGUID.ToString() );
+	UE_LOG( LogNetPackageMap, Log, TEXT( "RegisterNetGUID_Client: NetGUID: %s, Object: %s" ), *NetGUID.ToString(), Object ? *Object->GetName() : TEXT( "NULL" ) );
 	
 	//
 	// If we have an existing entry, make sure things match up properly
@@ -1530,7 +1527,7 @@ void FNetGUIDCache::RegisterNetGUID_Client( const FNetworkGUID& NetGUID, const U
 	{
 		if ( ExistingCacheObjectPtr->PathName != NAME_None )
 		{
-			UE_LOG( LogNetPackageMap, Warning, TEXT( "RegisterNetGUID_Client: Guid with pathname. Path: %s, NetGUID: %s, OuterGUID: %s" ), *ExistingCacheObjectPtr->PathName.ToString(), *NetGUID.ToString(), *ExistingCacheObjectPtr->OuterGUID.ToString() );
+			UE_LOG( LogNetPackageMap, Warning, TEXT( "RegisterNetGUID_Client: Guid with pathname. FullNetGUIDPath: %s" ), *FullNetGUIDPath( NetGUID ) );
 		}
 
 		// If this net guid was found but the old object is NULL, this can happen due to:
@@ -1586,6 +1583,8 @@ void FNetGUIDCache::RegisterNetGUIDFromPath_Client( const FNetworkGUID& NetGUID,
 	check( !IsNetGUIDAuthority() );		// Server never calls this locally
 	check( !NetGUID.IsDefault() );
 
+	UE_LOG( LogNetPackageMap, Log, TEXT( "RegisterNetGUIDFromPath_Client: NetGUID: %s, PathName: %s, OuterGUID: %s" ), *NetGUID.ToString(), *PathName, *OuterGUID.ToString() );
+
 	const FNetGuidCacheObject* ExistingCacheObjectPtr = ObjectLookup.Find( NetGUID );
 
 	// If we find this guid, make sure nothing changes
@@ -1711,7 +1710,7 @@ UObject* FNetGUIDCache::GetObjectFromNetGUID( const FNetworkGUID& NetGUID, const
 
 	if ( IsNetGUIDAuthority() )
 	{
-		UE_LOG( LogNetPackageMap, Error, TEXT( "GetObjectFromNetGUID: Guid with no object on server. Path: %s, NetGUID: %s" ), *CacheObjectPtr->PathName.ToString(), *NetGUID.ToString() );
+		UE_LOG( LogNetPackageMap, Error, TEXT( "GetObjectFromNetGUID: Guid with no object on server. FullNetGUIDPath: %s" ), *FullNetGUIDPath( NetGUID ) );
 		return NULL;
 	}
 
@@ -1735,7 +1734,7 @@ UObject* FNetGUIDCache::GetObjectFromNetGUID( const FNetworkGUID& NetGUID, const
 			// Shouldn't be possible, but just in case...
 			if ( CacheObjectPtr->OuterGUID.IsStatic() )
 			{
-				UE_LOG( LogNetPackageMap, Error, TEXT( "GetObjectFromNetGUID: Static outer not registered. Path: %s, NetGUID: %s, OuterGUID: %s" ), *CacheObjectPtr->PathName.ToString(), *NetGUID.ToString(), *CacheObjectPtr->OuterGUID.ToString() );
+				UE_LOG( LogNetPackageMap, Error, TEXT( "GetObjectFromNetGUID: Static outer not registered. FullNetGUIDPath: %s" ), *FullNetGUIDPath( NetGUID ) );
 				CacheObjectPtr->bIsBroken = 1;	// Set this to 1 so that we don't keep spamming
 			}
 			return NULL;
@@ -1744,7 +1743,7 @@ UObject* FNetGUIDCache::GetObjectFromNetGUID( const FNetworkGUID& NetGUID, const
 		// If outer is broken, we will never load, set outselves to broken as well and bail
 		if ( OuterCacheObject->bIsBroken )
 		{
-			UE_LOG( LogNetPackageMap, Error, TEXT( "GetObjectFromNetGUID: Outer is broken. Path: %s, Outer: %s, NetGUID: %s" ), *CacheObjectPtr->PathName.ToString(), *OuterCacheObject->PathName.ToString(), *NetGUID.ToString() );
+			UE_LOG( LogNetPackageMap, Error, TEXT( "GetObjectFromNetGUID: Outer is broken. FullNetGUIDPath: %s" ), *FullNetGUIDPath( NetGUID ) );
 			CacheObjectPtr->bIsBroken = 1;	// Set this to 1 so that we don't keep spamming
 			return NULL;
 		}
@@ -1758,7 +1757,7 @@ UObject* FNetGUIDCache::GetObjectFromNetGUID( const FNetworkGUID& NetGUID, const
 			// If the outer is missing, warn unless told to ignore
 			if ( !ShouldIgnoreWhenMissing( CacheObjectPtr->OuterGUID ) )
 			{
-				UE_LOG( LogNetPackageMap, Error, TEXT( "GetObjectFromNetGUID: Failed to find outer. Path: %s, Outer: %s, NetGUID: %s" ), *CacheObjectPtr->PathName.ToString(), *OuterCacheObject->PathName.ToString(), *NetGUID.ToString() );
+				UE_LOG( LogNetPackageMap, Error, TEXT( "GetObjectFromNetGUID: Failed to find outer. FullNetGUIDPath: %s" ), *FullNetGUIDPath( NetGUID ) );
 			}
 
 			return NULL;
@@ -1842,7 +1841,7 @@ UObject* FNetGUIDCache::GetObjectFromNetGUID( const FNetworkGUID& NetGUID, const
 		if ( !CacheObjectPtr->bIgnoreWhenMissing )
 		{
 			CacheObjectPtr->bIsBroken = 1;	// Set this to 1 so that we don't keep spamming
-			UE_LOG( LogNetPackageMap, Error, TEXT( "GetObjectFromNetGUID: Failed to resolve path. Path: %s, Outer: %s, NetGUID: %s" ), *CacheObjectPtr->PathName.ToString(), ObjOuter != NULL ? *ObjOuter->GetPathName() : TEXT( "NULL" ), *NetGUID.ToString() );
+			UE_LOG( LogNetPackageMap, Error, TEXT( "GetObjectFromNetGUID: Failed to resolve path. FullNetGUIDPath: %s" ), *FullNetGUIDPath( NetGUID ) );
 		}
 
 		return NULL;
@@ -1856,7 +1855,7 @@ UObject* FNetGUIDCache::GetObjectFromNetGUID( const FNetworkGUID& NetGUID, const
 		{
 			// This isn't really a package but it should be
 			CacheObjectPtr->bIsBroken = true;
-			UE_LOG( LogNetPackageMap, Error, TEXT( "GetObjectFromNetGUID: Object is not a package but shoule be! Path: %s, NetGUID: %s" ), *CacheObjectPtr->PathName.ToString(), *NetGUID.ToString() );
+			UE_LOG( LogNetPackageMap, Error, TEXT( "GetObjectFromNetGUID: Object is not a package but should be! Path: %s, NetGUID: %s" ), *CacheObjectPtr->PathName.ToString(), *NetGUID.ToString() );
 			return NULL;
 		}
 
@@ -2013,6 +2012,65 @@ bool FNetGUIDCache::IsGUIDBroken( const FNetworkGUID& NetGUID, const bool bMustB
 	return CacheObjectPtr->bIsBroken;
 }
 
+FString FNetGUIDCache::FullNetGUIDPath( const FNetworkGUID& NetGUID ) const
+{
+	FString FullPath;
+
+	GenerateFullNetGUIDPath_r( NetGUID, FullPath );
+
+	return FullPath;
+}
+
+void FNetGUIDCache::GenerateFullNetGUIDPath_r( const FNetworkGUID& NetGUID, FString& FullPath ) const
+{
+	if ( !NetGUID.IsValid() )
+	{
+		// This is the end of the outer chain, we're done
+		return;
+	}
+
+	const FNetGuidCacheObject* CacheObject = ObjectLookup.Find( NetGUID );
+
+	if ( CacheObject == NULL )
+	{
+		// Doh, this shouldn't be possible, but if this happens, we can't continue
+		// So warn, and return
+		FullPath += FString::Printf( TEXT( "[%s]NOT_IN_CACHE" ), *NetGUID.ToString() );
+		return;
+	}
+
+	GenerateFullNetGUIDPath_r( CacheObject->OuterGUID, FullPath );
+
+	if ( !FullPath.IsEmpty() )
+	{
+		FullPath += TEXT( "." );
+	}
+
+	// Prefer the object name first, since non stable named objects don't store the path
+	if ( CacheObject->Object.IsValid() )
+	{
+		// Sanity check that the names match if the path was stored
+		if ( CacheObject->PathName != NAME_None && CacheObject->Object->GetName() != CacheObject->PathName.ToString() )
+		{
+			UE_LOG( LogNetPackageMap, Warning, TEXT( "GenerateFullNetGUIDPath_r: Name mismatch! %s != %s" ), *CacheObject->PathName.ToString(), *CacheObject->Object->GetName() );	
+		}
+
+		FullPath += FString::Printf( TEXT( "[%s]%s" ), *NetGUID.ToString(), *CacheObject->Object->GetName() );
+	}
+	else
+	{
+		if ( CacheObject->PathName == NAME_None )
+		{
+			// This can happen when a non stably named object is NULL
+			FullPath += FString::Printf( TEXT( "[%s]EMPTY" ), *NetGUID.ToString() );
+		}
+		else
+		{
+			FullPath += FString::Printf( TEXT( "[%s]%s" ), *NetGUID.ToString(), *CacheObject->PathName.ToString() );			
+		}
+	}
+}
+
 //------------------------------------------------------
 // Debug command to see how many times we've exported each NetGUID
 // Used for measuring inefficiencies. Some duplication is unavoidable since we cannot garuntee atomicicity across multiple channels.
