Moving a latent open map command to the automation common files.

diff --git a/Engine/Source/Runtime/Engine/Private/Tests/AutomationCommon.cpp b/Engine/Source/Runtime/Engine/Private/Tests/AutomationCommon.cpp
--- a/Engine/Source/Runtime/Engine/Private/Tests/AutomationCommon.cpp
+++ b/Engine/Source/Runtime/Engine/Private/Tests/AutomationCommon.cpp
@@ -43,4 +43,14 @@ bool FTakeEditorScreenshotCommand::Update()
 {
 	AutomationCommon::SaveWindowAsScreenshot(ScreenshotParameters.CurrentWindow.ToSharedRef(), ScreenshotParameters.ScreenshotName);
 	return true;
+}
+
+bool FLoadGameMapCommand::Update()
+{
+	check(GEngine->GetWorldContexts().Num() == 1);
+	check(GEngine->GetWorldContexts()[0].WorldType == EWorldType::Game);
+
+	UE_LOG(LogEngineAutomationTests, Log, TEXT("Loading Map Now. '%s'"), *MapName);
+	GEngine->Exec(GEngine->GetWorldContexts()[0].World(), *FString::Printf(TEXT("Open %s"), *MapName));
+	return true;
 }
\ No newline at end of file
diff --git a/Engine/Source/Runtime/Engine/Private/Tests/EngineAutomationTests.cpp b/Engine/Source/Runtime/Engine/Private/Tests/EngineAutomationTests.cpp
--- a/Engine/Source/Runtime/Engine/Private/Tests/EngineAutomationTests.cpp
+++ b/Engine/Source/Runtime/Engine/Private/Tests/EngineAutomationTests.cpp
@@ -13,7 +13,6 @@
 #include "PlatformFeatures.h"
 #include "SaveGameSystem.h"
 
-DEFINE_LOG_CATEGORY_STATIC(LogEngineAutomationTests, Log, All);
 
 namespace
 {
@@ -336,22 +335,6 @@ bool FSaveGameTest::RunTest(const FString& Parameters)
 	return LoadedData == SavedData;
 }
 
-/**
- * Latent command to load a map in game
- */
-DEFINE_LATENT_AUTOMATION_COMMAND_ONE_PARAMETER(FLoadGameMapCommand, FString, MapName);
-
-bool FLoadGameMapCommand::Update()
-{
-	check(GEngine->GetWorldContexts().Num() == 1);
-	check(GEngine->GetWorldContexts()[0].WorldType == EWorldType::Game);
-
-	UE_LOG(LogEngineAutomationTests, Log, TEXT("Loading Map Now. '%s'"), *MapName);
-	GEngine->Exec(GEngine->GetWorldContexts()[0].World(), *FString::Printf(TEXT("Open %s"), *MapName));
-	return true;
-}
-
-
 /**
  * Automation test to load a map and capture FPS performance charts
  */
diff --git a/Engine/Source/Runtime/Engine/Public/Tests/AutomationCommon.h b/Engine/Source/Runtime/Engine/Public/Tests/AutomationCommon.h
--- a/Engine/Source/Runtime/Engine/Public/Tests/AutomationCommon.h
+++ b/Engine/Source/Runtime/Engine/Public/Tests/AutomationCommon.h
@@ -9,6 +9,7 @@
 // Common Latent commands which are used across test type. I.e. Engine, Network, etc...
 
 DEFINE_LOG_CATEGORY_STATIC(LogEditorAutomationTests, Log, All);
+DEFINE_LOG_CATEGORY_STATIC(LogEngineAutomationTests, Log, All);
 DEFINE_LOG_CATEGORY_STATIC(LogAnalytics, Log, All);
 
 /** Common automation functions */
@@ -102,3 +103,10 @@ DEFINE_ENGINE_LATENT_AUTOMATION_COMMAND_ONE_PARAMETER(FTakeActiveEditorScreensho
  * Take a screenshot of the active window
  */
 DEFINE_ENGINE_LATENT_AUTOMATION_COMMAND_ONE_PARAMETER(FTakeEditorScreenshotCommand, WindowScreenshotParameters, ScreenshotParameters);
+
+/**
+* Latent command to load a map in game
+*/
+DEFINE_LATENT_AUTOMATION_COMMAND_ONE_PARAMETER(FLoadGameMapCommand, FString, MapName);
+
+
