Compile fix

diff --git a/Engine/Source/Runtime/Engine/Public/MaterialShared.h b/Engine/Source/Runtime/Engine/Public/MaterialShared.h
--- a/Engine/Source/Runtime/Engine/Public/MaterialShared.h
+++ b/Engine/Source/Runtime/Engine/Public/MaterialShared.h
@@ -562,7 +562,7 @@ public:
 	 * @param Material - The material we also wish to apply the compiled shader map to.
 	 * @return True if the shader map was being compiled and we added Material to the list to be applied.
 	 */
-	bool TryToAddToSimultaneousCompilation(FMaterial* Material);
+	bool TryToAddToExistingCompilationTask(FMaterial* Material);
 
 	/** Builds a list of the shaders in a shader map. */
 	ENGINE_API void GetShaderList(TMap<FShaderId,FShader*>& OutShaders) const;
