Asserts to catch GC'ed textures used in rendering earlier

diff --git a/Engine/Source/Runtime/Engine/Private/Materials/MaterialUniformExpressions.cpp b/Engine/Source/Runtime/Engine/Private/Materials/MaterialUniformExpressions.cpp
--- a/Engine/Source/Runtime/Engine/Private/Materials/MaterialUniformExpressions.cpp
+++ b/Engine/Source/Runtime/Engine/Private/Materials/MaterialUniformExpressions.cpp
@@ -294,6 +294,8 @@ FUniformBufferRHIRef FUniformExpressionSet::CreateUniformBuffer(const FMaterialR
 		void** ResourceTable = (void**)((uint8*)TempBuffer + UniformBufferStruct->GetLayout().ResourceOffset);
 		check(((UPTRINT)ResourceTable & 0x7) == 0);
 
+		check(UniformBufferStruct->GetLayout().Resources.Num() == Uniform2DTextureExpressions.Num() * 2 + UniformCubeTextureExpressions.Num() * 2 + 2);
+
 		// Cache 2D texture uniform expressions.
 		for(int32 ExpressionIndex = 0;ExpressionIndex < Uniform2DTextureExpressions.Num();ExpressionIndex++)
 		{
@@ -302,6 +304,8 @@ FUniformBufferRHIRef FUniformExpressionSet::CreateUniformBuffer(const FMaterialR
 			Uniform2DTextureExpressions[ExpressionIndex]->GetTextureValue(MaterialRenderContext,MaterialRenderContext.Material,Value,SourceMode);
 			if(Value && Value->Resource)
 			{
+				// UMaterial / UMaterialInstance should have caused all dependent textures to be PostLoaded, which initializes their rendering resource
+				check(Value->TextureReference.TextureReferenceRHI);
 				*ResourceTable++ = Value->TextureReference.TextureReferenceRHI;
 				FSamplerStateRHIRef* SamplerSource = &Value->Resource->SamplerStateRHI;
 
@@ -331,6 +335,7 @@ FUniformBufferRHIRef FUniformExpressionSet::CreateUniformBuffer(const FMaterialR
 			UniformCubeTextureExpressions[ExpressionIndex]->GetTextureValue(MaterialRenderContext,MaterialRenderContext.Material,Value,SourceMode);
 			if(Value && Value->Resource)
 			{
+				check(Value->TextureReference.TextureReferenceRHI);
 				*ResourceTable++ = Value->TextureReference.TextureReferenceRHI;
 				FSamplerStateRHIRef* SamplerSource = &Value->Resource->SamplerStateRHI;
 
