UE4 - Per-frame cpu evaluated material expressions & added Period to Time expression

diff --git a/Engine/Shaders/ShaderVersion.usf b/Engine/Shaders/ShaderVersion.usf
--- a/Engine/Shaders/ShaderVersion.usf
+++ b/Engine/Shaders/ShaderVersion.usf
@@ -2,4 +2,4 @@
 // Each time the console command is executed it generates a new GUID. As this file is included
 // in common.usf (which should be included in any shader) it allows to invalidate the shader DDC.
 // 
-// GUID = 9F1892334C4B8A2429020D883842B32E
\ No newline at end of file
+// GUID = {CAC8D6F0-2E83-4C3D-A767-584D1E3DC80E}
diff --git a/Engine/Source/Editor/UnrealEd/Private/Lightmass/LightmassRender.cpp b/Engine/Source/Editor/UnrealEd/Private/Lightmass/LightmassRender.cpp
--- a/Engine/Source/Editor/UnrealEd/Private/Lightmass/LightmassRender.cpp
+++ b/Engine/Source/Editor/UnrealEd/Private/Lightmass/LightmassRender.cpp
@@ -25,101 +25,101 @@ struct FLightmassMaterialCompiler : public FProxyMaterialCompiler
 	{}
 
 	// gets value stored by SetMaterialProperty()
-	virtual EShaderFrequency GetCurrentShaderFrequency() const
+	virtual EShaderFrequency GetCurrentShaderFrequency() const override
 	{
 		// not used by Lightmass
 		return SF_Pixel;
 	}
 
-	virtual int32 ParticleMacroUV()
+	virtual int32 ParticleMacroUV() override
 	{
 		return Compiler->ParticleMacroUV();
 	}
 
-	virtual int32 ParticleRelativeTime()
+	virtual int32 ParticleRelativeTime() override
 	{
 		return Compiler->Constant(0.0f);
 	}
 
-	virtual int32 ParticleMotionBlurFade()
+	virtual int32 ParticleMotionBlurFade() override
 	{
 		return Compiler->Constant(1.0f);
 	}
 
-	virtual int32 ParticleDirection()
+	virtual int32 ParticleDirection() override
 	{
 		return Compiler->Constant3(0.0f, 0.0f, 0.0f);
 	}
 
-	virtual int32 ParticleSpeed()
+	virtual int32 ParticleSpeed() override
 	{
 		return Compiler->Constant(0.0f);
 	}
 	
-	virtual int32 ParticleSize()
+	virtual int32 ParticleSize() override
 	{
 		return Compiler->Constant2(0.0f,0.0f);
 	}
 
-	virtual int32 WorldPosition(EWorldPositionIncludedOffsets WorldPositionIncludedOffsets)
+	virtual int32 WorldPosition(EWorldPositionIncludedOffsets WorldPositionIncludedOffsets) override
 	{
 		//UE_LOG(LogLightmassRender, Log, TEXT("Lightmass material compiler has encountered WorldPosition... Forcing constant (0.0f,0.0f,0.0f)."));
 		return Compiler->Constant3(0.0f,0.0f,0.0f);
 	}
 
-	virtual int32 ObjectWorldPosition() 
+	virtual int32 ObjectWorldPosition() override
 	{
 		//UE_LOG(LogLightmassRender, Log, TEXT("Lightmass material compiler has encountered ObjectWorldPosition... Forcing constant (0.0f,0.0f,0.0f)."));
 		return Compiler->Constant3(0.0f,0.0f,0.0f);
 	}
 
-	virtual int32 ObjectRadius() 
+	virtual int32 ObjectRadius() override
 	{
 		//UE_LOG(LogLightmassRender, Log, TEXT("Lightmass material compiler has encountered ObjectRadius... Forcing constant 500.0f."));
 		return Compiler->Constant(500);
 	}
 
-	virtual int32 ObjectBounds() 
+	virtual int32 ObjectBounds() override
 	{
 		//UE_LOG(LogLightmassRender, Log, TEXT("Lightmass material compiler has encountered ObjectBounds... Forcing constant (0,0,0)."));
 		return Compiler->Constant3(0,0,0);
 	}
 
-	virtual int32 DistanceCullFade() 
+	virtual int32 DistanceCullFade() override
 	{
 		return Compiler->Constant(1.0f);
 	}
 
-	virtual int32 ActorWorldPosition() 
+	virtual int32 ActorWorldPosition() override
 	{
 		return Compiler->Constant3(0.0f,0.0f,0.0f);
 	}
 
-	virtual int32 CameraWorldPosition()
+	virtual int32 CameraWorldPosition() override
 	{
 		//UE_LOG(LogLightmassRender, Log, TEXT("Lightmass material compiler has encountered CameraWorldPosition... Forcing constant (0.0f,0.0f,0.0f)."));
 		return Compiler->Constant3(0.0f,0.0f,0.0f);
 	}
 
-	virtual int32 CameraVector()
+	virtual int32 CameraVector() override
 	{
 		//UE_LOG(LogLightmassRender, Log, TEXT("Lightmass material compiler has encountered CameraVector... Forcing constant (0.0f,0.0f,1.0f)."));
 		return Compiler->Constant3(0.0f,0.0f,1.0f);
 	}
 
-	virtual int32 LightVector()
+	virtual int32 LightVector() override
 	{
 		//UE_LOG(LogLightmassRender, Log, TEXT("Lightmass material compiler has encountered LightVector... Forcing constant (1.0f,0.0f,0.0f)."));
 		return Compiler->Constant3(1.0f,0.0f,0.0f);
 	}
 
-	virtual int32 ReflectionVector()
+	virtual int32 ReflectionVector() override
 	{
 		//UE_LOG(LogLightmassRender, Log, TEXT("Lightmass material compiler has encountered ReflectionVector... Forcing constant (0.0f,0.0f,-1.0f)."));
 		return Compiler->Constant3(0.0f,0.0f,-1.0f);
 	}
 
-	virtual int32 ReflectionAboutCustomWorldNormal(int32 CustomWorldNormal, int32 bNormalizeCustomWorldNormal)
+	virtual int32 ReflectionAboutCustomWorldNormal(int32 CustomWorldNormal, int32 bNormalizeCustomWorldNormal) override
 	{
 		//UE_LOG(LogLightmassRender, Log, TEXT("Lightmass material compiler has encountered ReflectionAboutCustomNormalVector... Forcing constant (0.0f,0.0f,-1.0f)."));
 		return Compiler->Constant3(0.0f,0.0f,-1.0f);
@@ -128,7 +128,7 @@ struct FLightmassMaterialCompiler : public FProxyMaterialCompiler
 	/**
 	 *	Generate shader code for transforming a vector
 	 */
-	virtual int32 TransformVector(uint8 SourceCoordType,uint8 DestCoordType,int32 A)
+	virtual int32 TransformVector(uint8 SourceCoordType,uint8 DestCoordType,int32 A) override
 	{
 		//UE_LOG(LogLightmassRender, Log, TEXT("Lightmass material compiler has encountered TransformVector... Passing thru source vector untouched."));
 		return A;
@@ -140,33 +140,33 @@ struct FLightmassMaterialCompiler : public FProxyMaterialCompiler
 	 *	@param	CoordType - type of transform to apply. see EMaterialExpressionTransformPosition 
 	 *	@param	A - index for input vector parameter's code
 	 */
-	virtual int32 TransformPosition(uint8 SourceCoordType,uint8 DestCoordType,int32 A)
+	virtual int32 TransformPosition(uint8 SourceCoordType,uint8 DestCoordType,int32 A) override
 	{
 		//UE_LOG(LogLightmassRender, Log, TEXT("Lightmass material compiler has encountered TransformPosition... Passing thru source vector untouched."));
 		return A;
 	}
 
-	virtual int32 VertexColor()
+	virtual int32 VertexColor() override
 	{
 		//UE_LOG(LogLightmassRender, Log, TEXT("Lightmass material compiler has encountered VertexColor... Forcing constant (1.0f,1.0f,1.0f,1.0f)."));
 		return Compiler->Constant4(1.0f,1.0f,1.0f,1.0f);
 	}
 
-	virtual int32 RealTime()
+	virtual int32 RealTime(bool bPeriodic, float Period) override
 	{
 		//UE_LOG(LogLightmassRender, Log, TEXT("Lightmass material compiler has encountered RealTime... Forcing constant 0.0f."));
 		return Compiler->Constant(0.0f);
 	}
 
-	virtual int32 GameTime()
+	virtual int32 GameTime(bool bPeriodic, float Period) override
 	{
 		//UE_LOG(LogLightmassRender, Log, TEXT("Lightmass material compiler has encountered GameTime... Forcing constant 0.0f."));
 		return Compiler->Constant(0.0f);
 	}
 
-	virtual int32 LightmassReplace(int32 Realtime, int32 Lightmass) { return Lightmass; }
+	virtual int32 LightmassReplace(int32 Realtime, int32 Lightmass) override { return Lightmass; }
 
-	virtual int32 GIReplace(int32 Direct, int32 StaticIndirect, int32 DynamicIndirect) { return StaticIndirect; }
+	virtual int32 GIReplace(int32 Direct, int32 StaticIndirect, int32 DynamicIndirect) override { return StaticIndirect; }
 };
 
 /**
diff --git a/Engine/Source/Runtime/Core/Public/UObject/ObjectVersion.h b/Engine/Source/Runtime/Core/Public/UObject/ObjectVersion.h
--- a/Engine/Source/Runtime/Core/Public/UObject/ObjectVersion.h
+++ b/Engine/Source/Runtime/Core/Public/UObject/ObjectVersion.h
@@ -442,6 +442,8 @@ enum EUnrealEngineObjectUE4Version
 	VER_UE4_FTEXT_HISTORY_DATE_TIMEZONE,
 	// Sort ActiveBoneIndices on lods so that we can avoid doing it at run time
 	VER_UE4_SORT_ACTIVE_BONE_INDICES,
+	// Added per-frame material uniform expressions
+	VER_UE4_PERFRAME_MATERIAL_UNIFORM_EXPRESSIONS,
 	// -----<new versions can be added before this line>-------------------------------------------------
 	// - this needs to be the last line (see note below)
 	VER_UE4_AUTOMATIC_VERSION_PLUS_ONE,
diff --git a/Engine/Source/Runtime/Engine/Classes/Materials/MaterialExpressionTime.h b/Engine/Source/Runtime/Engine/Classes/Materials/MaterialExpressionTime.h
--- a/Engine/Source/Runtime/Engine/Classes/Materials/MaterialExpressionTime.h
+++ b/Engine/Source/Runtime/Engine/Classes/Materials/MaterialExpressionTime.h
@@ -2,6 +2,7 @@
 
 
 #pragma once
+
 #include "Materials/MaterialExpression.h"
 #include "MaterialExpressionTime.generated.h"
 
@@ -14,6 +15,13 @@ class UMaterialExpressionTime : public UMaterialExpression
 	UPROPERTY(EditAnywhere, Category=MaterialExpressionTime)
 	uint32 bIgnorePause:1;
 
+	/** Enables or disables the Period value. */
+	UPROPERTY()
+	uint32 bOverride_Period:1;
+
+	/** Time will loop around once it gets to Period. */
+	UPROPERTY(EditAnywhere, Category=MaterialExpressionTime, meta = (editcondition = "bOverride_Period", ClampMin = "0.0", ToolTip = "Period at which to wrap around time"))
+	float Period;
 
 	// Begin UMaterialExpression Interface
 	virtual int32 Compile(class FMaterialCompiler* Compiler, int32 OutputIndex, int32 MultiplexIndex) override;
@@ -22,6 +30,3 @@ class UMaterialExpressionTime : public UMaterialExpression
 	// End UMaterialExpression Interface
 
 };
-
-
-
diff --git a/Engine/Source/Runtime/Engine/Private/Materials/HLSLMaterialTranslator.h b/Engine/Source/Runtime/Engine/Private/Materials/HLSLMaterialTranslator.h
--- a/Engine/Source/Runtime/Engine/Private/Materials/HLSLMaterialTranslator.h
+++ b/Engine/Source/Runtime/Engine/Private/Materials/HLSLMaterialTranslator.h
@@ -50,6 +50,19 @@ EMaterialValueType GetVectorType(uint32 NumComponents)
 	};
 }
 
+static inline int32 SwizzleComponentToIndex(TCHAR Component)
+{
+	switch (Component)
+	{
+	case TCHAR('x'): case TCHAR('X'): case TCHAR('r'): case TCHAR('R'): return 0;
+	case TCHAR('y'): case TCHAR('Y'): case TCHAR('g'): case TCHAR('G'): return 1;
+	case TCHAR('z'): case TCHAR('Z'): case TCHAR('b'): case TCHAR('B'): return 2;
+	case TCHAR('w'): case TCHAR('W'): case TCHAR('a'): case TCHAR('A'): return 3;
+	default:
+		return -1;
+	}
+}
+
 struct FShaderCodeChunk
 {
 	/** 
@@ -481,6 +494,19 @@ public:
 				ResourcesString += CustomExpressionImplementations[ExpressionIndex] + "\r\n\r\n";
 			}
 
+			// Per frame expressions
+			{
+				for (int32 Index = 0, Num = MaterialCompilationOutput.UniformExpressionSet.PerFrameUniformScalarExpressions.Num(); Index < Num; ++Index)
+				{
+					ResourcesString += FString::Printf(TEXT("float UE_Material_PerFrameScalarExpression%u;"), Index) + "\r\n\r\n";
+				}
+
+				for (int32 Index = 0, Num = MaterialCompilationOutput.UniformExpressionSet.PerFrameUniformVectorExpressions.Num(); Index < Num; ++Index)
+				{
+					ResourcesString += FString::Printf(TEXT("float4 UE_Material_PerFrameVectorExpression%u;"), Index) + "\r\n\r\n";
+				}
+			}
+
 			for(uint32 PropertyId = 0; PropertyId < MP_MAX; ++PropertyId)
 			{
 				if(PropertyId == MP_MaterialAttributes )
@@ -1000,9 +1026,9 @@ protected:
 			}
 		}
 
-		int32		BufferSize		= 256;
+		int32	BufferSize		= 256;
 		TCHAR*	FormattedCode	= NULL;
-		int32		Result			= -1;
+		int32	Result			= -1;
 
 		while(Result == -1)
 		{
@@ -1041,14 +1067,21 @@ protected:
 		TCHAR FormattedCode[MAX_SPRINTF]=TEXT("");
 		if(CodeChunk.Type == MCT_Float)
 		{
-			const static TCHAR IndexToMask[] = {'x', 'y', 'z', 'w'};
-			const int32 ScalarInputIndex = MaterialCompilationOutput.UniformExpressionSet.UniformScalarExpressions.AddUnique(CodeChunk.UniformExpression);
-			// Update the above FMemory::Malloc if this FCString::Sprintf grows in size, e.g. %s, ...
-			FCString::Sprintf(FormattedCode, TEXT("Material.ScalarExpressions[%u].%c"), ScalarInputIndex / 4, IndexToMask[ScalarInputIndex % 4]);
+			if (CodeChunk.UniformExpression->IsChangingPerFrame())
+			{
+				const int32 ScalarInputIndex = MaterialCompilationOutput.UniformExpressionSet.PerFrameUniformScalarExpressions.AddUnique(CodeChunk.UniformExpression);
+				FCString::Sprintf(FormattedCode, TEXT("UE_Material_PerFrameScalarExpression%u"), ScalarInputIndex);
+			}
+			else
+			{
+				const static TCHAR IndexToMask[] = {'x', 'y', 'z', 'w'};
+				const int32 ScalarInputIndex = MaterialCompilationOutput.UniformExpressionSet.UniformScalarExpressions.AddUnique(CodeChunk.UniformExpression);
+				// Update the above FMemory::Malloc if this FCString::Sprintf grows in size, e.g. %s, ...
+				FCString::Sprintf(FormattedCode, TEXT("Material.ScalarExpressions[%u].%c"), ScalarInputIndex / 4, IndexToMask[ScalarInputIndex % 4]);
+			}
 		}
 		else if(CodeChunk.Type & MCT_Float)
 		{
-			const int32 VectorInputIndex = MaterialCompilationOutput.UniformExpressionSet.UniformVectorExpressions.AddUnique(CodeChunk.UniformExpression);
 			const TCHAR* Mask;
 			switch(CodeChunk.Type)
 			{
@@ -1059,10 +1092,20 @@ protected:
 			default: Mask = TEXT(""); break;
 			};
 
-			FCString::Sprintf(FormattedCode, TEXT("Material.VectorExpressions[%u]%s"), VectorInputIndex, Mask);
+			if (CodeChunk.UniformExpression->IsChangingPerFrame())
+			{
+				const int32 VectorInputIndex = MaterialCompilationOutput.UniformExpressionSet.PerFrameUniformVectorExpressions.AddUnique(CodeChunk.UniformExpression);
+				FCString::Sprintf(FormattedCode, TEXT("UE_Material_PerFrameVectorExpression%u%s"), VectorInputIndex, Mask);
+			}
+			else
+			{
+				const int32 VectorInputIndex = MaterialCompilationOutput.UniformExpressionSet.UniformVectorExpressions.AddUnique(CodeChunk.UniformExpression);
+				FCString::Sprintf(FormattedCode, TEXT("Material.VectorExpressions[%u]%s"), VectorInputIndex, Mask);
+			}
 		}
 		else if(CodeChunk.Type & MCT_Texture)
 		{
+			check(!CodeChunk.UniformExpression->IsChangingPerFrame());
 			int32 TextureInputIndex = INDEX_NONE;
 			const TCHAR* BaseName = TEXT("");
 			switch(CodeChunk.Type)
@@ -1533,14 +1576,44 @@ protected:
 		return AddUniformExpression(new FMaterialUniformExpressionConstant(FLinearColor(X,Y,Z,W),MCT_Float4),MCT_Float4,TEXT("MaterialFloat4(%0.8f,%0.8f,%0.8f,%0.8f)"),X,Y,Z,W);
 	}
 
-	virtual int32 GameTime() override
+	virtual int32 GameTime(bool bPeriodic, float Period) override
 	{
-		return AddInlinedCodeChunk(MCT_Float, TEXT("View.GameTime"));
+		if (!bPeriodic)
+		{
+			return AddInlinedCodeChunk(MCT_Float, TEXT("View.GameTime"));
+		}
+		else if (Period == 0.0f)
+		{
+			return Constant(0.0f);
+		}
+
+		return AddUniformExpression(
+			new FMaterialUniformExpressionFmod(
+				new FMaterialUniformExpressionTime(),
+				new FMaterialUniformExpressionConstant(FLinearColor(Period, Period, Period, Period), MCT_Float)
+				),
+			MCT_Float, TEXT("")
+			);
 	}
 
-	virtual int32 RealTime() override
+	virtual int32 RealTime(bool bPeriodic, float Period) override
 	{
-		return AddInlinedCodeChunk(MCT_Float, TEXT("View.RealTime"));
+		if (!bPeriodic)
+		{
+			return AddInlinedCodeChunk(MCT_Float, TEXT("View.RealTime"));
+		}
+		else if (Period == 0.0f)
+		{
+			return Constant(0.0f);
+		}
+
+		return AddUniformExpression(
+			new FMaterialUniformExpressionFmod(
+				new FMaterialUniformExpressionRealTime(),
+				new FMaterialUniformExpressionConstant(FLinearColor(Period, Period, Period, Period), MCT_Float)
+				),
+			MCT_Float, TEXT("")
+			);
 	}
 
 	virtual int32 PeriodicHint(int32 PeriodicCode) override
@@ -2883,16 +2956,37 @@ protected:
 			return Errorf(TEXT("Couldn't determine result type of component mask %u%u%u%u"),R,G,B,A);
 		};
 
-		return AddInlinedCodeChunk(
-			ResultType,
-			TEXT("%s.%s%s%s%s"),
-			*GetParameterCode(Vector),
+		FString MaskString = FString::Printf(TEXT("%s%s%s%s"),
 			R ? TEXT("r") : TEXT(""),
 			// If VectorType is set to MCT_Float which means it could be any of the float types, assume it is a float1
 			G ? (VectorType == MCT_Float ? TEXT("r") : TEXT("g")) : TEXT(""),
 			B ? (VectorType == MCT_Float ? TEXT("r") : TEXT("b")) : TEXT(""),
 			A ? (VectorType == MCT_Float ? TEXT("r") : TEXT("a")) : TEXT("")
 			);
+
+		auto* Expression = GetParameterUniformExpression(Vector);
+		if (Expression)
+		{
+			int8 Mask[4] = {-1, -1, -1, -1};
+			for (int32 Index = 0; Index < MaskString.Len(); ++Index)
+			{
+				Mask[Index] = SwizzleComponentToIndex(MaskString[Index]);
+			}
+			return AddUniformExpression(
+				new FMaterialUniformExpressionComponentSwizzle(Expression, Mask[0], Mask[1], Mask[2], Mask[3]),
+				ResultType,
+				TEXT("%s.%s"),
+				*GetParameterCode(Vector),
+				*MaskString
+				);
+		}
+
+		return AddInlinedCodeChunk(
+			ResultType,
+			TEXT("%s.%s"),
+			*GetParameterCode(Vector),
+			*MaskString
+			);
 	}
 
 	virtual int32 AppendVector(int32 A,int32 B) override
diff --git a/Engine/Source/Runtime/Engine/Private/Materials/MaterialExpressions.cpp b/Engine/Source/Runtime/Engine/Private/Materials/MaterialExpressions.cpp
--- a/Engine/Source/Runtime/Engine/Private/Materials/MaterialExpressions.cpp
+++ b/Engine/Source/Runtime/Engine/Private/Materials/MaterialExpressions.cpp
@@ -2716,16 +2716,32 @@ UMaterialExpressionTime::UMaterialExpressionTime(const FObjectInitializer& Objec
 
 	MenuCategories.Add(ConstructorStatics.NAME_Constants);
 	bShaderInputData = true;
+	Period = 0.0f;
+	bOverride_Period = false;
 }
 
 int32 UMaterialExpressionTime::Compile(class FMaterialCompiler* Compiler, int32 OutputIndex, int32 MultiplexIndex)
 {
-	return bIgnorePause ? Compiler->RealTime() : Compiler->GameTime();
+	return bIgnorePause ? Compiler->RealTime(bOverride_Period, Period) : Compiler->GameTime(bOverride_Period, Period);
 }
 
 void UMaterialExpressionTime::GetCaption(TArray<FString>& OutCaptions) const
 {
-	OutCaptions.Add(TEXT("Time"));
+	if (bOverride_Period)
+	{
+		if (Period == 0.0f)
+		{
+			OutCaptions.Add(TEXT("Time (Stopped)"));
+		}
+		else
+		{
+			OutCaptions.Add(FString::Printf(TEXT("Time (Period of %.2f)"), Period));
+		}
+	}
+	else
+	{
+		OutCaptions.Add(TEXT("Time"));
+	}
 }
 
 
@@ -2857,7 +2873,7 @@ UMaterialExpressionPanner::UMaterialExpressionPanner(const FObjectInitializer& O
 
 int32 UMaterialExpressionPanner::Compile(class FMaterialCompiler* Compiler, int32 OutputIndex, int32 MultiplexIndex)
 {
-	int32 TimeArg = Time.Expression ? Time.Compile(Compiler) : Compiler->GameTime();
+	int32 TimeArg = Time.Expression ? Time.Compile(Compiler) : Compiler->GameTime(false, 0.0f);
 	int32 Arg1;
 	int32 Arg2;
 	if (bFractionalPart)
@@ -2917,8 +2933,8 @@ UMaterialExpressionRotator::UMaterialExpressionRotator(const FObjectInitializer&
 
 int32 UMaterialExpressionRotator::Compile(class FMaterialCompiler* Compiler, int32 OutputIndex, int32 MultiplexIndex)
 {
-	int32	Cosine = Compiler->Cosine(Compiler->Mul(Time.Expression ? Time.Compile(Compiler) : Compiler->GameTime(),Compiler->Constant(Speed))),
-		Sine = Compiler->Sine(Compiler->Mul(Time.Expression ? Time.Compile(Compiler) : Compiler->GameTime(),Compiler->Constant(Speed))),
+	int32	Cosine = Compiler->Cosine(Compiler->Mul(Time.Expression ? Time.Compile(Compiler) : Compiler->GameTime(false, 0.0f),Compiler->Constant(Speed))),
+		Sine = Compiler->Sine(Compiler->Mul(Time.Expression ? Time.Compile(Compiler) : Compiler->GameTime(false, 0.0f),Compiler->Constant(Speed))),
 		RowX = Compiler->AppendVector(Cosine,Compiler->Mul(Compiler->Constant(-1.0f),Sine)),
 		RowY = Compiler->AppendVector(Sine,Cosine),
 		Origin = Compiler->Constant2(CenterX,CenterY),
diff --git a/Engine/Source/Runtime/Engine/Private/Materials/MaterialShared.cpp b/Engine/Source/Runtime/Engine/Private/Materials/MaterialShared.cpp
--- a/Engine/Source/Runtime/Engine/Private/Materials/MaterialShared.cpp
+++ b/Engine/Source/Runtime/Engine/Private/Materials/MaterialShared.cpp
@@ -1498,6 +1498,8 @@ FMaterialRenderContext::FMaterialRenderContext(
 	const FSceneView* InView)
 		: MaterialRenderProxy(InMaterialRenderProxy)
 		, Material(InMaterial)
+		, Time(0.0f)
+		, RealTime(0.0f)
 {
 	bShowSelection = GIsEditor && InView && InView->Family->EngineShowFlags.Selection;
 }
diff --git a/Engine/Source/Runtime/Engine/Private/Materials/MaterialUniformExpressions.cpp b/Engine/Source/Runtime/Engine/Private/Materials/MaterialUniformExpressions.cpp
--- a/Engine/Source/Runtime/Engine/Private/Materials/MaterialUniformExpressions.cpp
+++ b/Engine/Source/Runtime/Engine/Private/Materials/MaterialUniformExpressions.cpp
@@ -88,6 +88,11 @@ void FUniformExpressionSet::Serialize(FArchive& Ar)
 	Ar << UniformScalarExpressions;
 	Ar << Uniform2DTextureExpressions;
 	Ar << UniformCubeTextureExpressions;
+	if (Ar.UE4Ver() >= VER_UE4_PERFRAME_MATERIAL_UNIFORM_EXPRESSIONS)
+	{
+		Ar << PerFrameUniformScalarExpressions;
+		Ar << PerFrameUniformVectorExpressions;
+	}
 
 	Ar << ParameterCollections;
 
@@ -104,6 +109,8 @@ bool FUniformExpressionSet::IsEmpty() const
 		&& UniformScalarExpressions.Num() == 0
 		&& Uniform2DTextureExpressions.Num() == 0
 		&& UniformCubeTextureExpressions.Num() == 0
+		&& PerFrameUniformScalarExpressions.Num() == 0
+		&& PerFrameUniformVectorExpressions.Num() == 0
 		&& ParameterCollections.Num() == 0;
 }
 
@@ -113,6 +120,8 @@ bool FUniformExpressionSet::operator==(const FUniformExpressionSet& ReferenceSet
 	|| UniformScalarExpressions.Num() != ReferenceSet.UniformScalarExpressions.Num()
 	|| Uniform2DTextureExpressions.Num() != ReferenceSet.Uniform2DTextureExpressions.Num()
 	|| UniformCubeTextureExpressions.Num() != ReferenceSet.UniformCubeTextureExpressions.Num()
+	|| PerFrameUniformScalarExpressions.Num() != ReferenceSet.PerFrameUniformScalarExpressions.Num()
+	|| PerFrameUniformVectorExpressions.Num() != ReferenceSet.PerFrameUniformVectorExpressions.Num()
 	|| ParameterCollections.Num() != ReferenceSet.ParameterCollections.Num())
 	{
 		return false;
@@ -150,6 +159,22 @@ bool FUniformExpressionSet::operator==(const FUniformExpressionSet& ReferenceSet
 		}
 	}
 
+	for (int32 i = 0; i < PerFrameUniformScalarExpressions.Num(); i++)
+	{
+		if (!PerFrameUniformScalarExpressions[i]->IsIdentical(ReferenceSet.PerFrameUniformScalarExpressions[i]))
+		{
+			return false;
+		}
+	}
+
+	for (int32 i = 0; i < PerFrameUniformVectorExpressions.Num(); i++)
+	{
+		if (!PerFrameUniformVectorExpressions[i]->IsIdentical(ReferenceSet.PerFrameUniformVectorExpressions[i]))
+		{
+			return false;
+		}
+	}
+
 	for (int32 i = 0; i < ParameterCollections.Num(); i++)
 	{
 		if (ParameterCollections[i] != ReferenceSet.ParameterCollections[i])
@@ -163,11 +188,13 @@ bool FUniformExpressionSet::operator==(const FUniformExpressionSet& ReferenceSet
 
 FString FUniformExpressionSet::GetSummaryString() const
 {
-	return FString::Printf(TEXT("(%u vectors, %u scalars, %u 2d tex, %u cube tex, %u collections)"),
+	return FString::Printf(TEXT("(%u vectors, %u scalars, %u 2d tex, %u cube tex, %u scalars/frame, %u vectors/frame, %u collections)"),
 		UniformVectorExpressions.Num(), 
 		UniformScalarExpressions.Num(),
 		Uniform2DTextureExpressions.Num(),
 		UniformCubeTextureExpressions.Num(),
+		PerFrameUniformScalarExpressions.Num(),
+		PerFrameUniformVectorExpressions.Num(),
 		ParameterCollections.Num()
 		);
 }
@@ -465,6 +492,7 @@ IMPLEMENT_MATERIALUNIFORMEXPRESSION_TYPE(FMaterialUniformExpressionMin);
 IMPLEMENT_MATERIALUNIFORMEXPRESSION_TYPE(FMaterialUniformExpressionMax);
 IMPLEMENT_MATERIALUNIFORMEXPRESSION_TYPE(FMaterialUniformExpressionClamp);
 IMPLEMENT_MATERIALUNIFORMEXPRESSION_TYPE(FMaterialUniformExpressionSaturate);
+IMPLEMENT_MATERIALUNIFORMEXPRESSION_TYPE(FMaterialUniformExpressionComponentSwizzle);
 IMPLEMENT_MATERIALUNIFORMEXPRESSION_TYPE(FMaterialUniformExpressionFloor);
 IMPLEMENT_MATERIALUNIFORMEXPRESSION_TYPE(FMaterialUniformExpressionCeil);
 IMPLEMENT_MATERIALUNIFORMEXPRESSION_TYPE(FMaterialUniformExpressionFrac);
diff --git a/Engine/Source/Runtime/Engine/Private/Materials/MaterialUniformExpressions.h b/Engine/Source/Runtime/Engine/Private/Materials/MaterialUniformExpressions.h
--- a/Engine/Source/Runtime/Engine/Private/Materials/MaterialUniformExpressions.h
+++ b/Engine/Source/Runtime/Engine/Private/Materials/MaterialUniformExpressions.h
@@ -61,12 +61,16 @@ public:
 	}
 	virtual void GetNumberValue(const FMaterialRenderContext& Context,FLinearColor& OutValue) const
 	{
-		//OutValue.R = Context.CurrentTime;
+		OutValue.R = Context.Time;
+		OutValue.G = Context.Time;
+		OutValue.B = Context.Time;
+		OutValue.A = Context.Time;
 	}
 	virtual bool IsConstant() const
 	{
 		return false;
 	}
+	virtual bool IsChangingPerFrame() const { return true; }
 	virtual bool IsIdentical(const FMaterialUniformExpression* OtherExpression) const
 	{
 		return GetType() == OtherExpression->GetType();
@@ -86,12 +90,16 @@ public:
 	}
 	virtual void GetNumberValue(const FMaterialRenderContext& Context,FLinearColor& OutValue) const
 	{
-		//OutValue.R = Context.CurrentRealTime;
+		OutValue.R = Context.RealTime;
+		OutValue.G = Context.RealTime;
+		OutValue.B = Context.RealTime;
+		OutValue.A = Context.RealTime;
 	}
 	virtual bool IsConstant() const
 	{
 		return false;
 	}
+	virtual bool IsChangingPerFrame() const { return true; }
 	virtual bool IsIdentical(const FMaterialUniformExpression* OtherExpression) const
 	{
 		return GetType() == OtherExpression->GetType();
@@ -381,6 +389,10 @@ public:
 	{
 		return X->IsConstant();
 	}
+	virtual bool IsChangingPerFrame() const
+	{
+		return X->IsChangingPerFrame();
+	}
 	virtual bool IsIdentical(const FMaterialUniformExpression* OtherExpression) const
 	{
 		if (GetType() != OtherExpression->GetType())
@@ -426,6 +438,10 @@ public:
 	{
 		return X->IsConstant();
 	}
+	virtual bool IsChangingPerFrame() const
+	{
+		return X->IsChangingPerFrame();
+	}
 	virtual bool IsIdentical(const FMaterialUniformExpression* OtherExpression) const
 	{
 		if (GetType() != OtherExpression->GetType())
@@ -467,6 +483,10 @@ public:
 	{
 		return X->IsConstant();
 	}
+	virtual bool IsChangingPerFrame() const
+	{
+		return X->IsChangingPerFrame();
+	}
 	virtual bool IsIdentical(const FMaterialUniformExpression* OtherExpression) const
 	{
 		if (GetType() != OtherExpression->GetType())
@@ -560,6 +580,12 @@ public:
 	{
 		return A->IsConstant() && B->IsConstant();
 	}
+	virtual bool IsChangingPerFrame() const
+	{
+		const bool AChangesPerFrame = A->IsChangingPerFrame();
+		const bool BChangesPerFrame = B->IsChangingPerFrame();
+		return AChangesPerFrame || BChangesPerFrame;
+	}
 	virtual bool IsIdentical(const FMaterialUniformExpression* OtherExpression) const
 	{
 		if (GetType() != OtherExpression->GetType())
@@ -657,6 +683,12 @@ public:
 	{
 		return A->IsConstant() && B->IsConstant();
 	}
+	virtual bool IsChangingPerFrame() const
+	{
+		const bool AChangesPerFrame = A->IsChangingPerFrame();
+		const bool BChangesPerFrame = B->IsChangingPerFrame();
+		return AChangesPerFrame || BChangesPerFrame;
+	}
 	virtual bool IsIdentical(const FMaterialUniformExpression* OtherExpression) const
 	{
 		if (GetType() != OtherExpression->GetType())
@@ -707,6 +739,12 @@ public:
 	{
 		return A->IsConstant() && B->IsConstant();
 	}
+	virtual bool IsChangingPerFrame() const
+	{
+		const bool AChangesPerFrame = A->IsChangingPerFrame();
+		const bool BChangesPerFrame = B->IsChangingPerFrame();
+		return AChangesPerFrame || BChangesPerFrame;
+	}
 	virtual bool IsIdentical(const FMaterialUniformExpression* OtherExpression) const
 	{
 		if (GetType() != OtherExpression->GetType())
@@ -756,6 +794,12 @@ public:
 	{
 		return A->IsConstant() && B->IsConstant();
 	}
+	virtual bool IsChangingPerFrame() const
+	{
+		const bool AChangesPerFrame = A->IsChangingPerFrame();
+		const bool BChangesPerFrame = B->IsChangingPerFrame();
+		return AChangesPerFrame || BChangesPerFrame;
+	}
 	virtual bool IsIdentical(const FMaterialUniformExpression* OtherExpression) const
 	{
 		if (GetType() != OtherExpression->GetType())
@@ -808,6 +852,10 @@ public:
 	{
 		return Input->IsConstant() && Min->IsConstant() && Max->IsConstant();
 	}
+	virtual bool IsChangingPerFrame() const
+	{
+		return Input->IsChangingPerFrame() || Min->IsChangingPerFrame() || Max->IsChangingPerFrame();
+	}
 	virtual bool IsIdentical(const FMaterialUniformExpression* OtherExpression) const
 	{
 		if (GetType() != OtherExpression->GetType())
@@ -855,6 +903,10 @@ public:
 	{
 		return Input->IsConstant();
 	}
+	virtual bool IsChangingPerFrame() const
+	{
+		return Input->IsChangingPerFrame();
+	}
 	virtual bool IsIdentical(const FMaterialUniformExpression* OtherExpression) const
 	{
 		if (GetType() != OtherExpression->GetType())
@@ -869,6 +921,109 @@ private:
 	TRefCountPtr<FMaterialUniformExpression> Input;
 };
 
+class FMaterialUniformExpressionComponentSwizzle : public FMaterialUniformExpression
+{
+	DECLARE_MATERIALUNIFORMEXPRESSION_TYPE(FMaterialUniformExpressionComponentSwizzle);
+public:
+
+	FMaterialUniformExpressionComponentSwizzle() {}
+	FMaterialUniformExpressionComponentSwizzle(FMaterialUniformExpression* InX, int8 InR, int8 InG, int8 InB, int8 InA) :
+		X(InX),
+		IndexR(InR),
+		IndexG(InG),
+		IndexB(InB),
+		IndexA(InA)
+	{
+		NumElements = 0;
+		if (InA >= 0)
+		{
+			check(InA <= 3);
+			++NumElements;
+			check(InB >= 0);
+		}
+
+		if (InB >= 0)
+		{
+			check(InB <= 3);
+			++NumElements;
+			check(InG >= 0);
+		}
+
+		if (InG >= 0)
+		{
+			check(InG <= 3);
+			++NumElements;
+		}
+
+		// At least one proper index
+		check(InR >= 0 && InR <= 3);
+		++NumElements;
+	}
+
+	// FMaterialUniformExpression interface.
+	virtual void Serialize(FArchive& Ar)
+	{
+		Ar << X;
+		Ar << IndexR;
+		Ar << IndexG;
+		Ar << IndexB;
+		Ar << IndexA;
+		Ar << NumElements;
+	}
+	virtual void GetNumberValue(const FMaterialRenderContext& Context, FLinearColor& OutValue) const
+	{
+		FLinearColor Temp = OutValue;
+		X->GetNumberValue(Context, Temp);
+		OutValue = Temp;
+		switch (NumElements)
+		{
+		case 4:
+			OutValue.A = Temp.Component(IndexA);
+			// Fallthrough...
+		case 3:
+			OutValue.B = Temp.Component(IndexB);
+			// Fallthrough...
+		case 2:
+			OutValue.G = Temp.Component(IndexG);
+			// Fallthrough...
+		case 1:
+			OutValue.R = Temp.Component(IndexR);
+			break;
+		default: UE_LOG(LogMaterial, Fatal, TEXT("Invalid number of swizzle elements: %d"), NumElements);
+		}
+	}
+	virtual bool IsConstant() const
+	{
+		return X->IsConstant();
+	}
+	virtual bool IsChangingPerFrame() const
+	{
+		return X->IsChangingPerFrame();
+	}
+	virtual bool IsIdentical(const FMaterialUniformExpression* OtherExpression) const
+	{
+		if (GetType() != OtherExpression->GetType())
+		{
+			return false;
+		}
+		auto* OtherSwizzle = (FMaterialUniformExpressionComponentSwizzle*)OtherExpression;
+		return X->IsIdentical(OtherSwizzle->X) &&
+			NumElements == OtherSwizzle->NumElements &&
+			IndexR == OtherSwizzle->IndexR &&
+			IndexG == OtherSwizzle->IndexG &&
+			IndexB == OtherSwizzle->IndexB &&
+			IndexA == OtherSwizzle->IndexA;
+	}
+
+private:
+	TRefCountPtr<FMaterialUniformExpression> X;
+	int8 IndexR;
+	int8 IndexG;
+	int8 IndexB;
+	int8 IndexA;
+	int8 NumElements;
+};
+
 /**
  */
 class FMaterialUniformExpressionFloor: public FMaterialUniformExpression
@@ -899,6 +1054,10 @@ public:
 	{
 		return X->IsConstant();
 	}
+	virtual bool IsChangingPerFrame() const
+	{
+		return X->IsChangingPerFrame();
+	}
 	virtual bool IsIdentical(const FMaterialUniformExpression* OtherExpression) const
 	{
 		if (GetType() != OtherExpression->GetType())
@@ -943,6 +1102,10 @@ public:
 	{
 		return X->IsConstant();
 	}
+	virtual bool IsChangingPerFrame() const
+	{
+		return X->IsChangingPerFrame();
+	}
 	virtual bool IsIdentical(const FMaterialUniformExpression* OtherExpression) const
 	{
 		if (GetType() != OtherExpression->GetType())
@@ -987,6 +1150,10 @@ public:
 	{
 		return X->IsConstant();
 	}
+	virtual bool IsChangingPerFrame() const
+	{
+		return X->IsChangingPerFrame();
+	}
 	virtual bool IsIdentical(const FMaterialUniformExpression* OtherExpression) const
 	{
 		if (GetType() != OtherExpression->GetType())
@@ -1035,6 +1202,12 @@ public:
 	{
 		return A->IsConstant() && B->IsConstant();
 	}
+	virtual bool IsChangingPerFrame() const
+	{
+		const bool AChangesPerFrame = A->IsChangingPerFrame();
+		const bool BChangesPerFrame = B->IsChangingPerFrame();
+		return AChangesPerFrame || BChangesPerFrame;
+	}
 	virtual bool IsIdentical(const FMaterialUniformExpression* OtherExpression) const
 	{
 		if (GetType() != OtherExpression->GetType())
@@ -1080,6 +1253,10 @@ public:
 	{
 		return X->IsConstant();
 	}
+	virtual bool IsChangingPerFrame() const
+	{
+		return X->IsChangingPerFrame();
+	}
 	virtual bool IsIdentical(const FMaterialUniformExpression* OtherExpression) const
 	{
 		if (GetType() != OtherExpression->GetType())
diff --git a/Engine/Source/Runtime/Engine/Public/MaterialCompiler.h b/Engine/Source/Runtime/Engine/Public/MaterialCompiler.h
--- a/Engine/Source/Runtime/Engine/Public/MaterialCompiler.h
+++ b/Engine/Source/Runtime/Engine/Public/MaterialCompiler.h
@@ -57,8 +57,8 @@ public:
 	virtual int32 Constant3(float X,float Y,float Z) = 0;
 	virtual int32 Constant4(float X,float Y,float Z,float W) = 0;
 
-	virtual int32 GameTime() = 0;
-	virtual int32 RealTime() = 0;
+	virtual int32 GameTime(bool bPeriodic, float Period) = 0;
+	virtual int32 RealTime(bool bPeriodic, float Period) = 0;
 	virtual int32 PeriodicHint(int32 PeriodicCode) { return PeriodicCode; }
 
 	virtual int32 Sine(int32 X) = 0;
@@ -216,8 +216,8 @@ public:
 	virtual int32 Constant3(float X,float Y,float Z) override { return Compiler->Constant3(X,Y,Z); }
 	virtual int32 Constant4(float X,float Y,float Z,float W) override { return Compiler->Constant4(X,Y,Z,W); }
 
-	virtual int32 GameTime() override { return Compiler->GameTime(); }
-	virtual int32 RealTime() override { return Compiler->RealTime(); }
+	virtual int32 GameTime(bool bPeriodic, float Period) override { return Compiler->GameTime(bPeriodic, Period); }
+	virtual int32 RealTime(bool bPeriodic, float Period) override { return Compiler->RealTime(bPeriodic, Period); }
 
 	virtual int32 PeriodicHint(int32 PeriodicCode) override { return Compiler->PeriodicHint(PeriodicCode); }
 
diff --git a/Engine/Source/Runtime/Engine/Public/MaterialShared.h b/Engine/Source/Runtime/Engine/Public/MaterialShared.h
--- a/Engine/Source/Runtime/Engine/Public/MaterialShared.h
+++ b/Engine/Source/Runtime/Engine/Public/MaterialShared.h
@@ -92,6 +92,11 @@ struct ENGINE_API FMaterialRenderContext
 	const FMaterialRenderProxy* MaterialRenderProxy;
 	/** Material resource to use. */
 	const FMaterial& Material;
+
+	// Used only when evaluating expressions per frame
+	float Time;
+	float RealTime;
+
 	/** Whether or not selected objects should use their selection color. */
 	bool bShowSelection;
 
@@ -167,6 +172,7 @@ public:
 	virtual void GetNumberValue(const struct FMaterialRenderContext& Context,FLinearColor& OutValue) const {}
 	virtual class FMaterialUniformExpressionTexture* GetTextureUniformExpression() { return NULL; }
 	virtual bool IsConstant() const { return false; }
+	virtual bool IsChangingPerFrame() const { return false; }
 	virtual bool IsIdentical(const FMaterialUniformExpression* OtherExpression) const { return false; }
 
 	friend FArchive& operator<<(FArchive& Ar,class FMaterialUniformExpression*& Ref);
@@ -232,6 +238,8 @@ public:
 			+ UniformScalarExpressions.GetAllocatedSize()
 			+ Uniform2DTextureExpressions.GetAllocatedSize()
 			+ UniformCubeTextureExpressions.GetAllocatedSize()
+			+ PerFrameUniformScalarExpressions.GetAllocatedSize()
+			+ PerFrameUniformVectorExpressions.GetAllocatedSize()
 			+ ParameterCollections.GetAllocatedSize()
 			+ (UniformBufferStruct ? (sizeof(FUniformBufferStruct) + UniformBufferStruct->GetMembers().GetAllocatedSize()) : 0);
 	}
@@ -242,6 +250,8 @@ protected:
 	TArray<TRefCountPtr<FMaterialUniformExpression> > UniformScalarExpressions;
 	TArray<TRefCountPtr<FMaterialUniformExpressionTexture> > Uniform2DTextureExpressions;
 	TArray<TRefCountPtr<FMaterialUniformExpressionTexture> > UniformCubeTextureExpressions;
+	TArray<TRefCountPtr<FMaterialUniformExpression> > PerFrameUniformScalarExpressions;
+	TArray<TRefCountPtr<FMaterialUniformExpression> > PerFrameUniformVectorExpressions;
 
 	/** Ids of parameter collections referenced by the material that was translated. */
 	TArray<FGuid> ParameterCollections;
diff --git a/Engine/Source/Runtime/Renderer/Private/ShaderBaseClasses.cpp b/Engine/Source/Runtime/Renderer/Private/ShaderBaseClasses.cpp
--- a/Engine/Source/Runtime/Renderer/Private/ShaderBaseClasses.cpp
+++ b/Engine/Source/Runtime/Renderer/Private/ShaderBaseClasses.cpp
@@ -34,6 +34,20 @@ FMaterialShader::FMaterialShader(const FMaterialShaderType::CompiledShaderInitia
 		ParameterCollectionUniformBuffers.Add(CollectionParameter);
 	}
 
+	for (int32 Index = 0; Index < Initializer.UniformExpressionSet.PerFrameUniformScalarExpressions.Num(); Index++)
+	{
+		FShaderParameter Parameter;
+		Parameter.Bind(Initializer.ParameterMap, *FString::Printf(TEXT("UE_Material_PerFrameScalarExpression%u"), Index));
+		PerFrameScalarExpressions.Add(Parameter);
+	}
+
+	for (int32 Index = 0; Index < Initializer.UniformExpressionSet.PerFrameUniformVectorExpressions.Num(); Index++)
+	{
+		FShaderParameter Parameter;
+		Parameter.Bind(Initializer.ParameterMap, *FString::Printf(TEXT("UE_Material_PerFrameVectorExpression%u"), Index));
+		PerFrameVectorExpressions.Add(Parameter);
+	}
+
 	DeferredParameters.Bind(Initializer.ParameterMap);
 	LightAttenuation.Bind(Initializer.ParameterMap, TEXT("LightAttenuationTexture"));
 	LightAttenuationSampler.Bind(Initializer.ParameterMap, TEXT("LightAttenuationTextureSampler"));
@@ -103,8 +117,8 @@ void FMaterialShader::SetParameters(
 			TEXT("%s shader uniform expression set mismatch for material %s/%s.\n")
 			TEXT("Shader compilation info:                %s\n")
 			TEXT("Material render proxy compilation info: %s\n")
-			TEXT("Shader uniform expression set:   %u vectors, %u scalars, %u 2D textures, %u cube textures, shader map %p\n")
-			TEXT("Material uniform expression set: %u vectors, %u scalars, %u 2D textures, %u cube textures, shader map %p\n"),
+			TEXT("Shader uniform expression set:   %u vectors, %u scalars, %u 2D textures, %u cube textures, %u scalars/frame, %u vectors/frame, shader map %p\n")
+			TEXT("Material uniform expression set: %u vectors, %u scalars, %u 2D textures, %u cube textures, %u scalars/frame, %u vectors/frame, shader map %p\n"),
 			GetType()->GetName(),
 			*MaterialRenderProxy->GetFriendlyName(),
 			*Material.GetFriendlyName(),
@@ -114,11 +128,15 @@ void FMaterialShader::SetParameters(
 			DebugUniformExpressionSet.NumScalarExpressions,
 			DebugUniformExpressionSet.Num2DTextureExpressions,
 			DebugUniformExpressionSet.NumCubeTextureExpressions,
+			DebugUniformExpressionSet.NumPerFrameScalarExpressions,
+			DebugUniformExpressionSet.NumPerFrameVectorExpressions,
 			UniformExpressionCache->CachedUniformExpressionShaderMap,
 			MaterialUniformExpressionSet.UniformVectorExpressions.Num(),
 			MaterialUniformExpressionSet.UniformScalarExpressions.Num(),
 			MaterialUniformExpressionSet.Uniform2DTextureExpressions.Num(),
 			MaterialUniformExpressionSet.UniformCubeTextureExpressions.Num(),
+			MaterialUniformExpressionSet.PerFrameUniformScalarExpressions.Num(),
+			MaterialUniformExpressionSet.PerFrameUniformVectorExpressions.Num(),
 			Material.GetRenderingThreadShaderMap()
 			);
 	}
@@ -136,16 +154,36 @@ void FMaterialShader::SetParameters(
 	}
 
 	{
-		const TArray<FGuid>& ParameterCollections = UniformExpressionCache->ParameterCollections;
-		const int32 ParameterCollectionsNum = ParameterCollections.Num();
-
-		check(ParameterCollectionUniformBuffers.Num() >= ParameterCollectionsNum);
+		// Per frame material expressions
+		const int32 NumScalarExpressions = PerFrameScalarExpressions.Num();
+		const int32 NumVectorExpressions = PerFrameVectorExpressions.Num();
 
-		// Find each referenced parameter collection's uniform buffer in the scene and set the parameter
-		for (int32 CollectionIndex = 0; CollectionIndex < ParameterCollectionsNum; CollectionIndex++)
+		if (NumScalarExpressions > 0 || NumVectorExpressions > 0)
 		{
-			FUniformBufferRHIParamRef UniformBuffer = GetParameterCollectionBuffer(ParameterCollections[CollectionIndex], View.Family->Scene);
-			SetUniformBufferParameter(RHICmdList, ShaderRHI,ParameterCollectionUniformBuffers[CollectionIndex],UniformBuffer);
+			FMaterialRenderContext MaterialRenderContext(MaterialRenderProxy, Material, &View);
+			MaterialRenderContext.Time = View.Family->CurrentWorldTime;
+			MaterialRenderContext.RealTime = View.Family->CurrentRealTime;
+			for (int32 Index = 0; Index < NumScalarExpressions; ++Index)
+			{
+				auto& Parameter = PerFrameScalarExpressions[Index];
+				if (Parameter.IsBound())
+				{
+					FLinearColor TempValue;
+					MaterialUniformExpressionSet.PerFrameUniformScalarExpressions[Index]->GetNumberValue(MaterialRenderContext, TempValue);
+					SetShaderValue(RHICmdList, ShaderRHI, Parameter, TempValue.R);
+				}
+			}
+
+			for (int32 Index = 0; Index < NumVectorExpressions; ++Index)
+			{
+				auto& Parameter = PerFrameVectorExpressions[Index];
+				if (Parameter.IsBound())
+				{
+					FLinearColor TempValue;
+					MaterialUniformExpressionSet.PerFrameUniformVectorExpressions[Index]->GetNumberValue(MaterialRenderContext, TempValue);
+					SetShaderValue(RHICmdList, ShaderRHI, Parameter, TempValue);
+				}
+			}
 		}
 	}
 
@@ -232,6 +270,11 @@ bool FMaterialShader::Serialize(FArchive& Ar)
 	const bool bShaderHasOutdatedParameters = FShader::Serialize(Ar);
 	Ar << MaterialUniformBuffer;
 	Ar << ParameterCollectionUniformBuffers;
+	if (Ar.UE4Ver() >= VER_UE4_PERFRAME_MATERIAL_UNIFORM_EXPRESSIONS)
+	{
+		Ar << PerFrameScalarExpressions;
+		Ar << PerFrameVectorExpressions;
+	}
 	Ar << DeferredParameters;
 	Ar << LightAttenuation;
 	Ar << LightAttenuationSampler;
diff --git a/Engine/Source/Runtime/Renderer/Public/MaterialShader.h b/Engine/Source/Runtime/Renderer/Public/MaterialShader.h
--- a/Engine/Source/Runtime/Renderer/Public/MaterialShader.h
+++ b/Engine/Source/Runtime/Renderer/Public/MaterialShader.h
@@ -34,12 +34,16 @@ public:
 	int32 NumScalarExpressions;
 	int32 Num2DTextureExpressions;
 	int32 NumCubeTextureExpressions;
+	int32 NumPerFrameScalarExpressions;
+	int32 NumPerFrameVectorExpressions;
 
 	FDebugUniformExpressionSet()
 		: NumVectorExpressions(0)
 		, NumScalarExpressions(0)
 		, Num2DTextureExpressions(0)
 		, NumCubeTextureExpressions(0)
+		, NumPerFrameScalarExpressions(0)
+		, NumPerFrameVectorExpressions(0)
 	{
 	}
 
@@ -55,6 +59,8 @@ public:
 		NumScalarExpressions = InUniformExpressionSet.UniformScalarExpressions.Num();
 		Num2DTextureExpressions = InUniformExpressionSet.Uniform2DTextureExpressions.Num();
 		NumCubeTextureExpressions = InUniformExpressionSet.UniformCubeTextureExpressions.Num();
+		NumPerFrameScalarExpressions = InUniformExpressionSet.PerFrameUniformScalarExpressions.Num();
+		NumPerFrameVectorExpressions = InUniformExpressionSet.PerFrameUniformVectorExpressions.Num();
 	}
 
 	/** Returns true if the number of uniform expressions matches those with which the debug set was initialized. */
@@ -63,7 +69,9 @@ public:
 		return NumVectorExpressions == InUniformExpressionSet.UniformVectorExpressions.Num()
 			&& NumScalarExpressions == InUniformExpressionSet.UniformScalarExpressions.Num()
 			&& Num2DTextureExpressions == InUniformExpressionSet.Uniform2DTextureExpressions.Num()
-			&& NumCubeTextureExpressions == InUniformExpressionSet.UniformCubeTextureExpressions.Num();
+			&& NumCubeTextureExpressions == InUniformExpressionSet.UniformCubeTextureExpressions.Num()
+			&& NumPerFrameScalarExpressions == InUniformExpressionSet.PerFrameUniformScalarExpressions.Num()
+			&& NumPerFrameVectorExpressions == InUniformExpressionSet.PerFrameUniformVectorExpressions.Num();
 	}
 };
 
@@ -74,6 +82,11 @@ inline FArchive& operator<<(FArchive& Ar, FDebugUniformExpressionSet& DebugExpre
 	Ar << DebugExpressionSet.NumScalarExpressions;
 	Ar << DebugExpressionSet.Num2DTextureExpressions;
 	Ar << DebugExpressionSet.NumCubeTextureExpressions;
+	if (Ar.UE4Ver() >= VER_UE4_PERFRAME_MATERIAL_UNIFORM_EXPRESSIONS)
+	{
+		Ar << DebugExpressionSet.NumPerFrameScalarExpressions;
+		Ar << DebugExpressionSet.NumPerFrameVectorExpressions;
+	}
 	return Ar;
 }
 
@@ -128,6 +141,8 @@ private:
 
 	FShaderUniformBufferParameter MaterialUniformBuffer;
 	TArray<FShaderUniformBufferParameter> ParameterCollectionUniformBuffers;
+	TArray<FShaderParameter> PerFrameScalarExpressions;
+	TArray<FShaderParameter> PerFrameVectorExpressions;
 	FDeferredPixelShaderParameters DeferredParameters;
 	FShaderResourceParameter LightAttenuation;
 	FShaderResourceParameter LightAttenuationSampler;
