Fix root motion from everywhere breaking character movement when we haven't extracted any root motion

diff --git a/Engine/Source/Runtime/Engine/Private/Animation/AnimInstance.cpp b/Engine/Source/Runtime/Engine/Private/Animation/AnimInstance.cpp
--- a/Engine/Source/Runtime/Engine/Private/Animation/AnimInstance.cpp
+++ b/Engine/Source/Runtime/Engine/Private/Animation/AnimInstance.cpp
@@ -450,7 +450,10 @@ void UAnimInstance::UpdateAnimation(float DeltaSeconds)
 
 	// We may have just partially blended root motion, so make it up to 1 by
 	// blending in identity too
-	ExtractedRootMotion.MakeUpToFullWeight();
+	if (ExtractedRootMotion.bHasRootMotion)
+	{
+		ExtractedRootMotion.MakeUpToFullWeight();
+	}
 
 	// now trigger Notifies
 	TriggerAnimNotifies(DeltaSeconds);
