Added IsValidClass to KismetSystemLibrary.

diff --git a/Engine/Source/Runtime/Engine/Classes/Kismet/KismetSystemLibrary.h b/Engine/Source/Runtime/Engine/Classes/Kismet/KismetSystemLibrary.h
--- a/Engine/Source/Runtime/Engine/Classes/Kismet/KismetSystemLibrary.h
+++ b/Engine/Source/Runtime/Engine/Classes/Kismet/KismetSystemLibrary.h
@@ -75,6 +75,10 @@ class UKismetSystemLibrary : public UBlueprintFunctionLibrary
 	UFUNCTION(BlueprintPure, Category = "Utilities")
 	static ENGINE_API bool IsValid(const UObject* Object);
 
+	// Return true if the class is usable : non-null and not pending kill
+	UFUNCTION(BlueprintPure, Category = "Utilities")
+	static ENGINE_API bool IsValidClass(UClass* Class);
+
 	// Returns the display name (or actor label), for displaying to end users.  Note:  In editor builds, this is the actor label.  In non-editor builds, this is the actual object name.  This function should not be used to uniquely identify actors!
 	UFUNCTION(BlueprintPure, Category="Utilities")
 	static ENGINE_API FString GetDisplayName(const UObject* Object);
diff --git a/Engine/Source/Runtime/Engine/Private/KismetSystemLibrary.cpp b/Engine/Source/Runtime/Engine/Private/KismetSystemLibrary.cpp
--- a/Engine/Source/Runtime/Engine/Private/KismetSystemLibrary.cpp
+++ b/Engine/Source/Runtime/Engine/Private/KismetSystemLibrary.cpp
@@ -32,6 +32,11 @@ bool UKismetSystemLibrary::IsValid(const UObject* Object)
 	return ::IsValid(Object);
 }
 
+bool UKismetSystemLibrary::IsValidClass(UClass* Class)
+{
+	return ::IsValid(Class);
+}
+
 FString UKismetSystemLibrary::GetDisplayName(const UObject* Object)
 {
 #if WITH_EDITOR
