Created a unit test for the automation logging abilities.

diff --git a/Engine/Source/Runtime/Engine/Private/Tests/EngineAutomationTests.cpp b/Engine/Source/Runtime/Engine/Private/Tests/EngineAutomationTests.cpp
--- a/Engine/Source/Runtime/Engine/Private/Tests/EngineAutomationTests.cpp
+++ b/Engine/Source/Runtime/Engine/Private/Tests/EngineAutomationTests.cpp
@@ -426,6 +426,32 @@ bool FCinematicFPSPerfTest::RunTest(const FString& Parameters)
 	return true;
 }
 
+
+IMPLEMENT_SIMPLE_AUTOMATION_TEST(FLogTypesTest, "Automation Framework.Logging Test", EAutomationTestFlags::ATF_None)
+
+bool FLogTypesTest::RunTest(const FString& Parameters)
+{
+	UE_LOG(LogEngineAutomationTests, Log, TEXT("This 'Log' log is created using the UE_LOG Macro"));
+	UE_LOG(LogEngineAutomationTests, Display, TEXT("This 'Display' log is created using the UE_LOG Macro"));
+
+	ExecutionInfo.LogItems.Add(TEXT("This was added to the ExecutionInfo.LogItems"));
+	ExecutionInfo.Errors.Add(TEXT("This was added to the ExecutionInfo.Error."));
+	if (ExecutionInfo.Errors.Last().Contains(TEXT("This was added to the ExecutionInfo.Error.")))
+	{
+		UE_LOG(LogEngineAutomationTests, Display, TEXT("ExecutionInfo.Error has the correct info when used."));
+		ExecutionInfo.Errors.Empty();
+	}
+
+	ExecutionInfo.Warnings.Add(TEXT("This was added to the ExecutionInfo.Warnings"));
+	if (ExecutionInfo.Warnings.Last().Contains(TEXT("This was added to the ExecutionInfo.Warnings")))
+	{
+		UE_LOG(LogEngineAutomationTests, Display, TEXT("ExecutionInfo.Warnings has the correct info when used."));
+		ExecutionInfo.Warnings.Empty();
+	}
+
+	return true;
+}
+
 /* UAutomationTestSettings interface
  *****************************************************************************/
 
