Add a simple generic quad tree

diff --git a/Engine/Source/Runtime/Engine/Public/GenericQuadTree.h b/Engine/Source/Runtime/Engine/Public/GenericQuadTree.h
new file mode 100644
--- /dev/null
+++ b/Engine/Source/Runtime/Engine/Public/GenericQuadTree.h
@@ -0,0 +1,273 @@
+// Copyright 1998-2015 Epic Games, Inc. All Rights Reserved.
+
+#pragma once
+template <typename ElementType, int32 NodeCapacity = 4>
+class TQuadTree
+{
+	typedef TQuadTree<ElementType, NodeCapacity> TreeType;
+public:
+	TQuadTree(const FBox2D& InBox);
+
+	/** Inserts an object of type ElementType with an associated 2D box of size Box (log n)*/
+	void Insert(const ElementType& Instance, const FBox2D& Box);
+
+	/** Given a 2D box, returns an array of elements within the box. This may have duplicates based on how many quads the element overlaps*/
+	void GetElements(const FBox2D& Box, TArray<ElementType>& ElementsOut);
+
+	/** Given a 2D box, returns an array of elements within the box. The array will have no duplicates*/
+	void GetElementsUnique(const FBox2D& Box, TArray<ElementType>& ElementsOut);
+
+	/** Removes an object of type ElementType with an associated 2D box of size Box (log n). Does not cleanup tree*/
+	void Remove(const ElementType& Instance, const FBox2D& Box);
+
+	~TQuadTree();
+
+private:
+	enum QuadNames
+	{
+		TopLeft = 0,
+		TopRight = 1,
+		BottomLeft = 2,
+		BottomRight = 3
+	};
+
+	/** Given a 2D box, returns an array of leaf "trees" to be used for insertion,removal, etc... */
+	void GetLeaves(const FBox2D& Box, TArray<TreeType*>& Leaves);
+
+	/** Given a 2D box, return the subtrees that are touched*/
+	int32 GetQuads(const FBox2D& Box, TreeType* Quads[4]) const;
+
+	/** Split the tree into 4 sub-trees */
+	void Split();
+
+
+	/** Node used to hold the element and its corresponding 2D box*/
+	struct FNode
+	{
+		FBox2D Box;
+		ElementType Element;
+
+		FNode(const ElementType& InElement, const FBox2D& InBox)
+		: Box(InBox)
+		, Element(InElement)
+		{}
+	};
+
+private:
+	
+	/** Array containing the actual data. The node is used because we want to store the AABB and the element */
+	TArray<FNode> Nodes;
+	TreeType* SubTrees[4];
+
+	/** AABB of the tree*/
+	FBox2D TreeBox;
+
+	/** Center position of the tree */
+	FVector2D Position;
+
+	/** Whether we are a leaf or an internal sub-tree */
+	bool bInternal;
+};
+
+template <typename ElementType, int32 NodeCapacity>
+TQuadTree<ElementType, NodeCapacity>::TQuadTree(const FBox2D& Box)
+: TreeBox(Box)
+, Position(Box.GetCenter())
+, bInternal(false)
+{
+	SubTrees[0] = SubTrees[1] = SubTrees[2] = SubTrees[3] = nullptr;
+}
+
+template <typename ElementType, int32 NodeCapacity>
+TQuadTree<ElementType, NodeCapacity>::~TQuadTree()
+{
+	for (TreeType* SubTree : SubTrees)
+	{
+		delete SubTree;
+	}
+}
+
+template <typename ElementType, int32 NodeCapacity>
+void TQuadTree<ElementType, NodeCapacity>::Split()
+{
+	check(bInternal == false);
+	
+	const FVector2D Extent = TreeBox.GetExtent();
+	const FVector2D XExtent = FVector2D(Extent.X, 0.f);
+	const FVector2D YExtent = FVector2D(0.f, Extent.Y);
+
+	/************************************************************************/
+	/*  ___________max
+	/* |     |     |
+	/* |     |     |
+	/* |-----c------
+	/* |     |     |
+	/* min___|_____|
+	/*
+	/* We create new quads by adding xExtent and yExtent
+	/************************************************************************/
+
+	const FVector2D C = Position;
+	const FVector2D TM = C + YExtent;
+	const FVector2D ML = C - XExtent;
+	const FVector2D MR = C + XExtent;
+	const FVector2D BM = C - YExtent;
+	const FVector2D BL = TreeBox.Min;
+	const FVector2D TR = TreeBox.Max;
+
+	SubTrees[TopLeft] = new TreeType(FBox2D(ML, TM));
+	SubTrees[TopRight] = new TreeType(FBox2D(C, TR));
+	SubTrees[BottomLeft] = new TreeType(FBox2D(BL, C));
+	SubTrees[BottomRight] = new TreeType(FBox2D(BM, MR));
+
+	//take existing nodes and split them into new sub-trees
+	for (const FNode& Node : Nodes)
+	{
+		TreeType* Leaves[4];
+		const int32 NumLeaves = GetQuads(Node.Box, Leaves);
+		for (int32 LeafIdx = 0; LeafIdx < NumLeaves; ++LeafIdx)
+		{
+			Leaves[LeafIdx]->Nodes.Add(FNode(Node.Element, Node.Box));
+		}
+	}
+
+	Nodes.Empty();
+	bInternal = true;	//mark as no longer a leaf
+}
+
+template <typename ElementType, int32 NodeCapacity>
+void TQuadTree<ElementType, NodeCapacity>::Insert(const ElementType& Element, const FBox2D& Box)
+{
+	TArray<TreeType*> Leaves;
+	GetLeaves(Box, Leaves);
+	
+	for (TreeType* Leaf : Leaves)
+	{
+		check(Leaf->bInternal == false);	//make sure we only got leaves back
+		TArray<FNode>& LeafNodes = Leaf->Nodes;
+
+		if (LeafNodes.Num() < NodeCapacity || Leaf->TreeBox.GetSize().Size() <= 1.f)	//it's possible that all elements in the leaf are bigger than the leaf, and we can get into an endless spiral of splitting
+		{
+			LeafNodes.Add(FNode(Element, Box));
+		}
+		else
+		{
+			//no room so split and then try to add again
+			Leaf->Split();
+			Leaf->Insert(Element, Box);
+		}
+	}
+}
+
+template <typename ElementType, int32 NodeCapacity>
+void TQuadTree<ElementType, NodeCapacity>::Remove(const ElementType& Element, const FBox2D& Box)
+{
+	TArray<TQuadTree*> Leaves;
+	GetLeaves(Box, Leaves);
+	for (TreeType* Leaf : Leaves)
+	{
+		int32 ElementIdx = INDEX_NONE;
+		for (int32 NodeIdx = 0, NumNodes = Leaf->Nodes.Num(); NodeIdx < NumNodes; ++NodeIdx)
+		{
+			if (Leaf->Nodes[NodeIdx].Element == Element)
+			{
+				ElementIdx = NodeIdx;
+				break;
+			}
+		}
+		
+		if (ElementIdx != INDEX_NONE)
+		{
+			Leaf->Nodes.RemoveAtSwap(ElementIdx);
+		}
+	}
+}
+
+template <typename ElementType, int32 NodeCapacity>
+void TQuadTree<ElementType, NodeCapacity>::GetElements(const FBox2D& Box, TArray<ElementType>& ElementsOut)
+{
+	TArray<TQuadTree*> Leaves;
+	GetLeaves(Box, Leaves);
+	for (TreeType* Leaf : Leaves)
+	{
+		ElementsOut.Reserve(ElementsOut.Num() + Leaf->Nodes.Num());
+		for (const FNode& Node : Leaf->Nodes)
+		{
+			ElementsOut.Add(Node.Element);
+		};
+	}
+}
+
+template <typename ElementType, int32 NodeCapacity>
+void TQuadTree<ElementType, NodeCapacity>::GetElementsUnique(const FBox2D& Box, TArray<ElementType>& ElementsOut)
+{
+	TSet<ElementType> Unique;
+	TArray<TQuadTree*> Leaves;
+	GetLeaves(Box, Leaves);
+	for (TreeType* Leaf : Leaves)
+	{
+		for (const FNode& Node : Leaf->Nodes)
+		{
+			Unique.Add(Node.Element);
+		};
+	}
+
+	for (ElementType Elem : Unique)
+	{
+		ElementsOut.Add(Elem);
+	}
+}
+
+
+template <typename ElementType, int32 NodeCapacity>
+int32 TQuadTree<ElementType, NodeCapacity>::GetQuads(const FBox2D& Box, TreeType* Quads[4]) const
+{
+	bool bNegX = Box.Min.X < Position.X;
+	bool bNegY = Box.Min.Y < Position.Y;
+
+	bool bPosX = Box.Max.X > Position.X;
+	bool bPosY = Box.Max.Y > Position.Y;
+
+	int32 QuadCount = 0;
+	if (bNegX && bNegY)
+	{
+		Quads[QuadCount++] = SubTrees[BottomLeft];
+	}
+
+	if (bPosX && bNegY)
+	{
+		Quads[QuadCount++] = SubTrees[BottomRight];
+	}
+
+	if (bNegX && bPosY)
+	{
+		Quads[QuadCount++] = SubTrees[TopLeft];
+	}
+
+	if (bPosX && bPosY)
+	{
+		Quads[QuadCount++] = SubTrees[TopRight];
+	}
+
+	return QuadCount;
+}
+
+template <typename ElementType, int32 NodeCapacity>
+void TQuadTree<ElementType, NodeCapacity>::GetLeaves(const FBox2D& Box, TArray<TreeType*>& LeavesOut)
+{
+	if (bInternal)
+	{
+		TreeType* UseNodes[4];
+		const int32 NumQuads = GetQuads(Box, UseNodes);
+		{
+			for (int32 NodeIdx = 0; NodeIdx < NumQuads; ++NodeIdx)
+			{
+				UseNodes[NodeIdx]->GetLeaves(Box, LeavesOut);
+			}
+		}
+	}
+	else
+	{
+		LeavesOut.Add(this);
+	}
+}
\ No newline at end of file
