Merge CL 2383970 to Main (Added SelectClass Kismet node)

diff --git a/Engine/Source/Runtime/Engine/Classes/Kismet/KismetMathLibrary.h b/Engine/Source/Runtime/Engine/Classes/Kismet/KismetMathLibrary.h
--- a/Engine/Source/Runtime/Engine/Classes/Kismet/KismetMathLibrary.h
+++ b/Engine/Source/Runtime/Engine/Classes/Kismet/KismetMathLibrary.h
@@ -1187,6 +1187,10 @@ class UKismetMathLibrary : public UBlueprintFunctionLibrary
 	UFUNCTION(BlueprintPure, Category="Utilities")
 	static UObject* SelectObject(UObject* A, UObject* B, bool bSelectA);
 
+	/** If bPickA is true, A is returned, otherwise B is */
+	UFUNCTION(BlueprintPure, Category = "Utilities")
+	static UClass* SelectClass(UClass* A, UClass* B, bool bSelectA);
+
 	// Build a reference frame from three axes
 	UFUNCTION(BlueprintPure, Category="Math|Rotator", meta=(Keywords="construct build rotation rotate"))
 	static FRotator MakeRotationFromAxes(FVector Forward, FVector Right, FVector Up);
diff --git a/Engine/Source/Runtime/Engine/Private/KismetMathLibrary.cpp b/Engine/Source/Runtime/Engine/Private/KismetMathLibrary.cpp
--- a/Engine/Source/Runtime/Engine/Private/KismetMathLibrary.cpp
+++ b/Engine/Source/Runtime/Engine/Private/KismetMathLibrary.cpp
@@ -1770,6 +1770,11 @@ UObject* UKismetMathLibrary::SelectObject(UObject* A, UObject* B, bool bSelectA)
 	return bSelectA ? A : B;
 }
 
+UClass* UKismetMathLibrary::SelectClass(UClass* A, UClass* B, bool bSelectA)
+{
+	return bSelectA ? A : B;
+}
+
 FRotator UKismetMathLibrary::MakeRotationFromAxes(FVector Forward, FVector Right, FVector Up)
 {
 	Forward.Normalize();
