build fix

diff --git a/Engine/Source/Runtime/Engine/Classes/Engine/NiagaraEffectRenderer.h b/Engine/Source/Runtime/Engine/Classes/Engine/NiagaraEffectRenderer.h
--- a/Engine/Source/Runtime/Engine/Classes/Engine/NiagaraEffectRenderer.h
+++ b/Engine/Source/Runtime/Engine/Classes/Engine/NiagaraEffectRenderer.h
@@ -23,6 +23,26 @@ struct FNiagaraDynamicDataBase
 };
 
 
+class SimpleTimer
+{
+public:
+	SimpleTimer()
+	{
+		StartTime = FPlatformTime::Seconds() * 1000.0;
+	}
+
+	double GetElapsedMilliseconds()
+	{
+		return (FPlatformTime::Seconds()*1000.0) - StartTime;
+	}
+
+	~SimpleTimer()
+	{
+	}
+
+private:
+	double StartTime;
+};
 
 
 /**
diff --git a/Engine/Source/Runtime/Engine/Classes/Engine/NiagaraSimulation.h b/Engine/Source/Runtime/Engine/Classes/Engine/NiagaraSimulation.h
--- a/Engine/Source/Runtime/Engine/Classes/Engine/NiagaraSimulation.h
+++ b/Engine/Source/Runtime/Engine/Classes/Engine/NiagaraSimulation.h
@@ -29,26 +29,6 @@ enum EEmitterRenderModuleType
 };
 
 
-class SimpleTimer
-{
-public:
-	SimpleTimer()
-	{
-		StartTime = FPlatformTime::Seconds() * 1000.0;
-	}
-
-	double GetElapsedMilliseconds()
-	{
-		return (FPlatformTime::Seconds()*1000.0) - StartTime;
-	}
-
-	~SimpleTimer()
-	{
-	}
-
-private:
-	double StartTime;
-};
 
 /** 
  *	FniagaraEmitterProperties stores the attributes of an FNiagaraSimulation
