#ue4 - AActor::SetScale3D takes the scale by value to enable direct editing in blueprints.

diff --git a/Engine/Source/Runtime/Engine/Classes/GameFramework/Actor.h b/Engine/Source/Runtime/Engine/Classes/GameFramework/Actor.h
--- a/Engine/Source/Runtime/Engine/Classes/GameFramework/Actor.h
+++ b/Engine/Source/Runtime/Engine/Classes/GameFramework/Actor.h
@@ -685,7 +685,7 @@ public:
 
 	/** Set the Actor's world-space scale. */
 	UFUNCTION(BlueprintCallable, Category="Utilities|Transformation")
-	void SetActorScale3D(const FVector& NewScale3D);
+	void SetActorScale3D(FVector NewScale3D);
 
 	/** Returns the Actor's world-space scale. */
 	UFUNCTION(BlueprintCallable, Category="Utilities|Orientation")
diff --git a/Engine/Source/Runtime/Engine/Private/Actor.cpp b/Engine/Source/Runtime/Engine/Private/Actor.cpp
--- a/Engine/Source/Runtime/Engine/Private/Actor.cpp
+++ b/Engine/Source/Runtime/Engine/Private/Actor.cpp
@@ -2461,7 +2461,7 @@ bool AActor::SetActorLocationAndRotation(FVector NewLocation, FRotator NewRotati
 	return false;
 }
 
-void AActor::SetActorScale3D(const FVector& NewScale3D)
+void AActor::SetActorScale3D(FVector NewScale3D)
 {
 	if (RootComponent)
 	{
