#ue4 - UPrimitiveComponent::UpdateOverlaps() optimizations:

- Avoids attempting any determination of old versus new overlaps if there are no current overlaps.
- Avoids dynamic array allocation when determining old and new overlaps (RemoveSingleSwap changed to avoid shrinking array).
- Uses TInlineAllocator presized a small amount to avoid dynamic allocation in cases with few overlaps.

diff --git a/Engine/Source/Runtime/Engine/Private/PrimitiveComponent.cpp b/Engine/Source/Runtime/Engine/Private/PrimitiveComponent.cpp
--- a/Engine/Source/Runtime/Engine/Private/PrimitiveComponent.cpp
+++ b/Engine/Source/Runtime/Engine/Private/PrimitiveComponent.cpp
@@ -2172,7 +2172,8 @@ void UPrimitiveComponent::UpdateOverlaps(TArray<FOverlapInfo> const* PendingOver
 
 			// now generate full list of new touches, so we can compare to existing list and
 			// determine what changed
-			TArray<FOverlapInfo> NewOverlappingComponents;
+			typedef TArray<FOverlapInfo, TInlineAllocator<4>> TInlineOverlapInfoArray;
+			TInlineOverlapInfoArray NewOverlappingComponents;
 
 			// If pending kill, we should not generate any new overlaps
 			if (!IsPendingKill())
@@ -2211,30 +2212,35 @@ void UPrimitiveComponent::UpdateOverlaps(TArray<FOverlapInfo> const* PendingOver
 				}
 			}
 
-			// make a copy of the old that we can manipulate to avoid n^2 searching later
-			TArray< FOverlapInfo > OldOverlappingComponents = OverlappingComponents;
-
-			// Now we want to compare the old and new overlap lists to determine 
-			// what overlaps are in old and not in new (need end overlap notifies), and 
-			// what overlaps are in new and not in old (need begin overlap notifies).
-			// We do this by removing common entries from both lists, since overlapping status has not changed for them.
-			// What is left over will be what has changed.
-			for (int32 CompIdx=0; CompIdx < OldOverlappingComponents.Num() && NewOverlappingComponents.Num() > 0; ++CompIdx)
+			if (OverlappingComponents.Num() > 0)
 			{
-				if (NewOverlappingComponents.RemoveSingleSwap(OldOverlappingComponents[CompIdx]) > 0)			// RemoveSingleSwap is ok, since it is not necessary to maintain order
+				// make a copy of the old that we can manipulate to avoid n^2 searching later
+				TInlineOverlapInfoArray OldOverlappingComponents(OverlappingComponents);
+
+				// Now we want to compare the old and new overlap lists to determine 
+				// what overlaps are in old and not in new (need end overlap notifies), and 
+				// what overlaps are in new and not in old (need begin overlap notifies).
+				// We do this by removing common entries from both lists, since overlapping status has not changed for them.
+				// What is left over will be what has changed.
+				for (int32 CompIdx=0; CompIdx < OldOverlappingComponents.Num() && NewOverlappingComponents.Num() > 0; ++CompIdx)
 				{
-					OldOverlappingComponents.RemoveAtSwap(CompIdx,1,false);
-					--CompIdx;
+					// RemoveSingleSwap is ok, since it is not necessary to maintain order
+					const bool bAllowShrinking = false;
+					if (NewOverlappingComponents.RemoveSingleSwap(OldOverlappingComponents[CompIdx], bAllowShrinking) > 0)
+					{
+						OldOverlappingComponents.RemoveAtSwap(CompIdx, 1, bAllowShrinking);
+						--CompIdx;
+					}
 				}
-			}
 
-			// OldOverlappingComponents now contains only previous overlaps that are confirmed to no longer be valid.
-			for (auto CompIt = OldOverlappingComponents.CreateIterator(); CompIt; ++CompIt)
-			{
-				const FOverlapInfo& OtherOverlap = *CompIt;
-				if (OtherOverlap.OverlapInfo.Component.IsValid())
+				// OldOverlappingComponents now contains only previous overlaps that are confirmed to no longer be valid.
+				for (auto CompIt = OldOverlappingComponents.CreateIterator(); CompIt; ++CompIt)
 				{
-					EndComponentOverlap(OtherOverlap, bDoNotifies, false);
+					const FOverlapInfo& OtherOverlap = *CompIt;
+					if (OtherOverlap.OverlapInfo.Component.IsValid())
+					{
+						EndComponentOverlap(OtherOverlap, bDoNotifies, false);
+					}
 				}
 			}
 
