Deprecated ConstAEqualsB material expression const

#branch UE4

diff --git a/Engine/Source/Runtime/Engine/Classes/Materials/MaterialExpressionIf.h b/Engine/Source/Runtime/Engine/Classes/Materials/MaterialExpressionIf.h
--- a/Engine/Source/Runtime/Engine/Classes/Materials/MaterialExpressionIf.h
+++ b/Engine/Source/Runtime/Engine/Classes/Materials/MaterialExpressionIf.h
@@ -19,7 +19,7 @@ class UMaterialExpressionIf : public UMaterialExpression
 	UPROPERTY()
 	FExpressionInput AGreaterThanB;
 
-	UPROPERTY(meta = (RequiredInput = "false", ToolTip = "Defaults to 'ConstAEqualsB' if not specified"))
+	UPROPERTY(meta = (RequiredInput = "false"))
 	FExpressionInput AEqualsB;
 
 	UPROPERTY()
@@ -32,9 +32,8 @@ class UMaterialExpressionIf : public UMaterialExpression
 	UPROPERTY(EditAnywhere, Category = MaterialExpressionIf)
 	float ConstB;
 
-	/** only used if AEqualsB is not hooked up */
-	UPROPERTY(EditAnywhere, Category = MaterialExpressionIf)
-	float ConstAEqualsB;
+	UPROPERTY()
+	float ConstAEqualsB_DEPRECATED;
 
 	// Begin UMaterialExpression Interface
 	virtual int32 Compile(class FMaterialCompiler* Compiler, int32 OutputIndex, int32 MultiplexIndex) override;
diff --git a/Engine/Source/Runtime/Engine/Private/Materials/MaterialExpressions.cpp b/Engine/Source/Runtime/Engine/Private/Materials/MaterialExpressions.cpp
--- a/Engine/Source/Runtime/Engine/Private/Materials/MaterialExpressions.cpp
+++ b/Engine/Source/Runtime/Engine/Private/Materials/MaterialExpressions.cpp
@@ -4960,7 +4960,6 @@ UMaterialExpressionIf::UMaterialExpressionIf(const FObjectInitializer& ObjectIni
 
 	EqualsThreshold = 0.00001f;
 	ConstB = 0.0f;
-	ConstAEqualsB = INDEX_NONE;
 }
 
 int32 UMaterialExpressionIf::Compile(class FMaterialCompiler* Compiler, int32 OutputIndex, int32 MultiplexIndex)
@@ -4992,7 +4991,7 @@ int32 UMaterialExpressionIf::Compile(class FMaterialCompiler* Compiler, int32 Ou
 	}
 
 	int32 Arg3 = AGreaterThanB.Compile(Compiler);
-	int32 Arg4 = AEqualsB.Expression ? AEqualsB.Compile(Compiler) : Compiler->Constant(ConstAEqualsB);
+	int32 Arg4 = AEqualsB.Expression ? AEqualsB.Compile(Compiler) : INDEX_NONE;
 	int32 Arg5 = ALessThanB.Compile(Compiler);
 	int32 ThresholdArg = Compiler->Constant(EqualsThreshold);
 
