#ue4 - Optimize UpdatePhysicsVolume() performance using a new list of all PhysicsVolumes stored per game world.

- 3 ms saving in test scene with many characters that disable overlap events (because usually we use those overlaps to do the physics volume tests).

diff --git a/Engine/Source/Runtime/Engine/Classes/Engine/World.h b/Engine/Source/Runtime/Engine/Classes/Engine/World.h
--- a/Engine/Source/Runtime/Engine/Classes/Engine/World.h
+++ b/Engine/Source/Runtime/Engine/Classes/Engine/World.h
@@ -39,6 +39,7 @@ typedef TArray<TAutoWeakObjectPtr<APlayerController> >::TConstIterator FConstPla
 typedef TArray<TAutoWeakObjectPtr<APawn> >::TConstIterator FConstPawnIterator;	
 typedef TArray<TAutoWeakObjectPtr<ACameraActor> >::TConstIterator FConstCameraActorIterator;
 typedef TArray<class ULevel*>::TConstIterator FConstLevelIterator;
+typedef TArray<TAutoWeakObjectPtr<APhysicsVolume> >::TConstIterator FConstPhysicsVolumeIterator;
 
 DECLARE_LOG_CATEGORY_EXTERN(LogSpawn, Warning, All);
 
@@ -679,17 +680,20 @@ public:
 private:
 
 	/** List of all the controllers in the world. */
-	TArray<TAutoWeakObjectPtr<class AController> >	ControllerList;
+	TArray<TAutoWeakObjectPtr<class AController> > ControllerList;
 
 	/** List of all the player controllers in the world. */
-	TArray<TAutoWeakObjectPtr<class APlayerController> >	PlayerControllerList;
+	TArray<TAutoWeakObjectPtr<class APlayerController> > PlayerControllerList;
 
 	/** List of all the pawns in the world. */
-	TArray<TAutoWeakObjectPtr<class APawn> >			PawnList;	
+	TArray<TAutoWeakObjectPtr<class APawn> > PawnList;	
 
 	/** List of all the cameras in the world that auto-activate for players. */
 	TArray<TAutoWeakObjectPtr<ACameraActor> > AutoCameraActorList;
 
+	/** List of all physics volumes in the world. Does not include the DefaultPhysicsVolume. */
+	TArray<TAutoWeakObjectPtr<APhysicsVolume> > NonDefaultPhysicsVolumeList;
+
 	/** Physics scene for this world. */
 	FPhysScene*									PhysicsScene;
 
@@ -1633,6 +1637,18 @@ public:
 	 */
 	APhysicsVolume* GetDefaultPhysicsVolume() const;
 
+	/** Add a physics volume to the list of those in the world. DefaultPhysicsVolume is not tracked. Used internally by APhysicsVolume. */
+	void AddPhysicsVolume(APhysicsVolume* Volume);
+
+	/** Removes a physics volume from the list of those in the world. */
+	void RemovePhysicsVolume(APhysicsVolume* Volume);
+
+	/** Get an iterator for all PhysicsVolumes in the world that are not a DefaultPhysicsVolume. */
+	FConstPhysicsVolumeIterator GetNonDefaultPhysicsVolumeIterator() const;
+
+	/** Get the count of all PhysicsVolumes in the world that are not a DefaultPhysicsVolume. */
+	int32 GetNonDefaultPhysicsVolumeCount() const;
+
 	/**
 	 * Returns the current (or specified) level's level scripting actor
 	 *
diff --git a/Engine/Source/Runtime/Engine/Classes/GameFramework/PhysicsVolume.h b/Engine/Source/Runtime/Engine/Classes/GameFramework/PhysicsVolume.h
--- a/Engine/Source/Runtime/Engine/Classes/GameFramework/PhysicsVolume.h
+++ b/Engine/Source/Runtime/Engine/Classes/GameFramework/PhysicsVolume.h
@@ -20,6 +20,9 @@ class ENGINE_API APhysicsVolume : public AVolume
 #endif // WITH_EDITOR	
 	// End UObject interface.
 
+	virtual void PostInitializeComponents() override;
+	virtual void Destroyed() override;
+
 	//======================================================================================
 	// Character Movement related properties
 
@@ -27,7 +30,7 @@ class ENGINE_API APhysicsVolume : public AVolume
 	UPROPERTY(EditAnywhere, BlueprintReadWrite, Category=CharacterMovement)
 	float TerminalVelocity;
 
-	/** determines which PhysicsVolume takes precedence if they overlap */
+	/** Determines which PhysicsVolume takes precedence if they overlap (higher number = higher priority). */
 	UPROPERTY(EditAnywhere, BlueprintReadWrite, Category=CharacterMovement)
 	int32 Priority;
 
diff --git a/Engine/Source/Runtime/Engine/Private/PhysicsVolume.cpp b/Engine/Source/Runtime/Engine/Private/PhysicsVolume.cpp
--- a/Engine/Source/Runtime/Engine/Private/PhysicsVolume.cpp
+++ b/Engine/Source/Runtime/Engine/Private/PhysicsVolume.cpp
@@ -35,6 +35,29 @@ void APhysicsVolume::LoadedFromAnotherClass(const FName& OldClassName)
 }
 #endif
 
+void APhysicsVolume::PostInitializeComponents()
+{
+	Super::PostInitializeComponents();
+
+	UWorld* MyWorld = GetWorld();
+	if (MyWorld)
+	{
+		MyWorld->AddPhysicsVolume(this);
+	}
+}
+
+void APhysicsVolume::Destroyed()
+{
+	UWorld* MyWorld = GetWorld();
+	if (MyWorld)
+	{
+		MyWorld->RemovePhysicsVolume(this);
+	}
+
+	Super::Destroyed();
+}
+
+
 bool APhysicsVolume::IsOverlapInVolume(const class USceneComponent& TestComponent) const
 {
 	bool bInsideVolume = true;
diff --git a/Engine/Source/Runtime/Engine/Private/SceneComponent.cpp b/Engine/Source/Runtime/Engine/Private/SceneComponent.cpp
--- a/Engine/Source/Runtime/Engine/Private/SceneComponent.cpp
+++ b/Engine/Source/Runtime/Engine/Private/SceneComponent.cpp
@@ -1212,34 +1212,61 @@ void USceneComponent::UpdatePhysicsVolume( bool bTriggerNotifiers )
 {
 	if ( bShouldUpdatePhysicsVolume && !IsPendingKill() && GetWorld() )
 	{
-		APhysicsVolume *NewVolume = GetWorld()->GetDefaultPhysicsVolume();
+		UWorld* const MyWorld = GetWorld();
+		APhysicsVolume *NewVolume = MyWorld->GetDefaultPhysicsVolume();
 		
-		// check for all volumes that overlap the component
-		TArray<FOverlapResult> Hits;
-		static FName NAME_PhysicsVolumeTrace = FName(TEXT("PhysicsVolumeTrace"));
-		FComponentQueryParams Params(NAME_PhysicsVolumeTrace, GetOwner());
-
-		bool bOverlappedOrigin = false;
-		const UPrimitiveComponent* SelfAsPrimitive = Cast<UPrimitiveComponent>(this);
-		if (SelfAsPrimitive)
+		// Avoid doing anything if there are no other physics volumes in the world.
+		if (MyWorld->GetNonDefaultPhysicsVolumeCount() > 0)
 		{
-			GetWorld()->ComponentOverlapMulti(Hits, SelfAsPrimitive, GetComponentLocation(), GetComponentRotation(), GetCollisionObjectType(), Params);
-		}
-		else
-		{
-			bOverlappedOrigin = true;
-			GetWorld()->OverlapMulti(Hits, GetComponentLocation(), FQuat::Identity, GetCollisionObjectType(), FCollisionShape::MakeSphere(0.f), Params);
-		}
-
-		for( int32 HitIdx = 0; HitIdx < Hits.Num(); HitIdx++ )
-		{
-			const FOverlapResult& Link = Hits[HitIdx];
-			APhysicsVolume* const V = Cast<APhysicsVolume>(Link.GetActor());
-			if (V && (V->Priority > NewVolume->Priority))
+			// Avoid a full overlap query if we can do some quick bounds tests against the volumes.
+			bool bAnyPotentialOverlap = false;
+			for (auto VolumeIter = MyWorld->GetNonDefaultPhysicsVolumeIterator(); VolumeIter && !bAnyPotentialOverlap; ++VolumeIter)
+			{
+				const APhysicsVolume* Volume = *VolumeIter;
+				const USceneComponent* VolumeRoot = Volume->GetRootComponent();
+				if (VolumeRoot)
+				{
+					if (FBoxSphereBounds::SpheresIntersect(VolumeRoot->Bounds, Bounds))
+					{
+						if (FBoxSphereBounds::BoxesIntersect(VolumeRoot->Bounds, Bounds))
+						{
+							bAnyPotentialOverlap = true;
+						}
+					}
+				}
+				// TODO: bail if too many volumes?				
+			}
+			
+			if (bAnyPotentialOverlap)
 			{
-				if (bOverlappedOrigin || V->IsOverlapInVolume(*this))
+				// check for all volumes that overlap the component
+				TArray<FOverlapResult> Hits;
+				static FName NAME_PhysicsVolumeTrace = FName(TEXT("PhysicsVolumeTrace"));
+				FComponentQueryParams Params(NAME_PhysicsVolumeTrace, GetOwner());
+
+				bool bOverlappedOrigin = false;
+				const UPrimitiveComponent* SelfAsPrimitive = Cast<UPrimitiveComponent>(this);
+				if (SelfAsPrimitive)
+				{
+					MyWorld->ComponentOverlapMulti(Hits, SelfAsPrimitive, GetComponentLocation(), GetComponentRotation(), GetCollisionObjectType(), Params);
+				}
+				else
 				{
-					NewVolume = V;
+					bOverlappedOrigin = true;
+					MyWorld->OverlapMulti(Hits, GetComponentLocation(), FQuat::Identity, GetCollisionObjectType(), FCollisionShape::MakeSphere(0.f), Params);
+				}
+
+				for (int32 HitIdx = 0; HitIdx < Hits.Num(); HitIdx++)
+				{
+					const FOverlapResult& Link = Hits[HitIdx];
+					APhysicsVolume* const V = Cast<APhysicsVolume>(Link.GetActor());
+					if (V && (V->Priority > NewVolume->Priority))
+					{
+						if (bOverlappedOrigin || V->IsOverlapInVolume(*this))
+						{
+							NewVolume = V;
+						}
+					}
 				}
 			}
 		}
diff --git a/Engine/Source/Runtime/Engine/Private/World.cpp b/Engine/Source/Runtime/Engine/Private/World.cpp
--- a/Engine/Source/Runtime/Engine/Private/World.cpp
+++ b/Engine/Source/Runtime/Engine/Private/World.cpp
@@ -3238,6 +3238,28 @@ APhysicsVolume* UWorld::GetDefaultPhysicsVolume() const
 	return DefaultPhysicsVolume;
 }
 
+void UWorld::AddPhysicsVolume(APhysicsVolume* Volume)
+{
+	if (!Cast<ADefaultPhysicsVolume>(Volume))
+	{
+		NonDefaultPhysicsVolumeList.Add(Volume);
+	}
+}
+
+void UWorld::RemovePhysicsVolume(APhysicsVolume* Volume)
+{
+	NonDefaultPhysicsVolumeList.Remove(Volume);
+}
+
+FConstPhysicsVolumeIterator UWorld::GetNonDefaultPhysicsVolumeIterator() const
+{
+	return NonDefaultPhysicsVolumeList.CreateConstIterator();
+}
+
+int32 UWorld::GetNonDefaultPhysicsVolumeCount() const
+{
+	return NonDefaultPhysicsVolumeList.Num();
+}
 
 ALevelScriptActor* UWorld::GetLevelScriptActor( ULevel* OwnerLevel ) const
 {
