[UE-5804] static components can again set non-static properties during UCS

change summary:
- added FUCSComponentManager (local type) in ActorConstruction.cpp as an updated form of FDeferRegisterComponents; mobility is now managed during UCS by this utility class
- removed DeferRegisterComponents.h (deferred registration is no longer in use)

diff --git a/Engine/Source/Runtime/Engine/Private/ActorConstruction.cpp b/Engine/Source/Runtime/Engine/Private/ActorConstruction.cpp
--- a/Engine/Source/Runtime/Engine/Private/ActorConstruction.cpp
+++ b/Engine/Source/Runtime/Engine/Private/ActorConstruction.cpp
@@ -3,7 +3,6 @@
 #include "EnginePrivate.h"
 #include "BlueprintUtilities.h"
 #include "LatentActions.h"
-#include "DeferRegisterComponents.h"
 #include "ComponentInstanceDataCache.h"
 #include "Engine/LevelScriptActor.h"
 
@@ -15,6 +14,97 @@
 
 DEFINE_LOG_CATEGORY(LogBlueprintUserMessages);
 
+//////////////////////////////////////////////////////////////////////////
+// Local Types
+
+namespace
+{
+	/** Tracks info for components instanced during UCS execution */
+	struct FUCSComponentInfo
+	{
+		UActorComponent* Component;
+		EComponentMobility::Type Mobility;
+
+		FUCSComponentInfo(UActorComponent* InComponent)
+			:Component(InComponent)
+		{
+			USceneComponent* SceneComponent = Cast<USceneComponent>(Component);
+			if(SceneComponent != nullptr)
+			{
+				// Save original mobility
+				Mobility = SceneComponent->Mobility;
+
+				// Temporarily set to 'movable' to allow changes during UCS
+				SceneComponent->Mobility = EComponentMobility::Movable;
+			}
+		}
+	};
+
+	/** Helper class to manage components instanced during UCS execution */
+	class FUCSComponentManager
+	{
+		/** Map of actors and any components created during UCS */
+		TMap<AActor*, TArray<FUCSComponentInfo> > UCSComponentsMap;
+
+	public:
+		/** Add a component instance constructed during UCS for the given Actor */
+		void AddComponent(AActor* InActor, UActorComponent* InComponent)
+		{
+			TArray<FUCSComponentInfo>* UCSComponentsList = UCSComponentsMap.Find(InActor);
+			if (UCSComponentsList == nullptr)
+			{
+				UCSComponentsList = &UCSComponentsMap.Add(InActor, TArray<FUCSComponentInfo>());
+			}
+
+			UCSComponentsList->Add(InComponent);
+		}
+
+		/** Called after UCS execution has finished for the given Actor */
+		void PostProcessComponents(AActor* InActor)
+		{
+			TArray<FUCSComponentInfo>* UCSComponentsList = UCSComponentsMap.Find(InActor);
+			if (UCSComponentsList != nullptr)
+			{
+				for (int32 ComponentIndex = 0; ComponentIndex < UCSComponentsList->Num(); ++ComponentIndex)
+				{
+					const FUCSComponentInfo& UCSComponentInfo = (*UCSComponentsList)[ComponentIndex];
+
+					USceneComponent* SceneComponent = Cast<USceneComponent>(UCSComponentInfo.Component);
+					if(SceneComponent != nullptr)
+					{
+						// Restore original mobility after UCS execution
+						SceneComponent->Mobility = UCSComponentInfo.Mobility;
+
+						// A parent component can't be more mobile than its children, so we check for that here and adjust as needed.
+						if(SceneComponent->AttachParent != nullptr && SceneComponent->AttachParent->Mobility > SceneComponent->Mobility)
+						{
+							if(SceneComponent->IsA<UStaticMeshComponent>())
+							{
+								// SMCs can't be stationary, so always set them (and any children) to be movable
+								SceneComponent->SetMobility(EComponentMobility::Movable);
+							}
+							else
+							{
+								// Set the new component (and any children) to be at least as mobile as its parent
+								SceneComponent->SetMobility(SceneComponent->Mobility);
+							}
+						}
+					}
+				}
+
+				UCSComponentsMap.Remove(InActor);
+			}
+		}
+
+		/** Gets the FUCSComponentManager singleton. */
+		static FUCSComponentManager& Get()
+		{
+			static FUCSComponentManager Singleton;
+			return Singleton;
+		}
+	};
+}
+
 //////////////////////////////////////////////////////////////////////////
 // AActor Blueprint Stuff
 
@@ -423,59 +513,15 @@ void AActor::ExecuteConstruction(const FTransform& Transform, const FComponentIn
 	OnConstruction(Transform);
 }
 
-void FDeferRegisterComponents::DeferComponentRegistration(AActor* Actor, UActorComponent* Component)
-{
-	TArray<UActorComponent*>* ActorComponentsToRegister = ComponentsToRegister.Find(Actor);
-	if (!ActorComponentsToRegister)
-	{
-		ActorComponentsToRegister = &ComponentsToRegister.Add(Actor, TArray<UActorComponent*>());
-	}
-	
-	ActorComponentsToRegister->Add(Component);
-}
-
-void FDeferRegisterComponents::RegisterComponents(AActor* Actor)
-{
-	TArray<UActorComponent*>* ActorComponentsToRegister = ComponentsToRegister.Find(Actor);
-	if (ActorComponentsToRegister)
-	{
-		for (int32 ComponentIndex = 0; ComponentIndex < ActorComponentsToRegister->Num(); ++ComponentIndex)
-		{
-			UActorComponent* Component = (*ActorComponentsToRegister)[ComponentIndex];
-			
-			Component->RegisterComponent();
-		}
-		ComponentsToRegister.Remove(Actor);
-	}
-}
-
-// FGCObject interface
-void FDeferRegisterComponents::AddReferencedObjects(FReferenceCollector& Collector)
-{
-	for (TMap<AActor*, TArray<UActorComponent*> >::TIterator It(ComponentsToRegister); It; ++It)
-	{
-		Collector.AddReferencedObject(It.Key());
-		TArray<UActorComponent*>& Components = It.Value();
-		for (int32 ComponentIndex = 0; ComponentIndex < Components.Num(); ++ComponentIndex)
-		{
-			Collector.AddReferencedObject(Components[ComponentIndex]);
-		}
-	}
-}
-
-/* Gets the FDeferRegisterComponents singleton. */
-FDeferRegisterComponents& FDeferRegisterComponents::Get()
-{
-	static FDeferRegisterComponents Singleton;
-	return Singleton;
-}
-
 void AActor::ProcessUserConstructionScript()
 {
 	// Set a flag that this actor is currently running UserConstructionScript.
 	bRunningUserConstructionScript = true;
 	UserConstructionScript();
 	bRunningUserConstructionScript = false;
+
+	// Perform any post processing for components instanced during UCS execution.
+	FUCSComponentManager::Get().PostProcessComponents(this);
 }
 
 void AActor::FinishAndRegisterComponent(UActorComponent* Component)
@@ -518,6 +564,9 @@ UActorComponent* AActor::AddComponent(FName TemplateName, bool bManualAttachment
 	UActorComponent* NewActorComp = CreateComponentFromTemplate(Template);
 	if(NewActorComp != nullptr)
 	{
+		// Keep track of the new component during UCS execution
+		FUCSComponentManager::Get().AddComponent(this, NewActorComp);
+
 		// The user has the option of doing attachment manually where they have complete control or via the automatic rule
 		// that the first component added becomes the root component, with subsequent components attached to the root.
 		USceneComponent* NewSceneComp = Cast<USceneComponent>(NewActorComp);
@@ -531,21 +580,6 @@ UActorComponent* AActor::AddComponent(FName TemplateName, bool bManualAttachment
 				}
 				else
 				{
-					// The root component can't be more mobile than its children, so we check for that here and adjust as needed before attaching.
-					if(RootComponent->Mobility > NewSceneComp->Mobility)
-					{
-						if(NewSceneComp->IsA<UStaticMeshComponent>())
-						{
-							// SMCs can't be stationary, so always set them to be movable
-							NewSceneComp->SetMobility(EComponentMobility::Movable);
-						}
-						else
-						{
-							// Set the new component to be at least as mobile as the root
-							NewSceneComp->SetMobility(RootComponent->Mobility);
-						}
-					}
-
 					NewSceneComp->AttachTo(RootComponent);
 				}
 			}
diff --git a/Engine/Source/Runtime/Engine/Private/DeferRegisterComponents.h b/Engine/Source/Runtime/Engine/Private/DeferRegisterComponents.h
deleted file mode 100644
diff --git a/Engine/Source/Runtime/Engine/Private/SCS_Node.cpp b/Engine/Source/Runtime/Engine/Private/SCS_Node.cpp
--- a/Engine/Source/Runtime/Engine/Private/SCS_Node.cpp
+++ b/Engine/Source/Runtime/Engine/Private/SCS_Node.cpp
@@ -6,7 +6,6 @@
 #if WITH_EDITOR
 #include "Kismet2/BlueprintEditorUtils.h"
 #endif
-#include "DeferRegisterComponents.h"
 #include "Engine/SCS_Node.h"
 #include "Engine/SimpleConstructionScript.h"
 
