Remove unused replication property

diff --git a/Engine/Source/Runtime/Engine/Classes/GameFramework/Actor.h b/Engine/Source/Runtime/Engine/Classes/GameFramework/Actor.h
--- a/Engine/Source/Runtime/Engine/Classes/GameFramework/Actor.h
+++ b/Engine/Source/Runtime/Engine/Classes/GameFramework/Actor.h
@@ -114,10 +114,6 @@ public:
 	UPROPERTY(Category=Replication, EditDefaultsOnly, BlueprintReadWrite)
 	uint32 bAlwaysRelevant:1;    
 
-	/** Replicate instigator to client. */
-	UPROPERTY()
-	uint32 bReplicateInstigator:1;    
-
 	/**
 	 * If true, replicate movement/location related properties.
 	 * Actor must also be set to replicate.
