#ue4 - Move stat counter for UpdateOverlaps() to the "Game" group from "Quick". Add stat counter for UpdatePhysicsVolume().

diff --git a/Engine/Source/Runtime/Engine/Private/LevelTick.cpp b/Engine/Source/Runtime/Engine/Private/LevelTick.cpp
--- a/Engine/Source/Runtime/Engine/Private/LevelTick.cpp
+++ b/Engine/Source/Runtime/Engine/Private/LevelTick.cpp
@@ -49,6 +49,8 @@ DEFINE_STAT(STAT_GCMarkTime);
 
 DEFINE_STAT(STAT_TeleportToTime);
 DEFINE_STAT(STAT_MoveComponentTime);
+DEFINE_STAT(STAT_UpdateOverlaps);
+DEFINE_STAT(STAT_UpdatePhysicsVolume);
 
 DEFINE_STAT(STAT_PostTickComponentLW);
 DEFINE_STAT(STAT_PostTickComponentRecreate);
diff --git a/Engine/Source/Runtime/Engine/Private/PrimitiveComponent.cpp b/Engine/Source/Runtime/Engine/Private/PrimitiveComponent.cpp
--- a/Engine/Source/Runtime/Engine/Private/PrimitiveComponent.cpp
+++ b/Engine/Source/Runtime/Engine/Private/PrimitiveComponent.cpp
@@ -2147,7 +2147,7 @@ void UPrimitiveComponent::ClearMoveIgnoreActors()
 
 void UPrimitiveComponent::UpdateOverlaps(TArray<FOverlapInfo> const* PendingOverlaps, bool bDoNotifies, const TArray<FOverlapInfo>* OverlapsAtEndLocation)
 {
-	QUICK_SCOPE_CYCLE_COUNTER(STAT_UpdateOverlaps); 
+	SCOPE_CYCLE_COUNTER(STAT_UpdateOverlaps); 
 
 	if (IsDeferringMovementUpdates())
 	{
@@ -2301,6 +2301,8 @@ void UPrimitiveComponent::UpdatePhysicsVolume( bool bTriggerNotifiers )
 {
 	if (bShouldUpdatePhysicsVolume && !IsPendingKill() && GetWorld())
 	{
+		SCOPE_CYCLE_COUNTER(STAT_UpdatePhysicsVolume);
+
 		if (bGenerateOverlapEvents && IsCollisionEnabled())
 		{
 			APhysicsVolume* BestVolume = GetWorld()->GetDefaultPhysicsVolume();
diff --git a/Engine/Source/Runtime/Engine/Private/SceneComponent.cpp b/Engine/Source/Runtime/Engine/Private/SceneComponent.cpp
--- a/Engine/Source/Runtime/Engine/Private/SceneComponent.cpp
+++ b/Engine/Source/Runtime/Engine/Private/SceneComponent.cpp
@@ -1212,6 +1212,8 @@ void USceneComponent::UpdatePhysicsVolume( bool bTriggerNotifiers )
 {
 	if ( bShouldUpdatePhysicsVolume && !IsPendingKill() && GetWorld() )
 	{
+		SCOPE_CYCLE_COUNTER(STAT_UpdatePhysicsVolume);
+
 		UWorld* const MyWorld = GetWorld();
 		APhysicsVolume *NewVolume = MyWorld->GetDefaultPhysicsVolume();
 		
@@ -1346,7 +1348,7 @@ bool USceneComponent::InternalSetWorldLocationAndRotation(FVector NewLocation, F
 
 void USceneComponent::UpdateOverlaps(TArray<FOverlapInfo> const* PendingOverlaps, bool bDoNotifies, const TArray<FOverlapInfo>* OverlapsAtEndLocation)
 {
-	QUICK_SCOPE_CYCLE_COUNTER(STAT_UpdateOverlaps); 
+	SCOPE_CYCLE_COUNTER(STAT_UpdateOverlaps); 
 
 	if (IsDeferringMovementUpdates())
 	{
diff --git a/Engine/Source/Runtime/Engine/Public/EngineStats.h b/Engine/Source/Runtime/Engine/Public/EngineStats.h
--- a/Engine/Source/Runtime/Engine/Public/EngineStats.h
+++ b/Engine/Source/Runtime/Engine/Public/EngineStats.h
@@ -36,6 +36,8 @@ DECLARE_CYCLE_STAT_EXTERN(TEXT("Deferred Tick Time"),STAT_DeferredTickTime,STATG
 DECLARE_CYCLE_STAT_EXTERN(TEXT("Async Physics Time"),STAT_PhysicsTime,STATGROUP_Game, );
 DECLARE_CYCLE_STAT_EXTERN(TEXT("Spawn Actor Time"),STAT_SpawnActorTime,STATGROUP_Game, );
 DECLARE_CYCLE_STAT_EXTERN(TEXT("MoveComponent Time"),STAT_MoveComponentTime,STATGROUP_Game, );
+DECLARE_CYCLE_STAT_EXTERN(TEXT("UpdateOverlaps Time"),STAT_UpdateOverlaps,STATGROUP_Game, );
+DECLARE_CYCLE_STAT_EXTERN(TEXT("UpdatePhysicsVolume Time"),STAT_UpdatePhysicsVolume,STATGROUP_Game, );
 DECLARE_CYCLE_STAT_EXTERN(TEXT("TeleportTo Time"),STAT_TeleportToTime,STATGROUP_Game, );
 DECLARE_CYCLE_STAT_EXTERN(TEXT("GC Mark Time"),STAT_GCMarkTime,STATGROUP_Game, );
 DECLARE_CYCLE_STAT_EXTERN(TEXT("GC Sweep Time"),STAT_GCSweepTime,STATGROUP_Game, );
