Fix clang

diff --git a/Engine/Source/Runtime/Engine/Public/GenericQuadTree.h b/Engine/Source/Runtime/Engine/Public/GenericQuadTree.h
--- a/Engine/Source/Runtime/Engine/Public/GenericQuadTree.h
+++ b/Engine/Source/Runtime/Engine/Public/GenericQuadTree.h
@@ -96,16 +96,16 @@ void TQuadTree<ElementType, NodeCapacity>::Split()
 	const FVector2D XExtent = FVector2D(Extent.X, 0.f);
 	const FVector2D YExtent = FVector2D(0.f, Extent.Y);
 
-	/************************************************************************/
-	/*  ___________max
-	/* |     |     |
-	/* |     |     |
-	/* |-----c------
-	/* |     |     |
-	/* min___|_____|
-	/*
-	/* We create new quads by adding xExtent and yExtent
-	/************************************************************************/
+	/************************************************************************
+	 *  ___________max
+	 * |     |     |
+	 * |     |     |
+	 * |-----c------
+	 * |     |     |
+	 * min___|_____|
+	 *
+	 * We create new quads by adding xExtent and yExtent
+	 ************************************************************************/
 
 	const FVector2D C = Position;
 	const FVector2D TM = C + YExtent;
