Fix for a logic change that caused not havig mouse cature to ignore the desired cursor.

diff --git a/Engine/Source/Runtime/Engine/Private/GameViewportClient.cpp b/Engine/Source/Runtime/Engine/Private/GameViewportClient.cpp
--- a/Engine/Source/Runtime/Engine/Private/GameViewportClient.cpp
+++ b/Engine/Source/Runtime/Engine/Private/GameViewportClient.cpp
@@ -539,7 +539,7 @@ EMouseCursor::Type UGameViewportClient::GetCursor(FViewport* InViewport, int32 X
 
 #endif
 
-	if (!InViewport->HasMouseCapture() || !InViewport->HasFocus() || (ViewportConsole && ViewportConsole->ConsoleActive()))
+	if ((!InViewport->HasMouseCapture() && !InViewport->HasFocus()) || (ViewportConsole && ViewportConsole->ConsoleActive()))
 	{
 		return EMouseCursor::Default;
 	}
