Increased the default navigation querying extent's Z to avoid embarrassing pathfinding issues #UE4

diff --git a/Engine/Source/Runtime/Engine/Classes/AI/Navigation/NavigationTypes.h b/Engine/Source/Runtime/Engine/Classes/AI/Navigation/NavigationTypes.h
--- a/Engine/Source/Runtime/Engine/Classes/AI/Navigation/NavigationTypes.h
+++ b/Engine/Source/Runtime/Engine/Classes/AI/Navigation/NavigationTypes.h
@@ -9,7 +9,7 @@
 #define INVALID_NAVEXTENT (FVector::ZeroVector)
 
 #define DEFAULT_NAV_QUERY_EXTENT_HORIZONTAL 50.f
-#define DEFAULT_NAV_QUERY_EXTENT_VERTICAL 100.f
+#define DEFAULT_NAV_QUERY_EXTENT_VERTICAL 250.f
 
 /** uniform identifier type for navigation data elements may it be a polygon or graph node */
 typedef uint64 NavNodeRef;
