Use MikkTSpace as the default method for computing tangent space on static meshes.

diff --git a/Engine/Source/Editor/StaticMeshEditor/Private/StaticMeshEditorTools.cpp b/Engine/Source/Editor/StaticMeshEditor/Private/StaticMeshEditorTools.cpp
--- a/Engine/Source/Editor/StaticMeshEditor/Private/StaticMeshEditorTools.cpp
+++ b/Engine/Source/Editor/StaticMeshEditor/Private/StaticMeshEditorTools.cpp
@@ -312,7 +312,7 @@ void FMeshBuildSettingsLayout::GenerateHeaderRowContent( FDetailWidgetRow& NodeR
 	];
 }
 
-TAutoConsoleVariable<int32> GEnableMikkTSpaceCVar(TEXT("r.MikkTSPace"),0,TEXT("Set to be non-zero to display the option of using MikkTSpace to generate tangents for static meshes."),ECVF_Default);
+TAutoConsoleVariable<int32> GEnableMikkTSpaceCVar(TEXT("r.MikkTSPaceOptional"),0,TEXT("Set to be non-zero to display the option of using MikkTSpace to generate tangents for static meshes. MikkTSpace is the default."),ECVF_Default);
 
 void FMeshBuildSettingsLayout::GenerateChildContent( IDetailChildrenBuilder& ChildrenBuilder )
 {
diff --git a/Engine/Source/Runtime/Core/Public/UObject/ObjectVersion.h b/Engine/Source/Runtime/Core/Public/UObject/ObjectVersion.h
--- a/Engine/Source/Runtime/Core/Public/UObject/ObjectVersion.h
+++ b/Engine/Source/Runtime/Core/Public/UObject/ObjectVersion.h
@@ -444,6 +444,9 @@ enum EUnrealEngineObjectUE4Version
 	VER_UE4_SORT_ACTIVE_BONE_INDICES,
 	// Added per-frame material uniform expressions
 	VER_UE4_PERFRAME_MATERIAL_UNIFORM_EXPRESSIONS,
+	// Make MikkTSpace the default tangent space calculation method for static meshes.
+	VER_UE4_MIKKTSPACE_IS_DEFAULT,
+
 	// -----<new versions can be added before this line>-------------------------------------------------
 	// - this needs to be the last line (see note below)
 	VER_UE4_AUTOMATIC_VERSION_PLUS_ONE,
diff --git a/Engine/Source/Runtime/Engine/Classes/Engine/EngineTypes.h b/Engine/Source/Runtime/Engine/Classes/Engine/EngineTypes.h
--- a/Engine/Source/Runtime/Engine/Classes/Engine/EngineTypes.h
+++ b/Engine/Source/Runtime/Engine/Classes/Engine/EngineTypes.h
@@ -1899,7 +1899,7 @@ struct FMeshBuildSettings
 
 	/** Default settings. */
 	FMeshBuildSettings()
-		: bUseMikkTSpace(false)
+		: bUseMikkTSpace(true)
 		, bRecomputeNormals(true)
 		, bRecomputeTangents(true)
 		, bRemoveDegenerates(true)
diff --git a/Engine/Source/Runtime/Engine/Private/StaticMesh.cpp b/Engine/Source/Runtime/Engine/Private/StaticMesh.cpp
--- a/Engine/Source/Runtime/Engine/Private/StaticMesh.cpp
+++ b/Engine/Source/Runtime/Engine/Private/StaticMesh.cpp
@@ -997,12 +997,7 @@ FArchive& operator<<(FArchive& Ar, FMeshBuildSettings& BuildSettings)
 	// Note: this serializer is currently only used to build the mesh DDC key, no versioning is required
 	Ar << BuildSettings.bRecomputeNormals;
 	Ar << BuildSettings.bRecomputeTangents;
-#if PLATFORM_WINDOWS
 	Ar << BuildSettings.bUseMikkTSpace;
-#else
-	// Until the library gets compiled for Mac + Linux ignore the user's selection.
-	bool bUseMikkTSpace = false;
-#endif
 	Ar << BuildSettings.bRemoveDegenerates;
 	Ar << BuildSettings.bUseFullPrecisionUVs;
 	Ar << BuildSettings.bGenerateLightmapUVs;
@@ -1832,6 +1827,14 @@ void UStaticMesh::PostLoad()
 		}
 	}
 
+	if (GetLinkerUE4Version() < VER_UE4_MIKKTSPACE_IS_DEFAULT)
+	{
+		for (int32 i = 0; i < SourceModels.Num(); ++i)
+		{
+			SourceModels[i].BuildSettings.bUseMikkTSpace = true;
+		}
+	}
+
 	CacheDerivedData();
 
 	// Only required in an editor build as other builds process this in a different place
