removed cvar project setting which should come from Scalability, editor should automatically set it
see UE-5977 Drop in performance for the viewport in the editor for TappyChicken

diff --git a/Engine/Source/Runtime/Engine/Classes/Engine/RendererSettings.h b/Engine/Source/Runtime/Engine/Classes/Engine/RendererSettings.h
--- a/Engine/Source/Runtime/Engine/Classes/Engine/RendererSettings.h
+++ b/Engine/Source/Runtime/Engine/Classes/Engine/RendererSettings.h
@@ -229,11 +229,6 @@ class ENGINE_API URendererSettings
 		ToolTip="Control how the scene is cleared before rendering"))
 	TEnumAsByte<EClearSceneOptions::Type> ClearSceneMethod;
 
-	UPROPERTY(config, EditAnywhere, Category=Editor, meta=(
-		ConsoleVariable="r.MSAA.CompositingSampleCount",DisplayName="Editor primitive MSAA",
-		ToolTip="Affects the render quality of 3D editor objects."))
-	TEnumAsByte<ECompositingSampleCount::Type> EditorPrimitiveMSAA;
-
 	UPROPERTY(config, EditAnywhere, Category=Editor, meta=(
 		ConsoleVariable="r.WireframeCullThreshold",DisplayName="Wireframe Cull Threshold",
 		ToolTip="Screen radius at which wireframe objects are culled. Larger values can improve performance when viewing a scene in wireframe."))
