Origin rebasing disabled by default when user creates new world composition map

diff --git a/Engine/Source/Runtime/Engine/Private/WorldSettings.cpp b/Engine/Source/Runtime/Engine/Private/WorldSettings.cpp
--- a/Engine/Source/Runtime/Engine/Private/WorldSettings.cpp
+++ b/Engine/Source/Runtime/Engine/Private/WorldSettings.cpp
@@ -32,7 +32,7 @@ AWorldSettings::AWorldSettings(const FObjectInitializer& ObjectInitializer)
 	bEnableWorldBoundsChecks = true;
 	bEnableNavigationSystem = true;
 	bEnableWorldComposition = false;
-	bEnableWorldOriginRebasing = true;
+	bEnableWorldOriginRebasing = false;
 
 	KillZ = -HALF_WORLD_MAX1;
 	KillZDamageType = ConstructorStatics.DmgType_Environmental_Object.Object;
