Editor: Remove some ancient references to RF_EdSelected and some unused methods

diff --git a/Engine/Source/Runtime/Engine/Classes/Engine/Selection.h b/Engine/Source/Runtime/Engine/Classes/Engine/Selection.h
--- a/Engine/Source/Runtime/Engine/Classes/Engine/Selection.h
+++ b/Engine/Source/Runtime/Engine/Classes/Engine/Selection.h
@@ -283,11 +283,6 @@ public:
 		return SelectedClasses[ InIndex ];
 	}
 
-	/**
-	 * Sync's all objects' RF_EdSelected flag based on the current selection list.
-	 */
-	void RefreshObjectFlags();
-
 
 	// Begin UObject Interface
 	virtual void Serialize(FArchive& Ar) override;
@@ -440,9 +435,3 @@ private:
 	USelection&	Selection;
 	int32			Index;
 };
-
-
-/**
- * Ensures that only selected objects are marked with what was the RF_EdSelected flag.
- */
-void RefreshSelectionSets();
diff --git a/Engine/Source/Runtime/Engine/Private/Selection.cpp b/Engine/Source/Runtime/Engine/Private/Selection.cpp
--- a/Engine/Source/Runtime/Engine/Private/Selection.cpp
+++ b/Engine/Source/Runtime/Engine/Private/Selection.cpp
@@ -158,34 +158,8 @@ bool USelection::IsSelected(const UObject* InObject) const
 }
 
 
-void USelection::RefreshObjectFlags()
-{
-	for( int32 Index = 0 ; Index < SelectedObjects.Num() ; ++Index )
-	{
-		if ( GetSelectedObject(Index) )
-		{
-			GSelectedAnnotation.Set(GetSelectedObject(Index));
-		}
-	}
-}
-
-
 void USelection::Serialize(FArchive& Ar)
 {
 	Super::Serialize( Ar );
 	Ar << SelectedObjects;
 }
-
-
-void RefreshSelectionSets()
-{
-	// Make a list of all selection sets and refresh RF_EdSelectedFlags.
-	GSelectedAnnotation.ClearAll();
-	for ( TObjectIterator<USelection> It ; It ; ++It )
-	{
-		It->RefreshObjectFlags();
-	}
-}
-
-
-
