Fix for broken stereo rendering.

diff --git a/Engine/Source/Runtime/Engine/Private/GameEngine.cpp b/Engine/Source/Runtime/Engine/Private/GameEngine.cpp
--- a/Engine/Source/Runtime/Engine/Private/GameEngine.cpp
+++ b/Engine/Source/Runtime/Engine/Private/GameEngine.cpp
@@ -96,7 +96,7 @@ void UGameEngine::CreateGameViewportWidget( UGameViewportClient* GameViewportCli
 			// Render directly to the window backbuffer unless capturing a movie or getting screenshots
 			// @todo TEMP
 			.RenderDirectlyToWindow( !GEngine->bStartWithMatineeCapture && GIsDumpingMovie == 0 )
-/*			.EnableStereoRendering(true)*/
+			.EnableStereoRendering(true)
 			[
 				SNew(SDPIScaler)
 				.DPIScale(TAttribute<float>::Create(TAttribute<float>::FGetter::CreateUObject(this, &UGameEngine::GetGameViewportDPIScale, GameViewportClient)))
