- #AnimLayerEditing: Fix crash when you create animation that has curves.

https://jira.ol.epicgames.net/browse/UE-6546?filter=12310

diff --git a/Engine/Source/Editor/UnrealEd/Private/EditorAnimUtils.cpp b/Engine/Source/Editor/UnrealEd/Private/EditorAnimUtils.cpp
--- a/Engine/Source/Editor/UnrealEd/Private/EditorAnimUtils.cpp
+++ b/Engine/Source/Editor/UnrealEd/Private/EditorAnimUtils.cpp
@@ -180,7 +180,9 @@ namespace EditorAnimUtils
 				EditorAnimUtils::CopyAnimCurves(OldSkeleton, NewSkeleton, AssetToRetarget, USkeleton::AnimCurveMappingName, FRawCurveTracks::FloatType);
 
 				// clear transform curves since those curves won't work in new skeleton
+				// since we're deleting curves, mark this rebake flag off
 				AssetToRetarget->RawCurveData.TransformCurves.Empty();
+				AssetToRetarget->bNeedsRebake = false;
 				// I can't copy transform curves yet because transform curves need retargeting. 
 				//EditorAnimUtils::CopyAnimCurves(OldSkeleton, NewSkeleton, AssetToRetarget, USkeleton::AnimTrackCurveMappingName, FRawCurveTracks::TransformType);
 			}	
diff --git a/Engine/Source/Runtime/Engine/Classes/Animation/AnimSequence.h b/Engine/Source/Runtime/Engine/Classes/Animation/AnimSequence.h
--- a/Engine/Source/Runtime/Engine/Classes/Animation/AnimSequence.h
+++ b/Engine/Source/Runtime/Engine/Classes/Animation/AnimSequence.h
@@ -549,6 +549,7 @@ public:
 	virtual void PreSave() override;
 #if WITH_EDITOR
 	ENGINE_API virtual void PostEditChangeProperty(FPropertyChangedEvent& PropertyChangedEvent) override;
+	virtual void PostDuplicate(bool bDuplicateForPIE) override;
 #endif // WITH_EDITOR
 	ENGINE_API virtual void BeginDestroy() override;
 	ENGINE_API virtual SIZE_T GetResourceSize(EResourceSizeMode::Type Mode) override;
diff --git a/Engine/Source/Runtime/Engine/Private/Animation/AnimSequence.cpp b/Engine/Source/Runtime/Engine/Private/Animation/AnimSequence.cpp
--- a/Engine/Source/Runtime/Engine/Private/Animation/AnimSequence.cpp
+++ b/Engine/Source/Runtime/Engine/Private/Animation/AnimSequence.cpp
@@ -575,6 +575,17 @@ void UAnimSequence::PostEditChangeProperty(FPropertyChangedEvent& PropertyChange
 		PostProcessSequence();
 	}
 }
+
+void UAnimSequence::PostDuplicate(bool bDuplicateForPIE)
+{
+	// if transform curve exists, mark as bake
+	if (DoesContainTransformCurves())
+	{
+		bNeedsRebake = true;
+	}
+
+	Super::PostDuplicate(bDuplicateForPIE);
+}
 #endif // WITH_EDITOR
 
 // @todo DB: Optimize!
@@ -3352,6 +3363,11 @@ bool UAnimSequence::CreateAnimation(UAnimSequence * Sequence)
 		AnimNotifyTracks = Sequence->AnimNotifyTracks;
 		RawCurveData = Sequence->RawCurveData;
 
+		if (DoesContainTransformCurves())
+		{
+			bNeedsRebake = true;
+		}
+
 		// refresh TrackToskeletonMapIndex
 		RefreshTrackMapFromAnimTrackNames();
 
