Added UWorld::GetNumPawns() to get pawn count without having to run through the iterator.

diff --git a/Engine/Source/Runtime/Engine/Classes/Engine/World.h b/Engine/Source/Runtime/Engine/Classes/Engine/World.h
--- a/Engine/Source/Runtime/Engine/Classes/Engine/World.h
+++ b/Engine/Source/Runtime/Engine/Classes/Engine/World.h
@@ -1551,13 +1551,16 @@ public:
 	/** Returns an iterator for the controller list. */
 	FConstControllerIterator GetControllerIterator() const;
 
-	/** Returns an iterator for the pawn list. */
+	/** @return Returns an iterator for the pawn list. */
 	FConstPawnIterator GetPawnIterator() const;
 	
-	/** Returns an iterator for the player controller list. */
+	/** @return Returns the number of Pawns. */
+	int32 GetNumPawns() const;
+
+	/** @return Returns an iterator for the player controller list. */
 	FConstPlayerControllerIterator GetPlayerControllerIterator() const;
 	
-	/** Get the first player controller, or NULL if there is not one. */	
+	/** @return Returns the first player controller, or NULL if there is not one. */	
 	APlayerController* GetFirstPlayerController() const;
 
 	/*
diff --git a/Engine/Source/Runtime/Engine/Private/World.cpp b/Engine/Source/Runtime/Engine/Private/World.cpp
--- a/Engine/Source/Runtime/Engine/Private/World.cpp
+++ b/Engine/Source/Runtime/Engine/Private/World.cpp
@@ -3076,6 +3076,11 @@ FConstPawnIterator UWorld::GetPawnIterator() const
 	return PawnList.CreateConstIterator();
 }
 
+int32 UWorld::GetNumPawns() const
+{
+	return PawnList.Num();
+}
+
 void UWorld::AddPawn( APawn* Pawn )
 {
 	check( Pawn );
