[UE-5804] a few additional fixes to component handling during UCS

change summary:
- added FUCSComponentManager::PreProcessComponents() to handle temp mobility swap for native/SCS-instanced components prior to UCS execution (in order to allow attribute changes on those as well)
- modified AActor::ProcessUserConstructionScript() to call FUCSComponentManager::PreProcessComponents()
- corrected a mobility parenting fixup issue in FUCSComponentManager::PostProcessComponents()
- changed AActor reference input parameter to 'const' in FUCSComponentManager methods

diff --git a/Engine/Source/Runtime/Engine/Private/ActorConstruction.cpp b/Engine/Source/Runtime/Engine/Private/ActorConstruction.cpp
--- a/Engine/Source/Runtime/Engine/Private/ActorConstruction.cpp
+++ b/Engine/Source/Runtime/Engine/Private/ActorConstruction.cpp
@@ -44,11 +44,22 @@ namespace
 	class FUCSComponentManager
 	{
 		/** Map of actors and any components created during UCS */
-		TMap<AActor*, TArray<FUCSComponentInfo> > UCSComponentsMap;
+		TMap<const AActor*, TArray<FUCSComponentInfo> > UCSComponentsMap;
 
 	public:
-		/** Add a component instance constructed during UCS for the given Actor */
-		void AddComponent(AActor* InActor, UActorComponent* InComponent)
+		/** Called before UCS execution has started for the given Actor */
+		void PreProcessComponents(const AActor* InActor)
+		{
+			TArray<UActorComponent*> ActorComponents;
+			InActor->GetComponents(ActorComponents);
+			for (auto CompIt = ActorComponents.CreateConstIterator(); CompIt; ++CompIt)
+			{
+				AddComponent(InActor, *CompIt);
+			}
+		}
+
+		/** Add a component instance for the given Actor */
+		void AddComponent(const AActor* InActor, UActorComponent* InComponent)
 		{
 			TArray<FUCSComponentInfo>* UCSComponentsList = UCSComponentsMap.Find(InActor);
 			if (UCSComponentsList == nullptr)
@@ -56,11 +67,11 @@ namespace
 				UCSComponentsList = &UCSComponentsMap.Add(InActor, TArray<FUCSComponentInfo>());
 			}
 
-			UCSComponentsList->Add(InComponent);
+			UCSComponentsList->Add(FUCSComponentInfo(InComponent));
 		}
 
 		/** Called after UCS execution has finished for the given Actor */
-		void PostProcessComponents(AActor* InActor)
+		void PostProcessComponents(const AActor* InActor)
 		{
 			TArray<FUCSComponentInfo>* UCSComponentsList = UCSComponentsMap.Find(InActor);
 			if (UCSComponentsList != nullptr)
@@ -86,7 +97,7 @@ namespace
 							else
 							{
 								// Set the new component (and any children) to be at least as mobile as its parent
-								SceneComponent->SetMobility(SceneComponent->Mobility);
+								SceneComponent->SetMobility(SceneComponent->AttachParent->Mobility);
 							}
 						}
 					}
@@ -515,13 +526,17 @@ void AActor::ExecuteConstruction(const FTransform& Transform, const FComponentIn
 
 void AActor::ProcessUserConstructionScript()
 {
+	// Process components that may have already been instanced before UCS execution.
+	FUCSComponentManager& UCSComponentManager = FUCSComponentManager::Get();
+	UCSComponentManager.PreProcessComponents(this);
+
 	// Set a flag that this actor is currently running UserConstructionScript.
 	bRunningUserConstructionScript = true;
 	UserConstructionScript();
 	bRunningUserConstructionScript = false;
 
-	// Perform any post processing for components instanced during UCS execution.
-	FUCSComponentManager::Get().PostProcessComponents(this);
+	// Perform any post processing on this Actor's components after UCS execution.
+	UCSComponentManager.PostProcessComponents(this);
 }
 
 void AActor::FinishAndRegisterComponent(UActorComponent* Component)
