Fixed crossfade duration of zero causing divide by zero in transition entry updates.

diff --git a/Engine/Source/Runtime/Engine/Private/Animation/AnimNode_StateMachine.cpp b/Engine/Source/Runtime/Engine/Private/Animation/AnimNode_StateMachine.cpp
--- a/Engine/Source/Runtime/Engine/Private/Animation/AnimNode_StateMachine.cpp
+++ b/Engine/Source/Runtime/Engine/Private/Animation/AnimNode_StateMachine.cpp
@@ -107,7 +107,15 @@ void FAnimationActiveTransitionEntry::Update(const FAnimationUpdateContext& Cont
 			bOutFinished = true;
 		}
 
-		Alpha = CalculateAlpha(ElapsedTime / CrossfadeDuration);
+		if(CrossfadeDuration <= 0.0f)
+		{
+			Alpha = 1.0f;
+		}
+		else
+		{
+			Alpha = CalculateAlpha(ElapsedTime / CrossfadeDuration);
+		}
+
 	}
 }
 
