[UE-3566] Header comments pass for various camera-related classes (PlayerCameraManager, CameraShake, CameraModifier, et al).  Also a few light cleanups.

diff --git a/Engine/Source/Runtime/Engine/Classes/Camera/CameraActor.h b/Engine/Source/Runtime/Engine/Classes/Camera/CameraActor.h
--- a/Engine/Source/Runtime/Engine/Classes/Camera/CameraActor.h
+++ b/Engine/Source/Runtime/Engine/Classes/Camera/CameraActor.h
@@ -1,11 +1,10 @@
 // Copyright 1998-2015 Epic Games, Inc. All Rights Reserved.
 
-
 #pragma once
 #include "CameraActor.generated.h"
 
 /** 
- * A camera that can be placed in a level.
+ * A CameraActor is a camera viewpoint that can be placed in a level.
  */
 UCLASS(ClassGroup=Common, hideCategories=(Input, Rendering), showcategories=("Input|MouseInput", "Input|TouchInput"), Blueprintable)
 class ENGINE_API ACameraActor : public AActor
diff --git a/Engine/Source/Runtime/Engine/Classes/Camera/CameraAnim.h b/Engine/Source/Runtime/Engine/Classes/Camera/CameraAnim.h
--- a/Engine/Source/Runtime/Engine/Classes/Camera/CameraAnim.h
+++ b/Engine/Source/Runtime/Engine/Classes/Camera/CameraAnim.h
@@ -43,7 +43,7 @@ class UCameraAnim : public UObject
 	float BasePostProcessBlendWeight;
 
 protected:
-	// @todo document
+	/** Internal. Computes and stores the local AABB of the camera's motion. */
 	void CalcLocalAABB();
 
 public:
@@ -59,11 +59,12 @@ public:
 	 */
 	ENGINE_API bool CreateFromInterpGroup(class UInterpGroup* SrcGroup, class AMatineeActor* InMatineeActor);
 	
-	// @todo document
-	FBox GetAABB(FVector const& BaseLoc, FRotator const& BaseRot, float Scale) const;
-
-	// @todo document
-	bool InitializeCamera(class UInterpGroup* SrcGroup, class USeqAct_Interp* Interp);
+	/** 
+	 * Gets AABB of the camera's path. Useful for rough testing if you can play an animation at a certain
+	 * location in the world without penetrating geometry.
+	 * @return Returns the local-space axis-aligned bounding box of the entire motion of this animation. 
+	 */
+	ENGINE_API FBox GetAABB(FVector const& BaseLoc, FRotator const& BaseRot, float Scale) const;
 };
 
 
diff --git a/Engine/Source/Runtime/Engine/Classes/Camera/CameraModifier.h b/Engine/Source/Runtime/Engine/Classes/Camera/CameraModifier.h
--- a/Engine/Source/Runtime/Engine/Classes/Camera/CameraModifier.h
+++ b/Engine/Source/Runtime/Engine/Classes/Camera/CameraModifier.h
@@ -1,42 +1,40 @@
 // Copyright 1998-2015 Epic Games, Inc. All Rights Reserved.
 
-
 #pragma once
+
 #include "CameraModifier.generated.h"
 
+//=============================================================================
+/**
+ * A CameraModifier is a base class for objects that may adjust the final camera properties after
+ * being computed by the APlayerCameraManager (@see ModifyCamera). A CameraModifier
+ * can be stateful, and is associated uniquely with a specific APlayerCameraManager.
+ */
 UCLASS()
 class ENGINE_API UCameraModifier : public UObject
 {
 	GENERATED_UCLASS_BODY()
 
 protected:
-	/* Interface class for all camera modifiers applied to
-		a player camera actor. Sub classes responsible for
-		implementing declared functions.
-	*/
-	
-	/* Do not apply this modifier */
+	/** If true, do not apply this modifier to the camera. */
 	UPROPERTY()
 	uint32 bDisabled:1;
 
 public:
-	/* This modifier is still being applied and will disable itself */
+	/** If true, this modifier will disable itself when finished interpolating out. */
 	UPROPERTY()
 	uint32 bPendingDisable:1;
 
-	/** Camera this object is attached to */
+	/** Camera this object is associated with. */
 	UPROPERTY()
 	class APlayerCameraManager* CameraOwner;
 
 protected:
-	/**
-	 * Priority of this modifier - determines where it is added in the modifier list.
-	 * 0 = highest priority, 255 = lowest 
-	 */
+	/** Priority value that determines the order in which modifiers are applied. 0 = highest priority, 255 = lowest. */
 	UPROPERTY()
 	uint8 Priority;
 
-	/** This modifier can only be used exclusively - no modifiers of same priority allowed */
+	/** If true, no other modifiers of same priority allowed. */
 	UPROPERTY()
 	uint32 bExclusive:1;
 
@@ -48,7 +46,7 @@ protected:
 	UPROPERTY()
 	float AlphaOutTime;
 
-	/** Current blend alpha */
+	/** Current blend alpha. */
 	UPROPERTY(transient)
 	float Alpha;
 
@@ -57,22 +55,23 @@ protected:
 	float TargetAlpha;
 
 public:
-	//debug
+	/** If true, enables certain debug visualization features. */
 	UPROPERTY(EditAnywhere, Category=Debug)
 	uint32 bDebug:1;
 
-
 protected:
+	/** @return Returns the ideal blend alpha for this modifier. Interpolation will seek this value. */
 	virtual float GetTargetAlpha(class APlayerCameraManager* Camera);
-public:
-	virtual void GetDebugText(TArray<FString>& Lines) {}
 
-	/** Allow anything to happen right after creation */
+public:
+	/** 
+	 * Allows any custom initialization. Called immediately after creation.
+	 * @param Camera - The camera this modifier should be associated with.
+	 */
 	virtual void Init( APlayerCameraManager* Camera );
 	
 	/**
 	 * Directly modifies variables in the camera actor
-	 *
 	 * @param	Camera		reference to camera actor we are modifying
 	 * @param	DeltaTime	Change in time since last update
 	 * @param	InOutPOV		current Point of View, to be updated.
@@ -80,7 +79,7 @@ public:
 	 */
 	virtual bool ModifyCamera(APlayerCameraManager* Camera, float DeltaTime, struct FMinimalViewInfo& InOutPOV);
 	
-	/** Accessor function to check if modifier is inactive */
+	/** @return Returns true if modifier is disabled, false otherwise. */
 	virtual bool IsDisabled() const;
 	
 	/**
@@ -89,31 +88,38 @@ public:
 	 * adding the same modifier twice.
 	 *
 	 * @param	Camera - reference to camera actor we want add this modifier to
-	 * @return	bool   - true if modifier added to camera's modifier list, false otherwise
+	 * @return	Returns true if modifier added to camera's modifier list, false otherwise.
 	 */
 	virtual bool AddCameraModifier( APlayerCameraManager* Camera );
 	
 	/**
 	 * Camera modifier removes itself from given camera's modifier list
-	 *
-	 * @param	Camera	- reference to camara actor we want to remove this modifier from
-	 * @return	bool	- true if modifier removed successfully, false otherwise
+	 * @param	Camera	- reference to camera actor we want to remove this modifier from
+	 * @return	Returns true if modifier removed successfully, false otherwise.
 	 */
 	virtual bool RemoveCameraModifier( APlayerCameraManager* Camera );
 	
 	/** 
-	 *  Accessor functions for changing disable flag
-	 *  
+	 *  Disables this modifier.
 	 *  @param  bImmediate  - true to disable with no blend out, false (default) to allow blend out
 	 */
 	virtual void DisableModifier(bool bImmediate = false);
+
+	/** Enables this modifier. */
 	virtual void EnableModifier();
+
+	/** Toggled disabled/enabled state of this modifier. */
 	virtual void ToggleModifier();
 	
 	/**
-	 * Allow this modifier a chance to change view rotation and deltarot
+	 * Called to give modifiers a chance to adjust view rotation updates before they are applied.
+	 *
 	 * Default just returns ViewRotation unchanged
-	 * @return	bool - true if should stop looping modifiers to adjust rotation, false otherwise
+	 * @param ViewTarget - Current view target.
+	 * @param DeltaTime - Frame time in seconds.
+	 * @param OutViewRotation - In/out. The view rotation to modify.
+	 * @param OutDeltaRot - In/out. How much the rotation changed this frame.
+	 * @return Return true to prevent subsequent (lower priority) modifiers to further adjust rotation, false otherwise.
 	 */
 	virtual bool ProcessViewRotation(class AActor* ViewTarget, float DeltaTime, FRotator& OutViewRotation, FRotator& OutDeltaRot);
 
@@ -125,9 +131,8 @@ public:
 	 */
 	virtual void UpdateAlpha( APlayerCameraManager* Camera, float DeltaTime );
 
-	/** Return a world context */
+	/** @return Returns the appropriate world context for this object. */
 	UWorld* GetWorld() const;
-		
 };
 
 
diff --git a/Engine/Source/Runtime/Engine/Classes/Camera/CameraModifier_CameraShake.h b/Engine/Source/Runtime/Engine/Classes/Camera/CameraModifier_CameraShake.h
--- a/Engine/Source/Runtime/Engine/Classes/Camera/CameraModifier_CameraShake.h
+++ b/Engine/Source/Runtime/Engine/Classes/Camera/CameraModifier_CameraShake.h
@@ -5,58 +5,68 @@
  */
 
 #pragma once
+
 #include "Camera/CameraModifier.h"
 #include "CameraModifier_CameraShake.generated.h"
 
-struct FFOscillator;
-
+/** 
+ * A CameraShakeInstance is an active instance of a shake asset. It contains
+ * the necessary state information to smoothly update this shake over time.
+ */
 USTRUCT()
 struct ENGINE_API FCameraShakeInstance
 {
 	GENERATED_USTRUCT_BODY()
 
-	/** source shake */
+	/** The source camera shake data. */
 	UPROPERTY()
 	const class UCameraShake* SourceShake;
 
-	/** Used to identify shakes when single instances are desired */
+	/** Optional name of shake. Used to identify shakes when single instances are desired. */
 	UPROPERTY()
 	FName SourceShakeName;
 
-	/** <0.f means play infinitely. */
+	/** Time remaining for oscillation shakes. Less than 0.f means shake infinitely. */
 	UPROPERTY()
 	float OscillatorTimeRemaining;
 
-	/** blend vars */
+	/** True if this shake is currently blending in. */
 	UPROPERTY()
 	uint32 bBlendingIn:1;
 
+	/** How long this instance has been blending in. */
 	UPROPERTY()
 	float CurrentBlendInTime;
 
+	/** True if this shake is currently blending out. */
 	UPROPERTY()
 	uint32 bBlendingOut:1;
 
+	/** How long this instance has been blending out. */
 	UPROPERTY()
 	float CurrentBlendOutTime;
 
-	/** Current offsets. */
+	/** Current location sinusoidal offset. */
 	UPROPERTY()
 	FVector LocSinOffset;
 
+	/** Current rotational sinusoidal offset. */
 	UPROPERTY()
 	FVector RotSinOffset;
 
+	/** Current FOV sinusoidal offset. */
 	UPROPERTY()
 	float FOVSinOffset;
 
+	/** Overall intensity scale for this shake instance. */
 	UPROPERTY()
 	float Scale;
 
+	/** The playing instance of the CameraAnim-based shake, if any. */
 	UPROPERTY()
 	class UCameraAnimInst* AnimInst;
 
-	/** What space to play the shake in before applying to the camera.  Affects Anim and Oscillation both. */
+	/** What space to play the shake in before applying to the camera.  Affects both Anim and Oscillation shakes. */
 	UPROPERTY()
 	TEnumAsByte<ECameraAnimPlaySpace::Type> PlaySpace;
 
@@ -64,7 +74,6 @@ struct ENGINE_API FCameraShakeInstance
 	UPROPERTY()
 	FMatrix UserPlaySpaceMatrix;
 
-
 	FCameraShakeInstance()
 		: SourceShake(NULL)
 		, OscillatorTimeRemaining(0)
@@ -79,43 +88,79 @@ struct ENGINE_API FCameraShakeInstance
 		, AnimInst(NULL)
 		, PlaySpace(0)
 		, UserPlaySpaceMatrix(ForceInit)
-	{
-	}
+	{}
 
 };
 
+//=============================================================================
+/**
+ * A UCameraModifier_CameraShake is a camera modifier that can apply a UCameraShake to 
+ * the owning camera.
+ */
 UCLASS(config=Camera)
 class ENGINE_API UCameraModifier_CameraShake : public UCameraModifier
 {
 	GENERATED_UCLASS_BODY()
 
-	/** Active CameraShakes array */
+public:
+	/** List of active CameraShakes */
 	UPROPERTY()
 	TArray<struct FCameraShakeInstance> ActiveShakes;
 
 protected:
-	/** Scalar applied to all camera shakes in splitscreen. Normally used to dampen, since shakes feel more intense in a smaller viewport. */
+	/** Scaling factor applied to all camera shakes in when in splitscreen mode. Normally used to dampen, since shakes feel more intense in a smaller viewport. */
 	UPROPERTY(EditAnywhere, Category=CameraModifier_CameraShake)
 	float SplitScreenShakeScale;
 
-	/** For situational scaling of individual shakes. */
+	/** @return Returns the current overall scaling factor for this shake. */
 	virtual float GetShakeScale(FCameraShakeInstance const& ShakeInst) const;
+	
+	/** 
+	 * Restarts a playing camerashake from the beginning. 
+	 * @param ShakeInst - The shake instance to restart.
+	 * @param Scale - The base scale to restart at.
+	 */
 	virtual void ReinitShake(FCameraShakeInstance& ShakeInst, float Scale);
 
-public:
-	float InitializeOffset( const FFOscillator& Param );
+	/** @return Returns an initial offset for the given oscillator parameter. */
+	float InitializeOffset( const struct FFOscillator& Param ) const;
 	
-	/** Initialize camera shake structure */
+	/** 
+	 * Internal. Creates a new camera shake instance.
+	 * @param NewShake - The class of camera shake to instantiate.
+	 * @param Scale - The scalar intensity to play the shake.
+	 * @param PlaySpace - Which coordinate system to play the shake in.
+	 * @param UserPlaySpaceRot - Coordinate system to play shake when PlaySpace == CAPS_UserDefined.
+	 * @return Returns the new camera shake instance.
+	 */
 	virtual FCameraShakeInstance InitializeShake(TSubclassOf<class UCameraShake> NewShake, float Scale, ECameraAnimPlaySpace::Type PlaySpace, FRotator UserPlaySpaceRot = FRotator::ZeroRotator);
 	
-	/** Add a new screen shake to the list */
+public:
+	/** 
+	 * Adds a new active screen shake to be applied. 
+	 * @param NewShake - The class of camera shake to instantiate.
+	 * @param Scale - The scalar intensity to play the shake.
+	 * @param PlaySpace - Which coordinate system to play the shake in.
+	 * @param UserPlaySpaceRot - Coordinate system to play shake when PlaySpace == CAPS_UserDefined.
+	 */
 	virtual void AddCameraShake(TSubclassOf<class UCameraShake> NewShake, float Scale, ECameraAnimPlaySpace::Type PlaySpace=ECameraAnimPlaySpace::CameraLocal, FRotator UserPlaySpaceRot = FRotator::ZeroRotator);
 	
+	/**
+	 * Stops and removes the camera shake of the given class from the camera.
+	 * @param Shake - the camera shake class to remove.
+	 */
 	virtual void RemoveCameraShake(TSubclassOf<class UCameraShake> Shake);
+
+	/** Stops and removes all camera shakes from the camera. */
 	virtual void RemoveAllCameraShakes();
 	
-	/** Update a CameraShake */
-	virtual void UpdateCameraShake(float DeltaTime, FCameraShakeInstance& Shake, struct FMinimalViewInfo& InOutPOV);
+	/** 
+	 * Called per-tick to update a CameraShake.
+	 * @param DeltaTime - Simulation timeslice, in seconds.
+	 * @param ShakeInst - The shake instance to update.
+	 * @param InOutPOV - the POV data to update with the results of the shake.
+	 */
+	virtual void UpdateCameraShake(float DeltaTime, FCameraShakeInstance& ShakeInst, struct FMinimalViewInfo& InOutPOV);
 	
 	// Begin UCameraModifer Interface
 	virtual bool ModifyCamera(APlayerCameraManager* Camera, float DeltaTime, struct FMinimalViewInfo& InOutPOV) override;
diff --git a/Engine/Source/Runtime/Engine/Classes/Camera/CameraShake.h b/Engine/Source/Runtime/Engine/Classes/Camera/CameraShake.h
--- a/Engine/Source/Runtime/Engine/Classes/Camera/CameraShake.h
+++ b/Engine/Source/Runtime/Engine/Classes/Camera/CameraShake.h
@@ -1,12 +1,7 @@
 // Copyright 1998-2015 Epic Games, Inc. All Rights Reserved.
 
-/**
- * 
- * Object that encapsulates parameters for defining a camera shake.
- * a code-driven (oscillating) screen shake.
- */
-
 #pragma once
+
 #include "CameraShake.generated.h"
 
 /************************************************************
@@ -17,9 +12,9 @@
 UENUM()
 enum EInitialOscillatorOffset
 {
-	// Start with random offset (default)
+	/** Start with random offset (default) */
 	EOO_OffsetRandom,
-	// Start with zero offset
+	/** Start with zero offset. */
 	EOO_OffsetZero,
 	EOO_MAX,
 };
@@ -30,23 +25,23 @@ struct FFOscillator
 {
 	GENERATED_USTRUCT_BODY()
 
+	/** Amplitude of the sinusoidal oscillation. */
 	UPROPERTY(EditAnywhere, Category=FOscillator)
 	float Amplitude;
 
-	UPROPERTY(EditAnywhere, Category=FOscillator)
+	/** Frequency of the sinusoidal oscillation. */
+	UPROPERTY(EditAnywhere, Category = FOscillator)
 	float Frequency;
 
+	/** Defines how to begin (either at zero, or at a randomized value. */
 	UPROPERTY(EditAnywhere, Category=FOscillator)
 	TEnumAsByte<enum EInitialOscillatorOffset> InitialOffset;
-
-
+	
 	FFOscillator()
 		: Amplitude(0)
 		, Frequency(0)
 		, InitialOffset(0)
-	{
-	}
-
+	{}
 };
 
 /** Defines FRotator oscillation. */
@@ -55,13 +50,16 @@ struct FROscillator
 {
 	GENERATED_USTRUCT_BODY()
 
+	/** Pitch oscillation. */
 	UPROPERTY(EditAnywhere, Category=ROscillator)
 	struct FFOscillator Pitch;
 
-	UPROPERTY(EditAnywhere, Category=ROscillator)
+	/** Yaw oscillation. */
+	UPROPERTY(EditAnywhere, Category = ROscillator)
 	struct FFOscillator Yaw;
 
-	UPROPERTY(EditAnywhere, Category=ROscillator)
+	/** Roll oscillation. */
+	UPROPERTY(EditAnywhere, Category = ROscillator)
 	struct FFOscillator Roll;
 
 };
@@ -72,37 +70,58 @@ struct FVOscillator
 {
 	GENERATED_USTRUCT_BODY()
 
-	UPROPERTY(EditAnywhere, Category=VOscillator)
+	/** Oscillation in the X axis. */
+	UPROPERTY(EditAnywhere, Category = VOscillator)
 	struct FFOscillator X;
 
-	UPROPERTY(EditAnywhere, Category=VOscillator)
+	/** Oscillation in the Y axis. */
+	UPROPERTY(EditAnywhere, Category = VOscillator)
 	struct FFOscillator Y;
 
-	UPROPERTY(EditAnywhere, Category=VOscillator)
+	/** Oscillation in the Z axis. */
+	UPROPERTY(EditAnywhere, Category = VOscillator)
 	struct FFOscillator Z;
 
 };
 
+
+/**
+ * A CameraShake is an asset that defines how to shake the camera in 
+ * a particular way. CameraShakes can be authored as either oscillating shakes, 
+ * animated shakes, or both.
+ *
+ * An oscillating shake will sinusoidally vibrate various camera parameters over time. Each location
+ * and rotation axis can be oscillated independently with different parameters to create complex and
+ * random-feeling shakes. These are easier to author and tweak, but can still feel mechanical and are
+ * limited to vibration-style shakes, such as earthquakes.
+ *
+ * Animated shakes play keyframed camera animations.  These can take more effort to author, but enable
+ * more natural-feeling results and things like directional shakes.  For instance, you can have an explosion
+ * to the camera's right push it primarily to the left.
+ */
+
 UCLASS(Blueprintable, editinlinenew)
 class ENGINE_API UCameraShake : public UObject
 {
 	GENERATED_UCLASS_BODY()
 
 	/** 
-	 *  true to only allow a single instance of this shake to play at any given time. 
+	 *  If true to only allow a single instance of this shake to play at any given time. 
 	 *  Subsequent attempts to play this shake will simply restart the timer.
 	 */
 	UPROPERTY(EditAnywhere, Category=CameraShake)
 	uint32 bSingleInstance:1;
 
-	/** Duration in seconds of current screen shake. <0 means indefinite, 0 means no oscillation */
+	/** Duration in seconds of current screen shake. Less than 0 means indefinite, 0 means no oscillation. */
 	UPROPERTY(EditAnywhere, Category=Oscillation)
 	float OscillationDuration;
 
+	/** Duration of the blend-in, where the oscillation scales from 0 to 1. */
 	UPROPERTY(EditAnywhere, Category=Oscillation, meta=(ClampMin = "0.0"))
 	float OscillationBlendInTime;
 
-	UPROPERTY(EditAnywhere, Category=Oscillation, meta=(ClampMin = "0.0"))
+	/** Duration of the blend-out, where the oscillation scales from 1 to 0. */
+	UPROPERTY(EditAnywhere, Category = Oscillation, meta = (ClampMin = "0.0"))
 	float OscillationBlendOutTime;
 
 	/** Rotational oscillation */
@@ -120,6 +139,8 @@ class ENGINE_API UCameraShake : public UObject
 	/************************************************************
 	 * Parameters for defining CameraAnim-driven camera shakes
 	 ************************************************************/
+
+	/** Source camera animation to play. Can be null. */
 	UPROPERTY(EditAnywhere, Category=AnimShake)
 	class UCameraAnim* Anim;
 
@@ -141,21 +162,14 @@ class ENGINE_API UCameraShake : public UObject
 
 	/**
 	 * If true, play a random snippet of the animation of length Duration.  Implies bLoop and bRandomStartTime = true for the CameraAnim.
-	 * If false, play the full anim once, non-looped.
+	 * If false, play the full anim once, non-looped. Useful for getting variety out of a single looped CameraAnim asset.
 	 */
 	UPROPERTY(EditAnywhere, Category=AnimShake)
 	uint32 bRandomAnimSegment:1;
 
-	/** When bRandomAnimSegment=true, this defines how long the anim should play. */
+	/** When bRandomAnimSegment is true, this defines how long the anim should play. */
 	UPROPERTY(EditAnywhere, Category=AnimShake, meta=(ClampMin = "0.0", editcondition = "bRandomAnimSegment"))
 	float RandomAnimSegmentDuration;
-
-
-	// @todo document
-	virtual float GetRotOscillationMagnitude();
-
-	// @todo document
-	virtual float GetLocOscillationMagnitude();
 };
 
 
diff --git a/Engine/Source/Runtime/Engine/Classes/Camera/PlayerCameraManager.h b/Engine/Source/Runtime/Engine/Classes/Camera/PlayerCameraManager.h
--- a/Engine/Source/Runtime/Engine/Classes/Camera/PlayerCameraManager.h
+++ b/Engine/Source/Runtime/Engine/Classes/Camera/PlayerCameraManager.h
@@ -1,13 +1,16 @@
 // Copyright 1998-2015 Epic Games, Inc. All Rights Reserved.
 
 #pragma once
+
 #include "GameFramework/Actor.h"
 #include "CameraTypes.h"
 #include "CameraShake.h"
 #include "PlayerCameraManager.generated.h"
 
-class APawn;
-
+/** 
+ * Options that define how to blend when changing view targets. 
+ * @see FViewTargetTransitionParams, SetViewTarget 
+ */
 UENUM()
 enum EViewTargetBlendFunction
 {
@@ -24,41 +27,41 @@ enum EViewTargetBlendFunction
 	VTBlend_MAX,
 };
 
-/* Caching Camera, for optimization */
+/** 
+ * Cached camera POV info, stored as optimization so we only
+ * need to do a full camera update once per tick.
+ */
 USTRUCT()
 struct FCameraCacheEntry
 {
 	GENERATED_USTRUCT_BODY()
+public:
 
-	/** Cached Time Stamp */
+	/** World time this entry was created. */
 	UPROPERTY()
 	float TimeStamp;
 
-	/** cached Point of View */
+	/** Camera POV to cache. */
 	UPROPERTY()
 	FMinimalViewInfo POV;
 
 	FCameraCacheEntry()
-		: TimeStamp(0.0f)
-	{
-	}
+		: TimeStamp(0.f)
+	{}
 };
 
-/**
- * View Target definition
- * A View Target is responsible for providing the Camera with an ideal Point of View (POV)
- */
+/** A ViewTarget is the primary actor the camera is associated with. */
 USTRUCT()
 struct ENGINE_API FTViewTarget
 {
 	GENERATED_USTRUCT_BODY()
+public:
 
-	/** Target  AActor  used to compute ideal POV */
+	/** Target Actor used to compute POV */
 	UPROPERTY(EditAnywhere, BlueprintReadWrite, Category=TViewTarget)
 	class AActor* Target;
 
-public:
-	/** Point of View */
+	/** Computed point of view */
 	UPROPERTY(EditAnywhere, BlueprintReadWrite, Category=TViewTarget)
 	struct FMinimalViewInfo POV;
 
@@ -69,34 +72,35 @@ protected:
 
 	//@TODO: All for GetNextViewablePlayer
 	friend class APlayerController;
+
 public:
 	void SetNewTarget(AActor* NewTarget);
 
-	APawn* GetTargetPawn() const;
+	class APawn* GetTargetPawn() const;
 
 	bool Equal(const FTViewTarget& OtherTarget) const;
 
 	FTViewTarget()
-		: Target(NULL)
-		, PlayerState(NULL)
-	{
-	}
+		: Target(nullptr)
+		, PlayerState(nullptr)
+	{}
 
 	/** Make sure ViewTarget is valid */
 	void CheckViewTarget(APlayerController* OwningController);
 };
 
-/** A set of parameters to describe how to transition between viewtargets. */
+/** A set of parameters to describe how to transition between view targets. */
 USTRUCT()
 struct FViewTargetTransitionParams
 {
 	GENERATED_USTRUCT_BODY()
+public:
 
-	/** Total duration of blend to pending view target.  0 means no blending. */
+	/** Total duration of blend to pending view target. 0 means no blending. */
 	UPROPERTY(EditAnywhere, BlueprintReadWrite, Category=ViewTargetTransitionParams)
 	float BlendTime;
 
-	/** Function to apply to the blend parameter */
+	/** Function to apply to the blend parameter. */
 	UPROPERTY(EditAnywhere, BlueprintReadWrite, Category=ViewTargetTransitionParams)
 	TEnumAsByte<enum EViewTargetBlendFunction> BlendFunction;
 
@@ -104,8 +108,10 @@ struct FViewTargetTransitionParams
 	UPROPERTY(EditAnywhere, BlueprintReadWrite, Category=ViewTargetTransitionParams)
 	float BlendExp;
 
-	/** If true, lock outgoing viewtarget to last frame's camera position for the remainder of the blend.
-	 *  This is useful if you plan to teleport the viewtarget, but want to keep the camera motion smooth. */
+	/** 
+	 * If true, lock outgoing viewtarget to last frame's camera POV for the remainder of the blend.
+	 * This is useful if you plan to teleport the old viewtarget, but don't want to affect the blend. 
+	 */
 	UPROPERTY(EditAnywhere, BlueprintReadWrite, Category=ViewTargetTransitionParams)
 	uint32 bLockOutgoing:1;
 
@@ -116,6 +122,7 @@ struct FViewTargetTransitionParams
 		, bLockOutgoing(false)
 	{}
 
+	/** For a given linear blend value (blend percentage), return the final blend alpha with the requested function applied */
 	float GetBlendAlpha(const float& TimePct) const
 	{
 		switch (BlendFunction)
@@ -133,15 +140,32 @@ struct FViewTargetTransitionParams
 	}
 };
 
+
+//=============================================================================
 /**
- *	Defines the point of view of a player in world space.
+ * A PlayerCameraManager is responsible for managing the camera for a particular
+ * player. It defines the final view properties used by other systems (e.g. the renderer),
+ * meaning you can think of it as your virtual eyeball in the world. It can compute the 
+ * final camera properties directly, or it can arbitrate/blend between other objects or 
+ * actors that influence the camera (e.g. blending from one CameraActor to another).
+ * 
+ * The PlayerCameraManagers primary external responsibility is to reliably respond to
+ * various Get*() functions, such as GetCameraViewPoint. Most everything else is
+ * implementation detail and overrideable by user projects.
+ * 
+ * By default, a PlayerCameraManager maintains a "view target", which is the primary actor
+ * the camera is associated with. It can also apply various "post" effects to the final 
+ * view state, such as camera animations, shakes, post-process effects or special 
+ * effects such as dirt on the lens.
+ *
+ * @see https://docs.unrealengine.com/latest/INT/Gameplay/Framework/Camera/
  */
 UCLASS(notplaceable, transient, BlueprintType, Blueprintable)
 class ENGINE_API APlayerCameraManager : public AActor
 {
 	GENERATED_UCLASS_BODY()
 
-	/** PlayerController owning this Camera actor  */
+	/** PlayerController that owns this Camera actor */
 	UPROPERTY()
 	class APlayerController* PCOwner;
 
@@ -150,17 +174,17 @@ private_subobject:
 	DEPRECATED_FORGAME(4.6, "TransformComponent should not be accessed directly, please use GetTransformComponent() function instead. TransformComponent will soon be private and your code will not compile.")
 	UPROPERTY(Category = PlayerCameraManager, VisibleAnywhere, BlueprintReadOnly, meta = (AllowPrivateAccess = "true"))
 	class USceneComponent* TransformComponent;
-public:
 
-	/** Camera Mode */
+public:
+	/** Usable to define different camera behaviors. A few simple styles are implemented by default. */
 	UPROPERTY()
 	FName CameraStyle;
 
-	/** default FOV */
+	/** FOV to use by default. */
 	UPROPERTY(EditAnywhere, BlueprintReadWrite, Category=PlayerCameraManager)
 	float DefaultFOV;
 
-	/** value FOV is locked at */
+	/** Value to lock FOV at, if bLockedFOV == true. Ignored otherwise. */
 	UPROPERTY()
 	float LockedFOV;
 
@@ -168,7 +192,7 @@ public:
 	UPROPERTY(EditAnywhere, BlueprintReadWrite, Category=PlayerCameraManager)
 	float DefaultOrthoWidth;
 
-	/** value OrthoWidth is locked at */
+	/** Value OrthoWidth is locked at, if bLockedOrthoWidth == true. Ignored otherwise. */
 	UPROPERTY()
 	float LockedOrthoWidth;
 
@@ -180,37 +204,39 @@ public:
 	UPROPERTY(EditAnywhere, BlueprintReadWrite, Category=PlayerCameraManager)
 	float DefaultAspectRatio;
 
-	/** Color to fade to. */
+	/** Color to fade to (when bEnableFading == true). */
 	UPROPERTY()
 	FColor FadeColor;
 
-	/** Amount of fading to apply. */
+	/** Amount of fading to apply (when bEnableFading == true). */
 	UPROPERTY()
 	float FadeAmount;
 
-	/** Allows control over scaling individual color channels in the final image. */
+	/** Allows control over scaling individual color channels in the final image (when bEnableColorScaling == true). */
 	UPROPERTY()
 	FVector ColorScale;
 
-	/** Desired color scale which ColorScale will interpolate to */
+	/** Desired color scale which ColorScale will interpolate to (when bEnableColorScaling and bEnableColorScaleInterp == true) */
 	UPROPERTY()
 	FVector DesiredColorScale;
 
-	/** Color scale value at start of interpolation */
+	/** Color scale value at start of interpolation (when bEnableColorScaling and bEnableColorScaleInterp == true) */
 	UPROPERTY()
 	FVector OriginalColorScale;
 
-	/** Total time for color scale interpolation to complete */
+	/** Total time for color scale interpolation to complete (when bEnableColorScaling and bEnableColorScaleInterp == true) */
 	UPROPERTY()
 	float ColorScaleInterpDuration;
 
-	/** Time at which interpolation started */
+	/** Time at which interpolation started (when bEnableColorScaling and bEnableColorScaleInterp == true) */
 	UPROPERTY()
 	float ColorScaleInterpStartTime;
 
+	/** Cached camera properties. */
 	UPROPERTY()
 	struct FCameraCacheEntry CameraCache;
 
+	/** Cached camera properties. */
 	UPROPERTY()
 	struct FCameraCacheEntry LastFrameCameraCache;
 
@@ -222,34 +248,35 @@ public:
 	UPROPERTY()
 	struct FTViewTarget PendingViewTarget;
 
-	/** Time left when blending to pending view target */
+	/** Time remaining in viewtarget blend. */
 	UPROPERTY()
 	float BlendTimeToGo;
 
+	/** Current view target transition blend parameters. */
 	UPROPERTY()
 	struct FViewTargetTransitionParams BlendParams;
 
-	/** List of camera modifiers to apply during update of camera position/ rotation */
+	/** List of active camera modifiers that have a chance to update the final camera POV */
 	UPROPERTY()
 	TArray<class UCameraModifier*> ModifierList;
 
-	/** Distance to place free camera from view target */
+	/** Distance to place free camera from view target (used in certain CameraStyles) */
 	UPROPERTY(EditAnywhere, BlueprintReadWrite, Category=Debug)
 	float FreeCamDistance;
 
-	/** Offset to Z free camera position */
+	/** Offset to Z free camera position (used in certain CameraStyles) */
 	UPROPERTY(EditAnywhere, BlueprintReadWrite, Category=Debug)
 	FVector FreeCamOffset;
 
-	/** camera fade management */
+	/** Current camera fade alpha value (when bEnableFading == true) */
 	UPROPERTY()
 	FVector2D FadeAlpha;
 
-	/** length of time for fade to occur over */
+	/** Total duration of the camera fade (when bEnableFading == true) */
 	UPROPERTY()
 	float FadeTime;
 
-	/** time remaining for fade to occur over */
+	/** Time remaining in camera fade (when bEnableFading == true) */
 	UPROPERTY()
 	float FadeTimeRemaining;
 
@@ -269,11 +296,11 @@ public:
 	class UCameraModifier_CameraShake* CameraShakeCamMod;
 
 protected:
-	/** Class to use when instantiating screenshake modifier object.  Provided to support overrides. */
+	/** Class to use when instantiating screenshake modifier object. */
 	UPROPERTY(EditAnywhere, BlueprintReadWrite, Category=PlayerCameraManager)
-	TSubclassOf<class UCameraModifier_CameraShake>  CameraShakeCamModClass;
+	TSubclassOf<class UCameraModifier_CameraShake> CameraShakeCamModClass;
 
-	/** By default camera modifiers are not applied to stock debug cameras. Setting to true will always apply modifiers. */
+	/** By default camera modifiers are not applied to stock debug cameras (e.g. CameraStyle == "FreeCam"). Setting to true will always apply modifiers. */
 	UPROPERTY()
 	bool bAlwaysApplyModifiers;
 	
@@ -281,109 +308,120 @@ protected:
 	// CameraAnim support
 	////////////////////////
 	
-	/** Pool of anim instance objects available with which to play camera animations */
+	static const int32 MAX_ACTIVE_CAMERA_ANIMS = 8;
+
+	/** Internal pool of camera anim instance objects available for playing camera animations. Defines the max number of camera anims that can play simultaneously. */
 	UPROPERTY()
-	class UCameraAnimInst* AnimInstPool[8];    /*MAX_ACTIVE_CAMERA_ANIMS @fixmeconst*/
+	class UCameraAnimInst* AnimInstPool[8];    /*MAX_ACTIVE_CAMERA_ANIMS @fixme constant */
 
-	/** Parallel arrays */
+	/** Internal list of active post process effects. Parallel array to PostProcessBlendCacheWeights. */
 	TArray<struct FPostProcessSettings> PostProcessBlendCache;
+	/** Internal list of weights for active post process effects. Parallel array to PostProcessBlendCache. */
 	TArray<float> PostProcessBlendCacheWeights;
 
+	/** Adds a postprocess effect at the given weight. */
 	void AddCachedPPBlend(struct FPostProcessSettings& PPSettings, float BlendWeight);
+	/** Removes all postprocess effects. */
 	void ClearCachedPPBlends();
 
 public:
-	/** Array of anim instances that are currently playing and in-use */
+	/** Array of camera anim instances that are currently playing and in-use */
 	UPROPERTY()
 	TArray<class UCameraAnimInst*> ActiveAnims;
 
+	/** Returns active post process info. */
 	void GetCachedPostProcessBlends(TArray<struct FPostProcessSettings> const*& OutPPSettings, TArray<float> const*& OutBlendWeigthts) const;
 	
 protected:
-	/** Array of anim instances that are not playing and available */
+	/** Array of camera anim instances that are not playing and available to be used. */
 	UPROPERTY()
 	TArray<class UCameraAnimInst*> FreeAnims;
 
-	/** Internal.  Receives the output of individual camera animations. */
+	/** Internal. Receives the output of individual camera animations. */
 	UPROPERTY(transient)
 	class ACameraActor* AnimCameraActor;
 
-	/** true if FOV is locked to a constant value*/
+	/** Tuue if FOV is locked to a constant value (@see LockedFOV) */
 	UPROPERTY()
 	uint32 bLockedFOV:1;
 
-	/** true if OrthoWidth is locked to a constant value*/
+	/** true if OrthoWidth is locked to a constant value (@see LockedOrthoWidth) */
 	UPROPERTY()
 	uint32 bLockedOrthoWidth:1;
 
 public:
-	/** If we should apply FadeColor/FadeAmount to the screen. */
+	/** True if we should apply FadeColor/FadeAmount to the screen. */
 	UPROPERTY()
 	uint32 bEnableFading:1;
 
-	/** Apply fading of audio alongside the video */
+	/** True to apply fading of audio alongside the video. */
 	UPROPERTY()
 	uint32 bFadeAudio:1;
 
-	/** Turn on scaling of color channels in final image using ColorScale property. */
+	/** True to turn on scaling of color channels in final image using ColorScale property. */
 	UPROPERTY()
 	uint32 bEnableColorScaling:1;
 
-	/** Should interpolate color scale values */
+	/** True to smoothly interpolate color scale values when they change. */
 	UPROPERTY()
 	uint32 bEnableColorScaleInterp:1;
 
-	/** True if clients are handling setting their own viewtarget and the server should not replicate it (e.g. during certain matinees) */
+	/** True if clients are handling setting their own viewtarget and the server should not replicate it (e.g. during certain Matinee sequences) */
 	UPROPERTY()
 	uint32 bClientSimulatingViewTarget:1;
 
-	/** if true, server will use camera positions replicated from the client instead of calculating locally. */
+	/** True if server will use camera positions replicated from the client instead of calculating them locally. */
 	UPROPERTY(EditAnywhere, BlueprintReadOnly, Category=PlayerCameraManager)
 	uint32 bUseClientSideCameraUpdates:1;
 
-	/** If true, replicate the client side camera position but don't use it, and draw the positions on the server */
+	/** For debugging. If true, replicate the client side camera position but don't use it, and draw the positions on the server */
 	UPROPERTY()
 	uint32 bDebugClientSideCamera:1;
 
-	/** if true, send a camera update to the server on next update */
+	/** If true, send a camera update to the server on next update. */
 	UPROPERTY(transient)
 	uint32 bShouldSendClientSideCameraUpdate:1;
 
-	// Whether we did a camera cut this frame. Automatically reset to false every frame.
-	// This flag affects various things in the renderer (such as whether to use the occlusion queries from last frame, and motion blur)
+	/** 
+	 * True if we did a camera cut this frame. Automatically reset to false every frame.
+	 * This flag affects various things in the renderer (such as whether to use the occlusion queries from last frame, and motion blur).
+	 */
 	UPROPERTY(transient)
 	uint32 bGameCameraCutThisFrame:1;
 
-	/** Whether this camera's orientation should be updated by most recent HMD orientation or not. */
+	/** True if camera's orientation should be updated by most recent HMD orientation. */
 	UPROPERTY(EditAnywhere, BlueprintReadWrite, Category = PlayerCameraManager)
-	uint32 bFollowHmdOrientation : 1;
+	uint32 bFollowHmdOrientation:1;
 
-	/** pitch limiting of camera facing direction */
+	/** Minimum view pitch, in degrees. */
 	UPROPERTY(EditAnywhere, BlueprintReadWrite, Category=PlayerCameraManager)
 	float ViewPitchMin;
 
-	/** pitch limiting of camera facing direction */
+	/** Maximum view pitch, in degrees. */
 	UPROPERTY(EditAnywhere, BlueprintReadWrite, Category=PlayerCameraManager)
 	float ViewPitchMax;
 
-	/** yaw limiting of camera facing direction */
+	/** Minimum view yaw, in degrees. */
 	UPROPERTY(EditAnywhere, BlueprintReadWrite, Category=PlayerCameraManager)
 	float ViewYawMin;
 
-	/** yaw limiting of camera facing direction */
+	/** Maximum view yaw, in degrees. */
 	UPROPERTY(EditAnywhere, BlueprintReadWrite, Category=PlayerCameraManager)
 	float ViewYawMax;
 
-	/** roll limiting of camera facing direction */
+	/** Minimum view roll, in degrees. */
 	UPROPERTY(EditAnywhere, BlueprintReadWrite, Category=PlayerCameraManager)
 	float ViewRollMin;
 
-	/** roll limiting of camera facing direction */
+	/** Maximum view roll, in degrees. */
 	UPROPERTY(EditAnywhere, BlueprintReadWrite, Category=PlayerCameraManager)
 	float ViewRollMax;
-
-
-	/** Blueprint hook to allow blueprints to implement custom cameras. Return true to use returned values, or false to ignore them. */
+	
+	/** 
+	 * Blueprint hook to allow blueprints to override existing camera behavior or implement custom cameras.
+	 * If this function returns true, we will use the given returned values and skip further calculations to determine
+	 * final camera POV. 
+	 */
 	UFUNCTION(BlueprintImplementableEvent, BlueprintCosmetic)
 	virtual bool BlueprintUpdateCamera(AActor* CameraTarget, FVector& NewCameraLocation, FRotator& NewCameraRotation, float& NewCameraFOV);
 
@@ -391,18 +429,15 @@ public:
 	UFUNCTION(BlueprintCallable, Category="Game|Player")
 	virtual APlayerController* GetOwningPlayerController() const;
 
-	static const int32 MAX_ACTIVE_CAMERA_ANIMS = 8;
-
-
 	virtual void AssignViewTarget(AActor* NewTarget, FTViewTarget& VT, struct FViewTargetTransitionParams TransitionParams=FViewTargetTransitionParams());
 
 	friend struct FTViewTarget;
 public:
-	/** Returns current ViewTarget */
+	/** @return the current ViewTarget. */
 	AActor* GetViewTarget() const;
 
-	/** Returns the pawn that is being viewed (if any) */
-	APawn* GetViewTargetPawn() const;
+	/** @return the ViewTarget if it is an APawn, or nullptr otherwise */
+	class APawn* GetViewTargetPawn() const;
 
 	// Begin AActor Interface
 	virtual bool ShouldTickIfViewportsOnly() const override;
@@ -412,112 +447,128 @@ public:
 	virtual void ApplyWorldOffset(const FVector& InOffset, bool bWorldShift) override;
 	// End AActor Interface
 
-	/** Static.  Plays an in-world camera shake that affects all nearby players, with distance-based attenuation.
-	 * @param InWorld - World context
-	 * @param Shake - Camera shake asset to use
-	 * @param Epicenter - location to place the effect in world space
-	 * @param InnerRadius - Cameras inside this radius are ignored
-	 * @param OuterRadius - Cameras outside of InnerRadius and inside this are effected
-	 * @param Falloff - Affects falloff of effect as it nears OuterRadius
-	 * @param bOrientShakeTowardsEpicenter - Changes the rotation of shake to point towards epicenter instead of forward
+	/** Static.  Plays an in-world camera shake that affects all nearby players, with radial distance-based attenuation.
+	 * @param InWorld - World context.
+	 * @param Shake - Camera shake asset to use.
+	 * @param Epicenter - Location to place the effect in world space
+	 * @param InnerRadius - Cameras inside this radius get the full intensity shake.
+	 * @param OuterRadius - Cameras outside this radius are not affected.
+	 * @param Falloff - Exponent that describes the shake intensity falloff curve between InnerRadius and OuterRadius. 1.0 is linear.
+	 * @param bOrientShakeTowardsEpicenter - Changes the rotation of shake to point towards epicenter instead of forward. Useful for things like directional hits.
 	 */
-	static void PlayWorldCameraShake(UWorld* InWorld, TSubclassOf<UCameraShake> Shake, FVector Epicenter, float InnerRadius, float OuterRadius, float Falloff, bool bOrientShakeTowardsEpicenter = false );
+	static void PlayWorldCameraShake(UWorld* InWorld, TSubclassOf<UCameraShake> Shake, FVector Epicenter, float InnerRadius, float OuterRadius, float Falloff, bool bOrientShakeTowardsEpicenter = false);
 
 protected:
-	/** Internal.  Returns intensity scalar in the range [0..1] for a shake originating at Epicenter. */
-	static float CalcRadialShakeScale(class APlayerCameraManager* Cam,FVector Epicenter,float InnerRadius,float OuterRadius,float Falloff);
+	/** 
+	 * Internal. Calculates shake scale for a particular camera.
+	 * @return Returns the intensity scalar in the range [0..1] for a shake originating at Epicenter. 
+	 */
+	static float CalcRadialShakeScale(class APlayerCameraManager* Cam, FVector Epicenter, float InnerRadius, float OuterRadius, float Falloff);
+
 public:
 
 	/**
-	 * Performs camera update.
-	 * Called once per frame after all actors have been ticked.
-	 * Non-local players replicate the POV if bUseClientSideCameraUpdates is true
+	 * Performs per-tick camera update. Called once per tick after all other actors have been ticked.
+	 * Non-local players replicate the POV if bUseClientSideCameraUpdates is true.
 	 */
 	virtual void UpdateCamera(float DeltaTime);
 
-	/** Internal. Creates and initializes a new camera modifier of the specified class, returns the object ref. */
+	/** 
+	 * Internal. Creates and initializes a new camera modifier of the specified class. 
+	 * @param ModifierClass - The class of camera modifier to create.
+	 * @return Returns the newly created camera modifier.
+	 */
 	virtual class UCameraModifier* CreateCameraModifier(TSubclassOf<class UCameraModifier> ModifierClass);
 	
-	
 	/**
-	 * Apply modifiers on Camera.
-	 * @param	DeltaTime	Time is seconds since last update
-	 * @param	InOutPOV		Point of View
+	 * Applies the current set of camera modifiers to the given camera POV.
+	 * @param	DeltaTime	Time in seconds since last update
+	 * @param	InOutPOV	Point of View
 	 */
 	virtual void ApplyCameraModifiers(float DeltaTime, FMinimalViewInfo& InOutPOV);
 	
 	/**
-	 * Initialize Camera for associated PlayerController
-	 * @param	PC	PlayerController attached to this Camera.
+	 * Initialize this PlayerCameraManager for the given associated PlayerController.
+	 * @param PC	PlayerController associated with this Camera.
 	 */
 	virtual void InitializeFor(class APlayerController* PC);
 	
-	/** Returns camera's current FOV angle */
+	/** @return Returns the camera's current full FOV angle, in degrees. */
 	virtual float GetFOVAngle() const;
 	
-	/** Lock FOV to a specific value. */
+	/** 
+	 * Locks the FOV to the given value.  Unlock with UnlockFOV.
+	 * @param NewFOV - New full FOV angle to use, in degrees.
+	 */
 	virtual void SetFOV(float NewFOV);
 
-	/** Remove lock from FOV */
+	/** Unlocks the FOV. */
 	virtual void UnlockFOV();
 
-	/** Returns true if this camera is using an orthographic perspective */
+	/** @return Returns true if this camera is using an orthographic perspective. */
 	virtual bool IsOrthographic() const;
 
-	/** Returns the current orthographic width for the camera. Only used if IsOrthographic returns true */
+	/** @return Returns the current orthographic width for the camera. */
 	virtual float GetOrthoWidth() const;
 
-	/** Sets the current orthographic width for the camera. Only used if IsOrthographic returns true */
+	/** 
+	 * Sets and locks the current orthographic width for the camera. Unlock with UnlockOrthoWidth. 
+	 * Only used if IsOrthographic returns true.
+	 * @param OrthoWidth - New orthographic width.
+	 */
 	virtual void SetOrthoWidth(float OrthoWidth);
 
-	/** Remove lock from OrthoWidth */
+	/** Unlocks OrthoWidth value */
 	virtual void UnlockOrthoWidth();
 
 	/**
 	 * Master function to retrieve Camera's actual view point.
-	 * do not call this directly, call PlayerController::GetPlayerViewPoint() instead.
+	 * Consider calling PlayerController::GetPlayerViewPoint() instead.
 	 *
-	 * @param	OutCamLoc	Camera Location
-	 * @param	OutCamRot	Camera Rotation
+	 * @param	OutCamLoc	Returned camera location
+	 * @param	OutCamRot	Returned camera rotation
 	 */
 	void GetCameraViewPoint(FVector& OutCamLoc, FRotator& OutCamRot) const;
 	
-	// @todo document
+	/** @return Returns camera's current rotation. */
 	UFUNCTION(BlueprintCallable, Category = "Camera")
 	FRotator GetCameraRotation() const;
 
-	// @todo document
+	/** @return Returns camera's current location. */
 	UFUNCTION(BlueprintCallable, Category = "Camera")
 	FVector GetCameraLocation() const;
 	
-	/** Sets the new desired color scale and enables interpolation. */
+	/** 
+	 * Sets the new desired color scale, enables color scaling, and enables color scale interpolation. 
+	 * @param NewColorScale - new color scale to use
+	 * @param InterpTime - duration of the interpolation from old to new, in seconds.
+	 */
 	virtual void SetDesiredColorScale(FVector NewColorScale, float InterpTime);
 	
 protected:
-	// @todo document
+	/** Internal function conditionally called from UpdateCamera to do the actual work of updating the camera. */
 	virtual void DoUpdateCamera(float DeltaTime);
 	
-	/** Apply audio fading */
+	/** Internal. Applies appropriate audio fading to the audio system. */
 	virtual void ApplyAudioFade();
 	
 	/**
-	 * Blend 2 viewtargets.
+	 * Internal helper to blend two viewtargets.
 	 *
 	 * @param	A		Source view target
-	 * @paramn	B		destination view target
-	 * @param	Alpha	Alpha, % of blend from A to B.
+	 * @param	B		destination view target
+	 * @param	Alpha	% of blend from A to B.
 	 */
-	FPOV BlendViewTargets(const FTViewTarget& A,const FTViewTarget& B, float Alpha);
+	FPOV BlendViewTargets(const FTViewTarget& A, const FTViewTarget& B, float Alpha);
 	
 public:
-	/** Cache update results */
+	/** Caches given final POV info for efficient access from other game code. */
 	void FillCameraCache(const FMinimalViewInfo& NewInfo);
 	
 protected:
 	/**
-	 * Query ViewTarget and outputs Point Of View.
-	 *
-	 * @param	OutVT		ViewTarget to use.
+	 * Calculates an updated POV for the given viewtarget.
+	 * @param	OutVT		ViewTarget to update.
 	 * @param	DeltaTime	Delta Time since last camera update (in seconds).
 	 */
 	virtual void UpdateViewTarget(FTViewTarget& OutVT, float DeltaTime);
@@ -526,24 +577,45 @@ protected:
 	virtual void UpdateCameraLensEffects( const FTViewTarget& OutVT );
 
 public:
-	/** Set a new ViewTarget with optional transition time */
+	/** 
+	 * Sets a new ViewTarget.
+	 * @param NewViewTarget - New viewtarget actor.
+	 * @param TransitionParams - Optional parameters to define the interpolation from the old viewtarget to the new. Transition will be instant by default.
+	 */
 	void SetViewTarget(class AActor* NewViewTarget, FViewTargetTransitionParams TransitionParams = FViewTargetTransitionParams());
 	
-	/** Give each modifier a chance to change view rotation/deltarot.  Base implementation also limits rotation components using LimitViewPitch, LimitViewYaw, and LimitViewRoll. */
+	/** 
+	 * Called to give PlayerCameraManager a chance to adjust view rotation updates before they are applied. 
+	 * e.g. The base implementation enforces view rotation limits using LimitViewPitch, et al. 
+	 * @param DeltaTime - Frame time in seconds.
+	 * @param OutViewRotation - In/out. The view rotation to modify.
+	 * @param OutDeltaRot - In/out. How much the rotation changed this frame.
+	 */
 	virtual void ProcessViewRotation(float DeltaTime, FRotator& OutViewRotation, FRotator& OutDeltaRot);
 
-	/** Finds the first instance of a lens effect of the given class, using linear search. */
+	//
+	// Camera Lens Effects
+	//
+	
+	/** @return Returns first instance of a lens effect of the given class. */
 	virtual class AEmitterCameraLensEffectBase* FindCameraLensEffect(TSubclassOf<class AEmitterCameraLensEffectBase> LensEffectEmitterClass);
 	
-	/** Initiates a camera lens effect of the given class on this camera. */
+	/** 
+	 * Creates a camera lens effect of the given class on this camera. 
+	 * @param LensEffectEmitterClass - Class of lens effect emitter to create.
+	 * @return Returns the new emitter actor.
+	 */
 	UFUNCTION(BlueprintCallable, Category = "Camera")
 	virtual AEmitterCameraLensEffectBase* AddCameraLensEffect(TSubclassOf<class AEmitterCameraLensEffectBase> LensEffectEmitterClass);
 	
-	/** Removes this particular lens effect from the camera. */
+	/** 
+	 * Removes the given lens effect from the camera. 
+	 * @param Emitter - the emitter actor to remove from the camera
+	 */
 	UFUNCTION(BlueprintCallable, Category = "Camera")
 	virtual void RemoveCameraLensEffect(class AEmitterCameraLensEffectBase* Emitter);
 	
-	/** Removes all Camera Lens Effects. */
+	/** Removes all camera lens effects. */
 	UFUNCTION(BlueprintCallable, Category = "Camera")
 	virtual void ClearCameraLensEffects();
 
@@ -551,13 +623,19 @@ public:
 	// Camera Shakes.
 	//
 
-	/** Play a camera shake */
+	/** 
+	 * Plays a camera shake on this camera.
+	 * @param Shake - The class of camera shake to play.
+	 * @param Scale - Scalar defining how "intense" to play the anim. 1.0 is normal or as authored.
+	 * @param PlaySpace - Which coordinate system to play the shake in.
+	 * @param UserPlaySpaceRot - Coordinate system to play shake when PlaySpace == CAPS_UserDefined.
+	 */
 	virtual void PlayCameraShake(TSubclassOf<class UCameraShake> Shake, float Scale, enum ECameraAnimPlaySpace::Type PlaySpace=ECameraAnimPlaySpace::CameraLocal, FRotator UserPlaySpaceRot = FRotator::ZeroRotator);
 	
-	/** Stop playing a camera shake. */
+	/** Stops playing CameraShake of the given class. */
 	virtual void StopCameraShake(TSubclassOf<class UCameraShake> Shake);
 
-	// @todo document
+	/** Stops playing all camera shakes. */
 	virtual void ClearAllCameraShakes();
 
 	//
@@ -603,34 +681,52 @@ public:
 	UFUNCTION(BlueprintCallable, Category = "Camera Animation")
 	virtual void StopAllCameraAnims(bool bImmediate = false);
 
-	/** Returns first existing instance of the specified camera anim, or NULL if none exists. */
-	UCameraAnimInst* FindInstanceOfCameraAnim(class UCameraAnim const* Anim) const;
+	/** @return Returns first existing instance of the specified camera anim, or NULL if none exists. */
+	class UCameraAnimInst* FindInstanceOfCameraAnim(class UCameraAnim const* Anim) const;
 
 protected:
 	/** Gets specified temporary CameraActor ready to update the specified Anim. */
 	void InitTempCameraActor(class ACameraActor* CamActor, class UCameraAnimInst const* AnimInstToInitFor) const;
 	void ApplyAnimToCamera(class ACameraActor const* AnimatedCamActor, class UCameraAnimInst const* AnimInst, FMinimalViewInfo& InOutPOV);
 
-	/** Returns an available CameraAnimInst, or NULL if no more are available. */
-	UCameraAnimInst* AllocCameraAnimInst();
+	/** @return Returns an available CameraAnimInst, or NULL if no more are available. */
+	class UCameraAnimInst* AllocCameraAnimInst();
 
-	/** Returns an available CameraAnimInst, or NULL if no more are available. */
-	void ReleaseCameraAnimInst(UCameraAnimInst* Inst);
+	/** 
+	 * Frees an allocated CameraAnimInst for future use. 
+	 * @param Inst - Instance to free. 
+	 */
+	void ReleaseCameraAnimInst(class UCameraAnimInst* Inst);
 
 protected:
-	/** Limit the player's view pitch.  */
-	virtual void LimitViewPitch( FRotator& ViewRotation, float InViewPitchMin, float InViewPitchMax );
+	/** 
+	 * Limit the player's view pitch. 
+	 * @param ViewRotation - ViewRotation to modify. Both input and output.
+	 * @param InViewPitchMin - Minimum view pitch, in degrees.
+	 * @param InViewPitchMax - Maximum view pitch, in degrees.
+	 */
+	virtual void LimitViewPitch(FRotator& ViewRotation, float InViewPitchMin, float InViewPitchMax);
 
-	/** Limit the player's view roll. */
+	/** 
+	 * Limit the player's view roll. 
+	 * @param ViewRotation - ViewRotation to modify. Both input and output.
+	 * @param InViewRollMin - Minimum view roll, in degrees.
+	 * @param InViewRollMax - Maximum view roll, in degrees.
+	 */
 	virtual void LimitViewRoll(FRotator& ViewRotation, float InViewRollMin, float InViewRollMax);
 
-	/** Limit the player's view yaw. */
+	/** 
+	 * Limit the player's view yaw. 
+	 * @param ViewRotation - ViewRotation to modify. Both input and output.
+	 * @param InViewYawMin - Minimum view yaw, in degrees.
+	 * @param InViewYawMax - Maximum view yaw, in degrees.
+	 */
 	virtual void LimitViewYaw(FRotator& ViewRotation, float InViewYawMin, float InViewYawMax);
 
 	/**
-	 * Query ViewTarget and updates the view target Point Of View.
-	 * This is called from UpdateViewTarget under normal circumstances (target is not a camera actor and no debug cameras are active)
-	 * Note: This function will eventually be replaced, but it currently provides a way for subclasses to override behavior without copy-pasting the special case code
+	 * Does the actual work to UpdateViewTarget. This is called from UpdateViewTarget under normal 
+	 * circumstances (target is not a camera actor and no debug cameras are active)
+	 * Provides a way for subclasses to override behavior without copy-pasting the special case code.
 	 *
 	 * @param	OutVT		ViewTarget to use.
 	 * @param	DeltaTime	Delta Time since last camera update (in seconds).
@@ -645,6 +741,6 @@ private:
 	FVector GetActorLocation() const { return Super::GetActorLocation(); }
 
 public:
-	/** Returns TransformComponent subobject **/
+	/** @return Returns TransformComponent subobject */
 	class USceneComponent* GetTransformComponent() const;
 };
diff --git a/Engine/Source/Runtime/Engine/Classes/GameFramework/PlayerController.h b/Engine/Source/Runtime/Engine/Classes/GameFramework/PlayerController.h
--- a/Engine/Source/Runtime/Engine/Classes/GameFramework/PlayerController.h
+++ b/Engine/Source/Runtime/Engine/Classes/GameFramework/PlayerController.h
@@ -675,7 +675,7 @@ public:
 	 * Play Camera Shake 
 	 * @param Shake - Camera shake animation to play
 	 * @param Scale - Scalar defining how "intense" to play the anim
-	 * @param PlaySpace - Animation play area
+	 * @param PlaySpace - Which coordinate system to play the shake in (used for CameraAnims within the shake).
 	 * @param UserPlaySpaceRot - Matrix used when PlaySpace = CAPS_UserDefined
 	 */
 	UFUNCTION(unreliable, client, BlueprintCallable, Category="Game|Feedback")
diff --git a/Engine/Source/Runtime/Engine/Private/Camera/CameraModifier_CameraShake.cpp b/Engine/Source/Runtime/Engine/Private/Camera/CameraModifier_CameraShake.cpp
--- a/Engine/Source/Runtime/Engine/Private/Camera/CameraModifier_CameraShake.cpp
+++ b/Engine/Source/Runtime/Engine/Private/Camera/CameraModifier_CameraShake.cpp
@@ -218,7 +218,7 @@ bool UCameraModifier_CameraShake::ModifyCamera(class APlayerCameraManager* Camer
 	return false;
 }
 
-float UCameraModifier_CameraShake::InitializeOffset( const FFOscillator& Param )
+float UCameraModifier_CameraShake::InitializeOffset( const FFOscillator& Param ) const
 {
 	switch( Param.InitialOffset )
 	{
diff --git a/Engine/Source/Runtime/Engine/Private/Camera/CameraShake.cpp b/Engine/Source/Runtime/Engine/Private/Camera/CameraShake.cpp
--- a/Engine/Source/Runtime/Engine/Private/Camera/CameraShake.cpp
+++ b/Engine/Source/Runtime/Engine/Private/Camera/CameraShake.cpp
@@ -15,21 +15,3 @@ UCameraShake::UCameraShake(const FObjectInitializer& ObjectInitializer)
 	OscillationBlendInTime = 0.1f;
 	OscillationBlendOutTime = 0.2f;
 }
-
-float UCameraShake::GetRotOscillationMagnitude()
-{
-	FVector V;
-	V.X = RotOscillation.Pitch.Amplitude;
-	V.Y = RotOscillation.Yaw.Amplitude;
-	V.Z = RotOscillation.Roll.Amplitude;
-	return V.Size();
-}
-
-float UCameraShake::GetLocOscillationMagnitude()
-{
-	FVector V;
-	V.X = LocOscillation.X.Amplitude;
-	V.Y = LocOscillation.Y.Amplitude;
-	V.Z = LocOscillation.Z.Amplitude;
-	return V.Size();
-}
