Fix incorrect interpretation of return value.

diff --git a/Engine/Source/Runtime/Engine/Private/DataReplication.cpp b/Engine/Source/Runtime/Engine/Private/DataReplication.cpp
--- a/Engine/Source/Runtime/Engine/Private/DataReplication.cpp
+++ b/Engine/Source/Runtime/Engine/Private/DataReplication.cpp
@@ -26,7 +26,13 @@ public:
 			UScriptStruct::ICppStructOps* CppStructOps = Struct->GetCppStructOps();
 			check(CppStructOps); // else should not have STRUCT_NetSerializeNative
 			check(!Struct->InheritedCppStructOps()); // else should not have STRUCT_NetSerializeNative
-			if (!CppStructOps->NetSerialize(Ar, Map, bHasUnmapped, Data))
+			bool bSuccess = true;
+			if (!CppStructOps->NetSerialize(Ar, Map, bSuccess, Data))
+			{
+				bHasUnmapped = true;
+			}
+
+			if (!bSuccess)
 			{
 				UE_LOG(LogNet, Warning, TEXT("NetSerializeStruct: Native NetSerialize %s failed."), *Struct->GetFullName());
 			}
