Expose angular impulse

diff --git a/Engine/Source/Runtime/Engine/Classes/Components/PrimitiveComponent.h b/Engine/Source/Runtime/Engine/Classes/Components/PrimitiveComponent.h
--- a/Engine/Source/Runtime/Engine/Classes/Components/PrimitiveComponent.h
+++ b/Engine/Source/Runtime/Engine/Classes/Components/PrimitiveComponent.h
@@ -672,6 +672,16 @@ public:
 	UFUNCTION(BlueprintCallable, Category="Physics")
 	virtual void AddImpulse(FVector Impulse, FName BoneName = NAME_None, bool bVelChange = false);
 
+	/**
+	*	Add an angular impulse to a single rigid body. Good for one time instant burst.
+	*
+	*	@param	AngularImpulse	Magnitude and direction of impulse to apply. Direction is axis of rotation.
+	*	@param	BoneName	If a SkeletalMeshComponent, name of body to apply angular impulse to. 'None' indicates root body.
+	*	@param	bVelChange	If true, the Strength is taken as a change in angular velocity instead of an impulse (ie. mass will have no affect).
+	*/
+	UFUNCTION(BlueprintCallable, Category = "Physics")
+	virtual void AddAngularImpulse(FVector Impulse, FName BoneName = NAME_None, bool bVelChange = false);
+
 	/**
 	 *	Add an impulse to a single rigid body at a specific location. 
 	 *
diff --git a/Engine/Source/Runtime/Engine/Classes/PhysicsEngine/BodyInstance.h b/Engine/Source/Runtime/Engine/Classes/PhysicsEngine/BodyInstance.h
--- a/Engine/Source/Runtime/Engine/Classes/PhysicsEngine/BodyInstance.h
+++ b/Engine/Source/Runtime/Engine/Classes/PhysicsEngine/BodyInstance.h
@@ -458,6 +458,8 @@ public:
 	void AddForceAtPosition(const FVector& Force, const FVector& Position, bool bAllowSubstepping = true);
 	/** Add a torque to this body */
 	void AddTorque(const FVector& Torque, bool bAllowSubstepping = true);
+	/** Add a rotational impulse to this body */
+	void AddAngularImpulse(const FVector& Impulse, bool bVelChange);
 	/** Add an impulse to this body */
 	void AddImpulse(const FVector& Impulse, bool bVelChange);
 	/** Add an impulse to this body and a particular world position */
diff --git a/Engine/Source/Runtime/Engine/Private/PhysicsEngine/BodyInstance.cpp b/Engine/Source/Runtime/Engine/Private/PhysicsEngine/BodyInstance.cpp
--- a/Engine/Source/Runtime/Engine/Private/PhysicsEngine/BodyInstance.cpp
+++ b/Engine/Source/Runtime/Engine/Private/PhysicsEngine/BodyInstance.cpp
@@ -3155,6 +3155,29 @@ void FBodyInstance::AddTorque(const FVector& Torque, bool bAllowSubstepping)
 #endif
 }
 
+void FBodyInstance::AddAngularImpulse(const FVector& AngularImpulse, bool bVelChange)
+{
+#if WITH_PHYSX
+	PxRigidBody* PRigidBody = GetPxRigidBody();
+	if (IsRigidBodyNonKinematic(PRigidBody))
+	{
+		PxForceMode::Enum Mode = bVelChange ? PxForceMode::eVELOCITY_CHANGE : PxForceMode::eIMPULSE;
+		PRigidBody->addTorque(U2PVector(AngularImpulse), Mode, true);
+	}
+#endif // WITH_PHYSX
+
+#if WITH_BOX2D
+	if (BodyInstancePtr)
+	{
+		if (!AngularImpulse.IsNearlyZero())
+		{
+			const float AngularImpulse2D = FPhysicsIntegration2D::ConvertUnrealTorqueToBox(AngularImpulse);
+			BodyInstancePtr->ApplyAngularImpulse(AngularImpulse2D, /*wake=*/ true);
+		}
+	}
+#endif
+}
+
 void FBodyInstance::AddImpulse(const FVector& Impulse, bool bVelChange)
 {
 #if WITH_PHYSX
diff --git a/Engine/Source/Runtime/Engine/Private/PrimitiveComponentPhysics.cpp b/Engine/Source/Runtime/Engine/Private/PrimitiveComponentPhysics.cpp
--- a/Engine/Source/Runtime/Engine/Private/PrimitiveComponentPhysics.cpp
+++ b/Engine/Source/Runtime/Engine/Private/PrimitiveComponentPhysics.cpp
@@ -175,6 +175,14 @@ void UPrimitiveComponent::AddImpulse(FVector Impulse, FName BoneName, bool bVelC
 	}
 }
 
+void UPrimitiveComponent::AddAngularImpulse(FVector Impulse, FName BoneName, bool bVelChange)
+{
+	if (FBodyInstance* BI = GetBodyInstance(BoneName))
+	{
+		BI->AddAngularImpulse(Impulse, bVelChange);
+	}
+}
+
 void UPrimitiveComponent::AddImpulseAtLocation(FVector Impulse, FVector Location, FName BoneName)
 {
 	FBodyInstance* BI = GetBodyInstance(BoneName);
