#AnimLayerEditing: Fixed crash when create animation from pose that has missing joint

https://jira.ol.epicgames.net/browse/UE-6483?filter=12310

diff --git a/Engine/Source/Runtime/Engine/Private/Animation/AnimSequence.cpp b/Engine/Source/Runtime/Engine/Private/Animation/AnimSequence.cpp
--- a/Engine/Source/Runtime/Engine/Private/Animation/AnimSequence.cpp
+++ b/Engine/Source/Runtime/Engine/Private/Animation/AnimSequence.cpp
@@ -3257,11 +3257,39 @@ void UAnimSequence::RefreshTrackMapFromAnimTrackNames()
 	const FReferenceSkeleton& RefSkeleton = MySkeleton->GetReferenceSkeleton();
 	const int32 NumBones = AnimationTrackNames.Num();
 	TrackToSkeletonMapTable.AddUninitialized(NumBones);
-	for(int32 BoneIndex=0; BoneIndex<AnimationTrackNames.Num(); ++BoneIndex)
+
+	bool bNeedsFixing = false;
+	const int32 NumTracks = AnimationTrackNames.Num();
+	for(int32 BoneIndex=0; BoneIndex<NumTracks; ++BoneIndex)
 	{
 		int32 BoneTreeIndex = RefSkeleton.FindBoneIndex(AnimationTrackNames[BoneIndex]);
+
+		if (BoneTreeIndex == INDEX_NONE)
+		{
+			bNeedsFixing = true;
+		}
+
 		TrackToSkeletonMapTable[BoneIndex] = FTrackToSkeletonMap(BoneTreeIndex);
 	}
+
+	if(bNeedsFixing)
+	{
+		UE_LOG(LogAnimation, Warning, TEXT("RESAVE ANIMATION NEEDED(%s): Fixing track index."), *GetName());
+
+		const TArray<FBoneNode>& BoneTree = MySkeleton->GetBoneTree();
+		for(int32 I=NumTracks-1; I>=0; --I)
+		{
+			int32 BoneTreeIndex = MySkeleton->GetReferenceSkeleton().FindBoneIndex(AnimationTrackNames[I]);
+			if(BoneTreeIndex == INDEX_NONE)
+			{
+				RemoveTrack(I);
+			}
+			else
+			{
+				TrackToSkeletonMapTable[I].BoneTreeIndex = BoneTreeIndex;
+			}
+		}
+	}
 }
 
 bool UAnimSequence::CreateAnimation(USkeletalMesh * Mesh)
