#AnimLayerEditing : Fixed better for missng bone issue

diff --git a/Engine/Source/Runtime/Engine/Private/Animation/AnimSequence.cpp b/Engine/Source/Runtime/Engine/Private/Animation/AnimSequence.cpp
--- a/Engine/Source/Runtime/Engine/Private/Animation/AnimSequence.cpp
+++ b/Engine/Source/Runtime/Engine/Private/Animation/AnimSequence.cpp
@@ -3260,34 +3260,16 @@ void UAnimSequence::RefreshTrackMapFromAnimTrackNames()
 
 	bool bNeedsFixing = false;
 	const int32 NumTracks = AnimationTrackNames.Num();
-	for(int32 BoneIndex=0; BoneIndex<NumTracks; ++BoneIndex)
+	for(int32 I=NumTracks-1; I>=0; --I)
 	{
-		int32 BoneTreeIndex = RefSkeleton.FindBoneIndex(AnimationTrackNames[BoneIndex]);
-
-		if (BoneTreeIndex == INDEX_NONE)
+		int32 BoneTreeIndex = RefSkeleton.FindBoneIndex(AnimationTrackNames[I]);
+		if(BoneTreeIndex == INDEX_NONE)
 		{
-			bNeedsFixing = true;
+			RemoveTrack(I);
 		}
-
-		TrackToSkeletonMapTable[BoneIndex] = FTrackToSkeletonMap(BoneTreeIndex);
-	}
-
-	if(bNeedsFixing)
-	{
-		UE_LOG(LogAnimation, Warning, TEXT("RESAVE ANIMATION NEEDED(%s): Fixing track index."), *GetName());
-
-		const TArray<FBoneNode>& BoneTree = MySkeleton->GetBoneTree();
-		for(int32 I=NumTracks-1; I>=0; --I)
+		else
 		{
-			int32 BoneTreeIndex = MySkeleton->GetReferenceSkeleton().FindBoneIndex(AnimationTrackNames[I]);
-			if(BoneTreeIndex == INDEX_NONE)
-			{
-				RemoveTrack(I);
-			}
-			else
-			{
-				TrackToSkeletonMapTable[I].BoneTreeIndex = BoneTreeIndex;
-			}
+			TrackToSkeletonMapTable[I].BoneTreeIndex = BoneTreeIndex;
 		}
 	}
 }
