Fix issue where merging replication change lists was reaching out of bounds

diff --git a/Engine/Source/Runtime/Engine/Private/RepLayout.cpp b/Engine/Source/Runtime/Engine/Private/RepLayout.cpp
--- a/Engine/Source/Runtime/Engine/Private/RepLayout.cpp
+++ b/Engine/Source/Runtime/Engine/Private/RepLayout.cpp
@@ -262,24 +262,42 @@ public:
 
 	void ProcessDataArrayElements_r( TStackState& StackState, const FRepLayoutCmd& Cmd )
 	{
-		for ( int32 i = 0; i < StackState.DataArray->Num(); i++ )
+		const int32 NumDataArrayElements	= StackState.DataArray		? StackState.DataArray->Num()	: 0;
+		const int32 NumShadowArrayElements	= StackState.ShadowArray	? StackState.ShadowArray->Num() : 0;
+
+		// Loop using the number of elements in data array
+		for ( int32 i = 0; i < NumDataArrayElements; i++ )
 		{
 			const int32 ElementOffset = i * Cmd.ElementSize;
-			ProcessCmds_r( StackState, StackState.ShadowBaseData + ElementOffset, StackState.BaseData + ElementOffset );
+
+			uint8* Data			= StackState.BaseData + ElementOffset;
+			uint8* ShadowData	= i < NumShadowArrayElements ? ( StackState.ShadowBaseData + ElementOffset ) : NULL;	// ShadowArray might be smaller than DataArray
+
+			ProcessCmds_r( StackState, ShadowData, Data );
 		}
 	}
 
 	void ProcessShadowArrayElements_r( TStackState& StackState, const FRepLayoutCmd& Cmd )
 	{
-		for ( int32 i = 0; i < StackState.ShadowArray->Num(); i++ )
+		const int32 NumDataArrayElements	= StackState.DataArray		? StackState.DataArray->Num()	: 0;
+		const int32 NumShadowArrayElements	= StackState.ShadowArray	? StackState.ShadowArray->Num() : 0;
+
+		// Loop using the number of elements in shadow array
+		for ( int32 i = 0; i < NumShadowArrayElements; i++ )
 		{
 			const int32 ElementOffset = i * Cmd.ElementSize;
-			ProcessCmds_r( StackState, StackState.ShadowBaseData + ElementOffset, StackState.BaseData + ElementOffset );
+
+			uint8* Data			= i < NumDataArrayElements ? ( StackState.BaseData + ElementOffset ) : NULL;	// DataArray might be smaller than ShadowArray
+			uint8* ShadowData	= StackState.ShadowBaseData + ElementOffset;
+
+			ProcessCmds_r( StackState, ShadowData, Data );
 		}
 	}
 
 	void ProcessArrayCmd_r( TStackState & PrevStackState, const FRepLayoutCmd& Cmd, const int32 CmdIndex, uint8* RESTRICT ShadowData, uint8* RESTRICT Data )
 	{
+		check( ShadowData != NULL || Data != NULL );
+
 		FScriptArray* ShadowArray	= (FScriptArray*)ShadowData;
 		FScriptArray* DataArray		= (FScriptArray*)Data;
 
@@ -290,6 +308,8 @@ public:
 
 	void ProcessCmds_r( TStackState& StackState, uint8* RESTRICT ShadowData, uint8* RESTRICT Data )
 	{
+		check( ShadowData != NULL || Data != NULL );
+
 		for ( int32 CmdIndex = StackState.CmdStart; CmdIndex < StackState.CmdEnd; CmdIndex++ )
 		{
 			const FRepLayoutCmd& Cmd = Cmds[ CmdIndex ];
@@ -300,7 +320,7 @@ public:
 			{
 				if ( static_cast< TImpl* >( this )->ShouldProcessNextCmd() )
 				{
-					ProcessArrayCmd_r( StackState, Cmd, CmdIndex, ShadowData + Cmd.Offset, Data + Cmd.Offset );
+					ProcessArrayCmd_r( StackState, Cmd, CmdIndex, ShadowData ? ( ShadowData + Cmd.Offset ) : NULL, Data ? ( Data + Cmd.Offset ) : NULL );
 				}
 				CmdIndex = Cmd.EndCmd - 1;	// Jump past children of this array (-1 for ++ in for loop)
 			}
