Fix to populating NavigationSystem::SupportedAgents with navigation data types from outside of engine code #UE4

diff --git a/Engine/Source/Runtime/Engine/Private/AI/Navigation/NavigationSystem.cpp b/Engine/Source/Runtime/Engine/Private/AI/Navigation/NavigationSystem.cpp
--- a/Engine/Source/Runtime/Engine/Private/AI/Navigation/NavigationSystem.cpp
+++ b/Engine/Source/Runtime/Engine/Private/AI/Navigation/NavigationSystem.cpp
@@ -344,15 +344,13 @@ void UNavigationSystem::PostInitProperties()
 		FCoreUObjectDelegates::PostLoadMap.AddUObject(this, &UNavigationSystem::OnPostLoadMap);
 		UNavigationSystem::NavigationDirtyEvent.AddUObject(this, &UNavigationSystem::OnNavigationDirtied);
 	}
-	else
+	
+	// update SupportedActors' navigation classes
+	for (FNavDataConfig& SupportedAgentConfig : SupportedAgents)
 	{
-		// update SupportedActors' navigation classes
-		for (FNavDataConfig& SupportedAgentConfig : SupportedAgents)
+		if (SupportedAgentConfig.NavigationDataClassName.IsValid())
 		{
-			if (SupportedAgentConfig.NavigationDataClassName.IsValid())
-			{
-				SupportedAgentConfig.NavigationDataClass = LoadClass<ANavigationData>(NULL, *SupportedAgentConfig.NavigationDataClassName.ToString(), NULL, LOAD_None, NULL);
-			}
+			SupportedAgentConfig.NavigationDataClass = LoadClass<ANavigationData>(NULL, *SupportedAgentConfig.NavigationDataClassName.ToString(), NULL, LOAD_None, NULL);
 		}
 	}
 }
