Made recast navmesh generator expandable by inherited classes #UE4

diff --git a/Engine/Source/Runtime/Engine/Private/AI/Navigation/RecastNavMesh.cpp b/Engine/Source/Runtime/Engine/Private/AI/Navigation/RecastNavMesh.cpp
--- a/Engine/Source/Runtime/Engine/Private/AI/Navigation/RecastNavMesh.cpp
+++ b/Engine/Source/Runtime/Engine/Private/AI/Navigation/RecastNavMesh.cpp
@@ -1754,7 +1754,7 @@ void ARecastNavMesh::ConstructGenerator()
 	NavDataGenerator.Reset();
 	if (SupportsRuntimeGeneration())
 	{
-		NavDataGenerator.Reset(new FRecastNavMeshGenerator(this));
+		NavDataGenerator.Reset(new FRecastNavMeshGenerator(*this));
 	}
 }
 
diff --git a/Engine/Source/Runtime/Engine/Private/AI/Navigation/RecastNavMeshGenerator.cpp b/Engine/Source/Runtime/Engine/Private/AI/Navigation/RecastNavMeshGenerator.cpp
--- a/Engine/Source/Runtime/Engine/Private/AI/Navigation/RecastNavMeshGenerator.cpp
+++ b/Engine/Source/Runtime/Engine/Private/AI/Navigation/RecastNavMeshGenerator.cpp
@@ -1431,33 +1431,32 @@ static FBox CalculateTileBounds(int32 X, int32 Y, const FVector& NavMeshOrigin,
 // FRecastTileGenerator
 //----------------------------------------------------------------------//
 
-FRecastTileGenerator::FRecastTileGenerator(
-	FRecastNavMeshGenerator* ParentGenerator,
-	const int32 X, 
-	const int32 Y,
-	TArray<FBox> DirtyAreas
-	)
+FRecastTileGenerator::FRecastTileGenerator(const FRecastNavMeshGenerator& ParentGenerator, const FIntPoint& Location)
 {
 	bSucceeded = false;
-	
-	TileX = X;
-	TileY = Y;
-		
-	TileConfig = ParentGenerator->GetConfig();
-	Version = ParentGenerator->GetVersion();
-	AdditionalCachedData = ParentGenerator->GetAdditionalCachedData();
+
+	TileX = Location.X;
+	TileY = Location.Y;
+
+	TileConfig = ParentGenerator.GetConfig();
+	Version = ParentGenerator.GetVersion();
+	AdditionalCachedData = ParentGenerator.GetAdditionalCachedData();
+}
+
+void FRecastTileGenerator::Setup(const FRecastNavMeshGenerator& ParentGenerator, const TArray<FBox>& DirtyAreas)
+{
 	const FVector NavMeshOrigin = FVector::ZeroVector;
-	const FBox NavTotalBounds = ParentGenerator->GetTotalBounds();
+	const FBox NavTotalBounds = ParentGenerator.GetTotalBounds();
 	const float TileCellSize = (TileConfig.tileSize * TileConfig.cs);
 
-	TileBB = CalculateTileBounds(X, Y, NavMeshOrigin, NavTotalBounds, TileCellSize);
+	TileBB = CalculateTileBounds(TileX, TileY, NavMeshOrigin, NavTotalBounds, TileCellSize);
 	const FBox RCBox = Unreal2RecastBox(TileBB);
 	rcVcopy(TileConfig.bmin, &RCBox.Min.X);
 	rcVcopy(TileConfig.bmax, &RCBox.Max.X);
 			
 	// from passed in boxes pick the ones overlapping with tile bounds
 	bFullyEncapsulatedByInclusionBounds = true;
-	const auto& ParentBounds = ParentGenerator->GetInclusionBounds();
+	const auto& ParentBounds = ParentGenerator.GetInclusionBounds();
 	if (ParentBounds.Num() > 0)
 	{
 		bFullyEncapsulatedByInclusionBounds = false;
@@ -1473,7 +1472,7 @@ FRecastTileGenerator::FRecastTileGenerator(
 	}
 
 	// Take ownership of tile cache data if it exist 
-	CompressedLayers = ParentGenerator->TakeIntermediateLayersData(FIntPoint(TileX, TileY));
+	CompressedLayers = ParentGenerator.TakeIntermediateLayersData(FIntPoint(TileX, TileY));
 
 	// We have to regenerate layers data in case geometry is changed or tile cache is missing
 	bRegenerateCompressedLayers = (DirtyAreas.Num() == 0 || CompressedLayers.Num() == 0);
@@ -1518,18 +1517,13 @@ void FRecastTileGenerator::DoWork()
 	bSucceeded = GenerateTile();
 }
 
-const TCHAR* FRecastTileGenerator::Name()
-{
-	return TEXT("FRecastTileGenerator");
-}
-
-void FRecastTileGenerator::GatherGeometry(const FRecastNavMeshGenerator* ParentGenerator, bool bGeometryChanged)
+void FRecastTileGenerator::GatherGeometry(const FRecastNavMeshGenerator& ParentGenerator, bool bGeometryChanged)
 {
-	const UNavigationSystem*	NavSys = UNavigationSystem::GetCurrent(ParentGenerator->GetWorld());
+	const UNavigationSystem*	NavSys = UNavigationSystem::GetCurrent(ParentGenerator.GetWorld());
 	const FNavigationOctree*	NavOctree = NavSys ? NavSys->GetNavOctree() : nullptr;
-	const FNavDataConfig*		NavDataConfig = &ParentGenerator->GetOwner()->NavDataConfig;
-	
-	for(FNavigationOctree::TConstElementBoxIterator<FNavigationOctree::DefaultStackAllocator> It(*NavOctree, ParentGenerator->GrowBoundingBox(TileBB, /*bIncludeAgentHeight*/ false));
+	const FNavDataConfig*		NavDataConfig = &ParentGenerator.GetOwner()->NavDataConfig;
+
+	for (FNavigationOctree::TConstElementBoxIterator<FNavigationOctree::DefaultStackAllocator> It(*NavOctree, ParentGenerator.GrowBoundingBox(TileBB, /*bIncludeAgentHeight*/ false));
 		It.HasPendingElements();
 		It.Advance())
 	{
@@ -2672,17 +2666,16 @@ static int32 CaclulateMaxTilesCount(const TNavStatArray<FBox>& NavigableAreas, f
 // FRecastNavMeshGenerator
 //----------------------------------------------------------------------//
 
-FRecastNavMeshGenerator::FRecastNavMeshGenerator(ARecastNavMesh* InDestNavMesh)
+FRecastNavMeshGenerator::FRecastNavMeshGenerator(ARecastNavMesh& InDestNavMesh)
 	: NumActiveTiles(0),
 	  MaxTileGeneratorTasks(1),
 	  AvgLayersPerTile(8.0f),
-	  DestNavMesh(InDestNavMesh),
+	  DestNavMesh(&InDestNavMesh),
 	  bInitialized(false),
 	  Version(0)
 {
 	INC_DWORD_STAT_BY(STAT_NavigationMemory, sizeof(*this));
 
-	check(InDestNavMesh);
 	Init();
 
 	// recreate navmesh if no data was loaded, or when loaded data doesn't match current grid layout
@@ -2708,7 +2701,7 @@ FRecastNavMeshGenerator::FRecastNavMeshGenerator(ARecastNavMesh* InDestNavMesh)
 	if (bRecreateNavmesh)
 	{
 		ConstructTiledNavMesh();
-		InDestNavMesh->MarkAsNeedingUpdate();
+		InDestNavMesh.MarkAsNeedingUpdate();
 	}
 	else
 	{
@@ -3210,7 +3203,7 @@ FBox FRecastNavMeshGenerator::GrowBoundingBox(const FBox& BBox, bool bIncludeAge
 	return FBox(BBox.Min - BBoxGrowOffsetBoth - BBoxGrowOffsetMin, BBox.Max + BBoxGrowOffsetBoth);
 }
 
-TArray<FNavMeshTileData> FRecastNavMeshGenerator::TakeIntermediateLayersData(FIntPoint GridCoord)
+TArray<FNavMeshTileData> FRecastNavMeshGenerator::TakeIntermediateLayersData(FIntPoint GridCoord) const
 {
 	TArray<FNavMeshTileData> Result;
 	IntermediateLayerDataMap.RemoveAndCopyValue(GridCoord, Result);
@@ -3428,6 +3421,13 @@ void FRecastNavMeshGenerator::SortPendingBuildTiles()
 	}
 }
 
+TSharedRef<FRecastTileGenerator> FRecastNavMeshGenerator::CreateTileGenerator(const FIntPoint& Coord, const TArray<FBox>& DirtyAreas)
+{
+	TSharedRef<FRecastTileGenerator> TileGenerator = MakeShareable(new FRecastTileGenerator(*this, Coord));
+	TileGenerator->Setup(*this, DirtyAreas);
+	return TileGenerator;
+}
+
 TArray<uint32> FRecastNavMeshGenerator::ProcessTileTasks(const int32 NumTasksToSubmit)
 {
 	QUICK_SCOPE_CYCLE_COUNTER(STAT_RecastNavMeshGenerator_ProcessTileTasks);
@@ -3446,15 +3446,10 @@ TArray<uint32> FRecastNavMeshGenerator::ProcessTileTasks(const int32 NumTasksToS
 		if (!RunningDirtyTiles.Contains(RunningElement))
 		{
 			// Spawn async task
-			TUniquePtr<FRecastTileGeneratorTask> TileTask = MakeUnique<FRecastTileGeneratorTask>(
-				this,
-				PendingElement.Coord.X,
-				PendingElement.Coord.Y,
-				PendingElement.DirtyAreas
-				);
+			TUniquePtr<FRecastTileGeneratorTask> TileTask = MakeUnique<FRecastTileGeneratorTask>(CreateTileGenerator(PendingElement.Coord, PendingElement.DirtyAreas));
 
 			// Start it in background in case it has something to build
-			if (TileTask->GetTask().HasDataToBuild())
+			if (TileTask->GetTask().TileGenerator->HasDataToBuild())
 			{
 				RunningElement.AsyncTask = TileTask.Release();
 				RunningElement.AsyncTask->StartBackgroundTask();
@@ -3495,7 +3490,7 @@ TArray<uint32> FRecastNavMeshGenerator::ProcessTileTasks(const int32 NumTasksToS
 			// Add generated tiles to navmesh
 			if (!Element.bShouldDiscard)
 			{
-				const FRecastTileGenerator& TileGenerator = Element.AsyncTask->GetTask();
+				const FRecastTileGenerator& TileGenerator = *(Element.AsyncTask->GetTask().TileGenerator);
 				TArray<uint32> UpdatedTileIndices = AddGeneratedTiles(TileGenerator);
 				UpdatedTiles.Append(UpdatedTileIndices);
 			
@@ -3631,7 +3626,7 @@ uint32 FRecastNavMeshGenerator::LogMemUsed() const
 	uint32 GeneratorsMem = 0;
 	for (const FRunningTileElement& Element : RunningDirtyTiles)
 	{
-		GeneratorsMem+= Element.AsyncTask->GetTask().UsedMemoryOnStartup; 
+		GeneratorsMem += Element.AsyncTask->GetTask().TileGenerator->UsedMemoryOnStartup;
 	}
 
 	UE_LOG(LogNavigation, Display, TEXT("    FRecastNavMeshGenerator: Total Generator\'s size %u, count %d"), GeneratorsMem, RunningDirtyTiles.Num());
diff --git a/Engine/Source/Runtime/Engine/Public/AI/Navigation/RecastNavMeshGenerator.h b/Engine/Source/Runtime/Engine/Public/AI/Navigation/RecastNavMeshGenerator.h
--- a/Engine/Source/Runtime/Engine/Public/AI/Navigation/RecastNavMeshGenerator.h
+++ b/Engine/Source/Runtime/Engine/Public/AI/Navigation/RecastNavMeshGenerator.h
@@ -10,6 +10,7 @@
 
 #define MAX_VERTS_PER_POLY	6
 
+class ARecastNavMesh;
 class FNavigationOctree;
 class FNavMeshBuildContext;
 class FRecastNavMeshGenerator;
@@ -130,16 +131,12 @@ struct FRecastAreaNavModifierElement
  */
 class ENGINE_API FRecastTileGenerator : public FNonAbandonableTask
 {
+	friend FRecastNavMeshGenerator;
+
 public:
-	FRecastTileGenerator(
-		FRecastNavMeshGenerator* ParentGenerator,
-		const int32 X, 
-		const int32 Y,
-		TArray<FBox> DirtyAreas
-	);
+	FRecastTileGenerator(const FRecastNavMeshGenerator& ParentGenerator, const FIntPoint& Location);
 		
 	void DoWork();
-	static const TCHAR *Name();
 
 	FORCEINLINE int32 GetTileX() const { return TileX; }
 	FORCEINLINE int32 GetTileY() const { return TileY; }
@@ -158,18 +155,20 @@ public:
 	// Memory amount used to construct generator 
 	uint32 UsedMemoryOnStartup;
 	
-private:
+protected:
 	/** Does the actual tile generations. 
 	 *	@note always trigger tile generation only via TriggerAsyncBuild. This is a worker function
 	 *	@return true if new tile navigation data has been generated and is ready to be added to navmesh instance, 
 	 *	false if failed or no need to generate (still valid)
 	 */
 	bool GenerateTile();
+
+	void Setup(const FRecastNavMeshGenerator& ParentGenerator, const TArray<FBox>& DirtyAreas);
 	
-	void GatherGeometry(const FRecastNavMeshGenerator* ParentGenerator, bool bGeometryChanged);
+	void GatherGeometry(const FRecastNavMeshGenerator& ParentGenerator, bool bGeometryChanged);
 
 	/** builds CompressedLayers array (geometry + modifiers) */
-	bool GenerateCompressedLayers(FNavMeshBuildContext& BuildContext);
+	virtual bool GenerateCompressedLayers(FNavMeshBuildContext& BuildContext);
 
 	/** builds NavigationData array (layers + obstacles) */
 	bool GenerateNavigationData(FNavMeshBuildContext& BuildContext);
@@ -225,7 +224,28 @@ protected:
 	TArray<FSimpleLinkNavModifier> OffmeshLinks;
 };
 
-typedef FAsyncTask<FRecastTileGenerator> FRecastTileGeneratorTask;
+struct ENGINE_API FRecastTileGeneratorWrapper : public FNonAbandonableTask
+{
+	TSharedRef<FRecastTileGenerator> TileGenerator;
+
+	FRecastTileGeneratorWrapper(TSharedRef<FRecastTileGenerator> InTileGenerator)
+		: TileGenerator(InTileGenerator)
+	{
+	}
+	
+	void DoWork()
+	{
+		TileGenerator->DoWork();
+	}
+
+	static const TCHAR *Name()
+	{
+		return TEXT("FRecastTileGenerator");
+	}
+};
+
+typedef FAsyncTask<FRecastTileGeneratorWrapper> FRecastTileGeneratorTask;
+//typedef FAsyncTask<FRecastTileGenerator> FRecastTileGeneratorTask;
 
 struct FPendingTileElement
 {
@@ -305,10 +325,10 @@ struct FTileTimestamp
 /**
  * Class that handles generation of the whole Recast-based navmesh.
  */
-class FRecastNavMeshGenerator : public FNavDataGenerator
+class ENGINE_API FRecastNavMeshGenerator : public FNavDataGenerator
 {
 public:
-	FRecastNavMeshGenerator(class ARecastNavMesh* InDestNavMesh);
+	FRecastNavMeshGenerator(ARecastNavMesh& InDestNavMesh);
 	virtual ~FRecastNavMeshGenerator();
 
 private:
@@ -361,8 +381,8 @@ public:
 	FBox GrowBoundingBox(const FBox& BBox, bool bIncludeAgentHeight) const;
 
 	/** Transfers ownership if tile cache data to the caller */
-	TArray<FNavMeshTileData> TakeIntermediateLayersData(FIntPoint GridCoord);
-
+	TArray<FNavMeshTileData> TakeIntermediateLayersData(FIntPoint GridCoord) const;
+	
 #if !(UE_BUILD_SHIPPING || UE_BUILD_TEST)
 	virtual void ExportNavigationData(const FString& FileName) const;
 #endif
@@ -375,8 +395,8 @@ public:
 
 	static void ExportRigidBodyGeometry(UBodySetup& BodySetup, TNavStatArray<FVector>& OutVertexBuffer, TNavStatArray<int32>& OutIndexBuffer, const FTransform& LocalToWorld = FTransform::Identity);
 	static void ExportRigidBodyGeometry(UBodySetup& BodySetup, TNavStatArray<FVector>& OutTriMeshVertexBuffer, TNavStatArray<int32>& OutTriMeshIndexBuffer, TNavStatArray<FVector>& OutConvexVertexBuffer, TNavStatArray<int32>& OutConvexIndexBuffer, TNavStatArray<int32>& OutShapeBuffer, const FTransform& LocalToWorld = FTransform::Identity);
-
-private:
+	
+protected:
 	// Performs initial setup of member variables so that generator is ready to do its thing from this point on
 	void Init();
 
@@ -407,6 +427,8 @@ private:
 	/** Blocks until build for specified list of tiles is complete and discard results */
 	void DiscardCurrentBuildingTasks();
 
+	virtual TSharedRef<FRecastTileGenerator> CreateTileGenerator(const FIntPoint& Coord, const TArray<FBox>& DirtyAreas);
+
 	//----------------------------------------------------------------------//
 	// debug
 	//----------------------------------------------------------------------//
@@ -445,7 +467,7 @@ private:
 	FRecastNavMeshCachedData AdditionalCachedData;
 
 	/** Compressed layers data, can be reused for tiles generation */
-	TMap<FIntPoint, TArray<FNavMeshTileData>> IntermediateLayerDataMap;
+	mutable TMap<FIntPoint, TArray<FNavMeshTileData>> IntermediateLayerDataMap;
 
 	/** */
 	TMapBase<const AActor*, FBox, false> ActorToAreaMap;
