Added RecastNavMesh::DistanceToWall, a wrapper for dtNavMeshQuery::findDistanceToWall #UE4

- note that it's not available directly from NavigationSystem yet, you need to grab a NavigationData from NavigationSystem and cast it to ARecastNavMesh

diff --git a/Engine/Source/Runtime/Engine/Classes/AI/Navigation/RecastNavMesh.h b/Engine/Source/Runtime/Engine/Classes/AI/Navigation/RecastNavMesh.h
--- a/Engine/Source/Runtime/Engine/Classes/AI/Navigation/RecastNavMesh.h
+++ b/Engine/Source/Runtime/Engine/Classes/AI/Navigation/RecastNavMesh.h
@@ -786,6 +786,9 @@ public:
 	/** Returns nearest navmesh polygon to Loc, or INVALID_NAVMESHREF if Loc is not on the navmesh. */
 	NavNodeRef FindNearestPoly(FVector const& Loc, FVector const& Extent, TSharedPtr<const FNavigationQueryFilter> Filter = NULL, const UObject* Querier = NULL) const;
 
+	/** Finds the distance to the closest wall, limited to MaxDistance */
+	float FindDistanceToWall(const FVector& StartLoc, TSharedPtr<const FNavigationQueryFilter> Filter = nullptr, float MaxDistance = FLT_MAX) const;
+
 	/** Retrieves center of the specified polygon. Returns false on error. */
 	bool GetPolyCenter(NavNodeRef PolyID, FVector& OutCenter) const;
 
diff --git a/Engine/Source/Runtime/Engine/Private/AI/Navigation/RecastNavMesh.cpp b/Engine/Source/Runtime/Engine/Private/AI/Navigation/RecastNavMesh.cpp
--- a/Engine/Source/Runtime/Engine/Private/AI/Navigation/RecastNavMesh.cpp
+++ b/Engine/Source/Runtime/Engine/Private/AI/Navigation/RecastNavMesh.cpp
@@ -998,6 +998,48 @@ NavNodeRef ARecastNavMesh::FindNearestPoly(FVector const& Loc, FVector const& Ex
 	return PolyRef;
 }
 
+float ARecastNavMesh::FindDistanceToWall(const FVector& StartLoc, TSharedPtr<const FNavigationQueryFilter> Filter, float MaxDistance) const
+{
+	if (HasValidNavmesh() == false)
+	{
+		return 0.f;
+	}
+
+	const FNavigationQueryFilter& FilterToUse = GetRightFilterRef(Filter);
+
+	INITIALIZE_NAVQUERY(NavQuery, FilterToUse.GetMaxSearchNodes());
+	const dtQueryFilter* QueryFilter = ((const FRecastQueryFilter*)(FilterToUse.GetImplementation()))->GetAsDetourQueryFilter();
+
+	if (QueryFilter == nullptr)
+	{
+		UE_VLOG(this, LogNavigation, Warning, TEXT("ARecastNavMesh::FindDistanceToWall failing due to QueryFilter == NULL"));
+		return 0.f;
+	}
+
+	const FVector& NavExtent = GetDefaultQueryExtent();
+	const float Extent[3] = { NavExtent.X, NavExtent.Z, NavExtent.Y };
+
+	const FVector RecastStart = Unreal2RecastPoint(StartLoc);
+
+	NavNodeRef StartNode = INVALID_NAVNODEREF;
+	NavQuery.findNearestPoly(&RecastStart.X, Extent, QueryFilter, &StartNode, NULL);
+
+	if (StartNode != INVALID_NAVNODEREF)
+	{
+		float TmpHitPos[3], TmpHitNormal[3];
+		float DistanceToWall = 0.f;
+		const dtStatus RaycastStatus = NavQuery.findDistanceToWall(StartNode, &RecastStart.X, MaxDistance, QueryFilter
+			, &DistanceToWall, TmpHitPos, TmpHitNormal);
+
+		if (dtStatusSucceed(RaycastStatus))
+		{
+			return DistanceToWall;
+		}
+	}
+
+	return 0.f;
+}
+
 void ARecastNavMesh::UpdateCustomLink(const INavLinkCustomInterface* CustomLink)
 {
 	TSubclassOf<UNavArea> AreaClass = CustomLink->GetLinkAreaClass();
