Back out revision 248 from //depot/UE4/Engine/Source/Runtime/Engine/Private/SkeletalMeshComponent.cpp

diff --git a/Engine/Source/Runtime/Engine/Private/SkeletalMeshComponent.cpp b/Engine/Source/Runtime/Engine/Private/SkeletalMeshComponent.cpp
--- a/Engine/Source/Runtime/Engine/Private/SkeletalMeshComponent.cpp
+++ b/Engine/Source/Runtime/Engine/Private/SkeletalMeshComponent.cpp
@@ -232,7 +232,6 @@ void USkeletalMeshComponent::InitAnim(bool bForceReinit)
         TickAnimation(0.f); 
 
 		RefreshBoneTransforms();
-		FlipEditableSpaceBases();
 		UpdateComponentToWorld();
 	}
 }
