fixed UE-6405
SceneCapture2D incorrect for section 1.7 in ContentExamples Reflection map

diff --git a/Engine/Source/Runtime/Engine/Private/SceneView.cpp b/Engine/Source/Runtime/Engine/Private/SceneView.cpp
--- a/Engine/Source/Runtime/Engine/Private/SceneView.cpp
+++ b/Engine/Source/Runtime/Engine/Private/SceneView.cpp
@@ -1121,7 +1121,7 @@ void FSceneView::EndFinalPostprocessSettings(const FSceneViewInitOptions& ViewIn
 		}
 	}
 
-	if(!Family->EngineShowFlags.ScreenPercentage)
+	if(!Family->EngineShowFlags.ScreenPercentage || bIsSceneCapture || bIsReflectionCapture)
 	{
 		FinalPostProcessSettings.ScreenPercentage = 100;
 	}
