Fix build error on Clang platforms.

diff --git a/Engine/Source/Runtime/Engine/Private/AI/Navigation/RecastNavMeshGenerator.cpp b/Engine/Source/Runtime/Engine/Private/AI/Navigation/RecastNavMeshGenerator.cpp
--- a/Engine/Source/Runtime/Engine/Private/AI/Navigation/RecastNavMeshGenerator.cpp
+++ b/Engine/Source/Runtime/Engine/Private/AI/Navigation/RecastNavMeshGenerator.cpp
@@ -1443,6 +1443,10 @@ FRecastTileGenerator::FRecastTileGenerator(const FRecastNavMeshGenerator& Parent
 	AdditionalCachedData = ParentGenerator.GetAdditionalCachedData();
 }
 
+FRecastTileGenerator::~FRecastTileGenerator()
+{
+}
+
 void FRecastTileGenerator::Setup(const FRecastNavMeshGenerator& ParentGenerator, const TArray<FBox>& DirtyAreas)
 {
 	const FVector NavMeshOrigin = FVector::ZeroVector;
diff --git a/Engine/Source/Runtime/Engine/Public/AI/Navigation/RecastNavMeshGenerator.h b/Engine/Source/Runtime/Engine/Public/AI/Navigation/RecastNavMeshGenerator.h
--- a/Engine/Source/Runtime/Engine/Public/AI/Navigation/RecastNavMeshGenerator.h
+++ b/Engine/Source/Runtime/Engine/Public/AI/Navigation/RecastNavMeshGenerator.h
@@ -135,6 +135,7 @@ class ENGINE_API FRecastTileGenerator : public FNonAbandonableTask
 
 public:
 	FRecastTileGenerator(const FRecastNavMeshGenerator& ParentGenerator, const FIntPoint& Location);
+	virtual ~FRecastTileGenerator();
 		
 	void DoWork();
 
