#ue4 - Add stat counter for EndScopedMovementUpdate.

diff --git a/Engine/Source/Runtime/Engine/Private/LevelTick.cpp b/Engine/Source/Runtime/Engine/Private/LevelTick.cpp
--- a/Engine/Source/Runtime/Engine/Private/LevelTick.cpp
+++ b/Engine/Source/Runtime/Engine/Private/LevelTick.cpp
@@ -51,6 +51,7 @@ DEFINE_STAT(STAT_TeleportToTime);
 DEFINE_STAT(STAT_MoveComponentTime);
 DEFINE_STAT(STAT_UpdateOverlaps);
 DEFINE_STAT(STAT_UpdatePhysicsVolume);
+DEFINE_STAT(STAT_EndScopedMovementUpdate);
 
 DEFINE_STAT(STAT_PostTickComponentLW);
 DEFINE_STAT(STAT_PostTickComponentRecreate);
diff --git a/Engine/Source/Runtime/Engine/Private/SceneComponent.cpp b/Engine/Source/Runtime/Engine/Private/SceneComponent.cpp
--- a/Engine/Source/Runtime/Engine/Private/SceneComponent.cpp
+++ b/Engine/Source/Runtime/Engine/Private/SceneComponent.cpp
@@ -193,6 +193,7 @@ void USceneComponent::BeginScopedMovementUpdate(class FScopedMovementUpdate& Sco
 
 void USceneComponent::EndScopedMovementUpdate(class FScopedMovementUpdate& CompletedScope)
 {
+	SCOPE_CYCLE_COUNTER(STAT_EndScopedMovementUpdate);
 	checkSlow(IsInGameThread());
 
 	// Special case when shutting down
diff --git a/Engine/Source/Runtime/Engine/Public/EngineStats.h b/Engine/Source/Runtime/Engine/Public/EngineStats.h
--- a/Engine/Source/Runtime/Engine/Public/EngineStats.h
+++ b/Engine/Source/Runtime/Engine/Public/EngineStats.h
@@ -38,6 +38,7 @@ DECLARE_CYCLE_STAT_EXTERN(TEXT("Spawn Actor Time"),STAT_SpawnActorTime,STATGROUP
 DECLARE_CYCLE_STAT_EXTERN(TEXT("MoveComponent Time"),STAT_MoveComponentTime,STATGROUP_Game, );
 DECLARE_CYCLE_STAT_EXTERN(TEXT("UpdateOverlaps Time"),STAT_UpdateOverlaps,STATGROUP_Game, );
 DECLARE_CYCLE_STAT_EXTERN(TEXT("UpdatePhysicsVolume Time"),STAT_UpdatePhysicsVolume,STATGROUP_Game, );
+DECLARE_CYCLE_STAT_EXTERN(TEXT("EndScopedMovementUpdate Time"),STAT_EndScopedMovementUpdate,STATGROUP_Game, );
 DECLARE_CYCLE_STAT_EXTERN(TEXT("TeleportTo Time"),STAT_TeleportToTime,STATGROUP_Game, );
 DECLARE_CYCLE_STAT_EXTERN(TEXT("GC Mark Time"),STAT_GCMarkTime,STATGROUP_Game, );
 DECLARE_CYCLE_STAT_EXTERN(TEXT("GC Sweep Time"),STAT_GCSweepTime,STATGROUP_Game, );
