#ue4 - Optimize PrimitiveComponent navigation relevancy updates to check the enabled flag to avoid the internal casts in naviagation updates.

diff --git a/Engine/Source/Runtime/Engine/Classes/Components/PrimitiveComponent.h b/Engine/Source/Runtime/Engine/Classes/Components/PrimitiveComponent.h
--- a/Engine/Source/Runtime/Engine/Classes/Components/PrimitiveComponent.h
+++ b/Engine/Source/Runtime/Engine/Classes/Components/PrimitiveComponent.h
@@ -396,6 +396,8 @@ public:
 
 protected:
 
+	virtual void UpdateNavigationData() override;
+
 	/** Returns true if all descendant components that we can possibly collide with use relative location and rotation. */
 	virtual bool AreAllCollideableDescendantsRelative(bool bAllowCachedValue = true) const;
 
diff --git a/Engine/Source/Runtime/Engine/Classes/Components/SceneComponent.h b/Engine/Source/Runtime/Engine/Classes/Components/SceneComponent.h
--- a/Engine/Source/Runtime/Engine/Classes/Components/SceneComponent.h
+++ b/Engine/Source/Runtime/Engine/Classes/Components/SceneComponent.h
@@ -779,7 +779,7 @@ protected:
 	virtual void OnChildDetached(USceneComponent* ChildComponent) {}
 
 	/** Called after changing transform, tries to update navigation octree */
-	void UpdateNavigationData();
+	virtual void UpdateNavigationData();
 };
 
 
diff --git a/Engine/Source/Runtime/Engine/Private/PrimitiveComponent.cpp b/Engine/Source/Runtime/Engine/Private/PrimitiveComponent.cpp
--- a/Engine/Source/Runtime/Engine/Private/PrimitiveComponent.cpp
+++ b/Engine/Source/Runtime/Engine/Private/PrimitiveComponent.cpp
@@ -2063,6 +2063,15 @@ void UPrimitiveComponent::GetOverlappingComponents(TArray<UPrimitiveComponent*>&
 }
 
 
+void UPrimitiveComponent::UpdateNavigationData()
+{
+	if (bNavigationRelevant)
+	{
+		Super::UpdateNavigationData();
+	}
+}
+
+
 bool UPrimitiveComponent::AreAllCollideableDescendantsRelative(bool bAllowCachedValue) const
 {
 	UPrimitiveComponent* MutableThis = const_cast<UPrimitiveComponent*>(this);
