Improved canvas string wrapping by using a cached pair of break iterators.

diff --git a/Engine/Source/Editor/UnrealEd/Private/Tests/EditorAssetAutomationTests.cpp b/Engine/Source/Editor/UnrealEd/Private/Tests/EditorAssetAutomationTests.cpp
--- a/Engine/Source/Editor/UnrealEd/Private/Tests/EditorAssetAutomationTests.cpp
+++ b/Engine/Source/Editor/UnrealEd/Private/Tests/EditorAssetAutomationTests.cpp
@@ -21,7 +21,6 @@
 #include "Matinee/InterpData.h"
 #include "Styling/ButtonWidgetStyle.h"
 
-#include "SlateWordWrapper.h"
 #include "AutomationCommon.h"
 #include "AutomationEditorCommon.h"
 #include "Engine/Selection.h"
diff --git a/Engine/Source/Editor/UnrealEd/Private/Tests/EditorAutomationTests.cpp b/Engine/Source/Editor/UnrealEd/Private/Tests/EditorAutomationTests.cpp
--- a/Engine/Source/Editor/UnrealEd/Private/Tests/EditorAutomationTests.cpp
+++ b/Engine/Source/Editor/UnrealEd/Private/Tests/EditorAutomationTests.cpp
@@ -19,7 +19,6 @@
 #include "Kismet2/BlueprintEditorUtils.h"
 #include "BlueprintGraphClasses.h"
 
-#include "SlateWordWrapper.h"
 #include "AutomationCommon.h"
 #include "AutomationEditorCommon.h"
 #include "AutomationTest.h"
@@ -191,89 +190,6 @@ bool FGenericImportAssetsAutomationTest::RunTest(const FString& Parameters)
 
 //////////////////////////////////////////////////////////////////////////
 
-/**
- * FTextWrappingAutomationTest
- * Simple unit test that wraps text using FSlateFontInfo::WrapTextToClippingWidth(). Cannot be run in a commandlet
- * as it executes code that routes through Slate.
- */
-IMPLEMENT_SIMPLE_AUTOMATION_TEST( FTextWrappingAutomationTest, "Engine.Rendering.Slate.Text Wrapping", EAutomationTestFlags::ATF_Editor | EAutomationTestFlags::ATF_NonNullRHI )
-
-/** 
- * Execute the text wrapping test
- *
- * @return	true if the test was successful, false otherwise
- */
-bool FTextWrappingAutomationTest::RunTest(const FString& Parameters)
-{
-	struct FTestStruct
-	{
-		static bool TestMethod(const FSlateFontInfo& FontInfo, const float WrapWidth, const FString& TestString, const int32 ArgumentCount, ...)
-		{
-			va_list Arguments;
-			va_start(Arguments, ArgumentCount);
-
-			TArray<FString> CorrectLines;
-			for(int32 i = 0; i < ArgumentCount; ++i)
-			{
-				const TCHAR* const Fragment = va_arg(Arguments, const TCHAR* const);
-				CorrectLines.Add(Fragment);
-			}
-
-			va_end(Arguments);
-
-			FString Result = SlateWordWrapper::WrapText(TestString, FontInfo, WrapWidth, 1.0f);
-			TArray<FString> ResultLines;
-			Result.ParseIntoArray(&ResultLines, TEXT("\n"), false);
-
-			return ResultLines == CorrectLines;
-		}
-	};
-
-	const FSlateFontInfo FontInfo = FEditorStyle::GetFontStyle("NormalFont");
-
-	const TSharedRef< FSlateFontMeasure > FontMeasureService = FSlateApplication::Get().GetRenderer()->GetFontMeasureService();
-	const float GlyphWidth = FontMeasureService->Measure( TEXT("A"), FontInfo).X;
-
-#define TEST(WrapWidth, TestString, ArgumentCount, ...)\
-	{\
-	if( !FTestStruct::TestMethod(FontInfo, (WrapWidth), (TestString), (ArgumentCount), __VA_ARGS__) )\
-		{\
-			AddError( TEXT("Testing string wrap failed, did not break line where expected.") );\
-		}\
-	}
-
-	TEST(GlyphWidth * 4, FString(TEXT("AAA    BBB CCC")), 3,   TEXT("AAA"),
-															TEXT("BBB"),
-															TEXT("CCC") );
-
-	TEST(GlyphWidth * 3, FString(TEXT("AAA BBB CCC")), 3,   TEXT("AAA"),
-															TEXT("BBB"),
-															TEXT("CCC") );
-
-	TEST(GlyphWidth * 2, FString(TEXT("AAA BBB CCC")), 6,   TEXT("AA"),
-															TEXT("A"),
-															TEXT("BB"),
-															TEXT("B"),
-															TEXT("CC"),
-															TEXT("C") );
-
-	TEST(GlyphWidth * 1, FString(TEXT("AAA BBB CCC")), 9,   TEXT("A"),
-															TEXT("A"),
-															TEXT("A"),
-															TEXT("B"),
-															TEXT("B"),
-															TEXT("B"),
-															TEXT("C"),
-															TEXT("C"),
-															TEXT("C") );
-
-#undef TEST
-
-	return true;
-}
-
-//////////////////////////////////////////////////////////////////////////
-
 /**
  * FJsonAutomationTest
  * Simple unit test that runs Json's in-built test cases
diff --git a/Engine/Source/Runtime/Core/Private/Internationalization/WordWrapper.cpp b/Engine/Source/Runtime/Core/Private/Internationalization/WordWrapper.cpp
deleted file mode 100644
diff --git a/Engine/Source/Runtime/Core/Public/Internationalization/WordWrapper.h b/Engine/Source/Runtime/Core/Public/Internationalization/WordWrapper.h
deleted file mode 100644
diff --git a/Engine/Source/Runtime/Engine/Classes/Engine/Canvas.h b/Engine/Source/Runtime/Engine/Classes/Engine/Canvas.h
--- a/Engine/Source/Runtime/Engine/Classes/Engine/Canvas.h
+++ b/Engine/Source/Runtime/Engine/Classes/Engine/Canvas.h
@@ -5,12 +5,11 @@
 //=============================================================================
 
 #pragma once
-#include "Internationalization/WordWrapper.h"
+
+#include "CanvasTypes.h"
 #include "Canvas.generated.h"
 
 class FSceneView;
-class FCanvas;
-class FCanvasItem;
 class UFont;
 
 /**
@@ -50,100 +49,6 @@ struct FCanvasIcon
 
 };
 
-
-/**
- * General purpose data structure for grouping all parameters needed when sizing or wrapping a string
- */
-USTRUCT()
-struct FTextSizingParameters
-{
-	GENERATED_USTRUCT_BODY()
-
-	/** a pixel value representing the horizontal screen location to begin rendering the string */
-	UPROPERTY()
-	float DrawX;
-
-	/** a pixel value representing the vertical screen location to begin rendering the string */
-	UPROPERTY()
-	float DrawY;
-
-	/** a pixel value representing the width of the area available for rendering the string */
-	UPROPERTY()
-	float DrawXL;
-
-	/** a pixel value representing the height of the area available for rendering the string */
-	UPROPERTY()
-	float DrawYL;
-
-	/** A value between 0.0 and 1.0, which represents how much the width/height should be scaled, where 1.0 represents 100% scaling. */
-	UPROPERTY()
-	FVector2D Scaling;
-
-	/** the font to use for sizing/wrapping the string */
-	UPROPERTY()
-	const UFont* DrawFont;
-
-	/** Horizontal spacing adjustment between characters and vertical spacing adjustment between wrapped lines */
-	UPROPERTY()
-	FVector2D SpacingAdjust;
-
-	FTextSizingParameters()
-		: DrawX(0)
-		, DrawY(0)
-		, DrawXL(0)
-		, DrawYL(0)
-		, Scaling(ForceInit)
-		, DrawFont(NULL)
-		, SpacingAdjust(ForceInit)
-	{
-	}
-
-
-		FTextSizingParameters( float inDrawX, float inDrawY, float inDrawXL, float inDrawYL, const UFont* inFont=NULL )
-		: DrawX(inDrawX), DrawY(inDrawY), DrawXL(inDrawXL), DrawYL(inDrawYL)
-		, Scaling(1.f,1.f), DrawFont(inFont)
-		, SpacingAdjust( 0.0f, 0.0f )
-		{
-		}
-		FTextSizingParameters( const UFont* inFont, float ScaleX, float ScaleY)
-		: DrawX(0.f), DrawY(0.f), DrawXL(0.f), DrawYL(0.f)
-		, Scaling(ScaleX,ScaleY), DrawFont(inFont)
-		, SpacingAdjust( 0.0f, 0.0f )
-		{
-		}
-	
-};
-
-/**
- * Used by UUIString::WrapString to track information about each line that is generated as the result of wrapping.
- */
-USTRUCT()
-struct FWrappedStringElement
-{
-	GENERATED_USTRUCT_BODY()
-
-	/** the string associated with this line */
-	UPROPERTY()
-	FString Value;
-
-	/** the size (in pixels) that it will take to render this string */
-	UPROPERTY()
-	FVector2D LineExtent;
-
-
-	FWrappedStringElement()
-		: LineExtent(ForceInit)
-	{
-	}
-
-
-		/** Constructor */
-		FWrappedStringElement( const TCHAR* InValue, float Width, float Height )
-		: Value(InValue), LineExtent(Width,Height)
-		{}
-	
-};
-
 UCLASS(transient)
 class ENGINE_API UCanvas : public UObject
 {
@@ -209,7 +114,6 @@ public:
 	FSceneView* SceneView;
 	FMatrix	ViewProjectionMatrix;
 	FQuat HmdOrientation;
-	
 
 	// UCanvas interface.
 	
@@ -346,7 +250,10 @@ public:
 	 *							not clear the array first.
 	 * @param	OutWrappedLineData An optional array to fill with the indices from the source string marking the begin and end points of the wrapped lines
 	 */
-	static void WrapString( FTextSizingParameters& Parameters, const float InCurX, const TCHAR* const pText, TArray<FWrappedStringElement>& out_Lines, FWordWrapper::FWrappedLineData* const OutWrappedLineData = nullptr);
+	static void WrapString( FCanvasWordWrapper& Wrapper, FTextSizingParameters& Parameters, const float InCurX, const TCHAR* const pText, TArray<FWrappedStringElement>& out_Lines, FCanvasWordWrapper::FWrappedLineData* const OutWrappedLineData = nullptr);
+
+	void WrapString( FTextSizingParameters& Parameters, const float InCurX, const TCHAR* const pText, TArray<FWrappedStringElement>& out_Lines, FCanvasWordWrapper::FWrappedLineData* const OutWrappedLineData = nullptr);
+
 
 	/** Transforms a 3D world-space vector into 2D screen coordinates. */
 	FVector Project(FVector Location) const;
diff --git a/Engine/Source/Runtime/Engine/Private/SubtitleManager.cpp b/Engine/Source/Runtime/Engine/Private/SubtitleManager.cpp
--- a/Engine/Source/Runtime/Engine/Private/SubtitleManager.cpp
+++ b/Engine/Source/Runtime/Engine/Private/SubtitleManager.cpp
@@ -163,7 +163,7 @@ void FSubtitleManager::DisplaySubtitle( FCanvas* Canvas, FActiveSubtitle* Subtit
  * If any of the active subtitles need to be split into multiple lines, do so now
  * - caveat - this assumes the width of the subtitle region does not change while the subtitle is active
  */
-void FSubtitleManager::SplitLinesToSafeZone( FIntRect& SubtitleRegion )
+void FSubtitleManager::SplitLinesToSafeZone( FCanvas* Canvas, FIntRect& SubtitleRegion )
 {
 	int32				i;
 	FString			Concatenated;
@@ -215,7 +215,7 @@ void FSubtitleManager::SplitLinesToSafeZone( FIntRect& SubtitleRegion )
 		// Word wrap into lines
 		TArray<FWrappedStringElement> Lines;
 		FTextSizingParameters RenderParms( 0.0f, 0.0f, SubtitleRegion.Width(), 0.0f, GEngine->GetSubtitleFont() );
-		UCanvas::WrapString( RenderParms, 0, *Concatenated, Lines );
+		UCanvas::WrapString( Canvas->WordWrapper, RenderParms, 0, *Concatenated, Lines );
 
 		// Set up the times
 		Subtitle.Subtitles.Empty();
@@ -333,7 +333,7 @@ void FSubtitleManager::DisplaySubtitles( FCanvas* InCanvas, FIntRect& InSubtitle
 		TrimRegionToSafeZone( InCanvas, InSubtitleRegion );
 
 		// If the lines have not already been split, split them to the safe zone now
-		SplitLinesToSafeZone( InSubtitleRegion );
+		SplitLinesToSafeZone( InCanvas, InSubtitleRegion );
 
 		// Find the subtitle to display
 		PTRINT HighestPriorityID = FindHighestPrioritySubtitle( InAudioTimeSeconds );
diff --git a/Engine/Source/Runtime/Engine/Private/UnrealEngine.cpp b/Engine/Source/Runtime/Engine/Private/UnrealEngine.cpp
--- a/Engine/Source/Runtime/Engine/Private/UnrealEngine.cpp
+++ b/Engine/Source/Runtime/Engine/Private/UnrealEngine.cpp
@@ -6560,7 +6560,7 @@ static void DrawProperty(UCanvas* CanvasObject, UObject* Obj, const FDebugDispla
 		CanvasObject->ClippedStrLen(GEngine->GetSmallFont(), 1.0f, 1.0f, XL, YL, *PropText);
 		FTextSizingParameters DrawParams(X, Y, CanvasObject->SizeX - X, 0, GEngine->GetSmallFont());
 		TArray<FWrappedStringElement> TextLines;
-		UCanvas::WrapString(DrawParams, X + XL, *Str, TextLines);
+		CanvasObject->WrapString(DrawParams, X + XL, *Str, TextLines);
 		int32 XL2 = XL;
 		if (TextLines.Num() > 0)
 		{
diff --git a/Engine/Source/Runtime/Engine/Private/UserInterface/Canvas.cpp b/Engine/Source/Runtime/Engine/Private/UserInterface/Canvas.cpp
--- a/Engine/Source/Runtime/Engine/Private/UserInterface/Canvas.cpp
+++ b/Engine/Source/Runtime/Engine/Private/UserInterface/Canvas.cpp
@@ -10,7 +10,7 @@
 #include "EngineFontServices.h"
 #include "TileRendering.h"
 #include "RHIStaticStates.h"
-#include "WordWrapper.h"
+#include "BreakIterator.h"
 
 #include "IHeadMountedDisplay.h"
 #include "Debug/ReporterGraph.h"
@@ -27,6 +27,180 @@ DEFINE_STAT(STAT_Canvas_GetBatchElementsTime);
 DEFINE_STAT(STAT_Canvas_AddTileRenderTime);
 DEFINE_STAT(STAT_Canvas_NumBatchesCreated);
 
+FCanvasWordWrapper::FCanvasWordWrapper()
+	: GraphemeBreakIterator(FBreakIterator::CreateCharacterBoundaryIterator())
+	, LineBreakIterator(FBreakIterator::CreateLineBreakIterator())
+{
+}
+
+void FCanvasWordWrapper::Execute(const TCHAR* const InString, const FTextSizingParameters& InParameters, TArray<FWrappedStringElement>& OutStrings, FWrappedLineData* const OutWrappedLineData)
+{
+	FWrappingState WrappingState(InString, FCString::Strlen(InString), InParameters, OutStrings, OutWrappedLineData);
+
+	if (WrappingState.WrappedLineData)
+	{
+		WrappingState.WrappedLineData->Empty();
+	}
+
+	GraphemeBreakIterator->SetString(WrappingState.String, WrappingState.StringLength);
+	LineBreakIterator->SetString(WrappingState.String, WrappingState.StringLength);
+
+	for(int32 i = 0; i < WrappingState.StringLength; ++i) // Sanity check: Doesn't seem valid to have more lines than code units.
+	{	
+		if( !ProcessLine(WrappingState) )
+		{
+			break;
+		}
+	}
+}
+
+bool FCanvasWordWrapper::ProcessLine(FWrappingState& WrappingState)
+{
+	bool bHasAddedLine = false;
+	if(WrappingState.StartIndex < WrappingState.StringLength)
+	{
+		int32 BreakIndex = FindFirstMandatoryBreakBetween(WrappingState, WrappingState.StringLength);
+
+		int32 NextStartIndex;
+		if( BreakIndex == INDEX_NONE || !DoesSubstringFit(WrappingState, BreakIndex) )
+		{
+			BreakIndex = INDEX_NONE;
+			int32 WrapIndex = FindIndexAtOrAfterWrapWidth(WrappingState);
+
+			if (WrapIndex == WrappingState.StringLength)
+			{
+				BreakIndex = WrapIndex;
+			}
+
+			if (BreakIndex <= WrappingState.StartIndex) // No mandatory break.
+			{
+				BreakIndex = FindLastBreakCandidateBetween(WrappingState.StartIndex, WrapIndex);
+			}
+
+			if (BreakIndex <= WrappingState.StartIndex) // No candidate break.
+			{
+				// Break after minimum length that would preserve the meaning/appearance.
+				BreakIndex = FindEndOfLastWholeGraphemeCluster(WrappingState.StartIndex, WrapIndex);
+			}
+
+			if (BreakIndex <= WrappingState.StartIndex) // No complete grapheme cluster.
+			{
+				BreakIndex = WrapIndex; // Break at wrap.
+			}
+
+			// Index for the next search
+			NextStartIndex = BreakIndex;
+		}
+		else
+		{
+			// Index for the next search
+			NextStartIndex = BreakIndex;
+			// The index is inclusive of the break - we don't want the break char in the string
+			--BreakIndex;			
+		}
+
+		while (BreakIndex > 0 && FText::IsWhitespace(WrappingState.String[BreakIndex - 1]))
+		{
+			--BreakIndex;
+		}
+
+		if (WrappingState.StartIndex <= BreakIndex)
+		{
+			AddLine(WrappingState, BreakIndex);
+
+			bHasAddedLine = true;
+		}
+
+		while (NextStartIndex < WrappingState.StringLength && FText::IsWhitespace(WrappingState.String[NextStartIndex]))
+		{
+			++NextStartIndex;
+		}
+
+		if(WrappingState.WrappedLineData)
+		{
+			WrappingState.WrappedLineData->Add(TPairInitializer<int32, int32>(WrappingState.StartIndex, BreakIndex));
+		}
+
+		WrappingState.StartIndex = NextStartIndex;
+	}
+	return bHasAddedLine;
+}
+
+bool FCanvasWordWrapper::DoesSubstringFit(FWrappingState& WrappingState, const int32 EndIndex)
+{
+	FTextSizingParameters MeasureParameters(WrappingState.Parameters);
+	int32 Unused;
+	UCanvas::MeasureStringInternal( MeasureParameters, WrappingState.String + WrappingState.StartIndex, EndIndex - WrappingState.StartIndex, 0, UCanvas::ELastCharacterIndexFormat::Unused, Unused );
+	return MeasureParameters.DrawXL <= WrappingState.Parameters.DrawXL;
+}
+
+int32 FCanvasWordWrapper::FindIndexAtOrAfterWrapWidth(FWrappingState& WrappingState)
+{
+	FTextSizingParameters MeasureParameters(WrappingState.Parameters);
+	int32 Return = INDEX_NONE;
+	UCanvas::MeasureStringInternal(MeasureParameters, WrappingState.String + WrappingState.StartIndex, WrappingState.StringLength - WrappingState.StartIndex, WrappingState.Parameters.DrawXL, UCanvas::ELastCharacterIndexFormat::CharacterAtOffset, Return);
+	return WrappingState.StartIndex + Return;
+}
+
+void FCanvasWordWrapper::AddLine(FWrappingState& WrappingState, const int32 EndIndex)
+{
+	FTextSizingParameters MeasureParameters(WrappingState.Parameters);
+	FString Substring(EndIndex - WrappingState.StartIndex, WrappingState.String + WrappingState.StartIndex);
+	FWrappedStringElement Element(*Substring, 0.0f, 0.0f);
+	UCanvas::CanvasStringSize(MeasureParameters, *Element.Value);
+	Element.LineExtent.X = MeasureParameters.DrawXL;
+	Element.LineExtent.Y = MeasureParameters.DrawYL;
+	WrappingState.Results.Add(Element);
+}
+
+int32 FCanvasWordWrapper::FindFirstMandatoryBreakBetween(FWrappingState& WrappingState, const int32 WrapIndex)
+{
+	int32 BreakIndex = INDEX_NONE;
+	for(int32 i = WrappingState.StartIndex + 1; i < WrapIndex; ++i)
+	{
+		const TCHAR Previous = WrappingState.String[i - 1];
+		if( FChar::IsLinebreak(Previous) ) // Line break occurs *after* linebreak character.
+		{
+			const TCHAR* const Current = i < WrapIndex ? WrappingState.String + i : NULL;
+			if(	Previous != FChar::CarriageReturn || !(Current && *Current == FChar::LineFeed) ) // Line break cannot occur within CR LF pair.
+			{
+				BreakIndex = i;
+				break;
+			}
+		}
+	}
+	// If we reached the end of the string we must also check that the last char is not a newline
+	if( BreakIndex == INDEX_NONE )
+	{
+		const TCHAR Previous = WrappingState.String[WrapIndex - 1];
+		if( FChar::IsLinebreak(Previous) ) // Line break occurs *after* linebreak character.
+		{
+			BreakIndex = WrapIndex;
+		}
+	}
+	return BreakIndex;
+}
+
+int32 FCanvasWordWrapper::FindLastBreakCandidateBetween(const int32 InStartIndex, const int32 WrapIndex)
+{
+	int32 BreakIndex = LineBreakIterator->MoveToCandidateBefore(WrapIndex + 1);
+	if(BreakIndex < InStartIndex)
+	{
+		BreakIndex = INDEX_NONE;
+	}
+	return BreakIndex;
+}
+
+int32 FCanvasWordWrapper::FindEndOfLastWholeGraphemeCluster(const int32 InStartIndex, const int32 WrapIndex)
+{
+	int32 BreakIndex = GraphemeBreakIterator->MoveToCandidateBefore(WrapIndex + 1);
+	if(BreakIndex < InStartIndex)
+	{
+		BreakIndex = INDEX_NONE;
+	}
+	return BreakIndex;
+}
+
 FCanvas::FCanvas(FRenderTarget* InRenderTarget, FHitProxyConsumer* InHitProxyConsumer, UWorld* InWorld, ERHIFeatureLevel::Type InFeatureLevel)
 :	ViewRect(0,0,0,0)
 ,	RenderTarget(InRenderTarget)
@@ -1033,81 +1207,14 @@ void UCanvas::CanvasStringSize( FTextSizingParameters& Parameters, const TCHAR*
 	MeasureStringInternal(Parameters, pText, FCString::Strlen(pText), 0, ELastCharacterIndexFormat::Unused, Unused);
 }
 
-namespace
-{
-	class FCanvasWordWrapper : public FWordWrapper
-	{
-	public:
-		/**
-		 * Used to generate multi-line/wrapped text.
-		 *
-		 * @param InString The unwrapped text.
-		 * @param InFontInfo The font used to render the text.
-		 * @param InWrapWidth The width available.
-		 * @param OutWrappedLineData An optional array to fill with the indices from the source string marking the begin and end points of the wrapped lines
-		 */
-		static void Execute(FTextSizingParameters& InParameters, const TCHAR* const InString, TArray<FWrappedStringElement>& OutStrings, FWrappedLineData* const OutWrappedLineData);
-	protected:
-		bool DoesSubstringFit(const int32 StartIndex, const int32 EndIndex);
-		int32 FindIndexAtOrAfterWrapWidth(const int32 StartIndex);
-		void AddLine(const int32 StartIndex, const int32 EndIndex);
-	private:
-		FCanvasWordWrapper(FTextSizingParameters& InParameters, const TCHAR* const InString, TArray<FWrappedStringElement>& OutStrings, FWrappedLineData* const OutWrappedLineData);
-	private:
-		FTextSizingParameters& Parameters;
-		TArray<FWrappedStringElement>& Results;
-	};
-
-	void FCanvasWordWrapper::Execute(FTextSizingParameters& InParameters, const TCHAR* const InString, TArray<FWrappedStringElement>& OutStrings, FWrappedLineData* const OutWrappedLineData)
-	{
-		FCanvasWordWrapper WordWrapper(InParameters, InString, OutStrings, OutWrappedLineData);
-		const int32 StringLength = FCString::Strlen(InString);
-		for(int32 i = 0; i < StringLength; ++i) // Sanity check: Doesn't seem valid to have more lines than code units.
-		{	
-			if( !WordWrapper.ProcessLine() )
-			{
-				break;
-			}
-		}
-	}
-
-	bool FCanvasWordWrapper::DoesSubstringFit(const int32 StartIndex, const int32 EndIndex)
-	{
-		FTextSizingParameters MeasureParameters(Parameters);
-		int32 Unused;
-		UCanvas::MeasureStringInternal( MeasureParameters, String + StartIndex, EndIndex - StartIndex, 0, UCanvas::ELastCharacterIndexFormat::Unused, Unused );
-		return MeasureParameters.DrawXL <= Parameters.DrawXL;
-	}
-
-	int32 FCanvasWordWrapper::FindIndexAtOrAfterWrapWidth(const int32 StartIndex)
-	{
-		FTextSizingParameters MeasureParameters(Parameters);
-		int32 Return = INDEX_NONE;
-		UCanvas::MeasureStringInternal(MeasureParameters, String + StartIndex, StringLength - StartIndex, Parameters.DrawXL, UCanvas::ELastCharacterIndexFormat::CharacterAtOffset, Return);
-		return StartIndex + Return;
-	}
-
-	void FCanvasWordWrapper::AddLine(const int32 StartIndex, const int32 EndIndex)
-	{
-		FTextSizingParameters MeasureParameters(Parameters);
-		FString Substring(EndIndex - StartIndex, String + StartIndex);
-		FWrappedStringElement Element(*Substring, 0.0f, 0.0f);
-		UCanvas::CanvasStringSize(MeasureParameters, *Element.Value);
-		Element.LineExtent.X = MeasureParameters.DrawXL;
-		Element.LineExtent.Y = MeasureParameters.DrawYL;
-		Results.Add(Element);
-	}
-
-	FCanvasWordWrapper::FCanvasWordWrapper(FTextSizingParameters& InParameters, const TCHAR* const InString, TArray<FWrappedStringElement>& OutStrings, FWrappedLineData* const OutWrappedLineData)
-		: FWordWrapper(InString, FString(InString).Len(), OutWrappedLineData), Parameters(InParameters), Results(OutStrings)
-	{
-		Results.Empty();
-	}
+void UCanvas::WrapString( FCanvasWordWrapper& Wrapper, FTextSizingParameters& Parameters, const float InCurX, const TCHAR* const pText, TArray<FWrappedStringElement>& out_Lines, FCanvasWordWrapper::FWrappedLineData* const OutWrappedLineData)
+{
+	Wrapper.Execute(pText, Parameters, out_Lines, OutWrappedLineData);
 }
 
-void UCanvas::WrapString( FTextSizingParameters& Parameters, const float InCurX, const TCHAR* const pText, TArray<FWrappedStringElement>& out_Lines, FWordWrapper::FWrappedLineData* const OutWrappedLineData)
+void UCanvas::WrapString( FTextSizingParameters& Parameters, const float InCurX, const TCHAR* const pText, TArray<FWrappedStringElement>& out_Lines, FCanvasWordWrapper::FWrappedLineData* const OutWrappedLineData)
 {
-	FCanvasWordWrapper::Execute(Parameters, pText, out_Lines, OutWrappedLineData);
+	UCanvas::WrapString( Canvas->WordWrapper, Parameters, InCurX, pText, out_Lines, OutWrappedLineData);
 }
 
 /*-----------------------------------------------------------------------------
diff --git a/Engine/Source/Runtime/Engine/Public/CanvasTypes.h b/Engine/Source/Runtime/Engine/Public/CanvasTypes.h
--- a/Engine/Source/Runtime/Engine/Public/CanvasTypes.h
+++ b/Engine/Source/Runtime/Engine/Public/CanvasTypes.h
@@ -7,8 +7,10 @@
 #pragma once
 
 #include "Engine/EngineTypes.h"
+#include "IBreakIterator.h"
 #include "BatchedElements.h"
 #include "CanvasItem.h"
+#include "CanvasTypes.generated.h"
 
 class FMaterialRenderProxy;
 class FTexture;
@@ -16,6 +18,178 @@ class FRenderTarget;
 class FHitProxyConsumer;
 class FRHICommandListImmediate;
 
+/**
+ * General purpose data structure for grouping all parameters needed when sizing or wrapping a string
+ */
+USTRUCT()
+struct FTextSizingParameters
+{
+	GENERATED_USTRUCT_BODY()
+
+	/** a pixel value representing the horizontal screen location to begin rendering the string */
+	UPROPERTY()
+	float DrawX;
+
+	/** a pixel value representing the vertical screen location to begin rendering the string */
+	UPROPERTY()
+	float DrawY;
+
+	/** a pixel value representing the width of the area available for rendering the string */
+	UPROPERTY()
+	float DrawXL;
+
+	/** a pixel value representing the height of the area available for rendering the string */
+	UPROPERTY()
+	float DrawYL;
+
+	/** A value between 0.0 and 1.0, which represents how much the width/height should be scaled, where 1.0 represents 100% scaling. */
+	UPROPERTY()
+	FVector2D Scaling;
+
+	/** the font to use for sizing/wrapping the string */
+	UPROPERTY()
+	const UFont* DrawFont;
+
+	/** Horizontal spacing adjustment between characters and vertical spacing adjustment between wrapped lines */
+	UPROPERTY()
+	FVector2D SpacingAdjust;
+
+	FTextSizingParameters()
+		: DrawX(0)
+		, DrawY(0)
+		, DrawXL(0)
+		, DrawYL(0)
+		, Scaling(ForceInit)
+		, DrawFont(NULL)
+		, SpacingAdjust(ForceInit)
+	{
+	}
+
+
+		FTextSizingParameters( float inDrawX, float inDrawY, float inDrawXL, float inDrawYL, const UFont* inFont=NULL )
+		: DrawX(inDrawX), DrawY(inDrawY), DrawXL(inDrawXL), DrawYL(inDrawYL)
+		, Scaling(1.f,1.f), DrawFont(inFont)
+		, SpacingAdjust( 0.0f, 0.0f )
+		{
+		}
+		FTextSizingParameters( const UFont* inFont, float ScaleX, float ScaleY)
+		: DrawX(0.f), DrawY(0.f), DrawXL(0.f), DrawYL(0.f)
+		, Scaling(ScaleX,ScaleY), DrawFont(inFont)
+		, SpacingAdjust( 0.0f, 0.0f )
+		{
+		}
+	
+};
+
+/**
+ * Used by UUIString::WrapString to track information about each line that is generated as the result of wrapping.
+ */
+USTRUCT()
+struct FWrappedStringElement
+{
+	GENERATED_USTRUCT_BODY()
+
+	/** the string associated with this line */
+	UPROPERTY()
+	FString Value;
+
+	/** the size (in pixels) that it will take to render this string */
+	UPROPERTY()
+	FVector2D LineExtent;
+
+
+	FWrappedStringElement()
+		: LineExtent(ForceInit)
+	{
+	}
+
+
+		/** Constructor */
+		FWrappedStringElement( const TCHAR* InValue, float Width, float Height )
+		: Value(InValue), LineExtent(Width,Height)
+		{}
+	
+};
+
+class FCanvasWordWrapper
+{
+public:
+	/** Array of indices where the wrapped lines begin and end in the source string */
+	typedef TArray< TPair<int32, int32> > FWrappedLineData;
+
+private:
+	struct FWrappingState
+	{
+		FWrappingState(	const TCHAR* const InString,
+						const int32 InStringLength,
+						const FTextSizingParameters& InParameters,
+						TArray<FWrappedStringElement>& InResults,
+						FWrappedLineData* const InWrappedLineData)
+			: String(InString)
+			, StringLength(InStringLength)
+			, Parameters(InParameters)
+			, StartIndex(0)
+			, Results(InResults)
+			, WrappedLineData(InWrappedLineData)
+		{}
+
+		const TCHAR* const String;
+		const int32 StringLength;
+		const FTextSizingParameters& Parameters;
+		int32 StartIndex;
+		TArray<FWrappedStringElement>& Results;
+		FWrappedLineData* const WrappedLineData;
+	};
+
+public:
+	FCanvasWordWrapper();
+
+	/**
+	* Used to generate multi-line/wrapped text.
+	*
+	* @param InString The unwrapped text.
+	* @param InFontInfo The font used to render the text.
+	* @param InWrapWidth The width available.
+	* @param OutWrappedLineData An optional array to fill with the indices from the source string marking the begin and end points of the wrapped lines
+	*/
+	void Execute(const TCHAR* const InString, const FTextSizingParameters& InParameters, TArray<FWrappedStringElement>& OutStrings, FWrappedLineData* const OutWrappedLineData);
+
+private:
+	/**
+		* Processes the string using a word wrapping algorithm, resulting in up to a single line.
+		*
+		* @return	True if a new line could be processed, false otherwise such as having reached the end of the string.
+		*/
+	bool ProcessLine(FWrappingState& WrappingState);
+
+	/**
+		* Stub method that should measure the substring of the range [StartIndex, EndIndex).
+		*
+		* @return	True if the substring fits the desired width.
+		*/
+	bool DoesSubstringFit(FWrappingState& WrappingState, const int32 EndIndex);
+
+	/**
+		* Stub method that should measure the substring starting from StartIndex until the wrap width is found or no more indices remain.
+		*
+		* @return	The index of the character that is at or after the desired width.
+		*/
+	int32 FindIndexAtOrAfterWrapWidth(FWrappingState& WrappingState);
+
+	/**
+		* Stub method for processing a single produced substring of the range [StartIndex, EndIndex) as a new line.
+		*/
+	void AddLine(FWrappingState& WrappingState, const int32 EndIndex);
+
+	int32 FindFirstMandatoryBreakBetween(FWrappingState& WrappingState, const int32 WrapIndex);
+	int32 FindLastBreakCandidateBetween(const int32 StartIndex, const int32 WrapIndex);
+	int32 FindEndOfLastWholeGraphemeCluster(const int32 StartIndex, const int32 WrapIndex);
+
+private:
+	TSharedPtr<IBreakIterator> GraphemeBreakIterator;
+	TSharedPtr<IBreakIterator> LineBreakIterator;
+};
+
 /**
  * Encapsulates the canvas state.
  */
@@ -362,6 +536,11 @@ public:
 		bScaledToRenderTarget = scale;
 	}
 	FORCEINLINE bool IsScaledToRenderTarget() const { return bScaledToRenderTarget; }
+
+public:
+	/** Private class for handling word wrapping behavior. */
+	FCanvasWordWrapper WordWrapper;
+
 private:
 	/** Stack of SortKeys. All rendering is done using the top most sort key */
 	TArray<int32> DepthSortKeyStack;	
diff --git a/Engine/Source/Runtime/Engine/Public/SubtitleManager.h b/Engine/Source/Runtime/Engine/Public/SubtitleManager.h
--- a/Engine/Source/Runtime/Engine/Public/SubtitleManager.h
+++ b/Engine/Source/Runtime/Engine/Public/SubtitleManager.h
@@ -57,7 +57,7 @@ public:
 	 * If any of the active subtitles need to be split into multiple lines, do so now
 	 * - caveat - this assumes the width of the subtitle region does not change while the subtitle is active
 	 */
-	void		SplitLinesToSafeZone( FIntRect & SubtitleRegion );
+	void		SplitLinesToSafeZone( FCanvas* Canvas, FIntRect & SubtitleRegion );
 
 	/**
 	 * Find the highest priority subtitle from the list of currently active ones
diff --git a/Engine/Source/Runtime/Slate/Private/Framework/Text/SlateWordWrapper.cpp b/Engine/Source/Runtime/Slate/Private/Framework/Text/SlateWordWrapper.cpp
deleted file mode 100644
diff --git a/Engine/Source/Runtime/Slate/Public/Framework/Text/SlateWordWrapper.h b/Engine/Source/Runtime/Slate/Public/Framework/Text/SlateWordWrapper.h
deleted file mode 100644
