UE4 - Fix for material time node not getting propagated

diff --git a/Engine/Source/Runtime/Engine/Private/Materials/MaterialUniformExpressions.h b/Engine/Source/Runtime/Engine/Private/Materials/MaterialUniformExpressions.h
--- a/Engine/Source/Runtime/Engine/Private/Materials/MaterialUniformExpressions.h
+++ b/Engine/Source/Runtime/Engine/Private/Materials/MaterialUniformExpressions.h
@@ -634,6 +634,10 @@ public:
 	{
 		return X->IsConstant();
 	}
+	virtual bool IsChangingPerFrame() const
+	{
+		return X->IsChangingPerFrame();
+	}
 	virtual bool IsIdentical(const FMaterialUniformExpression* OtherExpression) const
 	{
 		if (GetType() != OtherExpression->GetType())
