Remove parameterized constructor from FTimerHandle. The handle internals should be completely opaque to its users.

diff --git a/Engine/Source/Runtime/Engine/Public/TimerManager.h b/Engine/Source/Runtime/Engine/Public/TimerManager.h
--- a/Engine/Source/Runtime/Engine/Public/TimerManager.h
+++ b/Engine/Source/Runtime/Engine/Public/TimerManager.h
@@ -83,12 +83,6 @@ struct FTimerHandle
 
 	}
 
-	FTimerHandle(int32 InHandle)
-		: Handle(InHandle)
-	{
-
-	}
-
 	bool IsValid() const
 	{
 		return Handle != INDEX_NONE;
