export decal component
UE-6204 GitHub 652 : Export DecalComponent class from the engine dll

diff --git a/Engine/Source/Runtime/Engine/Classes/Components/DecalComponent.h b/Engine/Source/Runtime/Engine/Classes/Components/DecalComponent.h
--- a/Engine/Source/Runtime/Engine/Classes/Components/DecalComponent.h
+++ b/Engine/Source/Runtime/Engine/Classes/Components/DecalComponent.h
@@ -10,8 +10,8 @@ class FDeferredDecalProxy;
 /** 
  * A material that is rendered onto the surface of a mesh. A kind of 'bumper sticker' for a model.
  */
-UCLASS(hidecategories=(Collision, Object, Physics, SceneComponent, Activation, "Components|Activation", Mobility), ClassGroup=Rendering, MinimalAPI, meta=(BlueprintSpawnableComponent))
-class UDecalComponent : public USceneComponent
+UCLASS(hidecategories=(Collision, Object, Physics, SceneComponent, Activation, "Components|Activation", Mobility), ClassGroup=Rendering, meta=(BlueprintSpawnableComponent))
+class ENGINE_API UDecalComponent : public USceneComponent
 {
 	GENERATED_UCLASS_BODY()
 
@@ -29,15 +29,15 @@ class UDecalComponent : public USceneComponent
 
 	/** Sets the sort order for the decal component. Higher values draw later (on top). This will force the decal to reattach */
 	UFUNCTION(BlueprintCallable, Category = "Rendering|Components|Decal")
-	ENGINE_API void SetSortOrder(int32 Value);
+	void SetSortOrder(int32 Value);
 
 	/** setting decal material on decal component. This will force the decal to reattach */
 	UFUNCTION(BlueprintCallable, Category="Rendering|Components|Decal")
-	ENGINE_API void SetDecalMaterial(class UMaterialInterface* NewDecalMaterial);
+	void SetDecalMaterial(class UMaterialInterface* NewDecalMaterial);
 
 	/** Accessor for decal material */
 	UFUNCTION(BlueprintCallable, Category="Rendering|Components|Decal")
-	ENGINE_API class UMaterialInterface* GetDecalMaterial() const;
+	class UMaterialInterface* GetDecalMaterial() const;
 
 	/** Utility to allocate a new Dynamic Material Instance, set its parent to the currently applied material, and assign it */
 	UFUNCTION(BlueprintCallable, Category="Rendering|Components|Decal")
@@ -51,7 +51,7 @@ public:
 	/**
 	 * Pushes new selection state to the render thread primitive proxy
 	 */
-	ENGINE_API void PushSelectionToProxy();
+	void PushSelectionToProxy();
 
 protected:
 	/** Handle for efficient management of DestroyDecalComponent timer */
