Associate foliage on BSP with the Brush actor rather than the level's UModel. Painted instances now move when you move BSP brushes.
Fixes disappearing instances on BSP when rebuilding.

diff --git a/Engine/Source/Editor/FoliageEdit/Private/FoliageEdMode.cpp b/Engine/Source/Editor/FoliageEdit/Private/FoliageEdMode.cpp
--- a/Engine/Source/Editor/FoliageEdit/Private/FoliageEdMode.cpp
+++ b/Engine/Source/Editor/FoliageEdit/Private/FoliageEdMode.cpp
@@ -485,6 +485,16 @@ struct FPotentialInstance
 			Inst.Location = Inst.GetInstanceWorldTransform().TransformPosition(FVector(0, 0, Inst.ZOffset));
 		}
 
+		UModelComponent* ModelComponent = Cast<UModelComponent>(HitComponent);
+		if (ModelComponent)
+		{
+			ABrush* BrushActor = ModelComponent->GetModel()->FindBrush(HitLocation);
+			if (BrushActor)
+			{
+				HitComponent = BrushActor->GetBrushComponent();
+			}
+		}
+
 		Inst.Base = HitComponent;
 
 		return bSkipCollision || CheckCollisionWithWorld(Settings, Inst, HitNormal, HitLocation, InWorld, InIFA);
diff --git a/Engine/Source/Runtime/Engine/Private/InstancedFoliage.cpp b/Engine/Source/Runtime/Engine/Private/InstancedFoliage.cpp
--- a/Engine/Source/Runtime/Engine/Private/InstancedFoliage.cpp
+++ b/Engine/Source/Runtime/Engine/Private/InstancedFoliage.cpp
@@ -1212,8 +1212,9 @@ FVector AInstancedFoliageActor::GetSelectionLocation()
 
 void AInstancedFoliageActor::MapRebuild()
 {
-	// Map rebuilt - this may have modified the BSP components and thrown the previous ones away - if so we need to migrate the foliage across.
-	UE_LOG(LogInstancedFoliage, Log, TEXT("Map Rebuilt - Update all BSP painted foliage!"));
+	// Map rebuild may have modified the BSP's ModelComponents and thrown the previous ones away.
+	// Most BSP-painted foliage is attached to a Brush's UModelComponent which persist across rebuilds,
+	// but any foliage attached directly to the level BSP's ModelComponents will need to try to find a new base.
 
 	TMap<UFoliageType*, TArray<FFoliageInstance>> NewInstances;
 	TArray<UModelComponent*> RemovedModelComponents;
@@ -1232,7 +1233,8 @@ void AInstancedFoliageActor::MapRebuild()
 		{
 			// BSP components are UModelComponents - they are the only ones we need to change
 			UModelComponent* TargetComponent = Cast<UModelComponent>(ComponentFoliagePair.Key);
-			if (TargetComponent)
+			// Check if it's part of a brush. We only need to fix up model components that are part of the level BSP.
+			if (TargetComponent && Cast<ABrush>(TargetComponent->GetOuter()) == nullptr)
 			{
 				// Delete its instances later
 				RemovedModelComponents.Add(TargetComponent);
@@ -1242,7 +1244,7 @@ void AInstancedFoliageActor::MapRebuild()
 				// We have to test each instance to see if we can migrate it across
 				for (int32 InstanceIdx : FoliageInfo.Instances)
 				{
-					// Use a line test against the world, similar to FoliageEditMode.
+					// Use a line test against the world. This is not very reliable as we don't know the original trace direction.
 					check(MeshInfo.Instances.IsValidIndex(InstanceIdx));
 					FFoliageInstance const& Instance = MeshInfo.Instances[InstanceIdx];
 
diff --git a/Engine/Source/Runtime/Engine/Private/ModelCollision.cpp b/Engine/Source/Runtime/Engine/Private/ModelCollision.cpp
--- a/Engine/Source/Runtime/Engine/Private/ModelCollision.cpp
+++ b/Engine/Source/Runtime/Engine/Private/ModelCollision.cpp
@@ -121,6 +121,106 @@ float UModel::FindNearestVertex
 	return Nodes.Num() ? ::FindNearestVertex( *this,SourcePoint,DestPoint,MinRadius,0,pVertex ) : -1.f;
 }
 
+/*---------------------------------------------------------------------------------------
+	Find Brush Actor
+---------------------------------------------------------------------------------------*/
+//
+// ClipNode - return the node containing a specified location from a number of coplanar nodes.
+//
+int32 ClipNode(const UModel& Model, int32 iNode, const FVector& HitLocation)
+{
+	while (iNode != INDEX_NONE)
+	{
+		const FBspNode& Node = Model.Nodes[iNode];
+		int32 NumVertices = Node.NumVertices;
+
+		// Only consider this node if it has some vertices.
+		if (NumVertices > 0)
+		{
+			int32 iVertPool = Node.iVertPool;
+
+			FVector PrevPt = Model.Points[Model.Verts[iVertPool + NumVertices - 1].pVertex];
+			FVector Normal = Model.Surfs[Node.iSurf].Plane;
+			float PrevDot = 0.f;
+
+			for (int32 i = 0; i<NumVertices; i++)
+			{
+				FVector Pt = Model.Points[Model.Verts[iVertPool + i].pVertex];
+				float Dot = FPlane(Pt, Normal ^ (Pt - PrevPt)).PlaneDot(HitLocation);
+				// Check for sign change
+				if ((Dot < 0.f && PrevDot > 0.f) || (Dot > 0.f && PrevDot < 0.f))
+				{
+					goto TryNextNode;
+				}
+				PrevPt = Pt;
+				PrevDot = Dot;
+			}
+
+			return iNode;
+		}
+
+	TryNextNode:;
+		iNode = Node.iPlane;
+	}
+	return INDEX_NONE;
+}
+
+//
+// Find the iSurf for node the point lies upon. Returns INDEX_NONE if the point does not lie on a surface.
+//
+static int32 FindSurf(const UModel &Model, const FVector &SourcePoint, int32 iNode, float Tolerance)
+{
+	while (iNode != INDEX_NONE)
+	{
+		const FBspNode	*Node = &Model.Nodes[iNode];
+		const int32	    iBack = Node->iBack;
+		const float PlaneDist = Node->Plane.PlaneDot(SourcePoint);
+		if (PlaneDist >= -Tolerance && Node->iFront != INDEX_NONE)
+		{
+			// Check front.
+			int32 Front = FindSurf(Model, SourcePoint, Node->iFront, Tolerance);
+			if (Front != INDEX_NONE)
+			{
+				return Front;
+			}
+		}
+		if (PlaneDist > -Tolerance && PlaneDist <= Tolerance)
+		{
+			iNode = ClipNode(Model, iNode, SourcePoint);
+			if (iNode != INDEX_NONE)
+			{
+				return Model.Nodes[iNode].iSurf;
+			}
+			else
+			{
+				return Tolerance;
+			}
+		}
+		if (PlaneDist > KINDA_SMALL_NUMBER)
+		{
+			break;
+		}
+		iNode = iBack;
+	}
+	return INDEX_NONE;
+}
+
+//
+// Find the brush actor associated with this point, or NULL if the point does not lie on a BSP surface.
+//
+ABrush* UModel::FindBrush(const FVector	&SourcePoint) const
+{
+	if (Nodes.Num())
+	{
+		int32 iSurf = FindSurf(*this, SourcePoint, 0, 0.1f);
+		if (iSurf != INDEX_NONE)
+		{
+			return Surfs[iSurf].Actor;
+		}
+	}
+	return nullptr;
+}
+
 /*---------------------------------------------------------------------------------------
    Bound filter precompute.
 ---------------------------------------------------------------------------------------*/
diff --git a/Engine/Source/Runtime/Engine/Public/Model.h b/Engine/Source/Runtime/Engine/Public/Model.h
--- a/Engine/Source/Runtime/Engine/Public/Model.h
+++ b/Engine/Source/Runtime/Engine/Public/Model.h
@@ -523,6 +523,8 @@ public:
 		const FPlane	&Sphere
 	);
 
+	/* Find the source brush actor associated with this point, or NULL if the point does not lie on a BSP surface. */
+	ENGINE_API ABrush* FindBrush(const FVector &SourcePoint) const;
 
 	/**
 	 * Creates a bounding box for the passed in node
