Fix foliage flickering when painting

diff --git a/Engine/Source/Runtime/Engine/Private/HierarchicalInstancedStaticMesh.cpp b/Engine/Source/Runtime/Engine/Private/HierarchicalInstancedStaticMesh.cpp
--- a/Engine/Source/Runtime/Engine/Private/HierarchicalInstancedStaticMesh.cpp
+++ b/Engine/Source/Runtime/Engine/Private/HierarchicalInstancedStaticMesh.cpp
@@ -1252,7 +1252,7 @@ void FHierarchicalStaticMeshSceneProxy::GetDynamicMeshElements(const TArray<cons
 				else
 				{
 					// more than 1000, render them all at Max LOD
-					InstanceParams.AddRun(RenderData->LODResources.Num() - 1, RenderData->LODResources.Num() - 1, FirstUnbuiltIndex, FMath::Min<int32>(LastUnbuiltIndex - FirstUnbuiltIndex + 10000, LastUnbuiltIndex));
+					InstanceParams.AddRun(RenderData->LODResources.Num() - 1, RenderData->LODResources.Num() - 1, FirstUnbuiltIndex, LastUnbuiltIndex);
 				}
 				FillDynamicMeshElements(Collector, ElementParams, InstanceParams);
 			}
