Different approach to Landscape-specific objects duplication from package during UWorld duplication. Now it is delegated to FWorldDelegates::OnPostDuplicate event. This allows to back out change which was duplicating all objects from package referenced by UWorld, which was bad from the performance and blueprints point of view.

diff --git a/Engine/Source/Runtime/CoreUObject/Private/Serialization/DuplicateDataWriter.cpp b/Engine/Source/Runtime/CoreUObject/Private/Serialization/DuplicateDataWriter.cpp
--- a/Engine/Source/Runtime/CoreUObject/Private/Serialization/DuplicateDataWriter.cpp
+++ b/Engine/Source/Runtime/CoreUObject/Private/Serialization/DuplicateDataWriter.cpp
@@ -30,9 +30,6 @@ FDuplicateDataWriter::FDuplicateDataWriter( FUObjectAnnotationSparse<FDuplicated
 	ArAllowLazyLoading	= false;
 	ArPortFlags |= PPF_Duplicate | InPortFlags;
 
-	SourcePackage = SourceObject->GetOutermost();
-	DestPackage = DestObject->GetOutermost();
-
 	AddDuplicate(SourceObject,DestObject);
 }
 
@@ -61,6 +58,12 @@ FArchive& FDuplicateDataWriter::operator<<(UObject*& Object)
 
 FArchive& FDuplicateDataWriter::operator<<(FLazyObjectPtr& LazyObjectPtr)
 {
+	if (UObject* DuplicatedObject = GetDuplicatedObject(LazyObjectPtr.Get(), false))
+	{
+		FLazyObjectPtr Ptr(DuplicatedObject);
+		return *this << Ptr.GetUniqueID();
+	}
+
 	FUniqueObjectGuid ID = LazyObjectPtr.GetUniqueID();
 	return *this << ID;
 }
@@ -111,7 +114,7 @@ UObject* FDuplicateDataWriter::AddDuplicate(UObject* SourceObject,UObject* DupOb
  *
  * @return	a pointer to the duplicate of the specified object
  */
-UObject* FDuplicateDataWriter::GetDuplicatedObject( UObject* Object )
+UObject* FDuplicateDataWriter::GetDuplicatedObject(UObject* Object, bool bCreateIfMissing)
 {
 	UObject* Result = NULL;
 	if( Object != NULL )
@@ -122,45 +125,16 @@ UObject* FDuplicateDataWriter::GetDuplicatedObject( UObject* Object )
 		{
 			Result = DupObjectInfo.DuplicatedObject;
 		}
-		else
+		else if (bCreateIfMissing)
 		{
-			UObject* NewEmptyDuplicate = nullptr;
-
 			// Check to see if the object's outer is being duplicated.
 			UObject* DupOuter = GetDuplicatedObject(Object->GetOuter());
 			if(DupOuter != NULL)
 			{
-				// The object's outer is being duplicated, create a duplicate
-				// of this object.
-				NewEmptyDuplicate = StaticConstructObject(
-					Object->GetClass(),
-					DupOuter,
-					Object->GetFName(),
-					ApplyFlags|Object->GetMaskedFlags(FlagMask),
+				// The object's outer is being duplicated, create a duplicate of this object.
+				UObject* NewEmptyDuplicate = StaticConstructObject(Object->GetClass(), DupOuter, Object->GetFName(), ApplyFlags|Object->GetMaskedFlags(FlagMask),
 					Object->GetArchetype(), true, InstanceGraph);
 
-			}
-			else if (
-				!Object->HasAnyFlags(RF_ClassDefaultObject) // We shouldn't duplicate CDOs
-				&& (ArPortFlags & PPF_DuplicateForPIE) == 0 // We can also skip this if duplicating for PIE -- this is mainly performance optimization
-				&& SourcePackage != DestPackage
-				&& Object->GetOuter() == SourcePackage
-				&& DestPackage != GetTransientPackage())
-			{
-				// If not then maybe it's a package changing duplication and
-				// the object's outer is the source package itself. Then we
-				// need to duplicate that object also with changed package.
-
-				NewEmptyDuplicate = StaticConstructObject(
-					Object->GetClass(),
-					DestPackage,
-					Object->GetFName(),
-					ApplyFlags | Object->GetMaskedFlags(FlagMask),
-					nullptr, true, InstanceGraph);
-			}
-
-			if (NewEmptyDuplicate != nullptr)
-			{
 				Result = AddDuplicate(Object, NewEmptyDuplicate);
 			}
 		}
diff --git a/Engine/Source/Runtime/CoreUObject/Public/Serialization/ArchiveUObject.h b/Engine/Source/Runtime/CoreUObject/Public/Serialization/ArchiveUObject.h
--- a/Engine/Source/Runtime/CoreUObject/Public/Serialization/ArchiveUObject.h
+++ b/Engine/Source/Runtime/CoreUObject/Public/Serialization/ArchiveUObject.h
@@ -924,18 +924,6 @@ private:
 	EObjectFlags							FlagMask;
 	EObjectFlags							ApplyFlags;
 
-	/**
-	 * A pointer to source package, i.e. to the package the object being
-	 * duplicated is comming from.
-	 */
-	UPackage								*SourcePackage;
-
-	/**
-	 * A pointer to destination package, i.e. to the package the new object
-	 * should be created in.
-	 */
-	UPackage								*DestPackage;
-
 	/**
 	 * This is used to prevent object & component instancing resulting from the calls to StaticConstructObject(); instancing subobjects and components is pointless,
 	 * since we do that manually and replace the current value with our manually created object anyway.
@@ -1002,10 +990,11 @@ public:
 	 * Returns a pointer to the duplicate of a given object, creating the duplicate object if necessary.
 	 * 
 	 * @param	Object	the object to find a duplicate for
+	 * @param	bCreateIfMissing Create the duplicated object if it's missing.
 	 *
 	 * @return	a pointer to the duplicate of the specified object
 	 */
-	UObject* GetDuplicatedObject(UObject* Object);
+	UObject* GetDuplicatedObject(UObject* Object, bool bCreateIfMissing = true);
 
 	/**
 	 * Constructor
diff --git a/Engine/Source/Runtime/Engine/Classes/Engine/World.h b/Engine/Source/Runtime/Engine/Classes/Engine/World.h
--- a/Engine/Source/Runtime/Engine/Classes/Engine/World.h
+++ b/Engine/Source/Runtime/Engine/Classes/Engine/World.h
@@ -2642,7 +2642,22 @@ public:
 	DECLARE_MULTICAST_DELEGATE_TwoParams(FWorldInitializationEvent, UWorld* /*World*/, const UWorld::InitializationValues /*IVS*/);
 	DECLARE_MULTICAST_DELEGATE_ThreeParams(FWorldCleanupEvent, UWorld* /*World*/, bool /*bSessionEnded*/, bool /*bCleanupResources*/);
 	DECLARE_MULTICAST_DELEGATE_OneParam(FWorldEvent, UWorld* /*World*/);
-	DECLARE_MULTICAST_DELEGATE_TwoParams(FWorldPostDuplicateEvent, UWorld* /*World*/, bool /*bDuplicateForPIE*/);
+
+	/**
+	 * Post UWorld duplicate event.
+	 *
+	 * Sometimes there is a need to duplicate additional element after
+	 * duplicating UWorld. If you do this using this event you need also fill
+	 * ReplacementMap and ObjectsToFixReferences in order to properly fix
+	 * duplicated objects references.
+	 */
+	typedef TMap<UObject*, UObject*> FReplacementMap; // Typedef needed so the macro below can properly digest comma in template parameters.
+	DECLARE_MULTICAST_DELEGATE_FourParams(FWorldPostDuplicateEvent, UWorld* /*World*/, bool /*bDuplicateForPIE*/, FReplacementMap& /*ReplacementMap*/, TArray<UObject*>& /*ObjectsToFixReferences*/);
+
+#if WITH_EDITOR
+	DECLARE_MULTICAST_DELEGATE_FiveParams(FWorldRenameEvent, UWorld* /*World*/, const TCHAR* /*InName*/, UObject* /*NewOuter*/, ERenameFlags /*Flags*/, bool& /*bShouldFailRename*/);
+#endif // WITH_EDITOR
+
 	// Delegate type for level change events
 	DECLARE_MULTICAST_DELEGATE_TwoParams(FOnLevelChanged, ULevel*, UWorld*);
 
@@ -2655,6 +2670,11 @@ public:
 	// Callback for world initialization (post)
 	static FWorldInitializationEvent OnPostWorldInitialization;
 
+#if WITH_EDITOR
+	// Callback for world rename event (pre)
+	static FWorldRenameEvent OnPreWorldRename;
+#endif // WITH_EDITOR
+
 	// Post duplication event.
 	static FWorldPostDuplicateEvent OnPostDuplicate;
 
diff --git a/Engine/Source/Runtime/Engine/Private/World.cpp b/Engine/Source/Runtime/Engine/Private/World.cpp
--- a/Engine/Source/Runtime/Engine/Private/World.cpp
+++ b/Engine/Source/Runtime/Engine/Private/World.cpp
@@ -91,6 +91,9 @@ TMap<FName, EWorldType::Type> UWorld::WorldTypePreLoadMap;
 FWorldDelegates::FWorldEvent FWorldDelegates::OnPostWorldCreation;
 FWorldDelegates::FWorldInitializationEvent FWorldDelegates::OnPreWorldInitialization;
 FWorldDelegates::FWorldInitializationEvent FWorldDelegates::OnPostWorldInitialization;
+#if WITH_EDITOR
+FWorldDelegates::FWorldRenameEvent FWorldDelegates::OnPreWorldRename;
+#endif // WITH_EDITOR
 FWorldDelegates::FWorldPostDuplicateEvent FWorldDelegates::OnPostDuplicate;
 FWorldDelegates::FWorldCleanupEvent FWorldDelegates::OnWorldCleanup;
 FWorldDelegates::FWorldEvent FWorldDelegates::OnPreWorldFinishDestroy;
@@ -257,21 +260,12 @@ bool UWorld::Rename(const TCHAR* InName, UObject* NewOuter, ERenameFlags Flags)
 		}
 	}
 
-	// UWorld rename is also renaming package, so move all direct children
-	// of the old package.
+	bool bShouldFail = false;
+	FWorldDelegates::OnPreWorldRename.Broadcast(this, InName, NewOuter, Flags, bShouldFail);
 
-	UPackage* PersistentLevelPackage = PersistentLevel->GetOutermost();
-	if (PersistentLevelPackage != NewOuter)
+	if (bShouldFail)
 	{
-		TArray<UObject*> ReferencedObjects;
-		GetObjectsWithOuter(PersistentLevelPackage, ReferencedObjects, false);
-		for (auto* ReferencedObject : ReferencedObjects)
-		{
-			if (ReferencedObject != this && !ReferencedObject->Rename(nullptr, NewOuter, Flags))
-			{
-				return false;
-			}
-		}
+		return false;
 	}
 
 	// Rename the world itself
@@ -332,6 +326,9 @@ void UWorld::PostDuplicate(bool bDuplicateForPIE)
 {
 	Super::PostDuplicate(bDuplicateForPIE);
 
+	TArray<UObject*> ObjectsToFixReferences;
+	TMap<UObject*, UObject*> ReplacementMap;
+
 	// If we are not duplicating for PIE, fix up objects that travel with the world.
 	// Note that these objects should really be inners of the world, so if they become inners later, most of this code should not be necessary
 	if ( !bDuplicateForPIE )
@@ -354,9 +351,68 @@ void UWorld::PostDuplicate(bool bDuplicateForPIE)
 
 		// Make sure PKG_ContainsMap is set
 		MyPackage->ThisContainsMap();
+
+#if WITH_EDITOR
+		// Add the world to the list of objects in which to fix up references.
+		ObjectsToFixReferences.Add(this);
+
+		// Gather the textures
+		TArray<UTexture2D*> LightMapsAndShadowMaps;
+		GetLightMapsAndShadowMaps(PersistentLevel, LightMapsAndShadowMaps);
+
+		// Duplicate the textures, if any
+		for (auto* Tex : LightMapsAndShadowMaps)
+		{
+			if (Tex && Tex->GetOutermost() != MyPackage)
+			{
+				UObject* NewTex = StaticDuplicateObject(Tex, MyPackage, *Tex->GetName());
+				ReplacementMap.Add(Tex, NewTex);
+			}
+		}
+
+		// Duplicate the level script blueprint generated classes as well
+		const bool bDontCreate = true;
+		UBlueprint* LevelScriptBlueprint = PersistentLevel->GetLevelScriptBlueprint(bDontCreate);
+		if (LevelScriptBlueprint)
+		{
+			UObject* OldGeneratedClass = LevelScriptBlueprint->GeneratedClass;
+			if (OldGeneratedClass)
+			{
+				UObject* NewGeneratedClass = StaticDuplicateObject(OldGeneratedClass, MyPackage, *OldGeneratedClass->GetName());
+				ReplacementMap.Add(OldGeneratedClass, NewGeneratedClass);
+
+				// The class may have referenced a lightmap or landscape resource that is also being duplicated. Add it to the list of objects that need references fixed up.
+				ObjectsToFixReferences.Add(NewGeneratedClass);
+			}
+
+			UObject* OldSkeletonClass = LevelScriptBlueprint->SkeletonGeneratedClass;
+			if (OldSkeletonClass)
+			{
+				UObject* NewSkeletonClass = StaticDuplicateObject(OldSkeletonClass, MyPackage, *OldSkeletonClass->GetName());
+				ReplacementMap.Add(OldSkeletonClass, NewSkeletonClass);
+
+				// The class may have referenced a lightmap or landscape resource that is also being duplicated. Add it to the list of objects that need references fixed up.
+				ObjectsToFixReferences.Add(NewSkeletonClass);
+			}
+		}
+#endif // WITH_EDITOR
 	}
 
-	FWorldDelegates::OnPostDuplicate.Broadcast(this, bDuplicateForPIE);
+	FWorldDelegates::OnPostDuplicate.Broadcast(this, bDuplicateForPIE, ReplacementMap, ObjectsToFixReferences);
+
+#if WITH_EDITOR
+	// Now replace references from the old textures/classes to the new ones, if any were duplicated
+	if (ReplacementMap.Num() > 0)
+	{
+		const bool bNullPrivateRefs = false;
+		const bool bIgnoreOuterRef = true;
+		const bool bIgnoreArchetypeRef = false;
+		for (auto* Obj : ObjectsToFixReferences)
+		{
+			FArchiveReplaceObjectRef<UObject> ReplaceAr(Obj, ReplacementMap, bNullPrivateRefs, bIgnoreOuterRef, bIgnoreArchetypeRef);
+		}
+	}
+#endif // WITH_EDITOR
 }
 
 void UWorld::FinishDestroy()
@@ -2073,7 +2129,7 @@ UWorld* UWorld::DuplicateWorldForPIE(const FString& PackageName, UWorld* OwningW
 	FStringAssetReference::SetPackageNamesBeingDuplicatedForPIE(PackageNamesBeingDuplicatedForPIE);
 
 	ULevel::StreamedLevelsOwningWorld.Add(PIELevelPackage->GetFName(), OwningWorld);
-	UWorld* PIELevelWorld = CastChecked<UWorld>(StaticDuplicateObject(EditorLevelWorld, PIELevelPackage, *EditorLevelWorld->GetName(), RF_AllFlags, NULL, SDO_DuplicateForPie));
+	UWorld* PIELevelWorld = CastChecked<UWorld>(StaticDuplicateObject(EditorLevelWorld, PIELevelPackage, *EditorLevelWorld->GetName(), RF_AllFlags, nullptr, SDO_DuplicateForPie));
 	
 	// Clean up string asset reference fixups
 	FStringAssetReference::ClearPackageNamesBeingDuplicatedForPIE();
diff --git a/Engine/Source/Runtime/Landscape/Private/Landscape.cpp b/Engine/Source/Runtime/Landscape/Private/Landscape.cpp
--- a/Engine/Source/Runtime/Landscape/Private/Landscape.cpp
+++ b/Engine/Source/Runtime/Landscape/Private/Landscape.cpp
@@ -580,11 +580,6 @@ ULevel* ULandscapeComponent::GetLevel() const
 
 void ULandscapeComponent::GetGeneratedTexturesAndMaterialInstances(TArray<UObject*>& OutTexturesAndMaterials) const
 {
-	for (UMaterialInstance* CurrentMIC = MaterialInstance; CurrentMIC && CurrentMIC->IsA<ULandscapeMaterialInstanceConstant>(); CurrentMIC = Cast<UMaterialInstance>(CurrentMIC->Parent))
-	{
-		OutTexturesAndMaterials.Add(CurrentMIC);
-	}
-
 	if (HeightmapTexture)
 	{
 		OutTexturesAndMaterials.Add(HeightmapTexture);
@@ -599,6 +594,32 @@ void ULandscapeComponent::GetGeneratedTexturesAndMaterialInstances(TArray<UObjec
 	{
 		OutTexturesAndMaterials.Add(XYOffsetmapTexture);
 	}
+
+	for (UMaterialInstance* CurrentMIC = MaterialInstance; CurrentMIC && CurrentMIC->IsA<ULandscapeMaterialInstanceConstant>(); CurrentMIC = Cast<UMaterialInstance>(CurrentMIC->Parent))
+	{
+		OutTexturesAndMaterials.Add(CurrentMIC);
+
+		// Sometimes weight map is not registered in the WeightmapTextures, so
+		// we need to get it from here.
+		if (CurrentMIC->IsA<ULandscapeMaterialInstanceConstant>())
+		{
+			auto* LandscapeMIC = Cast<ULandscapeMaterialInstanceConstant>(CurrentMIC);
+
+			auto* WeightmapPtr = LandscapeMIC->TextureParameterValues.FindByPredicate(
+				[](const FTextureParameterValue& ParamValue)
+				{
+					static const FName WeightmapParamName("Weightmap0");
+					return ParamValue.ParameterName == WeightmapParamName;
+				}
+			);
+
+			if (WeightmapPtr != nullptr &&
+				!OutTexturesAndMaterials.Contains(WeightmapPtr->ParameterValue))
+			{
+				OutTexturesAndMaterials.Add(WeightmapPtr->ParameterValue);
+			}
+		}
+	}
 }
 
 ALandscapeProxy* ULandscapeComponent::GetLandscapeProxy() const
diff --git a/Engine/Source/Runtime/Landscape/Private/LandscapeCollision.cpp b/Engine/Source/Runtime/Landscape/Private/LandscapeCollision.cpp
--- a/Engine/Source/Runtime/Landscape/Private/LandscapeCollision.cpp
+++ b/Engine/Source/Runtime/Landscape/Private/LandscapeCollision.cpp
@@ -1315,8 +1315,15 @@ void ULandscapeHeightfieldCollisionComponent::PreSave()
 	{
 #if WITH_EDITOR
 		static FName PhysicsFormatName(FPlatformProperties::GetPhysicsFormat());
-		GetDerivedDataCacheRef().Put(*GetHFDDCKeyString(PhysicsFormatName, false, HeightfieldGuid), CookedCollisionData);
-		GetDerivedDataCacheRef().Put(*GetHFDDCKeyString(PhysicsFormatName, true, HeightfieldGuid), CookedCollisionDataEd);
+		if (CookedCollisionData.Num())
+		{
+			GetDerivedDataCacheRef().Put(*GetHFDDCKeyString(PhysicsFormatName, false, HeightfieldGuid), CookedCollisionData);
+		}
+
+		if (CookedCollisionDataEd.Num())
+		{
+			GetDerivedDataCacheRef().Put(*GetHFDDCKeyString(PhysicsFormatName, true, HeightfieldGuid), CookedCollisionDataEd);
+		}
 #endif// WITH_EDITOR
 	}
 }
diff --git a/Engine/Source/Runtime/Landscape/Private/LandscapeModule.cpp b/Engine/Source/Runtime/Landscape/Private/LandscapeModule.cpp
--- a/Engine/Source/Runtime/Landscape/Private/LandscapeModule.cpp
+++ b/Engine/Source/Runtime/Landscape/Private/LandscapeModule.cpp
@@ -2,6 +2,7 @@
 
 #include "Landscape.h"
 #include "LandscapeModule.h"
+#include "LandscapeComponent.h"
 
 class FLandscapeModule : public IModuleInterface
 {
@@ -58,13 +59,70 @@ void WorldCreationEventFunction(UWorld* World)
  */
 void WorldDestroyEventFunction(UWorld* World);
 
+/**
+ * Gets array of Landscape-specific textures and materials connected with given
+ * level.
+ *
+ * @param Level Level to search textures and materials in.
+ * @param OutTexturesAndMaterials (Output parameter) Array to fill.
+ */
+void GetLandscapeTexturesAndMaterials(ULevel* Level, TArray<UObject*>& OutTexturesAndMaterials)
+{
+	TArray<UObject*> ObjectsInLevel;
+	const bool bIncludeNestedObjects = true;
+	GetObjectsWithOuter(Level, ObjectsInLevel, bIncludeNestedObjects);
+	for (auto* ObjInLevel : ObjectsInLevel)
+	{
+		ULandscapeComponent* LandscapeComponent = Cast<ULandscapeComponent>(ObjInLevel);
+		if (LandscapeComponent)
+		{
+			LandscapeComponent->GetGeneratedTexturesAndMaterialInstances(OutTexturesAndMaterials);
+		}
+	}
+}
+
+/**
+ * A function that fires everytime a world is renamed.
+ *
+ * @param World A world that was renamed.
+ * @param InName New world name.
+ * @param NewOuter New outer of the world after rename.
+ * @param Flags Rename flags.
+ * @param bShouldFailRename (Output parameter) If you set it to true, then the renaming process should fail.
+ */
+void WorldRenameEventFunction(UWorld* World, const TCHAR* InName, UObject* NewOuter, ERenameFlags Flags, bool& bShouldFailRename)
+{
+	// Also rename all textures and materials used by landscape components
+	TArray<UObject*> LandscapeTexturesAndMaterials;
+	GetLandscapeTexturesAndMaterials(World->PersistentLevel, LandscapeTexturesAndMaterials);
+	UPackage* PersistentLevelPackage = World->PersistentLevel->GetOutermost();
+	for (auto* OldTexOrMat : LandscapeTexturesAndMaterials)
+	{
+		if (OldTexOrMat && OldTexOrMat->GetOuter() == PersistentLevelPackage)
+		{
+			// The names for these objects are not important, just generate a new name to avoid collisions
+			if (!OldTexOrMat->Rename(nullptr, NewOuter, Flags))
+			{
+				bShouldFailRename = true;
+			}
+		}
+	}
+}
+
 /**
  * A function that fires everytime a world is duplicated.
  *
+ * If there are some objects duplicated during this event fill out
+ * ReplacementMap and ObjectsToFixReferences in order to properly fix
+ * references in objects created during this duplication.
+ *
  * @param World A world that was duplicated.
  * @param bDuplicateForPIE If this duplication was done for PIE.
+ * @param ReplacementMap Replacement map (i.e. old object -> new object).
+ * @param ObjectsToFixReferences Array of objects that may contain bad
+ *		references to old objects.
  */
-void WorldDuplicateEventFunction(UWorld* World, bool bDuplicateForPIE)
+void WorldDuplicateEventFunction(UWorld* World, bool bDuplicateForPIE, TMap<UObject*, UObject*>& ReplacementMap, TArray<UObject*>& ObjectsToFixReferences)
 {
 	int32 Index;
 	ULandscapeInfoMap *InfoMap;
@@ -78,6 +136,32 @@ void WorldDuplicateEventFunction(UWorld* World, bool bDuplicateForPIE)
 	{
 		AddPerWorldLandscapeData(World);
 	}
+
+#if WITH_EDITOR
+	if (!bDuplicateForPIE)
+	{
+		UPackage* WorldPackage = World->GetOutermost();
+
+		// Also duplicate all textures and materials used by landscape components
+		TArray<UObject*> LandscapeTexturesAndMaterials;
+		GetLandscapeTexturesAndMaterials(World->PersistentLevel, LandscapeTexturesAndMaterials);
+		for (auto* OldTexOrMat : LandscapeTexturesAndMaterials)
+		{
+			if (OldTexOrMat && OldTexOrMat->GetOuter() != WorldPackage)
+			{
+				// The names for these objects are not important, just generate a new name to avoid collisions
+				UObject* NewTextureOrMaterial = StaticDuplicateObject(OldTexOrMat, WorldPackage, nullptr);
+				ReplacementMap.Add(OldTexOrMat, NewTextureOrMaterial);
+
+				// Materials hold references to the textures being moved, so they will need references fixed up as well
+				if (OldTexOrMat->IsA(UMaterialInterface::StaticClass()))
+				{
+					ObjectsToFixReferences.Add(NewTextureOrMaterial);
+				}
+			}
+		}
+	}
+#endif // WITH_EDITOR
 }
 
 void FLandscapeModule::StartupModule()
@@ -102,6 +186,12 @@ void FLandscapeModule::StartupModule()
 	);
 #endif // WITH_EDITORONLY_DATA
 
+#if WITH_EDITOR
+	FWorldDelegates::OnPreWorldRename.AddStatic(
+		&WorldRenameEventFunction
+	);
+#endif // WITH_EDITOR
+
 	FWorldDelegates::OnPostDuplicate.AddStatic(
 		&WorldDuplicateEventFunction
 	);
