Fix undo issues with hierarchical instanced foliage

diff --git a/Engine/Source/Runtime/Engine/Private/InstancedFoliage.cpp b/Engine/Source/Runtime/Engine/Private/InstancedFoliage.cpp
--- a/Engine/Source/Runtime/Engine/Private/InstancedFoliage.cpp
+++ b/Engine/Source/Runtime/Engine/Private/InstancedFoliage.cpp
@@ -295,10 +295,11 @@ void FFoliageMeshInfo::CheckValid()
 
 void FFoliageMeshInfo::AddInstance(AInstancedFoliageActor* InIFA, UFoliageType* InSettings, const FFoliageInstance& InNewInstance)
 {
-	// Want to make sure the component exists before we start a transaction, so it will be there on undo.
+	InIFA->Modify();
+
 	if (Component == nullptr)
 	{
-		Component = ConstructObject<UHierarchicalInstancedStaticMeshComponent>(UHierarchicalInstancedStaticMeshComponent::StaticClass(), InIFA);
+		Component = ConstructObject<UHierarchicalInstancedStaticMeshComponent>(UHierarchicalInstancedStaticMeshComponent::StaticClass(), InIFA, NAME_None, RF_Transactional);
 
 		Component->Mobility = EComponentMobility::Static;
 
@@ -326,14 +327,17 @@ void FFoliageMeshInfo::AddInstance(AInstancedFoliageActor* InIFA, UFoliageType*
 
 		// Use only instance translation as a component transform
 		Component->SetWorldTransform(InIFA->GetRootComponent()->ComponentToWorld);
+
+		// Add the new component to the transaction buffer so it will get destroyed on undo
+		Component->Modify();
+		// We don't want to track changes to instances later so we mark it as non-transactional
+		Component->ClearFlags(RF_Transactional);
 	}
 	else
 	{
 		Component->InvalidateLightingCache();
 	}
 
-	InIFA->Modify();
-
 	// Add the instance taking either a free slot or adding a new item.
 	int32 InstanceIndex = Instances.AddUninitialized();
 
