Implemented text scaling for composite fonts in Canvas

This was implemented for the bitmap fonts, but hadn't been implemented for composite fonts.

It scales the glyph texture when rendering, rather rendering the glyph with an updated scale. This saves us potentially thrashing the font cache (and having to implement non-uniform text scaling for the font cache), at the cost of text quality.

If you need crisp text, consider scaling the font size rather than using render scaling.

diff --git a/Engine/Source/Runtime/Engine/Private/UserInterface/CanvasItem.cpp b/Engine/Source/Runtime/Engine/Private/UserInterface/CanvasItem.cpp
--- a/Engine/Source/Runtime/Engine/Private/UserInterface/CanvasItem.cpp
+++ b/Engine/Source/Runtime/Engine/Private/UserInterface/CanvasItem.cpp
@@ -1038,8 +1038,6 @@ void FCanvasTextItem::DrawStringInternal_RuntimeCache( FCanvas* InCanvas, const
 	const FSlateFontInfo LegacyFontInfo = (SlateFontInfo.IsSet()) ? SlateFontInfo.GetValue() : Font->GetLegacySlateFontInfo();
 	FCharacterList& CharacterList = FontCache->GetCharacterList( LegacyFontInfo, FontScale );
 
-	const float MaxHeight = CharacterList.GetMaxHeight();
-
 	FHitProxyId HitProxyId = InCanvas->GetHitProxyId();
 
 	uint32 FontTextureIndex = 0;
@@ -1060,6 +1058,7 @@ void FCanvasTextItem::DrawStringInternal_RuntimeCache( FCanvas* InCanvas, const
 	float PosY = TopLeft.Y;
 
 	const float ScaledHorizSpacingAdjust = HorizSpacingAdjust * Scale.X;
+	const float ScaledMaxHeight = CharacterList.GetMaxHeight() * Scale.Y;
 
 	LineX = PosX;
 	
@@ -1070,7 +1069,7 @@ void FCanvasTextItem::DrawStringInternal_RuntimeCache( FCanvas* InCanvas, const
 		if (DrawnSize.Y == 0)
 		{
 			// We have a valid character so initialize vertical DrawnSize
-			DrawnSize.Y = MaxHeight;
+			DrawnSize.Y = ScaledMaxHeight;
 		}
 
 		const bool IsNewline = (CurrentChar == '\n');
@@ -1078,11 +1077,11 @@ void FCanvasTextItem::DrawStringInternal_RuntimeCache( FCanvas* InCanvas, const
 		if (IsNewline)
 		{
 			// Move down: we are drawing the next line.
-			PosY += MaxHeight;
+			PosY += ScaledMaxHeight;
 			// Carriage return 
 			LineX = PosX;
 			// Increase the vertical DrawnSize
-			DrawnSize.Y += MaxHeight;
+			DrawnSize.Y += ScaledMaxHeight;
 		}
 		else
 		{
@@ -1119,14 +1118,14 @@ void FCanvasTextItem::DrawStringInternal_RuntimeCache( FCanvas* InCanvas, const
 
 			if( !bIsWhitespace )
 			{
-				const float X = LineX + Entry.HorizontalOffset + DrawPos.X;
+				const float X = DrawPos.X + LineX + (Entry.HorizontalOffset * Scale.X);
 				// Note PosX,PosY is the upper left corner of the bounding box representing the string.  This computes the Y position of the baseline where text will sit
 
-				const float Y = PosY - Entry.VerticalOffset + MaxHeight + Entry.GlobalDescender + DrawPos.Y;
+				const float Y = DrawPos.Y + PosY - (Entry.VerticalOffset * Scale.Y) + (Entry.GlobalDescender * Scale.Y) + ScaledMaxHeight;
 				const float U = Entry.StartU * InvTextureSizeX;
 				const float V = Entry.StartV * InvTextureSizeY;
-				const float SizeX = Entry.USize;
-				const float SizeY = Entry.VSize;
+				const float SizeX = Entry.USize * Scale.X;
+				const float SizeY = Entry.VSize * Scale.Y;
 				const float SizeU = Entry.USize * InvTextureSizeX;
 				const float SizeV = Entry.VSize * InvTextureSizeY;
 
@@ -1161,7 +1160,7 @@ void FCanvasTextItem::DrawStringInternal_RuntimeCache( FCanvas* InCanvas, const
 				BatchedElements->AddTriangle(V00, V11, V01, FontTexture, BlendMode, FontRenderInfo.GlowInfo);
 			}
 
-			LineX += Entry.XAdvance;
+			LineX += Entry.XAdvance * Scale.X;
 			LineX += ScaledHorizSpacingAdjust;
 
 			// Increase the Horizontal DrawnSize
