Fix double buffering issue that causes skinned assets to collapse

diff --git a/Engine/Source/Runtime/Engine/Classes/Components/SkinnedMeshComponent.h b/Engine/Source/Runtime/Engine/Classes/Components/SkinnedMeshComponent.h
--- a/Engine/Source/Runtime/Engine/Classes/Components/SkinnedMeshComponent.h
+++ b/Engine/Source/Runtime/Engine/Classes/Components/SkinnedMeshComponent.h
@@ -572,6 +572,10 @@ public:
 	void SetSpaceBaseDoubleBuffering(bool bInDoubleBufferedBlendSpaces);
 
 protected:
+
+	/** Track whether we still need to flip to recently modified buffer */
+	bool bNeedToFlipSpaceBaseBuffers;
+
 	/**
 	* Combine CurveKeys (that reference morph targets by name) and ActiveAnims (that reference vertex anims by reference) into the ActiveVertexAnims array.
 	*/
diff --git a/Engine/Source/Runtime/Engine/Private/PoseableMeshComponent.cpp b/Engine/Source/Runtime/Engine/Private/PoseableMeshComponent.cpp
--- a/Engine/Source/Runtime/Engine/Private/PoseableMeshComponent.cpp
+++ b/Engine/Source/Runtime/Engine/Private/PoseableMeshComponent.cpp
@@ -35,6 +35,7 @@ bool UPoseableMeshComponent::AllocateTransformData()
 		}
 
 		FillSpaceBases();
+		FlipEditableSpaceBases();
 
 		return true;
 	}
@@ -62,6 +63,7 @@ void UPoseableMeshComponent::RefreshBoneTransforms(FActorComponentTickFunction*
 
 	// We need the mesh space bone transforms now for renderer to get delta from ref pose:
 	FillSpaceBases();
+	FlipEditableSpaceBases();
 
 	MarkRenderDynamicDataDirty();
 }
@@ -117,6 +119,7 @@ void UPoseableMeshComponent::FillSpaceBases()
 		checkSlow(GetEditableSpaceBases()[BoneIndex].IsRotationNormalized());
 		checkSlow(!GetEditableSpaceBases()[BoneIndex].ContainsNaN());
 	}
+	bNeedToFlipSpaceBaseBuffers = true;
 }
 
 void UPoseableMeshComponent::SetBoneTransformByName(FName BoneName, const FTransform& InTransform, EBoneSpaces::Type BoneSpace)
diff --git a/Engine/Source/Runtime/Engine/Private/SkeletalMesh.cpp b/Engine/Source/Runtime/Engine/Private/SkeletalMesh.cpp
--- a/Engine/Source/Runtime/Engine/Private/SkeletalMesh.cpp
+++ b/Engine/Source/Runtime/Engine/Private/SkeletalMesh.cpp
@@ -4238,9 +4238,6 @@ void FSkeletalMeshSceneProxy::GetMeshElementsConditionallySelectable(const TArra
 	{
 		if (VisibilityMap & (1 << ViewIndex))
 		{
-			// debug drawing
-			TArray<FTransform>* BoneSpaceBases = MeshObject->GetSpaceBases();
-
 			if( PhysicsAssetForDebug )
 			{
 				DebugDrawPhysicsAsset(ViewIndex, Collector, ViewFamily.EngineShowFlags);
@@ -4550,6 +4547,7 @@ USkinnedMeshComponent::USkinnedMeshComponent(const FObjectInitializer& ObjectIni
 	bDoubleBufferedBlendSpaces = true;
 	CurrentEditableSpaceBases = 0;
 	CurrentReadSpaceBases = 1;
+	bNeedToFlipSpaceBaseBuffers = false;
 }
 
 
diff --git a/Engine/Source/Runtime/Engine/Private/SkeletalMeshComponent.cpp b/Engine/Source/Runtime/Engine/Private/SkeletalMeshComponent.cpp
--- a/Engine/Source/Runtime/Engine/Private/SkeletalMeshComponent.cpp
+++ b/Engine/Source/Runtime/Engine/Private/SkeletalMeshComponent.cpp
@@ -232,6 +232,7 @@ void USkeletalMeshComponent::InitAnim(bool bForceReinit)
         TickAnimation(0.f); 
 
 		RefreshBoneTransforms();
+		FlipEditableSpaceBases();
 		UpdateComponentToWorld();
 	}
 }
@@ -1026,6 +1027,7 @@ void USkeletalMeshComponent::PostAnimEvaluation(FAnimationEvaluationContext& Eva
 		}
 		FAnimationRuntime::LerpBoneTransforms(GetEditableSpaceBases(), CachedSpaceBases, Alpha, RequiredBones);
 	}
+	bNeedToFlipSpaceBaseBuffers = true;
 
 	// Transforms updated, cached local bounds are now out of date.
 	InvalidateCachedBounds();
diff --git a/Engine/Source/Runtime/Engine/Private/SkinnedMeshComponent.cpp b/Engine/Source/Runtime/Engine/Private/SkinnedMeshComponent.cpp
--- a/Engine/Source/Runtime/Engine/Private/SkinnedMeshComponent.cpp
+++ b/Engine/Source/Runtime/Engine/Private/SkinnedMeshComponent.cpp
@@ -1836,14 +1836,18 @@ void USkinnedMeshComponent::AnimUpdateRateSetParams(const bool & bRecentlyRender
 
 void USkinnedMeshComponent::FlipEditableSpaceBases()
 {
-	if (bDoubleBufferedBlendSpaces)
+	if (bNeedToFlipSpaceBaseBuffers)
 	{
-		CurrentReadSpaceBases = CurrentEditableSpaceBases;
-		CurrentEditableSpaceBases = 1 - CurrentEditableSpaceBases;
-	}
-	else
-	{
-		CurrentReadSpaceBases = CurrentEditableSpaceBases = 0;
+		bNeedToFlipSpaceBaseBuffers = false;
+		if (bDoubleBufferedBlendSpaces)
+		{
+			CurrentReadSpaceBases = CurrentEditableSpaceBases;
+			CurrentEditableSpaceBases = 1 - CurrentEditableSpaceBases;
+		}
+		else
+		{
+			CurrentReadSpaceBases = CurrentEditableSpaceBases = 0;
+		}
 	}
 }
 
