Fix typo where we made a copy of the RefSkeleton instead of just getting a reference to it

diff --git a/Engine/Source/Runtime/Engine/Private/Animation/AnimSequence.cpp b/Engine/Source/Runtime/Engine/Private/Animation/AnimSequence.cpp
--- a/Engine/Source/Runtime/Engine/Private/Animation/AnimSequence.cpp
+++ b/Engine/Source/Runtime/Engine/Private/Animation/AnimSequence.cpp
@@ -745,7 +745,7 @@ FTransform UAnimSequence::ExtractRootTrackTransform(float Pos, const FBoneContai
 	// If we don't have a RequiredBones array, get root bone from default skeleton.
 	if( !RequiredBones &&  MySkeleton )
 	{
-		const FReferenceSkeleton RefSkeleton = MySkeleton->GetReferenceSkeleton();
+		const FReferenceSkeleton& RefSkeleton = MySkeleton->GetReferenceSkeleton();
 		if( RefSkeleton.GetNum() > 0 )
 		{
 			return RefSkeleton.GetRefBonePose()[0];
