UE-6207 - Fix mipbias being applied twice

diff --git a/Engine/Source/Runtime/Engine/Private/TextureDerivedData.cpp b/Engine/Source/Runtime/Engine/Private/TextureDerivedData.cpp
--- a/Engine/Source/Runtime/Engine/Private/TextureDerivedData.cpp
+++ b/Engine/Source/Runtime/Engine/Private/TextureDerivedData.cpp
@@ -1763,6 +1763,11 @@ void UTexture::SerializeCookedPlatformData(FArchive& Ar)
 			Ar << PixelFormatName;
 		}
 	}
+
+	if( Ar.IsLoading() )
+	{
+		LODBias = 0;
+	}
 }
 
 int32 UTexture2D::GMinTextureResidentMipCount = NUM_INLINE_DERIVED_MIPS;
