Fixed UFont::GetStringHeightAndWidth accumulating height rather than return the maximum height of any character in the string

diff --git a/Engine/Source/Runtime/Engine/Private/Font.cpp b/Engine/Source/Runtime/Engine/Private/Font.cpp
--- a/Engine/Source/Runtime/Engine/Private/Font.cpp
+++ b/Engine/Source/Runtime/Engine/Private/Font.cpp
@@ -301,7 +301,7 @@ void UFont::GetStringHeightAndWidth( const TCHAR *Text, int32& Height, int32& Wi
 		}
 
 		TotalWidth += TmpWidth + CharKerning;
-		TotalHeight += TmpHeight;
+		TotalHeight = FMath::Max( TotalHeight, TmpHeight );
 
 		PrevChar = Text++;
 	}
