Fixed notify string lists not correctly parsing on Mac
[UE-6475]

diff --git a/Engine/Source/Editor/Persona/Private/SSkeletonSmartNameManager.cpp b/Engine/Source/Editor/Persona/Private/SSkeletonSmartNameManager.cpp
--- a/Engine/Source/Editor/Persona/Private/SSkeletonSmartNameManager.cpp
+++ b/Engine/Source/Editor/Persona/Private/SSkeletonSmartNameManager.cpp
@@ -405,7 +405,7 @@ void SCurveNameManager::OnDeleteNameClicked()
 			}
 
 			TArray<FString> ParsedStringUids;
-			CurveDataCopy.ParseIntoArray(&ParsedStringUids, &USkeleton::CurveTagDelimiter, true);
+			CurveDataCopy.ParseIntoArray(&ParsedStringUids, *USkeleton::CurveTagDelimiter, true);
 
 			for(const FString& UidString : ParsedStringUids)
 			{
diff --git a/Engine/Source/Runtime/Engine/Classes/Animation/Skeleton.h b/Engine/Source/Runtime/Engine/Classes/Animation/Skeleton.h
--- a/Engine/Source/Runtime/Engine/Classes/Animation/Skeleton.h
+++ b/Engine/Source/Runtime/Engine/Classes/Animation/Skeleton.h
@@ -658,11 +658,11 @@ public:
 
 	// Asset registry information for animation notifies
 	static const FName AnimNotifyTag;
-	static const TCHAR AnimNotifyTagDelimiter;
+	static const FString AnimNotifyTagDelimiter;
 
 	// Asset registry information for animation curves
 	ENGINE_API static const FName CurveTag;
-	ENGINE_API static const TCHAR CurveTagDelimiter;
+	ENGINE_API static const FString CurveTagDelimiter;
 
 	// rig Configs
 	ENGINE_API static const FName RigTag;
diff --git a/Engine/Source/Runtime/Engine/Private/Animation/AnimSequenceBase.cpp b/Engine/Source/Runtime/Engine/Private/Animation/AnimSequenceBase.cpp
--- a/Engine/Source/Runtime/Engine/Private/Animation/AnimSequenceBase.cpp
+++ b/Engine/Source/Runtime/Engine/Private/Animation/AnimSequenceBase.cpp
@@ -794,7 +794,7 @@ void UAnimSequenceBase::GetAssetRegistryTags(TArray<FAssetRegistryTag>& OutTags)
 			// only add if not BP anim notify since they're handled separate
 			if(Iter->IsBlueprintNotify() == false)
 			{
-				NotifyList += FString::Printf(TEXT("%s%c"), *Iter->NotifyName.ToString(), USkeleton::AnimNotifyTagDelimiter);
+				NotifyList += FString::Printf(TEXT("%s%s"), *Iter->NotifyName.ToString(), *USkeleton::AnimNotifyTagDelimiter);
 			}
 		}
 		
@@ -812,7 +812,7 @@ void UAnimSequenceBase::GetAssetRegistryTags(TArray<FAssetRegistryTag>& OutTags)
 
 	for(const FFloatCurve& Curve : RawCurveData.FloatCurves)
 	{
-		CurveIdList += FString::Printf(TEXT("%u%c"), Curve.CurveUid, USkeleton::CurveTagDelimiter);
+		CurveIdList += FString::Printf(TEXT("%u%s"), Curve.CurveUid, *USkeleton::CurveTagDelimiter);
 	}
 	OutTags.Add(FAssetRegistryTag(USkeleton::CurveTag, CurveIdList, FAssetRegistryTag::TT_Hidden));
 }
diff --git a/Engine/Source/Runtime/Engine/Private/Animation/Skeleton.cpp b/Engine/Source/Runtime/Engine/Private/Animation/Skeleton.cpp
--- a/Engine/Source/Runtime/Engine/Private/Animation/Skeleton.cpp
+++ b/Engine/Source/Runtime/Engine/Private/Animation/Skeleton.cpp
@@ -16,10 +16,10 @@
 
 #if WITH_EDITOR
 const FName USkeleton::AnimNotifyTag = FName(TEXT("AnimNotifyList"));
-const TCHAR USkeleton::AnimNotifyTagDelimiter = TEXT(';');
+const FString USkeleton::AnimNotifyTagDelimiter = TEXT(";");
 
 const FName USkeleton::CurveTag = FName(TEXT("CurveUIDList"));
-const TCHAR USkeleton::CurveTagDelimiter = TEXT(';');
+const FString USkeleton::CurveTagDelimiter = TEXT(";");
 
 const FName USkeleton::RigTag = FName(TEXT("Rig"));
 #endif 
@@ -703,7 +703,7 @@ void USkeleton::CollectAnimationNotifies()
 			if (const FString* Value = Asset.TagsAndValues.Find(USkeleton::AnimNotifyTag))
 			{
 				TArray<FString> NotifyList;
-				Value->ParseIntoArray(&NotifyList, &AnimNotifyTagDelimiter, true);
+				Value->ParseIntoArray(&NotifyList, *USkeleton::AnimNotifyTagDelimiter, true);
 				for (auto NotifyIter = NotifyList.CreateConstIterator(); NotifyIter; ++NotifyIter)
 				{
 					FString NotifyName = *NotifyIter;
