Implemented bulk initialisation of physics bodies, currently used by instance static meshes

diff --git a/Engine/Source/Runtime/Core/Private/Misc/Core.cpp b/Engine/Source/Runtime/Core/Private/Misc/Core.cpp
--- a/Engine/Source/Runtime/Core/Private/Misc/Core.cpp
+++ b/Engine/Source/Runtime/Core/Private/Misc/Core.cpp
@@ -256,6 +256,7 @@ DEFINE_STAT(STAT_SkeletalMeshMotionBlurSkinningMemory);
 DEFINE_STAT(STAT_VertexShaderMemory);
 DEFINE_STAT(STAT_PixelShaderMemory);
 DEFINE_STAT(STAT_NavigationMemory);
+DEFINE_STAT(STAT_PhysSceneLock);
 
 DEFINE_STAT(STAT_ReflectionCaptureTextureMemory);
 DEFINE_STAT(STAT_ReflectionCaptureMemory);
diff --git a/Engine/Source/Runtime/Core/Public/Misc/CoreStats.h b/Engine/Source/Runtime/Core/Public/Misc/CoreStats.h
--- a/Engine/Source/Runtime/Core/Public/Misc/CoreStats.h
+++ b/Engine/Source/Runtime/Core/Public/Misc/CoreStats.h
@@ -35,6 +35,7 @@ DECLARE_MEMORY_STAT_EXTERN(TEXT("Navigation Memory"),STAT_NavigationMemory,STATG
 DECLARE_MEMORY_STAT_EXTERN(TEXT("PhysX Memory Used"),STAT_MemoryPhysXTotalAllocationSize,STATGROUP_Memory, CORE_API);
 DECLARE_MEMORY_STAT_EXTERN(TEXT("ICU Memory Used"),STAT_MemoryICUTotalAllocationSize,STATGROUP_Memory, CORE_API);
 DECLARE_MEMORY_STAT_EXTERN(TEXT("ICU Data File Memory Used"),STAT_MemoryICUDataFileAllocationSize,STATGROUP_Memory, CORE_API);
+DECLARE_CYCLE_STAT_EXTERN(TEXT("PhysX Scene Lock"), STAT_PhysSceneLock, STATGROUP_Physics, CORE_API);
 
 DECLARE_MEMORY_STAT_EXTERN(TEXT("StaticMesh Total Memory"),STAT_StaticMeshTotalMemory2,STATGROUP_MemoryStaticMesh, CORE_API);
 DECLARE_MEMORY_STAT_EXTERN(TEXT("StaticMesh Vertex Memory"),STAT_StaticMeshVertexMemory,STATGROUP_MemoryStaticMesh, CORE_API);
diff --git a/Engine/Source/Runtime/Engine/Classes/PhysicsEngine/BodyInstance.h b/Engine/Source/Runtime/Engine/Classes/PhysicsEngine/BodyInstance.h
--- a/Engine/Source/Runtime/Engine/Classes/PhysicsEngine/BodyInstance.h
+++ b/Engine/Source/Runtime/Engine/Classes/PhysicsEngine/BodyInstance.h
@@ -29,6 +29,7 @@ namespace physx
 	class PxRigidDynamic;
 	class PxGeometry;
 	class PxShape;
+	class PxMaterial;
 	class PxTransform;
 	struct PxContactPair;
 }
@@ -351,13 +352,48 @@ public:
 #endif
 
 #if UE_WITH_PHYSICS
+	/** Initialise a single rigid body (this FBodyInstance) for the given body setup
+	 *	@param Setup The setup to use to create the body
+	 *	@param Transform Transform of the body
+	 *	@param PrimComp The owning component
+	 *	@param InRBScene The physics scene to place the body into
+	 *	@param InAggregate An aggregate to place the body into
+	 */
 	void InitBody(class UBodySetup* Setup, const FTransform& Transform, class UPrimitiveComponent* PrimComp, class FPhysScene* InRBScene, PhysXAggregateType InAggregate = NULL);
-#endif	//WITH_PHYSX
 
-	void TermBody();
+	/** Initialise multiple rigid bodies 
+	 *	@param Bodies List of body instances to initialise
+	 *	@param Transforms List of transforms to use for the bodies
+	 *	@param BodySetup The setup to use for the bodies
+	 *	@param PrimitiveComp The owning component
+	 *	@param InRBScene The physics scene to put the bodies into
+	 *	@param InAggregate Aggregate to place the bodies into, be aware for PhysX we can only have 128 elements in an aggregate!
+	 */
+	static void InitBodies(TArray<FBodyInstance*>& Bodies, TArray<FTransform>& Transforms, class UBodySetup* BodySetup, class UPrimitiveComponent* PrimitiveComp, class FPhysScene* InRBScene, PhysXAggregateType InAggregate = NULL);
+
+	/** Initialises one or more rigid bodies
+	 *	@param Bodies First body instance to initialise
+	 *	@param Transforms First transform to use
+	 *	@param NumBodies Number of bodies to iterate
+	 *	@param BodySetup Setup to use for bodies
+	 *	@param PrimitiveComp The owning component for the bodies
+	 *	@param InRBScene Scene to place the bodies into
+	 *	@param InAggregate Aggregate to place the bodies into, be aware for PhysX we can only have 128 elements in an aggregate!
+	 */
+	static void InitBodies(FBodyInstance** Bodies, FTransform* Transforms, int32 NumBodies, class UBodySetup* BodySetup, class UPrimitiveComponent* PrimitiveComp, class FPhysScene* InRBScene, PhysXAggregateType InAggregate = NULL);
+
+	void InitBody_Legacy(class UBodySetup* Setup, const FTransform& Transform, class UPrimitiveComponent* PrimComp, class FPhysScene* InRBScene, PhysXAggregateType InAggregate = NULL);
 
+	/** Validate a body transform, outputting debug info
+	 *	@param Transform Transform to debug
+	 *	@param DebugName Name of the instance for logging
+	 *	@param Setup Body setup for this instance
+	 */
+	static bool ValidateTransform(const FTransform &Transform, FString& DebugName, UBodySetup* Setup);
 
+#endif	//UE_WITH_PHYSICS
 
+	void TermBody();
 
 	/** 
 	 * Takes two body instances and welds them together to create a single simulated rigid body. Returns true if success.
@@ -561,6 +597,12 @@ public:
 	/** Update the instance's material properties (friction, restitution) */
 	void UpdatePhysicalMaterials();
 
+#if WITH_PHYSX
+	static void ApplyMaterialToShape(physx::PxShape* PShape, physx::PxMaterial* PSimpleMat, TArray<UPhysicalMaterial*>& ComplexPhysMats);
+
+	void ApplyMaterialToInstanceShapes(physx::PxMaterial* PSimpleMat, TArray<UPhysicalMaterial*> ComplexPhysMats);
+#endif
+
 	/** Update the instances collision filtering data */
 	void UpdatePhysicsFilterData(bool bForceSimpleAsComplex = false);
 
diff --git a/Engine/Source/Runtime/Engine/Private/InstancedStaticMesh.cpp b/Engine/Source/Runtime/Engine/Private/InstancedStaticMesh.cpp
--- a/Engine/Source/Runtime/Engine/Private/InstancedStaticMesh.cpp
+++ b/Engine/Source/Runtime/Engine/Private/InstancedStaticMesh.cpp
@@ -879,10 +879,25 @@ void UInstancedStaticMeshComponent::CreateAllInstanceBodies()
 	int32 NumBodies = PerInstanceSMData.Num();
 	InstanceBodies.Init(NumBodies);
 
+	TArray<FTransform> Transforms;
+	Transforms.Reserve(NumBodies);
+
 	for (int32 i = 0; i < NumBodies; ++i)
 	{
 		InstanceBodies[i] = new FBodyInstance;
-		InitInstanceBody(i, InstanceBodies[i]);
+		FBodyInstance* Instance = InstanceBodies[i];
+		Transforms.Add(FTransform(PerInstanceSMData[i].Transform) * ComponentToWorld);
+
+		Instance->CopyBodyInstancePropertiesFrom(&BodyInstance);
+		Instance->InstanceBodyIndex = i; // Set body index 
+
+		// make sure we never enable bSimulatePhysics for ISMComps
+		Instance->bSimulatePhysics = false;
+	}
+
+	if(NumBodies > 0)
+	{
+		FBodyInstance::InitBodies(InstanceBodies, Transforms, GetBodySetup(), this, GetWorld()->GetPhysicsScene());
 	}
 }
 
diff --git a/Engine/Source/Runtime/Engine/Private/PhysicsEngine/BodyInstance.cpp b/Engine/Source/Runtime/Engine/Private/PhysicsEngine/BodyInstance.cpp
--- a/Engine/Source/Runtime/Engine/Private/PhysicsEngine/BodyInstance.cpp
+++ b/Engine/Source/Runtime/Engine/Private/PhysicsEngine/BodyInstance.cpp
@@ -354,10 +354,15 @@ void FBodyInstance::UpdateTriMeshVertices(const TArray<FVector> & NewPositions)
 #endif
 }
 
+DEFINE_STAT(STAT_UpdatePhysMats);
+DEFINE_STAT(STAT_UpdatePhysMatsSceneInteraction);
+
 void FBodyInstance::UpdatePhysicalMaterials()
 {
+	SCOPE_CYCLE_COUNTER(STAT_UpdatePhysMats);
 	UPhysicalMaterial* SimplePhysMat = GetSimplePhysicalMaterial();
 	check(SimplePhysMat != NULL);
+	TArray<UPhysicalMaterial*> ComplexPhysMats = GetComplexPhysicalMaterials();
 
 #if WITH_PHYSX
 	PxMaterial* PSimpleMat = SimplePhysMat->GetPhysXMaterial();
@@ -373,44 +378,7 @@ void FBodyInstance::UpdatePhysicalMaterials()
 		SCENE_LOCK_WRITE(RigidActorAsync->getScene());
 	}
 
-	int32 NumSyncShapes = 0;
-	TArray<PxShape*> AllShapes = GetAllShapes(NumSyncShapes);
-	
-	for(int32 ShapeIdx=0; ShapeIdx<AllShapes.Num(); ShapeIdx++)
-	{
-		PxShape* PShape = AllShapes[ShapeIdx];
-
-		// If a triangle mesh, need to get array of materials...
-		if(PShape->getGeometryType() == PxGeometryType::eTRIANGLEMESH)
-		{
-			TArray<UPhysicalMaterial*> ComplexPhysMats = GetComplexPhysicalMaterials();
-			TArray<PxMaterial*> PComplexMats;
-			PComplexMats.AddUninitialized( ComplexPhysMats.Num() );
-			for(int MatIdx=0; MatIdx<ComplexPhysMats.Num(); MatIdx++)
-			{
-				check(ComplexPhysMats[MatIdx] != NULL);
-				PComplexMats[MatIdx] = ComplexPhysMats[MatIdx]->GetPhysXMaterial();
-				check(PComplexMats[MatIdx] != NULL);
-			}
-
-			if (PComplexMats.Num())
-			{
-				PShape->setMaterials(PComplexMats.GetData(), PComplexMats.Num());
-			}
-			else
-			{
-				UE_LOG(LogPhysics, Warning, TEXT("FBodyInstance::UpdatePhysicalMaterials : PComplexMats is empty - falling back on simple physical material."));
-				PShape->setMaterials(&PSimpleMat, 1);
-			}
-
-		}
-		// Simple shape, 
-		else
-		{
-			PShape->setMaterials(&PSimpleMat, 1);
-		}
-	}
-
+	ApplyMaterialToInstanceShapes(PSimpleMat, ComplexPhysMats);
 
 	if (RigidActorSync != NULL)
 	{
@@ -907,10 +875,33 @@ void FBodyInstance::UpdatePhysicsFilterData(bool bForceSimpleAsComplex)
 #endif
 }
 
+DEFINE_STAT(STAT_InitBody);
+DEFINE_STAT(STAT_InitBodyDebug);
+DEFINE_STAT(STAT_InitBodySceneInteraction);
+DEFINE_STAT(STAT_InitBodyPostAdd);
+DEFINE_STAT(STAT_UpdatePhysFilter);
+
+DEFINE_STAT(STAT_BulkInit);
+DEFINE_STAT(STAT_BulkSceneAdd);
 
 #if UE_WITH_PHYSICS
-void FBodyInstance::InitBody(UBodySetup* Setup, const FTransform& Transform, UPrimitiveComponent* PrimComp, FPhysScene* InRBScene, PxAggregate* InAggregate)
+
+void FBodyInstance::InitBody(class UBodySetup* Setup, const FTransform& Transform, class UPrimitiveComponent* PrimComp, class FPhysScene* InRBScene, PhysXAggregateType InAggregate /*= NULL*/ )
+{
+	FBodyInstance* Instance = this;
+	FTransform TransformCopy = Transform;
+	InitBodies(&Instance, &TransformCopy, 1, Setup, PrimComp, InRBScene, InAggregate);
+}
+
+void FBodyInstance::InitBodies(TArray<FBodyInstance*>& Bodies, TArray<FTransform>& Transforms, class UBodySetup* BodySetup, class UPrimitiveComponent* PrimitiveComp, class FPhysScene* InRBScene, PhysXAggregateType InAggregate /*= NULL*/)
+{
+	check(Bodies.Num() == Transforms.Num());
+	InitBodies(Bodies.GetData(), Transforms.GetData(), Bodies.Num(), BodySetup, PrimitiveComp, InRBScene, InAggregate);
+}
+
+void FBodyInstance::InitBody_Legacy(UBodySetup* Setup, const FTransform& Transform, UPrimitiveComponent* PrimComp, FPhysScene* InRBScene, PxAggregate* InAggregate)
 {
+	SCOPE_CYCLE_COUNTER(STAT_InitBody);
 	check(Setup);
 	AActor* Owner = PrimComp ? PrimComp->GetOwner() : NULL;
 
@@ -923,43 +914,28 @@ void FBodyInstance::InitBody(UBodySetup* Setup, const FTransform& Transform, UPr
 #if (WITH_EDITORONLY_DATA || UE_BUILD_DEBUG || LOOKING_FOR_PERF_ISSUES) && !(UE_BUILD_SHIPPING || UE_BUILD_TEST) && !NO_LOGGING
 	if (PrimComp)
 	{
+		SCOPE_CYCLE_COUNTER(STAT_InitBodyDebug);
 		DebugName += FString::Printf(TEXT("Component: %s "), *PrimComp->GetReadableName());
 	}
 
 	if (Setup->BoneName != NAME_None)
 	{
+		SCOPE_CYCLE_COUNTER(STAT_InitBodyDebug);
 		DebugName += FString::Printf(TEXT("Bone: %s "), *Setup->BoneName.ToString());
 	}
 
 #if WITH_PHYSX
 	// Convert to char* for PhysX
-	CharDebugName = MakeShareable(new TArray<ANSICHAR>(StringToArray<ANSICHAR>(*DebugName, DebugName.Len() + 1)));
-#endif
-#endif
-
-	if (Transform.GetScale3D().IsNearlyZero())
 	{
-		UE_LOG(LogPhysics, Warning, TEXT("FBodyInstance::InitBody : Scale3D is (nearly) zero: %s"), *DebugName);
-		return;
-	}
-
-	// Check we support mirroring/non-mirroring
-	const float TransformDet = Transform.GetDeterminant();
-	if (TransformDet < 0.f && !Setup->bGenerateMirroredCollision)
-	{
-		UE_LOG(LogPhysics, Warning, TEXT("FBodyInstance::InitBody : Body is mirrored but bGenerateMirroredCollision == false: %s"), *DebugName);
-		return;
-	}
-
-	if (TransformDet > 0.f && !Setup->bGenerateNonMirroredCollision)
-	{
-		UE_LOG(LogPhysics, Warning, TEXT("FBodyInstance::InitBody : Body is not mirrored but bGenerateNonMirroredCollision == false: %s"), *DebugName);
-		return;
+		SCOPE_CYCLE_COUNTER(STAT_InitBodyDebug);
+		CharDebugName = MakeShareable(new TArray<ANSICHAR>(StringToArray<ANSICHAR>(*DebugName, DebugName.Len() + 1)));
 	}
+#endif
+#endif
 
-	if (Transform.ContainsNaN())
+	if(!ValidateTransform(Transform, DebugName, Setup))
 	{
-		UE_LOG(LogPhysics, Warning, TEXT("InitBody: Bad transform - %s %s\n%s"), *DebugName, *Setup->BoneName.ToString(), *Transform.ToString());
+		// Failed to validate transform.
 		return;
 	}
 
@@ -993,11 +969,12 @@ void FBodyInstance::InitBody(UBodySetup* Setup, const FTransform& Transform, UPr
 		{
 			bool bEnableSimulation = (SkelMeshComp && IsRunningDedicatedServer()) ? SkelMeshComp->bEnablePhysicsOnDedicatedServer : true;
 			bEnableSimulation &= ((BodySetup->PhysicsType == PhysType_Simulated) || (SkelMeshComp->BodyInstance.bSimulatePhysics));	//if unfixed enable. If default look at parent
-			if (bEnableSimulation)
+			
+			if(bEnableSimulation)
 			{
 				// set simulate to true if using physics
 				bSimulatePhysics = true;
-				if (BodySetup->PhysicsType == PhysType_Simulated)
+				if(BodySetup->PhysicsType == PhysType_Simulated)
 				{
 					PhysicsBlendWeight = 1.f;
 				}
@@ -1006,7 +983,7 @@ void FBodyInstance::InitBody(UBodySetup* Setup, const FTransform& Transform, UPr
 			else
 			{
 				bSimulatePhysics = false;
-				if (BodySetup->PhysicsType == PhysType_Simulated)
+				if(BodySetup->PhysicsType == PhysType_Simulated)
 				{
 					PhysicsBlendWeight = 0.f;
 				}
@@ -1225,7 +1202,7 @@ void FBodyInstance::InitBody(UBodySetup* Setup, const FTransform& Transform, UPr
 	{
 		// Put the static actor in both scenes
 		PNewActorSync = GPhysXSDK->createRigidStatic(PTransform);
-		if (PSceneAsync)
+		if(PSceneAsync)
 		{
 			PNewActorAsync = GPhysXSDK->createRigidStatic(PTransform);
 		}
@@ -1342,6 +1319,7 @@ void FBodyInstance::InitBody(UBodySetup* Setup, const FTransform& Transform, UPr
 		BodyAggregate = GPhysXSDK->createAggregate(AggregateMaxSize, true);
 		if(BodyAggregate)
 		{
+			SCOPE_CYCLE_COUNTER(STAT_InitBodySceneInteraction);
 			BodyAggregate->addActor(*PNewDynamic);
 			SCOPED_SCENE_WRITE_LOCK(PSceneForNewDynamic);
 			PSceneForNewDynamic->addAggregate(*BodyAggregate);
@@ -1352,11 +1330,13 @@ void FBodyInstance::InitBody(UBodySetup* Setup, const FTransform& Transform, UPr
 		// Actually add actor(s) to scene(s) (if not artic link)
 		if(PNewActorSync != NULL)
 		{
+			SCOPE_CYCLE_COUNTER(STAT_InitBodySceneInteraction);
 			SCOPED_SCENE_WRITE_LOCK(PSceneSync);
 			PSceneSync->addActor(*PNewActorSync);
 		}
 		if(PNewActorAsync != NULL)
 		{
+			SCOPE_CYCLE_COUNTER(STAT_InitBodySceneInteraction);
 			SCOPED_SCENE_WRITE_LOCK(PSceneAsync);
 			PSceneAsync->addActor(*PNewActorAsync);
 		}
@@ -1366,6 +1346,7 @@ void FBodyInstance::InitBody(UBodySetup* Setup, const FTransform& Transform, UPr
 	// Set initial velocities
 	if (PNewDynamic != NULL)
 	{
+		SCOPE_CYCLE_COUNTER(STAT_InitBodyPostAdd);
 		// Compute mass (call this after setting BodyData because it uses that pointer)
 		UpdateMassProperties();
 		// Update damping
@@ -1417,6 +1398,540 @@ void FBodyInstance::InitBody(UBodySetup* Setup, const FTransform& Transform, UPr
 	}
 #endif // WITH_PHYSX
 }
+
+
+void FBodyInstance::InitBodies(FBodyInstance** Bodies, FTransform* Transforms, int32 NumBodies, class UBodySetup* BodySetup, class UPrimitiveComponent* PrimitiveComp, class FPhysScene* InRBScene, PhysXAggregateType InAggregate /*= NULL*/ )
+{
+	SCOPE_CYCLE_COUNTER(STAT_BulkInit);
+
+	check(BodySetup);
+	check(InRBScene);
+	check(NumBodies > 0);
+
+	AActor* OwningActor = PrimitiveComp ? PrimitiveComp->GetOwner() : nullptr;
+
+	const bool bStatic = PrimitiveComp == nullptr || PrimitiveComp->Mobility != EComponentMobility::Movable;
+
+#if WITH_PHYSX
+	PxScene* PSceneForNewDynamic = nullptr;	// TBD
+	PxScene* PSceneSync = InRBScene->GetPhysXScene(PST_Sync);
+	PxScene* PSceneAsync = InRBScene->HasAsyncScene() ? InRBScene->GetPhysXScene(PST_Async) : nullptr;
+
+	TArray<PxActor*> PhysXSyncActors;
+	TArray<PxActor*> PhysXAsyncActors;
+	TArray<PxActor*> PhysXDynamicActors;
+	PhysXSyncActors.Reserve(NumBodies);
+
+	if(PSceneAsync)
+	{
+		// Only allocate this if we can have async bodies
+		PhysXAsyncActors.Reserve(NumBodies);
+	}
+
+	// Ensure we have the AggGeom inside the body setup so we can calculate the number of shapes
+	BodySetup->CreatePhysicsMeshes();
+	int32 NumShapes = BodySetup->AggGeom.GetElementCount() * NumBodies;
+	TArray<PxShape*> PhysxShapes;
+	int32 ShapesWritten = 0;
+	PhysxShapes.AddUninitialized(NumShapes);
+#endif
+
+	// Setup names
+	// Make the debug name for this geometry...
+	FString DebugName(TEXT(""));
+	TSharedPtr<TArray<ANSICHAR>> PhysXName;
+
+#if (WITH_EDITORONLY_DATA || UE_BUILD_DEBUG || LOOKING_FOR_PERF_ISSUES) && !(UE_BUILD_SHIPPING || UE_BUILD_TEST) && !NO_LOGGING
+	if(PrimitiveComp)
+	{
+		DebugName += FString::Printf(TEXT("Component: %s "), *PrimitiveComp->GetReadableName());
+	}
+
+	if(BodySetup->BoneName != NAME_None)
+	{
+		DebugName += FString::Printf(TEXT("Bone: %s "), *BodySetup->BoneName.ToString());
+	}
+
+	// Convert to char* for PhysX
+	PhysXName = MakeShareable(new TArray<ANSICHAR>(StringToArray<ANSICHAR>(*DebugName, DebugName.Len() + 1)));
+#endif
+
+	// Setup for skeletal components
+	float InstanceBlendWeight = -1.0f;
+	bool bInstanceSimulatePhysics = false;
+
+	USkeletalMeshComponent* SkelMeshComp = Cast<USkeletalMeshComponent>(PrimitiveComp);
+	if(SkelMeshComp)
+	{
+		if((BodySetup->PhysicsType == PhysType_Simulated) || (BodySetup->PhysicsType == PhysType_Default))
+		{
+			bool bEnableSim = (SkelMeshComp && IsRunningDedicatedServer()) ? SkelMeshComp->bEnablePhysicsOnDedicatedServer : true;
+			bEnableSim &= ((BodySetup->PhysicsType == PhysType_Simulated) || (SkelMeshComp->BodyInstance.bSimulatePhysics));	//if unfixed enable. If default look at parent
+
+			if(bEnableSim)
+			{
+				// set simulate to true if using physics
+				bInstanceSimulatePhysics = true;
+				if(BodySetup->PhysicsType == PhysType_Simulated)
+				{
+					InstanceBlendWeight = 1.f;
+				}
+			}
+			else
+			{
+				bInstanceSimulatePhysics = false;
+				if(BodySetup->PhysicsType == PhysType_Simulated)
+				{
+					InstanceBlendWeight = 0.f;
+				}
+			}
+		}
+	}
+
+	bool bStartAwakeOverride = false;
+
+	if(SkelMeshComp)
+	{
+		bStartAwakeOverride = SkelMeshComp->BodyInstance.bStartAwake;
+	}
+
+	FVector InitialLinVel(EForceInit::ForceInitToZero);
+	if(OwningActor)
+	{
+		InitialLinVel = OwningActor->GetVelocity();
+
+		if(InitialLinVel.Size() > KINDA_SMALL_NUMBER)
+		{
+			bStartAwakeOverride = true;
+		}
+	}
+	PxVec3 PxInitialLinVel = U2PVector(InitialLinVel);
+
+#if WITH_BOX2D
+	
+	if(UBodySetup2D* BodySetup2D = Cast<UBodySetup2D>(BodySetup))
+	{
+		if(b2World* BoxWorld = FPhysicsIntegration2D::FindAssociatedWorld(PrimitiveComp->GetWorld()))
+		{
+			for(int32 BodyIdx = 0 ; BodyIdx < NumBodies ; ++BodyIdx)
+			{
+				FBodyInstance* Instance = Bodies[BodyIdx];
+				FTransform& Transform = Transforms[BodyIdx];
+
+				const b2Vec2 Scale2D = FPhysicsIntegration2D::ConvertUnrealVectorToBox(Instance->Scale3D);
+				Instance->OwnerComponent = PrimitiveComp;
+
+				// Create the body definition
+				b2BodyDef BodyDefinition;
+				if(bStatic)
+				{
+					BodyDefinition.type = b2_staticBody;
+				}
+				else
+				{
+					BodyDefinition.type = Instance->ShouldInstanceSimulatingPhysics() ? b2_dynamicBody : b2_kinematicBody;
+				}
+
+				if(SkelMeshComp)
+				{
+					bStartAwakeOverride = Instance->bStartAwake && SkelMeshComp->BodyInstance.bStartAwake;
+				}
+
+				if(Instance->bStartAwake || bStartAwakeOverride)
+				{
+					BodyDefinition.awake = true;
+				}
+				else
+				{
+					BodyDefinition.awake = false;
+				}
+
+				if(Instance->ShouldInstanceSimulatingPhysics())
+				{
+					BodyDefinition.linearVelocity = FPhysicsIntegration2D::ConvertUnrealVectorToBox(InitialLinVel);
+				}
+
+				// Create the body
+				Instance->BodyInstancePtr = BoxWorld->CreateBody(&BodyDefinition);
+				Instance->BodyInstancePtr->SetUserData(Instance);
+
+				// Circles
+				for(const FCircleElement2D& Circle : BodySetup2D->AggGeom2D.CircleElements)
+				{
+					b2CircleShape CircleShape;
+					CircleShape.m_radius = Circle.Radius * Instance->Scale3D.Size() / UnrealUnitsPerMeter;
+					CircleShape.m_p.x = Circle.Center.X * Scale2D.x;
+					CircleShape.m_p.y = Circle.Center.Y * Scale2D.y;
+
+					b2FixtureDef FixtureDef;
+					FixtureDef.shape = &CircleShape;
+
+					Instance->BodyInstancePtr->CreateFixture(&FixtureDef);
+				}
+
+				// Boxes
+				for(const FBoxElement2D& Box : BodySetup2D->AggGeom2D.BoxElements)
+				{
+					const b2Vec2 HalfBoxSize(Box.Width * 0.5f * Scale2D.x, Box.Height * 0.5f * Scale2D.y);
+					const b2Vec2 BoxCenter(Box.Center.X * Scale2D.x, Box.Center.Y * Scale2D.y);
+
+					b2PolygonShape DynamicBox;
+					DynamicBox.SetAsBox(HalfBoxSize.x, HalfBoxSize.y, BoxCenter, FMath::DegreesToRadians(Box.Angle));
+
+					b2FixtureDef FixtureDef;
+					FixtureDef.shape = &DynamicBox;
+
+					Instance->BodyInstancePtr->CreateFixture(&FixtureDef);
+				}
+
+				// Convex hulls
+				for(const FConvexElement2D& Convex : BodySetup2D->AggGeom2D.ConvexElements)
+				{
+					const int32 NumVerts = Convex.VertexData.Num();
+
+					if(NumVerts <= b2_maxPolygonVertices)
+					{
+						TArray<b2Vec2, TInlineAllocator<b2_maxPolygonVertices>> Verts;
+
+						for(int32 VertexIndex = 0; VertexIndex < Convex.VertexData.Num(); ++VertexIndex)
+						{
+							const FVector2D SourceVert = Convex.VertexData[VertexIndex];
+							new (Verts)b2Vec2(SourceVert.X * Scale2D.x, SourceVert.Y * Scale2D.y);
+						}
+
+						b2PolygonShape ConvexPoly;
+						ConvexPoly.Set(Verts.GetData(), Verts.Num());
+
+						b2FixtureDef FixtureDef;
+						FixtureDef.shape = &ConvexPoly;
+
+						Instance->BodyInstancePtr->CreateFixture(&FixtureDef);
+					}
+					else
+					{
+						UE_LOG(LogPhysics, Warning, TEXT("Too many vertices in a 2D convex body")); //@TODO: Create a better error message that indicates the asset
+					}
+				}
+
+				// Make sure it contained at least one shape
+				if(Instance->BodyInstancePtr->GetFixtureList() == nullptr)
+				{
+					if(DebugName.Len())
+					{
+						UE_LOG(LogPhysics, Log, TEXT("InitBody: failed - no shapes: %s"), *DebugName);
+					}
+
+					Instance->BodyInstancePtr->SetUserData(nullptr);
+					BoxWorld->DestroyBody(Instance->BodyInstancePtr);
+					Instance->BodyInstancePtr = nullptr;
+
+					//clear Owner and Setup info as well to properly clean up the BodyInstance.
+					Instance->OwnerComponent = NULL;
+					Instance->BodySetup = NULL;
+
+					return;
+				}
+				else
+				{
+					// Position the body
+					Instance->SetBodyTransform(Transform, /*bTeleport=*/ true);
+				}
+
+				// Apply correct physical materials to shape we created.
+				Instance->UpdatePhysicalMaterials();
+
+				// Set the filter data on the shapes (call this after setting BodyData because it uses that pointer)
+				Instance->UpdatePhysicsFilterData();
+
+				if(!bStatic)
+				{
+					// Compute mass (call this after setting BodyData because it uses that pointer)
+					Instance->UpdateMassProperties();
+
+					// Update damping
+					Instance->UpdateDampingProperties();
+
+					Instance->SetMaxAngularVelocity(Instance->MaxAngularVelocity, false);
+
+					Instance->SetMaxDepenetrationVelocity(Instance->bOverrideMaxDepenetrationVelocity ? Instance->MaxDepenetrationVelocity : UPhysicsSettings::Get()->MaxDepenetrationVelocity);
+
+					//@TODO: BOX2D: Determine if sleep threshold and solver settings can be configured per-body or not
+#if 0
+					// Set the parameters for determining when to put the object to sleep.
+					float SleepEnergyThresh = PNewDynamic->getSleepThreshold();
+					if(SleepFamily == SF_Sensitive)
+					{
+						SleepEnergyThresh /= 20.f;
+					}
+					PNewDynamic->setSleepThreshold(SleepEnergyThresh);
+					// set solver iteration count 
+					int32 PositionIterCount = FMath::Clamp(PositionSolverIterationCount, 1, 255);
+					int32 VelocityIterCount = FMath::Clamp(VelocitySolverIterationCount, 1, 255);
+					PNewDynamic->setSolverIterationCounts(PositionIterCount, VelocityIterCount);
+#endif
+				}
+			}
+		}
+		return;
+	}
+#endif
+
+#if WITH_PHYSX
+	for(int32 BodyIdx = 0 ; BodyIdx < NumBodies ; ++BodyIdx)
+	{
+		FBodyInstance* Instance = Bodies[BodyIdx];
+		FTransform& Transform = Transforms[BodyIdx];
+
+		ValidateTransform(Transform, DebugName, BodySetup);
+
+		Instance->OwnerComponent = PrimitiveComp;
+		Instance->BodySetup = BodySetup;
+		Instance->Scale3D = Transform.GetScale3D();
+		Instance->CharDebugName = PhysXName;
+
+		// Handle autowelding here to avoid extra work
+		if(Instance->bAutoWeld)
+		{
+			UPrimitiveComponent * ParentPrimComponent = PrimitiveComp ? Cast<UPrimitiveComponent>(PrimitiveComp->AttachParent) : NULL;
+
+			//if we have a parent we will now do the weld and exit any further initialization
+			if(ParentPrimComponent && PrimitiveComp->GetWorld() && PrimitiveComp->GetWorld()->IsGameWorld())
+			{
+				if(PrimitiveComp->WeldToImplementation(ParentPrimComponent, PrimitiveComp->AttachSocketName, false))	//welded new simulated body so initialization is done
+				{
+					return;
+				}
+			}
+		}
+
+		PxRigidDynamic* PNewDynamic = nullptr;
+		PxRigidActor* PNewActorSync = nullptr;
+		PxRigidActor* PNewActorAsync = nullptr;
+		PxTransform PTransform = U2PTransform(Transform);
+		bool bInitFail = false;
+
+		// Don't process if we've already got a body
+		if(Instance->GetPxRigidActor())
+		{
+			Instance->OwnerComponent = nullptr;
+			Instance->BodySetup = nullptr;
+
+			continue;
+		}
+
+		// Set sim parameters for bodies from skeletal mesh components
+		if(SkelMeshComp)
+		{
+			Instance->bSimulatePhysics = bInstanceSimulatePhysics;
+			if(InstanceBlendWeight != -1.0f)
+			{
+				Instance->PhysicsBlendWeight = InstanceBlendWeight;
+			}
+		}
+
+		if(bStatic)
+		{
+			PNewActorSync = GPhysXSDK->createRigidStatic(PTransform);
+			if(PSceneAsync)
+			{
+				PNewActorAsync = GPhysXSDK->createRigidStatic(PTransform);
+			}
+		}
+		else
+		{
+			PNewDynamic = GPhysXSDK->createRigidDynamic(PTransform);
+
+			// Put the dynamic actor in one scene or the other
+			if(!Instance->UseAsyncScene())
+			{
+				PNewActorSync = PNewDynamic;
+				PSceneForNewDynamic = PSceneSync;
+			}
+			else
+			{
+				PNewActorAsync = PNewDynamic;
+				PSceneForNewDynamic = PSceneAsync;
+			}
+
+			// Set kinematic if desired
+			if(Instance->ShouldInstanceSimulatingPhysics())
+			{
+				PNewDynamic->setRigidDynamicFlag(PxRigidDynamicFlag::eKINEMATIC, false);
+			}
+			else
+			{
+				PNewDynamic->setRigidDynamicFlag(PxRigidDynamicFlag::eKINEMATIC, true);
+			}
+
+			// turn off gravity if desired
+			if(!Instance->bEnableGravity)
+			{
+				PNewDynamic->setActorFlag(PxActorFlag::eDISABLE_GRAVITY, true);
+			}
+
+			// If we ever drive this body kinematically, we want to use its target for scene queries, so collision is updated right away, not on the next physics sim
+			PNewDynamic->setRigidDynamicFlag(PxRigidDynamicFlag::eUSE_KINEMATIC_TARGET_FOR_SCENE_QUERIES, true);
+		}
+
+		Instance->PhysxUserData = FPhysxUserData(Instance);
+
+		if(PNewActorSync)
+		{
+			BodySetup->AddShapesToRigidActor(PNewActorSync, Instance->Scale3D);
+			bInitFail |= PNewActorSync->getNbShapes() == 0;
+			PNewActorSync->userData = &Instance->PhysxUserData;
+			PNewActorSync->setName(Instance->CharDebugName.IsValid() ? Instance->CharDebugName->GetData() : nullptr);
+
+			ShapesWritten += PNewActorSync->getShapes(PhysxShapes.GetData() + ShapesWritten, PhysxShapes.Num() - ShapesWritten);
+
+			check(FPhysxUserData::Get<FBodyInstance>(PNewActorSync->userData) == Instance && FPhysxUserData::Get<FBodyInstance>(PNewActorSync->userData)->OwnerComponent != NULL);
+
+		}
+
+		if(PNewActorAsync)
+		{
+			check(PSceneAsync);
+			BodySetup->AddShapesToRigidActor(PNewActorAsync, Instance->Scale3D);
+			bInitFail |= PNewActorAsync->getNbShapes() == 0;
+			PNewActorAsync->userData = &Instance->PhysxUserData;
+			PNewActorAsync->setName(Instance->CharDebugName.IsValid() ? Instance->CharDebugName->GetData() : nullptr);
+
+			ShapesWritten += PNewActorAsync->getShapes(PhysxShapes.GetData() + ShapesWritten, PhysxShapes.Num() - ShapesWritten);
+
+			check(FPhysxUserData::Get<FBodyInstance>(PNewActorAsync->userData) == Instance && FPhysxUserData::Get<FBodyInstance>(PNewActorAsync->userData)->OwnerComponent != NULL);
+		}
+
+		if(bInitFail)
+		{
+			UE_LOG(LogPhysics, Log, TEXT("Init Instance %d of Primitive Component %s failed"), BodyIdx, *PrimitiveComp->GetName());
+			if(PNewActorSync)
+			{
+				PNewActorSync->release();
+			}
+
+			if(PNewActorAsync)
+			{
+				PNewActorAsync->release();
+			}
+
+			Instance->OwnerComponent = nullptr;
+			Instance->BodySetup = nullptr;
+
+			continue;
+		}
+
+		PhysXSyncActors.Add(PNewActorSync);
+		if(PNewActorAsync)
+		{
+			PhysXAsyncActors.Add(PNewActorAsync);
+		}
+
+		if(PNewDynamic)
+		{
+			PhysXDynamicActors.Add(PNewDynamic);
+		}
+
+		Instance->RigidActorSync = PNewActorSync;
+		Instance->RigidActorAsync = PNewActorAsync;
+		Instance->SceneIndexSync = InRBScene->PhysXSceneIndex[PST_Sync];
+		Instance->SceneIndexAsync = PSceneAsync ? InRBScene->PhysXSceneIndex[PST_Async] : 0;
+		Instance->UpdatePhysicsFilterData();
+	}
+
+	// Grab a material to use, for now we only support batching of bodies using the same material
+	UPhysicalMaterial* SimpleMat = Bodies[0]->GetSimplePhysicalMaterial();
+	TArray<UPhysicalMaterial*> ComplexMats = Bodies[0]->GetComplexPhysicalMaterials();
+	PxMaterial* PhysxSimpleMat = SimpleMat->GetPhysXMaterial();
+	
+	// Update material parameters
+	for(int32 ShapeIdx = 0 ; ShapeIdx < ShapesWritten ; ++ShapeIdx)
+	{
+		PxShape* Shape = PhysxShapes[ShapeIdx];
+		ApplyMaterialToShape(Shape, PhysxSimpleMat, ComplexMats);
+	}
+
+	{
+		SCOPE_CYCLE_COUNTER(STAT_BulkSceneAdd);
+		// Use the aggregate if it exists, has enough slots and is in the correct scene (or no scene)
+		if(InAggregate && PhysXDynamicActors.Num() > 0 && (InAggregate->getMaxNbActors() - InAggregate->getNbActors()) >= (uint32)PhysXDynamicActors.Num() && (InAggregate->getScene() == PSceneForNewDynamic || InAggregate->getScene() == NULL))
+		{
+			for(PxActor* Actor : PhysXDynamicActors)
+			{
+				InAggregate->addActor(*Actor);
+			}
+		}
+		else
+		{
+			SCOPED_SCENE_WRITE_LOCK(PSceneSync);
+			PSceneSync->addActors(PhysXSyncActors.GetData(), PhysXSyncActors.Num());
+
+			if(PSceneAsync)
+			{
+				SCOPED_SCENE_WRITE_LOCK(PSceneAsync);
+				PSceneAsync->addActors(PhysXAsyncActors.GetData(), PhysXAsyncActors.Num());
+			}
+		}
+	}
+
+	// Set up dynamic instance data
+	if(!bStatic)
+	{
+		SCOPE_CYCLE_COUNTER(STAT_InitBodyPostAdd);
+		for(int32 BodyIdx = 0 ; BodyIdx < NumBodies ; ++BodyIdx)
+		{
+			FBodyInstance* Instance = Bodies[BodyIdx];
+			
+			if(PxRigidDynamic* RigidActor = Instance->GetPxRigidDynamic())
+			{
+				Instance->UpdateMassProperties();
+				Instance->UpdateDampingProperties();
+				Instance->SetMaxAngularVelocity(Instance->MaxAngularVelocity, false);
+				Instance->SetMaxDepenetrationVelocity(Instance->bOverrideMaxDepenetrationVelocity ? Instance->MaxDepenetrationVelocity : UPhysicsSettings::Get()->MaxDepenetrationVelocity);
+
+				if(Instance->ShouldInstanceSimulatingPhysics())
+				{
+					RigidActor->setLinearVelocity(PxInitialLinVel);
+				}
+
+				float SleepEnergyThresh = RigidActor->getSleepThreshold();
+				if(Instance->SleepFamily == SF_Sensitive)
+				{
+					SleepEnergyThresh /= 20.f;
+				}
+				RigidActor->setSleepThreshold(SleepEnergyThresh);
+				// set solver iteration count 
+				int32 PositionIterCount = FMath::Clamp(Instance->PositionSolverIterationCount, 1, 255);
+				int32 VelocityIterCount = FMath::Clamp(Instance->VelocitySolverIterationCount, 1, 255);
+				RigidActor->setSolverIterationCounts(PositionIterCount, VelocityIterCount);
+
+				if(RigidActor->getScene())
+				{
+					Instance->CreateDOFLock();
+					if(IsRigidBodyNonKinematic(RigidActor))
+					{
+						if(SkelMeshComp)
+						{
+							bStartAwakeOverride = Instance->bStartAwake && SkelMeshComp->BodyInstance.bStartAwake;
+						}
+
+						if(Instance->bStartAwake || bStartAwakeOverride)
+						{
+							RigidActor->wakeUp();
+						}
+						else
+						{
+							RigidActor->putToSleep();
+						}
+					}
+				}
+			}
+			
+		}
+	}
+#endif
+}
+
+
 #endif // UE_WITH_PHYSICS
 
 #if WITH_PHYSX
@@ -1515,11 +2030,13 @@ void TermBodyHelper(int32& SceneIndex, PxRigidActor*& PRigidActor, FBodyInstance
 
 #endif
 
+DEFINE_STAT(STAT_TermBody);
 /**
  *	Clean up the physics engine info for this instance.
  */
 void FBodyInstance::TermBody()
 {
+	SCOPE_CYCLE_COUNTER(STAT_TermBody);
 #if WITH_BOX2D
 	if (BodyInstancePtr != NULL)
 	{
@@ -3904,6 +4421,83 @@ void FBodyInstance::SetUseAsyncScene(bool bNewUseAsyncScene)
 	bUseAsyncScene = bNewUseAsyncScene;
 }
 
+void FBodyInstance::ApplyMaterialToShape(PxShape* PShape, PxMaterial* PSimpleMat, TArray<UPhysicalMaterial*>& ComplexPhysMats)
+{
+	// If a triangle mesh, need to get array of materials...
+	if(PShape->getGeometryType() == PxGeometryType::eTRIANGLEMESH)
+	{
+		//TArray<UPhysicalMaterial*> ComplexPhysMats = GetComplexPhysicalMaterials();
+		TArray<PxMaterial*> PComplexMats;
+		PComplexMats.AddUninitialized(ComplexPhysMats.Num());
+		for(int MatIdx = 0; MatIdx < ComplexPhysMats.Num(); MatIdx++)
+		{
+			check(ComplexPhysMats[MatIdx] != NULL);
+			PComplexMats[MatIdx] = ComplexPhysMats[MatIdx]->GetPhysXMaterial();
+			check(PComplexMats[MatIdx] != NULL);
+		}
+
+		if(PComplexMats.Num())
+		{
+			PShape->setMaterials(PComplexMats.GetData(), PComplexMats.Num());
+		}
+		else
+		{
+			UE_LOG(LogPhysics, Warning, TEXT("FBodyInstance::UpdatePhysicalMaterials : PComplexMats is empty - falling back on simple physical material."));
+			PShape->setMaterials(&PSimpleMat, 1);
+		}
+
+	}
+	// Simple shape, 
+	else
+	{
+		PShape->setMaterials(&PSimpleMat, 1);
+	}
+}
+
+void FBodyInstance::ApplyMaterialToInstanceShapes(PxMaterial* PSimpleMat, TArray<UPhysicalMaterial*> ComplexPhysMats)
+{
+	int32 NumSyncShapes = 0;
+	TArray<PxShape*> AllShapes = GetAllShapes(NumSyncShapes);
+
+	for(int32 ShapeIdx = 0; ShapeIdx < AllShapes.Num(); ShapeIdx++)
+	{
+		PxShape* PShape = AllShapes[ShapeIdx];
+
+		ApplyMaterialToShape(PShape, PSimpleMat, ComplexPhysMats);
+	}
+}
+
+bool FBodyInstance::ValidateTransform(const FTransform &Transform, FString& DebugName, UBodySetup* Setup)
+{
+	if(Transform.GetScale3D().IsNearlyZero())
+	{
+		UE_LOG(LogPhysics, Warning, TEXT("Initialising Body : Scale3D is (nearly) zero: %s"), *DebugName);
+		return false;
+	}
+
+	// Check we support mirroring/non-mirroring
+	const float TransformDet = Transform.GetDeterminant();
+	if(TransformDet < 0.f && !Setup->bGenerateMirroredCollision)
+	{
+		UE_LOG(LogPhysics, Warning, TEXT("Initialising Body : Body is mirrored but bGenerateMirroredCollision == false: %s"), *DebugName);
+		return false;
+	}
+
+	if(TransformDet > 0.f && !Setup->bGenerateNonMirroredCollision)
+	{
+		UE_LOG(LogPhysics, Warning, TEXT("Initialising Body : Body is not mirrored but bGenerateNonMirroredCollision == false: %s"), *DebugName);
+		return false;
+	}
+
+	if(Transform.ContainsNaN())
+	{
+		UE_LOG(LogPhysics, Warning, TEXT("Initialising Body : Bad transform - %s %s\n%s"), *DebugName, *Setup->BoneName.ToString(), *Transform.ToString());
+		return false;
+	}
+
+	return true;
+}
+
 ////////////////////////////////////////////////////////////////////////////
 // FBodyInstanceEditorHelpers
 
diff --git a/Engine/Source/Runtime/Engine/Private/PhysicsEngine/PhysXSupport.h b/Engine/Source/Runtime/Engine/Private/PhysicsEngine/PhysXSupport.h
--- a/Engine/Source/Runtime/Engine/Private/PhysicsEngine/PhysXSupport.h
+++ b/Engine/Source/Runtime/Engine/Private/PhysicsEngine/PhysXSupport.h
@@ -31,6 +31,7 @@ public:
 	FPhysXSceneReadLock(PxScene* PInScene)
 		: PScene(PInScene)
 	{
+		SCOPE_CYCLE_COUNTER(STAT_PhysSceneLock);
 		if(PScene)
 		{
 			PScene->lockRead();
@@ -39,6 +40,7 @@ public:
 
 	~FPhysXSceneReadLock()
 	{
+		SCOPE_CYCLE_COUNTER(STAT_PhysSceneLock);
 		if(PScene)
 		{
 			PScene->unlockRead();
@@ -56,6 +58,7 @@ public:
 	FPhysXSceneWriteLock(PxScene* PInScene)
 		: PScene(PInScene)
 	{
+		SCOPE_CYCLE_COUNTER(STAT_PhysSceneLock);
 		if(PScene)
 		{
 			PScene->lockWrite();
@@ -64,6 +67,7 @@ public:
 
 	~FPhysXSceneWriteLock()
 	{
+		SCOPE_CYCLE_COUNTER(STAT_PhysSceneLock);
 		if(PScene)
 		{
 			PScene->unlockWrite();
diff --git a/Engine/Source/Runtime/Engine/Public/PhysicsPublic.h b/Engine/Source/Runtime/Engine/Public/PhysicsPublic.h
--- a/Engine/Source/Runtime/Engine/Public/PhysicsPublic.h
+++ b/Engine/Source/Runtime/Engine/Public/PhysicsPublic.h
@@ -23,6 +23,17 @@ DECLARE_CYCLE_STAT_EXTERN(TEXT("Phys SubstepStart"), STAT_SubstepSimulationStart
 DECLARE_CYCLE_STAT_EXTERN(TEXT("Phys SubstepEnd"), STAT_SubstepSimulationEnd, STATGROUP_Physics, );
 DECLARE_CYCLE_STAT_EXTERN(TEXT("SyncComponentsToBodies"), STAT_SyncComponentsToBodies, STATGROUP_Physics, );
 DECLARE_CYCLE_STAT_EXTERN(TEXT("Init Articulated"), STAT_InitArticulated, STATGROUP_Physics, );
+DECLARE_CYCLE_STAT_EXTERN(TEXT("Init Body"), STAT_InitBody, STATGROUP_Physics, );
+DECLARE_CYCLE_STAT_EXTERN(TEXT("Init Body Debug"), STAT_InitBodyDebug, STATGROUP_Physics, );
+DECLARE_CYCLE_STAT_EXTERN(TEXT("Init Body Scene Interaction"), STAT_InitBodySceneInteraction, STATGROUP_Physics, );
+DECLARE_CYCLE_STAT_EXTERN(TEXT("Init Body Post Add to Scene"), STAT_InitBodyPostAdd, STATGROUP_Physics, );
+DECLARE_CYCLE_STAT_EXTERN(TEXT("Term Body"), STAT_TermBody, STATGROUP_Physics, );
+DECLARE_CYCLE_STAT_EXTERN(TEXT("Update Materials"), STAT_UpdatePhysMats, STATGROUP_Physics, );
+DECLARE_CYCLE_STAT_EXTERN(TEXT("Update Materials Scene Interaction"), STAT_UpdatePhysMatsSceneInteraction, STATGROUP_Physics, );
+DECLARE_CYCLE_STAT_EXTERN(TEXT("Filter Update"), STAT_UpdatePhysFilter, STATGROUP_Physics, );
+
+DECLARE_CYCLE_STAT_EXTERN(TEXT("Bulk Body Init"), STAT_BulkInit, STATGROUP_Physics, );
+DECLARE_CYCLE_STAT_EXTERN(TEXT("Bulk Body Scene Add"), STAT_BulkSceneAdd, STATGROUP_Physics, );
 
 DECLARE_DWORD_COUNTER_STAT_EXTERN(TEXT("Broadphase Adds"), STAT_NumBroadphaseAdds, STATGROUP_Physics, );
 DECLARE_DWORD_COUNTER_STAT_EXTERN(TEXT("Broadphase Removes"), STAT_NumBroadphaseRemoves, STATGROUP_Physics, );
