Added utility function to draw an oriented wire box with a matrix transform

diff --git a/Engine/Source/Runtime/Engine/Private/PrimitiveDrawingUtils.cpp b/Engine/Source/Runtime/Engine/Private/PrimitiveDrawingUtils.cpp
--- a/Engine/Source/Runtime/Engine/Private/PrimitiveDrawingUtils.cpp
+++ b/Engine/Source/Runtime/Engine/Private/PrimitiveDrawingUtils.cpp
@@ -836,6 +836,41 @@ void DrawWireBox(FPrimitiveDrawInterface* PDI, const FBox& Box, const FLinearCol
 	}
 }
 
+
+void DrawWireBox( FPrimitiveDrawInterface* PDI, const FMatrix& Matrix, const FBox& Box, const FLinearColor& Color, uint8 DepthPriority )
+{
+	FVector B[ 2 ];
+	B[ 0 ] = Box.Min;
+	B[ 1 ] = Box.Max;
+
+	for( int i = 0; i < 2; i++ )
+	{
+		for( int j = 0; j < 2; j++ )
+		{
+			FVector P, Q;
+
+			P.X = B[ i ].X; Q.X = B[ i ].X;
+			P.Y = B[ j ].Y; Q.Y = B[ j ].Y;
+			P.Z = B[ 0 ].Z; Q.Z = B[ 1 ].Z;
+			P = Matrix.TransformPosition( P ); Q = Matrix.TransformPosition( Q );
+			PDI->DrawLine( P, Q, Color, DepthPriority );
+
+			P.Y = B[ i ].Y; Q.Y = B[ i ].Y;
+			P.Z = B[ j ].Z; Q.Z = B[ j ].Z;
+			P.X = B[ 0 ].X; Q.X = B[ 1 ].X;
+			P = Matrix.TransformPosition( P ); Q = Matrix.TransformPosition( Q );
+			PDI->DrawLine( P, Q, Color, DepthPriority );
+
+			P.Z = B[ i ].Z; Q.Z = B[ i ].Z;
+			P.X = B[ j ].X; Q.X = B[ j ].X;
+			P.Y = B[ 0 ].Y; Q.Y = B[ 1 ].Y;
+			P = Matrix.TransformPosition( P ); Q = Matrix.TransformPosition( Q );
+			PDI->DrawLine( P, Q, Color, DepthPriority );
+		}
+	}
+}
+
+
 /**
  * Draws a circle using lines.
  *
diff --git a/Engine/Source/Runtime/Engine/Public/SceneManagement.h b/Engine/Source/Runtime/Engine/Public/SceneManagement.h
--- a/Engine/Source/Runtime/Engine/Public/SceneManagement.h
+++ b/Engine/Source/Runtime/Engine/Public/SceneManagement.h
@@ -1548,6 +1548,7 @@ extern ENGINE_API void DrawFlatArrow(class FPrimitiveDrawInterface* PDI,const FV
 
 // Line drawing utility functions.
 extern ENGINE_API void DrawWireBox(class FPrimitiveDrawInterface* PDI, const FBox& Box, const FLinearColor& Color, uint8 DepthPriority, float Thickness = 0, float DepthBias = 0.0f, bool bScreenSpace = false);
+extern ENGINE_API void DrawWireBox(class FPrimitiveDrawInterface* PDI, const FMatrix& Matrix, const FBox& Box, const FLinearColor& Color, uint8 DepthPriority );
 extern ENGINE_API void DrawCircle(class FPrimitiveDrawInterface* PDI, const FVector& Base, const FVector& X, const FVector& Y, const FLinearColor& Color, float Radius, int32 NumSides, uint8 DepthPriority, float Thickness = 0, float DepthBias = 0.0f, bool bScreenSpace = false);
 extern ENGINE_API void DrawArc(FPrimitiveDrawInterface* PDI, const FVector Base, const FVector X, const FVector Y, const float MinAngle, const float MaxAngle, const float Radius, const int32 Sections, const FLinearColor& Color, uint8 DepthPriority);
 extern ENGINE_API void DrawWireSphere(class FPrimitiveDrawInterface* PDI, const FVector& Base, const FLinearColor& Color, float Radius, int32 NumSides, uint8 DepthPriority, float Thickness = 0, float DepthBias = 0.0f, bool bScreenSpace = false);
