Fixed various spawned transient gameplay actors not actually transient
- This corrects issues where temporary actors could be saved into map packages

diff --git a/Engine/Source/Runtime/Engine/Private/Controller.cpp b/Engine/Source/Runtime/Engine/Private/Controller.cpp
--- a/Engine/Source/Runtime/Engine/Private/Controller.cpp
+++ b/Engine/Source/Runtime/Engine/Private/Controller.cpp
@@ -418,6 +418,7 @@ void AController::InitPlayerState()
 			SpawnInfo.Owner = this;
 			SpawnInfo.Instigator = Instigator;
 			SpawnInfo.bNoCollisionFail = true;
+			SpawnInfo.ObjectFlags |= RF_Transient;	// We never want player states to save into a map
 			PlayerState = World->SpawnActor<APlayerState>(GameMode->PlayerStateClass, SpawnInfo );
 	
 			// force a default player name if necessary
diff --git a/Engine/Source/Runtime/Engine/Private/DebugCameraController.cpp b/Engine/Source/Runtime/Engine/Private/DebugCameraController.cpp
--- a/Engine/Source/Runtime/Engine/Private/DebugCameraController.cpp
+++ b/Engine/Source/Runtime/Engine/Private/DebugCameraController.cpp
@@ -284,6 +284,7 @@ void ADebugCameraController::PostInitializeComponents()
 	SpawnInfo.Owner = this;
 	SpawnInfo.Instigator = Instigator;
 	SpawnInfo.bNoCollisionFail = true;
+	SpawnInfo.ObjectFlags |= RF_Transient;	// We never want these to save into a map
 	MyHUD = GetWorld()->SpawnActor<ADebugCameraHUD>( SpawnInfo );
 
 	ChangeState(NAME_Inactive);
diff --git a/Engine/Source/Runtime/Engine/Private/DemoNetDriver.cpp b/Engine/Source/Runtime/Engine/Private/DemoNetDriver.cpp
--- a/Engine/Source/Runtime/Engine/Private/DemoNetDriver.cpp
+++ b/Engine/Source/Runtime/Engine/Private/DemoNetDriver.cpp
@@ -805,7 +805,9 @@ void UDemoNetDriver::SpawnDemoRecSpectator( UNetConnection* Connection )
 		return;
 	}
 
-	APlayerController* Controller = World->SpawnActor<APlayerController>( C );
+	FActorSpawnParameters SpawnInfo;
+	SpawnInfo.ObjectFlags |= RF_Transient;	// We never want these to save into a map
+	APlayerController* Controller = World->SpawnActor<APlayerController>( C, SpawnInfo );
 
 	if ( Controller == NULL )
 	{
diff --git a/Engine/Source/Runtime/Engine/Private/GameMode.cpp b/Engine/Source/Runtime/Engine/Private/GameMode.cpp
--- a/Engine/Source/Runtime/Engine/Private/GameMode.cpp
+++ b/Engine/Source/Runtime/Engine/Private/GameMode.cpp
@@ -152,6 +152,7 @@ void AGameMode::InitGame(const FString& MapName, const FString& Options, FString
 	UClass* const SessionClass = GetGameSessionClass();
 	FActorSpawnParameters SpawnInfo;
 	SpawnInfo.Instigator = Instigator;
+	SpawnInfo.ObjectFlags |= RF_Transient;	// We never want to save game sessions into a map
 	GameSession = World->SpawnActor<AGameSession>(SessionClass, SpawnInfo);
 	GameSession->InitOptions(Options);
 
@@ -403,6 +404,7 @@ void AGameMode::PreInitializeComponents()
 	GetWorldTimerManager().SetTimer(TimerHandle_DefaultTimer, this, &AGameMode::DefaultTimer, GetWorldSettings()->GetEffectiveTimeDilation(), true);
 	FActorSpawnParameters SpawnInfo;
 	SpawnInfo.Instigator = Instigator;
+	SpawnInfo.ObjectFlags |= RF_Transient;	// We never want to save game states or network managers into a map
 
 	// Fallback to default GameState if none was specified.
 	if (GameStateClass == NULL)
@@ -1282,6 +1284,7 @@ APlayerController* AGameMode::SpawnPlayerController(FVector const& SpawnLocation
 {
 	FActorSpawnParameters SpawnInfo;
 	SpawnInfo.Instigator = Instigator;	
+	SpawnInfo.ObjectFlags |= RF_Transient;	// We never want to save player controllers into a map
 	return GetWorld()->SpawnActor<APlayerController>(PlayerControllerClass, SpawnLocation, SpawnRotation, SpawnInfo );
 }
 
@@ -1299,6 +1302,7 @@ APawn* AGameMode::SpawnDefaultPawnFor(AController* NewPlayer, AActor* StartSpot)
 
 	FActorSpawnParameters SpawnInfo;
 	SpawnInfo.Instigator = Instigator;	
+	SpawnInfo.ObjectFlags |= RF_Transient;	// We never want to save default player pawns into a map
 	APawn* ResultPawn = GetWorld()->SpawnActor<APawn>(GetDefaultPawnClassForController(NewPlayer), StartLocation, StartRotation, SpawnInfo );
 	if ( ResultPawn == NULL )
 	{
diff --git a/Engine/Source/Runtime/Engine/Private/LocalPlayer.cpp b/Engine/Source/Runtime/Engine/Private/LocalPlayer.cpp
--- a/Engine/Source/Runtime/Engine/Private/LocalPlayer.cpp
+++ b/Engine/Source/Runtime/Engine/Private/LocalPlayer.cpp
@@ -253,7 +253,9 @@ bool ULocalPlayer::SpawnPlayActor(const FString& URL,FString& OutError, UWorld*
 		// replaces this fake player controller with the real replicated one from the server
 		//
 
-		PlayerController = CastChecked<APlayerController>(InWorld->SpawnActor(PCClass));
+		FActorSpawnParameters SpawnInfo;
+		SpawnInfo.ObjectFlags |= RF_Transient;	// We never want to save player controllers into a map
+		PlayerController = CastChecked<APlayerController>(InWorld->SpawnActor<APlayerController>(PCClass, SpawnInfo));
 		const int32 PlayerIndex = GEngine->GetGamePlayers(InWorld).Find(this);
 		PlayerController->NetPlayerIndex = PlayerIndex;
 	}
diff --git a/Engine/Source/Runtime/Engine/Private/Pawn.cpp b/Engine/Source/Runtime/Engine/Private/Pawn.cpp
--- a/Engine/Source/Runtime/Engine/Private/Pawn.cpp
+++ b/Engine/Source/Runtime/Engine/Private/Pawn.cpp
@@ -320,6 +320,7 @@ void APawn::SpawnDefaultController()
 		SpawnInfo.Instigator = Instigator;
 		SpawnInfo.bNoCollisionFail = true;
 		SpawnInfo.OverrideLevel = GetLevel();
+		SpawnInfo.ObjectFlags |= RF_Transient;	// We never want to save AI controllers into a map
 		AController* NewController = GetWorld()->SpawnActor<AController>(AIControllerClass, GetActorLocation(), GetActorRotation(), SpawnInfo);
 		if (NewController != NULL)
 		{
diff --git a/Engine/Source/Runtime/Engine/Private/PlayerCameraManager.cpp b/Engine/Source/Runtime/Engine/Private/PlayerCameraManager.cpp
--- a/Engine/Source/Runtime/Engine/Private/PlayerCameraManager.cpp
+++ b/Engine/Source/Runtime/Engine/Private/PlayerCameraManager.cpp
@@ -633,6 +633,7 @@ void APlayerCameraManager::PostInitializeComponents()
 	SpawnInfo.Owner = this;
 	SpawnInfo.Instigator = Instigator;
 	SpawnInfo.bNoCollisionFail = true;
+	SpawnInfo.ObjectFlags |= RF_Transient;	// We never want to save these temp actors into a map
 	AnimCameraActor = GetWorld()->SpawnActor<ACameraActor>(SpawnInfo);
 }
 
@@ -1016,6 +1017,7 @@ AEmitterCameraLensEffectBase* APlayerCameraManager::AddCameraLensEffect(TSubclas
 			SpawnInfo.Owner = PCOwner->GetViewTarget();
 			SpawnInfo.Instigator = Instigator;
 			SpawnInfo.bNoCollisionFail = true;
+			SpawnInfo.ObjectFlags |= RF_Transient;	// We never want to save these into a map
 			LensEffect = GetWorld()->SpawnActor<AEmitterCameraLensEffectBase>(LensEffectEmitterClass, SpawnInfo);
 			if (LensEffect != NULL)
 			{
diff --git a/Engine/Source/Runtime/Engine/Private/PlayerController.cpp b/Engine/Source/Runtime/Engine/Private/PlayerController.cpp
--- a/Engine/Source/Runtime/Engine/Private/PlayerController.cpp
+++ b/Engine/Source/Runtime/Engine/Private/PlayerController.cpp
@@ -1023,6 +1023,7 @@ void APlayerController::SpawnDefaultHUD()
 	FActorSpawnParameters SpawnInfo;
 	SpawnInfo.Owner = this;
 	SpawnInfo.Instigator = Instigator;
+	SpawnInfo.ObjectFlags |= RF_Transient;	// We never want to save HUDs into a map
 	MyHUD = GetWorld()->SpawnActor<AHUD>( SpawnInfo );
 }
 
@@ -1177,6 +1178,7 @@ void APlayerController::ClientSetHUD_Implementation(TSubclassOf<AHUD> NewHUDClas
 	FActorSpawnParameters SpawnInfo;
 	SpawnInfo.Owner = this;
 	SpawnInfo.Instigator = Instigator;
+	SpawnInfo.ObjectFlags |= RF_Transient;	// We never want to save HUDs into a map
 	MyHUD = GetWorld()->SpawnActor<AHUD>(NewHUDClass, SpawnInfo );
 }
 
@@ -2479,6 +2481,7 @@ void APlayerController::SpawnPlayerCameraManager()
 	FActorSpawnParameters SpawnInfo;
 	SpawnInfo.Owner = this;
 	SpawnInfo.Instigator = Instigator;
+	SpawnInfo.ObjectFlags |= RF_Transient;	// We never want to save camera managers into a map
 	if (PlayerCameraManagerClass != NULL)
 	{
 		PlayerCameraManager = GetWorld()->SpawnActor<APlayerCameraManager>(PlayerCameraManagerClass, SpawnInfo);
@@ -3953,6 +3956,7 @@ ASpectatorPawn* APlayerController::SpawnSpectatorPawn()
 			FActorSpawnParameters SpawnParams;
 			SpawnParams.Owner = this;
 			SpawnParams.bNoCollisionFail = true;
+			SpawnParams.ObjectFlags |= RF_Transient;	// We never want to save spectator pawns into a map
 			SpawnedSpectator = GetWorld()->SpawnActor<ASpectatorPawn>(GameState->SpectatorClass, GetSpawnLocation(), GetControlRotation(), SpawnParams);
 			if (SpawnedSpectator)
 			{
diff --git a/Engine/Source/Runtime/Engine/Private/PlayerState.cpp b/Engine/Source/Runtime/Engine/Private/PlayerState.cpp
--- a/Engine/Source/Runtime/Engine/Private/PlayerState.cpp
+++ b/Engine/Source/Runtime/Engine/Private/PlayerState.cpp
@@ -288,6 +288,7 @@ APlayerState* APlayerState::Duplicate()
 	FActorSpawnParameters SpawnInfo;
 	SpawnInfo.Instigator = Instigator;
 	SpawnInfo.bNoCollisionFail = true;
+	SpawnInfo.ObjectFlags |= RF_Transient;	// We never want to save player states into a map
 	APlayerState* NewPlayerState = GetWorld()->SpawnActor<APlayerState>(GetClass(), SpawnInfo );
 	// Can fail in case of multiplayer PIE teardown
 	if (NewPlayerState)
diff --git a/Engine/Source/Runtime/Engine/Private/World.cpp b/Engine/Source/Runtime/Engine/Private/World.cpp
--- a/Engine/Source/Runtime/Engine/Private/World.cpp
+++ b/Engine/Source/Runtime/Engine/Private/World.cpp
@@ -2810,6 +2810,7 @@ bool UWorld::SetGameMode(const FURL& InURL)
 		UE_LOG(LogWorld, Log,  TEXT("Game class is '%s'"), *GameClass->GetName() );
 		FActorSpawnParameters SpawnInfo;
 		SpawnInfo.bNoCollisionFail = true;
+		SpawnInfo.ObjectFlags |= RF_Transient;	// We never want to save game modes into a map
 		AuthorityGameMode = SpawnActor<AGameMode>( GameClass, SpawnInfo );
 		
 		if( AuthorityGameMode != NULL )
diff --git a/Engine/Source/Runtime/Engine/Private/WorldSettings.cpp b/Engine/Source/Runtime/Engine/Private/WorldSettings.cpp
--- a/Engine/Source/Runtime/Engine/Private/WorldSettings.cpp
+++ b/Engine/Source/Runtime/Engine/Private/WorldSettings.cpp
@@ -83,6 +83,7 @@ void AWorldSettings::PreInitializeComponents()
 				FActorSpawnParameters SpawnParameters;
 				SpawnParameters.Owner = this;
 				SpawnParameters.Instigator = Instigator;
+				SpawnParameters.ObjectFlags |= RF_Transient;	// We never want to save particle event managers into a map
 				World->MyParticleEventManager = World->SpawnActor<AParticleEventManager>(ParticleEventManagerClass, SpawnParameters );
 			}
 		}
