Allowed face index checking to GetHitResultAtScreenPosition() in PlayerController
- Also refactored the code a bit to clean it up

diff --git a/Engine/Source/Runtime/Engine/Classes/GameFramework/PlayerController.h b/Engine/Source/Runtime/Engine/Classes/GameFramework/PlayerController.h
--- a/Engine/Source/Runtime/Engine/Classes/GameFramework/PlayerController.h
+++ b/Engine/Source/Runtime/Engine/Classes/GameFramework/PlayerController.h
@@ -415,6 +415,7 @@ public:
 	UFUNCTION(BlueprintCallable, Category="Input")
 	virtual void ResetIgnoreInputFlags();
 
+	bool GetHitResultAtScreenPosition(const FVector2D ScreenPosition, const ECollisionChannel TraceChannel, const struct FCollisionQueryParams& CollisionQueryParams, FHitResult& HitResult) const;
 	bool GetHitResultAtScreenPosition(const FVector2D ScreenPosition, const ECollisionChannel TraceChannel, bool bTraceComplex, FHitResult& HitResult) const;
 	bool GetHitResultAtScreenPosition(const FVector2D ScreenPosition, const ETraceTypeQuery TraceChannel, bool bTraceComplex, FHitResult& HitResult) const;
 	bool GetHitResultAtScreenPosition(const FVector2D ScreenPosition, const TArray<TEnumAsByte<EObjectTypeQuery> > & ObjectTypes, bool bTraceComplex, FHitResult& HitResult) const;
diff --git a/Engine/Source/Runtime/Engine/Private/PlayerController.cpp b/Engine/Source/Runtime/Engine/Private/PlayerController.cpp
--- a/Engine/Source/Runtime/Engine/Private/PlayerController.cpp
+++ b/Engine/Source/Runtime/Engine/Private/PlayerController.cpp
@@ -1917,7 +1917,7 @@ bool APlayerController::ProjectWorldLocationToScreen(FVector WorldLocation, FVec
 }
 
 
-bool APlayerController::GetHitResultAtScreenPosition(const FVector2D ScreenPosition, const ECollisionChannel TraceChannel, bool bTraceComplex, FHitResult& HitResult) const
+bool APlayerController::GetHitResultAtScreenPosition(const FVector2D ScreenPosition, const ECollisionChannel TraceChannel, const FCollisionQueryParams& CollisionQueryParams, FHitResult& HitResult) const
 {
 	// Early out if we clicked on a HUD hitbox
 	if( GetHUD() != NULL && GetHUD()->GetHitBoxAtCoordinates(ScreenPosition, true) )
@@ -1948,49 +1948,23 @@ bool APlayerController::GetHitResultAtScreenPosition(const FVector2D ScreenPosit
 			FVector WorldDirection;
 			SceneView->DeprojectFVector2D(ScreenPosition, WorldOrigin, WorldDirection);
 
-			return GetWorld()->LineTraceSingle(HitResult, WorldOrigin, WorldOrigin + WorldDirection * 100000.f, TraceChannel, FCollisionQueryParams("ClickableTrace", bTraceComplex));
+			return GetWorld()->LineTraceSingle(HitResult, WorldOrigin, WorldOrigin + WorldDirection * 100000.f, TraceChannel, CollisionQueryParams);
 		}
 	}
 
 	return false;
 }
 
-bool APlayerController::GetHitResultAtScreenPosition(const FVector2D ScreenPosition, const ETraceTypeQuery TraceChannel, bool bTraceComplex, FHitResult& HitResult) const
-{
-	// Early out if we clicked on a HUD hitbox
-	if (GetHUD() != NULL && GetHUD()->GetHitBoxAtCoordinates(ScreenPosition, true))
-	{
-		return false;
-	}
-
-	ULocalPlayer* LocalPlayer = Cast<ULocalPlayer>(Player);
-
-	if (LocalPlayer != NULL && LocalPlayer->ViewportClient != NULL && LocalPlayer->ViewportClient->Viewport != NULL)
-	{
-		// Create a view family for the game viewport
-		FSceneViewFamilyContext ViewFamily( FSceneViewFamily::ConstructionValues(
-			LocalPlayer->ViewportClient->Viewport,
-			GetWorld()->Scene,
-			LocalPlayer->ViewportClient->EngineShowFlags )
-			.SetRealtimeUpdate(true) );
-
 
-		// Calculate a view where the player is to update the streaming from the players start location
-		FVector ViewLocation;
-		FRotator ViewRotation;
-		FSceneView* SceneView = LocalPlayer->CalcSceneView( &ViewFamily, /*out*/ ViewLocation, /*out*/ ViewRotation, LocalPlayer->ViewportClient->Viewport );
-
-		if (SceneView)
-		{
-			FVector WorldOrigin;
-			FVector WorldDirection;
-			SceneView->DeprojectFVector2D(ScreenPosition, WorldOrigin, WorldDirection);
-
-			return GetWorld()->LineTraceSingle(HitResult, WorldOrigin, WorldOrigin + WorldDirection * 100000.f, UEngineTypes::ConvertToCollisionChannel(TraceChannel), FCollisionQueryParams("ClickableTrace", bTraceComplex));
-		}
-	}
+bool APlayerController::GetHitResultAtScreenPosition(const FVector2D ScreenPosition, const ECollisionChannel TraceChannel, bool bTraceComplex, FHitResult& HitResult) const
+{
+	FCollisionQueryParams CollisionQueryParams( "ClickableTrace", bTraceComplex );
+	return GetHitResultAtScreenPosition( ScreenPosition, TraceChannel, CollisionQueryParams, HitResult );
+}
 
-	return false;
+bool APlayerController::GetHitResultAtScreenPosition(const FVector2D ScreenPosition, const ETraceTypeQuery TraceChannel, bool bTraceComplex, FHitResult& HitResult) const
+{
+	return GetHitResultAtScreenPosition( ScreenPosition, UEngineTypes::ConvertToCollisionChannel( TraceChannel ), bTraceComplex, HitResult );
 }
 
 bool APlayerController::GetHitResultAtScreenPosition(const FVector2D ScreenPosition, const TArray<TEnumAsByte<EObjectTypeQuery> > & ObjectTypes, bool bTraceComplex, FHitResult& HitResult) const
