Slate: Added support for momentary mouse capture when only the RMB is pressed
Slate: Fixed cursor position not being restored to where it was hidden when using InputModes with mouse capture

diff --git a/Engine/Source/Runtime/Engine/Private/Slate/SceneViewport.cpp b/Engine/Source/Runtime/Engine/Private/Slate/SceneViewport.cpp
--- a/Engine/Source/Runtime/Engine/Private/Slate/SceneViewport.cpp
+++ b/Engine/Source/Runtime/Engine/Private/Slate/SceneViewport.cpp
@@ -31,6 +31,7 @@ FSceneViewport::FSceneViewport( FViewportClient* InViewportClient, TSharedPtr<SV
 	, bPlayInEditorGetsMouseControl( true )
 	, bPlayInEditorIsSimulate( false )
 	, bCursorHiddenDueToCapture( false )
+	, MousePosBeforeHiddenDueToCapture( -1, -1 )
 {
 	bIsSlateViewport = true;
 }
@@ -373,7 +374,10 @@ FReply FSceneViewport::OnMouseButtonDown( const FGeometry& InGeometry, const FPo
 			CurrentReplyState = FReply::Unhandled(); 
 		}
 
-		if (ViewportClient->CaptureMouseOnClick() != EMouseCaptureMode::NoCapture && !ViewportClient->IgnoreInput())
+		if (!ViewportClient->IgnoreInput() &&
+			( ViewportClient->CaptureMouseOnClick() == EMouseCaptureMode::CapturePermanently ||
+			  ViewportClient->CaptureMouseOnClick() == EMouseCaptureMode::CaptureDuringMouseDown ||
+			  ( ViewportClient->CaptureMouseOnClick() == EMouseCaptureMode::CaptureDuringRightMouseDown && InMouseEvent.GetEffectingButton() == EKeys::RightMouseButton ) ) )
 		{
 			TSharedRef<SViewport> ViewportWidgetRef = ViewportWidget.Pin().ToSharedRef();
 
@@ -390,6 +394,7 @@ FReply FSceneViewport::OnMouseButtonDown( const FGeometry& InGeometry, const FPo
 				if (ViewportClient->HideCursorDuringCapture() && bShouldShowMouseCursor)
 				{
 					bCursorHiddenDueToCapture = true;
+					MousePosBeforeHiddenDueToCapture = FIntPoint( InMouseEvent.GetScreenSpacePosition().X, InMouseEvent.GetScreenSpacePosition().Y );
 				}
 				if (bCursorHiddenDueToCapture || !bShouldShowMouseCursor)
 				{
@@ -430,7 +435,10 @@ FReply FSceneViewport::OnMouseButtonUp( const FGeometry& InGeometry, const FPoin
 			CurrentReplyState = FReply::Unhandled(); 
 		}
 		bIsCursorVisible = ViewportClient->GetCursor(this, GetMouseX(), GetMouseY()) != EMouseCursor::None;
-		bReleaseMouse = bIsCursorVisible || ViewportClient->CaptureMouseOnClick() == EMouseCaptureMode::CaptureDuringMouseDown;
+		bReleaseMouse = 
+			bIsCursorVisible || 
+			ViewportClient->CaptureMouseOnClick() == EMouseCaptureMode::CaptureDuringMouseDown ||
+			( ViewportClient->CaptureMouseOnClick() == EMouseCaptureMode::CaptureDuringRightMouseDown && InMouseEvent.GetEffectingButton() == EKeys::RightMouseButton );
 	}
 	if (!((FApp::IsGame() && !GIsEditor) || bIsPlayInEditorViewport) || bReleaseMouse)
 	{
@@ -438,7 +446,12 @@ FReply FSceneViewport::OnMouseButtonUp( const FGeometry& InGeometry, const FPoin
 		// as long as the left or right mouse buttons are not still down
 		if( !InMouseEvent.IsMouseButtonDown( EKeys::RightMouseButton ) && !InMouseEvent.IsMouseButtonDown( EKeys::LeftMouseButton ))
 		{
-			bCursorHiddenDueToCapture = false;
+			if( bCursorHiddenDueToCapture )
+			{
+				bCursorHiddenDueToCapture = false;
+				CurrentReplyState.SetMousePos( MousePosBeforeHiddenDueToCapture );
+				MousePosBeforeHiddenDueToCapture = FIntPoint( -1, -1 );
+			}
 
 			CurrentReplyState.ReleaseMouseCapture();
 			if (bIsCursorVisible)
diff --git a/Engine/Source/Runtime/Engine/Public/Slate/SceneViewport.h b/Engine/Source/Runtime/Engine/Public/Slate/SceneViewport.h
--- a/Engine/Source/Runtime/Engine/Public/Slate/SceneViewport.h
+++ b/Engine/Source/Runtime/Engine/Public/Slate/SceneViewport.h
@@ -322,6 +322,8 @@ private:
 	bool bPlayInEditorIsSimulate;
 	/** Whether or not the cursor is hidden when the viewport captures the mouse */
 	bool bCursorHiddenDueToCapture;
+	/** Position the cursor was at when we hid it due to capture, so we can put it back afterwards */
+	FIntPoint MousePosBeforeHiddenDueToCapture;
 };
 
 
diff --git a/Engine/Source/Runtime/Engine/Public/UnrealClient.h b/Engine/Source/Runtime/Engine/Public/UnrealClient.h
--- a/Engine/Source/Runtime/Engine/Public/UnrealClient.h
+++ b/Engine/Source/Runtime/Engine/Public/UnrealClient.h
@@ -645,7 +645,9 @@ namespace EMouseCaptureMode
 		/** Capture the mouse permanently when the viewport is clicked */
 		CapturePermanently,
 		/** Capture the mouse during a mouse down, releases on mouse up */
-		CaptureDuringMouseDown
+		CaptureDuringMouseDown,
+		/** Capture only when the right mouse button is down, not any of the other mouse buttons */
+		CaptureDuringRightMouseDown,
 	};
 }
 
