PR #542: Fixed crash when using Canvas Render Target objects
- Fixed dereferencing null world pointer in UCanvasRenderTarget2D to get RHI feature level (https://github.com/EpicGames/UnrealEngine/pull/542)
- The "Get Size" Blueprint node for Canvas Render Target is now a "pure" node

diff --git a/Engine/Source/Runtime/Engine/Classes/Engine/CanvasRenderTarget2D.h b/Engine/Source/Runtime/Engine/Classes/Engine/CanvasRenderTarget2D.h
--- a/Engine/Source/Runtime/Engine/Classes/Engine/CanvasRenderTarget2D.h
+++ b/Engine/Source/Runtime/Engine/Classes/Engine/CanvasRenderTarget2D.h
@@ -39,15 +39,16 @@ public:
 	/**
 	 * Creates a new canvas render target and initializes it to the specified dimensions
 	 *
-	 * @param	Class				Class of the render target.  Unless you want to use a special sub-class, you can simply pass UCanvasRenderTarget2D::StaticClass() here.
+	 * @param	WorldContextObject	The world where this render target will be rendered for
+	 * @param	CanvasRenderTarget2DClass	Class of the render target.  Unless you want to use a special sub-class, you can simply pass UCanvasRenderTarget2D::StaticClass() here.
 	 * @param	Width				Width of the render target.
 	 * @param	Height				Height of the render target.
 	 *
 	 * @return						Returns the instanced render target.
 	 */
-	UFUNCTION(BlueprintCallable, Category="Canvas Render Target 2D")
-	static UCanvasRenderTarget2D* CreateCanvasRenderTarget2D(TSubclassOf<UCanvasRenderTarget2D> CanvasRenderTarget2DClass, int32 Width, int32 Height);
-
+	UFUNCTION(BlueprintCallable, Category="Canvas Render Target 2D",meta=(HidePin="WorldContextObject", DefaultToSelf="WorldContextObject" ))
+	static UCanvasRenderTarget2D* CreateCanvasRenderTarget2D(UObject* WorldContextObject, TSubclassOf<UCanvasRenderTarget2D> CanvasRenderTarget2DClass, int32 Width = 1024, int32 Height = 1024);
+ 
 	/**
 	 * Allows a Blueprint to implement how this Canvas Render Target 2D should be updated.
 	 *
@@ -64,11 +65,22 @@ public:
 	 * @param	Width	Output variable for the render target's width
 	 * @param	Height	Output variable for the render target's height
 	 */
-	UFUNCTION(BlueprintCallable, Category="Canvas Render Target 2D")
+	UFUNCTION(BlueprintPure, Category="Canvas Render Target 2D")
 	void GetSize(int32& Width, int32& Height);
 
 
 	/** Called when this Canvas Render Target is asked to update its texture resource. */
 	UPROPERTY(BlueprintAssignable, Category="Canvas Render Target 2D")
 	FOnCanvasRenderTargetUpdate OnCanvasRenderTargetUpdate;
+
+	// UObject overrides
+	virtual UWorld* GetWorld() const override;
+	
+
+protected:
+
+	/* The world this render target will be used with */
+	UPROPERTY()
+	TWeakObjectPtr<UWorld> World;
+	
 };
\ No newline at end of file
diff --git a/Engine/Source/Runtime/Engine/Private/CanvasRenderTarget2D.cpp b/Engine/Source/Runtime/Engine/Private/CanvasRenderTarget2D.cpp
--- a/Engine/Source/Runtime/Engine/Private/CanvasRenderTarget2D.cpp
+++ b/Engine/Source/Runtime/Engine/Private/CanvasRenderTarget2D.cpp
@@ -4,7 +4,8 @@
 #include "Engine/CanvasRenderTarget2D.h"
 
 UCanvasRenderTarget2D::UCanvasRenderTarget2D( const FObjectInitializer& ObjectInitializer )
-	: Super(ObjectInitializer)
+	: Super(ObjectInitializer),
+	  World( nullptr )
 {
 	bNeedsTwoCopies = false;
 }
@@ -50,7 +51,9 @@ void UCanvasRenderTarget2D::UpdateResource()
 	);
 
 	// Create the FCanvas which does the actual rendering.
-	FCanvas RenderCanvas(GameThread_GetRenderTargetResource(), nullptr, FApp::GetCurrentTime() - GStartTime, FApp::GetDeltaTime(), FApp::GetCurrentTime() - GStartTime, GetWorld()->FeatureLevel);
+	const UWorld* WorldPtr = World.Get();
+	const ERHIFeatureLevel::Type FeatureLevel = WorldPtr != nullptr ? World->FeatureLevel : GMaxRHIFeatureLevel;
+	FCanvas RenderCanvas(GameThread_GetRenderTargetResource(), nullptr, FApp::GetCurrentTime() - GStartTime, FApp::GetDeltaTime(), FApp::GetCurrentTime() - GStartTime, FeatureLevel);
 	Canvas->Canvas = &RenderCanvas;
 
 	if (!IsPendingKill() && OnCanvasRenderTargetUpdate.IsBound())
@@ -78,13 +81,14 @@ void UCanvasRenderTarget2D::UpdateResource()
 	);
 }
 
-UCanvasRenderTarget2D* UCanvasRenderTarget2D::CreateCanvasRenderTarget2D(TSubclassOf<UCanvasRenderTarget2D> CanvasRenderTarget2DClass, int32 Width, int32 Height)
+UCanvasRenderTarget2D* UCanvasRenderTarget2D::CreateCanvasRenderTarget2D(UObject* WorldContextObject, TSubclassOf<UCanvasRenderTarget2D> CanvasRenderTarget2DClass, int32 Width, int32 Height)
 {
 	if ((Width > 0) && (Height > 0) && (CanvasRenderTarget2DClass != NULL))
 	{
 		UCanvasRenderTarget2D* NewCanvasRenderTarget = ConstructObject<UCanvasRenderTarget2D>(CanvasRenderTarget2DClass, GetTransientPackage());
 		if (NewCanvasRenderTarget)
 		{
+			NewCanvasRenderTarget->World = GEngine->GetWorldFromContextObject(WorldContextObject);
 			NewCanvasRenderTarget->InitAutoFormat(Width, Height);
 			return NewCanvasRenderTarget;
 		}
@@ -97,4 +101,10 @@ void UCanvasRenderTarget2D::GetSize(int32& Width, int32& Height)
 {
 	Width = GetSurfaceWidth();
 	Height = GetSurfaceHeight();
+}
+
+
+UWorld* UCanvasRenderTarget2D::GetWorld() const
+{
+	return World.Get();
 }
\ No newline at end of file
