Fix for physics material access crash

diff --git a/Engine/Source/Runtime/Engine/Private/PhysicsEngine/BodyInstance.cpp b/Engine/Source/Runtime/Engine/Private/PhysicsEngine/BodyInstance.cpp
--- a/Engine/Source/Runtime/Engine/Private/PhysicsEngine/BodyInstance.cpp
+++ b/Engine/Source/Runtime/Engine/Private/PhysicsEngine/BodyInstance.cpp
@@ -1777,6 +1777,7 @@ void FBodyInstance::InitBodies(FBodyInstance** Bodies, FTransform* Transforms, i
 
 		if(PNewActorSync)
 		{
+			SCOPED_SCENE_WRITE_LOCK(PSceneSync);
 			BodySetup->AddShapesToRigidActor(PNewActorSync, Instance->Scale3D);
 			bInitFail |= PNewActorSync->getNbShapes() == 0;
 			PNewActorSync->userData = &Instance->PhysxUserData;
@@ -1790,6 +1791,7 @@ void FBodyInstance::InitBodies(FBodyInstance** Bodies, FTransform* Transforms, i
 
 		if(PNewActorAsync)
 		{
+			SCOPED_SCENE_WRITE_LOCK(PSceneAsync);
 			check(PSceneAsync);
 			BodySetup->AddShapesToRigidActor(PNewActorAsync, Instance->Scale3D);
 			bInitFail |= PNewActorAsync->getNbShapes() == 0;
@@ -1844,12 +1846,24 @@ void FBodyInstance::InitBodies(FBodyInstance** Bodies, FTransform* Transforms, i
 	PxMaterial* PhysxSimpleMat = SimpleMat->GetPhysXMaterial();
 	
 	// Update material parameters
-	for(int32 ShapeIdx = 0 ; ShapeIdx < ShapesWritten ; ++ShapeIdx)
+	//for(int32 ShapeIdx = 0 ; ShapeIdx < ShapesWritten ; ++ShapeIdx)
+	//{
+	//	SCOPED_SCENE_WRITE_LOCK(PSceneSync);
+	//	SCOPED_SCENE_WRITE_LOCK(PSceneAsync);
+	//	PxShape* Shape = PhysxShapes[ShapeIdx];
+	//	ApplyMaterialToShape(Shape, PhysxSimpleMat, ComplexMats);
+	//}
+
 	{
-		PxShape* Shape = PhysxShapes[ShapeIdx];
-		ApplyMaterialToShape(Shape, PhysxSimpleMat, ComplexMats);
+		SCOPED_SCENE_WRITE_LOCK(PSceneSync);
+		SCOPED_SCENE_WRITE_LOCK(PSceneAsync);
+		for(int32 BodyIdx = 0 ; BodyIdx < NumBodies ; ++BodyIdx)
+		{
+			FBodyInstance* Instance = Bodies[BodyIdx];
+			Instance->ApplyMaterialToInstanceShapes(PhysxSimpleMat, ComplexMats);
+		}
 	}
-
+	
 	{
 		SCOPE_CYCLE_COUNTER(STAT_BulkSceneAdd);
 		// Use the aggregate if it exists, has enough slots and is in the correct scene (or no scene)
