Fixing Mono CIS errors.

diff --git a/Engine/Source/Runtime/Engine/Private/EdGraph/EdGraphSchema.cpp b/Engine/Source/Runtime/Engine/Private/EdGraph/EdGraphSchema.cpp
--- a/Engine/Source/Runtime/Engine/Private/EdGraph/EdGraphSchema.cpp
+++ b/Engine/Source/Runtime/Engine/Private/EdGraph/EdGraphSchema.cpp
@@ -440,7 +440,9 @@ FPinConnectionResponse UEdGraphSchema::MovePinLinks(UEdGraphPin& MoveFromPin, UE
 	TArray<UEdGraphPin*> CurrentLinks = MoveFromPin.LinkedTo;
 	// Then break all links at pin we are moving from
 	MoveFromPin.BreakAllPinLinks();
+#if WITH_EDITOR
 	MoveFromPin.GetOwningNode()->PinConnectionListChanged(&MoveFromPin);
+#endif // #if WITH_EDITOR
 	// Try and make each new connection
 	for (int32 i=0; i<CurrentLinks.Num(); i++)
 	{
@@ -454,13 +456,17 @@ FPinConnectionResponse UEdGraphSchema::MovePinLinks(UEdGraphPin& MoveFromPin, UE
 		{
 			FinalResponse = Response;
 		}
+#if WITH_EDITOR
 		NewLink->GetOwningNode()->PinConnectionListChanged(NewLink);
+#endif // #if WITH_EDITOR
 	}
 	// Move over the default values
 	MoveToPin.DefaultValue = MoveFromPin.DefaultValue;
 	MoveToPin.DefaultObject = MoveFromPin.DefaultObject;
 	MoveToPin.DefaultTextValue = MoveFromPin.DefaultTextValue;
+#if WITH_EDITOR
 	MoveToPin.GetOwningNode()->PinConnectionListChanged(&MoveToPin);
+#endif // #if WITH_EDITOR
 	return FinalResponse;
 }
 
