UE4 - Bump up material DDC to safely fix issues with per frame expressions adding/removal from assets

diff --git a/Engine/Shaders/ShaderVersion.usf b/Engine/Shaders/ShaderVersion.usf
--- a/Engine/Shaders/ShaderVersion.usf
+++ b/Engine/Shaders/ShaderVersion.usf
@@ -2,4 +2,4 @@
 // Each time the console command is executed it generates a new GUID. As this file is included
 // in common.usf (which should be included in any shader) it allows to invalidate the shader DDC.
 // 
-// GUID = {CAC8D6F0-2E83-4C3D-A767-584D1E3DC80E}a
+// GUID = {CAC8D6F0-2E83-4C3D-A767-584D1E3DC80g}
diff --git a/Engine/Source/Runtime/Engine/Private/ShaderDerivedDataVersion.h b/Engine/Source/Runtime/Engine/Private/ShaderDerivedDataVersion.h
--- a/Engine/Source/Runtime/Engine/Private/ShaderDerivedDataVersion.h
+++ b/Engine/Source/Runtime/Engine/Private/ShaderDerivedDataVersion.h
@@ -10,4 +10,4 @@
 // guid as version
 
 #define GLOBALSHADERMAP_DERIVEDDATA_VER			TEXT("af1606fcec114373844d2b6cf9df71c8")
-#define MATERIALSHADERMAP_DERIVEDDATA_VER		TEXT("e2c1e89ee9ae4f3ca7b8e0f2c49b34eb")
+#define MATERIALSHADERMAP_DERIVEDDATA_VER		TEXT("07f9bdef41d340979c7d94a22ad6a34e")
