UE4 - Back out the back out (2388995); proper fixes version check; restores deleted parameter collection code & restores per frame material expression code; still no generation of actual per frame expressions

diff --git a/Engine/Shaders/ShaderVersion.usf b/Engine/Shaders/ShaderVersion.usf
--- a/Engine/Shaders/ShaderVersion.usf
+++ b/Engine/Shaders/ShaderVersion.usf
@@ -2,4 +2,4 @@
 // Each time the console command is executed it generates a new GUID. As this file is included
 // in common.usf (which should be included in any shader) it allows to invalidate the shader DDC.
 // 
-// GUID = {CAC8D6F0-2E83-4C3D-A767-584D1E3DC80g}
+// GUID = {CAC8D6F0-2E83-4C3D-A767-584D1E3DC80x}
diff --git a/Engine/Source/Runtime/Engine/Private/Materials/HLSLMaterialTranslator.h b/Engine/Source/Runtime/Engine/Private/Materials/HLSLMaterialTranslator.h
--- a/Engine/Source/Runtime/Engine/Private/Materials/HLSLMaterialTranslator.h
+++ b/Engine/Source/Runtime/Engine/Private/Materials/HLSLMaterialTranslator.h
@@ -51,6 +51,19 @@ EMaterialValueType GetVectorType(uint32 NumComponents)
 	};
 }
 
+static inline int32 SwizzleComponentToIndex(TCHAR Component)
+{
+	switch (Component)
+	{
+	case TCHAR('x'): case TCHAR('X'): case TCHAR('r'): case TCHAR('R'): return 0;
+	case TCHAR('y'): case TCHAR('Y'): case TCHAR('g'): case TCHAR('G'): return 1;
+	case TCHAR('z'): case TCHAR('Z'): case TCHAR('b'): case TCHAR('B'): return 2;
+	case TCHAR('w'): case TCHAR('W'): case TCHAR('a'): case TCHAR('A'): return 3;
+	default:
+		return -1;
+	}
+}
+
 struct FShaderCodeChunk
 {
 	/** 
@@ -416,6 +429,19 @@ public:
 				ResourcesString += CustomExpressionImplementations[ExpressionIndex] + "\r\n\r\n";
 			}
 
+			// Per frame expressions
+			{
+				for (int32 Index = 0, Num = MaterialCompilationOutput.UniformExpressionSet.PerFrameUniformScalarExpressions.Num(); Index < Num; ++Index)
+				{
+					ResourcesString += FString::Printf(TEXT("float UE_Material_PerFrameScalarExpression%u;"), Index) + "\r\n\r\n";
+				}
+
+				for (int32 Index = 0, Num = MaterialCompilationOutput.UniformExpressionSet.PerFrameUniformVectorExpressions.Num(); Index < Num; ++Index)
+				{
+					ResourcesString += FString::Printf(TEXT("float4 UE_Material_PerFrameVectorExpression%u;"), Index) + "\r\n\r\n";
+				}
+			}
+
 			for(uint32 PropertyId = 0; PropertyId < MP_MAX; ++PropertyId)
 			{
 				if(PropertyId == MP_MaterialAttributes )
@@ -935,9 +961,9 @@ protected:
 			}
 		}
 
-		int32		BufferSize		= 256;
+		int32	BufferSize		= 256;
 		TCHAR*	FormattedCode	= NULL;
-		int32		Result			= -1;
+		int32	Result			= -1;
 
 		while(Result == -1)
 		{
@@ -976,14 +1002,21 @@ protected:
 		TCHAR FormattedCode[MAX_SPRINTF]=TEXT("");
 		if(CodeChunk.Type == MCT_Float)
 		{
-			const static TCHAR IndexToMask[] = {'x', 'y', 'z', 'w'};
-			const int32 ScalarInputIndex = MaterialCompilationOutput.UniformExpressionSet.UniformScalarExpressions.AddUnique(CodeChunk.UniformExpression);
-			// Update the above FMemory::Malloc if this FCString::Sprintf grows in size, e.g. %s, ...
-			FCString::Sprintf(FormattedCode, TEXT("Material.ScalarExpressions[%u].%c"), ScalarInputIndex / 4, IndexToMask[ScalarInputIndex % 4]);
+			if (CodeChunk.UniformExpression->IsChangingPerFrame())
+			{
+				const int32 ScalarInputIndex = MaterialCompilationOutput.UniformExpressionSet.PerFrameUniformScalarExpressions.AddUnique(CodeChunk.UniformExpression);
+				FCString::Sprintf(FormattedCode, TEXT("UE_Material_PerFrameScalarExpression%u"), ScalarInputIndex);
+			}
+			else
+			{
+				const static TCHAR IndexToMask[] = {'x', 'y', 'z', 'w'};
+				const int32 ScalarInputIndex = MaterialCompilationOutput.UniformExpressionSet.UniformScalarExpressions.AddUnique(CodeChunk.UniformExpression);
+				// Update the above FMemory::Malloc if this FCString::Sprintf grows in size, e.g. %s, ...
+				FCString::Sprintf(FormattedCode, TEXT("Material.ScalarExpressions[%u].%c"), ScalarInputIndex / 4, IndexToMask[ScalarInputIndex % 4]);
+			}
 		}
 		else if(CodeChunk.Type & MCT_Float)
 		{
-			const int32 VectorInputIndex = MaterialCompilationOutput.UniformExpressionSet.UniformVectorExpressions.AddUnique(CodeChunk.UniformExpression);
 			const TCHAR* Mask;
 			switch(CodeChunk.Type)
 			{
@@ -994,10 +1027,20 @@ protected:
 			default: Mask = TEXT(""); break;
 			};
 
-			FCString::Sprintf(FormattedCode, TEXT("Material.VectorExpressions[%u]%s"), VectorInputIndex, Mask);
+			if (CodeChunk.UniformExpression->IsChangingPerFrame())
+			{
+				const int32 VectorInputIndex = MaterialCompilationOutput.UniformExpressionSet.PerFrameUniformVectorExpressions.AddUnique(CodeChunk.UniformExpression);
+				FCString::Sprintf(FormattedCode, TEXT("UE_Material_PerFrameVectorExpression%u%s"), VectorInputIndex, Mask);
+			}
+			else
+			{
+				const int32 VectorInputIndex = MaterialCompilationOutput.UniformExpressionSet.UniformVectorExpressions.AddUnique(CodeChunk.UniformExpression);
+				FCString::Sprintf(FormattedCode, TEXT("Material.VectorExpressions[%u]%s"), VectorInputIndex, Mask);
+			}
 		}
 		else if(CodeChunk.Type & MCT_Texture)
 		{
+			check(!CodeChunk.UniformExpression->IsChangingPerFrame());
 			int32 TextureInputIndex = INDEX_NONE;
 			const TCHAR* BaseName = TEXT("");
 			switch(CodeChunk.Type)
diff --git a/Engine/Source/Runtime/Engine/Private/Materials/MaterialShared.cpp b/Engine/Source/Runtime/Engine/Private/Materials/MaterialShared.cpp
--- a/Engine/Source/Runtime/Engine/Private/Materials/MaterialShared.cpp
+++ b/Engine/Source/Runtime/Engine/Private/Materials/MaterialShared.cpp
@@ -1500,6 +1500,8 @@ FMaterialRenderContext::FMaterialRenderContext(
 	const FSceneView* InView)
 		: MaterialRenderProxy(InMaterialRenderProxy)
 		, Material(InMaterial)
+		, Time(0.0f)
+		, RealTime(0.0f)
 {
 	bShowSelection = GIsEditor && InView && InView->Family->EngineShowFlags.Selection;
 }
diff --git a/Engine/Source/Runtime/Engine/Private/Materials/MaterialUniformExpressions.cpp b/Engine/Source/Runtime/Engine/Private/Materials/MaterialUniformExpressions.cpp
--- a/Engine/Source/Runtime/Engine/Private/Materials/MaterialUniformExpressions.cpp
+++ b/Engine/Source/Runtime/Engine/Private/Materials/MaterialUniformExpressions.cpp
@@ -91,6 +91,9 @@ void FUniformExpressionSet::Serialize(FArchive& Ar)
 
 	Ar << ParameterCollections;
 
+	Ar << PerFrameUniformScalarExpressions;
+	Ar << PerFrameUniformVectorExpressions;
+
 	// Recreate the uniform buffer struct after loading.
 	if(Ar.IsLoading())
 	{
@@ -104,6 +107,8 @@ bool FUniformExpressionSet::IsEmpty() const
 		&& UniformScalarExpressions.Num() == 0
 		&& Uniform2DTextureExpressions.Num() == 0
 		&& UniformCubeTextureExpressions.Num() == 0
+		&& PerFrameUniformScalarExpressions.Num() == 0
+		&& PerFrameUniformVectorExpressions.Num() == 0
 		&& ParameterCollections.Num() == 0;
 }
 
@@ -113,6 +118,8 @@ bool FUniformExpressionSet::operator==(const FUniformExpressionSet& ReferenceSet
 	|| UniformScalarExpressions.Num() != ReferenceSet.UniformScalarExpressions.Num()
 	|| Uniform2DTextureExpressions.Num() != ReferenceSet.Uniform2DTextureExpressions.Num()
 	|| UniformCubeTextureExpressions.Num() != ReferenceSet.UniformCubeTextureExpressions.Num()
+	|| PerFrameUniformScalarExpressions.Num() != ReferenceSet.PerFrameUniformScalarExpressions.Num()
+	|| PerFrameUniformVectorExpressions.Num() != ReferenceSet.PerFrameUniformVectorExpressions.Num()
 	|| ParameterCollections.Num() != ReferenceSet.ParameterCollections.Num())
 	{
 		return false;
@@ -150,6 +157,22 @@ bool FUniformExpressionSet::operator==(const FUniformExpressionSet& ReferenceSet
 		}
 	}
 
+	for (int32 i = 0; i < PerFrameUniformScalarExpressions.Num(); i++)
+	{
+		if (!PerFrameUniformScalarExpressions[i]->IsIdentical(ReferenceSet.PerFrameUniformScalarExpressions[i]))
+		{
+			return false;
+		}
+	}
+
+	for (int32 i = 0; i < PerFrameUniformVectorExpressions.Num(); i++)
+	{
+		if (!PerFrameUniformVectorExpressions[i]->IsIdentical(ReferenceSet.PerFrameUniformVectorExpressions[i]))
+		{
+			return false;
+		}
+	}
+
 	for (int32 i = 0; i < ParameterCollections.Num(); i++)
 	{
 		if (ParameterCollections[i] != ReferenceSet.ParameterCollections[i])
@@ -163,11 +186,13 @@ bool FUniformExpressionSet::operator==(const FUniformExpressionSet& ReferenceSet
 
 FString FUniformExpressionSet::GetSummaryString() const
 {
-	return FString::Printf(TEXT("(%u vectors, %u scalars, %u 2d tex, %u cube tex, %u collections)"),
+	return FString::Printf(TEXT("(%u vectors, %u scalars, %u 2d tex, %u cube tex, %u scalars/frame, %u vectors/frame, %u collections)"),
 		UniformVectorExpressions.Num(), 
 		UniformScalarExpressions.Num(),
 		Uniform2DTextureExpressions.Num(),
 		UniformCubeTextureExpressions.Num(),
+		PerFrameUniformScalarExpressions.Num(),
+		PerFrameUniformVectorExpressions.Num(),
 		ParameterCollections.Num()
 		);
 }
diff --git a/Engine/Source/Runtime/Engine/Private/Materials/MaterialUniformExpressions.h b/Engine/Source/Runtime/Engine/Private/Materials/MaterialUniformExpressions.h
--- a/Engine/Source/Runtime/Engine/Private/Materials/MaterialUniformExpressions.h
+++ b/Engine/Source/Runtime/Engine/Private/Materials/MaterialUniformExpressions.h
@@ -61,7 +61,10 @@ public:
 	}
 	virtual void GetNumberValue(const FMaterialRenderContext& Context,FLinearColor& OutValue) const
 	{
-		//OutValue.R = Context.CurrentTime;
+		OutValue.R = Context.Time;
+		OutValue.G = Context.Time;
+		OutValue.B = Context.Time;
+		OutValue.A = Context.Time;
 	}
 	virtual bool IsConstant() const
 	{
@@ -86,7 +89,10 @@ public:
 	}
 	virtual void GetNumberValue(const FMaterialRenderContext& Context,FLinearColor& OutValue) const
 	{
-		//OutValue.R = Context.CurrentRealTime;
+		OutValue.R = Context.RealTime;
+		OutValue.G = Context.RealTime;
+		OutValue.B = Context.RealTime;
+		OutValue.A = Context.RealTime;
 	}
 	virtual bool IsConstant() const
 	{
@@ -381,6 +387,10 @@ public:
 	{
 		return X->IsConstant();
 	}
+	virtual bool IsChangingPerFrame() const
+	{
+		return X->IsChangingPerFrame();
+	}
 	virtual bool IsIdentical(const FMaterialUniformExpression* OtherExpression) const
 	{
 		if (GetType() != OtherExpression->GetType())
@@ -426,6 +436,10 @@ public:
 	{
 		return X->IsConstant();
 	}
+	virtual bool IsChangingPerFrame() const
+	{
+		return X->IsChangingPerFrame();
+	}
 	virtual bool IsIdentical(const FMaterialUniformExpression* OtherExpression) const
 	{
 		if (GetType() != OtherExpression->GetType())
@@ -467,6 +481,10 @@ public:
 	{
 		return X->IsConstant();
 	}
+	virtual bool IsChangingPerFrame() const
+	{
+		return X->IsChangingPerFrame();
+	}
 	virtual bool IsIdentical(const FMaterialUniformExpression* OtherExpression) const
 	{
 		if (GetType() != OtherExpression->GetType())
@@ -560,6 +578,12 @@ public:
 	{
 		return A->IsConstant() && B->IsConstant();
 	}
+	virtual bool IsChangingPerFrame() const
+	{
+		const bool AChangesPerFrame = A->IsChangingPerFrame();
+		const bool BChangesPerFrame = B->IsChangingPerFrame();
+		return AChangesPerFrame || BChangesPerFrame;
+	}
 	virtual bool IsIdentical(const FMaterialUniformExpression* OtherExpression) const
 	{
 		if (GetType() != OtherExpression->GetType())
@@ -608,6 +632,10 @@ public:
 	{
 		return X->IsConstant();
 	}
+	virtual bool IsChangingPerFrame() const
+	{
+		return X->IsChangingPerFrame();
+	}
 	virtual bool IsIdentical(const FMaterialUniformExpression* OtherExpression) const
 	{
 		if (GetType() != OtherExpression->GetType())
@@ -657,6 +685,12 @@ public:
 	{
 		return A->IsConstant() && B->IsConstant();
 	}
+	virtual bool IsChangingPerFrame() const
+	{
+		const bool AChangesPerFrame = A->IsChangingPerFrame();
+		const bool BChangesPerFrame = B->IsChangingPerFrame();
+		return AChangesPerFrame || BChangesPerFrame;
+	}
 	virtual bool IsIdentical(const FMaterialUniformExpression* OtherExpression) const
 	{
 		if (GetType() != OtherExpression->GetType())
@@ -707,6 +741,12 @@ public:
 	{
 		return A->IsConstant() && B->IsConstant();
 	}
+	virtual bool IsChangingPerFrame() const
+	{
+		const bool AChangesPerFrame = A->IsChangingPerFrame();
+		const bool BChangesPerFrame = B->IsChangingPerFrame();
+		return AChangesPerFrame || BChangesPerFrame;
+	}
 	virtual bool IsIdentical(const FMaterialUniformExpression* OtherExpression) const
 	{
 		if (GetType() != OtherExpression->GetType())
@@ -756,6 +796,12 @@ public:
 	{
 		return A->IsConstant() && B->IsConstant();
 	}
+	virtual bool IsChangingPerFrame() const
+	{
+		const bool AChangesPerFrame = A->IsChangingPerFrame();
+		const bool BChangesPerFrame = B->IsChangingPerFrame();
+		return AChangesPerFrame || BChangesPerFrame;
+	}
 	virtual bool IsIdentical(const FMaterialUniformExpression* OtherExpression) const
 	{
 		if (GetType() != OtherExpression->GetType())
@@ -808,6 +854,10 @@ public:
 	{
 		return Input->IsConstant() && Min->IsConstant() && Max->IsConstant();
 	}
+	virtual bool IsChangingPerFrame() const
+	{
+		return Input->IsChangingPerFrame() || Min->IsChangingPerFrame() || Max->IsChangingPerFrame();
+	}
 	virtual bool IsIdentical(const FMaterialUniformExpression* OtherExpression) const
 	{
 		if (GetType() != OtherExpression->GetType())
@@ -855,6 +905,10 @@ public:
 	{
 		return Input->IsConstant();
 	}
+	virtual bool IsChangingPerFrame() const
+	{
+		return Input->IsChangingPerFrame();
+	}
 	virtual bool IsIdentical(const FMaterialUniformExpression* OtherExpression) const
 	{
 		if (GetType() != OtherExpression->GetType())
@@ -899,6 +953,10 @@ public:
 	{
 		return X->IsConstant();
 	}
+	virtual bool IsChangingPerFrame() const
+	{
+		return X->IsChangingPerFrame();
+	}
 	virtual bool IsIdentical(const FMaterialUniformExpression* OtherExpression) const
 	{
 		if (GetType() != OtherExpression->GetType())
@@ -943,6 +1001,10 @@ public:
 	{
 		return X->IsConstant();
 	}
+	virtual bool IsChangingPerFrame() const
+	{
+		return X->IsChangingPerFrame();
+	}
 	virtual bool IsIdentical(const FMaterialUniformExpression* OtherExpression) const
 	{
 		if (GetType() != OtherExpression->GetType())
@@ -987,6 +1049,10 @@ public:
 	{
 		return X->IsConstant();
 	}
+	virtual bool IsChangingPerFrame() const
+	{
+		return X->IsChangingPerFrame();
+	}
 	virtual bool IsIdentical(const FMaterialUniformExpression* OtherExpression) const
 	{
 		if (GetType() != OtherExpression->GetType())
@@ -1035,6 +1101,12 @@ public:
 	{
 		return A->IsConstant() && B->IsConstant();
 	}
+	virtual bool IsChangingPerFrame() const
+	{
+		const bool AChangesPerFrame = A->IsChangingPerFrame();
+		const bool BChangesPerFrame = B->IsChangingPerFrame();
+		return AChangesPerFrame || BChangesPerFrame;
+	}
 	virtual bool IsIdentical(const FMaterialUniformExpression* OtherExpression) const
 	{
 		if (GetType() != OtherExpression->GetType())
@@ -1080,6 +1152,10 @@ public:
 	{
 		return X->IsConstant();
 	}
+	virtual bool IsChangingPerFrame() const
+	{
+		return X->IsChangingPerFrame();
+	}
 	virtual bool IsIdentical(const FMaterialUniformExpression* OtherExpression) const
 	{
 		if (GetType() != OtherExpression->GetType())
diff --git a/Engine/Source/Runtime/Engine/Private/ShaderDerivedDataVersion.h b/Engine/Source/Runtime/Engine/Private/ShaderDerivedDataVersion.h
--- a/Engine/Source/Runtime/Engine/Private/ShaderDerivedDataVersion.h
+++ b/Engine/Source/Runtime/Engine/Private/ShaderDerivedDataVersion.h
@@ -10,4 +10,4 @@
 // guid as version
 
 #define GLOBALSHADERMAP_DERIVEDDATA_VER			TEXT("af1606fcec114373844d2b6cf9df71c8")
-#define MATERIALSHADERMAP_DERIVEDDATA_VER		TEXT("07f9bdef41d340979c7d94a22ad6a34e")
+#define MATERIALSHADERMAP_DERIVEDDATA_VER		TEXT("85dd00daf10d45c79d4ed0306af26289")
diff --git a/Engine/Source/Runtime/Engine/Public/MaterialShared.h b/Engine/Source/Runtime/Engine/Public/MaterialShared.h
--- a/Engine/Source/Runtime/Engine/Public/MaterialShared.h
+++ b/Engine/Source/Runtime/Engine/Public/MaterialShared.h
@@ -92,6 +92,11 @@ struct ENGINE_API FMaterialRenderContext
 	const FMaterialRenderProxy* MaterialRenderProxy;
 	/** Material resource to use. */
 	const FMaterial& Material;
+
+	// Used only when evaluating expressions per frame
+	float Time;
+	float RealTime;
+
 	/** Whether or not selected objects should use their selection color. */
 	bool bShowSelection;
 
@@ -167,6 +172,7 @@ public:
 	virtual void GetNumberValue(const struct FMaterialRenderContext& Context,FLinearColor& OutValue) const {}
 	virtual class FMaterialUniformExpressionTexture* GetTextureUniformExpression() { return NULL; }
 	virtual bool IsConstant() const { return false; }
+	virtual bool IsChangingPerFrame() const { return false; }
 	virtual bool IsIdentical(const FMaterialUniformExpression* OtherExpression) const { return false; }
 
 	friend FArchive& operator<<(FArchive& Ar,class FMaterialUniformExpression*& Ref);
@@ -232,6 +238,8 @@ public:
 			+ UniformScalarExpressions.GetAllocatedSize()
 			+ Uniform2DTextureExpressions.GetAllocatedSize()
 			+ UniformCubeTextureExpressions.GetAllocatedSize()
+			+ PerFrameUniformScalarExpressions.GetAllocatedSize()
+			+ PerFrameUniformVectorExpressions.GetAllocatedSize()
 			+ ParameterCollections.GetAllocatedSize()
 			+ (UniformBufferStruct ? (sizeof(FUniformBufferStruct) + UniformBufferStruct->GetMembers().GetAllocatedSize()) : 0);
 	}
@@ -242,6 +250,8 @@ protected:
 	TArray<TRefCountPtr<FMaterialUniformExpression> > UniformScalarExpressions;
 	TArray<TRefCountPtr<FMaterialUniformExpressionTexture> > Uniform2DTextureExpressions;
 	TArray<TRefCountPtr<FMaterialUniformExpressionTexture> > UniformCubeTextureExpressions;
+	TArray<TRefCountPtr<FMaterialUniformExpression> > PerFrameUniformScalarExpressions;
+	TArray<TRefCountPtr<FMaterialUniformExpression> > PerFrameUniformVectorExpressions;
 
 	/** Ids of parameter collections referenced by the material that was translated. */
 	TArray<FGuid> ParameterCollections;
diff --git a/Engine/Source/Runtime/Renderer/Private/ShaderBaseClasses.cpp b/Engine/Source/Runtime/Renderer/Private/ShaderBaseClasses.cpp
--- a/Engine/Source/Runtime/Renderer/Private/ShaderBaseClasses.cpp
+++ b/Engine/Source/Runtime/Renderer/Private/ShaderBaseClasses.cpp
@@ -34,6 +34,20 @@ FMaterialShader::FMaterialShader(const FMaterialShaderType::CompiledShaderInitia
 		ParameterCollectionUniformBuffers.Add(CollectionParameter);
 	}
 
+	for (int32 Index = 0; Index < Initializer.UniformExpressionSet.PerFrameUniformScalarExpressions.Num(); Index++)
+	{
+		FShaderParameter Parameter;
+		Parameter.Bind(Initializer.ParameterMap, *FString::Printf(TEXT("UE_Material_PerFrameScalarExpression%u"), Index));
+		PerFrameScalarExpressions.Add(Parameter);
+	}
+
+	for (int32 Index = 0; Index < Initializer.UniformExpressionSet.PerFrameUniformVectorExpressions.Num(); Index++)
+	{
+		FShaderParameter Parameter;
+		Parameter.Bind(Initializer.ParameterMap, *FString::Printf(TEXT("UE_Material_PerFrameVectorExpression%u"), Index));
+		PerFrameVectorExpressions.Add(Parameter);
+	}
+
 	DeferredParameters.Bind(Initializer.ParameterMap);
 	LightAttenuation.Bind(Initializer.ParameterMap, TEXT("LightAttenuationTexture"));
 	LightAttenuationSampler.Bind(Initializer.ParameterMap, TEXT("LightAttenuationTextureSampler"));
@@ -89,9 +103,9 @@ void FMaterialShader::SetParameters(
 	check(Material.GetRenderingThreadShaderMap()->IsValidForRendering());
 	check(Material.GetFeatureLevel() == FeatureLevel);
 
+#if NO_LOGGING == 0
 	// Validate that the shader is being used for a material that matches the uniform expression set the shader was compiled for.
 	const FUniformExpressionSet& MaterialUniformExpressionSet = Material.GetRenderingThreadShaderMap()->GetUniformExpressionSet();
-#if NO_LOGGING == 0
 	const bool bUniformExpressionSetMismatch = !DebugUniformExpressionSet.Matches(MaterialUniformExpressionSet)
 		|| UniformExpressionCache->CachedUniformExpressionShaderMap != Material.GetRenderingThreadShaderMap();
 
@@ -103,8 +117,8 @@ void FMaterialShader::SetParameters(
 			TEXT("%s shader uniform expression set mismatch for material %s/%s.\n")
 			TEXT("Shader compilation info:                %s\n")
 			TEXT("Material render proxy compilation info: %s\n")
-			TEXT("Shader uniform expression set:   %u vectors, %u scalars, %u 2D textures, %u cube textures, shader map %p\n")
-			TEXT("Material uniform expression set: %u vectors, %u scalars, %u 2D textures, %u cube textures, shader map %p\n"),
+			TEXT("Shader uniform expression set:   %u vectors, %u scalars, %u 2D textures, %u cube textures, %u scalars/frame, %u vectors/frame, shader map %p\n")
+			TEXT("Material uniform expression set: %u vectors, %u scalars, %u 2D textures, %u cube textures, %u scalars/frame, %u vectors/frame, shader map %p\n"),
 			GetType()->GetName(),
 			*MaterialRenderProxy->GetFriendlyName(),
 			*Material.GetFriendlyName(),
@@ -114,11 +128,15 @@ void FMaterialShader::SetParameters(
 			DebugUniformExpressionSet.NumScalarExpressions,
 			DebugUniformExpressionSet.Num2DTextureExpressions,
 			DebugUniformExpressionSet.NumCubeTextureExpressions,
+			DebugUniformExpressionSet.NumPerFrameScalarExpressions,
+			DebugUniformExpressionSet.NumPerFrameVectorExpressions,
 			UniformExpressionCache->CachedUniformExpressionShaderMap,
 			MaterialUniformExpressionSet.UniformVectorExpressions.Num(),
 			MaterialUniformExpressionSet.UniformScalarExpressions.Num(),
 			MaterialUniformExpressionSet.Uniform2DTextureExpressions.Num(),
 			MaterialUniformExpressionSet.UniformCubeTextureExpressions.Num(),
+			MaterialUniformExpressionSet.PerFrameUniformScalarExpressions.Num(),
+			MaterialUniformExpressionSet.PerFrameUniformVectorExpressions.Num(),
 			Material.GetRenderingThreadShaderMap()
 			);
 	}
@@ -149,6 +167,40 @@ void FMaterialShader::SetParameters(
 		}
 	}
 
+	{
+		// Per frame material expressions
+		const int32 NumScalarExpressions = PerFrameScalarExpressions.Num();
+		const int32 NumVectorExpressions = PerFrameVectorExpressions.Num();
+
+		if (NumScalarExpressions > 0 || NumVectorExpressions > 0)
+		{
+			FMaterialRenderContext MaterialRenderContext(MaterialRenderProxy, Material, &View);
+			MaterialRenderContext.Time = View.Family->CurrentWorldTime;
+			MaterialRenderContext.RealTime = View.Family->CurrentRealTime;
+			for (int32 Index = 0; Index < NumScalarExpressions; ++Index)
+			{
+				auto& Parameter = PerFrameScalarExpressions[Index];
+				if (Parameter.IsBound())
+				{
+					FLinearColor TempValue;
+					MaterialUniformExpressionSet.PerFrameUniformScalarExpressions[Index]->GetNumberValue(MaterialRenderContext, TempValue);
+					SetShaderValue(RHICmdList, ShaderRHI, Parameter, TempValue.R);
+				}
+			}
+
+			for (int32 Index = 0; Index < NumVectorExpressions; ++Index)
+			{
+				auto& Parameter = PerFrameVectorExpressions[Index];
+				if (Parameter.IsBound())
+				{
+					FLinearColor TempValue;
+					MaterialUniformExpressionSet.PerFrameUniformVectorExpressions[Index]->GetNumberValue(MaterialRenderContext, TempValue);
+					SetShaderValue(RHICmdList, ShaderRHI, Parameter, TempValue);
+				}
+			}
+		}
+	}
+
 	DeferredParameters.Set(RHICmdList, ShaderRHI, View, TextureMode);
 
 	AtmosphericFogTextureParameters.Set(RHICmdList, ShaderRHI, View);
@@ -255,6 +307,10 @@ bool FMaterialShader::Serialize(FArchive& Ar)
 	Ar << PerlinNoiseGradientTextureSampler;
 	Ar << PerlinNoise3DTexture;
 	Ar << PerlinNoise3DTextureSampler;
+
+	Ar << PerFrameScalarExpressions;
+	Ar << PerFrameVectorExpressions;
+
 	return bShaderHasOutdatedParameters;
 }
 
diff --git a/Engine/Source/Runtime/Renderer/Public/MaterialShader.h b/Engine/Source/Runtime/Renderer/Public/MaterialShader.h
--- a/Engine/Source/Runtime/Renderer/Public/MaterialShader.h
+++ b/Engine/Source/Runtime/Renderer/Public/MaterialShader.h
@@ -34,12 +34,16 @@ public:
 	int32 NumScalarExpressions;
 	int32 Num2DTextureExpressions;
 	int32 NumCubeTextureExpressions;
+	int32 NumPerFrameScalarExpressions;
+	int32 NumPerFrameVectorExpressions;
 
 	FDebugUniformExpressionSet()
 		: NumVectorExpressions(0)
 		, NumScalarExpressions(0)
 		, Num2DTextureExpressions(0)
 		, NumCubeTextureExpressions(0)
+		, NumPerFrameScalarExpressions(0)
+		, NumPerFrameVectorExpressions(0)
 	{
 	}
 
@@ -55,6 +59,8 @@ public:
 		NumScalarExpressions = InUniformExpressionSet.UniformScalarExpressions.Num();
 		Num2DTextureExpressions = InUniformExpressionSet.Uniform2DTextureExpressions.Num();
 		NumCubeTextureExpressions = InUniformExpressionSet.UniformCubeTextureExpressions.Num();
+		NumPerFrameScalarExpressions = InUniformExpressionSet.PerFrameUniformScalarExpressions.Num();
+		NumPerFrameVectorExpressions = InUniformExpressionSet.PerFrameUniformVectorExpressions.Num();
 	}
 
 	/** Returns true if the number of uniform expressions matches those with which the debug set was initialized. */
@@ -63,7 +69,9 @@ public:
 		return NumVectorExpressions == InUniformExpressionSet.UniformVectorExpressions.Num()
 			&& NumScalarExpressions == InUniformExpressionSet.UniformScalarExpressions.Num()
 			&& Num2DTextureExpressions == InUniformExpressionSet.Uniform2DTextureExpressions.Num()
-			&& NumCubeTextureExpressions == InUniformExpressionSet.UniformCubeTextureExpressions.Num();
+			&& NumCubeTextureExpressions == InUniformExpressionSet.UniformCubeTextureExpressions.Num()
+			&& NumPerFrameScalarExpressions == InUniformExpressionSet.PerFrameUniformScalarExpressions.Num()
+			&& NumPerFrameVectorExpressions == InUniformExpressionSet.PerFrameUniformVectorExpressions.Num();
 	}
 };
 
@@ -72,6 +80,8 @@ inline FArchive& operator<<(FArchive& Ar, FDebugUniformExpressionSet& DebugExpre
 {
 	Ar << DebugExpressionSet.NumVectorExpressions;
 	Ar << DebugExpressionSet.NumScalarExpressions;
+	Ar << DebugExpressionSet.NumPerFrameScalarExpressions;
+	Ar << DebugExpressionSet.NumPerFrameVectorExpressions;
 	Ar << DebugExpressionSet.Num2DTextureExpressions;
 	Ar << DebugExpressionSet.NumCubeTextureExpressions;
 	return Ar;
@@ -128,6 +138,8 @@ private:
 
 	FShaderUniformBufferParameter MaterialUniformBuffer;
 	TArray<FShaderUniformBufferParameter> ParameterCollectionUniformBuffers;
+	TArray<FShaderParameter> PerFrameScalarExpressions;
+	TArray<FShaderParameter> PerFrameVectorExpressions;
 	FDeferredPixelShaderParameters DeferredParameters;
 	FShaderResourceParameter LightAttenuation;
 	FShaderResourceParameter LightAttenuationSampler;
