[UE-6714] Rolling back change to notify nodes of pin connection changes on MovePinLinks (leaves wildcard nodes erroring on compile).

diff --git a/Engine/Source/Editor/BlueprintGraph/Private/K2Node_Message.cpp b/Engine/Source/Editor/BlueprintGraph/Private/K2Node_Message.cpp
--- a/Engine/Source/Editor/BlueprintGraph/Private/K2Node_Message.cpp
+++ b/Engine/Source/Editor/BlueprintGraph/Private/K2Node_Message.cpp
@@ -156,6 +156,7 @@ void UK2Node_Message::ExpandNode(class FKismetCompilerContext& CompilerContext,
 
 		UEdGraphPin* CastToInterfaceSourceObjectPin = CastToInterfaceNode->GetCastSourcePin();
 		CompilerContext.MovePinLinksToIntermediate(*MessageSelfPin, *CastToInterfaceSourceObjectPin);
+		CastToInterfaceNode->PinConnectionListChanged(CastToInterfaceSourceObjectPin);
 
 		// Next, create the function call node
 		UK2Node_CallFunction* FunctionCallNode = CompilerContext.SpawnIntermediateNode<UK2Node_CallFunction>(this, SourceGraph);
diff --git a/Engine/Source/Runtime/Engine/Private/EdGraph/EdGraphSchema.cpp b/Engine/Source/Runtime/Engine/Private/EdGraph/EdGraphSchema.cpp
--- a/Engine/Source/Runtime/Engine/Private/EdGraph/EdGraphSchema.cpp
+++ b/Engine/Source/Runtime/Engine/Private/EdGraph/EdGraphSchema.cpp
@@ -440,9 +440,6 @@ FPinConnectionResponse UEdGraphSchema::MovePinLinks(UEdGraphPin& MoveFromPin, UE
 	TArray<UEdGraphPin*> CurrentLinks = MoveFromPin.LinkedTo;
 	// Then break all links at pin we are moving from
 	MoveFromPin.BreakAllPinLinks();
-#if WITH_EDITOR
-	MoveFromPin.GetOwningNode()->PinConnectionListChanged(&MoveFromPin);
-#endif // #if WITH_EDITOR
 	// Try and make each new connection
 	for (int32 i=0; i<CurrentLinks.Num(); i++)
 	{
@@ -456,17 +453,11 @@ FPinConnectionResponse UEdGraphSchema::MovePinLinks(UEdGraphPin& MoveFromPin, UE
 		{
 			FinalResponse = Response;
 		}
-#if WITH_EDITOR
-		NewLink->GetOwningNode()->PinConnectionListChanged(NewLink);
-#endif // #if WITH_EDITOR
 	}
 	// Move over the default values
 	MoveToPin.DefaultValue = MoveFromPin.DefaultValue;
 	MoveToPin.DefaultObject = MoveFromPin.DefaultObject;
 	MoveToPin.DefaultTextValue = MoveFromPin.DefaultTextValue;
-#if WITH_EDITOR
-	MoveToPin.GetOwningNode()->PinConnectionListChanged(&MoveToPin);
-#endif // #if WITH_EDITOR
 	return FinalResponse;
 }
 
