Fix Doxygen warning caused by mismatched declaration/definition signatures.

diff --git a/Engine/Source/Runtime/Engine/Classes/GameFramework/PlayerController.h b/Engine/Source/Runtime/Engine/Classes/GameFramework/PlayerController.h
--- a/Engine/Source/Runtime/Engine/Classes/GameFramework/PlayerController.h
+++ b/Engine/Source/Runtime/Engine/Classes/GameFramework/PlayerController.h
@@ -18,6 +18,7 @@
 
 
 class FPrimitiveComponentId;
+struct FCollisionQueryParams;
 
 /** Default delegate that provides an implementation for those that don't have special needs other than a toggle */
 DECLARE_DELEGATE_RetVal(bool, FCanUnpause);
@@ -415,7 +416,7 @@ public:
 	UFUNCTION(BlueprintCallable, Category="Input")
 	virtual void ResetIgnoreInputFlags();
 
-	bool GetHitResultAtScreenPosition(const FVector2D ScreenPosition, const ECollisionChannel TraceChannel, const struct FCollisionQueryParams& CollisionQueryParams, FHitResult& HitResult) const;
+	bool GetHitResultAtScreenPosition(const FVector2D ScreenPosition, const ECollisionChannel TraceChannel, const FCollisionQueryParams& CollisionQueryParams, FHitResult& HitResult) const;
 	bool GetHitResultAtScreenPosition(const FVector2D ScreenPosition, const ECollisionChannel TraceChannel, bool bTraceComplex, FHitResult& HitResult) const;
 	bool GetHitResultAtScreenPosition(const FVector2D ScreenPosition, const ETraceTypeQuery TraceChannel, bool bTraceComplex, FHitResult& HitResult) const;
 	bool GetHitResultAtScreenPosition(const FVector2D ScreenPosition, const TArray<TEnumAsByte<EObjectTypeQuery> > & ObjectTypes, bool bTraceComplex, FHitResult& HitResult) const;
