#Anim : fixed crash with extrem hierarchical changes when reimport

Merging

//depot/UE4-UT/Engine/Source/Runtime/Engine/*/Animation/...

to //depot/UE4/Engine/Source/Runtime/Engine/*/Animation/...

diff --git a/Engine/Source/Runtime/Engine/Classes/Animation/AnimSequence.h b/Engine/Source/Runtime/Engine/Classes/Animation/AnimSequence.h
--- a/Engine/Source/Runtime/Engine/Classes/Animation/AnimSequence.h
+++ b/Engine/Source/Runtime/Engine/Classes/Animation/AnimSequence.h
@@ -827,7 +827,7 @@ private:
 	int32 GetSpaceBasedAnimationData(TArray< TArray<FTransform> > & AnimationDataInComponentSpace, FAnimSequenceTrackContainer * RiggingAnimationData) const;
 
 	/** Verify Track Map is valid, if not, fix up */
-	void VerifyTrackMap();
+	void VerifyTrackMap(USkeleton* MySkeleton=NULL);
 	/** Reset Animation Data. Called before Creating new Animation data **/
 	void ResetAnimation();
 	/** Refresh Track Map from Animation Track Names **/
diff --git a/Engine/Source/Runtime/Engine/Private/Animation/AnimSequence.cpp b/Engine/Source/Runtime/Engine/Private/Animation/AnimSequence.cpp
--- a/Engine/Source/Runtime/Engine/Private/Animation/AnimSequence.cpp
+++ b/Engine/Source/Runtime/Engine/Private/Animation/AnimSequence.cpp
@@ -469,26 +469,27 @@ void UAnimSequence::PostLoad()
 }
 
 #if WITH_EDITOR
-void UAnimSequence::VerifyTrackMap()
+void UAnimSequence::VerifyTrackMap(USkeleton* MySkeleton)
 {
-	USkeleton* MySkeleton = GetSkeleton();
-	if(AnimationTrackNames.Num() != TrackToSkeletonMapTable.Num() && MySkeleton!=NULL)
+	USkeleton* UseSkeleton = (MySkeleton)? MySkeleton: GetSkeleton();
+
+	if(AnimationTrackNames.Num() != TrackToSkeletonMapTable.Num() && UseSkeleton!=NULL)
 	{
 		UE_LOG(LogAnimation, Warning, TEXT("RESAVE ANIMATION NEEDED(%s): Fixing track names."), *GetName());
 
-		const TArray<FBoneNode>& BoneTree = MySkeleton->GetBoneTree();
+		const TArray<FBoneNode>& BoneTree = UseSkeleton->GetBoneTree();
 		AnimationTrackNames.Empty();
 		AnimationTrackNames.AddUninitialized(TrackToSkeletonMapTable.Num());
 		for(int32 I=0; I<TrackToSkeletonMapTable.Num(); ++I)
 		{
 			const FTrackToSkeletonMap& TrackMap = TrackToSkeletonMapTable[I];
-			AnimationTrackNames[I] = MySkeleton->GetReferenceSkeleton().GetBoneName(TrackMap.BoneTreeIndex);
+			AnimationTrackNames[I] = UseSkeleton->GetReferenceSkeleton().GetBoneName(TrackMap.BoneTreeIndex);
 		}
 	}
-	else if (MySkeleton != NULL)
+	else if (UseSkeleton != NULL)
 	{
 		int32 NumTracks = AnimationTrackNames.Num();
-		int32 NumSkeletonBone = MySkeleton->GetReferenceSkeleton().GetNum();
+		int32 NumSkeletonBone = UseSkeleton->GetReferenceSkeleton().GetNum();
 
 		bool bNeedsFixing = false;
 		// verify all tracks are still valid
@@ -509,10 +510,10 @@ void UAnimSequence::VerifyTrackMap()
 		{
 			UE_LOG(LogAnimation, Warning, TEXT("RESAVE ANIMATION NEEDED(%s): Fixing track index."), *GetName());
 
-			const TArray<FBoneNode>& BoneTree = MySkeleton->GetBoneTree();
+			const TArray<FBoneNode>& BoneTree = UseSkeleton->GetBoneTree();
 			for(int32 I=NumTracks-1; I>=0; --I)
 			{
-				int32 BoneTreeIndex = MySkeleton->GetReferenceSkeleton().FindBoneIndex(AnimationTrackNames[I]);
+				int32 BoneTreeIndex = UseSkeleton->GetReferenceSkeleton().FindBoneIndex(AnimationTrackNames[I]);
 				if (BoneTreeIndex == INDEX_NONE)
 				{
 					RemoveTrack(I);
@@ -2306,6 +2307,10 @@ void UAnimSequence::RemapTracksToNewSkeleton( USkeleton* NewSkeleton, bool bConv
 		// that has to happen outside of this after Skeleton changes
 		SetSkeleton(NewSkeleton);
 	}
+	else
+	{
+		VerifyTrackMap(NewSkeleton);
+	}
 
 	PostProcessSequence();
 }
diff --git a/Engine/Source/Runtime/Engine/Private/Animation/Skeleton.cpp b/Engine/Source/Runtime/Engine/Private/Animation/Skeleton.cpp
--- a/Engine/Source/Runtime/Engine/Private/Animation/Skeleton.cpp
+++ b/Engine/Source/Runtime/Engine/Private/Animation/Skeleton.cpp
@@ -431,8 +431,27 @@ bool USkeleton::RecreateBoneTree(USkeletalMesh* InSkelMesh)
 		RegenerateGuid();	
 		BoneTree.Empty();
 		ReferenceSkeleton.Empty();
-		return MergeAllBonesToBoneTree(InSkelMesh);
+
+		bool bResult = MergeAllBonesToBoneTree(InSkelMesh);
+
+		if (bResult)
+		{
+			// this has to go through all assets and fix up
+			for (FObjectIterator Iter(UAnimationAsset::StaticClass()); Iter; ++Iter)
+			{
+				UAnimationAsset* AnimAsset = Cast<UAnimationAsset>(*Iter);
+				if (AnimAsset->GetSkeleton() == this)
+				{
+					AnimAsset->ValidateSkeleton();
+				}
+			}
+			
+			return bResult;
+		}
+		
+		return bResult;
 	}
+
 	return false;
 }
 
