Remove unnecessary code

diff --git a/Engine/Source/Runtime/Engine/Private/Animation/Skeleton.cpp b/Engine/Source/Runtime/Engine/Private/Animation/Skeleton.cpp
--- a/Engine/Source/Runtime/Engine/Private/Animation/Skeleton.cpp
+++ b/Engine/Source/Runtime/Engine/Private/Animation/Skeleton.cpp
@@ -445,8 +445,6 @@ bool USkeleton::RecreateBoneTree(USkeletalMesh* InSkelMesh)
 					AnimAsset->ValidateSkeleton();
 				}
 			}
-			
-			return bResult;
 		}
 		
 		return bResult;
