Fixed case where sub-levels was not fully loaded before player spawns. This change removes "view" requirement from UpdateLevelStreamign/FlushLevelStreaming. UE-5330

diff --git a/Engine/Source/Editor/UnrealEd/Private/Editor.cpp b/Engine/Source/Editor/UnrealEd/Private/Editor.cpp
--- a/Engine/Source/Editor/UnrealEd/Private/Editor.cpp
+++ b/Engine/Source/Editor/UnrealEd/Private/Editor.cpp
@@ -1657,16 +1657,8 @@ bool UEditorEngine::UpdateSingleViewportClient(FEditorViewportClient* InViewport
 			// this ourselves!
 			if( PlayWorld != NULL && bIsSimulatingInEditor && InViewportClient->IsSimulateInEditorViewport() )
 			{
-				// Create a scene view to use when calculating listener position.
-				FSceneViewFamilyContext ViewFamily( FSceneViewFamily::ConstructionValues( 
-					InViewportClient->Viewport, 
-					InViewportClient->GetScene(), 
-					InViewportClient->EngineShowFlags )
-					.SetRealtimeUpdate( InViewportClient->IsRealtime() ));
-				const FSceneView& View = *InViewportClient->CalcSceneView( &ViewFamily );
-	
 				// Update level streaming.
-				GWorld->UpdateLevelStreaming( &ViewFamily );
+				GWorld->UpdateLevelStreaming();
 
 				// Also make sure hit proxies are refreshed for SIE viewports, as the user may be trying to grab an object or widget manipulator that's moving!
 				if( InViewportClient->IsRealtime() )
diff --git a/Engine/Source/Runtime/Engine/Classes/Engine/LevelStreaming.h b/Engine/Source/Runtime/Engine/Classes/Engine/LevelStreaming.h
--- a/Engine/Source/Runtime/Engine/Classes/Engine/LevelStreaming.h
+++ b/Engine/Source/Runtime/Engine/Classes/Engine/LevelStreaming.h
@@ -188,19 +188,17 @@ class ULevelStreaming : public UObject
 	 * so the code needs to return true an appropriate amount of time before it is 
 	 * needed.
 	 *
-	 * @param ViewLocation	Location of the viewer
 	 * @return true if level should be loaded/ streamed in, false otherwise
 	 */
-	virtual bool ShouldBeLoaded( const FVector& ViewLocation );
+	virtual bool ShouldBeLoaded();
 
 	/**
 	 * Return whether this level should be visible/ associated with the world if it is
 	 * loaded.
 	 * 
-	 * @param ViewLocation	Location of the viewer
 	 * @return true if the level should be visible, false otherwise
 	 */
-	virtual bool ShouldBeVisible( const FVector& ViewLocation );
+	virtual bool ShouldBeVisible();
 
 	virtual bool ShouldBeAlwaysLoaded() const { return false; }
 	
diff --git a/Engine/Source/Runtime/Engine/Classes/Engine/LevelStreamingAlwaysLoaded.h b/Engine/Source/Runtime/Engine/Classes/Engine/LevelStreamingAlwaysLoaded.h
--- a/Engine/Source/Runtime/Engine/Classes/Engine/LevelStreamingAlwaysLoaded.h
+++ b/Engine/Source/Runtime/Engine/Classes/Engine/LevelStreamingAlwaysLoaded.h
@@ -16,7 +16,7 @@ class ULevelStreamingAlwaysLoaded : public ULevelStreaming
 	GENERATED_UCLASS_BODY()
 
 	// Begin ULevelStreaming Interface
-	virtual bool ShouldBeLoaded( const FVector& ViewLocation ) override;
+	virtual bool ShouldBeLoaded() override;
 	virtual bool ShouldBeAlwaysLoaded() const override { return true; } 
 	// End ULevelStreaming Interface
 };
diff --git a/Engine/Source/Runtime/Engine/Classes/Engine/LevelStreamingBounds.h b/Engine/Source/Runtime/Engine/Classes/Engine/LevelStreamingBounds.h
deleted file mode 100644
diff --git a/Engine/Source/Runtime/Engine/Classes/Engine/LevelStreamingKismet.h b/Engine/Source/Runtime/Engine/Classes/Engine/LevelStreamingKismet.h
--- a/Engine/Source/Runtime/Engine/Classes/Engine/LevelStreamingKismet.h
+++ b/Engine/Source/Runtime/Engine/Classes/Engine/LevelStreamingKismet.h
@@ -29,8 +29,8 @@ class ULevelStreamingKismet : public ULevelStreaming
 	// End UObject Interface
 
 	// Begin ULevelStreaming Interface
-	virtual bool ShouldBeLoaded( const FVector& ViewLocation ) override;
-	virtual bool ShouldBeVisible( const FVector& ViewLocation ) override;
+	virtual bool ShouldBeLoaded() override;
+	virtual bool ShouldBeVisible() override;
 	// End ULevelStreaming Interface
 };
 
diff --git a/Engine/Source/Runtime/Engine/Classes/Engine/LevelStreamingPersistent.h b/Engine/Source/Runtime/Engine/Classes/Engine/LevelStreamingPersistent.h
--- a/Engine/Source/Runtime/Engine/Classes/Engine/LevelStreamingPersistent.h
+++ b/Engine/Source/Runtime/Engine/Classes/Engine/LevelStreamingPersistent.h
@@ -17,11 +17,11 @@ class ULevelStreamingPersistent : public ULevelStreaming
 
 
 	// Begin ULevelStreaming Interface
-	virtual bool ShouldBeLoaded( const FVector& ViewLocation ) override
+	virtual bool ShouldBeLoaded() override
 	{
 		return true;
 	}
-	virtual bool ShouldBeVisible( const FVector& ViewLocation ) override
+	virtual bool ShouldBeVisible() override
 	{
 		return true;
 	}
diff --git a/Engine/Source/Runtime/Engine/Classes/Engine/World.h b/Engine/Source/Runtime/Engine/Classes/Engine/World.h
--- a/Engine/Source/Runtime/Engine/Classes/Engine/World.h
+++ b/Engine/Source/Runtime/Engine/Classes/Engine/World.h
@@ -1852,19 +1852,12 @@ public:
 	void RemoveFromWorld( ULevel* Level );
 
 	/**
-	 * Updates the all the visible worlds based on the current view location of the player and sets level LODs accordingly.	 
-	 *
-	 * @param ViewFamily	Optional collection of views to take into account
+	 * Updates sub-levels (load/unload/show/hide) using streaming levels current state
 	 */
-	void UpdateLevelStreaming( FSceneViewFamily* ViewFamily = NULL );
+	void UpdateLevelStreaming();
 
 private:
-	/**
-	 * Updates the world based on the current view location of the player and sets level LODs accordingly.	 
-	 * @param PersistentWorld	Persistent world
-	 * @param ViewFamily		Optional collection of views to take into account
-	 */
-	void UpdateLevelStreamingInner( UWorld* PersistentWorld, FSceneViewFamily* ViewFamily );
+	void UpdateLevelStreamingInner( ULevelStreaming* StreamingLevel );
 
 public:
 	/**
@@ -1872,11 +1865,10 @@ public:
 	 * so further calls to UpdateLevelStreaming won't do any work unless state changes. Basically blocks
 	 * on all async operation like updating components.
 	 *
-	 * @param ViewFamily				Optional collection of views to take into account
 	 * @param FlushType					Whether to only flush level visibility operations (optional)
 	 * @param ExcludeType				Exclude packages of this type from flushing
 	 */
-	void FlushLevelStreaming( FSceneViewFamily* ViewFamily = NULL, EFlushLevelStreamingType FlushType = EFlushLevelStreamingType::Full, FName ExcludeType = NAME_None );
+	void FlushLevelStreaming(EFlushLevelStreamingType FlushType = EFlushLevelStreamingType::Full, FName ExcludeType = NAME_None);
 
 	/**
 	 * Triggers a call to ULevel::BuildStreamingData(this,NULL,NULL) within a few seconds.
diff --git a/Engine/Source/Runtime/Engine/Private/GameEngine.cpp b/Engine/Source/Runtime/Engine/Private/GameEngine.cpp
--- a/Engine/Source/Runtime/Engine/Private/GameEngine.cpp
+++ b/Engine/Source/Runtime/Engine/Private/GameEngine.cpp
@@ -505,7 +505,7 @@ void UGameEngine::PreExit()
 
 			World->GetGameInstance()->Shutdown();
 
-			World->FlushLevelStreaming( NULL, EFlushLevelStreamingType::Visibility );
+			World->FlushLevelStreaming(EFlushLevelStreamingType::Visibility);
 			World->CleanupWorld();
 		}
 	}
diff --git a/Engine/Source/Runtime/Engine/Private/GameInstance.cpp b/Engine/Source/Runtime/Engine/Private/GameInstance.cpp
--- a/Engine/Source/Runtime/Engine/Private/GameInstance.cpp
+++ b/Engine/Source/Runtime/Engine/Private/GameInstance.cpp
@@ -184,7 +184,7 @@ bool UGameInstance::StartPIEGameInstance(ULocalPlayer* LocalPlayer, bool bInSimu
 		}
 		
 		// Make sure "always loaded" sub-levels are fully loaded
-		PlayWorld->FlushLevelStreaming(nullptr, EFlushLevelStreamingType::Visibility);
+		PlayWorld->FlushLevelStreaming(EFlushLevelStreamingType::Visibility);
 
 		UNavigationSystem::InitializeForWorld(PlayWorld, LocalPlayers.Num() > 0 ? FNavigationSystem::PIEMode : FNavigationSystem::SimulationMode);
 		PlayWorld->CreateAISystem();
diff --git a/Engine/Source/Runtime/Engine/Private/GameViewportClient.cpp b/Engine/Source/Runtime/Engine/Private/GameViewportClient.cpp
--- a/Engine/Source/Runtime/Engine/Private/GameViewportClient.cpp
+++ b/Engine/Source/Runtime/Engine/Private/GameViewportClient.cpp
@@ -912,7 +912,7 @@ void UGameViewportClient::Draw(FViewport* InViewport, FCanvas* SceneCanvas)
 	}
 
 	// Update level streaming.
-	GetWorld()->UpdateLevelStreaming( &ViewFamily );
+	GetWorld()->UpdateLevelStreaming();
 
 	// Draw the player views.
 	if (!bDisableWorldRendering && !bUIDisableWorldRendering && PlayerViewMap.Num() > 0)
diff --git a/Engine/Source/Runtime/Engine/Private/LevelStreaming.cpp b/Engine/Source/Runtime/Engine/Private/LevelStreaming.cpp
--- a/Engine/Source/Runtime/Engine/Private/LevelStreaming.cpp
+++ b/Engine/Source/Runtime/Engine/Private/LevelStreaming.cpp
@@ -2,7 +2,6 @@
 
 #include "EnginePrivate.h"
 #include "Engine/LevelStreamingAlwaysLoaded.h"
-#include "Engine/LevelStreamingBounds.h"
 #include "Engine/LevelStreamingPersistent.h"
 #include "Engine/LevelStreamingVolume.h"
 #include "Engine/LevelBounds.h"
@@ -770,17 +769,17 @@ void ULevelStreaming::RenameForPIE(int32 PIEInstanceID)
 	}
 }
 
-bool ULevelStreaming::ShouldBeLoaded( const FVector& ViewLocation )
+bool ULevelStreaming::ShouldBeLoaded()
 {
 	return true;
 }
 
-bool ULevelStreaming::ShouldBeVisible( const FVector& ViewLocation )
+bool ULevelStreaming::ShouldBeVisible()
 {
 	if( GetWorld()->IsGameWorld() )
 	{
 		// Game and play in editor viewport codepath.
-		return bShouldBeVisible && ShouldBeLoaded( ViewLocation );
+		return bShouldBeVisible && ShouldBeLoaded();
 	}
 	else
 	{
@@ -925,12 +924,12 @@ void ULevelStreamingKismet::PostLoad()
 	}
 }
 
-bool ULevelStreamingKismet::ShouldBeVisible( const FVector& ViewLocation )
+bool ULevelStreamingKismet::ShouldBeVisible()
 {
 	return bShouldBeVisible || (bShouldBeVisibleInEditor && !FApp::IsGame());
 }
 
-bool ULevelStreamingKismet::ShouldBeLoaded( const FVector& ViewLocation )
+bool ULevelStreamingKismet::ShouldBeLoaded()
 {
 	return bShouldBeLoaded;
 }
@@ -944,36 +943,9 @@ ULevelStreamingAlwaysLoaded::ULevelStreamingAlwaysLoaded(const FObjectInitialize
 	bShouldBeVisible = true;
 }
 
-bool ULevelStreamingAlwaysLoaded::ShouldBeLoaded( const FVector& ViewLocation )
+bool ULevelStreamingAlwaysLoaded::ShouldBeLoaded()
 {
 	return true;
 }
 
-/*-----------------------------------------------------------------------------
-	ULevelStreamingBounds implementation.
------------------------------------------------------------------------------*/
-ULevelStreamingBounds::ULevelStreamingBounds(const FObjectInitializer& ObjectInitializer)
-	: Super(ObjectInitializer)
-{
-}
-
-bool ULevelStreamingBounds::ShouldBeVisible( const FVector& ViewLocation )
-{
-	bShouldBeVisible = ShouldBeLoaded(ViewLocation);
-	return bShouldBeVisible;
-}
-
-bool ULevelStreamingBounds::ShouldBeLoaded( const FVector& ViewLocation )
-{
-	bShouldBeLoaded = false;
-	
-	UWorld* MyWorld = Cast<UWorld>(GetOuter());
-	if (MyWorld && MyWorld->PersistentLevel->LevelBoundsActor.IsValid())
-	{
-		return MyWorld->PersistentLevel->LevelBoundsActor.Get()->GetComponentsBoundingBox().IsInside(ViewLocation);
-	}
-	
-	return bShouldBeLoaded;
-}
-
 #undef LOCTEXT_NAMESPACE
diff --git a/Engine/Source/Runtime/Engine/Private/LevelTick.cpp b/Engine/Source/Runtime/Engine/Private/LevelTick.cpp
--- a/Engine/Source/Runtime/Engine/Private/LevelTick.cpp
+++ b/Engine/Source/Runtime/Engine/Private/LevelTick.cpp
@@ -551,7 +551,7 @@ void UWorld::ProcessLevelStreamingVolumes(FVector* OverrideViewLocation)
 			else
 			{
 				FRotator ViewRotation(0,0,0);
-				PlayerActor->GetPlayerViewPoint( ViewLocation, ViewRotation );
+				PlayerActor->GetActorEyesViewPoint( ViewLocation, ViewRotation );
 			}
 
 			TMap<AVolume*,bool> VolumeMap;
diff --git a/Engine/Source/Runtime/Engine/Private/PlayerController.cpp b/Engine/Source/Runtime/Engine/Private/PlayerController.cpp
--- a/Engine/Source/Runtime/Engine/Private/PlayerController.cpp
+++ b/Engine/Source/Runtime/Engine/Private/PlayerController.cpp
@@ -219,7 +219,7 @@ void APlayerController::ClientFlushLevelStreaming_Implementation()
 	if (GEngine->ShouldCommitPendingMapChange(GetWorld()))
 	{
 		// request level streaming be flushed next frame
-		GetWorld()->UpdateLevelStreaming(NULL);
+		GetWorld()->UpdateLevelStreaming();
 		GetWorld()->bRequestedBlockOnAsyncLoading = true;
 		// request GC as soon as possible to remove any unloaded levels from memory
 		GetWorld()->ForceGarbageCollection();
diff --git a/Engine/Source/Runtime/Engine/Private/UnrealClient.cpp b/Engine/Source/Runtime/Engine/Private/UnrealClient.cpp
--- a/Engine/Source/Runtime/Engine/Private/UnrealClient.cpp
+++ b/Engine/Source/Runtime/Engine/Private/UnrealClient.cpp
@@ -938,34 +938,8 @@ void FViewport::Draw( bool bShouldPresent /*= true */)
 		// if this is a game viewport, and game rendering is disabled, then we don't want to actually draw anything
 		if ( World && World->IsGameWorld() && !bIsGameRenderingEnabled)
 		{
-			// since we aren't drawing the viewport, we still need to update streaming, which needs valid view info
-			FSceneViewFamilyContext ViewFamily( FSceneViewFamily::ConstructionValues( this, World->Scene, FEngineShowFlags(ESFIM_Game)) );
-			for( FConstPlayerControllerIterator Iterator = World->GetPlayerControllerIterator(); Iterator; ++Iterator )
-			{
-				APlayerController* PlayerController = *Iterator;
-				if( PlayerController )
-				{
-					ULocalPlayer* Player = Cast<ULocalPlayer>(PlayerController->Player);
-					if(Player)
-					{
-						// Calculate the player's view information.
-						FVector		ViewLocation;
-						FRotator	ViewRotation;
-						FSceneView* View = Player->CalcSceneView( &ViewFamily, ViewLocation, ViewRotation, this);
-
-						// if we have a valid view, use it for resource streaming
-						if(View)
-						{
-							// Add view information for resource streaming.
-							IStreamingManager::Get().AddViewInformation( View->ViewMatrices.ViewOrigin, View->ViewRect.Width(), View->ViewRect.Width() * View->ViewMatrices.ProjMatrix.M[0][0] );
-							World->ViewLocationsRenderedLastFrame.Add(View->ViewMatrices.ViewOrigin);
-						}
-					}
-				}
-			}
-	
-			// Update level streaming.
-			World->UpdateLevelStreaming( &ViewFamily );
+			// since we aren't drawing the viewport, we still need to update streaming
+			World->UpdateLevelStreaming();
 		}
 		else
 		{
diff --git a/Engine/Source/Runtime/Engine/Private/UnrealEngine.cpp b/Engine/Source/Runtime/Engine/Private/UnrealEngine.cpp
--- a/Engine/Source/Runtime/Engine/Private/UnrealEngine.cpp
+++ b/Engine/Source/Runtime/Engine/Private/UnrealEngine.cpp
@@ -8639,7 +8639,7 @@ bool UEngine::LoadMap( FWorldContext& WorldContext, FURL URL, class UPendingNetG
 		ShutdownWorldNetDriver(WorldContext.World());
 
 		// Make sure there are no pending visibility requests.
-		WorldContext.World()->FlushLevelStreaming( NULL, EFlushLevelStreamingType::Visibility );
+		WorldContext.World()->FlushLevelStreaming(EFlushLevelStreamingType::Visibility);
 		
 		// send a message that all levels are going away (NULL means every sublevel is being removed
 		// without a call to RemoveFromWorld for each)
@@ -8990,7 +8990,7 @@ bool UEngine::LoadMap( FWorldContext& WorldContext, FURL URL, class UPendingNetG
 	LoadPackagesFully(WorldContext.World(), FULLYLOAD_Map, WorldContext.World()->PersistentLevel->GetOutermost()->GetName());
 
 	// Make sure "always loaded" sub-levels are fully loaded
-	WorldContext.World()->FlushLevelStreaming(nullptr, EFlushLevelStreamingType::Visibility);
+	WorldContext.World()->FlushLevelStreaming(EFlushLevelStreamingType::Visibility);
 	
 	UNavigationSystem::InitializeForWorld(WorldContext.World(), FNavigationSystem::GameMode);
 	
@@ -9078,28 +9078,8 @@ void UEngine::BlockTillLevelStreamingCompleted(UWorld* InWorld)
 		InWorld->WorldComposition->UpdateStreamingState();
 	}
 
-	// Create view collection
-	int32 NumPlayers = GetNumGamePlayers(InWorld);
-	TUniquePtr<FSceneViewFamilyContext> ViewFamily;
-	if (NumPlayers && GameViewport)
-	{
-		ViewFamily = MakeUnique<FSceneViewFamilyContext>(
-			FSceneViewFamily::ConstructionValues(GameViewport->Viewport, InWorld->Scene, GameViewport->EngineShowFlags).SetRealtimeUpdate(true));
-
-		for (int32 PlayerIndex = 0; PlayerIndex < NumPlayers; ++PlayerIndex)
-		{
-			ULocalPlayer* Player = GetGamePlayer(InWorld, PlayerIndex);
-			if (Player && Player->ViewportClient)
-			{
-				FVector ViewLocation;
-				FRotator ViewRotation;
-				Player->CalcSceneView(ViewFamily.Get(), /*out*/ ViewLocation, /*out*/ ViewRotation, Player->ViewportClient->Viewport);
-			}
-		}
-	}
-
 	// Probe if we have anything to do
-	InWorld->UpdateLevelStreaming(ViewFamily.Get());
+	InWorld->UpdateLevelStreaming();
 	bool bWorkToDo = (InWorld->IsVisibilityRequestPending() || IsAsyncLoading());
 	
 	if (bWorkToDo)
@@ -9110,7 +9090,7 @@ void UEngine::BlockTillLevelStreamingCompleted(UWorld* InWorld)
 		}	
 
 		// Flush level streaming requests, blocking till completion.
-		InWorld->FlushLevelStreaming(ViewFamily.Get(), EFlushLevelStreamingType::Full);
+		InWorld->FlushLevelStreaming(EFlushLevelStreamingType::Full);
 
 		if( GEngine->EndStreamingPauseDelegate && GEngine->EndStreamingPauseDelegate->IsBound() )
 		{
@@ -9993,7 +9973,7 @@ bool UEngine::CommitMapChange( FWorldContext &Context )
 		// Update level streaming, forcing existing levels to be unloaded and their streaming objects 
 		// removed from the world info.	We can't kick off async loading in this update as we want to 
 		// collect garbage right below.
-		Context.World()->FlushLevelStreaming( NULL, EFlushLevelStreamingType::Visibility );
+		Context.World()->FlushLevelStreaming(EFlushLevelStreamingType::Visibility);
 		
 		// make sure any looping sounds, etc are stopped
 		if (GetAudioDevice() != NULL)
@@ -10040,12 +10020,12 @@ bool UEngine::CommitMapChange( FWorldContext &Context )
 
 			Context.PendingLevelStreamingStatusUpdates.Empty();
 
-			Context.World()->FlushLevelStreaming( NULL, EFlushLevelStreamingType::Full );
+			Context.World()->FlushLevelStreaming(EFlushLevelStreamingType::Full);
 		}
 		else
 		{
 			// Make sure there are no pending visibility requests.
-			Context.World()->FlushLevelStreaming( NULL, EFlushLevelStreamingType::Visibility );
+			Context.World()->FlushLevelStreaming(EFlushLevelStreamingType::Visibility);
 		}
 
 		// delay the use of streaming volumes for a few frames
diff --git a/Engine/Source/Runtime/Engine/Private/World.cpp b/Engine/Source/Runtime/Engine/Private/World.cpp
--- a/Engine/Source/Runtime/Engine/Private/World.cpp
+++ b/Engine/Source/Runtime/Engine/Private/World.cpp
@@ -2179,151 +2179,113 @@ UWorld* UWorld::DuplicateWorldForPIE(const FString& PackageName, UWorld* OwningW
 	return PIELevelWorld;
 }
 
-void UWorld::UpdateLevelStreamingInner( UWorld* PersistentWorld, FSceneViewFamily* ViewFamily )
+void UWorld::UpdateLevelStreamingInner(ULevelStreaming* StreamingLevel)
 {
-	// Iterate over level collection to find out whether we need to load/ unload any levels.
-	for (int32 LevelIndex = 0; LevelIndex < StreamingLevels.Num(); LevelIndex++)
-	{
-		ULevelStreaming* StreamingLevel	= StreamingLevels[LevelIndex];
-		if( !StreamingLevel )
-		{
-			// This can happen when manually adding a new level to the array via the property editing code.
-			continue;
-		}
+	check(StreamingLevel != nullptr);
 		
-		// Don't bother loading sub-levels in PIE for levels that aren't visible in editor
-		if( PersistentWorld->IsPlayInEditor() && GEngine->OnlyLoadEditorVisibleLevelsInPIE() )
+	// Don't bother loading sub-levels in PIE for levels that aren't visible in editor
+	if (IsPlayInEditor() && GEngine->OnlyLoadEditorVisibleLevelsInPIE())
+	{
+		if (!StreamingLevel->bShouldBeVisibleInEditor)
 		{
-			// create a new streaming level helper object
-			if( !StreamingLevel->bShouldBeVisibleInEditor )
-			{
-				continue;
-			}
+			return;
 		}
+	}
 
-		// Work performed to make a level visible is spread across several frames and we can't unload/ hide a level that is currently pending
-		// to be made visible, so we fulfill those requests first.
-		bool bHasVisibilityRequestPending	= StreamingLevel->GetLoadedLevel() && StreamingLevel->GetLoadedLevel() == PersistentWorld->CurrentLevelPendingVisibility;
+	// Work performed to make a level visible is spread across several frames and we can't unload/ hide a level that is currently pending
+	// to be made visible, so we fulfill those requests first.
+	bool bHasVisibilityRequestPending	= StreamingLevel->GetLoadedLevel() && StreamingLevel->GetLoadedLevel() == CurrentLevelPendingVisibility;
 		
-		// Figure out whether level should be loaded, visible and block on load if it should be loaded but currently isn't.
-		bool bShouldBeLoaded				= bHasVisibilityRequestPending || (!GEngine->bUseBackgroundLevelStreaming && !PersistentWorld->bShouldForceUnloadStreamingLevels && !StreamingLevel->bIsRequestingUnloadAndRemoval);
-		bool bShouldBeVisible				= bHasVisibilityRequestPending || PersistentWorld->bShouldForceVisibleStreamingLevels;
-		bool bShouldBlockOnLoad				= StreamingLevel->bShouldBlockOnLoad;
+	// Figure out whether level should be loaded, visible and block on load if it should be loaded but currently isn't.
+	bool bShouldBeLoaded	= bHasVisibilityRequestPending || (!GEngine->bUseBackgroundLevelStreaming && !bShouldForceUnloadStreamingLevels && !StreamingLevel->bIsRequestingUnloadAndRemoval);
+	bool bShouldBeVisible	= bHasVisibilityRequestPending || bShouldForceVisibleStreamingLevels;
+	bool bShouldBlockOnLoad	= StreamingLevel->bShouldBlockOnLoad;
 
-		// Don't update if the code requested this level object to be unloaded and removed.
-		if( !PersistentWorld->bShouldForceUnloadStreamingLevels && !StreamingLevel->bIsRequestingUnloadAndRemoval )
-		{
-			// Take all views into account if any were passed in.
-			if( ViewFamily && ViewFamily->Views.Num() > 0 )
-			{
-				// Iterate over all views in collection.
-				for( int32 ViewIndex=0; ViewIndex<ViewFamily->Views.Num(); ViewIndex++ )
-				{
-					const FSceneView* View		= ViewFamily->Views[ViewIndex];
-					const FVector& ViewLocation	= View->ViewMatrices.ViewOrigin;
-					bShouldBeLoaded	= bShouldBeLoaded  || !IsGameWorld() || StreamingLevel->ShouldBeLoaded(ViewLocation);
-					bShouldBeVisible= bShouldBeVisible || ( bShouldBeLoaded && StreamingLevel->ShouldBeVisible(ViewLocation) );
-				}
-			}
-			// Or default to view location of 0,0,0
-			else
-			{
-				FVector ViewLocation(0,0,0);
-				bShouldBeLoaded		= bShouldBeLoaded  || !IsGameWorld() || StreamingLevel->ShouldBeLoaded(ViewLocation);
-				bShouldBeVisible	= bShouldBeVisible || ( bShouldBeLoaded && StreamingLevel->ShouldBeVisible(ViewLocation) );
-			}
-		}
+	// Don't update if the code requested this level object to be unloaded and removed.
+	if(!bShouldForceUnloadStreamingLevels && !StreamingLevel->bIsRequestingUnloadAndRemoval)
+	{
+		bShouldBeLoaded		= bShouldBeLoaded  || !IsGameWorld() || StreamingLevel->ShouldBeLoaded();
+		bShouldBeVisible	= bShouldBeVisible || (bShouldBeLoaded && StreamingLevel->ShouldBeVisible());
+	}
 
-		// We want to give the garbage collector a chance to remove levels before we stream in more. We can't do this in the
-		// case of a blocking load as it means those requests should be fulfilled right away. By waiting on GC before kicking
-		// off new levels we potentially delay streaming in maps, but AllowLevelLoadRequests already looks and checks whether
-		// async loading in general is active. E.g. normal package streaming would delay loading in this case. This is done
-		// on purpose as well so the GC code has a chance to execute between consecutive loads of maps.
-		//
-		// NOTE: AllowLevelLoadRequests not an invariant as streaming might affect the result, do NOT pulled out of the loop.
-		bool bAllowLevelLoadRequests =	PersistentWorld->AllowLevelLoadRequests() || bShouldBlockOnLoad;
+	// We want to give the garbage collector a chance to remove levels before we stream in more. We can't do this in the
+	// case of a blocking load as it means those requests should be fulfilled right away. By waiting on GC before kicking
+	// off new levels we potentially delay streaming in maps, but AllowLevelLoadRequests already looks and checks whether
+	// async loading in general is active. E.g. normal package streaming would delay loading in this case. This is done
+	// on purpose as well so the GC code has a chance to execute between consecutive loads of maps.
+	//
+	// NOTE: AllowLevelLoadRequests not an invariant as streaming might affect the result, do NOT pulled out of the loop.
+	bool bAllowLevelLoadRequests =	AllowLevelLoadRequests() || bShouldBlockOnLoad;
 
-		// Figure out whether there are any levels we haven't collected garbage yet.
-		bool bAreLevelsPendingPurge	=	FLevelStreamingGCHelper::GetNumLevelsPendingPurge() > 0;
-		// Request a 'soft' GC if there are levels pending purge and there are levels to be loaded. In the case of a blocking
-		// load this is going to guarantee GC firing first thing afterwards and otherwise it is going to sneak in right before
-		// kicking off the async load.
-		if (bAreLevelsPendingPurge)
-		{
-			PersistentWorld->ForceGarbageCollection( false );
-		}
+	// Figure out whether there are any levels we haven't collected garbage yet.
+	bool bAreLevelsPendingPurge	=	FLevelStreamingGCHelper::GetNumLevelsPendingPurge() > 0;
+	// Request a 'soft' GC if there are levels pending purge and there are levels to be loaded. In the case of a blocking
+	// load this is going to guarantee GC firing first thing afterwards and otherwise it is going to sneak in right before
+	// kicking off the async load.
+	if (bAreLevelsPendingPurge)
+	{
+		ForceGarbageCollection( false );
+	}
 
-		// See whether level is already loaded
-		if (bShouldBeLoaded)
-		{
-			const bool bBlockOnLoad = (!IsGameWorld() || !GEngine->bUseBackgroundLevelStreaming || bShouldBlockOnLoad);
-			// Request to load or duplicate existing level
-			StreamingLevel->RequestLevel(PersistentWorld, bAllowLevelLoadRequests, bBlockOnLoad);
-		}
+	// See whether level is already loaded
+	if (bShouldBeLoaded)
+	{
+		const bool bBlockOnLoad = (!IsGameWorld() || !GEngine->bUseBackgroundLevelStreaming || bShouldBlockOnLoad);
+		// Request to load or duplicate existing level
+		StreamingLevel->RequestLevel(this, bAllowLevelLoadRequests, bBlockOnLoad);
+	}
 		
-		// Cache pointer for convenience. This cannot happen before this point as e.g. flushing async loaders
-		// or such will modify StreamingLevel->LoadedLevel.
-		ULevel* Level = StreamingLevel->GetLoadedLevel();
+	// Cache pointer for convenience. This cannot happen before this point as e.g. flushing async loaders
+	// or such will modify StreamingLevel->LoadedLevel.
+	ULevel* Level = StreamingLevel->GetLoadedLevel();
 
-		// See whether we have a loaded level.
-		if (Level)
+	// See whether we have a loaded level.
+	if (Level)
+	{
+		// Update loaded level visibility
+		if (bShouldBeVisible)
 		{
-			// Update loaded level visibility
-			if (bShouldBeVisible)
-			{
-				// Add loaded level to a world if it's not there yet
-				if (!Level->bIsVisible)
-				{
-					PersistentWorld->AddToWorld(Level, StreamingLevel->LevelTransform);
-				}
-
-				// In case we have finished making level visible
-				// immediately discard previous level
-				if (Level->bIsVisible)
-				{
-					StreamingLevel->DiscardPendingUnloadLevel(PersistentWorld);
-				}
-			}
-			else
+			// Add loaded level to a world if it's not there yet
+			if (!Level->bIsVisible)
 			{
-				// Discard previous LOD level
-				StreamingLevel->DiscardPendingUnloadLevel(PersistentWorld);
-				// Hide loaded level
-				if (Level->bIsVisible)
-				{
-					PersistentWorld->RemoveFromWorld(Level);
-				}
+				AddToWorld(Level, StreamingLevel->LevelTransform);
 			}
 
-			if (!bShouldBeLoaded)
+			// In case we have finished making level visible
+			// immediately discard previous level
+			if (Level->bIsVisible)
 			{
-				if (!Level->bIsVisible && !PersistentWorld->IsVisibilityRequestPending())
-				{
-					StreamingLevel->DiscardPendingUnloadLevel(PersistentWorld);
-					StreamingLevel->ClearLoadedLevel();
-					StreamingLevel->DiscardPendingUnloadLevel(PersistentWorld);
-				}
+				StreamingLevel->DiscardPendingUnloadLevel(this);
 			}
 		}
 		else
 		{
-			StreamingLevel->DiscardPendingUnloadLevel(PersistentWorld);
+			// Discard previous LOD level
+			StreamingLevel->DiscardPendingUnloadLevel(this);
+			// Hide loaded level
+			if (Level->bIsVisible)
+			{
+				RemoveFromWorld(Level);
+			}
 		}
-								
-		// If requested, remove this level from iterated over array once it is unloaded.
-		if (StreamingLevel->bIsRequestingUnloadAndRemoval)
+
+		if (!bShouldBeLoaded)
 		{
-			if (StreamingLevel->HasLoadedLevel() == false && 
-				StreamingLevel->bHasLoadRequestPending == false)
+			if (!Level->bIsVisible && !IsVisibilityRequestPending())
 			{
-				// The -- is required as we're forward iterating over the StreamingLevels array.
-				StreamingLevels.RemoveAt( LevelIndex-- );
+				StreamingLevel->DiscardPendingUnloadLevel(this);
+				StreamingLevel->ClearLoadedLevel();
+				StreamingLevel->DiscardPendingUnloadLevel(this);
 			}
-		}	
+		}
+	}
+	else
+	{
+		StreamingLevel->DiscardPendingUnloadLevel(this);
 	}
 }
 
-void UWorld::UpdateLevelStreaming( FSceneViewFamily* ViewFamily )
+void UWorld::UpdateLevelStreaming()
 {
 	SCOPE_CYCLE_COUNTER(STAT_UpdateLevelStreamingTime);
 	// do nothing if level streaming is frozen
@@ -2332,22 +2294,29 @@ void UWorld::UpdateLevelStreaming( FSceneViewFamily* ViewFamily )
 		return;
 	}
 
-	// Store current count of pending unload levels
+	// Store current number of pending unload levels, it may change in loop bellow
 	const int32 NumLevelsPendingPurge = FLevelStreamingGCHelper::GetNumLevelsPendingPurge();
-	// Store current count of async loading packages
-	const int32	NumAsyncLoadingPackages = GetNumAsyncPackages(); 
-	
-	// Update level streaming objects state
-	UpdateLevelStreamingInner(this, ViewFamily);
 	
-	// Force initial loading to be "bShouldBlockOnLoad".
-	const bool bLevelsHaveLoadRequestPending = NumAsyncLoadingPackages < GetNumAsyncPackages();
-	if (bLevelsHaveLoadRequestPending && (!HasBegunPlay() || bWorldWasLoadedThisTick ))
+	for (int32 LevelIndex = 0; LevelIndex < StreamingLevels.Num(); LevelIndex++)
 	{
-		// Block till all async requests are finished.
-		FlushAsyncLoading( NAME_None );
+		ULevelStreaming* StreamingLevel = StreamingLevels[LevelIndex];
+		if (StreamingLevel)
+		{
+			UpdateLevelStreamingInner(StreamingLevel);
+
+			// If requested, remove this level from iterated over array once it is unloaded.
+			if (StreamingLevel->bIsRequestingUnloadAndRemoval)
+			{
+				if (StreamingLevel->HasLoadedLevel() == false && 
+					StreamingLevel->bHasLoadRequestPending == false)
+				{
+					// The -- is required as we're forward iterating over the StreamingLevels array.
+					StreamingLevels.RemoveAt(LevelIndex--);
+				}
+			}	
+		}
 	}
-	
+			
 	// In case more levels has been requested to unload, force GC on next tick 
 	if (NumLevelsPendingPurge < FLevelStreamingGCHelper::GetNumLevelsPendingPurge())
 	{
@@ -2355,20 +2324,20 @@ void UWorld::UpdateLevelStreaming( FSceneViewFamily* ViewFamily )
 	}
 }
 
-void UWorld::FlushLevelStreaming(FSceneViewFamily* ViewFamily, EFlushLevelStreamingType FlushType, FName ExcludeType)
+void UWorld::FlushLevelStreaming(EFlushLevelStreamingType FlushType, FName ExcludeType)
 {
 	AWorldSettings* WorldSettings = GetWorldSettings();
 
 	TGuardValue<EFlushLevelStreamingType> FlushingLevelStreamingGuard(FlushLevelStreamingType, FlushType);
 
 	// Update internals with current loaded/ visibility flags.
-	UpdateLevelStreaming(ViewFamily);
+	UpdateLevelStreaming();
 
 	// Make sure all outstanding loads are taken care of, other than ones associated with the excluded type
 	FlushAsyncLoading(ExcludeType);
 
 	// Kick off making levels visible if loading finished by flushing.
-	UpdateLevelStreaming(ViewFamily);
+	UpdateLevelStreaming();
 
 	// Making levels visible is spread across several frames so we simply loop till it is done.
 	bool bLevelsPendingVisibility = true;
@@ -2387,12 +2356,12 @@ void UWorld::FlushLevelStreaming(FSceneViewFamily* ViewFamily, EFlushLevelStream
 			}
 	
 			// Update level streaming.
-			UpdateLevelStreaming( ViewFamily );
+			UpdateLevelStreaming();
 		}
 	}
 	
 	// Update level streaming one last time to make sure all RemoveFromWorld requests made it.
-	UpdateLevelStreaming( ViewFamily );
+	UpdateLevelStreaming();
 
 	// we need this, or the traces will be abysmally slow
 	EnsureCollisionTreeIsBuilt();
@@ -4511,7 +4480,7 @@ UWorld* FSeamlessTravelHandler::Tick()
 		else
 		{
 			// Make sure there are no pending visibility requests.
-			CurrentWorld->FlushLevelStreaming( NULL, EFlushLevelStreamingType::Visibility );
+			CurrentWorld->FlushLevelStreaming(EFlushLevelStreamingType::Visibility);
 
 			if (CurrentWorld->GameState)
 			{
diff --git a/Engine/Source/Runtime/Engine/Private/WorldComposition.cpp b/Engine/Source/Runtime/Engine/Private/WorldComposition.cpp
--- a/Engine/Source/Runtime/Engine/Private/WorldComposition.cpp
+++ b/Engine/Source/Runtime/Engine/Private/WorldComposition.cpp
@@ -570,32 +570,25 @@ void UWorldComposition::UpdateStreamingState()
 	}
 	
 	int32 NumPlayers = GEngine->GetNumGamePlayers(PlayWorld);
-	if (NumPlayers == 0 || GEngine->GameViewport == nullptr)
+	if (NumPlayers == 0)
 	{
 		return;
 	}
 
 	// calculate centroid location using local players views
-	FSceneViewFamilyContext ViewFamily(
-		FSceneViewFamily::ConstructionValues(GEngine->GameViewport->Viewport, PlayWorld->Scene, GEngine->GameViewport->EngineShowFlags).SetRealtimeUpdate(true));
-
 	int32 NumViews = 0;
 	FVector CentroidLocation = FVector::ZeroVector;
 	
 	for (int32 PlayerIndex = 0; PlayerIndex < NumPlayers; ++PlayerIndex)
 	{
 		ULocalPlayer* Player = GEngine->GetGamePlayer(PlayWorld, PlayerIndex);
-		if (Player && 
-			Player->ViewportClient && 
-			Player->PlayerController)
+		if (Player && Player->PlayerController)
 		{
 			FVector ViewLocation;
 			FRotator ViewRotation;
-			if (Player->CalcSceneView(&ViewFamily, /*out*/ ViewLocation, /*out*/ ViewRotation, Player->ViewportClient->Viewport))
-			{
-				CentroidLocation+= ViewLocation;
-				NumViews++;
-			}
+			Player->PlayerController->GetActorEyesViewPoint(ViewLocation, ViewRotation);
+			CentroidLocation+= ViewLocation;
+			NumViews++;
 		}
 	}
 
diff --git a/Engine/Source/Runtime/Engine/Public/Engine.h b/Engine/Source/Runtime/Engine/Public/Engine.h
--- a/Engine/Source/Runtime/Engine/Public/Engine.h
+++ b/Engine/Source/Runtime/Engine/Public/Engine.h
@@ -239,7 +239,6 @@
 #include "Engine/Level.h"
 #include "Engine/LevelStreaming.h"
 #include "Engine/LevelStreamingAlwaysLoaded.h"
-#include "Engine/LevelStreamingBounds.h"
 #include "Engine/LevelStreamingKismet.h"
 #include "Engine/LevelStreamingPersistent.h"
 #include "Lightmass/LightmappedSurfaceCollection.h"
