Add 'Character' stats category

diff --git a/Engine/Source/Runtime/Engine/Private/CharacterMovementComponent.cpp b/Engine/Source/Runtime/Engine/Private/CharacterMovementComponent.cpp
--- a/Engine/Source/Runtime/Engine/Private/CharacterMovementComponent.cpp
+++ b/Engine/Source/Runtime/Engine/Private/CharacterMovementComponent.cpp
@@ -24,6 +24,21 @@
 
 DEFINE_LOG_CATEGORY_STATIC(LogCharacterMovement, Log, All);
 
+/**
+ * Character stats
+ */
+DECLARE_STATS_GROUP(TEXT("Character"), STATGROUP_Character, STATCAT_Advanced);
+
+DECLARE_CYCLE_STAT(TEXT("Char Movement Tick"), STAT_CharacterMovementTick, STATGROUP_Character);
+DECLARE_CYCLE_STAT(TEXT("Char Movement Authority Time"), STAT_CharacterMovementAuthority, STATGROUP_Character);
+DECLARE_CYCLE_STAT(TEXT("Char Movement Simulated Time"), STAT_CharacterMovementSimulated, STATGROUP_Character);
+DECLARE_CYCLE_STAT(TEXT("Char Physics Interation"), STAT_CharPhysicsInteraction, STATGROUP_Character);
+DECLARE_CYCLE_STAT(TEXT("StepUp"), STAT_CharStepUp, STATGROUP_Character);
+DECLARE_CYCLE_STAT(TEXT("Char Update Acceleration"), STAT_CharUpdateAcceleration, STATGROUP_Character);
+DECLARE_CYCLE_STAT(TEXT("Char MoveUpdateDelegate"), STAT_CharMoveUpdateDelegate, STATGROUP_Character);
+DECLARE_CYCLE_STAT(TEXT("PhysWalking"), STAT_CharPhysWalking, STATGROUP_Character);
+DECLARE_CYCLE_STAT(TEXT("PhysFalling"), STAT_CharPhysFalling, STATGROUP_Character);
+
 // MAGIC NUMBERS
 const float MAX_STEP_SIDE_Z = 0.08f;	// maximum z value for the normal on the vertical side of steps
 const float SWIMBOBSPEED = -80.f;
@@ -822,6 +837,8 @@ void UCharacterMovementComponent::TickComponent(float DeltaTime, enum ELevelTick
 {
 	Super::TickComponent(DeltaTime, TickType, ThisTickFunction);
 
+	SCOPE_CYCLE_COUNTER(STAT_CharacterMovementTick);
+
 	const FVector InputVector = ConsumeInputVector();
 	if (!HasValidData() || ShouldSkipUpdate(DeltaTime) || UpdatedComponent->IsSimulatingPhysics())
 	{
@@ -855,13 +872,17 @@ void UCharacterMovementComponent::TickComponent(float DeltaTime, enum ELevelTick
 		// Allow root motion to move characters that have no controller.
 		if( CharacterOwner->IsLocallyControlled() || (!CharacterOwner->Controller && bRunPhysicsWithNoController) || (!CharacterOwner->Controller && CharacterOwner->IsPlayingRootMotion()) )
 		{
-			// We need to check the jump state before adjusting input acceleration, to minimize latency
-			// and to make sure acceleration respects our potentially new falling state.
-			CharacterOwner->CheckJumpInput(DeltaTime);
+			{
+				SCOPE_CYCLE_COUNTER(STAT_CharUpdateAcceleration);
 
-			// apply input to acceleration
-			Acceleration = ScaleInputAcceleration(ConstrainInputAcceleration(InputVector));
-			AnalogInputModifier = ComputeAnalogInputModifier();
+				// We need to check the jump state before adjusting input acceleration, to minimize latency
+				// and to make sure acceleration respects our potentially new falling state.
+				CharacterOwner->CheckJumpInput(DeltaTime);
+
+				// apply input to acceleration
+				Acceleration = ScaleInputAcceleration(ConstrainInputAcceleration(InputVector));
+				AnalogInputModifier = ComputeAnalogInputModifier();
+			}
 
 			if (CharacterOwner->Role == ROLE_Authority)
 			{
@@ -891,6 +912,8 @@ void UCharacterMovementComponent::TickComponent(float DeltaTime, enum ELevelTick
 
 	if (bEnablePhysicsInteraction)
 	{
+		SCOPE_CYCLE_COUNTER(STAT_CharPhysicsInteraction);
+
 		if (CurrentFloor.HitResult.IsValidBlockingHit())
 		{
 			// Apply downwards force when walking on top of physics objects
@@ -963,6 +986,7 @@ void UCharacterMovementComponent::AdjustProxyCapsuleSize()
 
 void UCharacterMovementComponent::SimulatedTick(float DeltaSeconds)
 {
+	SCOPE_CYCLE_COUNTER(STAT_CharacterMovement);
 	SCOPE_CYCLE_COUNTER(STAT_CharacterMovementSimulated);
 
 	// If we are playing a RootMotion AnimMontage.
@@ -974,6 +998,7 @@ void UCharacterMovementComponent::SimulatedTick(float DeltaSeconds)
 		// Tick animations before physics.
 		if( CharacterOwner->GetMesh() )
 		{
+			SCOPE_CYCLE_COUNTER(STAT_CharacterMovement);
 			TickCharacterPose(DeltaSeconds);
 
 			// Make sure animation didn't trigger an event that destroyed us
@@ -1414,6 +1439,7 @@ void UCharacterMovementComponent::DisableMovement()
 
 void UCharacterMovementComponent::PerformMovement(float DeltaSeconds)
 {
+	SCOPE_CYCLE_COUNTER(STAT_CharacterMovement);
 	SCOPE_CYCLE_COUNTER(STAT_CharacterMovementAuthority);
 
 	if (!HasValidData())
@@ -1582,6 +1608,8 @@ void UCharacterMovementComponent::PerformMovement(float DeltaSeconds)
 
 void UCharacterMovementComponent::CallMovementUpdateDelegate(float DeltaTime, const FVector& OldLocation, const FVector& OldVelocity)
 {
+	SCOPE_CYCLE_COUNTER(STAT_CharMoveUpdateDelegate);
+
 	// Update component velocity in case events want to read it
 	UpdateComponentVelocity();
 
@@ -2935,6 +2963,8 @@ float UCharacterMovementComponent::BoostAirControl(float DeltaTime, float TickAi
 
 void UCharacterMovementComponent::PhysFalling(float deltaTime, int32 Iterations)
 {
+	SCOPE_CYCLE_COUNTER(STAT_CharPhysFalling);
+
 	if (deltaTime < MIN_TICK_TIME)
 	{
 		return;
@@ -3510,6 +3540,9 @@ void UCharacterMovementComponent::MaintainHorizontalGroundVelocity()
 
 void UCharacterMovementComponent::PhysWalking(float deltaTime, int32 Iterations)
 {
+	SCOPE_CYCLE_COUNTER(STAT_CharPhysWalking);
+
+
 	if (deltaTime < MIN_TICK_TIME)
 	{
 		return;
@@ -4867,6 +4900,9 @@ bool UCharacterMovementComponent::CanStepUp(const FHitResult& Hit) const
 
 bool UCharacterMovementComponent::StepUp(const FVector& InGravDir, const FVector& Delta, const FHitResult &InHit, FStepDownResult* OutStepDownResult)
 {
+	SCOPE_CYCLE_COUNTER(STAT_CharStepUp);
+
+
 	if (!CanStepUp(InHit))
 	{
 		return false;
diff --git a/Engine/Source/Runtime/Engine/Private/Collision/CollisionConversions.cpp b/Engine/Source/Runtime/Engine/Private/Collision/CollisionConversions.cpp
--- a/Engine/Source/Runtime/Engine/Private/Collision/CollisionConversions.cpp
+++ b/Engine/Source/Runtime/Engine/Private/Collision/CollisionConversions.cpp
@@ -16,6 +16,10 @@ static const PxQueryHit InvalidQueryHit;
 // Forward declare, I don't want to move the entire function right now or we lose change history.
 static bool ConvertOverlappedShapeToImpactHit(const PxLocationHit& PHit, const FVector& StartLoc, const FVector& EndLoc, FHitResult& OutResult, const PxGeometry& Geom, const PxTransform& QueryTM, const PxFilterData& QueryFilter, bool bReturnPhysMat);
 
+DECLARE_CYCLE_STAT(TEXT("ConvertQueryHit"), STAT_ConvertQueryImpactHit, STATGROUP_Collision);
+DECLARE_CYCLE_STAT(TEXT("ConvertOverlapToHit"), STAT_CollisionConvertOverlapToHit, STATGROUP_Collision);
+DECLARE_CYCLE_STAT(TEXT("ConvertOverlap"), STAT_CollisionConvertOverlap, STATGROUP_Collision);
+
 
 #if !(UE_BUILD_SHIPPING || UE_BUILD_TEST)
 /* Validate Normal of OutResult. We're on hunt for invalid normal */
@@ -333,6 +337,8 @@ static void SetHitResultFromShapeAndFaceIndex(const PxShape* PShape,  const PxRi
 
 void ConvertQueryImpactHit(const PxLocationHit& PHit, FHitResult& OutResult, float CheckLength, const PxFilterData& QueryFilter, const FVector& StartLoc, const FVector& EndLoc, const PxGeometry* const Geom, const PxTransform& QueryTM, bool bReturnFaceIndex, bool bReturnPhysMat)
 {
+	SCOPE_CYCLE_COUNTER(STAT_ConvertQueryImpactHit);
+
 	checkSlow(PHit.flags & PxHitFlag::eDISTANCE);
 	if (Geom != NULL && PHit.hadInitialOverlap())
 	{
@@ -623,6 +629,8 @@ static bool ComputeInflatedMTD(const float MtdInflation, const PxLocationHit& PH
 /** Util to convert an overlapped shape into a sweep hit result, returns whether it was a blocking hit. */
 static bool ConvertOverlappedShapeToImpactHit(const PxLocationHit& PHit, const FVector& StartLoc, const FVector& EndLoc, FHitResult& OutResult, const PxGeometry& Geom, const PxTransform& QueryTM, const PxFilterData& QueryFilter, bool bReturnPhysMat)
 {
+	SCOPE_CYCLE_COUNTER(STAT_CollisionConvertOverlapToHit);
+
 	OutResult.TraceStart = StartLoc;
 	OutResult.TraceEnd = EndLoc;
 
@@ -792,6 +800,8 @@ static void AddUniqueOverlap(TArray<FOverlapResult>& OutOverlaps, const FOverlap
 
 bool ConvertOverlapResults(int32 NumOverlaps, PxOverlapHit* POverlapResults, const PxFilterData& QueryFilter, TArray<FOverlapResult>& OutOverlaps)
 {
+	SCOPE_CYCLE_COUNTER(STAT_CollisionConvertOverlap);
+
 	OutOverlaps.Reserve(NumOverlaps);
 	bool bBlockingFound = false;
 
diff --git a/Engine/Source/Runtime/Engine/Private/LevelTick.cpp b/Engine/Source/Runtime/Engine/Private/LevelTick.cpp
--- a/Engine/Source/Runtime/Engine/Private/LevelTick.cpp
+++ b/Engine/Source/Runtime/Engine/Private/LevelTick.cpp
@@ -36,8 +36,6 @@ bool GLogDetailedDumpStats = true;
 
 
 // DECLARE_CYCLE_STAT is the reverse of what will be displayed in the game's stat game
-DEFINE_STAT(STAT_CharacterMovementAuthority);
-DEFINE_STAT(STAT_CharacterMovementSimulated)
 
 DEFINE_STAT(STAT_AsyncWorkWaitTime);
 DEFINE_STAT(STAT_PhysicsTime);
@@ -61,6 +59,7 @@ DEFINE_STAT(STAT_PostTickComponentUpdateWait);
 DEFINE_STAT(STAT_TickTime);
 DEFINE_STAT(STAT_WorldTickTime);
 DEFINE_STAT(STAT_UpdateCameraTime);
+DEFINE_STAT(STAT_CharacterMovement);
 
 DEFINE_STAT(STAT_VolumeStreamingTickTime);
 DEFINE_STAT(STAT_VolumeStreamingChecks);
diff --git a/Engine/Source/Runtime/Engine/Public/EngineStats.h b/Engine/Source/Runtime/Engine/Public/EngineStats.h
--- a/Engine/Source/Runtime/Engine/Public/EngineStats.h
+++ b/Engine/Source/Runtime/Engine/Public/EngineStats.h
@@ -45,14 +45,13 @@ DECLARE_CYCLE_STAT_EXTERN(TEXT("GC Sweep Time"),STAT_GCSweepTime,STATGROUP_Game,
 DECLARE_CYCLE_STAT_EXTERN(TEXT("Tick Time"),STAT_TickTime,STATGROUP_Game, );
 DECLARE_CYCLE_STAT_EXTERN(TEXT("World Tick Time"),STAT_WorldTickTime,STATGROUP_Game, );
 DECLARE_CYCLE_STAT_EXTERN(TEXT("Update Camera Time"),STAT_UpdateCameraTime,STATGROUP_Game, );
+DECLARE_CYCLE_STAT_EXTERN(TEXT("Char Movement Total"), STAT_CharacterMovement, STATGROUP_Game, );
 DECLARE_CYCLE_STAT_EXTERN(TEXT("Post Tick Component Update"),STAT_PostTickComponentUpdate,STATGROUP_Game, );
 DECLARE_CYCLE_STAT_EXTERN(TEXT("Post Tick Component Wait"),STAT_PostTickComponentUpdateWait,STATGROUP_Game, );
 DECLARE_CYCLE_STAT_EXTERN(TEXT("  Recreate"),STAT_PostTickComponentRecreate,STATGROUP_Game, );
 DECLARE_CYCLE_STAT_EXTERN(TEXT("  Transform or RenderData"),STAT_PostTickComponentLW,STATGROUP_Game, );
 DECLARE_CYCLE_STAT_EXTERN(TEXT("Update Particle Data"),STAT_ParticleManagerUpdateData,STATGROUP_Game, );
 DECLARE_CYCLE_STAT_EXTERN(TEXT("Async Work Wait"),STAT_AsyncWorkWaitTime,STATGROUP_Game, );
-DECLARE_CYCLE_STAT_EXTERN(TEXT("Char Movement Authority Time"),STAT_CharacterMovementAuthority,STATGROUP_Game, );
-DECLARE_CYCLE_STAT_EXTERN(TEXT("Char Movement Simulated Time"),STAT_CharacterMovementSimulated,STATGROUP_Game, );
 DECLARE_CYCLE_STAT_EXTERN(TEXT("Visual Logging time"), STAT_VisualLog, STATGROUP_Game, ENGINE_API);
 DECLARE_CYCLE_STAT_EXTERN(TEXT("Net Tick Time"),STAT_NetWorldTickTime,STATGROUP_Game, );
 DECLARE_CYCLE_STAT_EXTERN(TEXT("Nav Tick Time"),STAT_NavWorldTickTime,STATGROUP_Game, );
@@ -74,6 +73,7 @@ DECLARE_CYCLE_STAT_EXTERN(TEXT("  Rebuild Conditionals"),STAT_NetRebuildConditio
 DECLARE_CYCLE_STAT_EXTERN(TEXT("Net Post BC Tick Time"),STAT_NetBroadcastPostTickTime,STATGROUP_Game, );
 DECLARE_CYCLE_STAT_EXTERN(TEXT("Net PackageMap SerializeObject"),STAT_PackageMap_SerializeObjectTime,STATGROUP_Game, );
 
+
 /**
  * FPS chart stats
  */
