UE-6595 Fixed note's ArrowComponent attachment
https://github.com/EpicGames/UnrealEngine/pull/684 (Contributed by [REDACTED_PERSON])

diff --git a/Engine/Source/Runtime/Engine/Private/Note.cpp b/Engine/Source/Runtime/Engine/Private/Note.cpp
--- a/Engine/Source/Runtime/Engine/Private/Note.cpp
+++ b/Engine/Source/Runtime/Engine/Private/Note.cpp
@@ -41,6 +41,7 @@ ANote::ANote(const FObjectInitializer& ObjectInitializer)
 			ArrowComponent->SpriteInfo.DisplayName = ConstructorStatics.NAME_Notes;
 			ArrowComponent->Mobility = EComponentMobility::Static;
 			ArrowComponent->bIsScreenSizeScaled = true;
+			ArrowComponent->AttachParent = RootComponent;
 		}
 
 		if (SpriteComponent)
