Optimize UpdateKinematicBonesToAnim  (over 2x as fast in my test)
  - don't call FBodyInstance::SetBodyTransform for each body
  - do one write lock for the whole physics asset
  - do less redundant validation of transforms
Change FBodyInstance::SetBodyTransform to not check body is in correct place (which involves a read lock) before moving it.
Rename UpdateKinematicBonesToPhysics to UpdateKinematicBonesToAnim to correctly describe what it does!

diff --git a/Engine/Source/Runtime/Engine/Classes/Components/SkeletalMeshComponent.h b/Engine/Source/Runtime/Engine/Classes/Components/SkeletalMeshComponent.h
--- a/Engine/Source/Runtime/Engine/Classes/Components/SkeletalMeshComponent.h
+++ b/Engine/Source/Runtime/Engine/Classes/Components/SkeletalMeshComponent.h
@@ -330,6 +330,12 @@ public:
 	UPROPERTY(transient)
 	uint32 bHasValidBodies:1;
 
+	/** Set during InitArticulated, to indicate if there are bodies in the sync scene */
+	uint32 bHasBodiesInSyncScene:1;
+
+	/** Set during InitArticulated, to indicate if there are bodies in the async scene */
+	uint32 bHasBodiesInAsyncScene:1;
+
 	/** If we are running physics, should we update non-simulated bones based on the animation bone positions. */
 	UPROPERTY(EditAnywhere, AdvancedDisplay, BlueprintReadWrite, Category=SkeletalMesh)
 	TEnumAsByte<EKinematicBonesUpdateToPhysics::Type> KinematicBonesUpdateType;
@@ -796,7 +802,6 @@ public:
 	virtual FBoxSphereBounds CalcBounds(const FTransform& LocalToWorld) const override;
 	virtual bool IsAnySimulatingPhysics() const override;
 	virtual void OnUpdateTransform(bool bSkipPhysicsMove) override;
-	virtual void UpdateOverlaps(TArray<FOverlapInfo> const* PendingOverlaps=NULL, bool bDoNotifies=true, const TArray<FOverlapInfo>* OverlapsAtEndLocation=NULL) override;
 	
 	/**
 	 *  Test the collision of the supplied component at the supplied location/rotation, and determine the set of components that it overlaps
@@ -1012,7 +1017,7 @@ public:
 	 *	Iterate over each physics body in the physics for this mesh, and for each 'kinematic' (ie fixed or default if owner isn't simulating) one, update its
 	 *	transform based on the animated transform.
 	 */
-	void UpdateKinematicBonesToPhysics(const TArray<FTransform>& InSpaceBases, bool bTeleport, bool bNeedsSkinning, bool bForceUpdate = false);
+	void UpdateKinematicBonesToAnim(const TArray<FTransform>& InSpaceBases, bool bTeleport, bool bNeedsSkinning, bool bForceUpdate = false);
 	
 	/**
 	 * Look up all bodies for broken constraints.
diff --git a/Engine/Source/Runtime/Engine/Private/PhysicsEngine/BodyInstance.cpp b/Engine/Source/Runtime/Engine/Private/PhysicsEngine/BodyInstance.cpp
--- a/Engine/Source/Runtime/Engine/Private/PhysicsEngine/BodyInstance.cpp
+++ b/Engine/Source/Runtime/Engine/Private/PhysicsEngine/BodyInstance.cpp
@@ -2836,32 +2836,6 @@ void FBodyInstance::SetBodyTransform(const FTransform& NewTransform, bool bTelep
 
 	if (RigidActor != nullptr)
 	{
-		// Do nothing if already in correct place
-		{
-			SCOPED_SCENE_READ_LOCK(RigidActor->getScene());
-			if (PRigidDynamic && !IsRigidBodyNonKinematic(PRigidDynamic) && !bTeleport)
-			{
-				const PxScene* PScene = PRigidDynamic->getScene();
-				FPhysScene* PhysScene = FPhysxUserData::Get<FPhysScene>(PScene->userData);
-				FTransform CurrentTarget;
-				if (PhysScene->GetKinematicTarget(this, CurrentTarget))
-				{
-					if (CurrentTarget.Equals(NewTransform))
-					{
-						return;
-					}
-				}
-			}
-			else
-			{
-				const PxTransform PCurrentPose = RigidActor->getGlobalPose();
-				if (NewTransform.Equals(P2UTransform(PCurrentPose)))
-				{
-					return;
-				}
-			}
-		}
-
 		const PxTransform PNewPose = U2PTransform(NewTransform);
 
 		if (!PNewPose.isValid())
diff --git a/Engine/Source/Runtime/Engine/Private/PhysicsEngine/PhysAnim.cpp b/Engine/Source/Runtime/Engine/Private/PhysicsEngine/PhysAnim.cpp
--- a/Engine/Source/Runtime/Engine/Private/PhysicsEngine/PhysAnim.cpp
+++ b/Engine/Source/Runtime/Engine/Private/PhysicsEngine/PhysAnim.cpp
@@ -11,6 +11,7 @@
 #if WITH_PHYSX
 	#include "PhysXSupport.h"
 #endif // WITH_PHYSX
+#include "PhysicsPublic.h"
 #include "PhysicsEngine/PhysicsConstraintTemplate.h"
 #include "PhysicsEngine/PhysicsAsset.h"
 
@@ -228,7 +229,7 @@ void USkeletalMeshComponent::BlendInPhysics()
 
 
 
-void USkeletalMeshComponent::UpdateKinematicBonesToPhysics(const TArray<FTransform>& InSpaceBases, bool bTeleport, bool bNeedsSkinning, bool bForceUpdate)
+void USkeletalMeshComponent::UpdateKinematicBonesToAnim(const TArray<FTransform>& InSpaceBases, bool bTeleport, bool bNeedsSkinning, bool bForceUpdate)
 {
 	SCOPE_CYCLE_COUNTER(STAT_UpdateRBBones);
 
@@ -246,8 +247,21 @@ void USkeletalMeshComponent::UpdateKinematicBonesToPhysics(const TArray<FTransfo
 		return;
 	}
 
+	// Get the scene, and do nothing if we can't get one.
+	FPhysScene* PhysScene = nullptr;
+	if (GetWorld() != nullptr)
+	{
+		PhysScene = GetWorld()->GetPhysicsScene();
+	}
+
+	if(PhysScene == nullptr)
+	{
+		return;
+	}
+
 	const FTransform& CurrentLocalToWorld = ComponentToWorld;
 
+	// Gracefully handle NaN
 	if(CurrentLocalToWorld.ContainsNaN())
 	{
 		return;
@@ -272,10 +286,11 @@ void USkeletalMeshComponent::UpdateKinematicBonesToPhysics(const TArray<FTransfo
 #if !(UE_BUILD_SHIPPING || UE_BUILD_TEST)
 	if( !MeshScale3D.IsUniform() )
 	{
-		UE_LOG(LogPhysics, Log, TEXT("USkeletalMeshComponent::UpdateKinematicBonesToPhysics : Non-uniform scale factor (%s) can cause physics to mismatch for %s  SkelMesh: %s"), *MeshScale3D.ToString(), *GetFullName(), SkeletalMesh ? *SkeletalMesh->GetFullName() : TEXT("NULL"));
+		UE_LOG(LogPhysics, Log, TEXT("USkeletalMeshComponent::UpdateKinematicBonesToAnim : Non-uniform scale factor (%s) can cause physics to mismatch for %s  SkelMesh: %s"), *MeshScale3D.ToString(), *GetFullName(), SkeletalMesh ? *SkeletalMesh->GetFullName() : TEXT("NULL"));
 	}
 #endif
 
+
 	if (bEnablePerPolyCollision == false)
 	{
 		const UPhysicsAsset* const PhysicsAsset = GetPhysicsAsset();
@@ -291,6 +306,19 @@ void USkeletalMeshComponent::UpdateKinematicBonesToPhysics(const TArray<FTransfo
 			}
 #endif
 
+#if WITH_PHYSX
+			// Lock the scenes we need (flags set in InitArticulated)
+			if(bHasBodiesInSyncScene)
+			{
+				SCENE_LOCK_WRITE(PhysScene->GetPhysXScene(PST_Sync))
+			}
+
+			if (bHasBodiesInAsyncScene)
+			{
+				SCENE_LOCK_WRITE(PhysScene->GetPhysXScene(PST_Async))
+			}
+#endif
+
 			// Iterate over each body
 			for (int32 i = 0; i < Bodies.Num(); i++)
 			{
@@ -311,11 +339,26 @@ void USkeletalMeshComponent::UpdateKinematicBonesToPhysics(const TArray<FTransfo
 					}
 					else
 					{
+#if WITH_PHYSX
 						// update bone transform to world
-						FTransform BoneTransform = InSpaceBases[BoneIndex] * CurrentLocalToWorld;
+						const FTransform BoneTransform = InSpaceBases[BoneIndex] * CurrentLocalToWorld;
+						ensure(!BoneTransform.ContainsNaN());
+
+						// If kinematic and not teleporting, set kinematic target
+						PxRigidDynamic* PRigidDynamic = BodyInst->GetPxRigidDynamic();
+						if (!IsRigidBodyNonKinematic(PRigidDynamic) && !bTeleport)
+						{
+							PhysScene->SetKinematicTarget(BodyInst, BoneTransform, true);
+						}
+						// Otherwise, set global pose
+						else
+						{
+							const PxTransform PNewPose = U2PTransform(BoneTransform);
+							ensure(PNewPose.isValid());
+							PRigidDynamic->setGlobalPose(PNewPose);
+						}
+#endif
 
-						// move body
-						BodyInst->SetBodyTransform(BoneTransform, bTeleport);
 
 						// now update scale
 						// if uniform, we'll use BoneTranform
@@ -346,6 +389,19 @@ void USkeletalMeshComponent::UpdateKinematicBonesToPhysics(const TArray<FTransfo
 					}
 				}
 			}
+
+#if WITH_PHYSX
+			// Unlock the scenes 
+			if (bHasBodiesInSyncScene)
+			{
+				SCENE_UNLOCK_WRITE(PhysScene->GetPhysXScene(PST_Sync))
+			}
+
+			if (bHasBodiesInAsyncScene)
+			{
+				SCENE_UNLOCK_WRITE(PhysScene->GetPhysXScene(PST_Async))
+			}
+#endif
 		}
 	}
 	else
diff --git a/Engine/Source/Runtime/Engine/Private/SkeletalMeshComponent.cpp b/Engine/Source/Runtime/Engine/Private/SkeletalMeshComponent.cpp
--- a/Engine/Source/Runtime/Engine/Private/SkeletalMeshComponent.cpp
+++ b/Engine/Source/Runtime/Engine/Private/SkeletalMeshComponent.cpp
@@ -1033,7 +1033,7 @@ void USkeletalMeshComponent::PostAnimEvaluation(FAnimationEvaluationContext& Eva
 	InvalidateCachedBounds();
 
 	// update physics data from animated data
-	UpdateKinematicBonesToPhysics(GetEditableSpaceBases(), false, true);
+	UpdateKinematicBonesToAnim(GetEditableSpaceBases(), false, true);
 	UpdateRBJointMotors();
 
 	// @todo anim : hack TTP 224385	ANIM: Skeletalmesh double buffer
diff --git a/Engine/Source/Runtime/Engine/Private/SkeletalMeshComponentPhysics.cpp b/Engine/Source/Runtime/Engine/Private/SkeletalMeshComponentPhysics.cpp
--- a/Engine/Source/Runtime/Engine/Private/SkeletalMeshComponentPhysics.cpp
+++ b/Engine/Source/Runtime/Engine/Private/SkeletalMeshComponentPhysics.cpp
@@ -793,6 +793,16 @@ void USkeletalMeshComponent::InitArticulated(FPhysScene* PhysScene)
 			FTransform BoneTransform = GetBoneTransform( BoneIndex );
 			BodyInst->InitBody( BodySetup, BoneTransform, this, PhysScene, Aggregate);
 #endif //WITH_PHYSX
+
+			// Remember if we have bodies in sync/async scene, so we know which scene(s) to lock when moving bodies
+			if(BodyInst->UseAsyncScene())
+			{
+				bHasBodiesInAsyncScene = true;
+			}
+			else
+			{
+				bHasBodiesInSyncScene = true;
+			}
 		}
 	}
 
@@ -902,6 +912,10 @@ void USkeletalMeshComponent::TermArticulated()
 		Aggregate = NULL;
 	}
 #endif //WITH_PHYSX
+
+	// Reset bools for scenes
+	bHasBodiesInAsyncScene = false;
+	bHasBodiesInSyncScene = false;
 }
 
 void USkeletalMeshComponent::TermBodiesBelow(FName ParentBoneName)
@@ -1291,9 +1305,9 @@ void USkeletalMeshComponent::OnUpdateTransform(bool bSkipPhysicsMove)
 	Super::OnUpdateTransform(true);
 
 	// Always send new transform to physics
-	if(bPhysicsStateCreated && !bSkipPhysicsMove )
+	if(bPhysicsStateCreated && !bSkipPhysicsMove)
 	{
-		UpdateKinematicBonesToPhysics(GetSpaceBases(), false, false, true);
+		UpdateKinematicBonesToAnim(GetSpaceBases(), false, false, true);
 	}
 
 #if WITH_APEX_CLOTHING
@@ -1305,12 +1319,6 @@ void USkeletalMeshComponent::OnUpdateTransform(bool bSkipPhysicsMove)
 #endif //#if WITH_APEX_CLOTHING
 }
 
-void USkeletalMeshComponent::UpdateOverlaps(TArray<FOverlapInfo> const* PendingOverlaps, bool bDoNotifies, const TArray<FOverlapInfo>* OverlapsAtEndLocation)
-{
-		UPrimitiveComponent::UpdateOverlaps(PendingOverlaps, bDoNotifies, OverlapsAtEndLocation);
-	}
-
-
 void USkeletalMeshComponent::CreatePhysicsState()
 {
 	//	UE_LOG(LogSkeletalMesh, Warning, TEXT("Creating Physics State (%s : %s)"), *GetNameSafe(GetOuter()),  *GetName());
@@ -1326,8 +1334,6 @@ void USkeletalMeshComponent::CreatePhysicsState()
 		BodySetup->CreatePhysicsMeshes();
 		Super::CreatePhysicsState();	//If we're doing per poly we'll use the body instance of the primitive component
 	}
-
-	
 }
 
 
diff --git a/Engine/Source/Runtime/Engine/Public/AnimTree.h b/Engine/Source/Runtime/Engine/Public/AnimTree.h
--- a/Engine/Source/Runtime/Engine/Public/AnimTree.h
+++ b/Engine/Source/Runtime/Engine/Public/AnimTree.h
@@ -27,7 +27,7 @@ DECLARE_CYCLE_STAT_EXTERN(TEXT("Copy Pose"), STAT_AnimNativeCopyPoses, STATGROUP
 DECLARE_CYCLE_STAT_EXTERN(TEXT("StateMachine Eval"), STAT_AnimStateMachineEvaluate, STATGROUP_Anim, );
 DECLARE_CYCLE_STAT_EXTERN(TEXT("FillSpaceBases"), STAT_SkelComposeTime, STATGROUP_Anim, );
 DECLARE_CYCLE_STAT_EXTERN(TEXT("InterpolateSkippedFrames"), STAT_InterpolateSkippedFrames, STATGROUP_Anim, );
-DECLARE_CYCLE_STAT_EXTERN(TEXT("UpdateKinematicBonesToPhysics"), STAT_UpdateRBBones, STATGROUP_Anim, );
+DECLARE_CYCLE_STAT_EXTERN(TEXT("UpdateKinematicBonesToAnim"), STAT_UpdateRBBones, STATGROUP_Anim, );
 DECLARE_CYCLE_STAT_EXTERN(TEXT("UpdateRBJointsMotors"), STAT_UpdateRBJoints, STATGROUP_Anim, );
 DECLARE_CYCLE_STAT_EXTERN(TEXT("UpdateLocalToWorldAndOverlaps"), STAT_UpdateLocalToWorldAndOverlaps, STATGROUP_Anim, );
 DECLARE_CYCLE_STAT_EXTERN(TEXT("SkelComp UpdateTransform"), STAT_SkelCompUpdateTransform, STATGROUP_Anim, );
