Fixed some places where kerning wasn't being scaled when measuring/drawing in Canvas

diff --git a/Engine/Source/Runtime/Engine/Private/UserInterface/Canvas.cpp b/Engine/Source/Runtime/Engine/Private/UserInterface/Canvas.cpp
--- a/Engine/Source/Runtime/Engine/Private/UserInterface/Canvas.cpp
+++ b/Engine/Source/Runtime/Engine/Private/UserInterface/Canvas.cpp
@@ -1158,7 +1158,7 @@ void UCanvas::MeasureStringInternal( FTextSizingParameters& Parameters, const TC
 			float CharSpacing = DefaultCharIncrement;
 			if ( pPrevPos )
 			{
-				CharSpacing += Parameters.DrawFont->GetCharKerning( *pPrevPos, Ch );
+				CharSpacing += Parameters.DrawFont->GetCharKerning( *pPrevPos, Ch ) * ScaleX;
 			}
 
 			CharWidth *= ScaleX;
diff --git a/Engine/Source/Runtime/Engine/Private/UserInterface/CanvasItem.cpp b/Engine/Source/Runtime/Engine/Private/UserInterface/CanvasItem.cpp
--- a/Engine/Source/Runtime/Engine/Private/UserInterface/CanvasItem.cpp
+++ b/Engine/Source/Runtime/Engine/Private/UserInterface/CanvasItem.cpp
@@ -1106,7 +1106,7 @@ void FCanvasTextItem::DrawStringInternal_RuntimeCache( FCanvas* InCanvas, const
 
 			if( !bIsWhitespace && PreviousCharEntry.IsValidEntry() )
 			{
-				Kerning = CharacterList.GetKerning( PreviousCharEntry, Entry );
+				Kerning = CharacterList.GetKerning( PreviousCharEntry, Entry ) * Scale.X;
 			}
 			else
 			{
