UObject constructors create subobjects using UObject::Create* instead of FObjectInitializer::Create*. Updated Engine module.

diff --git a/Engine/Source/Runtime/Engine/Private/AI/Navigation/NavLinkProxy.cpp b/Engine/Source/Runtime/Engine/Private/AI/Navigation/NavLinkProxy.cpp
--- a/Engine/Source/Runtime/Engine/Private/AI/Navigation/NavLinkProxy.cpp
+++ b/Engine/Source/Runtime/Engine/Private/AI/Navigation/NavLinkProxy.cpp
@@ -13,19 +13,19 @@
 
 ANavLinkProxy::ANavLinkProxy(const FObjectInitializer& ObjectInitializer) : Super(ObjectInitializer)
 {
-	USceneComponent* SceneComponent = ObjectInitializer.CreateDefaultSubobject<USceneComponent>(this, TEXT("PositionComponent"));
+	USceneComponent* SceneComponent = CreateDefaultSubobject<USceneComponent>(TEXT("PositionComponent"));
 	RootComponent = SceneComponent;
 
 	bHidden = true;
 
 #if WITH_EDITORONLY_DATA
-	EdRenderComp = ObjectInitializer.CreateDefaultSubobject<UNavLinkRenderingComponent>(this, TEXT("EdRenderComp"));
+	EdRenderComp = CreateDefaultSubobject<UNavLinkRenderingComponent>(TEXT("EdRenderComp"));
 	EdRenderComp->PostPhysicsComponentTick.bCanEverTick = false;
 	EdRenderComp->AttachParent = RootComponent;
 #endif // WITH_EDITORONLY_DATA
 
 #if WITH_EDITOR
-	SpriteComponent = ObjectInitializer.CreateEditorOnlyDefaultSubobject<UBillboardComponent>(this, TEXT("Sprite"));
+	SpriteComponent = CreateEditorOnlyDefaultSubobject<UBillboardComponent>(TEXT("Sprite"));
 	if (!IsRunningCommandlet() && (SpriteComponent != NULL))
 	{
 		static ConstructorHelpers::FObjectFinderOptional<UTexture2D> SpriteTexture(TEXT("/Engine/EditorResources/AI/S_NavLink"));
@@ -41,7 +41,7 @@ ANavLinkProxy::ANavLinkProxy(const FObjectInitializer& ObjectInitializer) : Supe
 	}
 #endif
 
-	SmartLinkComp = ObjectInitializer.CreateDefaultSubobject<UNavLinkCustomComponent>(this, TEXT("SmartLinkComp"));
+	SmartLinkComp = CreateDefaultSubobject<UNavLinkCustomComponent>(TEXT("SmartLinkComp"));
 	SmartLinkComp->SetNavigationRelevancy(false);
 	SmartLinkComp->SetMoveReachedLink(this, &ANavLinkProxy::NotifySmartLinkReached);
 	bSmartLinkIsRelevant = false;
diff --git a/Engine/Source/Runtime/Engine/Private/AI/Navigation/NavigationTestingActor.cpp b/Engine/Source/Runtime/Engine/Private/AI/Navigation/NavigationTestingActor.cpp
--- a/Engine/Source/Runtime/Engine/Private/AI/Navigation/NavigationTestingActor.cpp
+++ b/Engine/Source/Runtime/Engine/Private/AI/Navigation/NavigationTestingActor.cpp
@@ -26,7 +26,7 @@ TStatId FNavTestTickHelper::GetStatId() const
 ANavigationTestingActor::ANavigationTestingActor(const FObjectInitializer& ObjectInitializer) : Super(ObjectInitializer)
 {
 #if WITH_EDITORONLY_DATA
-	EdRenderComp = ObjectInitializer.CreateDefaultSubobject<UNavTestRenderingComponent>(this, TEXT("EdRenderComp"));
+	EdRenderComp = CreateDefaultSubobject<UNavTestRenderingComponent>(TEXT("EdRenderComp"));
 	EdRenderComp->PostPhysicsComponentTick.bCanEverTick = false;
 
 #if WITH_RECAST
@@ -50,7 +50,7 @@ ANavigationTestingActor::ANavigationTestingActor(const FObjectInitializer& Objec
 	// collision profile name set up - found in baseengine.ini
 	static FName CollisionProfileName(TEXT("Pawn"));
 
-	CapsuleComponent = ObjectInitializer.CreateDefaultSubobject<UCapsuleComponent>(this, TEXT("CollisionCylinder"));
+	CapsuleComponent = CreateDefaultSubobject<UCapsuleComponent>(TEXT("CollisionCylinder"));
 	CapsuleComponent->InitCapsuleSize(NavAgentProps.AgentRadius, NavAgentProps.AgentHeight / 2);
 	CapsuleComponent->SetCollisionProfileName(CollisionProfileName);
 	CapsuleComponent->CanCharacterStepUpOn = ECB_No;
diff --git a/Engine/Source/Runtime/Engine/Private/AmbientSound.cpp b/Engine/Source/Runtime/Engine/Private/AmbientSound.cpp
--- a/Engine/Source/Runtime/Engine/Private/AmbientSound.cpp
+++ b/Engine/Source/Runtime/Engine/Private/AmbientSound.cpp
@@ -26,7 +26,7 @@
 AAmbientSound::AAmbientSound(const FObjectInitializer& ObjectInitializer)
 	: Super(ObjectInitializer)
 {
-	AudioComponent = ObjectInitializer.CreateDefaultSubobject<UAudioComponent>(this, TEXT("AudioComponent0"));
+	AudioComponent = CreateDefaultSubobject<UAudioComponent>(TEXT("AudioComponent0"));
 
 	AudioComponent->bAutoActivate = true;
 	AudioComponent->bStopWhenOwnerDestroyed = true;
diff --git a/Engine/Source/Runtime/Engine/Private/Atmosphere/Atmosphere.cpp b/Engine/Source/Runtime/Engine/Private/Atmosphere/Atmosphere.cpp
--- a/Engine/Source/Runtime/Engine/Private/Atmosphere/Atmosphere.cpp
+++ b/Engine/Source/Runtime/Engine/Private/Atmosphere/Atmosphere.cpp
@@ -14,11 +14,11 @@
 AAtmosphericFog::AAtmosphericFog(const FObjectInitializer& ObjectInitializer)
 	: Super(ObjectInitializer)
 {
-	AtmosphericFogComponent = ObjectInitializer.CreateDefaultSubobject<UAtmosphericFogComponent>(this, TEXT("AtmosphericFogComponent0"));
+	AtmosphericFogComponent = CreateDefaultSubobject<UAtmosphericFogComponent>(TEXT("AtmosphericFogComponent0"));
 	RootComponent = AtmosphericFogComponent;
 
 #if WITH_EDITORONLY_DATA
-	ArrowComponent = ObjectInitializer.CreateEditorOnlyDefaultSubobject<UArrowComponent>(this, TEXT("ArrowComponent0"));
+	ArrowComponent = CreateEditorOnlyDefaultSubobject<UArrowComponent>(TEXT("ArrowComponent0"));
 
 	if (!IsRunningCommandlet())
 	{
diff --git a/Engine/Source/Runtime/Engine/Private/Brush.cpp b/Engine/Source/Runtime/Engine/Private/Brush.cpp
--- a/Engine/Source/Runtime/Engine/Private/Brush.cpp
+++ b/Engine/Source/Runtime/Engine/Private/Brush.cpp
@@ -23,7 +23,7 @@ TArray< TWeakObjectPtr< ULevel > > ABrush::LevelsToRebuild;
 ABrush::ABrush(const FObjectInitializer& ObjectInitializer)
 	: Super(ObjectInitializer)
 {
-	BrushComponent = ObjectInitializer.CreateDefaultSubobject<UBrushComponent>(this, TEXT("BrushComponent0"));
+	BrushComponent = CreateDefaultSubobject<UBrushComponent>(TEXT("BrushComponent0"));
 	BrushComponent->Mobility = EComponentMobility::Static;
 	BrushComponent->bGenerateOverlapEvents = false;
 	BrushComponent->bCanEverAffectNavigation = false;
diff --git a/Engine/Source/Runtime/Engine/Private/Camera/CameraActor.cpp b/Engine/Source/Runtime/Engine/Private/Camera/CameraActor.cpp
--- a/Engine/Source/Runtime/Engine/Private/Camera/CameraActor.cpp
+++ b/Engine/Source/Runtime/Engine/Private/Camera/CameraActor.cpp
@@ -17,7 +17,7 @@ ACameraActor::ACameraActor(const FObjectInitializer& ObjectInitializer)
 	: Super(ObjectInitializer)
 {
 	// Setup camera defaults
-	CameraComponent = ObjectInitializer.CreateDefaultSubobject<UCameraComponent>(this, TEXT("CameraComponent"));
+	CameraComponent = CreateDefaultSubobject<UCameraComponent>(TEXT("CameraComponent"));
 	CameraComponent->FieldOfView = 90.0f;
 	CameraComponent->bConstrainAspectRatio = true;
 	CameraComponent->AspectRatio = 1.777778f;
diff --git a/Engine/Source/Runtime/Engine/Private/Camera/CameraAnimInst.cpp b/Engine/Source/Runtime/Engine/Private/Camera/CameraAnimInst.cpp
--- a/Engine/Source/Runtime/Engine/Private/Camera/CameraAnimInst.cpp
+++ b/Engine/Source/Runtime/Engine/Private/Camera/CameraAnimInst.cpp
@@ -23,7 +23,7 @@ UCameraAnimInst::UCameraAnimInst(const FObjectInitializer& ObjectInitializer)
 	TransientScaleModifier = 1.0f;
 	PlaySpace = ECameraAnimPlaySpace::CameraLocal;
 
-	InterpGroupInst = ObjectInitializer.CreateDefaultSubobject<UInterpGroupInst>(this, TEXT("InterpGroupInst0"));
+	InterpGroupInst = CreateDefaultSubobject<UInterpGroupInst>(TEXT("InterpGroupInst0"));
 }
 
 void UCameraAnimInst::AdvanceAnim(float DeltaTime, bool bJump)
diff --git a/Engine/Source/Runtime/Engine/Private/Character.cpp b/Engine/Source/Runtime/Engine/Private/Character.cpp
--- a/Engine/Source/Runtime/Engine/Private/Character.cpp
+++ b/Engine/Source/Runtime/Engine/Private/Character.cpp
@@ -46,7 +46,7 @@ ACharacter::ACharacter(const FObjectInitializer& ObjectInitializer)
 	bUseControllerRotationRoll = false;
 	bUseControllerRotationYaw = true;
 
-	CapsuleComponent = ObjectInitializer.CreateDefaultSubobject<UCapsuleComponent>(this, ACharacter::CapsuleComponentName);
+	CapsuleComponent = CreateDefaultSubobject<UCapsuleComponent>(ACharacter::CapsuleComponentName);
 	CapsuleComponent->InitCapsuleSize(34.0f, 88.0f);
 
 	static FName CollisionProfileName(TEXT("Pawn"));
@@ -62,7 +62,7 @@ ACharacter::ACharacter(const FObjectInitializer& ObjectInitializer)
 	JumpMaxHoldTime = 0.0f;
 
 #if WITH_EDITORONLY_DATA
-	ArrowComponent = ObjectInitializer.CreateEditorOnlyDefaultSubobject<UArrowComponent>(this, TEXT("Arrow"));
+	ArrowComponent = CreateEditorOnlyDefaultSubobject<UArrowComponent>(TEXT("Arrow"));
 	if (ArrowComponent)
 	{
 		ArrowComponent->ArrowColor = FColor(150, 200, 255);
@@ -74,14 +74,14 @@ ACharacter::ACharacter(const FObjectInitializer& ObjectInitializer)
 	}
 #endif // WITH_EDITORONLY_DATA
 
-	CharacterMovement = ObjectInitializer.CreateDefaultSubobject<UCharacterMovementComponent>(this, ACharacter::CharacterMovementComponentName);
+	CharacterMovement = CreateDefaultSubobject<UCharacterMovementComponent>(ACharacter::CharacterMovementComponentName);
 	if (CharacterMovement)
 	{
 		CharacterMovement->UpdatedComponent = CapsuleComponent;
 		CrouchedEyeHeight = CharacterMovement->CrouchedHalfHeight * 0.80f;
 	}
 
-	Mesh = ObjectInitializer.CreateOptionalDefaultSubobject<USkeletalMeshComponent>(this, ACharacter::MeshComponentName);
+	Mesh = CreateOptionalDefaultSubobject<USkeletalMeshComponent>(ACharacter::MeshComponentName);
 	if (Mesh)
 	{
 		Mesh->AlwaysLoadOnClient = true;
diff --git a/Engine/Source/Runtime/Engine/Private/Controller.cpp b/Engine/Source/Runtime/Engine/Private/Controller.cpp
--- a/Engine/Source/Runtime/Engine/Private/Controller.cpp
+++ b/Engine/Source/Runtime/Engine/Private/Controller.cpp
@@ -34,7 +34,7 @@ AController::AController(const FObjectInitializer& ObjectInitializer)
 #endif // WITH_EDITORONLY_DATA
 	bOnlyRelevantToOwner = true;
 
-	TransformComponent = ObjectInitializer.CreateDefaultSubobject<USceneComponent>(this, TEXT("TransformComponent0"));
+	TransformComponent = CreateDefaultSubobject<USceneComponent>(TEXT("TransformComponent0"));
 	RootComponent = TransformComponent;
 
 	bCanBeDamaged = false;
diff --git a/Engine/Source/Runtime/Engine/Private/DecalActor.cpp b/Engine/Source/Runtime/Engine/Private/DecalActor.cpp
--- a/Engine/Source/Runtime/Engine/Private/DecalActor.cpp
+++ b/Engine/Source/Runtime/Engine/Private/DecalActor.cpp
@@ -20,7 +20,7 @@ namespace DecalEditorConstants
 ADecalActor::ADecalActor(const FObjectInitializer& ObjectInitializer)
 : Super(ObjectInitializer)
 {
-	Decal = ObjectInitializer.CreateDefaultSubobject<UDecalComponent>(this, TEXT("NewDecalComponent"));
+	Decal = CreateDefaultSubobject<UDecalComponent>(TEXT("NewDecalComponent"));
 	Decal->RelativeScale3D = FVector(128.0f, 256.0f, 256.0f);
 
 	Decal->RelativeRotation = FRotator(-90, 0, 0);
@@ -28,7 +28,7 @@ ADecalActor::ADecalActor(const FObjectInitializer& ObjectInitializer)
 	RootComponent = Decal;
 
 #if WITH_EDITORONLY_DATA
-	BoxComponent = ObjectInitializer.CreateEditorOnlyDefaultSubobject<UBoxComponent>(this, TEXT("DrawBox0"));
+	BoxComponent = CreateEditorOnlyDefaultSubobject<UBoxComponent>(TEXT("DrawBox0"));
 	if (BoxComponent != nullptr)
 	{
 		BoxComponent->SetCollisionProfileName(UCollisionProfile::NoCollision_ProfileName);
@@ -40,8 +40,8 @@ ADecalActor::ADecalActor(const FObjectInitializer& ObjectInitializer)
 		BoxComponent->AttachParent = Decal;
 	}
 
-	ArrowComponent = ObjectInitializer.CreateEditorOnlyDefaultSubobject<UArrowComponent>(this, TEXT("ArrowComponent0"));
-	SpriteComponent = ObjectInitializer.CreateEditorOnlyDefaultSubobject<UBillboardComponent>(this, TEXT("Sprite"));
+	ArrowComponent = CreateEditorOnlyDefaultSubobject<UArrowComponent>(TEXT("ArrowComponent0"));
+	SpriteComponent = CreateEditorOnlyDefaultSubobject<UBillboardComponent>(TEXT("Sprite"));
 
 	if (!IsRunningCommandlet())
 	{
diff --git a/Engine/Source/Runtime/Engine/Private/DefaultPawn.cpp b/Engine/Source/Runtime/Engine/Private/DefaultPawn.cpp
--- a/Engine/Source/Runtime/Engine/Private/DefaultPawn.cpp
+++ b/Engine/Source/Runtime/Engine/Private/DefaultPawn.cpp
@@ -24,7 +24,7 @@ ADefaultPawn::ADefaultPawn(const FObjectInitializer& ObjectInitializer)
 	BaseEyeHeight = 0.0f;
 	bCollideWhenPlacing = false;
 
-	CollisionComponent = ObjectInitializer.CreateDefaultSubobject<USphereComponent>(this, ADefaultPawn::CollisionComponentName);
+	CollisionComponent = CreateDefaultSubobject<USphereComponent>(ADefaultPawn::CollisionComponentName);
 	CollisionComponent->InitSphereRadius(35.0f);
 
 	static FName CollisionProfileName(TEXT("Pawn"));
@@ -36,7 +36,7 @@ ADefaultPawn::ADefaultPawn(const FObjectInitializer& ObjectInitializer)
 
 	RootComponent = CollisionComponent;
 
-	MovementComponent = ObjectInitializer.CreateDefaultSubobject<UFloatingPawnMovement>(this, ADefaultPawn::MovementComponentName);
+	MovementComponent = CreateDefaultSubobject<UFloatingPawnMovement>(ADefaultPawn::MovementComponentName);
 	MovementComponent->UpdatedComponent = CollisionComponent;
 
 	// Structure to hold one-time initialization
@@ -49,7 +49,7 @@ ADefaultPawn::ADefaultPawn(const FObjectInitializer& ObjectInitializer)
 
 	static FConstructorStatics ConstructorStatics;
 
-	MeshComponent = ObjectInitializer.CreateOptionalDefaultSubobject<UStaticMeshComponent>(this, ADefaultPawn::MeshComponentName);
+	MeshComponent = CreateOptionalDefaultSubobject<UStaticMeshComponent>(ADefaultPawn::MeshComponentName);
 	if (MeshComponent)
 	{
 		MeshComponent->SetStaticMesh(ConstructorStatics.SphereMesh.Object);
diff --git a/Engine/Source/Runtime/Engine/Private/DestructibleActor.cpp b/Engine/Source/Runtime/Engine/Private/DestructibleActor.cpp
--- a/Engine/Source/Runtime/Engine/Private/DestructibleActor.cpp
+++ b/Engine/Source/Runtime/Engine/Private/DestructibleActor.cpp
@@ -13,7 +13,7 @@
 ADestructibleActor::ADestructibleActor(const FObjectInitializer& ObjectInitializer)
 	: Super(ObjectInitializer)
 {
-	DestructibleComponent = ObjectInitializer.CreateDefaultSubobject<UDestructibleComponent>(this, TEXT("DestructibleComponent0"));
+	DestructibleComponent = CreateDefaultSubobject<UDestructibleComponent>(TEXT("DestructibleComponent0"));
 	DestructibleComponent->bCanEverAffectNavigation = bAffectNavigation;
 	RootComponent = DestructibleComponent;
 }
diff --git a/Engine/Source/Runtime/Engine/Private/DocumentationActor.cpp b/Engine/Source/Runtime/Engine/Private/DocumentationActor.cpp
--- a/Engine/Source/Runtime/Engine/Private/DocumentationActor.cpp
+++ b/Engine/Source/Runtime/Engine/Private/DocumentationActor.cpp
@@ -10,12 +10,12 @@
 ADocumentationActor::ADocumentationActor(const FObjectInitializer& ObjectInitializer)
 	: Super(ObjectInitializer)
 {
-	USceneComponent* SceneComponent = ObjectInitializer.CreateDefaultSubobject<USceneComponent>(this, TEXT("SceneComp"));
+	USceneComponent* SceneComponent = CreateDefaultSubobject<USceneComponent>(TEXT("SceneComp"));
 	RootComponent = SceneComponent;	
 
 #if WITH_EDITORONLY_DATA
 	// Create a Material billboard to represent our actor
-	Billboard = ObjectInitializer.CreateDefaultSubobject<UMaterialBillboardComponent>(this, TEXT("BillboardComponent"));
+	Billboard = CreateDefaultSubobject<UMaterialBillboardComponent>(TEXT("BillboardComponent"));
 	if (!IsRunningCommandlet() && (Billboard != NULL))
 	{
 		static ConstructorHelpers::FObjectFinder<UMaterial> MaterialAsset(TEXT("/Engine/EditorMaterials/HelpActorMaterial"));
diff --git a/Engine/Source/Runtime/Engine/Private/HeightFogComponent.cpp b/Engine/Source/Runtime/Engine/Private/HeightFogComponent.cpp
--- a/Engine/Source/Runtime/Engine/Private/HeightFogComponent.cpp
+++ b/Engine/Source/Runtime/Engine/Private/HeightFogComponent.cpp
@@ -149,7 +149,7 @@ void UExponentialHeightFogComponent::SetStartDistance(float Value)
 AExponentialHeightFog::AExponentialHeightFog(const FObjectInitializer& ObjectInitializer)
 	: Super(ObjectInitializer)
 {
-	Component = ObjectInitializer.CreateDefaultSubobject<UExponentialHeightFogComponent>(this, TEXT("HeightFogComponent0"));
+	Component = CreateDefaultSubobject<UExponentialHeightFogComponent>(TEXT("HeightFogComponent0"));
 	RootComponent = Component;
 
 	bHidden = false;
diff --git a/Engine/Source/Runtime/Engine/Private/Info.cpp b/Engine/Source/Runtime/Engine/Private/Info.cpp
--- a/Engine/Source/Runtime/Engine/Private/Info.cpp
+++ b/Engine/Source/Runtime/Engine/Private/Info.cpp
@@ -6,7 +6,7 @@ AInfo::AInfo(const FObjectInitializer& ObjectInitializer)
 	: Super(ObjectInitializer)
 {
 #if WITH_EDITORONLY_DATA
-	SpriteComponent = ObjectInitializer.CreateEditorOnlyDefaultSubobject<UBillboardComponent>(this, TEXT("Sprite"));
+	SpriteComponent = CreateEditorOnlyDefaultSubobject<UBillboardComponent>(TEXT("Sprite"));
 	if (!IsRunningCommandlet() && (SpriteComponent != nullptr))
 	{
 		// Structure to hold one-time initialization
diff --git a/Engine/Source/Runtime/Engine/Private/InstancedFoliage.cpp b/Engine/Source/Runtime/Engine/Private/InstancedFoliage.cpp
--- a/Engine/Source/Runtime/Engine/Private/InstancedFoliage.cpp
+++ b/Engine/Source/Runtime/Engine/Private/InstancedFoliage.cpp
@@ -678,7 +678,7 @@ void FFoliageMeshInfo::SelectInstances(AInstancedFoliageActor* InIFA, bool bSele
 AInstancedFoliageActor::AInstancedFoliageActor(const FObjectInitializer& ObjectInitializer)
 	: Super(ObjectInitializer)
 {
-	USceneComponent* SceneComponent = ObjectInitializer.CreateDefaultSubobject<USceneComponent>(this, TEXT("RootComponent0"));
+	USceneComponent* SceneComponent = CreateDefaultSubobject<USceneComponent>(TEXT("RootComponent0"));
 	RootComponent = SceneComponent;
 	RootComponent->Mobility = EComponentMobility::Static;
 	
diff --git a/Engine/Source/Runtime/Engine/Private/InteractiveFoliageActor.cpp b/Engine/Source/Runtime/Engine/Private/InteractiveFoliageActor.cpp
--- a/Engine/Source/Runtime/Engine/Private/InteractiveFoliageActor.cpp
+++ b/Engine/Source/Runtime/Engine/Private/InteractiveFoliageActor.cpp
@@ -15,7 +15,7 @@ AInteractiveFoliageActor::AInteractiveFoliageActor(const FObjectInitializer& Obj
 	FoliageMeshComponent->SetCollisionProfileName(UCollisionProfile::NoCollision_ProfileName);
 	FoliageMeshComponent->Mobility = EComponentMobility::Static;
 
-	CapsuleComponent = ObjectInitializer.CreateDefaultSubobject<UCapsuleComponent>(this, TEXT("CollisionCylinder"));
+	CapsuleComponent = CreateDefaultSubobject<UCapsuleComponent>(TEXT("CollisionCylinder"));
 	CapsuleComponent->InitCapsuleSize(60.0f, 200.0f);
 	static FName CollisionProfileName(TEXT("OverlapAllDynamic"));
 	CapsuleComponent->SetCollisionProfileName(CollisionProfileName);
diff --git a/Engine/Source/Runtime/Engine/Private/Interpolation.cpp b/Engine/Source/Runtime/Engine/Private/Interpolation.cpp
--- a/Engine/Source/Runtime/Engine/Private/Interpolation.cpp
+++ b/Engine/Source/Runtime/Engine/Private/Interpolation.cpp
@@ -323,11 +323,11 @@ bool AMatineeActor::IgnoreActorSelection()
 AMatineeActor::AMatineeActor(const FObjectInitializer& ObjectInitializer)
 	: Super(ObjectInitializer)
 {
-	USceneComponent* SceneComponent = ObjectInitializer.CreateDefaultSubobject<USceneComponent>(this, TEXT("SceneComp"));
+	USceneComponent* SceneComponent = CreateDefaultSubobject<USceneComponent>(TEXT("SceneComp"));
 	RootComponent = SceneComponent;
 
 #if WITH_EDITORONLY_DATA
-	SpriteComponent = ObjectInitializer.CreateEditorOnlyDefaultSubobject<UBillboardComponent>(this, TEXT("Sprite"));
+	SpriteComponent = CreateEditorOnlyDefaultSubobject<UBillboardComponent>(TEXT("Sprite"));
 	if (!IsRunningCommandlet() && (SpriteComponent != nullptr))
 	{
 		// Structure to hold one-time initialization
@@ -8690,11 +8690,11 @@ static void GetListOfStaticActors(FString& OutString, const TArray<AActor*>& Act
 AMaterialInstanceActor::AMaterialInstanceActor(const FObjectInitializer& ObjectInitializer)
 	: Super(ObjectInitializer)
 {
-	USceneComponent* SceneComponent = ObjectInitializer.CreateDefaultSubobject<USceneComponent>(this, TEXT("SceneComp"));
+	USceneComponent* SceneComponent = CreateDefaultSubobject<USceneComponent>(TEXT("SceneComp"));
 	RootComponent = SceneComponent;
 
 #if WITH_EDITORONLY_DATA
-	SpriteComponent = ObjectInitializer.CreateEditorOnlyDefaultSubobject<UBillboardComponent>(this, TEXT("Sprite"));
+	SpriteComponent = CreateEditorOnlyDefaultSubobject<UBillboardComponent>(TEXT("Sprite"));
 	if (!IsRunningCommandlet() && (SpriteComponent != NULL))
 	{
 		// Structure to hold one-time initialization
diff --git a/Engine/Source/Runtime/Engine/Private/LevelBounds.cpp b/Engine/Source/Runtime/Engine/Private/LevelBounds.cpp
--- a/Engine/Source/Runtime/Engine/Private/LevelBounds.cpp
+++ b/Engine/Source/Runtime/Engine/Private/LevelBounds.cpp
@@ -10,7 +10,7 @@ static const FVector DefaultLevelSize = FVector(1000.f);
 ALevelBounds::ALevelBounds(const FObjectInitializer& ObjectInitializer)
 	: Super(ObjectInitializer)
 {
-	UBoxComponent* BoxComponent = ObjectInitializer.CreateDefaultSubobject<UBoxComponent>(this, TEXT("BoxComponent0"));
+	UBoxComponent* BoxComponent = CreateDefaultSubobject<UBoxComponent>(TEXT("BoxComponent0"));
 	RootComponent = BoxComponent;
 	RootComponent->Mobility = EComponentMobility::Static;
 	RootComponent->RelativeScale3D = DefaultLevelSize;
diff --git a/Engine/Source/Runtime/Engine/Private/Light.cpp b/Engine/Source/Runtime/Engine/Private/Light.cpp
--- a/Engine/Source/Runtime/Engine/Private/Light.cpp
+++ b/Engine/Source/Runtime/Engine/Private/Light.cpp
@@ -15,7 +15,7 @@
 ALight::ALight(const FObjectInitializer& ObjectInitializer)
 	: Super(ObjectInitializer)
 {
-	LightComponent = ObjectInitializer.CreateAbstractDefaultSubobject<ULightComponent>(this, TEXT("LightComponent0"));
+	LightComponent = CreateAbstractDefaultSubobject<ULightComponent>(TEXT("LightComponent0"));
 
 	bHidden = true;
 	bCollideWhenPlacing = true;
@@ -230,7 +230,7 @@ ADirectionalLight::ADirectionalLight(const FObjectInitializer& ObjectInitializer
 	RootComponent = DirectionalLightComponent;
 
 #if WITH_EDITORONLY_DATA
-	ArrowComponent = ObjectInitializer.CreateEditorOnlyDefaultSubobject<UArrowComponent>(this, TEXT("ArrowComponent0"));
+	ArrowComponent = CreateEditorOnlyDefaultSubobject<UArrowComponent>(TEXT("ArrowComponent0"));
 	if (ArrowComponent)
 	{
 		ArrowComponent->ArrowColor = FColor(150, 200, 255);
diff --git a/Engine/Source/Runtime/Engine/Private/NavigationObjectBase.cpp b/Engine/Source/Runtime/Engine/Private/NavigationObjectBase.cpp
--- a/Engine/Source/Runtime/Engine/Private/NavigationObjectBase.cpp
+++ b/Engine/Source/Runtime/Engine/Private/NavigationObjectBase.cpp
@@ -10,7 +10,7 @@ DEFINE_LOG_CATEGORY_STATIC(LogNavigationPoint, Log, All);
 ANavigationObjectBase::ANavigationObjectBase(const FObjectInitializer& ObjectInitializer)
 	: Super(ObjectInitializer)
 {
-	CapsuleComponent = ObjectInitializer.CreateDefaultSubobject<UCapsuleComponent>(this, TEXT("CollisionCapsule"));
+	CapsuleComponent = CreateDefaultSubobject<UCapsuleComponent>(TEXT("CollisionCapsule"));
 	CapsuleComponent->ShapeColor = FColor(255, 138, 5, 255);
 	CapsuleComponent->bDrawOnlyIfSelected = true;
 	CapsuleComponent->InitCapsuleSize(50.0f, 50.0f);
@@ -21,9 +21,9 @@ ANavigationObjectBase::ANavigationObjectBase(const FObjectInitializer& ObjectIni
 	RootComponent = CapsuleComponent;
 	bCollideWhenPlacing = true;
 
-	GoodSprite = ObjectInitializer.CreateEditorOnlyDefaultSubobject<UBillboardComponent>(this, TEXT("Sprite"));
+	GoodSprite = CreateEditorOnlyDefaultSubobject<UBillboardComponent>(TEXT("Sprite"));
 
-	BadSprite = ObjectInitializer.CreateEditorOnlyDefaultSubobject<UBillboardComponent>(this, TEXT("Sprite2"));
+	BadSprite = CreateEditorOnlyDefaultSubobject<UBillboardComponent>(TEXT("Sprite2"));
 
 #if WITH_EDITORONLY_DATA
 	if (!IsRunningCommandlet())
diff --git a/Engine/Source/Runtime/Engine/Private/NiagaraActor.cpp b/Engine/Source/Runtime/Engine/Private/NiagaraActor.cpp
--- a/Engine/Source/Runtime/Engine/Private/NiagaraActor.cpp
+++ b/Engine/Source/Runtime/Engine/Private/NiagaraActor.cpp
@@ -8,13 +8,13 @@
 ANiagaraActor::ANiagaraActor(const FObjectInitializer& ObjectInitializer)
 	: Super(ObjectInitializer)
 {
-	NiagaraComponent = ObjectInitializer.CreateDefaultSubobject<UNiagaraComponent>(this, TEXT("NiagaraComponent0"));
+	NiagaraComponent = CreateDefaultSubobject<UNiagaraComponent>(TEXT("NiagaraComponent0"));
 
 	RootComponent = NiagaraComponent;
 
 #if WITH_EDITORONLY_DATA
-	SpriteComponent = ObjectInitializer.CreateEditorOnlyDefaultSubobject<UBillboardComponent>(this, TEXT("Sprite"));
-	ArrowComponent = ObjectInitializer.CreateEditorOnlyDefaultSubobject<UArrowComponent>(this, TEXT("ArrowComponent0"));
+	SpriteComponent = CreateEditorOnlyDefaultSubobject<UBillboardComponent>(TEXT("Sprite"));
+	ArrowComponent = CreateEditorOnlyDefaultSubobject<UArrowComponent>(TEXT("ArrowComponent0"));
 
 	if (!IsRunningCommandlet())
 	{
diff --git a/Engine/Source/Runtime/Engine/Private/Note.cpp b/Engine/Source/Runtime/Engine/Private/Note.cpp
--- a/Engine/Source/Runtime/Engine/Private/Note.cpp
+++ b/Engine/Source/Runtime/Engine/Private/Note.cpp
@@ -6,13 +6,13 @@
 ANote::ANote(const FObjectInitializer& ObjectInitializer)
 	: Super(ObjectInitializer)
 {
-	USceneComponent* SceneComponent = ObjectInitializer.CreateDefaultSubobject<USceneComponent>(this, TEXT("SceneComp"));
+	USceneComponent* SceneComponent = CreateDefaultSubobject<USceneComponent>(TEXT("SceneComp"));
 	RootComponent = SceneComponent;
 	RootComponent->Mobility = EComponentMobility::Static;
 
 #if WITH_EDITORONLY_DATA
-	ArrowComponent = ObjectInitializer.CreateEditorOnlyDefaultSubobject<UArrowComponent>(this, TEXT("Arrow"));
-	SpriteComponent = ObjectInitializer.CreateEditorOnlyDefaultSubobject<UBillboardComponent>(this, TEXT("Sprite"));
+	ArrowComponent = CreateEditorOnlyDefaultSubobject<UArrowComponent>(TEXT("Arrow"));
+	SpriteComponent = CreateEditorOnlyDefaultSubobject<UBillboardComponent>(TEXT("Sprite"));
 
 	if (!IsRunningCommandlet())
 	{
diff --git a/Engine/Source/Runtime/Engine/Private/Particles/ParticleComponents.cpp b/Engine/Source/Runtime/Engine/Private/Particles/ParticleComponents.cpp
--- a/Engine/Source/Runtime/Engine/Private/Particles/ParticleComponents.cpp
+++ b/Engine/Source/Runtime/Engine/Private/Particles/ParticleComponents.cpp
@@ -180,13 +180,13 @@ void Particle_ModifyVectorDistribution(UDistributionVector* pkDistribution, FVec
 AEmitter::AEmitter(const FObjectInitializer& ObjectInitializer)
 	: Super(ObjectInitializer)
 {
-	ParticleSystemComponent = ObjectInitializer.CreateDefaultSubobject<UParticleSystemComponent>(this, TEXT("ParticleSystemComponent0"));
+	ParticleSystemComponent = CreateDefaultSubobject<UParticleSystemComponent>(TEXT("ParticleSystemComponent0"));
 	ParticleSystemComponent->SecondsBeforeInactive = 1;
 	RootComponent = ParticleSystemComponent;
 
 #if WITH_EDITORONLY_DATA
-	SpriteComponent = ObjectInitializer.CreateEditorOnlyDefaultSubobject<UBillboardComponent>(this, TEXT("Sprite"));
-	ArrowComponent = ObjectInitializer.CreateEditorOnlyDefaultSubobject<UArrowComponent>(this, TEXT("ArrowComponent0"));
+	SpriteComponent = CreateEditorOnlyDefaultSubobject<UBillboardComponent>(TEXT("Sprite"));
+	ArrowComponent = CreateEditorOnlyDefaultSubobject<UArrowComponent>(TEXT("ArrowComponent0"));
 
 	if (!IsRunningCommandlet())
 	{
diff --git a/Engine/Source/Runtime/Engine/Private/PhysicsEngine/PhysicsConstraintActor.cpp b/Engine/Source/Runtime/Engine/Private/PhysicsEngine/PhysicsConstraintActor.cpp
--- a/Engine/Source/Runtime/Engine/Private/PhysicsEngine/PhysicsConstraintActor.cpp
+++ b/Engine/Source/Runtime/Engine/Private/PhysicsEngine/PhysicsConstraintActor.cpp
@@ -20,7 +20,7 @@ APhysicsConstraintActor::APhysicsConstraintActor(const FObjectInitializer& Objec
 	};
 	static FConstructorStatics ConstructorStatics;
 
-	ConstraintComp = ObjectInitializer.CreateDefaultSubobject<UPhysicsConstraintComponent>(this, TEXT("MyConstraintComp"));
+	ConstraintComp = CreateDefaultSubobject<UPhysicsConstraintComponent>(TEXT("MyConstraintComp"));
 	RootComponent = ConstraintComp;
 	bHidden = true;
 }
diff --git a/Engine/Source/Runtime/Engine/Private/PhysicsEngine/PhysicsThruster.cpp b/Engine/Source/Runtime/Engine/Private/PhysicsEngine/PhysicsThruster.cpp
--- a/Engine/Source/Runtime/Engine/Private/PhysicsEngine/PhysicsThruster.cpp
+++ b/Engine/Source/Runtime/Engine/Private/PhysicsEngine/PhysicsThruster.cpp
@@ -36,12 +36,12 @@ void UPhysicsThrusterComponent::TickComponent(float DeltaTime, enum ELevelTick T
 APhysicsThruster::APhysicsThruster(const FObjectInitializer& ObjectInitializer)
 	: Super(ObjectInitializer)
 {
-	ThrusterComponent = ObjectInitializer.CreateDefaultSubobject<UPhysicsThrusterComponent>(this, TEXT("Thruster0"));
+	ThrusterComponent = CreateDefaultSubobject<UPhysicsThrusterComponent>(TEXT("Thruster0"));
 	RootComponent = ThrusterComponent;
 
 #if WITH_EDITORONLY_DATA
-	ArrowComponent = ObjectInitializer.CreateEditorOnlyDefaultSubobject<UArrowComponent>(this, TEXT("ArrowComponent0"));
-	SpriteComponent = ObjectInitializer.CreateEditorOnlyDefaultSubobject<UBillboardComponent>(this, TEXT("Sprite"));
+	ArrowComponent = CreateEditorOnlyDefaultSubobject<UArrowComponent>(TEXT("ArrowComponent0"));
+	SpriteComponent = CreateEditorOnlyDefaultSubobject<UBillboardComponent>(TEXT("Sprite"));
 
 	if (!IsRunningCommandlet())
 	{
diff --git a/Engine/Source/Runtime/Engine/Private/PhysicsEngine/RadialForceComponent.cpp b/Engine/Source/Runtime/Engine/Private/PhysicsEngine/RadialForceComponent.cpp
--- a/Engine/Source/Runtime/Engine/Private/PhysicsEngine/RadialForceComponent.cpp
+++ b/Engine/Source/Runtime/Engine/Private/PhysicsEngine/RadialForceComponent.cpp
@@ -180,10 +180,10 @@ void URadialForceComponent::UpdateCollisionObjectQueryParams()
 ARadialForceActor::ARadialForceActor(const FObjectInitializer& ObjectInitializer)
 	: Super(ObjectInitializer)
 {
-	ForceComponent = ObjectInitializer.CreateDefaultSubobject<URadialForceComponent>(this, TEXT("ForceComponent0"));
+	ForceComponent = CreateDefaultSubobject<URadialForceComponent>(TEXT("ForceComponent0"));
 
 #if WITH_EDITOR
-	SpriteComponent = ObjectInitializer.CreateEditorOnlyDefaultSubobject<UBillboardComponent>(this, TEXT("Sprite"));
+	SpriteComponent = CreateEditorOnlyDefaultSubobject<UBillboardComponent>(TEXT("Sprite"));
 	if (SpriteComponent)
 	{
 		// Structure to hold one-time initialization
diff --git a/Engine/Source/Runtime/Engine/Private/PlayerCameraManager.cpp b/Engine/Source/Runtime/Engine/Private/PlayerCameraManager.cpp
--- a/Engine/Source/Runtime/Engine/Private/PlayerCameraManager.cpp
+++ b/Engine/Source/Runtime/Engine/Private/PlayerCameraManager.cpp
@@ -45,7 +45,7 @@ APlayerCameraManager::APlayerCameraManager(const FObjectInitializer& ObjectIniti
 	bFollowHmdOrientation = false;
 
 	// create dummy transform component
-	TransformComponent = ObjectInitializer.CreateDefaultSubobject<USceneComponent>(this, TEXT("TransformComponent0"));
+	TransformComponent = CreateDefaultSubobject<USceneComponent>(TEXT("TransformComponent0"));
 	RootComponent = TransformComponent;
 }
 
diff --git a/Engine/Source/Runtime/Engine/Private/PlayerStart.cpp b/Engine/Source/Runtime/Engine/Private/PlayerStart.cpp
--- a/Engine/Source/Runtime/Engine/Private/PlayerStart.cpp
+++ b/Engine/Source/Runtime/Engine/Private/PlayerStart.cpp
@@ -11,7 +11,7 @@ APlayerStart::APlayerStart(const FObjectInitializer& ObjectInitializer)
 	GetCapsuleComponent()->InitCapsuleSize(40.0f, 92.0f);
 
 #if WITH_EDITORONLY_DATA
-	ArrowComponent = ObjectInitializer.CreateEditorOnlyDefaultSubobject<UArrowComponent>(this, TEXT("Arrow"));
+	ArrowComponent = CreateEditorOnlyDefaultSubobject<UArrowComponent>(TEXT("Arrow"));
 
 	if (!IsRunningCommandlet())
 	{
diff --git a/Engine/Source/Runtime/Engine/Private/ReflectionCaptureComponent.cpp b/Engine/Source/Runtime/Engine/Private/ReflectionCaptureComponent.cpp
--- a/Engine/Source/Runtime/Engine/Private/ReflectionCaptureComponent.cpp
+++ b/Engine/Source/Runtime/Engine/Private/ReflectionCaptureComponent.cpp
@@ -63,10 +63,10 @@ void UWorld::UpdateAllReflectionCaptures()
 AReflectionCapture::AReflectionCapture(const FObjectInitializer& ObjectInitializer)
 	: Super(ObjectInitializer)
 {
-	CaptureComponent = ObjectInitializer.CreateAbstractDefaultSubobject<UReflectionCaptureComponent>(this, TEXT("NewReflectionComponent"));
+	CaptureComponent = CreateAbstractDefaultSubobject<UReflectionCaptureComponent>(TEXT("NewReflectionComponent"));
 
 #if WITH_EDITORONLY_DATA
-	SpriteComponent = ObjectInitializer.CreateEditorOnlyDefaultSubobject<UBillboardComponent>(this, TEXT("Sprite"));
+	SpriteComponent = CreateEditorOnlyDefaultSubobject<UBillboardComponent>(TEXT("Sprite"));
 	if (!IsRunningCommandlet() && (SpriteComponent != nullptr))
 	{
 		// Structure to hold one-time initialization
@@ -113,14 +113,14 @@ ASphereReflectionCapture::ASphereReflectionCapture(const FObjectInitializer& Obj
 	}
 #endif	//WITH_EDITORONLY_DATA
 
-	UDrawSphereComponent* DrawInfluenceRadius = ObjectInitializer.CreateDefaultSubobject<UDrawSphereComponent>(this, TEXT("DrawRadius0"));
+	UDrawSphereComponent* DrawInfluenceRadius = CreateDefaultSubobject<UDrawSphereComponent>(TEXT("DrawRadius0"));
 	DrawInfluenceRadius->AttachParent = GetCaptureComponent();
 	DrawInfluenceRadius->bDrawOnlyIfSelected = true;
 	DrawInfluenceRadius->bUseEditorCompositing = true;
 	DrawInfluenceRadius->SetCollisionProfileName(UCollisionProfile::NoCollision_ProfileName);
 	SphereComponent->PreviewInfluenceRadius = DrawInfluenceRadius;
 
-	DrawCaptureRadius = ObjectInitializer.CreateDefaultSubobject<UDrawSphereComponent>(this, TEXT("DrawRadius1"));
+	DrawCaptureRadius = CreateDefaultSubobject<UDrawSphereComponent>(TEXT("DrawRadius1"));
 	DrawCaptureRadius->AttachParent = GetCaptureComponent();
 	DrawCaptureRadius->bDrawOnlyIfSelected = true;
 	DrawCaptureRadius->bUseEditorCompositing = true;
@@ -162,7 +162,7 @@ ABoxReflectionCapture::ABoxReflectionCapture(const FObjectInitializer& ObjectIni
 		GetSpriteComponent()->AttachParent = BoxComponent;
 	}
 #endif	//WITH_EDITORONLY_DATA
-	UBoxComponent* DrawInfluenceBox = ObjectInitializer.CreateDefaultSubobject<UBoxComponent>(this, TEXT("DrawBox0"));
+	UBoxComponent* DrawInfluenceBox = CreateDefaultSubobject<UBoxComponent>(TEXT("DrawBox0"));
 	DrawInfluenceBox->AttachParent = GetCaptureComponent();
 	DrawInfluenceBox->bDrawOnlyIfSelected = true;
 	DrawInfluenceBox->bUseEditorCompositing = true;
@@ -170,7 +170,7 @@ ABoxReflectionCapture::ABoxReflectionCapture(const FObjectInitializer& ObjectIni
 	DrawInfluenceBox->InitBoxExtent(FVector(1, 1, 1));
 	BoxComponent->PreviewInfluenceBox = DrawInfluenceBox;
 
-	UBoxComponent* DrawCaptureBox = ObjectInitializer.CreateDefaultSubobject<UBoxComponent>(this, TEXT("DrawBox1"));
+	UBoxComponent* DrawCaptureBox = CreateDefaultSubobject<UBoxComponent>(TEXT("DrawBox1"));
 	DrawCaptureBox->AttachParent = GetCaptureComponent();
 	DrawCaptureBox->bDrawOnlyIfSelected = true;
 	DrawCaptureBox->bUseEditorCompositing = true;
@@ -192,7 +192,7 @@ APlaneReflectionCapture::APlaneReflectionCapture(const FObjectInitializer& Objec
 		GetSpriteComponent()->AttachParent = PlaneComponent;
 	}
 #endif	//#if WITH_EDITORONLY_DATA
-	UDrawSphereComponent* DrawInfluenceRadius = ObjectInitializer.CreateDefaultSubobject<UDrawSphereComponent>(this, TEXT("DrawRadius0"));
+	UDrawSphereComponent* DrawInfluenceRadius = CreateDefaultSubobject<UDrawSphereComponent>(TEXT("DrawRadius0"));
 	DrawInfluenceRadius->AttachParent = GetCaptureComponent();
 	DrawInfluenceRadius->bDrawOnlyIfSelected = true;
 	DrawInfluenceRadius->bAbsoluteScale = true;
@@ -200,7 +200,7 @@ APlaneReflectionCapture::APlaneReflectionCapture(const FObjectInitializer& Objec
 	DrawInfluenceRadius->SetCollisionProfileName(UCollisionProfile::NoCollision_ProfileName);
 	PlaneComponent->PreviewInfluenceRadius = DrawInfluenceRadius;
 
-	UBoxComponent* DrawCaptureBox = ObjectInitializer.CreateDefaultSubobject<UBoxComponent>(this, TEXT("DrawBox1"));
+	UBoxComponent* DrawCaptureBox = CreateDefaultSubobject<UBoxComponent>(TEXT("DrawBox1"));
 	DrawCaptureBox->AttachParent = GetCaptureComponent();
 	DrawCaptureBox->bDrawOnlyIfSelected = true;
 	DrawCaptureBox->bUseEditorCompositing = true;
diff --git a/Engine/Source/Runtime/Engine/Private/SceneCaptureComponent.cpp b/Engine/Source/Runtime/Engine/Private/SceneCaptureComponent.cpp
--- a/Engine/Source/Runtime/Engine/Private/SceneCaptureComponent.cpp
+++ b/Engine/Source/Runtime/Engine/Private/SceneCaptureComponent.cpp
@@ -16,7 +16,7 @@
 ASceneCapture::ASceneCapture(const FObjectInitializer& ObjectInitializer)
 	: Super(ObjectInitializer)
 {
-	MeshComp = ObjectInitializer.CreateDefaultSubobject<UStaticMeshComponent>(this, TEXT("CamMesh0"));
+	MeshComp = CreateDefaultSubobject<UStaticMeshComponent>(TEXT("CamMesh0"));
 
 	MeshComp->SetCollisionProfileName(UCollisionProfile::NoCollision_ProfileName);
 
@@ -30,12 +30,12 @@ ASceneCapture::ASceneCapture(const FObjectInitializer& ObjectInitializer)
 ASceneCapture2D::ASceneCapture2D(const FObjectInitializer& ObjectInitializer)
 	: Super(ObjectInitializer)
 {
-	DrawFrustum = ObjectInitializer.CreateDefaultSubobject<UDrawFrustumComponent>(this, TEXT("DrawFrust0"));
+	DrawFrustum = CreateDefaultSubobject<UDrawFrustumComponent>(TEXT("DrawFrust0"));
 	DrawFrustum->AlwaysLoadOnClient = false;
 	DrawFrustum->AlwaysLoadOnServer = false;
 	DrawFrustum->AttachParent = GetMeshComp();
 
-	CaptureComponent2D = ObjectInitializer.CreateDefaultSubobject<USceneCaptureComponent2D>(this, TEXT("NewSceneCaptureComponent2D"));
+	CaptureComponent2D = CreateDefaultSubobject<USceneCaptureComponent2D>(TEXT("NewSceneCaptureComponent2D"));
 	CaptureComponent2D->AttachParent = GetMeshComp();
 }
 
@@ -86,12 +86,12 @@ void ASceneCapture2D::PostActorCreated()
 ASceneCaptureCube::ASceneCaptureCube(const FObjectInitializer& ObjectInitializer)
 	: Super(ObjectInitializer)
 {
-	DrawFrustum = ObjectInitializer.CreateDefaultSubobject<UDrawFrustumComponent>(this, TEXT("DrawFrust0"));
+	DrawFrustum = CreateDefaultSubobject<UDrawFrustumComponent>(TEXT("DrawFrust0"));
 	DrawFrustum->AlwaysLoadOnClient = false;
 	DrawFrustum->AlwaysLoadOnServer = false;
 	DrawFrustum->AttachParent = GetMeshComp();
 
-	CaptureComponentCube = ObjectInitializer.CreateDefaultSubobject<USceneCaptureComponentCube>(this, TEXT("NewSceneCaptureComponentCube"));
+	CaptureComponentCube = CreateDefaultSubobject<USceneCaptureComponentCube>(TEXT("NewSceneCaptureComponentCube"));
 	CaptureComponentCube->AttachParent = GetMeshComp();
 }
 
diff --git a/Engine/Source/Runtime/Engine/Private/SkeletalMesh.cpp b/Engine/Source/Runtime/Engine/Private/SkeletalMesh.cpp
--- a/Engine/Source/Runtime/Engine/Private/SkeletalMesh.cpp
+++ b/Engine/Source/Runtime/Engine/Private/SkeletalMesh.cpp
@@ -3614,7 +3614,7 @@ ASkeletalMeshActor::ASkeletalMeshActor(const FObjectInitializer& ObjectInitializ
 	: Super(ObjectInitializer)
 {
 
-	SkeletalMeshComponent = ObjectInitializer.CreateDefaultSubobject<USkeletalMeshComponent>(this, TEXT("SkeletalMeshComponent0"));
+	SkeletalMeshComponent = CreateDefaultSubobject<USkeletalMeshComponent>(TEXT("SkeletalMeshComponent0"));
 	SkeletalMeshComponent->MeshComponentUpdateFlag = EMeshComponentUpdateFlag::AlwaysTickPose;
 	// check BaseEngine.ini for profile setup
 	SkeletalMeshComponent->SetCollisionProfileName(UCollisionProfile::PhysicsActor_ProfileName);
diff --git a/Engine/Source/Runtime/Engine/Private/SkyLightComponent.cpp b/Engine/Source/Runtime/Engine/Private/SkyLightComponent.cpp
--- a/Engine/Source/Runtime/Engine/Private/SkyLightComponent.cpp
+++ b/Engine/Source/Runtime/Engine/Private/SkyLightComponent.cpp
@@ -497,7 +497,7 @@ void USkyLightComponent::RecaptureSky()
 ASkyLight::ASkyLight(const FObjectInitializer& ObjectInitializer)
 	: Super(ObjectInitializer)
 {
-	LightComponent = ObjectInitializer.CreateDefaultSubobject<USkyLightComponent>(this, TEXT("SkyLightComponent0"));
+	LightComponent = CreateDefaultSubobject<USkyLightComponent>(TEXT("SkyLightComponent0"));
 	RootComponent = LightComponent;
 
 #if WITH_EDITORONLY_DATA
diff --git a/Engine/Source/Runtime/Engine/Private/SpotLight.cpp b/Engine/Source/Runtime/Engine/Private/SpotLight.cpp
--- a/Engine/Source/Runtime/Engine/Private/SpotLight.cpp
+++ b/Engine/Source/Runtime/Engine/Private/SpotLight.cpp
@@ -28,7 +28,7 @@ ASpotLight::ASpotLight(const FObjectInitializer& ObjectInitializer)
 	RootComponent = SpotLightComponent;
 
 #if WITH_EDITORONLY_DATA
-	ArrowComponent = ObjectInitializer.CreateEditorOnlyDefaultSubobject<UArrowComponent>(this, TEXT("ArrowComponent0"));
+	ArrowComponent = CreateEditorOnlyDefaultSubobject<UArrowComponent>(TEXT("ArrowComponent0"));
 	if (ArrowComponent)
 	{
 		ArrowComponent->ArrowColor = GetLightColor();
diff --git a/Engine/Source/Runtime/Engine/Private/StaticMeshActor.cpp b/Engine/Source/Runtime/Engine/Private/StaticMeshActor.cpp
--- a/Engine/Source/Runtime/Engine/Private/StaticMeshActor.cpp
+++ b/Engine/Source/Runtime/Engine/Private/StaticMeshActor.cpp
@@ -18,7 +18,7 @@ AStaticMeshActor::AStaticMeshActor(const FObjectInitializer& ObjectInitializer)
 {
 	bCanBeDamaged = false;
 
-	StaticMeshComponent = ObjectInitializer.CreateDefaultSubobject<UStaticMeshComponent>(this, TEXT("StaticMeshComponent0"));
+	StaticMeshComponent = CreateDefaultSubobject<UStaticMeshComponent>(TEXT("StaticMeshComponent0"));
 	StaticMeshComponent->SetCollisionProfileName(UCollisionProfile::BlockAll_ProfileName);
 	StaticMeshComponent->Mobility = EComponentMobility::Static;
 	StaticMeshComponent->bGenerateOverlapEvents = false;
diff --git a/Engine/Source/Runtime/Engine/Private/TargetPoint.cpp b/Engine/Source/Runtime/Engine/Private/TargetPoint.cpp
--- a/Engine/Source/Runtime/Engine/Private/TargetPoint.cpp
+++ b/Engine/Source/Runtime/Engine/Private/TargetPoint.cpp
@@ -6,13 +6,13 @@
 ATargetPoint::ATargetPoint(const FObjectInitializer& ObjectInitializer)
 	: Super(ObjectInitializer)
 {
-	USceneComponent* SceneComponent = ObjectInitializer.CreateDefaultSubobject<USceneComponent>(this, TEXT("SceneComp"));
+	USceneComponent* SceneComponent = CreateDefaultSubobject<USceneComponent>(TEXT("SceneComp"));
 
 	RootComponent = SceneComponent;
 
 #if WITH_EDITORONLY_DATA
-	SpriteComponent = ObjectInitializer.CreateEditorOnlyDefaultSubobject<UBillboardComponent>(this, TEXT("Sprite"));
-	ArrowComponent = ObjectInitializer.CreateEditorOnlyDefaultSubobject<UArrowComponent>(this, TEXT("Arrow"));
+	SpriteComponent = CreateEditorOnlyDefaultSubobject<UBillboardComponent>(TEXT("Sprite"));
+	ArrowComponent = CreateEditorOnlyDefaultSubobject<UArrowComponent>(TEXT("Arrow"));
 
 	if (!IsRunningCommandlet())
 	{
diff --git a/Engine/Source/Runtime/Engine/Private/TextRenderComponent.cpp b/Engine/Source/Runtime/Engine/Private/TextRenderComponent.cpp
--- a/Engine/Source/Runtime/Engine/Private/TextRenderComponent.cpp
+++ b/Engine/Source/Runtime/Engine/Private/TextRenderComponent.cpp
@@ -10,11 +10,11 @@
 ATextRenderActor::ATextRenderActor(const FObjectInitializer& ObjectInitializer)
 	: Super(ObjectInitializer)
 {
-	TextRender = ObjectInitializer.CreateDefaultSubobject<UTextRenderComponent>(this, TEXT("NewTextRenderComponent"));
+	TextRender = CreateDefaultSubobject<UTextRenderComponent>(TEXT("NewTextRenderComponent"));
 	RootComponent = TextRender;
 
 #if WITH_EDITORONLY_DATA
-	SpriteComponent = ObjectInitializer.CreateEditorOnlyDefaultSubobject<UBillboardComponent>(this, TEXT("Sprite"));
+	SpriteComponent = CreateEditorOnlyDefaultSubobject<UBillboardComponent>(TEXT("Sprite"));
 
 	if (!IsRunningCommandlet() && (SpriteComponent != nullptr))
 	{
diff --git a/Engine/Source/Runtime/Engine/Private/TriggerBase.cpp b/Engine/Source/Runtime/Engine/Private/TriggerBase.cpp
--- a/Engine/Source/Runtime/Engine/Private/TriggerBase.cpp
+++ b/Engine/Source/Runtime/Engine/Private/TriggerBase.cpp
@@ -26,14 +26,14 @@ ATriggerBase::ATriggerBase(const FObjectInitializer& ObjectInitializer)
 
 	// ATriggerBase is requesting UShapeComponent which is abstract, however it is responsibility
 	// of a derived class to override this type with ObjectInitializer.SetDefaultSubobjectClass.
-	CollisionComponent = ObjectInitializer.CreateAbstractDefaultSubobject<UShapeComponent>(this, TEXT("CollisionComp"));
+	CollisionComponent = CreateAbstractDefaultSubobject<UShapeComponent>(TEXT("CollisionComp"));
 	if (CollisionComponent)
 	{
 		RootComponent = CollisionComponent;
 		CollisionComponent->bHiddenInGame = false;
 	}
 
-	SpriteComponent = ObjectInitializer.CreateDefaultSubobject<UBillboardComponent>(this, TEXT("Sprite"));
+	SpriteComponent = CreateDefaultSubobject<UBillboardComponent>(TEXT("Sprite"));
 	if (SpriteComponent)
 	{
 		SpriteComponent->Sprite = ConstructorStatics.TriggerTextureObject.Get();
diff --git a/Engine/Source/Runtime/Engine/Private/VectorFieldVolume.cpp b/Engine/Source/Runtime/Engine/Private/VectorFieldVolume.cpp
--- a/Engine/Source/Runtime/Engine/Private/VectorFieldVolume.cpp
+++ b/Engine/Source/Runtime/Engine/Private/VectorFieldVolume.cpp
@@ -8,11 +8,11 @@
 AVectorFieldVolume::AVectorFieldVolume(const FObjectInitializer& ObjectInitializer)
 	: Super(ObjectInitializer)
 {
-	VectorFieldComponent = ObjectInitializer.CreateDefaultSubobject<UVectorFieldComponent>(this, TEXT("VectorFieldComponent0"));
+	VectorFieldComponent = CreateDefaultSubobject<UVectorFieldComponent>(TEXT("VectorFieldComponent0"));
 	RootComponent = VectorFieldComponent;
 
 #if WITH_EDITORONLY_DATA
-	SpriteComponent = ObjectInitializer.CreateEditorOnlyDefaultSubobject<UBillboardComponent>(this, TEXT("Sprite"));
+	SpriteComponent = CreateEditorOnlyDefaultSubobject<UBillboardComponent>(TEXT("Sprite"));
 
 	if (!IsRunningCommandlet() && (SpriteComponent != nullptr))
 	{
diff --git a/Engine/Source/Runtime/Engine/Private/WheeledVehicle.cpp b/Engine/Source/Runtime/Engine/Private/WheeledVehicle.cpp
--- a/Engine/Source/Runtime/Engine/Private/WheeledVehicle.cpp
+++ b/Engine/Source/Runtime/Engine/Private/WheeledVehicle.cpp
@@ -16,7 +16,7 @@ FName AWheeledVehicle::VehicleMeshComponentName(TEXT("VehicleMesh"));
 AWheeledVehicle::AWheeledVehicle(const FObjectInitializer& ObjectInitializer)
 	: Super(ObjectInitializer)
 {
-	Mesh = ObjectInitializer.CreateDefaultSubobject<USkeletalMeshComponent>(this, VehicleMeshComponentName);
+	Mesh = CreateDefaultSubobject<USkeletalMeshComponent>(VehicleMeshComponentName);
 	Mesh->SetCollisionProfileName(UCollisionProfile::Vehicle_ProfileName);
 	Mesh->BodyInstance.bSimulatePhysics = true;
 	Mesh->BodyInstance.bNotifyRigidBodyCollision = true;
@@ -26,7 +26,7 @@ AWheeledVehicle::AWheeledVehicle(const FObjectInitializer& ObjectInitializer)
 	Mesh->bCanEverAffectNavigation = false;
 	RootComponent = Mesh;
 
-	VehicleMovement = ObjectInitializer.CreateDefaultSubobject<UWheeledVehicleMovementComponent, UWheeledVehicleMovementComponent4W>(this, VehicleMovementComponentName);
+	VehicleMovement = CreateDefaultSubobject<UWheeledVehicleMovementComponent, UWheeledVehicleMovementComponent4W>(VehicleMovementComponentName);
 	VehicleMovement->SetIsReplicated(true); // Enable replication by default
 	VehicleMovement->UpdatedComponent = Mesh;
 }
diff --git a/Engine/Source/Runtime/Engine/Private/WindDirectionalSource.cpp b/Engine/Source/Runtime/Engine/Private/WindDirectionalSource.cpp
--- a/Engine/Source/Runtime/Engine/Private/WindDirectionalSource.cpp
+++ b/Engine/Source/Runtime/Engine/Private/WindDirectionalSource.cpp
@@ -7,11 +7,11 @@
 AWindDirectionalSource::AWindDirectionalSource(const FObjectInitializer& ObjectInitializer)
 	: Super(ObjectInitializer)
 {
-	Component = ObjectInitializer.CreateDefaultSubobject<UWindDirectionalSourceComponent>(this, TEXT("WindDirectionalSourceComponent0"));
+	Component = CreateDefaultSubobject<UWindDirectionalSourceComponent>(TEXT("WindDirectionalSourceComponent0"));
 	RootComponent = Component;
 
 #if WITH_EDITORONLY_DATA
-	ArrowComponent = ObjectInitializer.CreateEditorOnlyDefaultSubobject<UArrowComponent>(this, TEXT("ArrowComponent0"));
+	ArrowComponent = CreateEditorOnlyDefaultSubobject<UArrowComponent>(TEXT("ArrowComponent0"));
 
 	if (!IsRunningCommandlet())
 	{
