Add scale support to UpdateInstanceTransform
[UE-6533]

diff --git a/Engine/Source/Runtime/Engine/Private/InstancedStaticMesh.cpp b/Engine/Source/Runtime/Engine/Private/InstancedStaticMesh.cpp
--- a/Engine/Source/Runtime/Engine/Private/InstancedStaticMesh.cpp
+++ b/Engine/Source/Runtime/Engine/Private/InstancedStaticMesh.cpp
@@ -1262,6 +1262,7 @@ bool UInstancedStaticMeshComponent::UpdateInstanceTransform(int32 InstanceIndex,
 #if WITH_PHYSX
 		// Update transform.
 		InstanceBodyInstance->SetBodyTransform(WorldTransform, false);
+		InstanceBodyInstance->UpdateBodyScale(WorldTransform.GetScale3D());
 #endif //WITH_PHYSX
 	}
 
