Fixed PhysX crash caused by PxScene::addActors

* PxActors created after simulation started would crash because they weren't getting added to the PxScene.  PxScene::addActors doesn't work when simulation is running.  I changed it to call PxScene::addActor in a loop.  Fixes blocker bug OR-446.

diff --git a/Engine/Source/Runtime/Engine/Private/PhysicsEngine/BodyInstance.cpp b/Engine/Source/Runtime/Engine/Private/PhysicsEngine/BodyInstance.cpp
--- a/Engine/Source/Runtime/Engine/Private/PhysicsEngine/BodyInstance.cpp
+++ b/Engine/Source/Runtime/Engine/Private/PhysicsEngine/BodyInstance.cpp
@@ -1877,12 +1877,18 @@ void FBodyInstance::InitBodies(FBodyInstance** Bodies, FTransform* Transforms, i
 		else
 		{
 			SCOPED_SCENE_WRITE_LOCK(PSceneSync);
-			PSceneSync->addActors(PhysXSyncActors.GetData(), PhysXSyncActors.Num());
+			for(PxActor* Actor : PhysXSyncActors)
+			{
+				PSceneSync->addActor(*Actor);
+			}
 
 			if(PSceneAsync)
 			{
 				SCOPED_SCENE_WRITE_LOCK(PSceneAsync);
-				PSceneAsync->addActors(PhysXAsyncActors.GetData(), PhysXAsyncActors.Num());
+				for (PxActor* Actor : PhysXAsyncActors)
+				{
+					PSceneAsync->addActor(*Actor);
+				}
 			}
 		}
 	}
