Disabled bAffectDistanceFieldLighting on foliage by default, as the most common use case is grass which causes performance problems

diff --git a/Engine/Source/Runtime/Engine/Classes/Foliage/FoliageType.h b/Engine/Source/Runtime/Engine/Classes/Foliage/FoliageType.h
--- a/Engine/Source/Runtime/Engine/Classes/Foliage/FoliageType.h
+++ b/Engine/Source/Runtime/Engine/Classes/Foliage/FoliageType.h
@@ -188,7 +188,7 @@ class UFoliageType : public UObject
 	uint32 bAffectDynamicIndirectLighting:1;
 
 	/** Controls whether the primitive should affect dynamic distance field lighting methods.  This flag is only used if CastShadow is true. **/
-	UPROPERTY(EditAnywhere, BlueprintReadOnly, Category=Lighting, AdvancedDisplay, meta=(EditCondition="CastShadow"))
+	UPROPERTY(EditAnywhere, BlueprintReadOnly, Category=Lighting, meta=(EditCondition="CastShadow"))
 	uint32 bAffectDistanceFieldLighting:1;
 
 	/** Controls whether the foliage should cast shadows in the case of non precomputed shadowing.  This flag is only used if CastShadow is true. **/
diff --git a/Engine/Source/Runtime/Engine/Private/InstancedFoliage.cpp b/Engine/Source/Runtime/Engine/Private/InstancedFoliage.cpp
--- a/Engine/Source/Runtime/Engine/Private/InstancedFoliage.cpp
+++ b/Engine/Source/Runtime/Engine/Private/InstancedFoliage.cpp
@@ -182,7 +182,8 @@ UFoliageType::UFoliageType(const FObjectInitializer& ObjectInitializer)
 	bCastDynamicShadow = true;
 	bCastStaticShadow = true;
 	bAffectDynamicIndirectLighting = false;
-	bAffectDistanceFieldLighting = true;
+	// Most of the high instance count foliage like grass causes performance problems with distance field lighting
+	bAffectDistanceFieldLighting = false;
 	bCastShadowAsTwoSided = false;
 	bReceivesDecals = false;
 
