XGE Distributed Shader Compilation from Luke Thatcher ([REDACTED_EMAIL]) of Lionhead Studios
 - Enabled with either r.XGEShaderCompile=1 in ConsoleVariables.ini, or -xgeshadercompile command line argument.

diff --git a/Engine/Config/ConsoleVariables.ini b/Engine/Config/ConsoleVariables.ini
--- a/Engine/Config/ConsoleVariables.ini
+++ b/Engine/Config/ConsoleVariables.ini
@@ -27,6 +27,8 @@
 ; Uncomment to enable parallel rendering at startup
 ;r.RHICmdBypass=0
 
+; Uncomment to enable XGE shader compilation.
+;r.XGEShaderCompile = 1
 ; Uncomment when running with a graphical debugger (but not when profiling)
 ;r.Shaders.Optimize=0
 ;r.Shaders.KeepDebugInfo=1
\ No newline at end of file
diff --git a/Engine/Source/Programs/ShaderCompileWorker/Private/ShaderCompileWorker.cpp b/Engine/Source/Programs/ShaderCompileWorker/Private/ShaderCompileWorker.cpp
--- a/Engine/Source/Programs/ShaderCompileWorker/Private/ShaderCompileWorker.cpp
+++ b/Engine/Source/Programs/ShaderCompileWorker/Private/ShaderCompileWorker.cpp
@@ -18,6 +18,13 @@ const int32 ShaderCompileWorkerOutputVersion = 1;
 
 double LastCompileTime = 0.0;
 
+static bool GShaderCompileUseXGE = false;
+static void WriteXGESuccessFile(const TCHAR* WorkingDirectory)
+{
+	// To signal compilation completion, create a zero length file in the working directory.
+	delete IFileManager::Get().CreateFileWriter(*FString::Printf(TEXT("%s/Success"), WorkingDirectory), FILEWRITE_EvenIfReadOnly);
+}
+
 const TArray<const IShaderFormat*>& GetShaderFormats()
 {
 	static bool bInitialized = false;
@@ -178,6 +185,15 @@ public:
 				LastConnectionTime = FPlatformTime::Seconds();
 			}
 #endif	// PLATFORM_SUPPORTS_NAMED_PIPES
+
+			if (GShaderCompileUseXGE)
+			{
+				// To signal compilation completion, create a zero length file in the working directory.
+				WriteXGESuccessFile(*WorkingDirectory);
+
+				// We only do one pass per process when using XGE.
+				break;
+			}
 		}
 
 		UE_LOG(LogShaders, Log, TEXT("Exiting job loop"));
@@ -311,16 +327,22 @@ private:
 			const double StartTime = FPlatformTime::Seconds();
 			bool bResult = false;
 
-			do 
+			// It seems XGE does not support deleting files.
+			// Don't delete the input file if we are running under Incredibuild.
+			// Instead, we signal completion by creating a zero byte "Success" file after the output file has been fully written.
+			if (!GShaderCompileUseXGE)
 			{
-				// Remove the input file so that it won't get processed more than once
-				bResult = IFileManager::Get().Delete(*InputFilePath);
-			} 
-			while (!bResult && (FPlatformTime::Seconds() - StartTime < 2));
+				do 
+				{
+					// Remove the input file so that it won't get processed more than once
+					bResult = IFileManager::Get().Delete(*InputFilePath);
+				} 
+				while (!bResult && (FPlatformTime::Seconds() - StartTime < 2));
 
-			if (!bResult)
-			{
-				UE_LOG(LogShaders, Fatal,TEXT("Couldn't delete input file %s, is it readonly?"), *InputFilePath);
+				if (!bResult)
+				{
+					UE_LOG(LogShaders, Fatal,TEXT("Couldn't delete input file %s, is it readonly?"), *InputFilePath);
+				}
 			}
 
 #if PLATFORM_MAC || PLATFORM_LINUX
@@ -518,11 +540,11 @@ int32 GuardedMain(int32 argc, TCHAR* argv[])
 
 	FString InCommunicating = argv[6];
 #if PLATFORM_SUPPORTS_NAMED_PIPES
-	const bool bThroughFile = (InCommunicating == FString(TEXT("-communicatethroughfile")));
+	const bool bThroughFile = GShaderCompileUseXGE || (InCommunicating == FString(TEXT("-communicatethroughfile")));
 	const bool bThroughNamedPipe = (InCommunicating == FString(TEXT("-communicatethroughnamedpipe")));
 	const bool bThroughNamedPipeOnce = (InCommunicating == FString(TEXT("-communicatethroughnamedpipeonce")));
 #else
-	const bool bThroughFile = true;
+	const bool bThroughFile = GShaderCompileUseXGE;
 	const bool bThroughNamedPipe = false;
 	const bool bThroughNamedPipeOnce = false;
 #endif
@@ -538,6 +560,10 @@ int32 GuardedMain(int32 argc, TCHAR* argv[])
 
 int32 GuardedMainWrapper(int32 ArgC, TCHAR* ArgV[], const TCHAR* CrashOutputFile)
 {
+	// We need to know whether we are using XGE now, in case an exception
+	// is thrown before we parse the command line inside GuardedMain.
+	GShaderCompileUseXGE = FCString::Strcmp(ArgV[6], TEXT("-xge")) == 0;
+
 	int32 ReturnCode = 0;
 #if PLATFORM_WINDOWS
 	if (FPlatformMisc::IsDebuggerPresent())
@@ -579,6 +605,12 @@ int32 GuardedMainWrapper(int32 ArgC, TCHAR* ArgV[], const TCHAR* CrashOutputFile
 
 			// Close the output file.
 			delete &OutputFile;
+
+			if (GShaderCompileUseXGE)
+			{
+				ReturnCode = 1;
+				WriteXGESuccessFile(ArgV[1]);
+			}
 		}
 	}
 #endif
@@ -601,6 +633,32 @@ IMPLEMENT_APPLICATION(ShaderCompileWorker, "ShaderCompileWorker")
 
 INT32_MAIN_INT32_ARGC_TCHAR_ARGV()
 {
+	if (ArgC == 4 && FCString::Strcmp(ArgV[1], TEXT("-xgemonitor")) == 0)
+	{
+		// Open handles to the two processes
+		FProcHandle EngineProc = FPlatformProcess::OpenProcess(FCString::Atoi(ArgV[2]));
+		FProcHandle BuildProc = FPlatformProcess::OpenProcess(FCString::Atoi(ArgV[3]));
+
+		if (EngineProc.IsValid() && BuildProc.IsValid())
+		{
+			// Whilst the build is still in progress
+			while (FPlatformProcess::IsProcRunning(BuildProc))
+			{
+				// Check that the engine is still alive.
+				if (!FPlatformProcess::IsProcRunning(EngineProc))
+				{
+					// The engine has shutdown before the build was stopped.
+					// Kill off the build process
+					FPlatformProcess::TerminateProc(BuildProc);
+					break;
+				}
+
+				FPlatformProcess::Sleep(0.01f);
+			}
+		}
+		return 0;
+	}
+
 	if(ArgC < 6)
 	{
 		printf("ShaderCompileWorker is called by UE4, it requires specific command like arguments.\n");
diff --git a/Engine/Source/Programs/ShaderCompileWorker/ShaderCompileWorker.Target.cs b/Engine/Source/Programs/ShaderCompileWorker/ShaderCompileWorker.Target.cs
--- a/Engine/Source/Programs/ShaderCompileWorker/ShaderCompileWorker.Target.cs
+++ b/Engine/Source/Programs/ShaderCompileWorker/ShaderCompileWorker.Target.cs
@@ -68,6 +68,9 @@ public class ShaderCompileWorkerTarget : TargetRules
 
 		// Disable logging, as the workers are spawned often and logging will just slow them down
 		OutCPPEnvironmentConfiguration.Definitions.Add("ALLOW_LOG_FILE=0");
+
+        // Linking against wer.lib/wer.dll causes XGE to bail when the worker is run on a Windows 8 machine, so turn this off.
+        OutCPPEnvironmentConfiguration.Definitions.Add("ALLOW_WINDOWS_ERROR_REPORT_LIB=0");
 	}
     public override bool GUBP_AlwaysBuildWithBaseEditor()
     {
diff --git a/Engine/Source/Runtime/Core/Private/Misc/ConfigCacheIni.cpp b/Engine/Source/Runtime/Core/Private/Misc/ConfigCacheIni.cpp
--- a/Engine/Source/Runtime/Core/Private/Misc/ConfigCacheIni.cpp
+++ b/Engine/Source/Runtime/Core/Private/Misc/ConfigCacheIni.cpp
@@ -2774,7 +2774,10 @@ bool FConfigCacheIni::LoadGlobalIniFile(FString& FinalIniFilename, const TCHAR*
 	NewConfigFile.Name = BaseIniName;
 
 	// don't write anything to disk in cooked builds - we will always use re-generated INI files anyway.
-	if (!FPlatformProperties::RequiresCookedData() || bAllowGeneratedIniWhenCooked)
+	if ((!FPlatformProperties::RequiresCookedData() || bAllowGeneratedIniWhenCooked)
+			// We shouldn't save config files when in multiprocess mode,
+			// otherwise we get file contention in XGE shader builds.
+			&& !FParse::Param(FCommandLine::Get(), TEXT("Multiprocess")))
 	{
 		// Check the config system for any changes made to defaults and propagate through to the saved.
 		NewConfigFile.ProcessSourceAndCheckAgainstBackup();
diff --git a/Engine/Source/Runtime/Core/Private/Windows/WindowsPlatformProcess.cpp b/Engine/Source/Runtime/Core/Private/Windows/WindowsPlatformProcess.cpp
--- a/Engine/Source/Runtime/Core/Private/Windows/WindowsPlatformProcess.cpp
+++ b/Engine/Source/Runtime/Core/Private/Windows/WindowsPlatformProcess.cpp
@@ -443,7 +443,7 @@ bool FWindowsPlatformProcess::GetProcReturnCode( FProcHandle & ProcHandle, int32
 bool FWindowsPlatformProcess::GetApplicationMemoryUsage(uint32 ProcessId, SIZE_T* OutMemoryUsage)
 {
 	bool bSuccess = false;
-	HANDLE ProcessHandle = OpenProcess(PROCESS_QUERY_INFORMATION | PROCESS_VM_READ, false, ProcessId);
+	HANDLE ProcessHandle = ::OpenProcess(PROCESS_QUERY_INFORMATION | PROCESS_VM_READ, false, ProcessId);
 
 	if (ProcessHandle != NULL)
 	{
@@ -464,7 +464,7 @@ bool FWindowsPlatformProcess::GetApplicationMemoryUsage(uint32 ProcessId, SIZE_T
 bool FWindowsPlatformProcess::IsApplicationRunning( uint32 ProcessId )
 {
 	bool bApplicationRunning = true;
-	HANDLE ProcessHandle = OpenProcess(SYNCHRONIZE, false, ProcessId);
+	HANDLE ProcessHandle = ::OpenProcess(SYNCHRONIZE, false, ProcessId);
 	if (ProcessHandle == NULL)
 	{
 		bApplicationRunning = false;
@@ -517,7 +517,7 @@ bool FWindowsPlatformProcess::IsApplicationRunning( const TCHAR* ProcName )
 FString FWindowsPlatformProcess::GetApplicationName( uint32 ProcessId )
 {
 	FString Output = TEXT("");
-	HANDLE ProcessHandle = OpenProcess(PROCESS_QUERY_INFORMATION, false, ProcessId);
+	HANDLE ProcessHandle = ::OpenProcess(PROCESS_QUERY_INFORMATION, false, ProcessId);
 	if (ProcessHandle != NULL)
 	{
 		const int32 ProcessNameBufferSize = 4096;
@@ -1215,4 +1215,9 @@ bool FWindowsPlatformProcess::Daemonize()
 	return true;
 }
 
+FProcHandle FWindowsPlatformProcess::OpenProcess(uint32 ProcessID)
+{
+	return FProcHandle(::OpenProcess(PROCESS_ALL_ACCESS, 0, ProcessID));
+}
+
 #include "HideWindowsPlatformTypes.h"
diff --git a/Engine/Source/Runtime/Core/Public/Windows/WindowsPlatformProcess.h b/Engine/Source/Runtime/Core/Public/Windows/WindowsPlatformProcess.h
--- a/Engine/Source/Runtime/Core/Public/Windows/WindowsPlatformProcess.h
+++ b/Engine/Source/Runtime/Core/Public/Windows/WindowsPlatformProcess.h
@@ -122,7 +122,7 @@ public:
 	static FSemaphore* NewInterprocessSynchObject(const FString& Name, bool bCreate, uint32 MaxLocks = 1);
 	static bool DeleteInterprocessSynchObject(FSemaphore * Object);
 	static bool Daemonize();
-
+	static FProcHandle OpenProcess(uint32 ProcessID);
 protected:
 
 	/**
diff --git a/Engine/Source/Runtime/Engine/Private/ShaderCompiler/ShaderCompiler.cpp b/Engine/Source/Runtime/Engine/Private/ShaderCompiler/ShaderCompiler.cpp
--- a/Engine/Source/Runtime/Engine/Private/ShaderCompiler/ShaderCompiler.cpp
+++ b/Engine/Source/Runtime/Engine/Private/ShaderCompiler/ShaderCompiler.cpp
@@ -56,8 +56,52 @@ static TAutoConsoleVariable<FString> CVarD3DCompilerPath(TEXT("r.D3DCompilerPath
 	ECVF_Cheat);
 #endif
 
+namespace XGEConsoleVariables
+{
+	int32 Enabled = 0;
+	FAutoConsoleVariableRef CVarXGEShaderCompile(
+		TEXT("r.XGEShaderCompile"),
+		Enabled,
+		TEXT("Enables or disables the use of XGE to build shaders.\n")
+		TEXT("0: Local builds only. \n")
+		TEXT("1: Distribute builds using XGE (default)."),
+		ECVF_Default);
+
+	/** The maximum number of shaders to group into a single XGE task. */
+	int32 BatchSize = 16;
+	FAutoConsoleVariableRef CVarXGEShaderCompileBatchSize(
+		TEXT("r.XGEShaderCompile.BatchSize"),
+		BatchSize,
+		TEXT("Specifies the number of shaders to batch together into a single XGE task.\n")
+		TEXT("Default = 16\n"),
+		ECVF_Default);
+
+	/** The total number of batches to fill with shaders before creating another group of batches. */
+	int32 BatchGroupSize = 128;
+	FAutoConsoleVariableRef CVarXGEShaderCompileBatchGroupSize(
+		TEXT("r.XGEShaderCompile.BatchGroupSize"),
+		BatchGroupSize,
+		TEXT("Specifies the number of batches to fill with shaders.\n")
+		TEXT("Shaders are spread across this number of batches until all the batches are full.\n")
+		TEXT("This allows the XGE compile to go wider when compiling a small number of shaders.\n")
+		TEXT("Default = 128\n"),
+		ECVF_Default);
+
+	/**
+	 * The number of seconds to wait after a job is submitted before kicking off the XGE process.
+	 * This allows time for the engine to enqueue more shaders, so we get better batching.
+	 */
+	float JobTimeout = 0.5f;
+	FAutoConsoleVariableRef CVarXGEShaderCompileJobTimeout(
+		TEXT("r.XGEShaderCompile.JobTimeout"),
+		JobTimeout,
+		TEXT("The number of seconds to wait for additional shader jobs to be submitted before starting a build.\n")
+		TEXT("Default = 0.5\n"),
+		ECVF_Default);
+}
+
 // Serialize Queued Job information
-static void DoWriteTasks(TArray<FShaderCompileJob*>& QueuedJobs, FArchive& TransferFile)
+static void DoWriteTasks(const TArray<FShaderCompileJob*>& QueuedJobs, FArchive& TransferFile)
 {
 	int32 ShaderCompileWorkerInputVersion = 2;
 	TransferFile << ShaderCompileWorkerInputVersion;
@@ -74,7 +118,7 @@ static void DoWriteTasks(TArray<FShaderCompileJob*>& QueuedJobs, FArchive& Trans
 }
 
 // Process results from Worker Process
-static void DoReadTaskResults(TArray<FShaderCompileJob*>& QueuedJobs, FArchive& OutputFile)
+static void DoReadTaskResults(const TArray<FShaderCompileJob*>& QueuedJobs, FArchive& OutputFile)
 {
 	int32 ShaderCompileWorkerOutputVersion;
 	OutputFile << ShaderCompileWorkerOutputVersion;
@@ -360,7 +404,11 @@ struct FShaderCompileWorkerInfo
 	// warning: not virtual
 	~FShaderCompileWorkerInfo()
 	{
-		check(!WorkerProcess.IsValid() || !FPlatformProcess::IsProcRunning(WorkerProcess));
+		if(WorkerProcess.IsValid())
+		{
+			FPlatformProcess::TerminateProc(WorkerProcess);
+			WorkerProcess.Close();
+		}
 	}
 
 	void CreatePipeAndNewTask(uint32 WorkerIndex, uint32 ProcessId)
@@ -456,23 +504,34 @@ struct FShaderCompileWorkerInfo
 
 };
 
-FShaderCompileThreadRunnable::FShaderCompileThreadRunnable(FShaderCompilingManager* InManager) :
-	Manager(InManager),
-	Thread(NULL),
-	bTerminatedByError(false),
-	bForceFinish( false ),
-	bIsRunning( false )
+//--------------------
+// @LH : LUTHA - XGE Distributed Shader Compilation
+//-----
+FShaderCompileThreadRunnableBase::FShaderCompileThreadRunnableBase(FShaderCompilingManager* InManager)
+	: Manager(InManager)
+	, Thread(nullptr)
+	, bTerminatedByError(false)
+	, bForceFinish(false)
 {
-	LastCheckForWorkersTime = 0;
-
-	for (uint32 WorkerIndex = 0; WorkerIndex < Manager->NumShaderCompilingThreads; WorkerIndex++)
+}
+void FShaderCompileThreadRunnableBase::StartThread()
+{
+//-----
+// @LH : END
+//--------------------
+	if (Manager->bAllowAsynchronousShaderCompiling && !FPlatformProperties::RequiresCookedData())
 	{
-		WorkerInfos.Add(new FShaderCompileWorkerInfo());
+		Thread = FRunnableThread::Create(this, TEXT("ShaderCompilingThread"), 0, TPri_Normal, FPlatformAffinity::GetPoolThreadMask());
 	}
+}
 
-	if (Manager->bAllowAsynchronousShaderCompiling && !FPlatformProperties::RequiresCookedData())
+FShaderCompileThreadRunnable::FShaderCompileThreadRunnable(FShaderCompilingManager* InManager)
+	: FShaderCompileThreadRunnableBase(InManager)
+	, LastCheckForWorkersTime(0)
+{
+	for (uint32 WorkerIndex = 0; WorkerIndex < Manager->NumShaderCompilingThreads; WorkerIndex++)
 	{
-		Thread = FRunnableThread::Create(this, TEXT("ShaderCompilingThread"), 0, TPri_Normal, FPlatformAffinity::GetPoolThreadMask());
+		WorkerInfos.Add(new FShaderCompileWorkerInfo());
 	}
 }
 
@@ -487,7 +546,7 @@ FShaderCompileThreadRunnable::~FShaderCompileThreadRunnable()
 }
 
 /** Entry point for the shader compiling thread. */
-uint32 FShaderCompileThreadRunnable::Run()
+uint32 FShaderCompileThreadRunnableBase::Run()
 {
 #if _MSC_VER
 	if(!FPlatformMisc::IsDebuggerPresent())
@@ -537,7 +596,7 @@ uint32 FShaderCompileThreadRunnable::Run()
 }
 
 /** Called by the main thread only, reports exceptions in the worker threads */
-void FShaderCompileThreadRunnable::CheckHealth() const
+void FShaderCompileThreadRunnableBase::CheckHealth() const
 {
 	if (bTerminatedByError)
 	{
@@ -1116,7 +1175,17 @@ FShaderCompilingManager::FShaderCompilingManager() :
 
 	NumShaderCompilingThreadsDuringGame = FMath::Min<int32>(NumShaderCompilingThreadsDuringGame, NumShaderCompilingThreads);
 
-	Thread = new FShaderCompileThreadRunnable(this);
+	if (FShaderCompileXGEThreadRunnable::IsSupported())
+	{
+		UE_LOG(LogShaderCompilers, Display, TEXT("Using XGE Shader Compiler."));
+		Thread = new FShaderCompileXGEThreadRunnable(this);
+	}
+	else
+	{
+		UE_LOG(LogShaderCompilers, Display, TEXT("Using Local Shader Compiler."));
+		Thread = new FShaderCompileThreadRunnable(this);
+	}
+	Thread->StartThread();
 }
 
 void FShaderCompilingManager::AddJobs(TArray<FShaderCompileJob*>& NewJobs, bool bApplyCompletedShaderMapForRendering, bool bOptimizeForLowLatency)
@@ -1548,10 +1617,7 @@ void FShaderCompilingManager::ProcessCompiledShaderMaps(
 void FShaderCompilingManager::Shutdown()
 {
 	Thread->Stop();
-	while ( Thread->IsRunning() )
-	{
-		FPlatformProcess::Sleep(0.01f);
-	}
+	Thread->WaitForCompletion();
 }
 
 
@@ -2412,3 +2478,575 @@ bool RecompileShaders(const TCHAR* Cmd, FOutputDevice& Ar)
 	UE_LOG(LogShaderCompilers, Warning, TEXT("Invalid parameter. Options are: \n'Changed', 'Global', 'Material [name]', 'All' 'Platform [name]'\nNote: Platform implies Changed, and requires the proper target platform modules to be compiled."));
 	return 1;
 }
+
+static FString XGE_ConsolePath;
+static const FString XGE_ScriptFileName(TEXT("xgscript.xml"));
+static const FString XGE_InputFileName (TEXT("WorkerInput.in"));
+static const FString XGE_OutputFileName(TEXT("WorkerOutput.out"));
+static const FString XGE_SuccessFileName(TEXT("Success"));
+
+bool FShaderCompileXGEThreadRunnable::IsSupported()
+{
+	// List of possible paths to xgconsole.exe
+	static const TCHAR* Paths[] =
+	{
+		TEXT("C:\\Program Files\\Xoreax\\IncrediBuild\\xgConsole.exe"),
+		TEXT("C:\\Program Files (x86)\\Xoreax\\IncrediBuild\\xgConsole.exe")
+	};
+
+	// Check the command line to see if XGE shader compilation has been enabled/disabled.
+	// This overrides the value of the console variable.
+	if (FParse::Param(FCommandLine::Get(), TEXT("xgeshadercompile")))
+	{
+		XGEConsoleVariables::Enabled = 1;
+	}
+	if (FParse::Param(FCommandLine::Get(), TEXT("noxgeshadercompile")))
+	{
+		XGEConsoleVariables::Enabled = 0;
+	}
+
+	// Check for a valid installation of Incredibuild by seeing if xgconsole.exe exists.
+	if (XGEConsoleVariables::Enabled == 1)
+	{
+		IPlatformFile& PlatformFile = FPlatformFileManager::Get().GetPlatformFile();
+
+		bool bFound = false;
+		for (int PathIndex = 0; PathIndex < ARRAYSIZE(Paths); PathIndex++)
+		{
+			if (PlatformFile.FileExists(Paths[PathIndex]))
+			{
+				bFound = true;
+				XGE_ConsolePath = Paths[PathIndex];
+				break;
+			}
+		}
+
+		if (!bFound)
+		{
+			UE_LOG(LogShaderCompilers, Warning, TEXT("Cannot use XGE Shader Compiler as Incredibuild is not installed on this machine."));
+			XGEConsoleVariables::Enabled = 0;
+		}
+	}
+
+	return XGEConsoleVariables::Enabled == 1;
+}
+
+/** Initialization constructor. */
+FShaderCompileXGEThreadRunnable::FShaderCompileXGEThreadRunnable(class FShaderCompilingManager* InManager)
+	: FShaderCompileThreadRunnableBase(InManager)
+	, BuildProcessID(INDEX_NONE)
+	, LastAddTime(0)
+	, StartTime(0)
+	, BatchIndexToFill(0)
+	, BatchIndexToCreate(0)
+	, XGEDirectoryIndex(0)
+	, XGEWorkingDirectory(InManager->AbsoluteShaderBaseWorkingDirectory / TEXT("XGE"))
+{
+}
+
+FShaderCompileXGEThreadRunnable::~FShaderCompileXGEThreadRunnable()
+{
+	if (BuildProcessHandle.IsValid())
+	{
+		// We still have a build in progress.
+		// Kill it...
+		FPlatformProcess::TerminateProc(BuildProcessHandle);
+		BuildProcessHandle.Close();
+	}
+
+	// Clean up any intermediate files/directories we've got left over.
+	IFileManager::Get().DeleteDirectory(*XGEWorkingDirectory, false, true);
+
+	// Delete all the shader batch instances we have.
+	for (FShaderBatch* Batch : ShaderBatchesIncomplete)
+		delete Batch;
+
+	for (FShaderBatch* Batch : ShaderBatchesInFlight)
+		delete Batch;
+
+	for (FShaderBatch* Batch : ShaderBatchesFull)
+		delete Batch;
+
+	ShaderBatchesIncomplete.Empty();
+	ShaderBatchesInFlight.Empty();
+	ShaderBatchesFull.Empty();
+}
+
+static FArchive* CreateFileHelper(const FString& Filename)
+{
+	// TODO: This logic came from FShaderCompileThreadRunnable::WriteNewTasks().
+	// We can't avoid code duplication unless we refactored the local worker too.
+
+	FArchive* File = nullptr;
+	int32 RetryCount = 0;
+	// Retry over the next two seconds if we can't write out the file.
+	// Anti-virus and indexing applications can interfere and cause this to fail.
+	while (File == nullptr && RetryCount < 200)
+	{
+		if (RetryCount > 0)
+		{
+			FPlatformProcess::Sleep(0.01f);
+		}
+		File = IFileManager::Get().CreateFileWriter(*Filename, FILEWRITE_EvenIfReadOnly);
+		RetryCount++;
+	}
+	if (File == nullptr)
+	{
+		File = IFileManager::Get().CreateFileWriter(*Filename, FILEWRITE_EvenIfReadOnly | FILEWRITE_NoFail);
+	}
+	checkf(File, TEXT("Failed to create file %s!"), *Filename);
+	return File;
+}
+
+static void MoveFileHelper(const FString& To, const FString& From)
+{
+	IPlatformFile& PlatformFile = FPlatformFileManager::Get().GetPlatformFile();
+
+	if (PlatformFile.FileExists(*From))
+	{
+		FString DirectoryName;
+		int32 LastSlashIndex;
+		if (To.FindLastChar('/', LastSlashIndex))
+		{
+			DirectoryName = To.Left(LastSlashIndex);
+		}
+		else
+		{
+			DirectoryName = To;
+		}
+
+		// TODO: This logic came from FShaderCompileThreadRunnable::WriteNewTasks().
+		// We can't avoid code duplication unless we refactored the local worker too.
+
+		bool Success = false;
+		int32 RetryCount = 0;
+		// Retry over the next two seconds if we can't move the file.
+		// Anti-virus and indexing applications can interfere and cause this to fail.
+		while (!Success && RetryCount < 200)
+		{
+			if (RetryCount > 0)
+			{
+				FPlatformProcess::Sleep(0.01f);
+			}
+
+			// MoveFile does not create the directory tree, so try to do that now...
+			Success = PlatformFile.CreateDirectoryTree(*DirectoryName);
+			if (Success)
+			{
+				Success = PlatformFile.MoveFile(*To, *From);
+			}
+			RetryCount++;
+		}
+		checkf(Success, TEXT("Failed to move file %s to %s!"), *From, *To);
+	}
+}
+
+static void DeleteFileHelper(const FString& Filename)
+{
+	// TODO: This logic came from FShaderCompileThreadRunnable::WriteNewTasks().
+	// We can't avoid code duplication unless we refactored the local worker too.
+
+	if (FPlatformFileManager::Get().GetPlatformFile().FileExists(*Filename))
+	{
+		bool bDeletedOutput = IFileManager::Get().Delete(*Filename, true, true);
+
+		// Retry over the next two seconds if we couldn't delete it
+		int32 RetryCount = 0;
+		while (!bDeletedOutput && RetryCount < 200)
+		{
+			FPlatformProcess::Sleep(0.01f);
+			bDeletedOutput = IFileManager::Get().Delete(*Filename, true, true);
+			RetryCount++;
+		}
+		checkf(bDeletedOutput, TEXT("Failed to delete %s!"), *Filename);
+	}
+}
+
+void FShaderCompileXGEThreadRunnable::PostCompletedJobsForBatch(FShaderBatch* Batch)
+{
+	// Enter the critical section so we can access the input and output queues
+	FScopeLock Lock(&Manager->CompileQueueSection);
+	for (FShaderCompileJob* Job : Batch->GetJobs())
+	{
+		FShaderMapCompileResults& ShaderMapResults = Manager->ShaderMapJobs.FindChecked(Job->Id);
+		ShaderMapResults.FinishedJobs.Add(Job);
+		ShaderMapResults.bAllJobsSucceeded = ShaderMapResults.bAllJobsSucceeded && Job->bSucceeded;
+	}
+
+	// Using atomics to update NumOutstandingJobs since it is read outside of the critical section
+	FPlatformAtomics::InterlockedAdd(&Manager->NumOutstandingJobs, -Batch->NumJobs());
+}
+
+void FShaderCompileXGEThreadRunnable::FShaderBatch::AddJob(FShaderCompileJob* Job)
+{
+	// We can only add jobs to a batch which hasn't been written out yet.
+	if (bTransferFileWritten)
+	{
+		UE_LOG(LogShaderCompilers, Fatal, TEXT("Attempt to add shader compile jobs to an XGE shader batch which has already been written to disk."));
+	}
+	else
+	{
+		Jobs.Add(Job);
+	}
+}
+		
+void FShaderCompileXGEThreadRunnable::FShaderBatch::WriteTransferFile()
+{
+	// Write out the file that the worker app is waiting for, which has all the information needed to compile the shader.
+	FArchive* TransferFile = CreateFileHelper(InputFileNameAndPath);
+	DoWriteTasks(Jobs, *TransferFile);
+	delete TransferFile;
+
+	bTransferFileWritten = true;
+}
+
+void FShaderCompileXGEThreadRunnable::FShaderBatch::SetIndices(int32 InDirectoryIndex, int32 InBatchIndex)
+{
+	DirectoryIndex = InDirectoryIndex;
+	BatchIndex = InBatchIndex;
+
+	WorkingDirectory = FString::Printf(TEXT("%s/%d/%d"), *DirectoryBase, DirectoryIndex, BatchIndex);
+
+	InputFileNameAndPath = WorkingDirectory / InputFileName;
+	OutputFileNameAndPath = WorkingDirectory / OutputFileName;
+	SuccessFileNameAndPath = WorkingDirectory / SuccessFileName;
+}
+
+void FShaderCompileXGEThreadRunnable::FShaderBatch::CleanUpFiles(bool keepInputFile)
+{
+	if (!keepInputFile)
+	{
+		DeleteFileHelper(InputFileNameAndPath);
+	}
+
+	DeleteFileHelper(OutputFileNameAndPath);
+	DeleteFileHelper(SuccessFileNameAndPath);
+}
+
+static void WriteScriptFileHeader(FArchive* ScriptFile, const FString& WorkerName)
+{
+	static const TCHAR HeaderTemplate[] =
+		TEXT("<?xml version=\"1.0\" encoding=\"UTF-8\" standalone=\"no\"?>\r\n"																						)
+		TEXT("<BuildSet FormatVersion=\"1\">\r\n"																													)
+		TEXT(	"\t<Environments>\r\n"																																)
+		TEXT(		"\t\t<Environment Name=\"Default\">\r\n"																										)
+		TEXT(			"\t\t\t<Tools>\r\n"																															)
+		TEXT(				"\t\t\t\t<Tool Name=\"ShaderCompiler\" Path=\"%s\" OutputFileMasks=\"%s,%s\" AllowRemote=\"true\" AllowRestartOnLocal=\"true\" />\r\n"	)
+		TEXT(			"\t\t\t</Tools>\r\n"																														)
+		TEXT(		"\t\t</Environment>\r\n"																														)
+		TEXT(	"\t</Environments>\r\n"																																)
+		TEXT(	"\t<Project Env=\"Default\" Name=\"Shader Compilation Project\">\r\n"																				)
+		TEXT(		"\t\t<TaskGroup Name=\"Compiling Shaders\" Tool=\"ShaderCompiler\">\r\n"																		);
+
+	FString HeaderXML = FString::Printf(HeaderTemplate, *WorkerName, *XGE_OutputFileName, *XGE_SuccessFileName);
+	ScriptFile->Serialize((void*)StringCast<ANSICHAR>(*HeaderXML, HeaderXML.Len()).Get(), sizeof(ANSICHAR) * HeaderXML.Len());
+}
+
+static void WriteScriptFileFooter(FArchive* ScriptFile)
+{
+	static const ANSICHAR HeaderFooter[] =
+				"\t\t</TaskGroup>\r\n"
+			"\t</Project>\r\n"
+		"</BuildSet>\r\n";
+
+	ScriptFile->Serialize((void*)HeaderFooter, sizeof(ANSICHAR) * (ARRAYSIZE(HeaderFooter) - 1));
+}
+
+void FShaderCompileXGEThreadRunnable::GatherResultsFromXGE()
+{
+	IPlatformFile& PlatformFile = FPlatformFileManager::Get().GetPlatformFile();
+	IFileManager& FileManager = IFileManager::Get();
+
+	// Reverse iterate so we can remove batches that have completed as we go.
+	for (int32 Index = ShaderBatchesInFlight.Num() - 1; Index >= 0; Index--)
+	{
+		FShaderBatch* Batch = ShaderBatchesInFlight[Index];
+
+		// Check to see if the shader compile worker has finished for this batch. This will be indicated by a zero length file placed in the working directory.
+		// We also check the timestamp of the success file to determine if it came from the current build and is not simply a leftover from a previous build.
+		if (PlatformFile.FileExists(*Batch->SuccessFileNameAndPath) && 
+			PlatformFile.GetTimeStamp(*Batch->SuccessFileNameAndPath) >= ScriptFileCreationTime)
+		{
+			// Perform the same checks on the worker output file to verify it came from this build.
+			if (PlatformFile.FileExists(*Batch->OutputFileNameAndPath) && 
+				PlatformFile.GetTimeStamp(*Batch->OutputFileNameAndPath) >= ScriptFileCreationTime)
+			{
+				FArchive* OutputFilePtr = FileManager.CreateFileReader(*Batch->OutputFileNameAndPath, FILEREAD_Silent);
+				if (OutputFilePtr)
+				{
+					FArchive& OutputFile = *OutputFilePtr;
+					DoReadTaskResults(Batch->GetJobs(), OutputFile);
+
+					// Close the output file.
+					delete OutputFilePtr;
+
+					// Cleanup the worker files
+					Batch->CleanUpFiles(false);			// (false = don't keep the input file)
+					PostCompletedJobsForBatch(Batch);
+					ShaderBatchesInFlight.RemoveAt(Index);
+					delete Batch;
+				}
+			}
+		}
+	}
+}
+		
+int32 FShaderCompileXGEThreadRunnable::CompilingLoop()
+{
+	bool bWorkRemaining = false;
+
+	// We can only run one XGE build at a time.
+	// Check if a build is currently in progress.
+	if (BuildProcessHandle.IsValid())
+	{
+		// Read back results from the current batches in progress.
+		GatherResultsFromXGE();
+
+		bool bDoExitCheck = false;
+		if (FPlatformProcess::IsProcRunning(BuildProcessHandle))
+		{
+			if (ShaderBatchesInFlight.Num() == 0)
+			{
+				// We've processed all batches.
+				// Wait for the XGE console process to exit
+				FPlatformProcess::WaitForProc(BuildProcessHandle);
+				bDoExitCheck = true;
+			}
+		}
+		else
+		{
+			bDoExitCheck = true;
+		}
+
+		if (bDoExitCheck)
+		{
+			if (ShaderBatchesInFlight.Num() > 0)
+			{
+				// The build process has stopped.
+				// Do one final pass over the output files to gather any remaining results.
+				GatherResultsFromXGE();
+			}
+
+			// The build process is no longer running.
+			// We need to check the return code for possible failure
+			int32 ReturnCode = 0;
+			FPlatformProcess::GetProcReturnCode(BuildProcessHandle, &ReturnCode);
+
+			switch (ReturnCode)
+			{
+			case 0:
+				// No error
+				break;
+
+			case 1:
+				// One or more of the shader compile worker processes crashed.
+				UE_LOG(LogShaderCompilers, Fatal, TEXT("An error occurred during an XGE shader compilation job. One or more of the shader compile worker processes exited unexpectedly (Code 1)."));
+				break;
+
+			case 2:
+				// Fatal IncrediBuild error
+				UE_LOG(LogShaderCompilers, Fatal, TEXT("An error occurred during an XGE shader compilation job. XGConsole.exe returned a fatal Incredibuild error (Code 2)."));
+				break;
+
+			case 3:
+				// User canceled the build
+				UE_LOG(LogShaderCompilers, Display, TEXT("The user terminated an XGE shader compilation job. Incomplete shader jobs will be redispatched in another XGE build."));
+				break;
+
+			default:
+				UE_LOG(LogShaderCompilers, Fatal, TEXT("An unknown error occurred during an XGE shader compilation job (Code %d)."), ReturnCode);
+				break;
+			}
+
+			// Reclaim jobs from the workers which did not succeed (if any).
+			for (FShaderBatch* Batch : ShaderBatchesInFlight)
+			{
+				// Delete any output/success files, but keep the input file so we don't have to write it out again.
+				Batch->CleanUpFiles(true);
+
+				// We can't add any jobs to a shader batch which has already been written out to disk,
+				// so put the batch back into the full batches list, even if the batch isn't full.
+				ShaderBatchesFull.Add(Batch);
+
+				// Reset the batch/directory indices and move the input file to the correct place.
+				FString OldInputFilename = Batch->InputFileNameAndPath;
+				Batch->SetIndices(XGEDirectoryIndex, BatchIndexToCreate++);
+				MoveFileHelper(Batch->InputFileNameAndPath, OldInputFilename);
+			}
+
+			ShaderBatchesInFlight.Empty();
+			BuildProcessHandle.Close();
+		}
+
+		bWorkRemaining |= ShaderBatchesInFlight.Num() > 0;
+	}
+	// No build process running. Check if we can kick one off now.
+	else
+	{
+		// Determine if enough time has passed to allow a build to kick off.
+		// Since shader jobs are added to the shader compile manager asynchronously by the engine, 
+		// we want to give the engine enough time to queue up a large number of shaders.
+		// Otherwise we will only be kicking off a small number of shader jobs at once.
+		bool BuildDelayElapsed = (((FPlatformTime::Cycles() - LastAddTime) * FPlatformTime::GetSecondsPerCycle()) >= XGEConsoleVariables::JobTimeout);
+		bool HasJobsToRun = (ShaderBatchesIncomplete.Num() > 0 || ShaderBatchesFull.Num() > 0);
+
+		if (BuildDelayElapsed && HasJobsToRun && ShaderBatchesInFlight.Num() == 0)
+		{
+			// Move all the pending shader batches into the in-flight list.
+			ShaderBatchesInFlight.Reserve(ShaderBatchesIncomplete.Num() + ShaderBatchesFull.Num());
+
+			for (FShaderBatch* Batch : ShaderBatchesIncomplete)
+			{
+				// Check we've actually got jobs for this batch.
+				check(Batch->NumJobs() > 0);
+
+				// Make sure we've written out the worker files for any incomplete batches.
+				Batch->WriteTransferFile();
+				ShaderBatchesInFlight.Add(Batch);
+			}
+
+			for (FShaderBatch* Batch : ShaderBatchesFull)
+			{
+				// Check we've actually got jobs for this batch.
+				check(Batch->NumJobs() > 0);
+
+				ShaderBatchesInFlight.Add(Batch);
+			}
+			
+			ShaderBatchesFull.Empty();
+			ShaderBatchesIncomplete.Empty(XGEConsoleVariables::BatchGroupSize);
+			
+			FString ScriptFilename = XGEWorkingDirectory / FString::FromInt(XGEDirectoryIndex) / XGE_ScriptFileName;
+
+			// Create the XGE script file.
+			FArchive* ScriptFile = CreateFileHelper(ScriptFilename);
+			WriteScriptFileHeader(ScriptFile, Manager->ShaderCompileWorkerName);
+
+			// Write the XML task line for each shader batch
+			for (FShaderBatch* Batch : ShaderBatchesInFlight)
+			{
+				FString WorkerAbsoluteDirectory = IFileManager::Get().ConvertToAbsolutePathForExternalAppForWrite(*Batch->WorkingDirectory);
+				FPaths::NormalizeDirectoryName(WorkerAbsoluteDirectory);
+
+				FString WorkerParameters = FString::Printf(TEXT("&quot;%s/&quot; %d %d &quot;%s&quot; &quot;%s&quot; -xge %s"),
+					*WorkerAbsoluteDirectory,
+					Manager->ProcessId,
+					Batch->BatchIndex,
+					*XGE_InputFileName,
+					*XGE_OutputFileName,
+					*FCommandLine::GetSubprocessCommandline());
+				FString TaskXML = FString::Printf(TEXT("\t\t\t<Task Caption=\"Compiling %d Shaders (Batch %d)\" Params=\"%s\" />\r\n"), Batch->NumJobs(), Batch->BatchIndex, *WorkerParameters);
+	
+				ScriptFile->Serialize((void*)StringCast<ANSICHAR>(*TaskXML, TaskXML.Len()).Get(), sizeof(ANSICHAR) * TaskXML.Len());
+			}
+
+			// End the XML script file and close it.
+			WriteScriptFileFooter(ScriptFile);
+			delete ScriptFile;
+			ScriptFile = nullptr;
+
+			// Grab the timestamp from the script file.
+			// We use this to ignore any left over files from previous builds by only accepting files created after the script file.
+			ScriptFileCreationTime = IFileManager::Get().GetTimeStamp(*ScriptFilename);
+
+			StartTime = FPlatformTime::Cycles();
+
+			// Use stop on errors so we can respond to shader compile worker crashes immediately.
+			// Regular shader compilation errors are not returned as worker errors.
+			FString XGConsoleArgs = TEXT("/STOPONERRORS /BUILD \"") + ScriptFilename + TEXT("\"");
+
+			// Kick off the XGE process...
+			BuildProcessHandle = FPlatformProcess::CreateProc(*XGE_ConsolePath, *XGConsoleArgs, false, false, true, &BuildProcessID, 0, nullptr, nullptr);
+			if(!BuildProcessHandle.IsValid())
+			{
+				UE_LOG(LogShaderCompilers, Fatal, TEXT("Failed to launch %s during shader compilation."), *XGE_ConsolePath);
+			}
+
+			// If the engine crashes, we don't get a chance to kill the build process.
+			// Start up the build monitor process to monitor for engine crashes.
+			uint32 BuildMonitorProcessID;
+			FProcHandle BuildMonitorHandle = FPlatformProcess::CreateProc(*Manager->ShaderCompileWorkerName, *FString::Printf(TEXT("-xgemonitor %d %d"), Manager->ProcessId, BuildProcessID), true, false, false, &BuildMonitorProcessID, 0, nullptr, nullptr);
+			BuildMonitorHandle.Close();
+
+			// Reset batch counters and switch directories
+			BatchIndexToFill = 0;
+			BatchIndexToCreate = 0;
+			XGEDirectoryIndex = 1 - XGEDirectoryIndex;
+
+			bWorkRemaining = true;
+		}
+	}
+
+	// Try to prepare more shader jobs (even if a build is in flight).
+	TArray<FShaderCompileJob*> JobQueue;
+	{
+		// Enter the critical section so we can access the input and output queues
+		FScopeLock Lock(&Manager->CompileQueueSection);
+
+		// Grab as many jobs from the job queue as we can.
+		int32 NumNewJobs = Manager->CompileQueue.Num();
+		if (NumNewJobs > 0)
+		{
+			int32 DestJobIndex = JobQueue.AddUninitialized(NumNewJobs);
+			for (int32 SrcJobIndex = 0; SrcJobIndex < NumNewJobs; SrcJobIndex++, DestJobIndex++)
+			{
+				JobQueue[DestJobIndex] = Manager->CompileQueue[SrcJobIndex];
+			}
+		
+			Manager->CompileQueue.RemoveAt(0, NumNewJobs);
+		}
+	}
+
+	if (JobQueue.Num() > 0)
+	{
+		// We have new jobs in the queue.
+		// Group the jobs into batches and create the worker input files.
+		for (int32 JobIndex = 0; JobIndex < JobQueue.Num(); JobIndex++)
+		{
+			if (BatchIndexToFill >= ShaderBatchesIncomplete.GetMaxIndex() || !ShaderBatchesIncomplete.IsAllocated(BatchIndexToFill))
+			{
+				// There are no more incomplete shader batches available.
+				// Create another one...
+				ShaderBatchesIncomplete.Insert(BatchIndexToFill, new FShaderBatch(
+					XGEWorkingDirectory,
+					XGE_InputFileName,
+					XGE_SuccessFileName,
+					XGE_OutputFileName,
+					XGEDirectoryIndex,
+					BatchIndexToCreate));
+
+				BatchIndexToCreate++;
+			}
+
+			// Add a single job to this batch
+			FShaderBatch* CurrentBatch = ShaderBatchesIncomplete[BatchIndexToFill];
+			CurrentBatch->AddJob(JobQueue[JobIndex]);
+
+			// If the batch is now full...
+			if (CurrentBatch->NumJobs() == XGEConsoleVariables::BatchSize)
+			{
+				CurrentBatch->WriteTransferFile();
+
+				// Move the batch to the full list.
+				ShaderBatchesFull.Add(CurrentBatch);
+				ShaderBatchesIncomplete.RemoveAt(BatchIndexToFill);
+			}
+
+			BatchIndexToFill++;
+			BatchIndexToFill %= XGEConsoleVariables::BatchGroupSize;
+		}
+
+		// Keep track of the last time we added jobs.
+		LastAddTime = FPlatformTime::Cycles();
+
+		bWorkRemaining = true;
+	}
+
+	if (Manager->bAllowAsynchronousShaderCompiling)
+	{
+		// Yield for a short while to stop this thread continuously polling the disk.
+		FPlatformProcess::Sleep(0.01f);
+	}
+
+	return bWorkRemaining ? 1 : 0;
+}
diff --git a/Engine/Source/Runtime/Engine/Public/ShaderCompiler.h b/Engine/Source/Runtime/Engine/Public/ShaderCompiler.h
--- a/Engine/Source/Runtime/Engine/Public/ShaderCompiler.h
+++ b/Engine/Source/Runtime/Engine/Public/ShaderCompiler.h
@@ -48,14 +48,11 @@ public:
 	}
 };
 
-/** 
- * Shader compiling thread
- * This runs in the background while UE4 is running, launches shader compile worker processes when necessary, and feeds them inputs and reads back the outputs.
- */
-class FShaderCompileThreadRunnable : public FRunnable
+class FShaderCompileThreadRunnableBase : public FRunnable
 {
 	friend class FShaderCompilingManager;
-private:
+
+protected:
 	/** The manager for this thread */
 	class FShaderCompilingManager* Manager;
 	/** The runnable thread */
@@ -65,27 +62,50 @@ private:
 	FString ErrorMessage;
 	/** true if the thread has been terminated by an unhandled exception. */
 	bool bTerminatedByError;
+
+	volatile bool bForceFinish;
+
+public:
+	FShaderCompileThreadRunnableBase(class FShaderCompilingManager* InManager);
+	virtual ~FShaderCompileThreadRunnableBase()
+	{}
+
+	void StartThread();
+
+	// FRunnable interface.
+	virtual void Stop() { bForceFinish = true; }
+	virtual uint32 Run();
+	inline void WaitForCompletion() const
+	{
+		Thread->WaitForCompletion();
+	}
+
+	/** Checks the thread's health, and passes on any errors that have occured.  Called by the main thread. */
+	void CheckHealth() const;
+
+	/** Main work loop. */
+	virtual int32 CompilingLoop() = 0;
+};
+
+/** 
+ * Shader compiling thread
+ * This runs in the background while UE4 is running, launches shader compile worker processes when necessary, and feeds them inputs and reads back the outputs.
+ */
+class FShaderCompileThreadRunnable : public FShaderCompileThreadRunnableBase
+{
+	friend class FShaderCompilingManager;
+private:
+
 	/** Information about the active workers that this thread is tracking. */
 	TArray<struct FShaderCompileWorkerInfo*> WorkerInfos;
 	/** Tracks the last time that this thread checked if the workers were still active. */
 	double LastCheckForWorkersTime;
 
-	volatile bool bForceFinish;
-	volatile bool bIsRunning;
 public:
 	/** Initialization constructor. */
 	FShaderCompileThreadRunnable(class FShaderCompilingManager* InManager);
 	virtual ~FShaderCompileThreadRunnable();
 
-	// FRunnable interface.
-	virtual bool Init(void) { bIsRunning = true; return true; }
-	virtual void Exit(void) { bIsRunning = false; }
-	virtual void Stop(void) { bForceFinish = true; }
-	virtual uint32 Run(void);
-
-	/** Checks the thread's health, and passes on any errors that have occured.  Called by the main thread. */
-	void CheckHealth() const;
-
 private:
 
 	/** 
@@ -107,12 +127,99 @@ private:
 	void CompileDirectlyThroughDll();
 
 	/** Main work loop. */
-	int32 CompilingLoop();
+	virtual int32 CompilingLoop() override;
 
 	/** Used when compiling through workers, launches worker processes if needed. */
 	void LaunchWorkerIfNeeded(FShaderCompileWorkerInfo & CurrentWorkerInfo, uint32 WorkerIndex);
+};
+
+class FShaderCompileXGEThreadRunnable : public FShaderCompileThreadRunnableBase
+{
+private:
+	/** The handle referring to the XGE console process, if a build is in progress. */
+	FProcHandle BuildProcessHandle;
+	
+	/** Process ID of the XGE console, if a build is in progress. */
+	uint32 BuildProcessID;
+
+	/**
+	 * A map of directory paths to shader jobs contained within that directory.
+	 * One entry per XGE task.
+	 */
+	class FShaderBatch
+	{
+		TArray<FShaderCompileJob*> Jobs;
+		bool bTransferFileWritten;
+
+	public:
+		const FString& DirectoryBase;
+		const FString& InputFileName;
+		const FString& SuccessFileName;
+		const FString& OutputFileName;
+
+		int32 BatchIndex;
+		int32 DirectoryIndex;
+
+		FString WorkingDirectory;
+		FString OutputFileNameAndPath;
+		FString SuccessFileNameAndPath;
+		FString InputFileNameAndPath;
+		
+		FShaderBatch(const FString& InDirectoryBase, const FString& InInputFileName, const FString& InSuccessFileName, const FString& InOutputFileName, int32 InDirectoryIndex, int32 InBatchIndex)
+			: DirectoryBase(InDirectoryBase)
+			, InputFileName(InInputFileName)
+			, SuccessFileName(InSuccessFileName)
+			, OutputFileName(InOutputFileName)
+			, bTransferFileWritten(false)
+		{
+			SetIndices(InDirectoryIndex, InBatchIndex);
+		}
+
+		void SetIndices(int32 InDirectoryIndex, int32 InBatchIndex);
+
+		void CleanUpFiles(bool keepInputFile);
+
+		inline int32 NumJobs()
+		{
+			return Jobs.Num();
+		}
+		inline const TArray<FShaderCompileJob*>& GetJobs() const
+		{
+			return Jobs;
+		}
+
+		void AddJob(FShaderCompileJob* Job);
+		
+		void WriteTransferFile();
+	};
+	TArray<FShaderBatch*> ShaderBatchesInFlight;
+	TArray<FShaderBatch*> ShaderBatchesFull;
+	TSparseArray<FShaderBatch*> ShaderBatchesIncomplete;
+
+	/** The full path to the two working directories for XGE shader builds. */
+	const FString XGEWorkingDirectory;
+	uint32 XGEDirectoryIndex;
+
+	uint64 LastAddTime;
+	uint64 StartTime;
+	int32 BatchIndexToCreate;
+	int32 BatchIndexToFill;
+
+	FDateTime ScriptFileCreationTime;
+
+	void PostCompletedJobsForBatch(FShaderBatch* Batch);
+
+	void GatherResultsFromXGE();
+
+public:
+	/** Initialization constructor. */
+	FShaderCompileXGEThreadRunnable(class FShaderCompilingManager* InManager);
+	virtual ~FShaderCompileXGEThreadRunnable();
+
+	/** Main work loop. */
+	virtual int32 CompilingLoop() override;
 
-	inline bool IsRunning() { return bIsRunning; }
+	static bool IsSupported();
 };
 
 /** Results for a single compiled shader map. */
@@ -148,7 +255,9 @@ struct FShaderMapFinalizeResults : public FShaderMapCompileResults
  */
 class FShaderCompilingManager
 {
+	friend class FShaderCompileThreadRunnableBase;
 	friend class FShaderCompileThreadRunnable;
+	friend class FShaderCompileXGEThreadRunnable;
 private:
 
 	//////////////////////////////////////////////////////
@@ -173,7 +282,7 @@ private:
 	TMap<int32, FShaderMapFinalizeResults> PendingFinalizeShaderMaps;
 
 	/** The threads spawned for shader compiling. */
-	TScopedPointer<FShaderCompileThreadRunnable> Thread;
+	TScopedPointer<FShaderCompileThreadRunnableBase> Thread;
 
 	//////////////////////////////////////////////////////
 	// Configuration properties - these are set only on initialization and can be read from either thread
diff --git a/Engine/Source/Runtime/Launch/Private/LaunchEngineLoop.cpp b/Engine/Source/Runtime/Launch/Private/LaunchEngineLoop.cpp
--- a/Engine/Source/Runtime/Launch/Private/LaunchEngineLoop.cpp
+++ b/Engine/Source/Runtime/Launch/Private/LaunchEngineLoop.cpp
@@ -2750,6 +2750,9 @@ void FEngineLoop::AppPreExit( )
 	if ( GShaderCompilingManager )
 	{
 		GShaderCompilingManager->Shutdown();
+
+		delete GShaderCompilingManager;
+		GShaderCompilingManager = nullptr;
 	}
 #endif
 }
