Mac compile fixes

diff --git a/Engine/Source/Runtime/Engine/Private/ShaderCompiler/ShaderCompiler.cpp b/Engine/Source/Runtime/Engine/Private/ShaderCompiler/ShaderCompiler.cpp
--- a/Engine/Source/Runtime/Engine/Private/ShaderCompiler/ShaderCompiler.cpp
+++ b/Engine/Source/Runtime/Engine/Private/ShaderCompiler/ShaderCompiler.cpp
@@ -2511,7 +2511,7 @@ bool FShaderCompileXGEThreadRunnable::IsSupported()
 		IPlatformFile& PlatformFile = FPlatformFileManager::Get().GetPlatformFile();
 
 		bool bFound = false;
-		for (int PathIndex = 0; PathIndex < ARRAYSIZE(Paths); PathIndex++)
+		for (int PathIndex = 0; PathIndex < ARRAY_COUNT(Paths); PathIndex++)
 		{
 			if (PlatformFile.FileExists(Paths[PathIndex]))
 			{
@@ -2535,12 +2535,12 @@ bool FShaderCompileXGEThreadRunnable::IsSupported()
 FShaderCompileXGEThreadRunnable::FShaderCompileXGEThreadRunnable(class FShaderCompilingManager* InManager)
 	: FShaderCompileThreadRunnableBase(InManager)
 	, BuildProcessID(INDEX_NONE)
+	, XGEWorkingDirectory(InManager->AbsoluteShaderBaseWorkingDirectory / TEXT("XGE"))
+	, XGEDirectoryIndex(0)
 	, LastAddTime(0)
 	, StartTime(0)
-	, BatchIndexToFill(0)
 	, BatchIndexToCreate(0)
-	, XGEDirectoryIndex(0)
-	, XGEWorkingDirectory(InManager->AbsoluteShaderBaseWorkingDirectory / TEXT("XGE"))
+	, BatchIndexToFill(0)
 {
 }
 
@@ -2749,7 +2749,7 @@ static void WriteScriptFileFooter(FArchive* ScriptFile)
 			"\t</Project>\r\n"
 		"</BuildSet>\r\n";
 
-	ScriptFile->Serialize((void*)HeaderFooter, sizeof(ANSICHAR) * (ARRAYSIZE(HeaderFooter) - 1));
+	ScriptFile->Serialize((void*)HeaderFooter, sizeof(ANSICHAR) * (ARRAY_COUNT(HeaderFooter) - 1));
 }
 
 void FShaderCompileXGEThreadRunnable::GatherResultsFromXGE()
diff --git a/Engine/Source/Runtime/Engine/Public/ShaderCompiler.h b/Engine/Source/Runtime/Engine/Public/ShaderCompiler.h
--- a/Engine/Source/Runtime/Engine/Public/ShaderCompiler.h
+++ b/Engine/Source/Runtime/Engine/Public/ShaderCompiler.h
@@ -166,11 +166,11 @@ private:
 		FString InputFileNameAndPath;
 		
 		FShaderBatch(const FString& InDirectoryBase, const FString& InInputFileName, const FString& InSuccessFileName, const FString& InOutputFileName, int32 InDirectoryIndex, int32 InBatchIndex)
-			: DirectoryBase(InDirectoryBase)
+			: bTransferFileWritten(false)
+			, DirectoryBase(InDirectoryBase)
 			, InputFileName(InInputFileName)
 			, SuccessFileName(InSuccessFileName)
 			, OutputFileName(InOutputFileName)
-			, bTransferFileWritten(false)
 		{
 			SetIndices(InDirectoryIndex, InBatchIndex);
 		}
diff --git a/Engine/Source/Runtime/Renderer/Private/DistanceFieldSurfaceCacheLighting.h b/Engine/Source/Runtime/Renderer/Private/DistanceFieldSurfaceCacheLighting.h
--- a/Engine/Source/Runtime/Renderer/Private/DistanceFieldSurfaceCacheLighting.h
+++ b/Engine/Source/Runtime/Renderer/Private/DistanceFieldSurfaceCacheLighting.h
@@ -220,7 +220,7 @@ public:
 /** Generates unit length, stratified and uniformly distributed direction samples in a hemisphere. */
 extern void GenerateStratifiedUniformHemisphereSamples2(int32 NumThetaSteps, int32 NumPhiSteps, FRandomStream& RandomStream, TArray<FVector4>& Samples);
 
-extern void GetSpacedVectors(TArray<TInlineAllocator<9> >& OutVectors);
+extern void GetSpacedVectors(TArray<FVector, TInlineAllocator<9> >& OutVectors);
 
 BEGIN_UNIFORM_BUFFER_STRUCT(FAOSampleData2,)
 	DECLARE_UNIFORM_BUFFER_STRUCT_MEMBER_ARRAY(FVector4,SampleDirections,[NumConeSampleDirections])
