Made certain material node properties multiline in the details panel.  Fixed a small bug where comment bubbles weren't updating their node comment text.

diff --git a/Engine/Source/Editor/GraphEditor/Private/SCommentBubble.cpp b/Engine/Source/Editor/GraphEditor/Private/SCommentBubble.cpp
--- a/Engine/Source/Editor/GraphEditor/Private/SCommentBubble.cpp
+++ b/Engine/Source/Editor/GraphEditor/Private/SCommentBubble.cpp
@@ -337,7 +337,8 @@ void SCommentBubble::OnCommentTextCommitted( const FText& NewText, ETextCommit::
 	{
 		const FScopedTransaction Transaction( NSLOCTEXT( "CommentBubble", "CommentCommitted", "Comment Changed" ) );
 		GraphNode->Modify();
-		GraphNode->NodeComment = NewText.ToString();
+		GraphNode->OnRenameNode(NewText.ToString());
+//		GraphNode->NodeComment = NewText.ToString();
 	}
 }
 
diff --git a/Engine/Source/Runtime/Engine/Classes/Materials/MaterialExpression.h b/Engine/Source/Runtime/Engine/Classes/Materials/MaterialExpression.h
--- a/Engine/Source/Runtime/Engine/Classes/Materials/MaterialExpression.h
+++ b/Engine/Source/Runtime/Engine/Classes/Materials/MaterialExpression.h
@@ -132,7 +132,7 @@ class ENGINE_API UMaterialExpression : public UObject
 	class UMaterialFunction* Function;
 
 	/** A description that level designers can add (shows in the material editor UI). */
-	UPROPERTY(EditAnywhere, Category=MaterialExpression)
+	UPROPERTY(EditAnywhere, Category=MaterialExpression, meta=(MultiLine=true))
 	FString Desc;
 
 	/** Color of the expression's border outline. */
diff --git a/Engine/Source/Runtime/Engine/Classes/Materials/MaterialExpressionComment.h b/Engine/Source/Runtime/Engine/Classes/Materials/MaterialExpressionComment.h
--- a/Engine/Source/Runtime/Engine/Classes/Materials/MaterialExpressionComment.h
+++ b/Engine/Source/Runtime/Engine/Classes/Materials/MaterialExpressionComment.h
@@ -16,7 +16,7 @@ class UMaterialExpressionComment : public UMaterialExpression
 	UPROPERTY()
 	int32 SizeY;
 
-	UPROPERTY(EditAnywhere, Category=MaterialExpressionComment)
+	UPROPERTY(EditAnywhere, Category=MaterialExpressionComment, meta=(MultiLine=true))
 	FString Text;
 
 	/** Color to style comment with */
diff --git a/Engine/Source/Runtime/Engine/Classes/Materials/MaterialExpressionFunctionOutput.h b/Engine/Source/Runtime/Engine/Classes/Materials/MaterialExpressionFunctionOutput.h
--- a/Engine/Source/Runtime/Engine/Classes/Materials/MaterialExpressionFunctionOutput.h
+++ b/Engine/Source/Runtime/Engine/Classes/Materials/MaterialExpressionFunctionOutput.h
@@ -15,7 +15,7 @@ class UMaterialExpressionFunctionOutput : public UMaterialExpression
 	FString OutputName;
 
 	/** The output's description, which will be used as a tooltip on the connector in function call expressions that use this function. */
-	UPROPERTY(EditAnywhere, Category=MaterialExpressionFunctionOutput)
+	UPROPERTY(EditAnywhere, Category=MaterialExpressionFunctionOutput, meta=(MultiLine=true))
 	FString Description;
 
 	/** Controls where the output is displayed relative to the other outputs. */
