Additional file intended for CL 2390838.

diff --git a/Engine/Source/Runtime/Engine/Private/SkeletalMesh.cpp b/Engine/Source/Runtime/Engine/Private/SkeletalMesh.cpp
--- a/Engine/Source/Runtime/Engine/Private/SkeletalMesh.cpp
+++ b/Engine/Source/Runtime/Engine/Private/SkeletalMesh.cpp
@@ -2698,6 +2698,13 @@ void USkeletalMesh::GetAssetRegistryTags(TArray<FAssetRegistryTag>& OutTags) con
 	Super::GetAssetRegistryTags(OutTags);
 }
 
+#if WITH_EDITOR
+void USkeletalMesh::GetAssetRegistryTagMetadata(TMap<FName, FAssetRegistryTagMetadata>& OutMetadata) const
+{
+	OutMetadata.Add("PhysicsAsset", FAssetRegistryTagMetadata().SetImportantValue(TEXT("None")));
+}
+#endif
+
 void USkeletalMesh::DebugVerifySkeletalMeshLOD()
 {
 	// if LOD do not have displayfactor set up correctly
