Fixed slow foliage tool instance removal algorithm which was significant when removing a very large number of instances.

diff --git a/Engine/Source/Runtime/Engine/Private/InstancedFoliage.cpp b/Engine/Source/Runtime/Engine/Private/InstancedFoliage.cpp
--- a/Engine/Source/Runtime/Engine/Private/InstancedFoliage.cpp
+++ b/Engine/Source/Runtime/Engine/Private/InstancedFoliage.cpp
@@ -371,19 +371,23 @@ void FFoliageMeshInfo::AddInstance(AInstancedFoliageActor* InIFA, UFoliageType*
 
 void FFoliageMeshInfo::RemoveInstances(AInstancedFoliageActor* InIFA, const TArray<int32>& InInstancesToRemove)
 {
-	bool bRemoved = false;
-
-	TArray<int32> InstancesToRemove = InInstancesToRemove;
-
-	if (InstancesToRemove.Num())
+	if (InInstancesToRemove.Num())
 	{
 		check(Component);
 		InIFA->Modify();
 
-		// Remove instances
-		for (int32 Idx = 0; Idx < InstancesToRemove.Num(); Idx++)
+		TSet<int32> InstancesToRemove;
+		for (int32 Instance : InInstancesToRemove)
+		{
+			InstancesToRemove.Add(Instance);
+		}
+
+		while(InstancesToRemove.Num())
 		{
-			int32 InstanceIndex = InstancesToRemove[Idx];
+			// Get an item from the set for processing
+			auto It = InstancesToRemove.CreateConstIterator();
+			int32 InstanceIndex = *It;		
+			int32 InstanceIndexToRemove = InstanceIndex;
 
 			FFoliageInstance& Instance = Instances[InstanceIndex];
 
@@ -431,18 +435,19 @@ void FFoliageMeshInfo::RemoveInstances(AInstancedFoliageActor* InIFA, const TArr
 				{
 					SelectedIndices.Remove(Instances.Num());
 					SelectedIndices.Add(InstanceIndex);
-				}				
-			}
+				}
 
-			// See if any item we're about to remove was swapped into this slot.
-			for (int32 NextIdx = Idx + 1; NextIdx < InstancesToRemove.Num(); NextIdx++)
-			{
-				if (InstancesToRemove[NextIdx] == Instances.Num())
+				// Removal list
+				if (InstancesToRemove.Contains(Instances.Num()))
 				{
-					InstancesToRemove[NextIdx] = InstanceIndex;
-					break;
+					// The item from the end of the array that we swapped in to InstanceIndex is also on the list to remove.
+					// Remove the item at the end of the array and leave InstanceIndex in the removal list.
+					InstanceIndexToRemove = Instances.Num();
 				}
 			}
+
+			// Remove the removed item from the removal list
+			InstancesToRemove.Remove(InstanceIndexToRemove);
 		}
 
 		InIFA->CheckSelection();
