Fix crash bug when retargetting skeletal meshes with LODs

diff --git a/Engine/Source/Runtime/Engine/Private/Animation/AnimationRuntime.cpp b/Engine/Source/Runtime/Engine/Private/Animation/AnimationRuntime.cpp
--- a/Engine/Source/Runtime/Engine/Private/Animation/AnimationRuntime.cpp
+++ b/Engine/Source/Runtime/Engine/Private/Animation/AnimationRuntime.cpp
@@ -707,12 +707,14 @@ void FAnimationRuntime::FillWithRetargetBaseRefPose( TArray<FTransform> & OutAto
 	// Copy Target Asset's ref pose.
 	if (Mesh)
 	{
+		const int32 NumBones = RequiredBones.GetNumBones();
 		const TArray<FBoneIndexType> & BoneIndices = RequiredBones.GetBoneIndicesArray();
-		OutAtoms.Empty(BoneIndices.Num());
-		OutAtoms.AddUninitialized(BoneIndices.Num());
+		OutAtoms.Empty(NumBones);
+		OutAtoms.AddUninitialized(NumBones);
 		for (int32 Idx = 0; Idx < BoneIndices.Num(); ++Idx)
 		{
-			OutAtoms[Idx] = Mesh->RetargetBasePose[BoneIndices[Idx]];
+			const int32 PoseIdx = BoneIndices[Idx];
+			OutAtoms[PoseIdx] = Mesh->RetargetBasePose[PoseIdx];
 		}
 	}
 	else
