Cache bone index on FBodyInstance to avoid lookup by FName each time

diff --git a/Engine/Source/Runtime/Engine/Classes/PhysicsEngine/BodyInstance.h b/Engine/Source/Runtime/Engine/Classes/PhysicsEngine/BodyInstance.h
--- a/Engine/Source/Runtime/Engine/Classes/PhysicsEngine/BodyInstance.h
+++ b/Engine/Source/Runtime/Engine/Classes/PhysicsEngine/BodyInstance.h
@@ -115,13 +115,15 @@ struct ENGINE_API FBodyInstance
 {
 	GENERATED_USTRUCT_BODY()
 
-
 	/** 
 	 *	Index of this BodyInstance within the SkeletalMeshComponent/PhysicsAsset. 
 	 *	Is INDEX_NONE if a single body component
 	 */
 	int32 InstanceBodyIndex;
 
+	/** When we are a body within a SkeletalMeshComponent, we cache the index of the bone we represent, to speed up sync'ing physics to anim. */
+	int32 InstanceBoneIndex;
+
 	/** Current scale of physics - used to know when and how physics must be rescaled to match current transform of OwnerComponent. */
 	UPROPERTY()
 	FVector Scale3D;
diff --git a/Engine/Source/Runtime/Engine/Private/PhysicsEngine/BodyInstance.cpp b/Engine/Source/Runtime/Engine/Private/PhysicsEngine/BodyInstance.cpp
--- a/Engine/Source/Runtime/Engine/Private/PhysicsEngine/BodyInstance.cpp
+++ b/Engine/Source/Runtime/Engine/Private/PhysicsEngine/BodyInstance.cpp
@@ -199,6 +199,7 @@ void FCollisionResponse::UpdateResponseContainerFromArray()
 
 FBodyInstance::FBodyInstance()
 : InstanceBodyIndex(INDEX_NONE)
+, InstanceBoneIndex(INDEX_NONE)
 , Scale3D(1.0f)
 , SceneIndexSync(0)
 , SceneIndexAsync(0)
diff --git a/Engine/Source/Runtime/Engine/Private/PhysicsEngine/PhysAnim.cpp b/Engine/Source/Runtime/Engine/Private/PhysicsEngine/PhysAnim.cpp
--- a/Engine/Source/Runtime/Engine/Private/PhysicsEngine/PhysAnim.cpp
+++ b/Engine/Source/Runtime/Engine/Private/PhysicsEngine/PhysAnim.cpp
@@ -328,13 +328,12 @@ void USkeletalMeshComponent::UpdateKinematicBonesToAnim(const TArray<FTransform>
 
 				if (bForceUpdate || (BodyInst->IsValidBodyInstance() && !BodyInst->IsInstanceSimulatingPhysics()))
 				{
-					// Find the graphics bone index that corresponds to this physics body.
-					FName const BodyName = PhysicsAsset->BodySetup[i]->BoneName;
-					int32 const BoneIndex = SkeletalMesh->RefSkeleton.FindBoneIndex(BodyName);
+					const int32 BoneIndex = BodyInst->InstanceBoneIndex;
 
 					// If we could not find it - warn.
 					if (BoneIndex == INDEX_NONE || BoneIndex >= GetNumSpaceBases())
 					{
+						const FName BodyName = PhysicsAsset->BodySetup[i]->BoneName;
 						UE_LOG(LogPhysics, Log, TEXT("UpdateRBBones: WARNING: Failed to find bone '%s' need by PhysicsAsset '%s' in SkeletalMesh '%s'."), *BodyName.ToString(), *PhysicsAsset->GetName(), *SkeletalMesh->GetName());
 					}
 					else
diff --git a/Engine/Source/Runtime/Engine/Private/SkeletalMeshComponentPhysics.cpp b/Engine/Source/Runtime/Engine/Private/SkeletalMeshComponentPhysics.cpp
--- a/Engine/Source/Runtime/Engine/Private/SkeletalMeshComponentPhysics.cpp
+++ b/Engine/Source/Runtime/Engine/Private/SkeletalMeshComponentPhysics.cpp
@@ -775,6 +775,7 @@ void USkeletalMeshComponent::InitArticulated(FPhysScene* PhysScene)
 			// this is not true for all other BodyInstance, but for physics assets it is true. 
 			BodyInst->bSimulatePhysics = false;
 			BodyInst->InstanceBodyIndex = i; // Set body index 
+			BodyInst->InstanceBoneIndex = BoneIndex; // Set bone index
 
 			if (i == RootBodyIndex)
 			{
