Quad tree can now duplicate and also added const correct functions

diff --git a/Engine/Source/Runtime/Engine/Public/GenericQuadTree.h b/Engine/Source/Runtime/Engine/Public/GenericQuadTree.h
--- a/Engine/Source/Runtime/Engine/Public/GenericQuadTree.h
+++ b/Engine/Source/Runtime/Engine/Public/GenericQuadTree.h
@@ -12,14 +12,17 @@ public:
 	void Insert(const ElementType& Instance, const FBox2D& Box);
 
 	/** Given a 2D box, returns an array of elements within the box. This may have duplicates based on how many quads the element overlaps*/
-	void GetElements(const FBox2D& Box, TArray<ElementType>& ElementsOut);
+	void GetElements(const FBox2D& Box, TArray<ElementType>& ElementsOut) const;
 
 	/** Given a 2D box, returns an array of elements within the box. The array will have no duplicates*/
-	void GetElementsUnique(const FBox2D& Box, TArray<ElementType>& ElementsOut);
+	void GetElementsUnique(const FBox2D& Box, TArray<ElementType>& ElementsOut) const;
 
 	/** Removes an object of type ElementType with an associated 2D box of size Box (log n). Does not cleanup tree*/
 	void Remove(const ElementType& Instance, const FBox2D& Box);
 
+	/** Does a deep copy of the tree by going through and re-creating the internal data. Cheaper than re-insertion as it should be linear instead of nlogn */
+	void Duplicate(TreeType* OutDuplicate) const;
+
 	~TQuadTree();
 
 private:
@@ -33,6 +36,7 @@ private:
 
 	/** Given a 2D box, returns an array of leaf "trees" to be used for insertion,removal, etc... */
 	void GetLeaves(const FBox2D& Box, TArray<TreeType*>& Leaves);
+	void GetLeaves(const FBox2D& Box, TArray<const TreeType*>& Leaves) const;
 
 	/** Given a 2D box, return the subtrees that are touched*/
 	int32 GetQuads(const FBox2D& Box, TreeType* Quads[4]) const;
@@ -40,7 +44,6 @@ private:
 	/** Split the tree into 4 sub-trees */
 	void Split();
 
-
 	/** Node used to hold the element and its corresponding 2D box*/
 	struct FNode
 	{
@@ -184,11 +187,11 @@ void TQuadTree<ElementType, NodeCapacity>::Remove(const ElementType& Element, co
 }
 
 template <typename ElementType, int32 NodeCapacity>
-void TQuadTree<ElementType, NodeCapacity>::GetElements(const FBox2D& Box, TArray<ElementType>& ElementsOut)
+void TQuadTree<ElementType, NodeCapacity>::GetElements(const FBox2D& Box, TArray<ElementType>& ElementsOut) const
 {
-	TArray<TQuadTree*> Leaves;
+	TArray<const TQuadTree*> Leaves;
 	GetLeaves(Box, Leaves);
-	for (TreeType* Leaf : Leaves)
+	for (const TreeType* Leaf : Leaves)
 	{
 		ElementsOut.Reserve(ElementsOut.Num() + Leaf->Nodes.Num());
 		for (const FNode& Node : Leaf->Nodes)
@@ -199,12 +202,12 @@ void TQuadTree<ElementType, NodeCapacity>::GetElements(const FBox2D& Box, TArray
 }
 
 template <typename ElementType, int32 NodeCapacity>
-void TQuadTree<ElementType, NodeCapacity>::GetElementsUnique(const FBox2D& Box, TArray<ElementType>& ElementsOut)
+void TQuadTree<ElementType, NodeCapacity>::GetElementsUnique(const FBox2D& Box, TArray<ElementType>& ElementsOut) const
 {
 	TSet<ElementType> Unique;
-	TArray<TQuadTree*> Leaves;
+	TArray<const TQuadTree*> Leaves;
 	GetLeaves(Box, Leaves);
-	for (TreeType* Leaf : Leaves)
+	for (const TreeType* Leaf : Leaves)
 	{
 		for (const FNode& Node : Leaf->Nodes)
 		{
@@ -259,15 +262,50 @@ void TQuadTree<ElementType, NodeCapacity>::GetLeaves(const FBox2D& Box, TArray<T
 	{
 		TreeType* UseNodes[4];
 		const int32 NumQuads = GetQuads(Box, UseNodes);
+		for (int32 NodeIdx = 0; NodeIdx < NumQuads; ++NodeIdx)
 		{
-			for (int32 NodeIdx = 0; NodeIdx < NumQuads; ++NodeIdx)
-			{
-				UseNodes[NodeIdx]->GetLeaves(Box, LeavesOut);
-			}
+			UseNodes[NodeIdx]->GetLeaves(Box, LeavesOut);
 		}
 	}
 	else
 	{
 		LeavesOut.Add(this);
 	}
+}
+
+template <typename ElementType, int32 NodeCapacity>
+void TQuadTree<ElementType, NodeCapacity>::GetLeaves(const FBox2D& Box, TArray<const TreeType*>& LeavesOut) const
+{
+	if (bInternal)
+	{
+		TreeType* UseNodes[4];
+		const int32 NumQuads = GetQuads(Box, UseNodes);
+		for (int32 NodeIdx = 0; NodeIdx < NumQuads; ++NodeIdx)
+		{
+			UseNodes[NodeIdx]->GetLeaves(Box, LeavesOut);
+		}
+	}
+	else
+	{
+		LeavesOut.Add(this);
+	}
+}
+
+template <typename ElementType, int32 NodeCapacity>
+void TQuadTree<ElementType, NodeCapacity>::Duplicate(TreeType* OutDuplicate) const
+{
+	for (int32 TreeIdx = 0; TreeIdx < 4; ++TreeIdx)
+	{
+		if (TreeType* SubTree = SubTrees[TreeIdx])
+		{
+			OutDuplicate->SubTrees[TreeIdx] = new TreeType(FBox2D(0,0));
+			SubTree->Duplicate(OutDuplicate->SubTrees[TreeIdx]);	//duplicate sub trees
+		}
+		
+	}
+
+	OutDuplicate->Nodes = Nodes;
+	OutDuplicate->TreeBox = TreeBox;
+	OutDuplicate->Position = Position;
+	OutDuplicate->bInternal = bInternal;
 }
\ No newline at end of file
