Added notification when running help - opens the folder where the help was dumped

diff --git a/Engine/Source/Runtime/Core/Private/Linux/LinuxPlatformProcess.cpp b/Engine/Source/Runtime/Core/Private/Linux/LinuxPlatformProcess.cpp
--- a/Engine/Source/Runtime/Core/Private/Linux/LinuxPlatformProcess.cpp
+++ b/Engine/Source/Runtime/Core/Private/Linux/LinuxPlatformProcess.cpp
@@ -801,11 +801,17 @@ bool FLinuxPlatformProcess::ExecProcess( const TCHAR* URL, const TCHAR* Params,
 
 			if (WIFEXITED(status))
 			{
-				*OutReturnCode = WEXITSTATUS(status);
+				if(OutReturnCode)
+				{
+					*OutReturnCode = WEXITSTATUS(status);
+				}
 			}
 			else if (WIFSIGNALED(status))
 			{
-				*OutReturnCode = WTERMSIG(status);
+				if(OutReturnCode)
+				{
+					*OutReturnCode = WTERMSIG(status);
+				}
 			}
 
 		} while (!WIFEXITED(status) && !WIFSIGNALED(status));
diff --git a/Engine/Source/Runtime/Core/Private/Mac/MacPlatformProcess.cpp b/Engine/Source/Runtime/Core/Private/Mac/MacPlatformProcess.cpp
--- a/Engine/Source/Runtime/Core/Private/Mac/MacPlatformProcess.cpp
+++ b/Engine/Source/Runtime/Core/Private/Mac/MacPlatformProcess.cpp
@@ -220,8 +220,14 @@ bool FMacPlatformProcess::ExecProcess( const TCHAR* URL, const TCHAR* Params, in
 	
 	if(LaunchPath == NULL)
 	{
-		*OutReturnCode = ENOENT;
-		*OutStdErr = TEXT("No such executable");
+		if(OutReturnCode)
+		{
+			*OutReturnCode = ENOENT;
+		}
+		if(OutStdErr)
+		{
+			*OutStdErr = TEXT("No such executable");
+		}
 		return false;
 	}
 	
@@ -311,8 +317,14 @@ bool FMacPlatformProcess::ExecProcess( const TCHAR* URL, const TCHAR* Params, in
 		}
 		@catch (NSException* Exc)
 		{
-			*OutReturnCode = ENOENT;
-			*OutStdErr = FString([Exc reason]);
+			if(OutReturnCode)
+			{
+				*OutReturnCode = ENOENT;
+			}
+			if(OutStdErr)
+			{
+				*OutStdErr = FString([Exc reason]);
+			}
 			return false;
 		}
 	}
diff --git a/Engine/Source/Runtime/Core/Public/GenericPlatform/GenericPlatformProcess.h b/Engine/Source/Runtime/Core/Public/GenericPlatform/GenericPlatformProcess.h
--- a/Engine/Source/Runtime/Core/Public/GenericPlatform/GenericPlatformProcess.h
+++ b/Engine/Source/Runtime/Core/Public/GenericPlatform/GenericPlatformProcess.h
@@ -353,6 +353,9 @@ struct CORE_API FGenericPlatformProcess
 	/**
 	 * Executes a process, returning the return code, stdout, and stderr. This
 	 * call blocks until the process has returned.
+	 * @param OutReturnCode may be 0
+	 * @param OutStdOut may be 0
+	 * @param OutStdErr may be 0
 	 */
 	static bool ExecProcess( const TCHAR* URL, const TCHAR* Params, int32* OutReturnCode, FString* OutStdOut, FString* OutStdErr );
 
diff --git a/Engine/Source/Runtime/Engine/Private/UnrealEngine.cpp b/Engine/Source/Runtime/Engine/Private/UnrealEngine.cpp
--- a/Engine/Source/Runtime/Engine/Private/UnrealEngine.cpp
+++ b/Engine/Source/Runtime/Engine/Private/UnrealEngine.cpp
@@ -97,6 +97,8 @@
 #include "Components/BrushComponent.h"
 #include "GameFramework/GameUserSettings.h"
 #include "PhysicalMaterials/PhysicalMaterial.h"
+#include "NotificationManager.h"
+#include "SNotificationList.h"
 
 DEFINE_LOG_CATEGORY_STATIC(LogEngine, Log, All);
 
@@ -2699,6 +2701,28 @@ static void DumpHelp(UWorld* InWorld)
 	UE_LOG(LogEngine, Display, TEXT(" "));
 
 	ConsoleCommandLibrary_DumpLibraryHTML(InWorld, *GEngine, FilePath);
+
+	// Notification in editor
+	{
+		auto Message = NSLOCTEXT("UnrealEd", "ConsoleHelpExported", "ConsoleHelp.html was saved as");
+		FNotificationInfo Info(Message);
+		Info.bFireAndForget = true;
+		Info.ExpireDuration = 5.0f;
+		Info.bUseSuccessFailIcons = false;
+		Info.bUseLargeFont = false;
+
+		const FString HyperLinkText = FPaths::ConvertRelativePathToFull(FilePath);
+		Info.Hyperlink = FSimpleDelegate::CreateStatic([](FString SourceFilePath) 
+		{
+			// open default browser, usually internet explorer (bit slower than Chrome)
+//			FPlatformProcess::LaunchURL(*SourceFilePath, TEXT(""), 0);
+			// open folder, you can choose the browser yourself
+			FPlatformProcess::ExploreFolder(*(FPaths::GetPath(SourceFilePath)));
+		}, HyperLinkText);
+		Info.HyperlinkText = FText::FromString(HyperLinkText);
+
+		FSlateNotificationManager::Get().AddNotification(Info);
+	}
 }
 static FAutoConsoleCommandWithWorld GConsoleCommandHelp(
 	TEXT("help"),
