#ue4 - Fix incorrect comment on "bShouldUpdatePhysicsVolume".

diff --git a/Engine/Source/Runtime/Engine/Classes/Components/SceneComponent.h b/Engine/Source/Runtime/Engine/Classes/Components/SceneComponent.h
--- a/Engine/Source/Runtime/Engine/Classes/Components/SceneComponent.h
+++ b/Engine/Source/Runtime/Engine/Classes/Components/SceneComponent.h
@@ -126,7 +126,10 @@ public:
 	UPROPERTY(EditDefaultsOnly, BlueprintReadOnly, Category=Rendering)
 	uint32 bHiddenInGame:1;
 
-	/** Whether or not PhysicsVolume should be updated when the component is moved. For PrimitiveComponents, bGenerateOverlapEvents must be true as well. */
+	/**
+	 * Whether or not the cached PhysicsVolume this component overlaps should be updated when the component is moved.
+	 * @see GetPhysicsVolume()
+	 */
 	UPROPERTY(EditAnywhere, BlueprintReadWrite, Category=PhysicsVolume)
 	uint32 bShouldUpdatePhysicsVolume:1;
 
