#ue4 - Reserve array space for raycast results before adding them one-by-one.

diff --git a/Engine/Source/Runtime/Engine/Private/Collision/CollisionConversions.cpp b/Engine/Source/Runtime/Engine/Private/Collision/CollisionConversions.cpp
--- a/Engine/Source/Runtime/Engine/Private/Collision/CollisionConversions.cpp
+++ b/Engine/Source/Runtime/Engine/Private/Collision/CollisionConversions.cpp
@@ -419,6 +419,8 @@ void ConvertQueryImpactHit(const PxLocationHit& PHit, FHitResult& OutResult, flo
 
 void ConvertRaycastResults(int32 NumHits, PxRaycastHit* Hits, float CheckLength, const PxFilterData& QueryFilter, TArray<FHitResult>& OutHits, const FVector& StartLoc, const FVector& EndLoc, bool bReturnFaceIndex, bool bReturnPhysMat)
 {
+	OutHits.Reserve(NumHits);
+
 	PxTransform PStartTM(U2PVector(StartLoc));
 	for(int32 i=0; i<NumHits; i++)
 	{
