#ue4 - Add stat tracking to all PhysX lock/unlock calls, not just the scoped locks.

diff --git a/Engine/Source/Runtime/Engine/Private/PhysicsEngine/PhysXSupport.h b/Engine/Source/Runtime/Engine/Private/PhysicsEngine/PhysXSupport.h
--- a/Engine/Source/Runtime/Engine/Private/PhysicsEngine/PhysXSupport.h
+++ b/Engine/Source/Runtime/Engine/Private/PhysicsEngine/PhysXSupport.h
@@ -83,10 +83,10 @@ private:
 #define SCOPED_SCENE_READ_LOCK( _scene ) FPhysXSceneReadLock SCOPED_SCENE_TOKENPASTE(_rlock,__LINE__)(_scene)
 #define SCOPED_SCENE_WRITE_LOCK( _scene ) FPhysXSceneWriteLock SCOPED_SCENE_TOKENPASTE(_wlock,__LINE__)(_scene)
 
-#define SCENE_LOCK_READ( _scene ) if((_scene) != NULL) { (_scene)->lockRead(); }
-#define SCENE_UNLOCK_READ( _scene ) if((_scene) != NULL) { (_scene)->unlockRead(); }
-#define SCENE_LOCK_WRITE( _scene ) if((_scene) != NULL) { (_scene)->lockWrite(); }
-#define SCENE_UNLOCK_WRITE( _scene ) if((_scene) != NULL) { (_scene)->unlockWrite(); }
+#define SCENE_LOCK_READ( _scene )		{ SCOPE_CYCLE_COUNTER(STAT_PhysSceneLock); if((_scene) != NULL) { (_scene)->lockRead(); } }
+#define SCENE_UNLOCK_READ( _scene )		{ SCOPE_CYCLE_COUNTER(STAT_PhysSceneLock); if((_scene) != NULL) { (_scene)->unlockRead(); } }
+#define SCENE_LOCK_WRITE( _scene )		{ SCOPE_CYCLE_COUNTER(STAT_PhysSceneLock); if((_scene) != NULL) { (_scene)->lockWrite(); } }
+#define SCENE_UNLOCK_WRITE( _scene )	{ SCOPE_CYCLE_COUNTER(STAT_PhysSceneLock); if((_scene) != NULL) { (_scene)->unlockWrite(); } }
 #else
 #define SCOPED_SCENE_READ_LOCK_INDEXED( _scene, _index )
 #define SCOPED_SCENE_READ_LOCK( _scene )
