#ue4 - Limit number of PhysicsVolumes we test in USceneComponent::UpdatePhysicsVolume() in the pre-process pass before using the full physics scene query.

diff --git a/Engine/Source/Runtime/Engine/Private/PrimitiveComponent.cpp b/Engine/Source/Runtime/Engine/Private/PrimitiveComponent.cpp
--- a/Engine/Source/Runtime/Engine/Private/PrimitiveComponent.cpp
+++ b/Engine/Source/Runtime/Engine/Private/PrimitiveComponent.cpp
@@ -2335,7 +2335,10 @@ void UPrimitiveComponent::UpdatePhysicsVolume( bool bTriggerNotifiers )
 				}
 			}
 
-			SetPhysicsVolume(BestVolume, bTriggerNotifiers);
+			if (BestVolume != PhysicsVolume)
+			{
+				SetPhysicsVolume(BestVolume, bTriggerNotifiers);
+			}
 		}
 		else
 		{
diff --git a/Engine/Source/Runtime/Engine/Private/SceneComponent.cpp b/Engine/Source/Runtime/Engine/Private/SceneComponent.cpp
--- a/Engine/Source/Runtime/Engine/Private/SceneComponent.cpp
+++ b/Engine/Source/Runtime/Engine/Private/SceneComponent.cpp
@@ -1222,8 +1222,10 @@ void USceneComponent::UpdatePhysicsVolume( bool bTriggerNotifiers )
 		if (MyWorld->GetNonDefaultPhysicsVolumeCount() > 0)
 		{
 			// Avoid a full overlap query if we can do some quick bounds tests against the volumes.
+			static uint32 MaxVolumesToCheck = 100;
+			uint32 VolumeIndex = 0;
 			bool bAnyPotentialOverlap = false;
-			for (auto VolumeIter = MyWorld->GetNonDefaultPhysicsVolumeIterator(); VolumeIter && !bAnyPotentialOverlap; ++VolumeIter)
+			for (auto VolumeIter = MyWorld->GetNonDefaultPhysicsVolumeIterator(); VolumeIter && !bAnyPotentialOverlap; ++VolumeIter, ++VolumeIndex)
 			{
 				const APhysicsVolume* Volume = *VolumeIter;
 				const USceneComponent* VolumeRoot = Volume->GetRootComponent();
@@ -1237,7 +1239,13 @@ void USceneComponent::UpdatePhysicsVolume( bool bTriggerNotifiers )
 						}
 					}
 				}
-				// TODO: bail if too many volumes?				
+
+				// Bail if too many volumes. Later we'll probably convert to using an octree so this wouldn't be a concern.
+				if (VolumeIndex >= MaxVolumesToCheck)
+				{
+					bAnyPotentialOverlap = true;
+					break;
+				}
 			}
 			
 			if (bAnyPotentialOverlap)
@@ -1274,7 +1282,10 @@ void USceneComponent::UpdatePhysicsVolume( bool bTriggerNotifiers )
 			}
 		}
 
-		SetPhysicsVolume( NewVolume, bTriggerNotifiers );
+		if (PhysicsVolume != NewVolume)
+		{
+			SetPhysicsVolume( NewVolume, bTriggerNotifiers );
+		}
 	}
 }
 
