Prevent a crash when force deleting certain groups of assets which contain maps and actors.

diff --git a/Engine/Source/Runtime/Engine/Private/ActorEditor.cpp b/Engine/Source/Runtime/Engine/Private/ActorEditor.cpp
--- a/Engine/Source/Runtime/Engine/Private/ActorEditor.cpp
+++ b/Engine/Source/Runtime/Engine/Private/ActorEditor.cpp
@@ -145,7 +145,7 @@ void AActor::PostEditMove(bool bFinished)
 
 bool AActor::ReregisterComponentsWhenModified() const
 {
-	return !IsTemplate() && ( GetOutermost()->PackageFlags & PKG_PlayInEditor ) == 0;
+	return !IsTemplate() && ( GetOutermost()->PackageFlags & PKG_PlayInEditor ) == 0 && GetWorld() != nullptr;
 }
 
 void AActor::DebugShowComponentHierarchy(  const TCHAR* Info, bool bShowPosition )
diff --git a/Engine/Source/Runtime/Engine/Private/WorldSettings.cpp b/Engine/Source/Runtime/Engine/Private/WorldSettings.cpp
--- a/Engine/Source/Runtime/Engine/Private/WorldSettings.cpp
+++ b/Engine/Source/Runtime/Engine/Private/WorldSettings.cpp
@@ -252,7 +252,7 @@ void AWorldSettings::PostEditChangeProperty(FPropertyChangedEvent& PropertyChang
 	// Ensure texture size is power of two between 512 and 4096.
 	PackedLightAndShadowMapTextureSize = FMath::Clamp<uint32>( FMath::RoundUpToPowerOfTwo( PackedLightAndShadowMapTextureSize ), 512, 4096 );
 
-	if (GetWorld()->PersistentLevel->GetWorldSettings() == this)
+	if (GetWorld() != nullptr && GetWorld()->PersistentLevel->GetWorldSettings() == this)
 	{
 		if (GIsEditor)
 		{
