#ue4 - AActor::GetComponents() can now support TArrays with different allocator types.

- It's now recommended to use TArrays with a TInlineAllocator to potentially avoid memory allocation costs.
- TInlineComponentArray is defined to make this easier, for example:
{
    TInlineComponentArray<UPrimitiveComponent*> PrimComponents;
    Actor->GetComponents(PrimComponents);
}

diff --git a/Engine/Source/Runtime/Engine/Classes/Engine/EngineTypes.h b/Engine/Source/Runtime/Engine/Classes/Engine/EngineTypes.h
--- a/Engine/Source/Runtime/Engine/Classes/Engine/EngineTypes.h
+++ b/Engine/Source/Runtime/Engine/Classes/Engine/EngineTypes.h
@@ -10,6 +10,24 @@
 #include "GameFramework/DamageType.h"
 #include "EngineTypes.generated.h"
 
+
+
+/**
+ * Default number of components to expect in TInlineAllocators used with AActor component arrays.
+ * Used by engine code to try to avoid allocations in AActor::GetComponents(), among others.
+ */
+enum { NumInlinedActorComponents = 24 };
+
+/**
+ * TInlineComponentArray is simply a TArray that reserves a fixed amount of space on the stack
+ * to try to avoid heap allocation when there are fewer than a specified number of elements expected in the result.
+ */
+template<class T, uint32 NumElements = NumInlinedActorComponents>
+class TInlineComponentArray : public TArray<T, TInlineAllocator<NumElements>>
+{
+};
+
+
 UENUM()
 enum EAspectRatioAxisConstraint
 {
diff --git a/Engine/Source/Runtime/Engine/Classes/GameFramework/Actor.h b/Engine/Source/Runtime/Engine/Classes/GameFramework/Actor.h
--- a/Engine/Source/Runtime/Engine/Classes/GameFramework/Actor.h
+++ b/Engine/Source/Runtime/Engine/Classes/GameFramework/Actor.h
@@ -2121,8 +2121,17 @@ public:
 		return (T*)FindComponentByClass(T::StaticClass());
 	}
 
-	template<class T>
-	void GetComponents(TArray<T*>& OutComponents) const
+	/**
+	 * Get all components derived from class 'T' and fill in the OutComponents array with the result.
+	 * It's recommended to use TArrays with a TInlineAllocator to potentially avoid memory allocation costs.
+	 * TInlineComponentArray is defined to make this easier, for example:
+	 * {
+	 * 	   TInlineComponentArray<UPrimitiveComponent*> PrimComponents;
+	 *     Actor->GetComponents(PrimComponents);
+	 * }
+	 */
+	template<class T, class AllocatorType = FDefaultAllocator>
+	void GetComponents(TArray<T*, AllocatorType>& OutComponents) const
 	{
 		static_assert(CanConvertPointerFromTo<T, UActorComponent>::Result, "'T' template parameter to GetComponents must be derived from ActorComponent");
 
@@ -2140,8 +2149,9 @@ public:
 		}
 	}
 
-	// UActorComponent specialization to avoid unnecessary casts
-	void GetComponents(TArray<UActorComponent*>& OutComponents) const
+	/** UActorComponent specialization of GetComponents() to avoid unnecessary casts. */
+	template<class AllocatorType = FDefaultAllocator>
+	void GetComponents(TArray<UActorComponent*, AllocatorType>& OutComponents) const
 	{
 		SCOPE_CYCLE_COUNTER(STAT_GetComponentsTime);
 
diff --git a/Engine/Source/Runtime/Engine/Private/Actor.cpp b/Engine/Source/Runtime/Engine/Private/Actor.cpp
--- a/Engine/Source/Runtime/Engine/Private/Actor.cpp
+++ b/Engine/Source/Runtime/Engine/Private/Actor.cpp
@@ -129,7 +129,7 @@ bool AActor::CheckActorComponents()
 	DEFINE_LOG_CATEGORY_STATIC(LogCheckComponents, Warning, All);
 
 	bool bResult = true;
-	TArray<UActorComponent*> Components;
+	TInlineComponentArray<UActorComponent*> Components;
 	GetComponents(Components);
 
 	for (int32 Index = 0; Index < Components.Num(); Index++)
@@ -565,7 +565,7 @@ void AActor::RegisterAllActorTickFunctions(bool bRegister, bool bDoComponents)
 
 		if (bDoComponents)
 		{
-			TArray<UActorComponent*> Components;
+			TInlineComponentArray<UActorComponent*> Components;
 			GetComponents(Components);
 
 			for (int32 Index = 0; Index < Components.Num(); Index++)
@@ -662,7 +662,7 @@ void AActor::Tick( float DeltaSeconds )
 		bool bOKToDestroy = true;
 
 		// @todo: naive implementation, needs improved
-		TArray<UActorComponent*> Components;
+		TInlineComponentArray<UActorComponent*> Components;
 		GetComponents(Components);
 
 		for (int32 CompIdx=0; CompIdx<Components.Num(); ++CompIdx)
@@ -748,7 +748,7 @@ void AActor::GetComponentsBoundingCylinder(float& OutCollisionRadius, float& Out
 	float Radius = 0.f;
 	float HalfHeight = 0.f;
 
-	TArray<UPrimitiveComponent*> Components;
+	TInlineComponentArray<UPrimitiveComponent*> Components;
 	GetComponents(Components);
 
 	for (int32 CompIdx = 0; CompIdx < Components.Num(); CompIdx++)
@@ -839,7 +839,7 @@ FBox AActor::GetComponentsBoundingBox(bool bNonColliding) const
 {
 	FBox Box(0);
 
-	TArray<UPrimitiveComponent*> Components;
+	TInlineComponentArray<UPrimitiveComponent*> Components;
 	GetComponents(Components);
 
 	for(int32 CompIdx = 0; CompIdx < Components.Num(); CompIdx++)
@@ -899,7 +899,7 @@ void AActor::SetTickGroup(ETickingGroup NewTickGroup)
 
 void AActor::ClearComponentOverlaps()
 {
-	TArray<UPrimitiveComponent*> PrimitiveComponents;
+	TInlineComponentArray<UPrimitiveComponent*> PrimitiveComponents;
 	GetComponents(PrimitiveComponents);
 
 	// Remove owned components from overlap tracking
@@ -1035,7 +1035,7 @@ void AActor::GetOverlappingComponents(TArray<UPrimitiveComponent*>& OutOverlappi
  */
 static void MarkOwnerRelevantComponentsDirty(AActor* TheActor)
 {
-	TArray<UPrimitiveComponent*> Components;
+	TInlineComponentArray<UPrimitiveComponent*> Components;
 	TheActor->GetComponents(Components);
 
 	for (int32 i = 0; i < Components.Num(); i++)
@@ -1062,7 +1062,7 @@ float AActor::GetLastRenderTime() const
 {
 	// return most recent of Components' LastRenderTime
 	// @todo UE4 maybe check base component and components attached to it instead?
-	TArray<UPrimitiveComponent*> Components;
+	TInlineComponentArray<UPrimitiveComponent*> Components;
 	GetComponents(Components);
 
 	float LastRenderTime = -1000.f;
@@ -1237,7 +1237,7 @@ void AActor::DetachSceneComponentsFromParent(USceneComponent* InParentComponent,
 {
 	if (InParentComponent != NULL)
 	{
-		TArray<USceneComponent*> Components;
+		TInlineComponentArray<USceneComponent*> Components;
 		GetComponents(Components);
 
 		for (int32 Index = 0; Index < Components.Num(); ++Index)
@@ -1291,10 +1291,10 @@ void AActor::GetAttachedActors(TArray<class AActor*>& OutActors) const
 	if (RootComponent != NULL)
 	{
 		// Current set of components to check
-		TArray< USceneComponent*, TInlineAllocator<8> > CompsToCheck;
+		TArray< USceneComponent*, TInlineAllocator<NumInlinedActorComponents> > CompsToCheck;
 
 		// Set of all components we have checked
-		TArray< USceneComponent*, TInlineAllocator<8> > CheckedComps;
+		TArray< USceneComponent*, TInlineAllocator<NumInlinedActorComponents> > CheckedComps;
 
 		CompsToCheck.Push(RootComponent);
 
@@ -1495,7 +1495,7 @@ void AActor::PrestreamTextures( float Seconds, bool bEnableStreaming, int32 Cine
 	}
 
 	// Iterate over all components of that actor
-	TArray<UMeshComponent*> Components;
+	TInlineComponentArray<UMeshComponent*> Components;
 	GetComponents(Components);
 
 	for (int32 ComponentIndex=0; ComponentIndex < Components.Num(); ComponentIndex++)
@@ -1554,7 +1554,7 @@ FVector AActor::GetPlacementExtent() const
 	FVector Extent(0.f);
 	if( (RootComponent && GetRootComponent()->ShouldCollideWhenPlacing()) && bCollideWhenPlacing) 
 	{
-		TArray<USceneComponent*> Components;
+		TInlineComponentArray<USceneComponent*> Components;
 		GetComponents(Components);
 
 		FBox ActorBox(0.f);
@@ -1876,7 +1876,7 @@ void AActor::DisplayDebug(UCanvas* Canvas, const FDebugDisplayInfo& DebugDisplay
 	static FName NAME_Bones = FName(TEXT("Bones"));
 	if (DebugDisplay.IsDisplayOn(NAME_Animation) || DebugDisplay.IsDisplayOn(NAME_Bones))
 	{
-		TArray<USkeletalMeshComponent*> Components;
+		TInlineComponentArray<USkeletalMeshComponent*> Components;
 		GetComponents(Components);
 
 		if (DebugDisplay.IsDisplayOn(NAME_Animation))
@@ -1938,7 +1938,7 @@ void AActor::CalcCamera(float DeltaTime, FMinimalViewInfo& OutResult)
 	if (bFindCameraComponentWhenViewTarget)
 	{
 		// Look for the first active camera component and use that for the view
-		TArray<UCameraComponent*> Cameras;
+		TInlineComponentArray<UCameraComponent*> Cameras;
 		GetComponents<UCameraComponent>(/*out*/ Cameras);
 
 		for (UCameraComponent* CameraComponent : Cameras)
@@ -1984,7 +1984,7 @@ enum ECollisionResponse AActor::GetComponentsCollisionResponseToChannel(enum ECo
 {
 	ECollisionResponse OutResponse = ECR_Ignore;
 
-	TArray<UPrimitiveComponent*> Components;
+	TInlineComponentArray<UPrimitiveComponent*> Components;
 	GetComponents(Components);
 
 	for (int32 i = 0; i < Components.Num(); i++)
@@ -2102,7 +2102,7 @@ TArray<UActorComponent*> AActor::GetComponentsByClass(TSubclassOf<UActorComponen
 
 void AActor::DisableComponentsSimulatePhysics()
 {
-	TArray<UPrimitiveComponent*> Components;
+	TInlineComponentArray<UPrimitiveComponent*> Components;
 	GetComponents(Components);
 
 	for (UPrimitiveComponent* Component : Components)
@@ -2118,7 +2118,7 @@ void AActor::PostRegisterAllComponents()
 /** Util to call OnComponentCreated on components */
 static void DispatchOnComponentsCreated(AActor* NewActor)
 {
-	TArray<UActorComponent*> Components;
+	TInlineComponentArray<UActorComponent*> Components;
 	NewActor->GetComponents(Components);
 
 	for (int32 Idx = 0; Idx < Components.Num(); Idx++)
@@ -2655,7 +2655,7 @@ void AActor::SetActorEnableCollision(bool bNewActorEnableCollision)
 		bActorEnableCollision = bNewActorEnableCollision;
 
 		// Notify components about the change
-		TArray<UActorComponent*> Components;
+		TInlineComponentArray<UActorComponent*> Components;
 		GetComponents(Components);
 
 		for(int32 CompIdx=0; CompIdx<Components.Num(); CompIdx++)
@@ -3060,7 +3060,7 @@ void AActor::DispatchPhysicsCollisionHit(const FRigidBodyCollisionInfo& MyInfo,
 
 void AActor::UnregisterAllComponents()
 {
-	TArray<UActorComponent*> Components;
+	TInlineComponentArray<UActorComponent*> Components;
 	GetComponents(Components);
 
 	for(int32 CompIdx = 0; CompIdx < Components.Num(); CompIdx++)
@@ -3129,7 +3129,7 @@ bool AActor::IncrementalRegisterComponents(int32 NumComponentsToRegister)
 
 	int32 NumTotalRegisteredComponents = 0;
 	int32 NumRegisteredComponentsThisRun = 0;
-	TArray<UActorComponent*> Components;
+	TInlineComponentArray<UActorComponent*> Components;
 	GetComponents(Components);
 	
 	for (int32 CompIdx = 0; CompIdx < Components.Num() && NumRegisteredComponentsThisRun < NumComponentsToRegister; CompIdx++)
@@ -3179,7 +3179,7 @@ bool AActor::HasValidRootComponent()
 void AActor::MarkComponentsAsPendingKill()
 {
 	// Iterate components and mark them all as pending kill.
-	TArray<UActorComponent*> Components;
+	TInlineComponentArray<UActorComponent*> Components;
 	GetComponents(Components);
 
 	for (int32 Index = 0; Index < Components.Num(); Index++)
@@ -3204,7 +3204,7 @@ void AActor::ReregisterAllComponents()
 
 void AActor::UpdateComponentTransforms()
 {
-	TArray<UActorComponent*> Components;
+	TInlineComponentArray<UActorComponent*> Components;
 	GetComponents(Components);
 
 	for (int32 Idx = 0; Idx < Components.Num(); Idx++)
@@ -3219,7 +3219,7 @@ void AActor::UpdateComponentTransforms()
 
 void AActor::MarkComponentsRenderStateDirty()
 {
-	TArray<UActorComponent*> Components;
+	TInlineComponentArray<UActorComponent*> Components;
 	GetComponents(Components);
 
 	for (int32 Idx = 0; Idx < Components.Num(); Idx++)
@@ -3234,7 +3234,7 @@ void AActor::MarkComponentsRenderStateDirty()
 
 void AActor::InitializeComponents()
 {
-	TArray<UActorComponent*> Components;
+	TInlineComponentArray<UActorComponent*> Components;
 	GetComponents(Components);
 
 	for (UActorComponent* ActorComp : Components)
@@ -3248,7 +3248,7 @@ void AActor::InitializeComponents()
 
 void AActor::UninitializeComponents()
 {
-	TArray<UActorComponent*> Components;
+	TInlineComponentArray<UActorComponent*> Components;
 	GetComponents(Components);
 
 	for (UActorComponent* ActorComp : Components)
@@ -3262,7 +3262,7 @@ void AActor::UninitializeComponents()
 
 void AActor::DrawDebugComponents(FColor const& BaseColor) const
 {
-	TArray<USceneComponent*> Components;
+	TInlineComponentArray<USceneComponent*> Components;
 	GetComponents(Components);
 
 	for(int32 ComponentIndex = 0; ComponentIndex < Components.Num(); ComponentIndex++)
@@ -3291,7 +3291,7 @@ void AActor::DrawDebugComponents(FColor const& BaseColor) const
 
 void AActor::InvalidateLightingCacheDetailed(bool bTranslationOnly)
 {
-	TArray<UActorComponent*> Components;
+	TInlineComponentArray<UActorComponent*> Components;
 	GetComponents(Components);
 
 	for(int32 ComponentIndex = 0;ComponentIndex < Components.Num();ComponentIndex++)
@@ -3306,7 +3306,7 @@ void AActor::InvalidateLightingCacheDetailed(bool bTranslationOnly)
 
 #if !(UE_BUILD_SHIPPING || UE_BUILD_TEST)
 	// Validate that we didn't change it during this action
-	TArray<UActorComponent*> NewComponents;
+	TInlineComponentArray<UActorComponent*> NewComponents;
 	GetComponents(NewComponents);
 	check(Components == NewComponents);
 #endif
@@ -3321,7 +3321,7 @@ bool AActor::ActorLineTraceSingle(struct FHitResult& OutHit, const FVector& Star
 	OutHit.TraceEnd = End;
 	bool bHasHit = false;
 	
-	TArray<UPrimitiveComponent*> Components;
+	TInlineComponentArray<UPrimitiveComponent*> Components;
 	GetComponents(Components);
 
 	for (int32 ComponentIndex=0; ComponentIndex<Components.Num(); ComponentIndex++)
@@ -3349,7 +3349,7 @@ float AActor::ActorGetDistanceToCollision(const FVector& Point, ECollisionChanne
 	ClosestPointOnCollision = Point;
 	float ClosestPointDistance = -1.f;
 
-	TArray<UPrimitiveComponent*> Components;
+	TInlineComponentArray<UPrimitiveComponent*> Components;
 	GetComponents(Components);
 
 	for (int32 ComponentIndex=0; ComponentIndex<Components.Num(); ComponentIndex++)
