#ue4 - Optimize collision queries to reduce dynamic allocations in ignored component list.

diff --git a/Engine/Source/Runtime/Engine/Private/Collision/Collision.cpp b/Engine/Source/Runtime/Engine/Private/Collision/Collision.cpp
--- a/Engine/Source/Runtime/Engine/Private/Collision/Collision.cpp
+++ b/Engine/Source/Runtime/Engine/Private/Collision/Collision.cpp
@@ -133,11 +133,23 @@ void FCollisionQueryParams::AddIgnoredActor(const AActor* InIgnoreActor)
 {
 	if (InIgnoreActor)
 	{
-		TArray<UPrimitiveComponent*> PrimComponents;
+		TInlineComponentArray<UPrimitiveComponent*> PrimComponents;
 		InIgnoreActor->GetComponents(PrimComponents);
-		for (UPrimitiveComponent* PrimComponent : PrimComponents)
+
+		if (IgnoreComponents.Num() == 0)
 		{
-			IgnoreComponents.AddUnique(PrimComponent->GetUniqueID());
+			// Don't need to AddUnique if we start empty, as GetComponents should never contain duplicate components.
+			for (UPrimitiveComponent* PrimComponent : PrimComponents)
+			{
+				IgnoreComponents.Add(PrimComponent->GetUniqueID());
+			}
+		}
+		else
+		{
+			for (UPrimitiveComponent* PrimComponent : PrimComponents)
+			{
+				IgnoreComponents.AddUnique(PrimComponent->GetUniqueID());
+			}
 		}
 	}
 }
diff --git a/Engine/Source/Runtime/Engine/Private/Collision/PhysXCollision.h b/Engine/Source/Runtime/Engine/Private/Collision/PhysXCollision.h
--- a/Engine/Source/Runtime/Engine/Private/Collision/PhysXCollision.h
+++ b/Engine/Source/Runtime/Engine/Private/Collision/PhysXCollision.h
@@ -74,13 +74,20 @@ PxGeometry * GetGeometryFromShape(GeometryFromShapeStorage & LocalStorage, const
 
 // FILTER
 
+/** TArray typedef of components to ignore. */
+typedef TArray<uint32, TInlineAllocator<NumInlinedActorComponents>> FilterIgnoreComponentsArrayType;
+
+
 /** Unreal PhysX scene query filter callback object */
 class FPxQueryFilterCallback : public PxSceneQueryFilterCallback
 {
 public:
+
 	/** List of ActorIds for this query to ignore */
-	TArray<uint32, TInlineAllocator<1> >	IgnoreComponents;
-	PxSceneQueryHitType::Enum				PrefilterReturnValue;
+	FilterIgnoreComponentsArrayType IgnoreComponents;
+	
+	/** Result of PreFilter callback. */
+	PxSceneQueryHitType::Enum PrefilterReturnValue;
 
 	/** Whether this is a raycastSingle or sweepSingle, because in 3.3 you can't return eTOUCH for single queries */
 	bool bSingleQuery;
@@ -92,7 +99,7 @@ public:
 		bSingleQuery = false;
 	}
 
-	FPxQueryFilterCallback(const TArray<uint32, TInlineAllocator<1> > InIgnoreComponents)
+	FPxQueryFilterCallback(const FilterIgnoreComponentsArrayType& InIgnoreComponents)
 	{
 		PrefilterReturnValue = PxSceneQueryHitType::eNONE;
 		
@@ -126,7 +133,7 @@ class FPxQueryFilterCallbackSweep : public FPxQueryFilterCallback
 public:
 	bool DiscardInitialOverlaps;
 
-	FPxQueryFilterCallbackSweep(const TArray<uint32, TInlineAllocator<1> > InIgnoreComponents)
+	FPxQueryFilterCallbackSweep(const FilterIgnoreComponentsArrayType& InIgnoreComponents)
 		: FPxQueryFilterCallback(InIgnoreComponents)
 	{
 		DiscardInitialOverlaps = false;
diff --git a/Engine/Source/Runtime/Engine/Public/CollisionQueryParams.h b/Engine/Source/Runtime/Engine/Public/CollisionQueryParams.h
--- a/Engine/Source/Runtime/Engine/Public/CollisionQueryParams.h
+++ b/Engine/Source/Runtime/Engine/Public/CollisionQueryParams.h
@@ -35,8 +35,12 @@ struct ENGINE_API FCollisionQueryParams
 	/** Only fill in the PhysMaterial field of  */
 	bool bReturnPhysicalMaterial;
 
+	/** TArray typedef of components to ignore. */
+	typedef TArray<uint32, TInlineAllocator<NumInlinedActorComponents>> IgnoreComponentsArrayType;
+
 	/** Set of components to ignore during the trace */
-	TArray<uint32, TInlineAllocator<1> > IgnoreComponents;
+	IgnoreComponentsArrayType IgnoreComponents;
+
 
 	// Constructors
 	FCollisionQueryParams(bool bInTraceComplex=false)
@@ -65,10 +69,10 @@ struct ENGINE_API FCollisionQueryParams
 	/** Add a component for this trace to ignore */
 	void AddIgnoredComponent(const UPrimitiveComponent* InIgnoreComponent);
 
-	/** Add a collection of componnets for this trace to ignore */
+	/** Add a collection of components for this trace to ignore */
 	void AddIgnoredComponents(const TArray<UPrimitiveComponent*>& InIgnoreComponents);
 	
-	/** Varient that uses an array of TWeakObjectPtrs */
+	/** Variant that uses an array of TWeakObjectPtrs */
 	void AddIgnoredComponents(const TArray<TWeakObjectPtr<UPrimitiveComponent>>& InIgnoreComponents);
 
 	FString ToString() const
