#ue4 - Removed the "lock-per-shape" and redundant lock/unlocks in ComponentSweepMulti.

diff --git a/Engine/Source/Runtime/Engine/Private/Collision/WorldCollision.cpp b/Engine/Source/Runtime/Engine/Private/Collision/WorldCollision.cpp
--- a/Engine/Source/Runtime/Engine/Private/Collision/WorldCollision.cpp
+++ b/Engine/Source/Runtime/Engine/Private/Collision/WorldCollision.cpp
@@ -458,11 +458,11 @@ bool UWorld::ComponentSweepMulti(TArray<struct FHitResult>& OutHits, class UPrim
 	if (PxRigidActor* PRigidActor = PrimComp->BodyInstance.GetPxRigidActor())
 	{
 		PxScene* const PScene = PRigidActor->getScene();
+		SCOPED_SCENE_READ_LOCK(PScene);
 
 		// Get all the shapes from the actor
-		TArray<PxShape*, TInlineAllocator<8>> PShapes;
+		TArray<PxShape*, TInlineAllocator<32>> PShapes;
 		{
-			SCOPED_SCENE_READ_LOCK(PScene);
 			PShapes.AddZeroed(PRigidActor->getNbShapes());
 			PRigidActor->getShapes(PShapes.GetData(), PShapes.Num());
 		}
@@ -474,7 +474,6 @@ bool UWorld::ComponentSweepMulti(TArray<struct FHitResult>& OutHits, class UPrim
 		PxQuat PGeomRot = U2PQuat(Rot.Quaternion());
 
 		// Iterate over each shape
-		SCENE_LOCK_READ(PScene);
 		for(int32 ShapeIdx=0; ShapeIdx<PShapes.Num(); ShapeIdx++)
 		{
 			PxShape* PShape = PShapes[ShapeIdx];
@@ -493,17 +492,14 @@ bool UWorld::ComponentSweepMulti(TArray<struct FHitResult>& OutHits, class UPrim
 
 			if(PGeom != NULL)
 			{
-				SCENE_UNLOCK_READ(PScene);
 				if (GeomSweepMulti_PhysX(this, *PGeom, PShapeRot, Hits, P2UVector(PShapeGlobalStartPose.p), P2UVector(PShapeGlobalEndPose.p), TraceChannel, Params, FCollisionResponseParams(PrimComp->GetCollisionResponseToChannels())))
 				{
 					bHaveBlockingHit = true;
 				}
 
 				OutHits.Append(Hits);
-				SCENE_LOCK_READ(PScene);
 			}
 		}
-		SCENE_UNLOCK_READ(PScene);
 	}
 #endif //WITH_PHYSX
 
