Add spawner hash for foliage actor

diff --git a/Engine/Source/Runtime/Engine/Classes/Foliage/InstancedFoliageActor.h b/Engine/Source/Runtime/Engine/Classes/Foliage/InstancedFoliageActor.h
--- a/Engine/Source/Runtime/Engine/Classes/Foliage/InstancedFoliageActor.h
+++ b/Engine/Source/Runtime/Engine/Classes/Foliage/InstancedFoliageActor.h
@@ -62,6 +62,9 @@ public:
 	// Deletes the instances attached to a component
 	ENGINE_API void DeleteInstancesForComponent(UActorComponent* InComponent);
 
+	// Deletes the instances spawned by a component
+	ENGINE_API void DeleteInstancesForSpawner(UActorComponent* InComponent);
+
 	// Finds a mesh entry
 	ENGINE_API FFoliageMeshInfo* FindMesh(UFoliageType* InType);
 
diff --git a/Engine/Source/Runtime/Engine/Private/InstancedFoliage.cpp b/Engine/Source/Runtime/Engine/Private/InstancedFoliage.cpp
--- a/Engine/Source/Runtime/Engine/Private/InstancedFoliage.cpp
+++ b/Engine/Source/Runtime/Engine/Private/InstancedFoliage.cpp
@@ -353,6 +353,13 @@ void FFoliageMeshInfo::AddInstance(AInstancedFoliageActor* InIFA, UFoliageType*
 		ComponentHashInfo = &ComponentHash.Add(InNewInstance.Base, FFoliageComponentHashInfo(InNewInstance.Base));
 	}
 	ComponentHashInfo->Instances.Add(InstanceIndex);
+	
+	FFoliageComponentHashInfo* SpawnerHashInfo = SpawnerHash.Find(InNewInstance.Spawner);
+	if (SpawnerHashInfo == nullptr)
+	{
+		SpawnerHashInfo = &SpawnerHash.Add(InNewInstance.Spawner, FFoliageComponentHashInfo(InNewInstance.Spawner));
+	}
+	SpawnerHashInfo->Instances.Add(InstanceIndex);
 
 	// Calculate transform for the instance
 	FTransform InstanceToWorld = InNewInstance.GetInstanceWorldTransform();
@@ -404,6 +411,17 @@ void FFoliageMeshInfo::RemoveInstances(AInstancedFoliageActor* InIFA, const TArr
 			    }
 		    }
 
+			FFoliageComponentHashInfo* SpawnerHashInfo = SpawnerHash.Find(Instance.Spawner);
+			if (SpawnerHashInfo)
+			{
+				SpawnerHashInfo->Instances.Remove(InstanceIndex);
+				if (SpawnerHashInfo->Instances.Num() == 0)
+				{
+					// Remove the component from the spawner hash if this is the last instance.
+					SpawnerHash.Remove(Instance.Spawner);
+				}
+			}
+
 			// remove from the component
 			Component->RemoveInstance(InstanceIndex);
 
@@ -430,6 +448,15 @@ void FFoliageMeshInfo::RemoveInstances(AInstancedFoliageActor* InIFA, const TArr
 					ComponentHashInfo->Instances.Add(InstanceIndex);
 				}
 
+				// Spawner hash
+				if (SwappedInstance.Spawner)
+				{
+					FFoliageComponentHashInfo* SpawnerHashInfo = SpawnerHash.Find(SwappedInstance.Spawner);
+					check(SpawnerHashInfo);
+					SpawnerHashInfo->Instances.Remove(Instances.Num());
+					SpawnerHashInfo->Instances.Add(InstanceIndex);
+				}
+
 				// Selection
 				if (SelectedIndices.Contains(Instances.Num()))
 				{
@@ -533,6 +560,7 @@ void FFoliageMeshInfo::ReallocateClusters(AInstancedFoliageActor* InIFA, UFoliag
 	Exchange(Instances, OldInstances);
 	InstanceHash->Empty();
 	ComponentHash.Empty();
+	SpawnerHash.Empty();
 	SelectedIndices.Empty();
 
 	// Re-add
@@ -798,6 +826,19 @@ void AInstancedFoliageActor::DeleteInstancesForComponent(UActorComponent* InComp
 	}
 }
 
+void AInstancedFoliageActor::DeleteInstancesForSpawner(UActorComponent* InComponent)
+{
+	for (auto& MeshPair : FoliageMeshes)
+	{
+		FFoliageMeshInfo& MeshInfo = *MeshPair.Value;
+		const FFoliageComponentHashInfo* SpawnerHashInfo = MeshInfo.SpawnerHash.Find(InComponent);
+		if (SpawnerHashInfo)
+		{
+			MeshInfo.RemoveInstances(this, SpawnerHashInfo->Instances.Array());
+		}
+	}
+}
+
 void AInstancedFoliageActor::MoveInstancesForComponentToCurrentLevel(UActorComponent* InComponent)
 {
 	AInstancedFoliageActor* NewIFA = AInstancedFoliageActor::GetInstancedFoliageActorForCurrentLevel(InComponent->GetWorld());
@@ -1461,6 +1502,13 @@ void AInstancedFoliageActor::PostLoad()
 						ComponentHashInfo = &MeshInfo.ComponentHash.Add(Instance.Base, FFoliageComponentHashInfo(Instance.Base));
 					}
 					ComponentHashInfo->Instances.Add(InstanceIdx);
+
+					FFoliageComponentHashInfo* SpawnerHashInfo = MeshInfo.SpawnerHash.Find(Instance.Spawner);
+					if (SpawnerHashInfo == nullptr)
+					{
+						SpawnerHashInfo = &MeshInfo.SpawnerHash.Add(Instance.Spawner, FFoliageComponentHashInfo(Instance.Spawner));
+					}
+					SpawnerHashInfo->Instances.Add(InstanceIdx);
 				}
 			}
 
diff --git a/Engine/Source/Runtime/Engine/Public/InstancedFoliage.h b/Engine/Source/Runtime/Engine/Public/InstancedFoliage.h
--- a/Engine/Source/Runtime/Engine/Public/InstancedFoliage.h
+++ b/Engine/Source/Runtime/Engine/Public/InstancedFoliage.h
@@ -54,8 +54,11 @@ struct FFoliageInstance : public FFoliageInstancePlacementInfo
 {
 	UActorComponent* Base;
 
+	UActorComponent* Spawner;
+
 	FFoliageInstance()
 	: Base(NULL)
+	, Spawner(NULL)
 	{}
 
 
@@ -162,6 +165,9 @@ struct FFoliageMeshInfo
 	// Transient, editor-only set of instances per component
 	TMap<UActorComponent*, FFoliageComponentHashInfo> ComponentHash;
 
+	// Transient, editor-only set of instances per spawner
+	TMap<UActorComponent*, FFoliageComponentHashInfo> SpawnerHash;
+
 	// Transient, editor-only list of selected instances.
 	TSet<int32> SelectedIndices;
 #endif
