(Merge 2391682) Made DefaultPawn more extensible by allowing subclasses to use any subclass of PawnMovementComponent.

diff --git a/Engine/Source/Runtime/Engine/Classes/GameFramework/DefaultPawn.h b/Engine/Source/Runtime/Engine/Classes/GameFramework/DefaultPawn.h
--- a/Engine/Source/Runtime/Engine/Classes/GameFramework/DefaultPawn.h
+++ b/Engine/Source/Runtime/Engine/Classes/GameFramework/DefaultPawn.h
@@ -90,7 +90,7 @@ public:
 private:
 	/** DefaultPawn movement component */
 	UPROPERTY(Category = Pawn, VisibleAnywhere, BlueprintReadOnly, meta = (AllowPrivateAccess = "true"))
-	UFloatingPawnMovement* MovementComponent;
+	UPawnMovementComponent* MovementComponent;
 
 public:
 	/** Name of the CollisionComponent. */
diff --git a/Engine/Source/Runtime/Engine/Private/DefaultPawn.cpp b/Engine/Source/Runtime/Engine/Private/DefaultPawn.cpp
--- a/Engine/Source/Runtime/Engine/Private/DefaultPawn.cpp
+++ b/Engine/Source/Runtime/Engine/Private/DefaultPawn.cpp
@@ -36,7 +36,7 @@ ADefaultPawn::ADefaultPawn(const FObjectInitializer& ObjectInitializer)
 
 	RootComponent = CollisionComponent;
 
-	MovementComponent = CreateDefaultSubobject<UFloatingPawnMovement>(ADefaultPawn::MovementComponentName);
+	MovementComponent = CreateDefaultSubobject<UPawnMovementComponent, UFloatingPawnMovement>(ADefaultPawn::MovementComponentName);
 	MovementComponent->UpdatedComponent = CollisionComponent;
 
 	// Structure to hold one-time initialization
