UE-4402 - A more robust way to retire FNetworkGUID's

diff --git a/Engine/Source/Runtime/Engine/Classes/Engine/PackageMapClient.h b/Engine/Source/Runtime/Engine/Classes/Engine/PackageMapClient.h
--- a/Engine/Source/Runtime/Engine/Classes/Engine/PackageMapClient.h
+++ b/Engine/Source/Runtime/Engine/Classes/Engine/PackageMapClient.h
@@ -27,7 +27,7 @@
 class FNetGuidCacheObject
 {
 public:
-	FNetGuidCacheObject() : GuidSequence( 0 ), InitialQueryTime( -1.0f ), bNoLoad( 0 ), bIgnoreWhenMissing( 0 ), bIsPending( 0 ), bIsBroken( 0 )
+	FNetGuidCacheObject() : ReadOnlyTimestamp( 0 ), bNoLoad( 0 ), bIgnoreWhenMissing( 0 ), bIsPending( 0 ), bIsBroken( 0 )
 	{
 	}
 
@@ -36,10 +36,9 @@ public:
 	// These fields are set when this guid is static
 	FNetworkGUID				OuterGUID;
 	FName						PathName;
-	int32						GuidSequence;				// We remember the guid sequence this net guid was generated on so we can tell how old it is
 	FGuid						PackageGuid;				// If this is a package, this is the guid to expect
 
-	float						InitialQueryTime;
+	double						ReadOnlyTimestamp;			// Time in second when we should start timing out after going read only
 
 	uint8						bNoLoad				: 1;	// Don't load this, only do a find
 	uint8						bIgnoreWhenMissing	: 1;	// Don't warn when this asset can't be found or loaded
@@ -76,7 +75,6 @@ public:
 	TMap< TWeakObjectPtr< UObject >, FNetworkGUID >	NetGUIDLookup;
 	int32											UniqueNetIDs[2];
 
-	int32											GuidSequence;
 	bool											IsExportingNetGUIDBunch;
 
 	UNetDriver *									Driver;
diff --git a/Engine/Source/Runtime/Engine/Private/NetworkDriver.cpp b/Engine/Source/Runtime/Engine/Private/NetworkDriver.cpp
--- a/Engine/Source/Runtime/Engine/Private/NetworkDriver.cpp
+++ b/Engine/Source/Runtime/Engine/Private/NetworkDriver.cpp
@@ -2726,9 +2726,8 @@ void UNetDriver::ResetGameWorldState()
 
 void UNetDriver::CleanPackageMaps()
 {
-	if ( IsServer() && GuidCache.IsValid()  )
+	if ( GuidCache.IsValid() )
 	{ 
-		// Only the server makes this call initially, client will call when the GuidSequence increments on their end
 		GuidCache->CleanReferences();
 	}
 }
diff --git a/Engine/Source/Runtime/Engine/Private/PackageMapClient.cpp b/Engine/Source/Runtime/Engine/Private/PackageMapClient.cpp
--- a/Engine/Source/Runtime/Engine/Private/PackageMapClient.cpp
+++ b/Engine/Source/Runtime/Engine/Private/PackageMapClient.cpp
@@ -23,7 +23,6 @@ static const int GUID_PACKET_ACKED		= -1;
 static const int INTERNAL_LOAD_OBJECT_RECURSION_LIMIT = 16;
 
 static TAutoConsoleVariable<int32> CVarAllowAsyncLoading( TEXT( "net.AllowAsyncLoading" ), 0, TEXT( "Allow async loading" ) );
-static TAutoConsoleVariable<int32> CVarSimulateAsyncLoading( TEXT( "net.SimulateAsyncLoading" ), 0, TEXT( "Simulate async loading" ) );
 static TAutoConsoleVariable<int32> CVarIgnorePackageMismatch( TEXT( "net.IgnorePackageMismatch" ), 0, TEXT( "Ignore when package versions are different" ) );
 
 /*-----------------------------------------------------------------------------
@@ -725,7 +724,6 @@ bool UPackageMapClient::ExportNetGUID( FNetworkGUID NetGUID, const UObject* Obje
 
 			ExportNetGUIDCount = 0;
 			*CurrentExportBunch << ExportNetGUIDCount;
-			*CurrentExportBunch << GuidCache->GuidSequence;
 			NET_CHECKSUM( *CurrentExportBunch );
 		}
 
@@ -804,7 +802,6 @@ void UPackageMapClient::ExportNetGUIDHeader()
 	FBitWriterMark Restore(Out);
 	Reset.PopWithoutClear(Out);
 	Out << ExportNetGUIDCount;
-	Out << GuidCache->GuidSequence;
 	Restore.PopWithoutClear(Out);
 
 	// If we've written new NetGUIDs to the 'bunch' set (current+1)
@@ -839,25 +836,6 @@ void UPackageMapClient::ReceiveNetGUIDBunch( FInBunch &InBunch )
 	int32 NumGUIDsInBunch = 0;
 	InBunch << NumGUIDsInBunch;
 
-	int32 NewGuidSequence = -1;
-
-	InBunch << NewGuidSequence;
-
-	if ( NewGuidSequence < GuidCache->GuidSequence )
-	{
-		// Older sequence, ignore
-		UE_LOG( LogNetPackageMap, Warning, TEXT( "UPackageMapClient::ReceiveNetGUIDBunch: NewGuidSequence < GuidSequence (%i / %i)" ), GuidCache->GuidSequence, NewGuidSequence );
-		GuidCache->IsExportingNetGUIDBunch = false;
-		return;
-	}
-	else if ( NewGuidSequence > GuidCache->GuidSequence )
-	{
-		// Newer sequence, this is our signal to clean the package to stay in sync with the server
-		UE_LOG( LogNetPackageMap, Verbose, TEXT( "UPackageMapClient::ReceiveNetGUIDBunch: NewGuidSequence > GuidSequence (%i / %i)" ), GuidCache->GuidSequence, NewGuidSequence );
-		GuidCache->GuidSequence = NewGuidSequence;
-		GuidCache->CleanReferences();
-	}
-
 	static const int32 MAX_GUID_COUNT = 2048;
 
 	if ( NumGUIDsInBunch > MAX_GUID_COUNT )
@@ -1262,7 +1240,6 @@ UObject* UPackageMapClient::GetObjectFromNetGUID( const FNetworkGUID& NetGUID, c
 FNetGUIDCache::FNetGUIDCache( UNetDriver * InDriver ) : Driver( InDriver )
 {
 	UniqueNetIDs[0] = UniqueNetIDs[1] = 0;
-	GuidSequence = 0;
 	IsExportingNetGUIDBunch = false;
 }
 
@@ -1276,44 +1253,27 @@ public:
 
 void FNetGUIDCache::CleanReferences()
 {
-	if ( IsNetGUIDAuthority() )
-	{
-		// The server increments this to signify to clients to also call CleanReferences when they detect this change
-		GuidSequence++;
-	}
-
-	// Server:
-	//	Clean all static or non valid guids
-	//	This means static objects will have new guids assigned, we're assuming clients are ok with this
-	// Client:
-	//	Clean all guids that are old enough
-	//	This means guids that are older than SEAMLESS_TRAVEL_COUNT_FOR_CLEAN seamless travels
-	//		We use a GuidSequence counter to determine how often the server has seamlessly traveled
-	// 
-	// It's possible for guids that have been cleaned to get lost if they were sent from a level that was 
-	// was older than (SEAMLESS_TRAVEL_COUNT_FOR_CLEAN) seamless travels, but we're assuming the game code 
-	// is resilient to this, and only needs information from the previous level at most
+	// Mark all static or non valid dynamic guids to timeout after NETWORK_GUID_TIMEOUT seconds
+	// We want to leave them around for a certain amount of time to allow in-flight references to these guids to continue to resolve
 	for ( auto It = ObjectLookup.CreateIterator(); It; ++It )
 	{
-		if ( IsNetGUIDAuthority() )
+		if ( It.Value().ReadOnlyTimestamp != 0 )
 		{
-			// Server cleans references that are NULL and static
-			// (server will assign a new guid to static references after a seamless travel)
-			if ( !It.Value().Object.IsValid() || It.Key().IsStatic() )
+			// If this guid was suppose to time out, check to see if it has, otherwise ignore it
+			const double NETWORK_GUID_TIMEOUT = 90;
+
+			if ( FPlatformTime::Seconds() - It.Value().ReadOnlyTimestamp > NETWORK_GUID_TIMEOUT )
 			{
 				It.RemoveCurrent();
 			}
+
+			continue;
 		}
-		else
-		{
-			static const int32 SEAMLESS_TRAVEL_COUNT_FOR_CLEAN = 2;
 
-			// Clients only clean references that are older than SEAMLESS_TRAVEL_COUNT_FOR_CLEAN 
-			//	(this guid has seamless traveled this many times, so we can safely punt it once it's old enough)
-			if ( ( GuidSequence - It.Value().GuidSequence ) > SEAMLESS_TRAVEL_COUNT_FOR_CLEAN )
-			{
-				It.RemoveCurrent();
-			}
+		if ( !It.Value().Object.IsValid() || It.Key().IsStatic() )
+		{
+			// We will leave this guid around for NETWORK_GUID_TIMEOUT seconds to make sure any in-flight guids can be resolved
+			It.Value().ReadOnlyTimestamp = FPlatformTime::Seconds();
 		}
 	}
 
@@ -1332,12 +1292,9 @@ void FNetGUIDCache::CleanReferences()
 		check( !It.Key().IsDefault() );
 		check( It.Key().IsStatic() != It.Key().IsDynamic() );
 
-		checkf( !It.Value().Object.IsValid() || NetGUIDLookup.FindRef( It.Value().Object ) == It.Key(), TEXT( "Failed to validate ObjectLookup map in UPackageMap. Object '%s' was not in the NetGUIDLookup map with with value '%s'." ), *It.Value().Object.Get()->GetPathName(), *It.Key().ToString() );
+		checkf( !It.Value().Object.IsValid() || NetGUIDLookup.FindRef( It.Value().Object ) == It.Key() || It.Value().ReadOnlyTimestamp != 0, TEXT( "Failed to validate ObjectLookup map in UPackageMap. Object '%s' was not in the NetGUIDLookup map with with value '%s'." ), *It.Value().Object.Get()->GetPathName(), *It.Key().ToString() );
 	}
 
-	// Server lists should be completely in sync
-	check( ObjectLookup.Num() == NetGUIDLookup.Num() || !IsNetGUIDAuthority() );
-
 	for ( auto It = NetGUIDLookup.CreateIterator(); It; ++It )
 	{
 		check( It.Key().IsValid() );
@@ -1422,7 +1379,21 @@ FNetworkGUID FNetGUIDCache::GetOrAssignNetGUID( const UObject* Object )
 
 	if ( NetGUID.IsValid() )
 	{
-		return NetGUID;
+		const FNetGuidCacheObject* CacheObject = ObjectLookup.Find( NetGUID );
+
+		// Check to see if this guid is read only
+		// If so, we should ignore this entry, and create a new one (or send default as client)
+		const bool bReadOnly = CacheObject != NULL && CacheObject->ReadOnlyTimestamp > 0;
+
+		if ( bReadOnly )
+		{
+			// Reset this object's guid, we will re-assign below (or send default as a client)
+			NetGUIDLookup.Remove( Object );
+		}
+		else
+		{
+			return NetGUID;
+		}
 	}
 
 	if ( !IsNetGUIDAuthority() )
@@ -1460,7 +1431,6 @@ void FNetGUIDCache::RegisterNetGUID_Internal( const FNetworkGUID& NetGUID, const
 {
 	// We're pretty strict in this function, we expect everything to have been handled before we get here
 	check( !ObjectLookup.Contains( NetGUID ) );
-	check( CacheObject.GuidSequence == GuidSequence );
 
 	ObjectLookup.Add( NetGUID, CacheObject );
 
@@ -1498,7 +1468,6 @@ void FNetGUIDCache::RegisterNetGUID_Server( const FNetworkGUID& NetGUID, const U
 	FNetGuidCacheObject CacheObject;
 
 	CacheObject.Object			= Object;
-	CacheObject.GuidSequence	= GuidSequence;
 	
 	RegisterNetGUID_Internal( NetGUID, CacheObject );
 }
@@ -1568,8 +1537,7 @@ void FNetGUIDCache::RegisterNetGUID_Client( const FNetworkGUID& NetGUID, const U
 
 	FNetGuidCacheObject CacheObject;
 
-	CacheObject.Object			= Object;
-	CacheObject.GuidSequence	= GuidSequence;
+	CacheObject.Object = Object;
 
 	RegisterNetGUID_Internal( NetGUID, CacheObject );
 }
@@ -1621,7 +1589,6 @@ void FNetGUIDCache::RegisterNetGUIDFromPath_Client( const FNetworkGUID& NetGUID,
 	CacheObject.PackageGuid			= PackageGuid;
 	CacheObject.bNoLoad				= bNoLoad;
 	CacheObject.bIgnoreWhenMissing	= bIgnoreWhenMissing;
-	CacheObject.GuidSequence		= GuidSequence;
 
 	RegisterNetGUID_Internal( NetGUID, CacheObject );
 }
@@ -1764,25 +1731,6 @@ UObject* FNetGUIDCache::GetObjectFromNetGUID( const FNetworkGUID& NetGUID, const
 		}
 	}
 
-	// Simulate async loading if desired
-	if ( CVarAllowAsyncLoading.GetValueOnGameThread() > 0 && CVarSimulateAsyncLoading.GetValueOnGameThread() > 0 && !IsNetGUIDAuthority() && !CacheObjectPtr->OuterGUID.IsValid() )
-	{
-		if ( !CacheObjectPtr->bNoLoad && !CacheObjectPtr->bIsPending && CacheObjectPtr->InitialQueryTime <= 0.0f )
-		{
-			CacheObjectPtr->InitialQueryTime = FPlatformTime::Seconds();
-			CacheObjectPtr->bIsPending = true;
-		}
-
-		if ( CacheObjectPtr->bIsPending && CacheObjectPtr->InitialQueryTime >= 0.0f )
-		{
-			if ( FPlatformTime::Seconds() - CacheObjectPtr->InitialQueryTime > (float)CVarSimulateAsyncLoading.GetValueOnGameThread() )
-			{
-				//UE_LOG( LogNetPackageMap, Warning, TEXT( "GetObjectFromNetGUID: FINISH simulating pending load. Path: %s, NetGUID: %s" ), *CacheObjectPtr->PathName.ToString(), *NetGUID.ToString() );
-				CacheObjectPtr->bIsPending = false;
-			}
-		}
-	}
-
 	// At this point, we either have an outer, or we are a package
 	check( !CacheObjectPtr->bIsPending );
 	check( ObjOuter == NULL || ObjOuter->GetOutermost()->IsFullyLoaded() );
@@ -1803,7 +1751,7 @@ UObject* FNetGUIDCache::GetObjectFromNetGUID( const FNetworkGUID& NetGUID, const
 			//	3. We're actually suppose to load (levels don't load here for example)
 			check( !PendingAsyncPackages.Contains( CacheObjectPtr->PathName ) );
 
-			if ( CVarAllowAsyncLoading.GetValueOnGameThread() > 0 && CVarSimulateAsyncLoading.GetValueOnGameThread() == 0 )
+			if ( CVarAllowAsyncLoading.GetValueOnGameThread() > 0 )
 			{
 				PendingAsyncPackages.Add( CacheObjectPtr->PathName, NetGUID );
 				CacheObjectPtr->bIsPending = true;
@@ -1887,22 +1835,16 @@ UObject* FNetGUIDCache::GetObjectFromNetGUID( const FNetworkGUID& NetGUID, const
 	CacheObjectPtr->Object = Object;		
 
 	// Assign the guid to the object 
-	// Don't allow the assignment if this guid is from an older sequence
-	if ( CacheObjectPtr->GuidSequence == GuidSequence )
-	{
-		// We're good, assign the guid (this may replace the old guid, but we're ok with that, since this guid is newer)
-		NetGUIDLookup.Add( Object, NetGUID );
-	}
-	else
+	// We don't want to assign this guid to the object if this guid is timing out
+	// But we'll have to if there is no other guid yet
+	if ( CacheObjectPtr->ReadOnlyTimestamp == 0 || !NetGUIDLookup.Contains( Object ) )
 	{
-		UE_LOG( LogNetPackageMap, Warning, TEXT( "GetObjectFromNetGUID: Attempt to reassign guid from older sequence (%i / %i). Path: %s, Outer: %s, NetGUID: %s" ), GuidSequence, CacheObjectPtr->GuidSequence, *CacheObjectPtr->PathName.ToString(), ObjOuter != NULL ? *ObjOuter->GetPathName() : TEXT( "NULL" ), *NetGUID.ToString() );
-
-		// If we haven't set the guid yet for this object, go ahead and accept the older guid until the newer one comes in
-		if ( !NetGUIDLookup.Contains( Object ) )
+		if ( CacheObjectPtr->ReadOnlyTimestamp > 0 )
 		{
-			UE_LOG( LogNetPackageMap, Warning, TEXT( "GetObjectFromNetGUID: Forced assign guid from older sequence (%i / %i). Path: %s, Outer: %s, NetGUID: %s" ), GuidSequence, CacheObjectPtr->GuidSequence, *CacheObjectPtr->PathName.ToString(), ObjOuter != NULL ? *ObjOuter->GetPathName() : TEXT( "NULL" ), *NetGUID.ToString() );
-			NetGUIDLookup.Add( Object, NetGUID );
+			UE_LOG( LogNetPackageMap, Warning, TEXT( "GetObjectFromNetGUID: Attempt to reassign read-only guid. FullNetGUIDPath: %s" ), *FullNetGUIDPath( NetGUID ) );
 		}
+
+		NetGUIDLookup.Add( Object, NetGUID );
 	}
 
 	return Object;
