Default template parameters don't work on VS2012 compiler

diff --git a/Engine/Source/Runtime/Engine/Classes/GameFramework/Actor.h b/Engine/Source/Runtime/Engine/Classes/GameFramework/Actor.h
--- a/Engine/Source/Runtime/Engine/Classes/GameFramework/Actor.h
+++ b/Engine/Source/Runtime/Engine/Classes/GameFramework/Actor.h
@@ -2130,7 +2130,7 @@ public:
 	 *     Actor->GetComponents(PrimComponents);
 	 * }
 	 */
-	template<class T, class AllocatorType = FDefaultAllocator>
+	template<class T, class AllocatorType>
 	void GetComponents(TArray<T*, AllocatorType>& OutComponents) const
 	{
 		static_assert(CanConvertPointerFromTo<T, UActorComponent>::Result, "'T' template parameter to GetComponents must be derived from ActorComponent");
@@ -2150,7 +2150,7 @@ public:
 	}
 
 	/** UActorComponent specialization of GetComponents() to avoid unnecessary casts. */
-	template<class AllocatorType = FDefaultAllocator>
+	template<class AllocatorType>
 	void GetComponents(TArray<UActorComponent*, AllocatorType>& OutComponents) const
 	{
 		SCOPE_CYCLE_COUNTER(STAT_GetComponentsTime);
