#ue4 Fixing conditionals in the HMDFunctionLibrary

#github 690 https://github.com/EpicGames/UnrealEngine/pull/690/

diff --git a/Engine/Source/Runtime/Engine/Private/HeadMountedDisplayFunctionLibrary.cpp b/Engine/Source/Runtime/Engine/Private/HeadMountedDisplayFunctionLibrary.cpp
--- a/Engine/Source/Runtime/Engine/Private/HeadMountedDisplayFunctionLibrary.cpp
+++ b/Engine/Source/Runtime/Engine/Private/HeadMountedDisplayFunctionLibrary.cpp
@@ -129,11 +129,11 @@ void UHeadMountedDisplayFunctionLibrary::SetBaseRotationAndPositionOffset(const
 {
 	if (GEngine->HMDDevice.IsValid() && GEngine->HMDDevice->IsHeadTrackingAllowed())
 	{
-		if (Options == EOrientPositionSelector::Orientation || EOrientPositionSelector::OrientationAndPosition)
+		if ((Options == EOrientPositionSelector::Orientation) || (Options == EOrientPositionSelector::OrientationAndPosition))
 		{
 			GEngine->HMDDevice->SetBaseRotation(BaseRot);
 		}
-		if (Options == EOrientPositionSelector::Position || EOrientPositionSelector::OrientationAndPosition)
+		if ((Options == EOrientPositionSelector::Position) || (Options == EOrientPositionSelector::OrientationAndPosition))
 		{
 			GEngine->HMDDevice->SetPositionOffset(PosOffset);
 		}
