Removed SpawnerHash

diff --git a/Engine/Source/Runtime/Engine/Private/InstancedFoliage.cpp b/Engine/Source/Runtime/Engine/Private/InstancedFoliage.cpp
--- a/Engine/Source/Runtime/Engine/Private/InstancedFoliage.cpp
+++ b/Engine/Source/Runtime/Engine/Private/InstancedFoliage.cpp
@@ -354,13 +354,6 @@ void FFoliageMeshInfo::AddInstance(AInstancedFoliageActor* InIFA, UFoliageType*
 	}
 	ComponentHashInfo->Instances.Add(InstanceIndex);
 	
-	FFoliageComponentHashInfo* SpawnerHashInfo = SpawnerHash.Find(InNewInstance.Spawner);
-	if (SpawnerHashInfo == nullptr)
-	{
-		SpawnerHashInfo = &SpawnerHash.Add(InNewInstance.Spawner, FFoliageComponentHashInfo(InNewInstance.Spawner));
-	}
-	SpawnerHashInfo->Instances.Add(InstanceIndex);
-
 	// Calculate transform for the instance
 	FTransform InstanceToWorld = InNewInstance.GetInstanceWorldTransform();
 
@@ -411,17 +404,6 @@ void FFoliageMeshInfo::RemoveInstances(AInstancedFoliageActor* InIFA, const TArr
 			    }
 		    }
 
-			FFoliageComponentHashInfo* SpawnerHashInfo = SpawnerHash.Find(Instance.Spawner);
-			if (SpawnerHashInfo)
-			{
-				SpawnerHashInfo->Instances.Remove(InstanceIndex);
-				if (SpawnerHashInfo->Instances.Num() == 0)
-				{
-					// Remove the component from the spawner hash if this is the last instance.
-					SpawnerHash.Remove(Instance.Spawner);
-				}
-			}
-
 			// remove from the component
 			Component->RemoveInstance(InstanceIndex);
 
@@ -448,15 +430,6 @@ void FFoliageMeshInfo::RemoveInstances(AInstancedFoliageActor* InIFA, const TArr
 					ComponentHashInfo->Instances.Add(InstanceIndex);
 				}
 
-				// Spawner hash
-				if (SwappedInstance.Spawner)
-				{
-					FFoliageComponentHashInfo* SpawnerHashInfo = SpawnerHash.Find(SwappedInstance.Spawner);
-					check(SpawnerHashInfo);
-					SpawnerHashInfo->Instances.Remove(Instances.Num());
-					SpawnerHashInfo->Instances.Add(InstanceIndex);
-				}
-
 				// Selection
 				if (SelectedIndices.Contains(Instances.Num()))
 				{
@@ -560,7 +533,6 @@ void FFoliageMeshInfo::ReallocateClusters(AInstancedFoliageActor* InIFA, UFoliag
 	Exchange(Instances, OldInstances);
 	InstanceHash->Empty();
 	ComponentHash.Empty();
-	SpawnerHash.Empty();
 	SelectedIndices.Empty();
 
 	// Re-add
@@ -831,10 +803,18 @@ void AInstancedFoliageActor::DeleteInstancesForSpawner(UActorComponent* InCompon
 	for (auto& MeshPair : FoliageMeshes)
 	{
 		FFoliageMeshInfo& MeshInfo = *MeshPair.Value;
-		const FFoliageComponentHashInfo* SpawnerHashInfo = MeshInfo.SpawnerHash.Find(InComponent);
-		if (SpawnerHashInfo)
+		TArray<int32> InstancesToRemove;
+		for (int32 InstanceIdx = 0; InstanceIdx < MeshInfo.Instances.Num(); InstanceIdx++)
 		{
-			MeshInfo.RemoveInstances(this, SpawnerHashInfo->Instances.Array());
+			if (MeshInfo.Instances[InstanceIdx].Spawner == InComponent)
+			{
+				InstancesToRemove.Add(InstanceIdx);
+			}
+		}
+
+		if (InstancesToRemove.Num())
+		{
+			MeshInfo.RemoveInstances(this, InstancesToRemove);
 		}
 	}
 }
@@ -1502,13 +1482,6 @@ void AInstancedFoliageActor::PostLoad()
 						ComponentHashInfo = &MeshInfo.ComponentHash.Add(Instance.Base, FFoliageComponentHashInfo(Instance.Base));
 					}
 					ComponentHashInfo->Instances.Add(InstanceIdx);
-
-					FFoliageComponentHashInfo* SpawnerHashInfo = MeshInfo.SpawnerHash.Find(Instance.Spawner);
-					if (SpawnerHashInfo == nullptr)
-					{
-						SpawnerHashInfo = &MeshInfo.SpawnerHash.Add(Instance.Spawner, FFoliageComponentHashInfo(Instance.Spawner));
-					}
-					SpawnerHashInfo->Instances.Add(InstanceIdx);
 				}
 			}
 
diff --git a/Engine/Source/Runtime/Engine/Public/InstancedFoliage.h b/Engine/Source/Runtime/Engine/Public/InstancedFoliage.h
--- a/Engine/Source/Runtime/Engine/Public/InstancedFoliage.h
+++ b/Engine/Source/Runtime/Engine/Public/InstancedFoliage.h
@@ -165,9 +165,6 @@ struct FFoliageMeshInfo
 	// Transient, editor-only set of instances per component
 	TMap<UActorComponent*, FFoliageComponentHashInfo> ComponentHash;
 
-	// Transient, editor-only set of instances per spawner
-	TMap<UActorComponent*, FFoliageComponentHashInfo> SpawnerHash;
-
 	// Transient, editor-only list of selected instances.
 	TSet<int32> SelectedIndices;
 #endif
