Forgot about few files.

diff --git a/Engine/Source/Runtime/Engine/Classes/VisualLogger/VisualLoggerBinaryFileDevice.h b/Engine/Source/Runtime/Engine/Classes/VisualLogger/VisualLoggerBinaryFileDevice.h
--- a/Engine/Source/Runtime/Engine/Classes/VisualLogger/VisualLoggerBinaryFileDevice.h
+++ b/Engine/Source/Runtime/Engine/Classes/VisualLogger/VisualLoggerBinaryFileDevice.h
@@ -20,7 +20,7 @@ public:
 	virtual void StartRecordingToFile(float TImeStamp) override;
 	virtual void StopRecordingToFile(float TImeStamp) override;
 	virtual void SetFileName(const FString& InFileName) override;
-	virtual void Serialize(const class UObject* LogOwner, FName OwnerName, const FVisualLogEntry& LogEntry) override;
+	virtual void Serialize(const class UObject* LogOwner, FName OwnerName, FName OwnerClassName, const FVisualLogEntry& LogEntry) override;
 	virtual void GetRecordedLogs(TArray<FVisualLogEntryItem>& RecordedLogs) const override { RecordedLogs = FrameCache; }
 	virtual bool HasFlags(int32 InFlags) const { return !!(InFlags & (EVisualLoggerDeviceFlags::CanSaveToFile | EVisualLoggerDeviceFlags::StoreLogsLocally)); }
 
diff --git a/Engine/Source/Runtime/Engine/Private/VisualLogger/VisualLoggerAutomationTests.cpp b/Engine/Source/Runtime/Engine/Private/VisualLogger/VisualLoggerAutomationTests.cpp
--- a/Engine/Source/Runtime/Engine/Private/VisualLogger/VisualLoggerAutomationTests.cpp
+++ b/Engine/Source/Runtime/Engine/Private/VisualLogger/VisualLoggerAutomationTests.cpp
@@ -19,7 +19,7 @@ class FVisualLoggerTestDevice : public FVisualLogDevice
 public:
 	FVisualLoggerTestDevice();
 	virtual void Cleanup(bool bReleaseMemory = false) override;
-	virtual void Serialize(const class UObject* LogOwner, FName OwnerName, const FVisualLogEntry& LogEntry) override;
+	virtual void Serialize(const class UObject* LogOwner, FName OwnerName, FName InOwnerClassName, const FVisualLogEntry& LogEntry) override;
 
 	class UObject* LastObject;
 	FVisualLogEntry LastEntry;
@@ -36,7 +36,7 @@ void FVisualLoggerTestDevice::Cleanup(bool bReleaseMemory)
 	LastEntry.Reset();
 }
 
-void FVisualLoggerTestDevice::Serialize(const UObject* LogOwner, FName OwnerName, const FVisualLogEntry& LogEntry)
+void FVisualLoggerTestDevice::Serialize(const UObject* LogOwner, FName OwnerName, FName InOwnerClassName, const FVisualLogEntry& LogEntry)
 {
 	LastObject = const_cast<class UObject*>(LogOwner);
 	LastEntry = LogEntry;
diff --git a/Engine/Source/Runtime/Engine/Private/VisualLogger/VisualLoggerBinaryFileDevice.cpp b/Engine/Source/Runtime/Engine/Private/VisualLogger/VisualLoggerBinaryFileDevice.cpp
--- a/Engine/Source/Runtime/Engine/Private/VisualLogger/VisualLoggerBinaryFileDevice.cpp
+++ b/Engine/Source/Runtime/Engine/Private/VisualLogger/VisualLoggerBinaryFileDevice.cpp
@@ -77,7 +77,7 @@ void FVisualLoggerBinaryFileDevice::SetFileName(const FString& InFileName)
 	FileName = InFileName;
 }
 
-void FVisualLoggerBinaryFileDevice::Serialize(const UObject* LogOwner, FName OwnerName, const FVisualLogEntry& LogEntry)
+void FVisualLoggerBinaryFileDevice::Serialize(const UObject* LogOwner, FName OwnerName, FName OwnerClassName, const FVisualLogEntry& LogEntry)
 {
 	const int32 NumEntries = FrameCache.Num();
 	if (NumEntries> 0 && LastLogTimeStamp + FrameCacheLenght <= LogEntry.TimeStamp && FileArchive)
@@ -87,6 +87,6 @@ void FVisualLoggerBinaryFileDevice::Serialize(const UObject* LogOwner, FName Own
 	}
 
 	LastLogTimeStamp = LogEntry.TimeStamp;
-	FrameCache.Add(FVisualLogEntryItem(OwnerName, LogEntry));
+	FrameCache.Add(FVisualLogEntryItem(OwnerName, OwnerClassName, LogEntry));
 }
 #endif
diff --git a/Engine/Source/Runtime/Engine/Private/VisualLogger/VisualLoggerTypes.cpp b/Engine/Source/Runtime/Engine/Private/VisualLogger/VisualLoggerTypes.cpp
--- a/Engine/Source/Runtime/Engine/Private/VisualLogger/VisualLoggerTypes.cpp
+++ b/Engine/Source/Runtime/Engine/Private/VisualLogger/VisualLoggerTypes.cpp
@@ -337,6 +337,11 @@ FArchive& operator<<(FArchive& Ar, FVisualLogEntry& LogEntry)
 FArchive& operator<<(FArchive& Ar, FVisualLogDevice::FVisualLogEntryItem& FrameCacheItem)
 {
 	FVisualLoggerHelpers::Serialize(Ar, FrameCacheItem.OwnerName);
+	const int32 VLogsVer = Ar.CustomVer(EVisualLoggerVersion::GUID);
+	if (VLogsVer >= EVisualLoggerVersion::AddedOwnerClassName)
+	{
+		FVisualLoggerHelpers::Serialize(Ar, FrameCacheItem.OwnerClassName);
+	}
 	Ar << FrameCacheItem.Entry;
 	return Ar;
 }
diff --git a/Engine/Source/Runtime/Engine/Public/VisualLogger/VisualLogger.h b/Engine/Source/Runtime/Engine/Public/VisualLogger/VisualLogger.h
--- a/Engine/Source/Runtime/Engine/Public/VisualLogger/VisualLogger.h
+++ b/Engine/Source/Runtime/Engine/Public/VisualLogger/VisualLogger.h
@@ -232,6 +232,8 @@ protected:
 	TMap<const class UObject*, FVisualLogEntry>	CurrentEntryPerObject;
 	// Map to contain names for Objects (they can be destroyed after while)
 	TMap<const class UObject*, FName> ObjectToNameMap;
+	// Map to contain class names for Objects (they can be destroyed after while)
+	TMap<const class UObject*, FName> ObjectToClassNameMap;
 	// Map to contain information about pointers in game
 	TMap<const class UObject*, TWeakObjectPtr<const class UObject> > ObjectToPointerMap;
 	// if set all categories are blocked from logging
diff --git a/Engine/Source/Runtime/Engine/Public/VisualLogger/VisualLogger.inl b/Engine/Source/Runtime/Engine/Public/VisualLogger/VisualLogger.inl
--- a/Engine/Source/Runtime/Engine/Public/VisualLogger/VisualLogger.inl
+++ b/Engine/Source/Runtime/Engine/Public/VisualLogger/VisualLogger.inl
@@ -82,6 +82,7 @@ FVisualLogEntry* FVisualLogger::GetEntryToWrite(const class UObject* Object, flo
 	{
 		CurrentEntry = &CurrentEntryPerObject.Add(LogOwner);
 		ObjectToNameMap.Add(LogOwner, LogOwner->GetFName());
+		ObjectToClassNameMap.Add(LogOwner, *(LogOwner->GetClass()->GetName()));
 		ObjectToPointerMap.Add(LogOwner, LogOwner);
 		InitializeNewEntry = true;
 	}
@@ -135,7 +136,7 @@ FVisualLogEntry* FVisualLogger::GetEntryToWrite(const class UObject* Object, flo
 					{
 						for (auto* Device : OutputDevices)
 						{
-							Device->Serialize(CurrentPair.Key, ObjectToNameMap[CurrentPair.Key], *Entry);
+							Device->Serialize(CurrentPair.Key, ObjectToNameMap[CurrentPair.Key], ObjectToClassNameMap[CurrentPair.Key], *Entry);
 						}
 						Entry->Reset();
 					}
@@ -157,7 +158,7 @@ void FVisualLogger::Flush()
 		{
 			for (auto* Device : OutputDevices)
 			{
-				Device->Serialize(CurrentEntry.Key, ObjectToNameMap[CurrentEntry.Key], CurrentEntry.Value);
+				Device->Serialize(CurrentEntry.Key, ObjectToNameMap[CurrentEntry.Key], ObjectToClassNameMap[CurrentEntry.Key], CurrentEntry.Value);
 			}
 			CurrentEntry.Value.Reset();
 		}
diff --git a/Engine/Source/Runtime/Engine/Public/VisualLogger/VisualLoggerTypes.h b/Engine/Source/Runtime/Engine/Public/VisualLogger/VisualLoggerTypes.h
--- a/Engine/Source/Runtime/Engine/Public/VisualLogger/VisualLoggerTypes.h
+++ b/Engine/Source/Runtime/Engine/Public/VisualLogger/VisualLoggerTypes.h
@@ -27,6 +27,7 @@ namespace EVisualLoggerVersion
 	{
 		Initial = 0,
 		HistogramGraphsSerialization = 1,
+		AddedOwnerClassName = 2,
 		// -----<new versions can be added before this line>-------------------------------------------------
 		// - this needs to be the last line (see note below)
 		VersionPlusOne,
@@ -216,14 +217,15 @@ public:
 	struct FVisualLogEntryItem
 	{
 		FVisualLogEntryItem() {}
-		FVisualLogEntryItem(FName InOwnerName, const FVisualLogEntry& LogEntry) : OwnerName(InOwnerName), Entry(LogEntry) { }
+		FVisualLogEntryItem(FName InOwnerName, FName InOwnerClassName, const FVisualLogEntry& LogEntry) : OwnerName(InOwnerName), OwnerClassName(InOwnerClassName), Entry(LogEntry) { }
 
 		FName OwnerName;
+		FName OwnerClassName;
 		FVisualLogEntry Entry;
 	};
 
 
-	virtual void Serialize(const class UObject* LogOwner, FName OwnerName, const FVisualLogEntry& LogEntry) = 0;
+	virtual void Serialize(const class UObject* LogOwner, FName OwnerName, FName InOwnerClassName, const FVisualLogEntry& LogEntry) = 0;
 	virtual void Cleanup(bool bReleaseMemory = false) { /* Empty */ }
 	virtual void StartRecordingToFile(float TImeStamp) { /* Empty */ }
 	virtual void StopRecordingToFile(float TImeStamp) { /* Empty */ }
