UE4 - fix crash when copy-pasting static mesh actors with vertex colors and no static mesh

diff --git a/Engine/Source/Runtime/Engine/Private/StaticMeshComponent.cpp b/Engine/Source/Runtime/Engine/Private/StaticMeshComponent.cpp
--- a/Engine/Source/Runtime/Engine/Private/StaticMeshComponent.cpp
+++ b/Engine/Source/Runtime/Engine/Private/StaticMeshComponent.cpp
@@ -732,8 +732,7 @@ bool UStaticMeshComponent::RequiresOverrideVertexColorsFixup( TArray<int32>& Out
 void UStaticMeshComponent::RemoveInstanceVertexColorsFromLOD( int32 LODToRemoveColorsFrom )
 {
 #if WITH_EDITORONLY_DATA
-	if (( LODToRemoveColorsFrom < StaticMesh->GetNumLODs() ) &&
-		( LODToRemoveColorsFrom < LODData.Num() ))
+	if (StaticMesh && LODToRemoveColorsFrom < StaticMesh->GetNumLODs() && LODToRemoveColorsFrom < LODData.Num())
 	{
 		FStaticMeshComponentLODInfo& CurrentLODInfo = LODData[LODToRemoveColorsFrom];
 
