polish

diff --git a/Engine/Source/Runtime/Engine/Classes/Components/TextRenderComponent.h b/Engine/Source/Runtime/Engine/Classes/Components/TextRenderComponent.h
--- a/Engine/Source/Runtime/Engine/Classes/Components/TextRenderComponent.h
+++ b/Engine/Source/Runtime/Engine/Classes/Components/TextRenderComponent.h
@@ -78,11 +78,11 @@ class ENGINE_API UTextRenderComponent : public UPrimitiveComponent
 	float InvDefaultSize;
 
 	/** Horizontal adjustment per character, default is 0.0 */
-	UPROPERTY(EditAnywhere, BlueprintReadOnly, Category=Text)
+	UPROPERTY(EditAnywhere, BlueprintReadOnly, AdvancedDisplay, Category=Text)
 	float HorizSpacingAdjust;
 
 	/** Allows text to draw unmodified when using debug visualization modes. **/
-	UPROPERTY(EditAnywhere, BlueprintReadOnly, Category=Rendering)
+	UPROPERTY(EditAnywhere, BlueprintReadOnly, AdvancedDisplay, Category=Rendering)
 	uint32 bAlwaysRenderAsText:1;
 
 	// -----------------------------
