polish material UI

diff --git a/Engine/Source/Runtime/Engine/Classes/Materials/Material.h b/Engine/Source/Runtime/Engine/Classes/Materials/Material.h
--- a/Engine/Source/Runtime/Engine/Classes/Materials/Material.h
+++ b/Engine/Source/Runtime/Engine/Classes/Materials/Material.h
@@ -403,33 +403,33 @@ public:
 	int32 NumCustomizedUVs;
 
 	/** Sets the lighting mode that will be used on this material if it is translucent. */
-	UPROPERTY(EditAnywhere, Category=Translucency)
+	UPROPERTY(EditAnywhere, Category=Translucency, meta=(DisplayName = "Lighting Mode"))
 	TEnumAsByte<enum ETranslucencyLightingMode> TranslucencyLightingMode;
 
 	/** 
 	 * Useful for artificially increasing the influence of the normal on the lighting result for translucency.
 	 * A value larger than 1 increases the influence of the normal, a value smaller than 1 makes the lighting more ambient.
 	 */
-	UPROPERTY(EditAnywhere, Category=Translucency)
+	UPROPERTY(EditAnywhere, Category=Translucency, meta=(DisplayName = "Directional Lighting Intensity"))
 	float TranslucencyDirectionalLightingIntensity;
 
 	/** Scale used to make translucent shadows more or less opaque than the material's actual opacity. */
-	UPROPERTY(EditAnywhere, Category=TranslucencySelfShadowing)
+	UPROPERTY(EditAnywhere, Category=TranslucencySelfShadowing, meta=(DisplayName = "Density Scale"))
 	float TranslucentShadowDensityScale;
 
 	/** 
 	 * Scale used to make translucent self-shadowing more or less opaque than the material's shadow on other objects. 
 	 * This is only used when the object is casting a volumetric translucent shadow.
 	 */
-	UPROPERTY(EditAnywhere, Category=TranslucencySelfShadowing)
+	UPROPERTY(EditAnywhere, Category=TranslucencySelfShadowing, meta=(DisplayName = "Density Scale"))
 	float TranslucentSelfShadowDensityScale;
 
 	/** Used to make a second self shadow gradient, to add interesting shading in the shadow of the first. */
-	UPROPERTY(EditAnywhere, Category=TranslucencySelfShadowing)
+	UPROPERTY(EditAnywhere, Category=TranslucencySelfShadowing, meta=(DisplayName = "Second Density Scale"))
 	float TranslucentSelfShadowSecondDensityScale;
 
 	/** Controls the strength of the second self shadow gradient. */
-	UPROPERTY(EditAnywhere, Category=TranslucencySelfShadowing)
+	UPROPERTY(EditAnywhere, Category=TranslucencySelfShadowing, meta=(DisplayName = "Second Opacity"))
 	float TranslucentSelfShadowSecondOpacity;
 
 	/** 
@@ -437,18 +437,18 @@ public:
 	 * Larger exponents give a less diffuse look (smaller, brighter backscattering highlight).
 	 * This is only used when the object is casting a volumetric translucent shadow from a directional light.
 	 */
-	UPROPERTY(EditAnywhere, Category=TranslucencySelfShadowing)
+	UPROPERTY(EditAnywhere, Category=TranslucencySelfShadowing, meta=(DisplayName = "Backscattering Exponent"))
 	float TranslucentBackscatteringExponent;
 
 	/** 
 	 * Colored extinction factor used to approximate multiple scattering in dense volumes. 
 	 * This is only used when the object is casting a volumetric translucent shadow.
 	 */
-	UPROPERTY(EditAnywhere, Category=TranslucencySelfShadowing)
+	UPROPERTY(EditAnywhere, Category=TranslucencySelfShadowing, meta=(DisplayName = "Multiple Scattering Extinction"))
 	FLinearColor TranslucentMultipleScatteringExtinction;
 
 	/** Local space distance to bias the translucent shadow.  Positive values move the shadow away from the light. */
-	UPROPERTY(EditAnywhere, Category=TranslucencySelfShadowing)
+	UPROPERTY(EditAnywhere, Category=TranslucencySelfShadowing, meta=(DisplayName = "Start Offset"))
 	float TranslucentShadowStartOffset;
 
 	/** Whether to draw on top of opaque pixels even if behind them. This only has meaning for translucency. */
