Fix for UE-6678. If both end points for the line attractor are the same, NaNs in the update results were causing crashes.

diff --git a/Engine/Source/Runtime/Engine/Private/Particles/ParticleModules.cpp b/Engine/Source/Runtime/Engine/Private/Particles/ParticleModules.cpp
--- a/Engine/Source/Runtime/Engine/Private/Particles/ParticleModules.cpp
+++ b/Engine/Source/Runtime/Engine/Private/Particles/ParticleModules.cpp
@@ -3433,6 +3433,13 @@ void UParticleModuleAttractorLine::PostEditChangeProperty(FPropertyChangedEvent&
 void UParticleModuleAttractorLine::Update(FParticleEmitterInstance* Owner, int32 Offset, float DeltaTime)
 {
 	FVector Line = EndPoint1 - EndPoint0;
+
+	// if both end points are the same, we end up with NaNs in the results of the update
+	if (Line.Size() == 0.0f)
+	{
+		Line = FVector(SMALL_NUMBER, SMALL_NUMBER, SMALL_NUMBER);
+	}
+
 	FVector LineNorm = Line;
 	LineNorm.Normalize();
 
