When adding a track in Matinee, the Choose Property box will now display both the variable name and the display name. Helpful for matching variables to properties in the details panel when properties have DisplayName metadata differs from the variable name.

diff --git a/Engine/Source/Runtime/Engine/Private/MatineeUtils.cpp b/Engine/Source/Runtime/Engine/Private/MatineeUtils.cpp
--- a/Engine/Source/Runtime/Engine/Private/MatineeUtils.cpp
+++ b/Engine/Source/Runtime/Engine/Private/MatineeUtils.cpp
@@ -44,7 +44,8 @@ protected:
 				// Is this property the desired type?
 				if (IsDesiredProperty(ClassMemberProperty))
 				{
-					const FString QualifiedFullPath = FString::Printf(TEXT("%s%s"), *Prefix, *ClassMemberProperty->GetName());
+					const FText DisplayName = ClassMemberProperty->GetDisplayNameText();
+					const FString QualifiedFullPath = FString::Printf(TEXT("%s (%s%s)"), *DisplayName.ToString(), *Prefix, *ClassMemberProperty->GetName());
 					GatheredPropertyPaths.Add(*QualifiedFullPath);
 				}
 
@@ -56,7 +57,8 @@ protected:
 						UProperty* StructMemberProperty = *StructFieldIt;
 						if (StructMemberProperty->HasAnyPropertyFlags(CPF_Interp) && IsDesiredProperty(StructMemberProperty))
 						{
-							const FString QualifiedFullPath = FString::Printf(TEXT("%s%s.%s"), *Prefix, *OuterStructProperty->GetName(), *StructMemberProperty->GetName());
+							const FText DisplayName = StructMemberProperty->GetDisplayNameText();
+							const FString QualifiedFullPath = FString::Printf(TEXT("%s (%s%s.%s)"), *DisplayName.ToString(), *Prefix, *OuterStructProperty->GetName(), *StructMemberProperty->GetName());
 							GatheredPropertyPaths.Add(*QualifiedFullPath);
 						}
 					}
