Fix for some bad data being passed into the viewport client from FSceneViewport, this was exposedby gamepad data going throught this api as well as keyboard.

diff --git a/Engine/Source/Runtime/Engine/Private/Slate/SceneViewport.cpp b/Engine/Source/Runtime/Engine/Private/Slate/SceneViewport.cpp
--- a/Engine/Source/Runtime/Engine/Private/Slate/SceneViewport.cpp
+++ b/Engine/Source/Runtime/Engine/Private/Slate/SceneViewport.cpp
@@ -727,7 +727,7 @@ FReply FSceneViewport::OnKeyDown( const FGeometry& InGeometry, const FKeyEvent&
 		// Switch to the viewport clients world before processing input
 		FScopedConditionalWorldSwitcher WorldSwitcher( ViewportClient );
 
-		if (!ViewportClient->InputKey(this, InKeyEvent.GetUserIndex(), Key, InKeyEvent.IsRepeat() ? IE_Repeat : IE_Pressed))
+		if (!ViewportClient->InputKey(this, InKeyEvent.GetUserIndex(), Key, InKeyEvent.IsRepeat() ? IE_Repeat : IE_Pressed, 1.0f, Key.IsGamepadKey()))
 		{
 			CurrentReplyState = FReply::Unhandled();
 		}
@@ -748,7 +748,7 @@ FReply FSceneViewport::OnKeyUp( const FGeometry& InGeometry, const FKeyEvent& In
 		// Switch to the viewport clients world before processing input
 		FScopedConditionalWorldSwitcher WorldSwitcher( ViewportClient );
 
-		if (!ViewportClient->InputKey(this, InKeyEvent.GetUserIndex(), Key, IE_Released))
+		if (!ViewportClient->InputKey(this, InKeyEvent.GetUserIndex(), Key, IE_Released, 1.0f, Key.IsGamepadKey()))
 		{
 			CurrentReplyState = FReply::Unhandled();
 		}
@@ -762,14 +762,15 @@ FReply FSceneViewport::OnAnalogValueChanged(const FGeometry& MyGeometry, const F
 	// Start a new reply state
 	CurrentReplyState = FReply::Handled();
 
-	KeyStateMap.Add(InAnalogInputEvent.GetKey(), true);
+	FKey Key = InAnalogInputEvent.GetKey();
+	KeyStateMap.Add(Key, true);
 
 	if (ViewportClient)
 	{
 		// Switch to the viewport clients world before processing input
 		FScopedConditionalWorldSwitcher WorldSwitcher(ViewportClient);
 
-		if (!ViewportClient->InputAxis(this, InAnalogInputEvent.GetUserIndex(), InAnalogInputEvent.GetKey(), InAnalogInputEvent.GetKey() == EKeys::Gamepad_RightY ? -InAnalogInputEvent.GetAnalogValue() : InAnalogInputEvent.GetAnalogValue(), FApp::GetDeltaTime(), 1, true))
+		if (!ViewportClient->InputAxis(this, InAnalogInputEvent.GetUserIndex(), Key, Key == EKeys::Gamepad_RightY ? -InAnalogInputEvent.GetAnalogValue() : InAnalogInputEvent.GetAnalogValue(), FApp::GetDeltaTime(), 1, Key.IsGamepadKey()))
 		{
 			CurrentReplyState = FReply::Unhandled();
 		}
@@ -788,7 +789,7 @@ FReply FSceneViewport::OnKeyChar( const FGeometry& InGeometry, const FCharacterE
 		// Switch to the viewport clients world before processing input
 		FScopedConditionalWorldSwitcher WorldSwitcher( ViewportClient );
 
-		if( !ViewportClient->InputChar( this,0, InCharacterEvent.GetCharacter() ) )
+		if (!ViewportClient->InputChar(this, InCharacterEvent.GetUserIndex(), InCharacterEvent.GetCharacter()))
 		{
 			CurrentReplyState = FReply::Unhandled();
 		}
