Fixed WITH_EDITOR-related compile error.

diff --git a/Engine/Source/Runtime/Engine/Private/MatineeUtils.cpp b/Engine/Source/Runtime/Engine/Private/MatineeUtils.cpp
--- a/Engine/Source/Runtime/Engine/Private/MatineeUtils.cpp
+++ b/Engine/Source/Runtime/Engine/Private/MatineeUtils.cpp
@@ -44,8 +44,12 @@ protected:
 				// Is this property the desired type?
 				if (IsDesiredProperty(ClassMemberProperty))
 				{
+#if WITH_EDITOR
 					const FText DisplayName = ClassMemberProperty->GetDisplayNameText();
 					const FString QualifiedFullPath = FString::Printf(TEXT("%s (%s%s)"), *DisplayName.ToString(), *Prefix, *ClassMemberProperty->GetName());
+#else
+					const FString QualifiedFullPath = FString::Printf(TEXT("%s%s"), *Prefix, *ClassMemberProperty->GetName());
+#endif
 					GatheredPropertyPaths.Add(*QualifiedFullPath);
 				}
 
@@ -57,8 +61,12 @@ protected:
 						UProperty* StructMemberProperty = *StructFieldIt;
 						if (StructMemberProperty->HasAnyPropertyFlags(CPF_Interp) && IsDesiredProperty(StructMemberProperty))
 						{
+#if WITH_EDITOR
 							const FText DisplayName = StructMemberProperty->GetDisplayNameText();
 							const FString QualifiedFullPath = FString::Printf(TEXT("%s (%s%s.%s)"), *DisplayName.ToString(), *Prefix, *OuterStructProperty->GetName(), *StructMemberProperty->GetName());
+#else
+							const FString QualifiedFullPath = FString::Printf(TEXT("%s%s.%s"), *Prefix, *OuterStructProperty->GetName(), *StructMemberProperty->GetName());
+#endif
 							GatheredPropertyPaths.Add(*QualifiedFullPath);
 						}
 					}
