fixed UE-5791
Adding a Scene Capture 2D to Render Target will override Motion Blur settings in level

diff --git a/Engine/Source/Runtime/Engine/Public/SceneManagement.h b/Engine/Source/Runtime/Engine/Public/SceneManagement.h
--- a/Engine/Source/Runtime/Engine/Public/SceneManagement.h
+++ b/Engine/Source/Runtime/Engine/Public/SceneManagement.h
@@ -1448,6 +1448,7 @@ private:
 	bool bKeepAndUpdateThisFrame;
 };
 
+// stored in the scene, can be shared for multiple views
 class FMotionBlurInfoData
 {
 public:
@@ -1457,19 +1458,21 @@ public:
 	/** 
 	 *	Set the primitives motion blur info
 	 * 
-	 *	@param PrimitiveSceneInfo	The primitive to add
+	 *	@param PrimitiveSceneInfo The primitive to add, must not be 0
 	 */
 	void UpdatePrimitiveMotionBlur(FPrimitiveSceneInfo* PrimitiveSceneInfo);
 
 	/** 
-	 *	Set the primitives motion blur info
+	 *	Unsets the primitives motion blur info
 	 * 
-	 *	@param PrimitiveSceneInfo	The primitive to add
+	 *	@param PrimitiveSceneInfo The primitive to remove, must not be 0
 	 */
 	void RemovePrimitiveMotionBlur(FPrimitiveSceneInfo* PrimitiveSceneInfo);
 
 	/**
 	 * Creates any needed motion blur infos if needed and saves the transforms of the frame we just completed
+	 * called in RenderFinish()
+	 * @param InScene must not be 0
 	 */
 	void UpdateMotionBlurCache(class FScene* InScene);
 
@@ -1487,7 +1490,7 @@ public:
 	 */
 	bool GetPrimitiveMotionBlurInfo(const FPrimitiveSceneInfo* PrimitiveSceneInfo, FMatrix& OutPreviousLocalToWorld);
 
-	/** */
+	/** Request to clear all stored motion blur data for this scene. */
 	void SetClearMotionBlurInfo();
 
 	/**
@@ -1498,8 +1501,6 @@ public:
 private:
 	/** The motion blur info entries for the frame. Accessed on Renderthread only! */
 	TMap<FPrimitiveComponentId, FMotionBlurInfo> MotionBlurInfos;
-	/** Unique "frame number" counter to make sure we don't double update */
-	uint32 CacheUpdateCount;	
 	/** */
 	bool bShouldClearMotionBlurInfo;
 
diff --git a/Engine/Source/Runtime/Renderer/Private/DeferredShadingRenderer.cpp b/Engine/Source/Runtime/Renderer/Private/DeferredShadingRenderer.cpp
--- a/Engine/Source/Runtime/Renderer/Private/DeferredShadingRenderer.cpp
+++ b/Engine/Source/Runtime/Renderer/Private/DeferredShadingRenderer.cpp
@@ -565,12 +565,6 @@ void FDeferredShadingSceneRenderer::RenderFinish(FRHICommandListImmediate& RHICm
 
 	FSceneRenderer::RenderFinish(RHICmdList);
 
-	//grab the new transform out of the proxies for next frame
-	if(ViewFamily.EngineShowFlags.MotionBlur || ViewFamily.EngineShowFlags.TemporalAA)
-	{
-		Scene->MotionBlurInfoData.UpdateMotionBlurCache(Scene);
-	}
-
 	// Some RT should be released as early as possible to allow sharing of that memory for other purposes.
 	// SceneColor is be released in tone mapping, if not we want to get access to the HDR scene color after this pass so we keep it.
 	// This becomes even more important with some limited VRam (XBoxOne).
@@ -1023,6 +1017,12 @@ void FDeferredShadingSceneRenderer::Render(FRHICommandListImmediate& RHICmdList)
 		GSceneRenderTargets.AdjustGBufferRefCount(-1);
 	}
 
+	//grab the new transform out of the proxies for next frame
+	if(VelocityRT)
+	{
+		Scene->MotionBlurInfoData.UpdateMotionBlurCache(Scene);
+	}
+
 	VelocityRT.SafeRelease();
 
 	RenderFinish(RHICmdList);
diff --git a/Engine/Source/Runtime/Renderer/Private/RendererScene.cpp b/Engine/Source/Runtime/Renderer/Private/RendererScene.cpp
--- a/Engine/Source/Runtime/Renderer/Private/RendererScene.cpp
+++ b/Engine/Source/Runtime/Renderer/Private/RendererScene.cpp
@@ -2464,8 +2464,7 @@ TStaticMeshDrawList<TBasePassForForwardShadingDrawingPolicy<FMovableDirectionalL
 -----------------------------------------------------------------------------*/
 
 FMotionBlurInfoData::FMotionBlurInfoData()
-	: CacheUpdateCount(0)
-	, bShouldClearMotionBlurInfo(false)
+	: bShouldClearMotionBlurInfo(false)
 {
 
 }
@@ -2552,7 +2551,7 @@ void FMotionBlurInfoData::RestoreForPausedMotionBlur()
 
 void FMotionBlurInfoData::UpdateMotionBlurCache(FScene* InScene)
 {
-	check(IsInRenderingThread());
+	check(InScene && IsInRenderingThread());
 
 	if (InScene->GetFeatureLevel() >= ERHIFeatureLevel::SM4)
 	{
diff --git a/Engine/Source/Runtime/Renderer/Private/VelocityRendering.cpp b/Engine/Source/Runtime/Renderer/Private/VelocityRendering.cpp
--- a/Engine/Source/Runtime/Renderer/Private/VelocityRendering.cpp
+++ b/Engine/Source/Runtime/Renderer/Private/VelocityRendering.cpp
@@ -634,11 +634,14 @@ void FDeferredShadingSceneRenderer::RenderVelocities(FRHICommandListImmediate& R
 
 		bNeedsVelocity |= bMotionBlur || bTemporalAA;
 	}
-	if( !bNeedsVelocity || GPixelFormats[PF_G16R16].Supported == false )
+	if( !bNeedsVelocity || !GPixelFormats[PF_G16R16].Supported )
 	{
 		return;
 	}
 
+	// this is not supported
+	check(!Views[0].bIsSceneCapture);
+
 	SCOPED_DRAW_EVENT(RHICmdList, RenderVelocities);
 
 	const FIntPoint BufferSize = GSceneRenderTargets.GetBufferSizeXY();
