UE-6804 - Demo recording filename needs a replacement key for the map name

diff --git a/Engine/Source/Runtime/Engine/Private/World.cpp b/Engine/Source/Runtime/Engine/Private/World.cpp
--- a/Engine/Source/Runtime/Engine/Private/World.cpp
+++ b/Engine/Source/Runtime/Engine/Private/World.cpp
@@ -2531,6 +2531,8 @@ bool UWorld::HandleDemoRecordCommand( const TCHAR* Cmd, FOutputDevice& Ar, UWorl
 		return true;
 	}
 
+	DemoName.ReplaceInline( TEXT( "%m" ), *GetMapName() );
+
 	int32 Year, Month, DayOfWeek, Day, Hour, Min, Sec, MSec;
 	FPlatformTime::SystemTime( Year, Month, DayOfWeek, Day, Hour, Min, Sec, MSec );
 
