#vr Added GetTextSafeRegionBounds() function to stereo interface, to allow for displaying of stats in a more readable region of the screen

diff --git a/Engine/Source/Runtime/Engine/Public/StereoRendering.h b/Engine/Source/Runtime/Engine/Public/StereoRendering.h
--- a/Engine/Source/Runtime/Engine/Public/StereoRendering.h
+++ b/Engine/Source/Runtime/Engine/Public/StereoRendering.h
@@ -25,6 +25,12 @@ public:
      */
     virtual void AdjustViewRect(enum EStereoscopicPass StereoPass, int32& X, int32& Y, uint32& SizeX, uint32& SizeY) const = 0;
 
+	/**
+	 * Gets the percentage bounds of the safe region to draw in.  This allows things like stat rendering to appear within the readable portion of the stereo view.
+	 * @return	The centered percentage of the view that is safe to draw readable text in
+	 */
+	virtual FVector2D GetTextSafeRegionBounds() const { return FVector2D(0.75f, 0.75f); }
+
     /**
 	 * Calculates the offset for the camera position, given the specified position, rotation, and world scale
 	 */
