#ue4 - UPrimitiveComponent::AreAllCollideableDescendantsRelative() optimized to reduce potential dynamic allocations.

diff --git a/Engine/Source/Runtime/Engine/Private/PrimitiveComponent.cpp b/Engine/Source/Runtime/Engine/Private/PrimitiveComponent.cpp
--- a/Engine/Source/Runtime/Engine/Private/PrimitiveComponent.cpp
+++ b/Engine/Source/Runtime/Engine/Private/PrimitiveComponent.cpp
@@ -2087,13 +2087,13 @@ bool UPrimitiveComponent::AreAllCollideableDescendantsRelative(bool bAllowCached
 		}
 
 		// Check all descendant PrimitiveComponents
-		TArray<USceneComponent*, TInlineAllocator<16>> ComponentStack;
-		ComponentStack.Reserve(FMath::Max(16, AttachChildren.Num()));
+		TInlineComponentArray<USceneComponent*> ComponentStack;
+		ComponentStack.Reserve(FMath::Max<uint32>(NumInlinedActorComponents, AttachChildren.Num()));
 
 		ComponentStack.Append(AttachChildren);
 		while (ComponentStack.Num() > 0)
 		{
-			USceneComponent* const CurrentComp = ComponentStack.Pop();
+			USceneComponent* const CurrentComp = ComponentStack.Pop(false);
 			if (CurrentComp)
 			{
 				// Is the component not using relative position?
