#ue4 - USkinnedMeshComponent::AnimUpdateRateTick() optimized to reduce potential dynamic allocations. 0.40ms -> 0.26ms in test scene with many characters.

diff --git a/Engine/Source/Runtime/Engine/Private/SkinnedMeshComponent.cpp b/Engine/Source/Runtime/Engine/Private/SkinnedMeshComponent.cpp
--- a/Engine/Source/Runtime/Engine/Private/SkinnedMeshComponent.cpp
+++ b/Engine/Source/Runtime/Engine/Private/SkinnedMeshComponent.cpp
@@ -1762,9 +1762,10 @@ void USkinnedMeshComponent::AnimUpdateRateTick()
 	float MaxDistanceFactor = 0.f;
 
 	// Gather Actor's components
-	TArray<USceneComponent *> ComponentStack;
-	USceneComponent* CurrentComponent = Owner->GetRootComponent();
+	TInlineComponentArray<USceneComponent *> ComponentStack;
+	ComponentStack.Reserve(NumInlinedActorComponents);
 
+	USceneComponent* CurrentComponent = Owner->GetRootComponent();
 	while (CurrentComponent)
 	{
 		// push children on the stack so they get tested later
@@ -1781,7 +1782,7 @@ void USkinnedMeshComponent::AnimUpdateRateTick()
 			}
 		}
 
-		CurrentComponent = (ComponentStack.Num() > 0) ? ComponentStack.Pop() : NULL;
+		CurrentComponent = (ComponentStack.Num() > 0) ? ComponentStack.Pop(false) : NULL;
 	}
 
 	// Figure out which update rate should be used.
