#ue4 - Remove check for DestructibleActor in CharacterMovementComponent floor check, it is handled by component mobility check.

diff --git a/Engine/Source/Runtime/Engine/Private/CharacterMovementComponent.cpp b/Engine/Source/Runtime/Engine/Private/CharacterMovementComponent.cpp
--- a/Engine/Source/Runtime/Engine/Private/CharacterMovementComponent.cpp
+++ b/Engine/Source/Runtime/Engine/Private/CharacterMovementComponent.cpp
@@ -4623,9 +4623,7 @@ void UCharacterMovementComponent::FindFloor(const FVector& CapsuleLocation, FFin
 			{
 				MutableThis->bForceNextFloorCheck = !MovementBase->IsCollisionEnabled()
 				|| MovementBase->GetCollisionResponseToChannel(CollisionChannel) != ECR_Block
-				|| (MovementBase->Mobility == EComponentMobility::Movable)
-				|| MovementBaseUtility::IsDynamicBase(MovementBase)
-				|| (Cast<const ADestructibleActor>(BaseActor) != NULL);
+				|| MovementBaseUtility::IsDynamicBase(MovementBase);
 			}
 
 			const bool IsActorBasePendingKill = BaseActor && BaseActor->IsPendingKill();
