#ue4 - Check that mobility is not Movable before doing other expensive validation in SceneComponent::MoveComponent().

diff --git a/Engine/Source/Runtime/Engine/Private/SceneComponent.cpp b/Engine/Source/Runtime/Engine/Private/SceneComponent.cpp
--- a/Engine/Source/Runtime/Engine/Private/SceneComponent.cpp
+++ b/Engine/Source/Runtime/Engine/Private/SceneComponent.cpp
@@ -1389,16 +1389,19 @@ bool USceneComponent::MoveComponent( const FVector& Delta, const FRotator& NewRo
 {
 #if !(UE_BUILD_SHIPPING || UE_BUILD_TEST)
 	// make sure mobility is movable, otherwise you shouldn't try to move
-	if ( UWorld * World = GetWorld() )
+	if (Mobility != EComponentMobility::Movable)
 	{
-		ULevel* Level = GetComponentLevel();
-
-		// It's only a problem if we're in gameplay, and the owning level is visible
-		if (World->HasBegunPlay() && IsRegistered() && Level && Level->bIsVisible && Mobility != EComponentMobility::Movable)
+		if (UWorld * World = GetWorld())
 		{
-			FMessageLog("Performance").Warning( FText::Format(LOCTEXT("InvalidMove", "Mobility of {0} : {1} has to be 'Movable' if you'd like to move. "), 
-				FText::FromString(GetNameSafe(GetOwner())), FText::FromString(GetName())));
-			return false;
+			ULevel* Level = GetComponentLevel();
+
+			// It's only a problem if we're in gameplay, and the owning level is visible
+			if (World->HasBegunPlay() && IsRegistered() && Level && Level->bIsVisible)
+			{
+				FMessageLog("Performance").Warning(FText::Format(LOCTEXT("InvalidMove", "Mobility of {0} : {1} has to be 'Movable' if you'd like to move. "),
+					FText::FromString(GetNameSafe(GetOwner())), FText::FromString(GetName())));
+				return false;
+			}
 		}
 	}
 #endif
