Marked CharacterMovementComponent::SetAvoidanceEnabled as unsafe for ConstructionScript use #UE4

- also, added a safety check for character owner in SetAvoidanceEnabled
- plus minor refactor

diff --git a/Engine/Source/Runtime/Engine/Classes/GameFramework/CharacterMovementComponent.h b/Engine/Source/Runtime/Engine/Classes/GameFramework/CharacterMovementComponent.h
--- a/Engine/Source/Runtime/Engine/Classes/GameFramework/CharacterMovementComponent.h
+++ b/Engine/Source/Runtime/Engine/Classes/GameFramework/CharacterMovementComponent.h
@@ -12,6 +12,7 @@
 #include "CharacterMovementComponent.generated.h"
 
 class FDebugDisplayInfo;
+class ACharacter;
 
 /** Data about the floor for walking movement, used by CharacterMovementComponent. */
 USTRUCT(BlueprintType)
@@ -138,7 +139,7 @@ protected:
 
 	/** Character movement component belongs to */
 	UPROPERTY()
-	class ACharacter* CharacterOwner;
+	ACharacter* CharacterOwner;
 
 public:
 
@@ -727,12 +728,12 @@ public:
 	FNavLocation CachedNavLocation;
 
 	/** Change avoidance state and registers in RVO manager if needed */
-	UFUNCTION(BlueprintCallable, Category="Pawn|Components|CharacterMovement")
+	UFUNCTION(BlueprintCallable, Category = "Pawn|Components|CharacterMovement", meta = (UnsafeDuringActorConstruction = "true"))
 	void SetAvoidanceEnabled(bool bEnable);
 
 	/** Get the Character that owns UpdatedComponent. */
 	UFUNCTION(BlueprintCallable, Category="Pawn|Components|CharacterMovement")
-	class ACharacter* GetCharacterOwner() const;
+	ACharacter* GetCharacterOwner() const;
 
 	/**
 	 * Change movement mode.
@@ -1803,10 +1804,10 @@ public:
 	virtual void Clear();
 
 	/** Called to set up this saved move (when initially created) to make a predictive correction. */
-	virtual void SetMoveFor(class ACharacter* C, float DeltaTime, FVector const& NewAccel, class FNetworkPredictionData_Client_Character & ClientData);
+	virtual void SetMoveFor(ACharacter* C, float DeltaTime, FVector const& NewAccel, class FNetworkPredictionData_Client_Character & ClientData);
 
 	/** Set the properties describing the position, etc. of the moved pawn at the start of the move. */
-	virtual void SetInitialPosition(class ACharacter* C);
+	virtual void SetInitialPosition(ACharacter* C);
 
 	/** @Return true if this move is an "important" move that should be sent again if not acked by the server */
 	virtual bool IsImportantMove(const FSavedMovePtr& LastAckedMove) const;
@@ -1821,13 +1822,13 @@ public:
 	};
 
 	/** Set the properties describing the final position, etc. of the moved pawn. */
-	virtual void PostUpdate(class ACharacter* C, EPostUpdateMode PostUpdateMode);
+	virtual void PostUpdate(ACharacter* C, EPostUpdateMode PostUpdateMode);
 	
 	/** @Return true if this move can be combined with NewMove for replication without changing any behavior */
-	virtual bool CanCombineWith(const FSavedMovePtr& NewMove, class ACharacter* InPawn, float MaxDelta) const;
+	virtual bool CanCombineWith(const FSavedMovePtr& NewMove, ACharacter* InPawn, float MaxDelta) const;
 	
 	/** Called before ClientUpdatePosition uses this SavedMove to make a predictive correction	 */
-	virtual void PrepMoveFor( class ACharacter* C );
+	virtual void PrepMoveFor(ACharacter* C);
 
 	/** @returns a byte containing encoded special movement information (jumping, crouching, etc.)	 */
 	virtual uint8 GetCompressedFlags() const;
@@ -1945,9 +1946,13 @@ public:
 		It will measure the accuracy between passed in DeltaTime and how Server will calculate its DeltaTime.
 		If inaccuracy is too high, it will reset CurrentTimeStamp to maintain a high level of accuracy.
 		@return DeltaTime to use for Client's physics simulation prior to replicate move to server. */
-	float UpdateTimeStampAndDeltaTime(float DeltaTime, class ACharacter & CharacterOwner, class UCharacterMovementComponent & CharacterMovementComponent);
+	float UpdateTimeStampAndDeltaTime(float DeltaTime, ACharacter & CharacterOwner, class UCharacterMovementComponent & CharacterMovementComponent);
 };
 
+FORCEINLINE ACharacter* UCharacterMovementComponent::GetCharacterOwner() const
+{
+	return CharacterOwner;
+}
 
 class ENGINE_API FNetworkPredictionData_Server_Character : public FNetworkPredictionData_Server
 {
diff --git a/Engine/Source/Runtime/Engine/Private/CharacterMovementComponent.cpp b/Engine/Source/Runtime/Engine/Private/CharacterMovementComponent.cpp
--- a/Engine/Source/Runtime/Engine/Private/CharacterMovementComponent.cpp
+++ b/Engine/Source/Runtime/Engine/Private/CharacterMovementComponent.cpp
@@ -422,12 +422,6 @@ bool UCharacterMovementComponent::HasValidData() const
 	return UpdatedComponent && IsValid(CharacterOwner);
 }
 
-ACharacter* UCharacterMovementComponent::GetCharacterOwner() const
-{
-	return CharacterOwner;
-}
-
-
 FCollisionShape UCharacterMovementComponent::GetPawnCapsuleCollisionShape(const EShrinkCapsuleExtent ShrinkMode, const float CustomShrinkAmount) const
 {
 	FVector Extent = GetPawnCapsuleExtent(ShrinkMode, CustomShrinkAmount);
@@ -6497,10 +6491,16 @@ void UCharacterMovementComponent::SetAvoidanceEnabled(bool bEnable)
 	{
 		bUseRVOAvoidance = bEnable;
 
-		UAvoidanceManager* AvoidanceManager = GetWorld()->GetAvoidanceManager();
-		if (AvoidanceManager && bEnable && AvoidanceUID == 0)
+		// this is a safety check - it's possible to not have CharacterOwner at this point if this function gets
+		// called too early
+		ensure(GetCharacterOwner());
+		if (GetCharacterOwner() != nullptr)
 		{
-			AvoidanceManager->RegisterMovementComponent(this, AvoidanceWeight);
+			UAvoidanceManager* AvoidanceManager = GetWorld()->GetAvoidanceManager();
+			if (AvoidanceManager && bEnable && AvoidanceUID == 0)
+			{
+				AvoidanceManager->RegisterMovementComponent(this, AvoidanceWeight);
+			}
 		}
 	}
 }
