Fixed AI Logging on dedicated servers using EnableAILogging commandline argument. Change integrated from Dev branch.

diff --git a/Engine/Source/Runtime/Engine/Private/VisualLogger/VisualLogger.cpp b/Engine/Source/Runtime/Engine/Private/VisualLogger/VisualLogger.cpp
--- a/Engine/Source/Runtime/Engine/Private/VisualLogger/VisualLogger.cpp
+++ b/Engine/Source/Runtime/Engine/Private/VisualLogger/VisualLogger.cpp
@@ -36,25 +36,24 @@ FVisualLogger::FVisualLogger()
 	}
 }
 
-UWorld* FVisualLogger::GetWorld()
+UWorld* FVisualLogger::GetWorld(const class UObject* Object)
 {
-	UWorld* World = GWorld;
+	UWorld* World = Object ? GEngine->GetWorldFromContextObject(Object, false) : NULL;
 #if WITH_EDITOR
 	UEditorEngine *EEngine = Cast<UEditorEngine>(GEngine);
-	if (GIsEditor && EEngine != NULL)
+	if (GIsEditor && EEngine != NULL && World == NULL)
 	{
 		// lets use PlayWorld during PIE/Simulate and regular world from editor otherwise, to draw debug information
 		World = EEngine->PlayWorld != NULL ? EEngine->PlayWorld : EEngine->GetEditorWorldContext().World();
 	}
-	else 
+
 #endif
-	if (!GIsEditor)
+	if (!GIsEditor && World == NULL)
 	{
-
 		World = GEngine->GetWorld();
 	}
 
-	return World;
+	return World != NULL ? World : GWorld;
 }
 
 void FVisualLogger::Shutdown()
diff --git a/Engine/Source/Runtime/Engine/Public/VisualLogger/VisualLogger.h b/Engine/Source/Runtime/Engine/Public/VisualLogger/VisualLogger.h
--- a/Engine/Source/Runtime/Engine/Public/VisualLogger/VisualLogger.h
+++ b/Engine/Source/Runtime/Engine/Public/VisualLogger/VisualLogger.h
@@ -215,7 +215,7 @@ public:
 private:
 	FVisualLogger();
 	virtual void Serialize(const TCHAR* V, ELogVerbosity::Type Verbosity, const class FName& Category) override { ensureMsgf(0, TEXT("Regular serialize is forbiden for visual logs")); }
-	class UWorld* GetWorld();
+	class UWorld* GetWorld(const class UObject* Object);
 
 protected:
 	/** Array of output devices to redirect to */
diff --git a/Engine/Source/Runtime/Engine/Public/VisualLogger/VisualLogger.inl b/Engine/Source/Runtime/Engine/Public/VisualLogger/VisualLogger.inl
--- a/Engine/Source/Runtime/Engine/Public/VisualLogger/VisualLogger.inl
+++ b/Engine/Source/Runtime/Engine/Public/VisualLogger/VisualLogger.inl
@@ -67,7 +67,8 @@ FVisualLogEntry* FVisualLogger::GetEntryToWrite(const class UObject* Object, flo
 
 	bool InitializeNewEntry = false;
 
-	UWorld* World = GetWorld();
+	UWorld* World = GetWorld(Object);
+
 	if (CurrentEntryPerObject.Contains(LogOwner))
 	{
 		CurrentEntry = &CurrentEntryPerObject[LogOwner];
@@ -123,27 +124,27 @@ FVisualLogEntry* FVisualLogger::GetEntryToWrite(const class UObject* Object, flo
 				ObjectAsActor->GrabDebugSnapshot(CurrentEntry);
 			}
 		}
+	}
 
-		if (World)
-		{
-			//set next tick timer to flush obsolete/old entries
-			World->GetTimerManager().SetTimerForNextTick(FTimerDelegate::CreateLambda(
-				[this](){
-				for (auto& CurrentPair : CurrentEntryPerObject)
+	if (World)
+	{
+		//set next tick timer to flush obsolete/old entries
+		World->GetTimerManager().SetTimerForNextTick(FTimerDelegate::CreateLambda(
+			[this](){
+			for (auto& CurrentPair : CurrentEntryPerObject)
+			{
+				FVisualLogEntry* Entry = &CurrentPair.Value;
+				if (Entry->TimeStamp >= 0) // CurrentEntry->TimeStamp == -1 means it's not initialized entry information
 				{
-					FVisualLogEntry* Entry = &CurrentPair.Value;
-					if (Entry->TimeStamp >= 0) // CurrentEntry->TimeStamp == -1 means it's not initialized entry information
+					for (auto* Device : OutputDevices)
 					{
-						for (auto* Device : OutputDevices)
-						{
-							Device->Serialize(CurrentPair.Key, ObjectToNameMap[CurrentPair.Key], ObjectToClassNameMap[CurrentPair.Key], *Entry);
-						}
-						Entry->Reset();
+						Device->Serialize(CurrentPair.Key, ObjectToNameMap[CurrentPair.Key], ObjectToClassNameMap[CurrentPair.Key], *Entry);
 					}
+					Entry->Reset();
 				}
 			}
-			));
 		}
+		));
 	}
 
 	return CurrentEntry;
