UE4 - fix crash with missing cluster tree

diff --git a/Engine/Source/Runtime/Engine/Private/HierarchicalInstancedStaticMesh.cpp b/Engine/Source/Runtime/Engine/Private/HierarchicalInstancedStaticMesh.cpp
--- a/Engine/Source/Runtime/Engine/Private/HierarchicalInstancedStaticMesh.cpp
+++ b/Engine/Source/Runtime/Engine/Private/HierarchicalInstancedStaticMesh.cpp
@@ -1451,7 +1451,6 @@ void UHierarchicalInstancedStaticMeshComponent::ClearInstances()
 		bConcurrentRemoval = true;
 	}
 
-	ClusterTree.Empty();
 	NumBuiltInstances = 0;
 	UnbuiltInstanceBounds.Init();
 
