BP: new timer functions with delegate input.

diff --git a/Engine/Source/Runtime/Engine/Classes/Kismet/KismetSystemLibrary.h b/Engine/Source/Runtime/Engine/Classes/Kismet/KismetSystemLibrary.h
--- a/Engine/Source/Runtime/Engine/Classes/Kismet/KismetSystemLibrary.h
+++ b/Engine/Source/Runtime/Engine/Classes/Kismet/KismetSystemLibrary.h
@@ -249,26 +249,84 @@ class UKismetSystemLibrary : public UBlueprintFunctionLibrary
 	UFUNCTION(BlueprintCallable, meta=(Latent, LatentInfo="LatentInfo", WorldContext="WorldContextObject", ExpandEnumAsExecs="MoveAction", OverTime="0.2"), Category="Components")
 	static void MoveComponentTo(USceneComponent* Component, FVector TargetRelativeLocation, FRotator TargetRelativeRotation, bool bEaseOut, bool bEaseIn, float OverTime, TEnumAsByte<EMoveComponentAction::Type> MoveAction, FLatentActionInfo LatentInfo);
 
-	// --- Timer functions ------------------------------
+	// --- Timer functions with delegate input ----------
 
 	/**
 	 * Set a timer to execute delegate. Setting an existing timer will reset that timer with updated parameters.
-	 * @param Object		Object that implements the delegate function. Defaults to self (this blueprint)
-	 * @param FunctionName	Delegate function name. Can be a K2 function or a Custom Event.
+	 * @param Delegate		Delegate. Can be a K2 function or a Custom Event.
 	 * @param Time			How long to wait before executing the delegate, in seconds. Setting a timer to <= 0 seconds will clear it if it is set.
 	 * @param bLooping		true to keep executing the delegate every Time seconds, false to execute delegate only once.
 	 */
-	UFUNCTION(BlueprintCallable, meta=(FriendlyName = "SetTimer", DefaultToSelf = "Object"), Category="Utilities|Time")
-	static void K2_SetTimer(UObject* Object, FString FunctionName, float Time, bool bLooping);
+	UFUNCTION(BlueprintCallable, meta=(FriendlyName = "SetTimer_Delegate"), Category="Utilities|Time")
+	static void K2_SetTimerDelegate(FTimerDynamicDelegate Delegate, float Time, bool bLooping);
 
 	/**
-	 * Set a timer to execute delegate. Setting an existing timer will reset that timer with updated parameters.
+	 * Clears a set timer.
+	 * @param Delegate		Delegate. Can be a K2 function or a Custom Event.
+	 */
+	UFUNCTION(BlueprintCallable, meta=(FriendlyName = "ClearTimer_Delegate", DefaultToSelf = "Object"), Category="Utilities|Time")
+	static void K2_ClearTimerDelegate(FTimerDynamicDelegate Delegate);
+
+	/**
+	 * Pauses a set timer at its current elapsed time.
 	 * @param Delegate		Delegate. Can be a K2 function or a Custom Event.
+	 */
+	UFUNCTION(BlueprintCallable, meta=(FriendlyName = "PauseTimer_Delegate", DefaultToSelf = "Object"), Category="Utilities|Time")
+	static void K2_PauseTimerDelegate(FTimerDynamicDelegate Delegate);
+
+	/**
+	 * Resumes a paused timer from its current elapsed time.
+	 * @param Delegate		Delegate. Can be a K2 function or a Custom Event.
+	 */
+	UFUNCTION(BlueprintCallable, meta=(FriendlyName = "UnPauseTimer_Delegate", DefaultToSelf = "Object"), Category="Utilities|Time")
+	static void K2_UnPauseTimerDelegate(FTimerDynamicDelegate Delegate);
+
+	/**
+	 * Returns true is a timer exists and is active for the given delegate, false otherwise.
+	 * @param Delegate		Delegate. Can be a K2 function or a Custom Event.
+	 */
+	UFUNCTION(BlueprintPure, meta=(FriendlyName = "IsTimerActive_Delegate", DefaultToSelf = "Object"), Category="Utilities|Time")
+	static bool K2_IsTimerActiveDelegate(FTimerDynamicDelegate Delegate);
+
+	/**
+	 * Returns true is a timer exists and is paused for the given delegate, false otherwise.
+	 * @param Delegate		Delegate. Can be a K2 function or a Custom Event.
+	 */
+	UFUNCTION(BlueprintPure, meta = (FriendlyName = "IsTimerPaused_Delegate", DefaultToSelf = "Object"), Category = "Utilities|Time")
+	static bool K2_IsTimerPausedDelegate(FTimerDynamicDelegate Delegate);
+
+	/**
+	 * Returns true is a timer for the given delegate exists, false otherwise.
+	 * @param Delegate		Delegate. Can be a K2 function or a Custom Event.
+	 */
+	UFUNCTION(BlueprintPure, meta = (FriendlyName = "TimerExists_Delegate", DefaultToSelf = "Object"), Category = "Utilities|Time")
+	static bool K2_TimerExistsDelegate(FTimerDynamicDelegate Delegate);
+	
+	/**
+	 * Returns elapsed time for the given delegate (time since current countdown iteration began).
+	 * @param Delegate		Delegate. Can be a K2 function or a Custom Event.
+	 */
+	UFUNCTION(BlueprintPure, meta=(FriendlyName = "GetTimerElapsedTime_Delegate", DefaultToSelf = "Object"), Category="Utilities|Time")
+	static float K2_GetTimerElapsedTimeDelegate(FTimerDynamicDelegate Delegate);
+
+	/**
+	 * Returns time until the timer will next execute its delegate.
+	 * @param Delegate		Delegate. Can be a K2 function or a Custom Event.
+	 */
+	UFUNCTION(BlueprintPure, meta=(FriendlyName = "GetTimerRemainingTime_Delegate", DefaultToSelf = "Object"), Category="Utilities|Time")
+	static float K2_GetTimerRemainingTimeDelegate(FTimerDynamicDelegate Delegate);
+
+	// --- Timer functions ------------------------------
+
+	/**
+	 * Set a timer to execute delegate. Setting an existing timer will reset that timer with updated parameters.
+	 * @param Object		Object that implements the delegate function. Defaults to self (this blueprint)
+	 * @param FunctionName	Delegate function name. Can be a K2 function or a Custom Event.
 	 * @param Time			How long to wait before executing the delegate, in seconds. Setting a timer to <= 0 seconds will clear it if it is set.
 	 * @param bLooping		true to keep executing the delegate every Time seconds, false to execute delegate only once.
 	 */
-	UFUNCTION(BlueprintCallable, meta=(FriendlyName = "SetTimerDelegate"), Category="Utilities|Time")
-	static void K2_SetTimerDelegate(FTimerDynamicDelegate Delegate, float Time, bool bLooping);
+	UFUNCTION(BlueprintCallable, meta=(FriendlyName = "SetTimer", DefaultToSelf = "Object"), Category="Utilities|Time")
+	static void K2_SetTimer(UObject* Object, FString FunctionName, float Time, bool bLooping);
 
 	/**
 	 * Clears a set timer.
diff --git a/Engine/Source/Runtime/Engine/Private/KismetSystemLibrary.cpp b/Engine/Source/Runtime/Engine/Private/KismetSystemLibrary.cpp
--- a/Engine/Source/Runtime/Engine/Private/KismetSystemLibrary.cpp
+++ b/Engine/Source/Runtime/Engine/Private/KismetSystemLibrary.cpp
@@ -295,7 +295,9 @@ void UKismetSystemLibrary::K2_SetTimerDelegate(FTimerDynamicDelegate Delegate, f
 	}
 	else
 	{
-		UE_LOG(LogBlueprintUserMessages, Warning, TEXT("SetTimer passed a bad function (%s) or object (%s)"), *Delegate.GetFunctionName().ToString(), *GetNameSafe(Delegate.GetUObject()));
+		UE_LOG(LogBlueprintUserMessages, Warning, 
+			TEXT("SetTimer passed a bad function (%s) or object (%s)"),
+			*Delegate.GetFunctionName().ToString(), *GetNameSafe(Delegate.GetUObject()));
 	}
 }
 
@@ -304,9 +306,14 @@ void UKismetSystemLibrary::K2_ClearTimer(UObject* Object, FString FunctionName)
 	FTimerDynamicDelegate Delegate;
 	Delegate.BindUFunction(Object, *FunctionName);
 
+	K2_ClearTimerDelegate(Delegate);
+}
+
+void UKismetSystemLibrary::K2_ClearTimerDelegate(FTimerDynamicDelegate Delegate)
+{
 	if (Delegate.IsBound())
 	{
-		UWorld* World = GEngine->GetWorldFromContextObject(Object);
+		UWorld* World = GEngine->GetWorldFromContextObject(Delegate.GetUObject());
 		if (World)
 		{
 			World->GetTimerManager().ClearTimer(Delegate);
@@ -314,8 +321,9 @@ void UKismetSystemLibrary::K2_ClearTimer(UObject* Object, FString FunctionName)
 	}
 	else
 	{
-		const FString ObjectNameStr = Object != NULL ? Object->GetName() : TEXT("None");
-		UE_LOG(LogBlueprintUserMessages, Warning, TEXT("ClearTimer passed a bad function (%s) or object (%s)"), *FunctionName, *ObjectNameStr);
+		UE_LOG(LogBlueprintUserMessages, Warning, 
+			TEXT("ClearTimer passed a bad function (%s) or object (%s)"),
+			*Delegate.GetFunctionName().ToString(), *GetNameSafe(Delegate.GetUObject()));
 	}
 }
 
@@ -324,15 +332,21 @@ void UKismetSystemLibrary::K2_PauseTimer(UObject* Object, FString FunctionName)
 	FTimerDynamicDelegate Delegate;
 	Delegate.BindUFunction(Object, *FunctionName);
 
+	K2_PauseTimerDelegate(Delegate);
+}
+
+void UKismetSystemLibrary::K2_PauseTimerDelegate(FTimerDynamicDelegate Delegate)
+{
 	if (Delegate.IsBound())
 	{
-		UWorld* World = GEngine->GetWorldFromContextObject(Object);
+		UWorld* World = GEngine->GetWorldFromContextObject(Delegate.GetUObject());
 		World->GetTimerManager().PauseTimer(Delegate);
 	}
 	else
 	{
-		const FString ObjectNameStr = Object != NULL ? Object->GetName() : TEXT("None");
-		UE_LOG(LogBlueprintUserMessages, Warning, TEXT("PauseTimer passed a bad function (%s) or object (%s)"), *FunctionName, *ObjectNameStr);
+		UE_LOG(LogBlueprintUserMessages, Warning, 
+			TEXT("PauseTimer passed a bad function (%s) or object (%s)"),
+			*Delegate.GetFunctionName().ToString(), *GetNameSafe(Delegate.GetUObject()));
 	}
 }
 
@@ -341,15 +355,21 @@ void UKismetSystemLibrary::K2_UnPauseTimer(UObject* Object, FString FunctionName
 	FTimerDynamicDelegate Delegate;
 	Delegate.BindUFunction(Object, *FunctionName);
 
+	K2_UnPauseTimerDelegate(Delegate);
+}
+
+void UKismetSystemLibrary::K2_UnPauseTimerDelegate(FTimerDynamicDelegate Delegate)
+{
 	if (Delegate.IsBound())
 	{
-		UWorld* World = GEngine->GetWorldFromContextObject(Object);
+		UWorld* World = GEngine->GetWorldFromContextObject(Delegate.GetUObject());
 		World->GetTimerManager().UnPauseTimer(Delegate);
 	}
 	else
 	{
-		const FString ObjectNameStr = Object != NULL ? Object->GetName() : TEXT("None");
-		UE_LOG(LogBlueprintUserMessages, Warning, TEXT("UnPauseTimer passed a bad function (%s) or object (%s)"), *FunctionName, *ObjectNameStr);
+		UE_LOG(LogBlueprintUserMessages, Warning,
+			TEXT("UnPauseTimer passed a bad function (%s) or object (%s)"),
+			*Delegate.GetFunctionName().ToString(), *GetNameSafe(Delegate.GetUObject()));
 	}
 }
 
@@ -358,15 +378,21 @@ bool UKismetSystemLibrary::K2_IsTimerActive(UObject* Object, FString FunctionNam
 	FTimerDynamicDelegate Delegate;
 	Delegate.BindUFunction(Object, *FunctionName);
 
+	return K2_IsTimerActiveDelegate(Delegate);
+}
+
+bool UKismetSystemLibrary::K2_IsTimerActiveDelegate(FTimerDynamicDelegate Delegate)
+{
 	if (Delegate.IsBound())
 	{
-		UWorld* World = GEngine->GetWorldFromContextObject(Object);
+		UWorld* World = GEngine->GetWorldFromContextObject(Delegate.GetUObject());
 		return World->GetTimerManager().IsTimerActive(Delegate);
 	}
 	else
 	{
-		const FString ObjectNameStr = Object != NULL ? Object->GetName() : TEXT("None");
-		UE_LOG(LogBlueprintUserMessages, Warning, TEXT("IsTimerActive passed a bad function (%s) or object (%s)"), *FunctionName, *ObjectNameStr);
+		UE_LOG(LogBlueprintUserMessages, Warning, 
+			TEXT("IsTimerActive passed a bad function (%s) or object (%s)"),
+			*Delegate.GetFunctionName().ToString(), *GetNameSafe(Delegate.GetUObject()));
 		return false;
 	}
 }
@@ -376,15 +402,21 @@ bool UKismetSystemLibrary::K2_IsTimerPaused(UObject* Object, FString FunctionNam
 	FTimerDynamicDelegate Delegate;
 	Delegate.BindUFunction(Object, *FunctionName);
 
+	return K2_IsTimerPausedDelegate(Delegate);
+}
+
+bool UKismetSystemLibrary::K2_IsTimerPausedDelegate(FTimerDynamicDelegate Delegate)
+{
 	if (Delegate.IsBound())
 	{
-		UWorld* World = GEngine->GetWorldFromContextObject(Object);
+		UWorld* World = GEngine->GetWorldFromContextObject(Delegate.GetUObject());
 		return World->GetTimerManager().IsTimerPaused(Delegate);
 	}
 	else
 	{
-		const FString ObjectNameStr = Object != NULL ? Object->GetName() : TEXT("None");
-		UE_LOG(LogBlueprintUserMessages, Warning, TEXT("IsTimerPaused passed a bad function (%s) or object (%s)"), *FunctionName, *ObjectNameStr);
+		UE_LOG(LogBlueprintUserMessages, Warning, 
+			TEXT("IsTimerPaused passed a bad function (%s) or object (%s)"),
+			*Delegate.GetFunctionName().ToString(), *GetNameSafe(Delegate.GetUObject()));
 		return false;
 	}
 }
@@ -394,15 +426,21 @@ bool UKismetSystemLibrary::K2_TimerExists(UObject* Object, FString FunctionName)
 	FTimerDynamicDelegate Delegate;
 	Delegate.BindUFunction(Object, *FunctionName);
 
+	return K2_TimerExistsDelegate(Delegate);
+}
+
+bool UKismetSystemLibrary::K2_TimerExistsDelegate(FTimerDynamicDelegate Delegate)
+{
 	if (Delegate.IsBound())
 	{
-		UWorld* World = GEngine->GetWorldFromContextObject(Object);
+		UWorld* World = GEngine->GetWorldFromContextObject(Delegate.GetUObject());
 		return World->GetTimerManager().TimerExists(Delegate);
 	}
 	else
 	{
-		const FString ObjectNameStr = Object != NULL ? Object->GetName() : TEXT("None");
-		UE_LOG(LogBlueprintUserMessages, Warning, TEXT("TimerExists passed a bad function (%s) or object (%s)"), *FunctionName, *ObjectNameStr);
+		UE_LOG(LogBlueprintUserMessages, Warning,
+			TEXT("TimerExists passed a bad function (%s) or object (%s)"),
+			*Delegate.GetFunctionName().ToString(), *GetNameSafe(Delegate.GetUObject()));
 		return false;
 	}
 }
@@ -412,15 +450,21 @@ float UKismetSystemLibrary::K2_GetTimerElapsedTime(UObject* Object, FString Func
 	FTimerDynamicDelegate Delegate;
 	Delegate.BindUFunction(Object, *FunctionName);
 
+	return K2_GetTimerElapsedTimeDelegate(Delegate);
+}
+
+float UKismetSystemLibrary::K2_GetTimerElapsedTimeDelegate(FTimerDynamicDelegate Delegate)
+{
 	if (Delegate.IsBound())
 	{
-		UWorld* World = GEngine->GetWorldFromContextObject(Object);		
+		UWorld* World = GEngine->GetWorldFromContextObject(Delegate.GetUObject());
 		return World->GetTimerManager().GetTimerElapsed(Delegate);
 	}
 	else
 	{
-		const FString ObjectNameStr = Object != NULL ? Object->GetName() : TEXT("None");
-		UE_LOG(LogBlueprintUserMessages, Warning, TEXT("GetTimerElapsedTime passed a bad function (%s) or object (%s)"), *FunctionName, *ObjectNameStr);
+		UE_LOG(LogBlueprintUserMessages, Warning, 
+			TEXT("GetTimerElapsedTime passed a bad function (%s) or object (%s)"), 
+			*Delegate.GetFunctionName().ToString(), *GetNameSafe(Delegate.GetUObject()));
 		return 0.0f;
 	}
 }
@@ -430,15 +474,21 @@ float UKismetSystemLibrary::K2_GetTimerRemainingTime(UObject* Object, FString Fu
 	FTimerDynamicDelegate Delegate;
 	Delegate.BindUFunction(Object, *FunctionName);
 
+	return K2_GetTimerRemainingTimeDelegate(Delegate);
+}
+
+float UKismetSystemLibrary::K2_GetTimerRemainingTimeDelegate(FTimerDynamicDelegate Delegate)
+{
 	if (Delegate.IsBound())
 	{
-		UWorld* World = GEngine->GetWorldFromContextObject(Object);
+		UWorld* World = GEngine->GetWorldFromContextObject(Delegate.GetUObject());
 		return World->GetTimerManager().GetTimerRemaining(Delegate);
 	}
 	else
 	{
-		const FString ObjectNameStr = Object != NULL ? Object->GetName() : TEXT("None");
-		UE_LOG(LogBlueprintUserMessages, Warning, TEXT("GetTimerRemainingTime passed a bad function (%s) or object (%s)"), *FunctionName, *ObjectNameStr);
+		UE_LOG(LogBlueprintUserMessages, Warning, 
+			TEXT("GetTimerRemainingTime passed a bad function (%s) or object (%s)"), 
+			*Delegate.GetFunctionName().ToString(), *GetNameSafe(Delegate.GetUObject()));
 		return 0.0f;
 	}
 }
