UE4 - Better handling of teh shared cluster tree between the HISMC and the proxy to prevent potential threading problems

diff --git a/Engine/Source/Runtime/Engine/Classes/Components/HierarchicalInstancedStaticMeshComponent.h b/Engine/Source/Runtime/Engine/Classes/Components/HierarchicalInstancedStaticMeshComponent.h
--- a/Engine/Source/Runtime/Engine/Classes/Components/HierarchicalInstancedStaticMeshComponent.h
+++ b/Engine/Source/Runtime/Engine/Classes/Components/HierarchicalInstancedStaticMeshComponent.h
@@ -56,8 +56,7 @@ class ENGINE_API UHierarchicalInstancedStaticMeshComponent : public UInstancedSt
 {
 	GENERATED_UCLASS_BODY()
 
-	UPROPERTY(Transient, DuplicateTransient)
-	TArray<FClusterNode> ClusterTree;
+	TSharedPtr<TArray<FClusterNode>, ESPMode::ThreadSafe> ClusterTreePtr;
 	
 	// Table for remaping instances from cluster tree to PerInstanceSMData order
 	UPROPERTY()
diff --git a/Engine/Source/Runtime/Engine/Private/HierarchicalInstancedStaticMesh.cpp b/Engine/Source/Runtime/Engine/Private/HierarchicalInstancedStaticMesh.cpp
--- a/Engine/Source/Runtime/Engine/Private/HierarchicalInstancedStaticMesh.cpp
+++ b/Engine/Source/Runtime/Engine/Private/HierarchicalInstancedStaticMesh.cpp
@@ -570,6 +570,7 @@ struct FFoliageCullInstanceParams;
 
 class FHierarchicalStaticMeshSceneProxy : public FInstancedStaticMeshSceneProxy
 {
+	TSharedPtr<TArray<FClusterNode>, ESPMode::ThreadSafe> ClusterTreePtr;
 	const TArray<FClusterNode>& ClusterTree;
 
 	int32 FirstUnbuiltIndex;
@@ -587,7 +588,8 @@ public:
 
 	FHierarchicalStaticMeshSceneProxy(UHierarchicalInstancedStaticMeshComponent* InComponent, ERHIFeatureLevel::Type InFeatureLevel)
 		: FInstancedStaticMeshSceneProxy(InComponent, InFeatureLevel)
-		, ClusterTree(InComponent->ClusterTree)
+		, ClusterTreePtr(InComponent->ClusterTreePtr)
+		, ClusterTree(*InComponent->ClusterTreePtr)
 		, FirstUnbuiltIndex(InComponent->NumBuiltInstances)
 		, LastUnbuiltIndex(InComponent->PerInstanceSMData.Num()+InComponent->RemovedInstances.Num()-1)
 #if !(UE_BUILD_SHIPPING || UE_BUILD_TEST)
@@ -1292,6 +1294,8 @@ FBoxSphereBounds UHierarchicalInstancedStaticMeshComponent::CalcBounds(const FTr
 		// @todo: Level error or something to let LDs know this
 		1;//StaticMesh->LODModels(0).Elements.Num() == 1;
 
+	const TArray<FClusterNode>& ClusterTree = *ClusterTreePtr;
+
 	if ((ClusterTree.Num() || UnbuiltInstanceBounds.IsValid) && bMeshIsValid)
 	{
 		FBoxSphereBounds Result;
@@ -1319,6 +1323,7 @@ FBoxSphereBounds UHierarchicalInstancedStaticMeshComponent::CalcBounds(const FTr
 
 UHierarchicalInstancedStaticMeshComponent::UHierarchicalInstancedStaticMeshComponent(const FObjectInitializer& ObjectInitializer)
 	: Super(ObjectInitializer)
+	, ClusterTreePtr(MakeShareable(new TArray<FClusterNode>))
 	, NumBuiltInstances(0)
 	, UnbuiltInstanceBounds(0)
 	, bIsAsyncBuilding(false)
@@ -1330,7 +1335,11 @@ UHierarchicalInstancedStaticMeshComponent::UHierarchicalInstancedStaticMeshCompo
 void UHierarchicalInstancedStaticMeshComponent::Serialize(FArchive& Ar)
 {
 	Super::Serialize(Ar);
-
+	if (Ar.IsLoading())
+	{
+		ClusterTreePtr = MakeShareable(new TArray<FClusterNode>);
+	}
+	TArray<FClusterNode>& ClusterTree = *ClusterTreePtr;
 	ClusterTree.BulkSerialize(Ar);
 }
 
@@ -1451,6 +1460,7 @@ void UHierarchicalInstancedStaticMeshComponent::ClearInstances()
 		bConcurrentRemoval = true;
 	}
 
+	ClusterTreePtr = MakeShareable(new TArray<FClusterNode>);
 	NumBuiltInstances = 0;
 	UnbuiltInstanceBounds.Init();
 
@@ -1504,7 +1514,8 @@ void UHierarchicalInstancedStaticMeshComponent::BuildTree()
 		UnbuiltInstanceBounds.Init();
 		RemovedInstances.Empty();
 
-		Exchange(ClusterTree, Builder->Result->Nodes);
+		ClusterTreePtr = MakeShareable(new TArray<FClusterNode>);
+		Exchange(*ClusterTreePtr, Builder->Result->Nodes);
 		Exchange(InstanceReorderTable, Builder->Result->InstanceReorderTable);
 		Exchange(SortedInstances, Builder->Result->SortedInstances);
 		
@@ -1548,7 +1559,7 @@ void UHierarchicalInstancedStaticMeshComponent::AcceptPrebuiltTree(
 	NumBuiltInstances = PerInstanceSMData.Num();
 	UnbuiltInstanceBounds.Init();
 	RemovedInstances.Empty();
-	ClusterTree.Empty();
+	ClusterTreePtr = MakeShareable(new TArray<FClusterNode>);
 	InstanceReorderTable.Empty();
 	SortedInstances.Empty();
 
@@ -1574,7 +1585,7 @@ void UHierarchicalInstancedStaticMeshComponent::AcceptPrebuiltTree(
 			InstancingRandomSeed = FMath::Rand();
 		}
 
-		Exchange(ClusterTree, InClusterTree);
+		Exchange(*ClusterTreePtr, InClusterTree);
 		Exchange(InstanceReorderTable, InInstanceReorderTable);
 		Exchange(SortedInstances, InSortedInstances);
 
@@ -1614,6 +1625,9 @@ void UHierarchicalInstancedStaticMeshComponent::BuildFlatTree(const TArray<int32
 
 		const FBox InstBox = StaticMesh->GetBounds().GetBox();
 
+		ClusterTreePtr = MakeShareable(new TArray<FClusterNode>);
+		TArray<FClusterNode>& ClusterTree = *ClusterTreePtr;
+
 		ClusterTree.Empty(LeafInstanceCounts.Num() + 1);
 		FClusterNode Root;
 		Root.FirstChild = 1;
@@ -1702,6 +1716,8 @@ void UHierarchicalInstancedStaticMeshComponent::ApplyBuildTreeAsync(ENamedThread
 		    }
 	    }
     
+		ClusterTreePtr = MakeShareable(new TArray<FClusterNode>);
+		TArray<FClusterNode>& ClusterTree = *ClusterTreePtr;
 	    Exchange(ClusterTree, Builder->Result->Nodes);
 		Exchange(InstanceReorderTable, Builder->Result->InstanceReorderTable);
 		Exchange(SortedInstances, Builder->Result->SortedInstances);
@@ -1783,6 +1799,7 @@ FPrimitiveSceneProxy* UHierarchicalInstancedStaticMeshComponent::CreateSceneProx
 	const bool bMeshIsValid = 
 		// make sure we have instances
 		PerInstanceSMData.Num() > 0 &&
+		ClusterTreePtr->Num() &&
 		// make sure we have an actual staticmesh
 		StaticMesh &&
 		StaticMesh->HasValidRenderData() &&
@@ -1807,7 +1824,8 @@ FPrimitiveSceneProxy* UHierarchicalInstancedStaticMeshComponent::CreateSceneProx
 
 static void GatherInstanceTransformsInArea(const UHierarchicalInstancedStaticMeshComponent& Component, const FBox& AreaBox, int32 Child, TArray<FTransform>& InstanceData)
 {
-	const FClusterNode& ChildNode = Component.ClusterTree[Child];
+	const TArray<FClusterNode>& ClusterTree = *Component.ClusterTreePtr;
+	const FClusterNode& ChildNode = ClusterTree[Child];
 	const FVector WorldMinMax[] = {
 		Component.ComponentToWorld.TransformPosition(ChildNode.BoundMin),
 		Component.ComponentToWorld.TransformPosition(ChildNode.BoundMax)
@@ -1847,6 +1865,7 @@ void UHierarchicalInstancedStaticMeshComponent::GetNavigationPerInstanceTransfor
 {
 	if (IsTreeFullyBuilt())
 	{
+		const TArray<FClusterNode>& ClusterTree = *ClusterTreePtr;
 		if (ClusterTree.Num())
 		{
 			GatherInstanceTransformsInArea(*this, AreaBox, 0, InstanceData);
@@ -1886,6 +1905,7 @@ void UHierarchicalInstancedStaticMeshComponent::FlushAccumulatedNavigationUpdate
 {
 	if (AccumulatedNavigationDirtyArea.IsValid)
 	{
+		const TArray<FClusterNode>& ClusterTree = *ClusterTreePtr;
 		UNavigationSystem* NavSys = UNavigationSystem::GetCurrent(GetWorld());
 		// Check if this component is registered in navigation system
 		if (ClusterTree.Num() && NavSys && NavSys->GetObjectsNavOctreeId(this))
diff --git a/Engine/Source/Runtime/Landscape/Private/Landscape.cpp b/Engine/Source/Runtime/Landscape/Private/Landscape.cpp
--- a/Engine/Source/Runtime/Landscape/Private/Landscape.cpp
+++ b/Engine/Source/Runtime/Landscape/Private/Landscape.cpp
@@ -2730,7 +2730,7 @@ void ALandscapeProxy::Tick(float DeltaSeconds)
 
 				FBoxSphereBounds WorldBounds = HierarchicalInstancedStaticMeshComponent->CalcBounds(HierarchicalInstancedStaticMeshComponent->ComponentToWorld);
 				UE_LOG(LogTemp, Display, TEXT("Comp %d instances at %.0f %.0f %.0f"), HierarchicalInstancedStaticMeshComponent->PerInstanceSMData.Num(), WorldBounds.Origin.X, WorldBounds.Origin.Y, WorldBounds.Origin.Z);
-				TreeSize += HierarchicalInstancedStaticMeshComponent->ClusterTree.GetAllocatedSize();
+				TreeSize += HierarchicalInstancedStaticMeshComponent->ClusterTreePtr->GetAllocatedSize();
 				InstSize += HierarchicalInstancedStaticMeshComponent->PerInstanceSMData.GetAllocatedSize();
 				InstBufSize += HierarchicalInstancedStaticMeshComponent->PerInstanceSMData.Num() * 7 * 16; // there is no place where this constant exists.
 			}
