Changed navigation octree creation policy so that it gets created only when navigation generation is supposed to work #UE4

diff --git a/Engine/Config/BaseEngine.ini b/Engine/Config/BaseEngine.ini
--- a/Engine/Config/BaseEngine.ini
+++ b/Engine/Config/BaseEngine.ini
@@ -1933,7 +1933,6 @@ EditorViewButtonsObject=EditorCylinder
 
 [/Script/Engine.NavigationSystem]
 bAutoCreateNavigationData=true
-bBuildNavigationAtRuntime=true
 bAddPlayersToGenerationSeeds=true
 +RequiredNavigationDataClassNames=/Script/Engine.RecastNavMesh
 
diff --git a/Engine/Source/Runtime/Engine/Classes/AI/Navigation/NavigationSystem.h b/Engine/Source/Runtime/Engine/Classes/AI/Navigation/NavigationSystem.h
--- a/Engine/Source/Runtime/Engine/Classes/AI/Navigation/NavigationSystem.h
+++ b/Engine/Source/Runtime/Engine/Classes/AI/Navigation/NavigationSystem.h
@@ -138,15 +138,18 @@ class ENGINE_API UNavigationSystem : public UBlueprintFunctionLibrary
 	UPROPERTY()
 	ANavigationData* AbstractNavData;
 
+protected:
 	/** Should navigation system spawn default Navigation Data when there's none and there are navigation bounds present? */
 	UPROPERTY(config, EditAnywhere, Category=NavigationSystem)
 	uint32 bAutoCreateNavigationData:1;
 
-	/** Should navigation system rebuild navigation data at game/pie runtime? 
-	 *	It also includes spawning navigation data instances if there are none present. */
-	UPROPERTY(config, EditAnywhere, Category=NavigationSystem)
-	uint32 bBuildNavigationAtRuntime:1;
+	/** gets set to true if gathering navigation data (like in navoctree) is required due to the need of navigation generation 
+	 *	Is always true in Editor Mode. In other modes it depends on bRebuildAtRuntime of every required NavigationData class' CDO
+	 */
+	UPROPERTY()
+	uint32 bSupportRebuilding : 1; 
 
+public:
 	/** if set to true will result navigation system not rebuild navigation until 
 	 *	a call to ReleaseInitialBuildingLock() is called. Does not influence 
 	 *	editor-time generation (i.e. does influence PIE and Game).
@@ -437,6 +440,8 @@ public:
 	/** checks if dirty navigation data can rebuild itself */
 	bool CanRebuildDirtyNavigation() const;
 
+	FORCEINLINE bool SupportsNavigationGeneration() const { return bSupportRebuilding; }
+
 	static bool DoesPathIntersectBox(const FNavigationPath* Path, const FBox& Box, uint32 StartingIndex = 0);
 	static bool DoesPathIntersectBox(const FNavigationPath* Path, const FBox& Box, const FVector& AgentLocation, uint32 StartingIndex = 0);
 
@@ -452,6 +457,16 @@ public:
 	void RequestRegistration(ANavigationData* NavData, bool bTriggerRegistrationProcessing = true);
 
 protected:
+
+	/** Processes all NavigationData instances in UWorld owning navigation system instance, and registers
+	 *	all previously unregistered */
+	void RegisterNavigationDataInstances();
+
+	/** called in places where we need to spawn the NavOctree, but is checking additional conditions if we really want to do that
+	 *	depending on project setup among others 
+	 *	@return true if NavOctree instance has been created, or if one is already present */
+	virtual bool ConditionallyCreateNavOctree();
+
 	/** Processes registration of candidates queues via RequestRegistration and stored in NavDataRegistrationQueue */
 	void ProcessRegistrationCandidates();
 
diff --git a/Engine/Source/Runtime/Engine/Private/AI/Navigation/NavigationSystem.cpp b/Engine/Source/Runtime/Engine/Private/AI/Navigation/NavigationSystem.cpp
--- a/Engine/Source/Runtime/Engine/Private/AI/Navigation/NavigationSystem.cpp
+++ b/Engine/Source/Runtime/Engine/Private/AI/Navigation/NavigationSystem.cpp
@@ -227,10 +227,6 @@ UNavigationSystem::UNavigationSystem(const FObjectInitializer& ObjectInitializer
 		AsyncPathFindingQueries.Reserve( INITIAL_ASYNC_QUERIES_SIZE );
 		NavDataRegistrationQueue.Reserve( REGISTRATION_QUEUE_SIZE );
 	
-		NavOctree = new FNavigationOctree(FVector(0,0,0), 64000);
-#if WITH_RECAST
-		NavOctree->ComponentExportDelegate = FNavigationOctree::FNavigableGeometryComponentExportDelegate::CreateStatic(&FRecastNavMeshGenerator::ExportComponentGeometry);
-#endif // WITH_RECAST
 		FWorldDelegates::LevelAddedToWorld.AddUObject(this, &UNavigationSystem::OnLevelAddedToWorld);
 		FWorldDelegates::LevelRemovedFromWorld.AddUObject(this, &UNavigationSystem::OnLevelRemovedFromWorld);
 	}
@@ -265,30 +261,7 @@ void UNavigationSystem::DoInitialSetup()
 	{
 		return;
 	}
-
-	// gather navigation creators
-	NavDataClasses.Empty(RequiredNavigationDataClassNames.Num());
-	for (int32 Index = 0; Index < RequiredNavigationDataClassNames.Num(); ++Index)
-	{
-		TSubclassOf<ANavigationData> NavDataClass = LoadClass<ANavigationData>(NULL, *RequiredNavigationDataClassNames[Index].ToString(), NULL, LOAD_None, NULL);
-		if (NavDataClass)
-		{
-			NavDataClasses.AddUnique(NavDataClass);
-		}
-		else
-		{
-			UE_LOG(LogNavigation, Warning, TEXT("Unable to find navigation data class \'%s\' while setting up require navigation types")
-				, *RequiredNavigationDataClassNames[Index].ToString());
-		}
-	}
-
-	if (NavDataClasses.Num() == 0)
-	{
-		// @note: if you don't want navigation system to be created at all you can disable it by 
-		// setting AWorldSettings.bEnableNavigationSystem to false
-		UE_LOG(LogNavigation, Error, TEXT("No navigation data types found while setting up required navigation types!"));
-	}
-
+	
 	// spawn abstract nav data separately
 	// it's responsible for direct paths and shouldn't be picked for any agent type as default one
 	UWorld* NavWorld = GetWorld();
@@ -326,10 +299,30 @@ void UNavigationSystem::PostInitProperties()
 			SupportedAgents.Add(FNavDataConfig(FNavigationSystem::FallbackAgentRadius, FNavigationSystem::FallbackAgentHeight));
 		}
 
-		const bool bShouldStoreNavigableGeometry = bBuildNavigationAtRuntime || !GetWorld()->IsGameWorld();
-		NavOctree->SetNavigableGeometryStoringMode(bShouldStoreNavigableGeometry 
-			? FNavigationOctree::StoreNavGeometry
-			: FNavigationOctree::SkipNavGeometry);
+		// gather navigation creators
+		NavDataClasses.Empty(RequiredNavigationDataClassNames.Num());
+		for (int32 Index = 0; Index < RequiredNavigationDataClassNames.Num(); ++Index)
+		{
+			TSubclassOf<ANavigationData> NavDataClass = LoadClass<ANavigationData>(NULL, *RequiredNavigationDataClassNames[Index].ToString(), NULL, LOAD_None, NULL);
+			if (NavDataClass)
+			{
+				NavDataClasses.AddUnique(NavDataClass);
+			}
+			else
+			{
+				UE_LOG(LogNavigation, Warning, TEXT("Unable to find navigation data class \'%s\' while setting up require navigation types")
+					, *RequiredNavigationDataClassNames[Index].ToString());
+			}
+		}
+
+		if (NavDataClasses.Num() == 0)
+		{
+			// @note: if you don't want navigation system to be created at all you can disable it by 
+			// setting AWorldSettings.bEnableNavigationSystem to false
+			UE_LOG(LogNavigation, Error, TEXT("No navigation data types found while setting up required navigation types!"));
+		}
+
+		ConditionallyCreateNavOctree();
 	
 		bInitialBuildingLockActive = bInitialBuildingLocked;
 		InitializeLevelCollisions();
@@ -355,6 +348,35 @@ void UNavigationSystem::PostInitProperties()
 	}
 }
 
+bool UNavigationSystem::ConditionallyCreateNavOctree()
+{
+	ensure(NavOctree == nullptr);
+	if (NavOctree != nullptr)
+	{
+		return true;
+	}
+
+	bSupportRebuilding = GetWorld()->IsGameWorld() == false;
+	for (int32 NavDataClassIndex = 0; NavDataClassIndex < NavDataClasses.Num() && bSupportRebuilding == false; ++NavDataClassIndex)
+	{
+		const ANavigationData* NavDataCDO = GetDefault<ANavigationData>(NavDataClasses[NavDataClassIndex]);
+		check(NavDataCDO);
+		bSupportRebuilding = NavDataCDO->bRebuildAtRuntime;
+	}
+
+	if (bSupportRebuilding)
+	{
+		NavOctree = new FNavigationOctree(FVector(0, 0, 0), 64000);
+#if WITH_RECAST
+		NavOctree->ComponentExportDelegate = FNavigationOctree::FNavigableGeometryComponentExportDelegate::CreateStatic(&FRecastNavMeshGenerator::ExportComponentGeometry);
+#endif // WITH_RECAST
+		NavOctree->SetNavigableGeometryStoringMode(FNavigationOctree::StoreNavGeometry);
+	}
+
+	return NavOctree != nullptr;
+}
+
+
 #if WITH_EDITOR
 void UNavigationSystem::PostEditChangeProperty(FPropertyChangedEvent& PropertyChangedEvent)
 {
@@ -394,121 +416,135 @@ void UNavigationSystem::OnInitializeActors()
 void UNavigationSystem::OnWorldInitDone(FNavigationSystem::EMode Mode)
 {
 	OperationMode = Mode;
-
-	UWorld* World = GetWorld();
-
-	if (IsThereAnywhereToBuildNavigation() == false
-		// Simulation mode is a special case - better not do it in this case
-		&& OperationMode != FNavigationSystem::SimulationMode)
+	
+	if (bSupportRebuilding)
 	{
-		// remove all navigation data instances
-		for (TActorIterator<ANavigationData> It(World); It; ++It)
+		UWorld* World = GetWorld();
+
+		if (IsThereAnywhereToBuildNavigation() == false
+			// Simulation mode is a special case - better not do it in this case
+			&& OperationMode != FNavigationSystem::SimulationMode)
 		{
-			ANavigationData* Nav = (*It);
-			if (Nav != NULL && Nav->IsPendingKill() == false)
+			// remove all navigation data instances
+			for (TActorIterator<ANavigationData> It(World); It; ++It)
 			{
-				UnregisterNavData(Nav);
-				Nav->CleanUpAndMarkPendingKill();
+				ANavigationData* Nav = (*It);
+				if (Nav != NULL && Nav->IsPendingKill() == false)
+				{
+					UnregisterNavData(Nav);
+					Nav->CleanUpAndMarkPendingKill();
+				}
 			}
-		}
-
-		if (OperationMode == FNavigationSystem::EditorMode)
-		{
-			bInitialBuildingLockActive = false;
-		}
 
-		bNavDataRemovedDueToMissingNavBounds = true;
-	}
-	else
-	{
-		InitializeLevelCollisions();
-		PopulateNavOctree();
-		
-		// gather navigable bounds
-		GatherNavigationBounds();
-		
-		// gather all navigation data instances and register all not-yet-registered
-		// (since it's quite possible navigation system was not ready by the time
-		// those instances were serialized-in or spawned)
-		bool bProcessRegistration = false;
-		for (TActorIterator<ANavigationData> It(World); It; ++It)
-		{
-			ANavigationData* Nav = (*It);
-			if (Nav != NULL && Nav->IsPendingKill() == false && Nav->IsRegistered() == false)
+			if (OperationMode == FNavigationSystem::EditorMode)
 			{
-				RequestRegistration(Nav, false);
-				bProcessRegistration = true;
+				bInitialBuildingLockActive = false;
 			}
-		}
-		if (bProcessRegistration == true)
-		{
-			ProcessRegistrationCandidates();
-		}
 
-		if (OperationMode == FNavigationSystem::EditorMode)
-		{
-			// don't lock navigation building in editor
-			bInitialBuildingLockActive = false;
-		}
-
-		if (bAutoCreateNavigationData == true)
-		{
-			SpawnMissingNavigationData();
-			// in case anything spawned has registered
-			ProcessRegistrationCandidates();
+			bNavDataRemovedDueToMissingNavBounds = true;
 		}
 		else
 		{
-			if (GetMainNavData(FNavigationSystem::DontCreate) != NULL)
+			InitializeLevelCollisions();
+			PopulateNavOctree();
+
+			// gather navigable bounds
+			GatherNavigationBounds();
+
+			// gather all navigation data instances and register all not-yet-registered
+			// (since it's quite possible navigation system was not ready by the time
+			// those instances were serialized-in or spawned)
+			RegisterNavigationDataInstances();
+
+			if (OperationMode == FNavigationSystem::EditorMode)
+			{
+				// don't lock navigation building in editor
+				bInitialBuildingLockActive = false;
+			}
+
+			if (bAutoCreateNavigationData == true)
+			{
+				SpawnMissingNavigationData();
+				// in case anything spawned has registered
+				ProcessRegistrationCandidates();
+			}
+			else
 			{
-				// trigger navmesh update
-				for (TActorIterator<ANavigationData> It(World); It; ++It)
+				if (GetMainNavData(FNavigationSystem::DontCreate) != NULL)
 				{
-					ANavigationData* NavData = (*It);
-					if (NavData != NULL)
+					// trigger navmesh update
+					for (TActorIterator<ANavigationData> It(World); It; ++It)
 					{
-						ERegistrationResult Result = RegisterNavData(NavData);
-
-						if (Result == RegistrationSuccessful)
+						ANavigationData* NavData = (*It);
+						if (NavData != NULL)
 						{
-#if WITH_RECAST
-							if (Cast<ARecastNavMesh>(NavData) != NULL)
+							ERegistrationResult Result = RegisterNavData(NavData);
+
+							if (Result == RegistrationSuccessful)
 							{
-								if (bInitialBuildingLockActive == false && bNavigationAutoUpdateEnabled)
+#if WITH_RECAST
+								if (Cast<ARecastNavMesh>(NavData) != NULL)
 								{
-									NavData->RebuildAll();
+									if (bInitialBuildingLockActive == false && bNavigationAutoUpdateEnabled)
+									{
+										NavData->RebuildAll();
+									}
 								}
-							}
 #endif // WITH_RECAST
-						}
-						else if (Result != RegistrationFailed_DataPendingKill
-							&& Result != RegistrationFailed_AgentNotValid
-							)
-						{
-							NavData->CleanUpAndMarkPendingKill();					
+							}
+							else if (Result != RegistrationFailed_DataPendingKill
+								&& Result != RegistrationFailed_AgentNotValid
+								)
+							{
+								NavData->CleanUpAndMarkPendingKill();
+							}
 						}
 					}
 				}
 			}
-		}
 
-		// All navigation actors are registered
-		// Add NavMesh parts from all sub-levels that were streamed in prior NavMesh registration
-		if (World->IsGameWorld())
-		{
-			const auto& Levels = World->GetLevels();
-			for (ULevel* Level : Levels)
+			// All navigation actors are registered
+			// Add NavMesh parts from all sub-levels that were streamed in prior NavMesh registration
+			if (World->IsGameWorld())
 			{
-				if (!Level->IsPersistentLevel() && Level->bIsVisible)
+				const auto& Levels = World->GetLevels();
+				for (ULevel* Level : Levels)
 				{
-					for (ANavigationData* NavData : NavDataSet)
+					if (!Level->IsPersistentLevel() && Level->bIsVisible)
 					{
-						NavData->OnStreamingLevelAdded(Level);
+						for (ANavigationData* NavData : NavDataSet)
+						{
+							NavData->OnStreamingLevelAdded(Level);
+						}
 					}
 				}
 			}
 		}
 	}
+	else // just register data already present
+	{
+		RegisterNavigationDataInstances();
+	}
+}
+
+void UNavigationSystem::RegisterNavigationDataInstances()
+{
+	UWorld* World = GetWorld();
+
+	bool bProcessRegistration = false;
+	for (TActorIterator<ANavigationData> It(World); It; ++It)
+	{
+		ANavigationData* Nav = (*It);
+		if (Nav != NULL && Nav->IsPendingKill() == false && Nav->IsRegistered() == false)
+		{
+			RequestRegistration(Nav, false);
+			bProcessRegistration = true;
+		}
+	}
+	if (bProcessRegistration == true)
+	{
+		ProcessRegistrationCandidates();
+	}
 }
 
 void UNavigationSystem::CreateCrowdManager()
@@ -2902,7 +2938,7 @@ void UNavigationSystem::NavigationBuildingLock()
 		return;
 	}
 
-	GetMainNavData(bBuildNavigationAtRuntime && IsThereAnywhereToBuildNavigation() ? FNavigationSystem::Create : FNavigationSystem::DontCreate);
+	GetMainNavData(bAutoCreateNavigationData && NavOctree != nullptr && IsThereAnywhereToBuildNavigation() ? FNavigationSystem::Create : FNavigationSystem::DontCreate);
 
 	bNavigationBuildingLocked = true;
 }
diff --git a/Engine/Source/Runtime/Engine/Private/UnrealEngine.cpp b/Engine/Source/Runtime/Engine/Private/UnrealEngine.cpp
--- a/Engine/Source/Runtime/Engine/Private/UnrealEngine.cpp
+++ b/Engine/Source/Runtime/Engine/Private/UnrealEngine.cpp
@@ -6789,7 +6789,7 @@ void DrawStatsHUD( UWorld* World, FViewport* Viewport, FCanvas* Canvas, UCanvas*
 #endif
 			UNavigationSystem* NavSys = UNavigationSystem::GetCurrent(World);
 			if (NavSys && NavSys->IsNavigationDirty() &&
-				(!bIsNavigationAutoUpdateEnabled || !NavSys->bBuildNavigationAtRuntime || !NavSys->CanRebuildDirtyNavigation()))
+				(!bIsNavigationAutoUpdateEnabled || !NavSys->SupportsNavigationGeneration() || !NavSys->CanRebuildDirtyNavigation()))
 			{
 				SmallTextItem.SetColor( FLinearColor::White );
 				SmallTextItem.Text =  LOCTEXT("NAVMESHERROR", "NAVMESH NEEDS TO BE REBUILT");				
