- Split up an array of 3D textures that Metal can't support

diff --git a/Engine/Shaders/ParticleSimulationShader.usf b/Engine/Shaders/ParticleSimulationShader.usf
--- a/Engine/Shaders/ParticleSimulationShader.usf
+++ b/Engine/Shaders/ParticleSimulationShader.usf
@@ -46,7 +46,7 @@ void VertexMain(
 {
 	uint InstanceId = Input.InstanceId * TILES_PER_INSTANCE + Input.VertexId / 4;
 	
-	float2 TileCoord = Input.TexCoord.xy * float2(TILE_SIZE_X,TILE_SIZE_Y) + TileOffsets[InstanceId];
+float2 TileCoord = Input.TexCoord.xy * float2(TILE_SIZE_X,TILE_SIZE_Y) + TileOffsets[InstanceId];
 
 	OutPosition = float4(
 		TileCoord.xy * float2(2.0f,-2.0f) + float2(-1.0f,1.0f),
@@ -98,8 +98,14 @@ SamplerState CurveTextureSampler;
 #if MAX_VECTOR_FIELDS != 4
 #error This must match MAX_VECTOR_FIELDS in C++ land
 #endif
-Texture3D VectorFieldTextures[MAX_VECTOR_FIELDS];
-SamplerState VectorFieldTexturesSampler;
+Texture3D VectorFieldTextures0;
+Texture3D VectorFieldTextures1;
+Texture3D VectorFieldTextures2;
+Texture3D VectorFieldTextures3;
+SamplerState VectorFieldTexturesSampler0;
+SamplerState VectorFieldTexturesSampler1;
+SamplerState VectorFieldTexturesSampler2;
+SamplerState VectorFieldTexturesSampler3;
 
 /**
  * Computes the orbit velocity to apply to the particle based on time.
@@ -131,10 +137,10 @@ float3 ComputeOrbitVelocity(float Time, float RandomOrbit)
  */
  float3 SampleVectorFieldTexture(int Index, float3 UV)
  {
-	if (Index == 0) return Texture3DSample(VectorFieldTextures[0], VectorFieldTexturesSampler, UV).xyz;
-	if (Index == 1) return Texture3DSample(VectorFieldTextures[1], VectorFieldTexturesSampler, UV).xyz;
-	if (Index == 2) return Texture3DSample(VectorFieldTextures[2], VectorFieldTexturesSampler, UV).xyz;
-	return Texture3DSample(VectorFieldTextures[3], VectorFieldTexturesSampler, UV).xyz;
+	if (Index == 0) return Texture3DSample(VectorFieldTextures0, VectorFieldTexturesSampler0, UV).xyz;
+	if (Index == 1) return Texture3DSample(VectorFieldTextures1, VectorFieldTexturesSampler1, UV).xyz;
+	if (Index == 2) return Texture3DSample(VectorFieldTextures2, VectorFieldTexturesSampler2, UV).xyz;
+	return Texture3DSample(VectorFieldTextures3, VectorFieldTexturesSampler3, UV).xyz;
  }
 
 /**
@@ -168,7 +174,7 @@ void EvaluateVectorFields(out float3 OutForce, out float4 OutVelocity, float3 Po
 
 		// @todo compat hack: Some compilers only allow constant indexing into a texture array
 //		float3 VectorSample = Texture3DSample(VectorFieldTextures[VectorFieldIndex], VectorFieldTexturesSampler, saturate(VolumeUV)).xyz;
-		float3 VectorSample = SampleVectorFieldTexture(VectorFieldIndex, saturate(VolumeUV));
+float3 VectorSample = SampleVectorFieldTexture(VectorFieldIndex, saturate(VolumeUV));
 
 		float3 Vec = mul(float4(VectorSample,0), VectorFields.VolumeToWorld[VectorFieldIndex]).xyz;
 		TotalForce += (Vec * DistanceWeight * Intensity);
diff --git a/Engine/Source/Runtime/Engine/Private/Particles/ParticleGpuSimulation.cpp b/Engine/Source/Runtime/Engine/Private/Particles/ParticleGpuSimulation.cpp
--- a/Engine/Source/Runtime/Engine/Private/Particles/ParticleGpuSimulation.cpp
+++ b/Engine/Source/Runtime/Engine/Private/Particles/ParticleGpuSimulation.cpp
@@ -874,8 +874,14 @@ public:
 		RenderAttributesTextureSampler.Bind(Initializer.ParameterMap, TEXT("RenderAttributesTextureSampler"));
 		CurveTexture.Bind(Initializer.ParameterMap, TEXT("CurveTexture"));
 		CurveTextureSampler.Bind(Initializer.ParameterMap, TEXT("CurveTextureSampler"));
-		VectorFieldTextures.Bind(Initializer.ParameterMap, TEXT("VectorFieldTextures"));
-		VectorFieldTexturesSampler.Bind(Initializer.ParameterMap, TEXT("VectorFieldTexturesSampler"));
+		VectorFieldTextures0.Bind(Initializer.ParameterMap, TEXT("VectorFieldTextures0"));
+		VectorFieldTextures1.Bind(Initializer.ParameterMap, TEXT("VectorFieldTextures1"));
+		VectorFieldTextures2.Bind(Initializer.ParameterMap, TEXT("VectorFieldTextures2"));
+		VectorFieldTextures3.Bind(Initializer.ParameterMap, TEXT("VectorFieldTextures3"));
+		VectorFieldTexturesSampler0.Bind(Initializer.ParameterMap, TEXT("VectorFieldTexturesSampler0"));
+		VectorFieldTexturesSampler1.Bind(Initializer.ParameterMap, TEXT("VectorFieldTexturesSampler1"));
+		VectorFieldTexturesSampler2.Bind(Initializer.ParameterMap, TEXT("VectorFieldTexturesSampler2"));
+		VectorFieldTexturesSampler3.Bind(Initializer.ParameterMap, TEXT("VectorFieldTexturesSampler3"));
 		SceneDepthTextureParameter.Bind(Initializer.ParameterMap,TEXT("SceneDepthTexture"));
 		SceneDepthTextureParameterSampler.Bind(Initializer.ParameterMap,TEXT("SceneDepthTextureSampler"));
 		GBufferATextureParameter.Bind(Initializer.ParameterMap,TEXT("GBufferATexture"));
@@ -898,8 +904,14 @@ public:
 		Ar << RenderAttributesTextureSampler;
 		Ar << CurveTexture;
 		Ar << CurveTextureSampler;
-		Ar << VectorFieldTextures;
-		Ar << VectorFieldTexturesSampler;
+		Ar << VectorFieldTextures0;
+		Ar << VectorFieldTextures1;
+		Ar << VectorFieldTextures2;
+		Ar << VectorFieldTextures3;
+		Ar << VectorFieldTexturesSampler0;
+		Ar << VectorFieldTexturesSampler1;
+		Ar << VectorFieldTexturesSampler2;
+		Ar << VectorFieldTexturesSampler3;
 		Ar << SceneDepthTextureParameter;
 		Ar << SceneDepthTextureParameterSampler;
 		Ar << GBufferATextureParameter;
@@ -970,10 +982,10 @@ public:
 		FPixelShaderRHIParamRef PixelShaderRHI = GetPixelShader();
 		FSamplerStateRHIParamRef SamplerStateLinear = TStaticSamplerState<SF_Bilinear,AM_Clamp,AM_Clamp,AM_Clamp>::GetRHI();
 		SetUniformBufferParameter(RHICmdList, PixelShaderRHI, GetUniformBufferParameter<FVectorFieldUniformParameters>(), UniformBuffer);
-		for (int32 TextureIndex = 0; TextureIndex < MAX_VECTOR_FIELDS; ++TextureIndex)
-		{
-			SetTextureParameter(RHICmdList, PixelShaderRHI, VectorFieldTextures, VectorFieldTexturesSampler, SamplerStateLinear, VolumeTexturesRHI[TextureIndex], TextureIndex);
-		}
+		SetTextureParameter(RHICmdList, PixelShaderRHI, VectorFieldTextures0, VectorFieldTexturesSampler0, SamplerStateLinear, VolumeTexturesRHI[0], 0);
+		SetTextureParameter(RHICmdList, PixelShaderRHI, VectorFieldTextures1, VectorFieldTexturesSampler1, SamplerStateLinear, VolumeTexturesRHI[1], 0);
+		SetTextureParameter(RHICmdList, PixelShaderRHI, VectorFieldTextures2, VectorFieldTexturesSampler2, SamplerStateLinear, VolumeTexturesRHI[2], 0);
+		SetTextureParameter(RHICmdList, PixelShaderRHI, VectorFieldTextures3, VectorFieldTexturesSampler3, SamplerStateLinear, VolumeTexturesRHI[3], 0);
 	}
 
 	/**
@@ -993,13 +1005,21 @@ public:
 	{
 		FPixelShaderRHIParamRef PixelShaderRHI = GetPixelShader();
 		FShaderResourceViewRHIParamRef NullSRV = FShaderResourceViewRHIParamRef();
-		if (VectorFieldTextures.IsBound())
+		if (VectorFieldTextures0.IsBound())
 		{
-			const uint32 ResourceCount = VectorFieldTextures.GetNumResources();
-			for (uint32 ResourceIndex = 0; ResourceIndex < ResourceCount; ++ResourceIndex)
-			{
-				RHICmdList.SetShaderResourceViewParameter(PixelShaderRHI, VectorFieldTextures.GetBaseIndex() + ResourceIndex, NullSRV);
-			}
+			RHICmdList.SetShaderResourceViewParameter(PixelShaderRHI, VectorFieldTextures0.GetBaseIndex(), NullSRV);
+		}
+		if (VectorFieldTextures1.IsBound())
+		{
+			RHICmdList.SetShaderResourceViewParameter(PixelShaderRHI, VectorFieldTextures1.GetBaseIndex(), NullSRV);
+		}
+		if (VectorFieldTextures2.IsBound())
+		{
+			RHICmdList.SetShaderResourceViewParameter(PixelShaderRHI, VectorFieldTextures2.GetBaseIndex(), NullSRV);
+		}
+		if (VectorFieldTextures3.IsBound())
+		{
+			RHICmdList.SetShaderResourceViewParameter(PixelShaderRHI, VectorFieldTextures3.GetBaseIndex(), NullSRV);
 		}
 	}
 
@@ -1021,8 +1041,14 @@ private:
 	FShaderResourceParameter CurveTexture;
 	FShaderResourceParameter CurveTextureSampler;
 	/** Vector fields. */
-	FShaderResourceParameter VectorFieldTextures;
-	FShaderResourceParameter VectorFieldTexturesSampler;
+	FShaderResourceParameter VectorFieldTextures0;
+	FShaderResourceParameter VectorFieldTextures1;
+	FShaderResourceParameter VectorFieldTextures2;
+	FShaderResourceParameter VectorFieldTextures3;
+	FShaderResourceParameter VectorFieldTexturesSampler0;
+	FShaderResourceParameter VectorFieldTexturesSampler1;
+	FShaderResourceParameter VectorFieldTexturesSampler2;
+	FShaderResourceParameter VectorFieldTexturesSampler3;
 	/** The SceneDepthTexture parameter for depth buffer collision. */
 	FShaderResourceParameter SceneDepthTextureParameter;
 	FShaderResourceParameter SceneDepthTextureParameterSampler;
@@ -3248,7 +3274,7 @@ public:
 			
 			EmitterInstance->ParticleBoundingBox = ComputeParticleBounds(
 				RHICmdList,
-				EmitterInstance->Simulation->VertexBuffer.VertexBufferSRV, 
+				EmitterInstance->Simulation->VertexBuffer.VertexBufferSRV,
 				EmitterInstance->FXSystem->GetParticleSimulationResources()->GetCurrentStateTextures().PositionTextureRHI,
 				EmitterInstance->Simulation->VertexBuffer.ParticleCount
 				);
@@ -4013,20 +4039,16 @@ void FFXSystem::SimulateGPUParticles(
 	// On some platforms, the textures are filled with garbage after creation, so we need to clear them to black the first time we use them
 	if ( !CurrentStateTextures.bTexturesCleared )
 	{
-		SetRenderTarget(RHICmdList, CurrentStateTextures.PositionTextureTargetRHI, FTextureRHIRef());
-		RHICmdList.Clear(true, FLinearColor::Black, false, 1.0f, false, 0, FIntRect());
-		SetRenderTarget(RHICmdList, CurrentStateTextures.VelocityTextureTargetRHI, FTextureRHIRef());
-		RHICmdList.Clear(true, FLinearColor::Black, false, 1.0f, false, 0, FIntRect());
+		SetRenderTarget(RHICmdList, CurrentStateTextures.PositionTextureTargetRHI, FTextureRHIRef(), ESimpleRenderTargetMode::EClearToDefault);
+		SetRenderTarget(RHICmdList, CurrentStateTextures.VelocityTextureTargetRHI, FTextureRHIRef(), ESimpleRenderTargetMode::EClearToDefault);
 
 		CurrentStateTextures.bTexturesCleared = true;
 	}
 
 	if ( !PrevStateTextures.bTexturesCleared )
 	{
-		SetRenderTarget(RHICmdList, PrevStateTextures.PositionTextureTargetRHI, FTextureRHIRef());
-		RHICmdList.Clear(true, FLinearColor::Black, false, 1.0f, false, 0, FIntRect());
-		SetRenderTarget(RHICmdList, PrevStateTextures.VelocityTextureTargetRHI, FTextureRHIRef());
-		RHICmdList.Clear(true, FLinearColor::Black, false, 1.0f, false, 0, FIntRect());
+		SetRenderTarget(RHICmdList, PrevStateTextures.PositionTextureTargetRHI, FTextureRHIRef(), ESimpleRenderTargetMode::EClearToDefault);
+		SetRenderTarget(RHICmdList, PrevStateTextures.VelocityTextureTargetRHI, FTextureRHIRef(), ESimpleRenderTargetMode::EClearToDefault);
 
 		PrevStateTextures.bTexturesCleared = true;
 	}
