Removed deprecated functions

diff --git a/Engine/Source/Runtime/Engine/Classes/Animation/AnimNotifies/AnimNotify.h b/Engine/Source/Runtime/Engine/Classes/Animation/AnimNotifies/AnimNotify.h
--- a/Engine/Source/Runtime/Engine/Classes/Animation/AnimNotifies/AnimNotify.h
+++ b/Engine/Source/Runtime/Engine/Classes/Animation/AnimNotifies/AnimNotify.h
@@ -48,14 +48,6 @@ class ENGINE_API UAnimNotify : public UObject
 #endif // WITH_EDITORONLY_DATA
 	}
 
-	/**
-	 *	Called by the AnimSet viewer when the 'parent' FAnimNotifyEvent is edited.
-	 *
-	 *	@param	AnimSeq			The animation sequence this notify is associated with.
-	 *	@param	OwnerEvent		The event that 'owns' this AnimNotify.
-	 */
-	virtual void AnimNotifyEventChanged(USkeletalMeshComponent* MeshComp, UAnimSequenceBase* Animation, FAnimNotifyEvent* OwnerEvent) {}
-
 	/**
 	 * We don't instance UAnimNotify objects along with the animations they belong to, but
 	 * we still need a way to see which world this UAnimNotify is currently operating on.
diff --git a/Engine/Source/Runtime/Engine/Classes/Animation/AnimNotifies/AnimNotifyState.h b/Engine/Source/Runtime/Engine/Classes/Animation/AnimNotifies/AnimNotifyState.h
--- a/Engine/Source/Runtime/Engine/Classes/Animation/AnimNotifies/AnimNotifyState.h
+++ b/Engine/Source/Runtime/Engine/Classes/Animation/AnimNotifies/AnimNotifyState.h
@@ -56,14 +56,6 @@ class ENGINE_API UAnimNotifyState : public UObject
 #endif // WITH_EDITORONLY_DATA
 	}
 
-	/**
-	 *	Called by the AnimSet viewer when the 'parent' FAnimNotifyEvent is edited.
-	 *
-	 *	@param	AnimSeq			The animation sequence this notify is associated with.
-	 *	@param	OwnerEvent		The event that 'owns' this AnimNotify.
-	 */
-	virtual void AnimNotifyEventChanged(USkeletalMeshComponent* MeshComp, UAnimSequenceBase* Animation, FAnimNotifyEvent * OwnerEvent) {}
-
 	/** UObject Interface */
 	virtual void PostLoad();
 	/** End UObject Interface */
diff --git a/Engine/Source/Runtime/Engine/Classes/Animation/AnimNotifies/AnimNotifyState_Trail.h b/Engine/Source/Runtime/Engine/Classes/Animation/AnimNotifies/AnimNotifyState_Trail.h
--- a/Engine/Source/Runtime/Engine/Classes/Animation/AnimNotifies/AnimNotifyState_Trail.h
+++ b/Engine/Source/Runtime/Engine/Classes/Animation/AnimNotifies/AnimNotifyState_Trail.h
@@ -62,14 +62,6 @@ class UAnimNotifyState_Trail : public UAnimNotifyState
 	virtual void NotifyEnd(class USkeletalMeshComponent * MeshComp, class UAnimSequenceBase * Animation) override;
 
 	bool ValidateInput(class USkeletalMeshComponent * MeshComp, bool bReportErrors = false);
-
-	/**
-	 *	Called by the AnimSet viewer when the 'parent' FAnimNotifyEvent is edited.
-	 *
-	 *	@param	AnimSeq			The animation sequence this notify is associated with.
-	 *	@param	OwnerEvent		The event that 'owns' this AnimNotify.
-	 */
-	virtual void AnimNotifyEventChanged(class USkeletalMeshComponent* MeshComp, class UAnimSequenceBase* Animation, FAnimNotifyEvent * OwnerEvent) override;
 };
 
 
diff --git a/Engine/Source/Runtime/Engine/Private/Animation/AnimNotifyState_Trail.cpp b/Engine/Source/Runtime/Engine/Private/Animation/AnimNotifyState_Trail.cpp
--- a/Engine/Source/Runtime/Engine/Private/Animation/AnimNotifyState_Trail.cpp
+++ b/Engine/Source/Runtime/Engine/Private/Animation/AnimNotifyState_Trail.cpp
@@ -34,10 +34,6 @@ UAnimNotifyState_Trail::UAnimNotifyState_Trail(const FObjectInitializer& ObjectI
 #endif // WITH_EDITORONLY_DATA
 }
 
-void UAnimNotifyState_Trail::AnimNotifyEventChanged(class USkeletalMeshComponent* MeshComp, class UAnimSequenceBase* Animation, FAnimNotifyEvent * OwnerEvent)
-{
-}
-
 void UAnimNotifyState_Trail::NotifyBegin(class USkeletalMeshComponent * MeshComp, class UAnimSequenceBase * Animation, float TotalDuration)
 {
 	bool bError = ValidateInput(MeshComp);
