Fix stat fps and stat unit not displaying correct in safe zone on vr headsets.

diff --git a/Engine/Source/Runtime/Engine/Private/UnrealClient.cpp b/Engine/Source/Runtime/Engine/Private/UnrealClient.cpp
--- a/Engine/Source/Runtime/Engine/Private/UnrealClient.cpp
+++ b/Engine/Source/Runtime/Engine/Private/UnrealClient.cpp
@@ -356,15 +356,15 @@ int32 FStatUnitData::DrawStat(FViewport* InViewport, FCanvas* InCanvas, int32 In
 
 	// Render CPU thread and GPU frame times.
 	UFont* Font = (!FPlatformProperties::SupportsWindowedMode() && GEngine->GetMediumFont()) ? GEngine->GetMediumFont() : GEngine->GetSmallFont();
-	const int32 SafeZone = FPlatformProperties::SupportsWindowedMode() ? 0 : FMath::TruncToInt(InViewport->GetSizeXY().X * 0.05f);
 
 	FColor Color;
-	int32 X3 = InViewport->GetSizeXY().X - SafeZone;
+	// Note InX should already be within the safe zone
+	int32 X3 = InX;
 	if (bShowUnitMaxTimes)
 	{
 		X3 -= Font->GetStringSize(TEXT(" 0000.00 ms "));
 	}
-	const int32 X2 = X3 - Font->GetStringSize(TEXT(" 000.00 ms "));
+	const int32 X2 = bShowUnitMaxTimes ? X3 - Font->GetStringSize(TEXT(" 000.00 ms ")) : X3;
 	const int32 X1 = X2 - Font->GetStringSize(TEXT("Frame: "));
 	const int32 RowHeight = FMath::TruncToInt(Font->GetMaxCharHeight() * 1.1f);
 	const bool bShowUnitTimeGraph = InViewport->GetClient() ? InViewport->GetClient()->IsStatEnabled(TEXT("UnitGraph")) : false;
diff --git a/Engine/Source/Runtime/Engine/Private/UserInterface/Canvas.cpp b/Engine/Source/Runtime/Engine/Private/UserInterface/Canvas.cpp
--- a/Engine/Source/Runtime/Engine/Private/UserInterface/Canvas.cpp
+++ b/Engine/Source/Runtime/Engine/Private/UserInterface/Canvas.cpp
@@ -1379,17 +1379,29 @@ void UCanvas::PopSafeZoneTransform()
 
 void UCanvas::UpdateSafeZoneData()
 {
-	if (FSlateApplication::IsInitialized())
+	if(GEngine && GEngine->IsStereoscopic3D())
+	{
+		FVector2D SafeRegionPercentage = GEngine->StereoRenderingDevice->GetTextSafeRegionBounds();
+
+		CachedDisplayWidth = UnsafeSizeX;
+		CachedDisplayHeight = UnsafeSizeY;
+
+		SafeZonePadX = (CachedDisplayWidth - (CachedDisplayWidth * SafeRegionPercentage.X))/2.f;
+		SafeZonePadY = CachedDisplayHeight - (CachedDisplayHeight * SafeRegionPercentage.Y)/2.f;
+	}
+	else if(FSlateApplication::IsInitialized())
 	{
 		FDisplayMetrics DisplayMetrics;
+
 		FSlateApplication::Get().GetDisplayMetrics(DisplayMetrics);
-	
+
 		SafeZonePadX = FMath::CeilToInt(DisplayMetrics.TitleSafePaddingSize.X);
 		SafeZonePadY = FMath::CeilToInt(DisplayMetrics.TitleSafePaddingSize.Y);
 
 		CachedDisplayWidth = DisplayMetrics.PrimaryDisplayWidth;
 		CachedDisplayHeight = DisplayMetrics.PrimaryDisplayHeight;
 	}
+
 }
 
 void UCanvas::UpdateAllCanvasSafeZoneData()
