Remove unused member variables

diff --git a/Engine/Source/Runtime/Engine/Classes/Engine/ActorChannel.h b/Engine/Source/Runtime/Engine/Classes/Engine/ActorChannel.h
--- a/Engine/Source/Runtime/Engine/Classes/Engine/ActorChannel.h
+++ b/Engine/Source/Runtime/Engine/Classes/Engine/ActorChannel.h
@@ -106,9 +106,6 @@ class ENGINE_API UActorChannel : public UChannel
 
 	// End UChannel interface.
 
-	/** Sets an override GUID for this actor channel */
-	void SetActorNetGUIDOverride(const FNetworkGUID &NetGUID);
-
 	/** UActorChannel interface and accessors. */
 	AActor* GetActor() {return Actor;}
 
