UE4 - UE-6923 - Bad replication of swizzles

diff --git a/Engine/Source/Runtime/Engine/Private/Materials/MaterialUniformExpressions.h b/Engine/Source/Runtime/Engine/Private/Materials/MaterialUniformExpressions.h
--- a/Engine/Source/Runtime/Engine/Private/Materials/MaterialUniformExpressions.h
+++ b/Engine/Source/Runtime/Engine/Private/Materials/MaterialUniformExpressions.h
@@ -962,6 +962,15 @@ public:
 		// At least one proper index
 		check(InR >= 0 && InR <= 3);
 		++NumElements;
+
+		if (NumElements == 1)
+		{
+			// Replicate scalar
+			IndexG = IndexR;
+			IndexB = IndexR;
+			IndexA = IndexR;
+			NumElements = 4;
+		}
 	}
 
 	// FMaterialUniformExpression interface.
@@ -978,7 +987,8 @@ public:
 	{
 		FLinearColor Temp = OutValue;
 		X->GetNumberValue(Context, Temp);
-		OutValue = Temp;
+		// Clear
+		OutValue *= 0;
 		switch (NumElements)
 		{
 		case 4:
diff --git a/Engine/Source/Runtime/Engine/Private/ShaderDerivedDataVersion.h b/Engine/Source/Runtime/Engine/Private/ShaderDerivedDataVersion.h
--- a/Engine/Source/Runtime/Engine/Private/ShaderDerivedDataVersion.h
+++ b/Engine/Source/Runtime/Engine/Private/ShaderDerivedDataVersion.h
@@ -10,4 +10,4 @@
 // guid as version
 
 #define GLOBALSHADERMAP_DERIVEDDATA_VER			TEXT("af1606fcec114373844d2b6cf9df71c8")
-#define MATERIALSHADERMAP_DERIVEDDATA_VER		TEXT("85dd00daf10d45c79d4ed0306af26289")
+#define MATERIALSHADERMAP_DERIVEDDATA_VER		TEXT("85dd00daf10d45c79d4ed0306af2628b")
