Remove unneeded physics override check (our call to GetPhysicsAsset after it will check the override for us)

diff --git a/Engine/Source/Runtime/Engine/Private/SkeletalMeshComponent.cpp b/Engine/Source/Runtime/Engine/Private/SkeletalMeshComponent.cpp
--- a/Engine/Source/Runtime/Engine/Private/SkeletalMeshComponent.cpp
+++ b/Engine/Source/Runtime/Engine/Private/SkeletalMeshComponent.cpp
@@ -1714,11 +1714,7 @@ float USkeletalMeshComponent::CalculateMass(FName BoneName)
 	else	//We want to calculate mass before we've initialized body instances - in this case use physics asset setup
 	{
 		TArray<class UBodySetup*> * BodySetups = NULL;
-		if (PhysicsAssetOverride)
-		{
-			BodySetups = &PhysicsAssetOverride->BodySetup;
-		}
-		else if (UPhysicsAsset * PhysicsAsset = GetPhysicsAsset())
+		if (UPhysicsAsset * PhysicsAsset = GetPhysicsAsset())
 		{
 			BodySetups = &PhysicsAsset->BodySetup;
 		}
