Fix for UE-5189 Skel/Vert/Surf module is spawning sprites at emitter origin.

This was being caused by a combination of the daft way the skel/vert/surf module limits particles to an array of valid bone names.
Particles are assigned a vertex/surf index at spawn and this is then checked against the valid bone list. If it is not influenced by a bone on the list it is rejected by setting it's RelativeLife to 1.1.
If a Lifetime module is placed after the skel vert module then this rejection is overwritten and the particle spwans at 0,0,0.
I've added a check for this in the lifetime module to preserve the original RelativeTime if it's already set at > 1.0.
A proper fix to this will require bigger changes to the skel/vert module than I'm willing to delve into now that cascade is being ushered off to the retirement home.

diff --git a/Engine/Source/Runtime/Engine/Private/Particles/ParticleModules.cpp b/Engine/Source/Runtime/Engine/Private/Particles/ParticleModules.cpp
--- a/Engine/Source/Runtime/Engine/Private/Particles/ParticleModules.cpp
+++ b/Engine/Source/Runtime/Engine/Private/Particles/ParticleModules.cpp
@@ -3300,7 +3300,8 @@ void UParticleModuleLifetime::SpawnEx(FParticleEmitterInstance* Owner, int32 Off
 			// First module to modify lifetime.
 			Particle.OneOverMaxLifetime = MaxLifetime > 0.f ? 1.f / MaxLifetime : 0.f;
 		}
-		Particle.RelativeTime = SpawnTime * Particle.OneOverMaxLifetime;
+		//If the relative time is already > 1.0f then we don't want to be setting it. Some modules use this to mark a particle as dead during spawn.
+		Particle.RelativeTime = Particle.RelativeTime > 1.0f ? Particle.RelativeTime : SpawnTime * Particle.OneOverMaxLifetime;
 	}
 }
 
