UE 3849 - UE4: Cascade: Regression: Spawn using Surface Normal stopped working  Skel Vert Surf location module.

diff --git a/Engine/Source/Runtime/Core/Private/GenericPlatform/GenericPlatformMath.cpp b/Engine/Source/Runtime/Core/Private/GenericPlatform/GenericPlatformMath.cpp
--- a/Engine/Source/Runtime/Core/Private/GenericPlatform/GenericPlatformMath.cpp
+++ b/Engine/Source/Runtime/Core/Private/GenericPlatform/GenericPlatformMath.cpp
@@ -11,6 +11,11 @@ void FGenericPlatformMath::SRandInit( int32 Seed )
 	GSRandSeed = Seed; 
 }
 
+int32 FGenericPlatformMath::GetRandSeed()
+{
+	return GSRandSeed;
+}
+
 float FGenericPlatformMath::SRand() 
 { 
 	GSRandSeed = (GSRandSeed * 196314165) + 907633515;
diff --git a/Engine/Source/Runtime/Core/Public/GenericPlatform/GenericPlatformMath.h b/Engine/Source/Runtime/Core/Public/GenericPlatform/GenericPlatformMath.h
--- a/Engine/Source/Runtime/Core/Public/GenericPlatform/GenericPlatformMath.h
+++ b/Engine/Source/Runtime/Core/Public/GenericPlatform/GenericPlatformMath.h
@@ -239,6 +239,9 @@ struct FGenericPlatformMath
 	/** Seeds future calls to SRand() */
 	static CORE_API void SRandInit( int32 Seed );
 
+	/** Returns the current seed for SRand(). */
+	static CORE_API int32 GetRandSeed();
+
 	/** Returns a seeded random float in the range [0,1), using the seed from SRandInit(). */
 	static CORE_API float SRand();
 
diff --git a/Engine/Source/Runtime/Core/Public/UObject/ObjectVersion.h b/Engine/Source/Runtime/Core/Public/UObject/ObjectVersion.h
--- a/Engine/Source/Runtime/Core/Public/UObject/ObjectVersion.h
+++ b/Engine/Source/Runtime/Core/Public/UObject/ObjectVersion.h
@@ -448,6 +448,8 @@ enum EUnrealEngineObjectUE4Version
 	VER_UE4_MIKKTSPACE_IS_DEFAULT,
 	// Only applies to cooked files, grass cooking support.
 	VER_UE4_LANDSCAPE_GRASS_COOKING,
+	// Fixed code for using the bOrientMeshEmitters property.
+	VER_UE4_FIX_SKEL_VERT_ORIENT_MESH_PARTICLES,
 
 	// -----<new versions can be added before this line>-------------------------------------------------
 	// - this needs to be the last line (see note below)
diff --git a/Engine/Source/Runtime/Engine/Classes/Particles/Location/ParticleModuleLocationSkelVertSurface.h b/Engine/Source/Runtime/Engine/Classes/Particles/Location/ParticleModuleLocationSkelVertSurface.h
--- a/Engine/Source/Runtime/Engine/Classes/Particles/Location/ParticleModuleLocationSkelVertSurface.h
+++ b/Engine/Source/Runtime/Engine/Classes/Particles/Location/ParticleModuleLocationSkelVertSurface.h
@@ -136,7 +136,7 @@ class UParticleModuleLocationSkelVertSurface : public UParticleModuleLocationBas
 	 *	
 	 *	@return	bool					true if successful, false if not
 	 */
-	bool GetParticleLocation(FParticleEmitterInstance* Owner, USkeletalMeshComponent* InSkelMeshComponent, int32 InPrimaryVertexIndex, FVector& OutPosition, FQuat* OutRotation, bool bSpawning = false);
+	bool GetParticleLocation(FParticleEmitterInstance* Owner, USkeletalMeshComponent* InSkelMeshComponent, int32 InPrimaryVertexIndex, FVector& OutPosition, FQuat& OutRotation, bool bSpawning = false);
 
 	/**
 	 *  Check to see if the vert is influenced by a bone on our approved list.
diff --git a/Engine/Source/Runtime/Engine/Private/Particles/ParticleModules_Location.cpp b/Engine/Source/Runtime/Engine/Private/Particles/ParticleModules_Location.cpp
--- a/Engine/Source/Runtime/Engine/Private/Particles/ParticleModules_Location.cpp
+++ b/Engine/Source/Runtime/Engine/Private/Particles/ParticleModules_Location.cpp
@@ -1973,6 +1973,13 @@ void UParticleModuleLocationSkelVertSurface::PostLoad()
 	}
 
 	NormalCheckTolerance = ((1.0f-(NormalCheckToleranceDegrees/180.0f))*2.0f)-1.0f;
+
+	if (GetLinkerUE4Version() < VER_UE4_FIX_SKEL_VERT_ORIENT_MESH_PARTICLES)
+	{
+		//The code to actually do this hasn't been present ever in UE4 so I'm disabling it for old emitters.
+		//I expect that some users will have this set to true and it will alter the behavior of their emitters under their feet.
+		bOrientMeshEmitters = false;
+	}
 }
 
 #if WITH_EDITOR
@@ -2101,10 +2108,9 @@ void UParticleModuleLocationSkelVertSurface::Spawn(FParticleEmitterInstance* Own
 	}
 
 	FVector SourceLocation;
-	FQuat RotationQuat;
+	FQuat SourceRotation;
 	const int32 MeshRotationOffset = Owner->GetMeshRotationOffset();
 	const bool bMeshRotationActive = MeshRotationOffset > 0 && Owner->IsMeshRotationActive();
-	FQuat* SourceRotation = (bMeshRotationActive) ? NULL : &RotationQuat;
 	if (GetParticleLocation(Owner, SourceComponent, SourceIndex, SourceLocation, SourceRotation, true) == true)
 	{
 		SPAWN_INIT
@@ -2153,11 +2159,23 @@ void UParticleModuleLocationSkelVertSurface::Spawn(FParticleEmitterInstance* Own
 			if (bMeshRotationActive)
 			{
 				FMeshRotationPayloadData* PayloadData = (FMeshRotationPayloadData*)((uint8*)&Particle + MeshRotationOffset);
-				PayloadData->Rotation = RotationQuat.Euler();
+				if (bOrientMeshEmitters)
+				{
+					//We have the mesh oriented to the normal of the triangle it's on but this looks fugly as particles on each triangle are facing the same way.
+					//The only valid orientation reference should be the normal. So add an additional random rotation around it.
+					int32 OldRandSeed = FMath::GetRandSeed();
+					FMath::SRandInit((int32)ParticleBase);
+					SourceRotation = SourceRotation * FQuat(FVector::UpVector, FMath::SRand()*(PI*2.0f));
+					FMath::SRandInit(OldRandSeed);
+				}
+
+				FVector Rot = SourceRotation.Euler();
 				if (Owner->CurrentLODLevel->RequiredModule->bUseLocalSpace == true)
 				{
-					PayloadData->Rotation = Owner->Component->ComponentToWorld.InverseTransformVectorNoScale(PayloadData->Rotation);
+					Rot = Owner->Component->ComponentToWorld.InverseTransformVectorNoScale(Rot);
 				}
+				PayloadData->Rotation = Rot;
+				PayloadData->InitRotation = Rot;
 			}
 		}
 	}
@@ -2213,10 +2231,9 @@ void UParticleModuleLocationSkelVertSurface::Update(FParticleEmitterInstance* Ow
 	}
 
 	FVector SourceLocation;
-	FQuat RotationQuat;
+	FQuat SourceRotation;
 	const int32 MeshRotationOffset = Owner->GetMeshRotationOffset();
 	const bool bMeshRotationActive = MeshRotationOffset > 0 && Owner->IsMeshRotationActive();
-	FQuat* SourceRotation = (bMeshRotationActive) ? NULL : &RotationQuat;
 	BEGIN_UPDATE_LOOP;
 	{
 		FModuleLocationVertSurfaceParticlePayload* ParticlePayload = (FModuleLocationVertSurfaceParticlePayload*)((uint8*)&Particle + Offset);
@@ -2225,12 +2242,24 @@ void UParticleModuleLocationSkelVertSurface::Update(FParticleEmitterInstance* Ow
 			Particle.Location = SourceLocation;
 			if (bMeshRotationActive)
 			{
+				if (bOrientMeshEmitters)
+				{
+					//We have the mesh oriented to the normal of the triangle it's on but this looks fugly as particles on each triangle are facing the same way.
+					//The only valid orientation reference should be the normal. So add an additional random rotation around it.
+
+					int32 OldRandSeed = FMath::GetRandSeed();
+					FMath::SRandInit((int32)ParticleBase);
+					SourceRotation = SourceRotation * FQuat(FVector::UpVector, FMath::SRand()*(PI*2.0f));
+					FMath::SRandInit(OldRandSeed);
+				}
+
 				FMeshRotationPayloadData* PayloadData = (FMeshRotationPayloadData*)((uint8*)&Particle + MeshRotationOffset);
-				PayloadData->Rotation = RotationQuat.Euler();
+				FVector Rot = SourceRotation.Euler();
 				if (Owner->CurrentLODLevel->RequiredModule->bUseLocalSpace == true)
 				{
-					PayloadData->Rotation = Owner->Component->ComponentToWorld.InverseTransformVectorNoScale(PayloadData->Rotation);
+					Rot = Owner->Component->ComponentToWorld.InverseTransformVectorNoScale(Rot);
 				}
+				PayloadData->Rotation = Rot;
 			}
 		}
 	}
@@ -2493,7 +2522,7 @@ USkeletalMeshComponent* UParticleModuleLocationSkelVertSurface::GetSkeletalMeshC
 
 bool UParticleModuleLocationSkelVertSurface::GetParticleLocation(FParticleEmitterInstance* Owner, 
 	USkeletalMeshComponent* InSkelMeshComponent, int32 InPrimaryVertexIndex, 
-	FVector& OutPosition, FQuat* OutRotation, bool bSpawning /* = false*/)
+	FVector& OutPosition, FQuat& OutRotation, bool bSpawning /* = false*/)
 {
 	check(InSkelMeshComponent);
 	FSkeletalMeshResource* SkelMeshResource = InSkelMeshComponent->GetSkeletalMeshResource();
@@ -2504,10 +2533,7 @@ bool UParticleModuleLocationSkelVertSurface::GetParticleLocation(FParticleEmitte
 		{
 			FVector VertPos = InSkelMeshComponent->GetSkinnedVertexPosition(InPrimaryVertexIndex);
 			OutPosition = InSkelMeshComponent->ComponentToWorld.TransformPosition(VertPos);
-			if (OutRotation != NULL)
-			{
-				*OutRotation = FQuat::Identity;
-			}
+			OutRotation = FQuat::Identity;
 		}
 		else if (SourceType == VERTSURFACESOURCE_Surface)
 		{
@@ -2522,13 +2548,16 @@ bool UParticleModuleLocationSkelVertSurface::GetParticleLocation(FParticleEmitte
 			Verts[1] = InSkelMeshComponent->ComponentToWorld.TransformPosition(InSkelMeshComponent->GetSkinnedVertexPosition(VertIndex[1]));
 			Verts[2] = InSkelMeshComponent->ComponentToWorld.TransformPosition(InSkelMeshComponent->GetSkinnedVertexPosition(VertIndex[2]));
 
+			FVector V0ToV2 = (Verts[2] - Verts[0]);
+			V0ToV2.Normalize();
+			FVector V0ToV1 = (Verts[1] - Verts[0]);
+			V0ToV1.Normalize();
+			FVector Normal = V0ToV2 ^ V0ToV1;
+			Normal.Normalize();
+
 			if(bEnforceNormalCheck && bSpawning)
 			{
-
-				FVector Direction = (Verts[2]-Verts[0]) ^ (Verts[1]-Verts[0]);
-				Direction.Normalize();
-
-				float Dot = Direction | NormalToCompare;
+				float Dot = Normal | NormalToCompare;
 
 				if(Dot < ((2.0f*NormalCheckTolerance)-1.0f))
 				{
@@ -2544,9 +2573,19 @@ bool UParticleModuleLocationSkelVertSurface::GetParticleLocation(FParticleEmitte
 				OutPosition = (Verts[0] + Verts[1] + Verts[2]) / 3.0f;
 			}
 
-			if (OutRotation != NULL)
+			if (bOrientMeshEmitters)
+			{
+				FVector Fwd = Normal ^ V0ToV1;
+				Fwd.Normalize();
+				FMatrix Orientation(FMatrix::Identity);
+				Orientation.SetAxes(&V0ToV1, &Fwd, &Normal);
+				OutRotation = FQuat(Orientation);
+				OutRotation.Normalize();
+
+			}
+			else
 			{
-				*OutRotation = FQuat::Identity;
+				OutRotation = FQuat::Identity;
 			}
 		}
 	}
