Speculative fix for casting etiquette required to get the bottom 32 bits of a pointer being used as the seed for a rand.
As usual, VS is happy with any old rubbish but the other compilers are more careful.

diff --git a/Engine/Source/Runtime/Engine/Private/Particles/ParticleModules_Location.cpp b/Engine/Source/Runtime/Engine/Private/Particles/ParticleModules_Location.cpp
--- a/Engine/Source/Runtime/Engine/Private/Particles/ParticleModules_Location.cpp
+++ b/Engine/Source/Runtime/Engine/Private/Particles/ParticleModules_Location.cpp
@@ -2164,7 +2164,7 @@ void UParticleModuleLocationSkelVertSurface::Spawn(FParticleEmitterInstance* Own
 					//We have the mesh oriented to the normal of the triangle it's on but this looks fugly as particles on each triangle are facing the same way.
 					//The only valid orientation reference should be the normal. So add an additional random rotation around it.
 					int32 OldRandSeed = FMath::GetRandSeed();
-					FMath::SRandInit((int32)ParticleBase);
+					FMath::SRandInit((int32)((intptr_t)ParticleBase));
 					SourceRotation = SourceRotation * FQuat(FVector::UpVector, FMath::SRand()*(PI*2.0f));
 					FMath::SRandInit(OldRandSeed);
 				}
@@ -2248,7 +2248,7 @@ void UParticleModuleLocationSkelVertSurface::Update(FParticleEmitterInstance* Ow
 					//The only valid orientation reference should be the normal. So add an additional random rotation around it.
 
 					int32 OldRandSeed = FMath::GetRandSeed();
-					FMath::SRandInit((int32)ParticleBase);
+					FMath::SRandInit((int32)((intptr_t)ParticleBase));
 					SourceRotation = SourceRotation * FQuat(FVector::UpVector, FMath::SRand()*(PI*2.0f));
 					FMath::SRandInit(OldRandSeed);
 				}
