Particles - Marking ParticleModuleVelocityCone as MinimalAPI so that its symbols get exported.

diff --git a/Engine/Source/Runtime/Engine/Classes/Particles/Velocity/ParticleModuleVelocityCone.h b/Engine/Source/Runtime/Engine/Classes/Particles/Velocity/ParticleModuleVelocityCone.h
--- a/Engine/Source/Runtime/Engine/Classes/Particles/Velocity/ParticleModuleVelocityCone.h
+++ b/Engine/Source/Runtime/Engine/Classes/Particles/Velocity/ParticleModuleVelocityCone.h
@@ -5,7 +5,7 @@
 #include "Particles/Velocity/ParticleModuleVelocityBase.h"
 #include "ParticleModuleVelocityCone.generated.h"
 
-UCLASS(editinlinenew, hidecategories=Object, meta=(DisplayName = "Velocity Cone"))
+UCLASS(editinlinenew, hidecategories=Object, MinimalAPI, meta=( DisplayName = "Velocity Cone" ))
 class UParticleModuleVelocityCone : public UParticleModuleVelocityBase
 {
 	GENERATED_UCLASS_BODY()
