Fix for foliage insatnces deleted pre-4.7 being resurrected when opened in 4.7.

diff --git a/Engine/Source/Runtime/Engine/Private/InstancedFoliage.cpp b/Engine/Source/Runtime/Engine/Private/InstancedFoliage.cpp
--- a/Engine/Source/Runtime/Engine/Private/InstancedFoliage.cpp
+++ b/Engine/Source/Runtime/Engine/Private/InstancedFoliage.cpp
@@ -63,14 +63,28 @@ FArchive& operator<<(FArchive& Ar, FFoliageInstance& Instance)
 	Ar << Instance.Location;
 	Ar << Instance.Rotation;
 	Ar << Instance.DrawScale3D;
+
 	if (Ar.CustomVer(FFoliageCustomVersion::GUID) < FFoliageCustomVersion::FoliageUsingHierarchicalISMC)
 	{
 		int32 OldClusterIndex;
 		Ar << OldClusterIndex;
+		Ar << Instance.PreAlignRotation;
+		Ar << Instance.Flags;
+
+		if (OldClusterIndex == INDEX_NONE)
+		{
+			// When converting, we need to skip over any instance that was previously deleted but still in the Instances array.
+			Instance.Flags |= FOLIAGE_InstanceDeleted;
+		}
+	}
+	else
+	{
+		Ar << Instance.PreAlignRotation;
+		Ar << Instance.Flags;
 	}
-	Ar << Instance.PreAlignRotation;
-	Ar << Instance.Flags;
+	
 	Ar << Instance.ZOffset;
+
 	return Ar;
 }
 
@@ -527,7 +541,6 @@ void FFoliageMeshInfo::ReallocateClusters(AInstancedFoliageActor* InIFA, UFoliag
 		Component = nullptr;
 	}
 
-
 	// Remove everything
 	TArray<FFoliageInstance> OldInstances;
 	Exchange(Instances, OldInstances);
@@ -538,7 +551,10 @@ void FFoliageMeshInfo::ReallocateClusters(AInstancedFoliageActor* InIFA, UFoliag
 	// Re-add
 	for (FFoliageInstance& Instance : OldInstances)
 	{
-		AddInstance(InIFA, InSettings, Instance);
+		if ((Instance.Flags & FOLIAGE_InstanceDeleted) == 0)
+		{
+			AddInstance(InIFA, InSettings, Instance);
+		}
 	}
 
 	InIFA->RegisterAllComponents();
diff --git a/Engine/Source/Runtime/Engine/Public/InstancedFoliage.h b/Engine/Source/Runtime/Engine/Public/InstancedFoliage.h
--- a/Engine/Source/Runtime/Engine/Public/InstancedFoliage.h
+++ b/Engine/Source/Runtime/Engine/Public/InstancedFoliage.h
@@ -20,9 +20,10 @@ struct FFoliageInstanceHash;
  */
 enum EFoliageInstanceFlags
 {
-	FOLIAGE_AlignToNormal = 0x00000001,
-	FOLIAGE_NoRandomYaw   = 0x00000002,
-	FOLIAGE_Readjusted    = 0x00000004,
+	FOLIAGE_AlignToNormal	= 0x00000001,
+	FOLIAGE_NoRandomYaw		= 0x00000002,
+	FOLIAGE_Readjusted		= 0x00000004,
+	FOLIAGE_InstanceDeleted	= 0x00000008,	// Used only for migration from pre-HierarchicalISM foliage.
 };
 
 /**
