For Japanese keyboards on Windows platforms, use @ key as the default in-game console key

diff --git a/Engine/Source/Runtime/Engine/Private/UserInterface/InputSettings.cpp b/Engine/Source/Runtime/Engine/Private/UserInterface/InputSettings.cpp
--- a/Engine/Source/Runtime/Engine/Private/UserInterface/InputSettings.cpp
+++ b/Engine/Source/Runtime/Engine/Private/UserInterface/InputSettings.cpp
@@ -60,6 +60,10 @@ void UInputSettings::PostInitProperties()
 		case LANG_SPANISH:
 			DefaultConsoleKey = FInputKeyManager::Get().GetKeyFromCodes(VK_OEM_5, 0);
 			break;
+			
+		case LANG_JAPANESE:
+			DefaultConsoleKey = FInputKeyManager::Get().GetKeyFromCodes(VK_OEM_3, 0);
+			break;
 		}
 
 		ConsoleKeys[0] = DefaultConsoleKey;
