Fixed a handful of IntelliSense compiler problems with engine header files
- This reduces the chance of seeing false errors in your Error List window or false "squiggles" on code

diff --git a/Engine/Source/Runtime/Core/Private/Serialization/Archive.cpp b/Engine/Source/Runtime/Core/Private/Serialization/Archive.cpp
--- a/Engine/Source/Runtime/Core/Private/Serialization/Archive.cpp
+++ b/Engine/Source/Runtime/Core/Private/Serialization/Archive.cpp
@@ -193,7 +193,7 @@ void FArchive::ResetCustomVersions()
 	bCustomVersionsAreReset = true;
 }
 
-void FArchive::UsingCustomVersion(FGuid Key)
+void FArchive::UsingCustomVersion(const FGuid& Key)
 {
 	// If we're loading, we want to use the version that the archive was serialized with, not register a new one.
 	if (IsLoading())
@@ -208,7 +208,7 @@ void FArchive::UsingCustomVersion(FGuid Key)
 	const_cast<FCustomVersionContainer&>(GetCustomVersions()).SetVersion(Key, RegisteredVersion->Version, RegisteredVersion->FriendlyName);
 }
 
-int32 FArchive::CustomVer(FGuid Key) const
+int32 FArchive::CustomVer(const FGuid& Key) const
 {
 	auto* CustomVersion = GetCustomVersions().GetVersion(Key);
 
@@ -219,7 +219,7 @@ int32 FArchive::CustomVer(FGuid Key) const
 	return CustomVersion ? CustomVersion->Version : -1;
 }
 
-void FArchive::SetCustomVersion(FGuid Key, int32 Version, FString FriendlyName)
+void FArchive::SetCustomVersion(const FGuid&  Key, int32 Version, FString FriendlyName)
 {
 	const_cast<FCustomVersionContainer&>(GetCustomVersions()).SetVersion(Key, Version, FriendlyName);
 }
diff --git a/Engine/Source/Runtime/Core/Public/Serialization/ArchiveBase.h b/Engine/Source/Runtime/Core/Public/Serialization/ArchiveBase.h
--- a/Engine/Source/Runtime/Core/Public/Serialization/ArchiveBase.h
+++ b/Engine/Source/Runtime/Core/Public/Serialization/ArchiveBase.h
@@ -553,7 +553,7 @@ public:
 	 *
 	 * @param Guid The guid of the custom version.  This must have previously been registered with FCustomVersionRegistration.
 	 */
-	void UsingCustomVersion(FGuid Guid);
+	void UsingCustomVersion(const class FGuid& Guid);
 
 	/**
 	 * Queries a custom version from the archive.  If the archive is being used to write, the custom version must have already been registered.
@@ -561,7 +561,7 @@ public:
 	 * @param Key The guid of the custom version to query.
 	 * @return The version number, or 0 if the custom tag isn't stored in the archive.
 	 */
-	int32 CustomVer(FGuid Key) const;
+	int32 CustomVer(const class FGuid& Key) const;
 
 	FORCEINLINE bool IsLoading() const
 	{
@@ -751,7 +751,7 @@ public:
 	 * @param Version - The version number to set key to
 	 * @param FriendlyName - Friendly name corresponding to the key
 	 */
-	void SetCustomVersion(FGuid Key, int32 Version, FString FriendlyName);
+	void SetCustomVersion(const class FGuid& Key, int32 Version, FString FriendlyName);
 
 	/**
 	 * Toggle saving as Unicode. This is needed when we need to make sure ANSI strings are saved as Unicode
diff --git a/Engine/Source/Runtime/Engine/Classes/Particles/ParticleSystemComponent.h b/Engine/Source/Runtime/Engine/Classes/Particles/ParticleSystemComponent.h
--- a/Engine/Source/Runtime/Engine/Classes/Particles/ParticleSystemComponent.h
+++ b/Engine/Source/Runtime/Engine/Classes/Particles/ParticleSystemComponent.h
@@ -887,7 +887,7 @@ public:
 	 *  @param  InEventData         Gamespecific event data payload
 	 */
 	void ReportEventSpawn(const FName InEventName, const float InEmitterTime,
-		const FVector InLocation, const FVector InVelocity, const TArray<UParticleModuleEventSendToGame*>& InEventData);
+		const FVector InLocation, const FVector InVelocity, const TArray<class UParticleModuleEventSendToGame*>& InEventData);
 
 	/**
 	 *	Record a death event.
@@ -900,7 +900,7 @@ public:
 	 *	@param	InParticleTime		The relative life of the particle when the event fired.
 	 */
 	void ReportEventDeath(const FName InEventName, const float InEmitterTime,
-		const FVector InLocation, const FVector InVelocity, const TArray<UParticleModuleEventSendToGame*>& InEventData, const float InParticleTime);
+		const FVector InLocation, const FVector InVelocity, const TArray<class UParticleModuleEventSendToGame*>& InEventData, const float InParticleTime);
 
 	/**
 	 *	Record a collision event.
@@ -918,7 +918,7 @@ public:
 	 *	@param	InBoneName		Name of bone we hit (for skeletal meshes).
 	 */
 	void ReportEventCollision(const FName InEventName, const float InEmitterTime, const FVector InLocation,
-		const FVector InDirection, const FVector InVelocity, const TArray<UParticleModuleEventSendToGame*>& InEventData, 
+		const FVector InDirection, const FVector InVelocity, const TArray<class UParticleModuleEventSendToGame*>& InEventData, 
 		const float InParticleTime, const FVector InNormal, const float InTime, const int32 InItem, const FName InBoneName);
 
 	/**
@@ -930,7 +930,7 @@ public:
 	 *	@param	InLocation			The location of the particle emitter when the event fired.
 	 */
 	void ReportEventBurst(const FName InEventName, const float InEmitterTime, const int32 ParticleCount,
-		const FVector InLocation, const TArray<UParticleModuleEventSendToGame*>& InEventData);
+		const FVector InLocation, const TArray<class UParticleModuleEventSendToGame*>& InEventData);
 
 	/**
 	 *	Record a kismet event.
diff --git a/Engine/Source/Runtime/Engine/Public/BlueprintUtilities.h b/Engine/Source/Runtime/Engine/Public/BlueprintUtilities.h
--- a/Engine/Source/Runtime/Engine/Public/BlueprintUtilities.h
+++ b/Engine/Source/Runtime/Engine/Public/BlueprintUtilities.h
@@ -5,8 +5,8 @@
 #include "GraphEditAction.h"
 
 DECLARE_MULTICAST_DELEGATE_OneParam( FOnGraphChanged, const FEdGraphEditAction& );
-DECLARE_DELEGATE_OneParam( FSingleNodeEvent, UEdGraphNode* );
-DECLARE_DELEGATE_OneParam( FEdGraphEvent, UEdGraph* );
+DECLARE_DELEGATE_OneParam( FSingleNodeEvent, class UEdGraphNode* );
+DECLARE_DELEGATE_OneParam( FEdGraphEvent, class UEdGraph* );
 /** Delegate for notification when property changed */
 DECLARE_MULTICAST_DELEGATE_TwoParams( FOnPropertyChanged, const FPropertyChangedEvent&, const FString& );
 
diff --git a/Engine/Source/Runtime/Engine/Public/UnrealClient.h b/Engine/Source/Runtime/Engine/Public/UnrealClient.h
--- a/Engine/Source/Runtime/Engine/Public/UnrealClient.h
+++ b/Engine/Source/Runtime/Engine/Public/UnrealClient.h
@@ -60,7 +60,7 @@ public:
 	* @param InRect - source rect of the image to capture
 	* @return True if the read succeeded.
 	*/
-	ENGINE_API bool ReadPixels(TArray< FColor >& OutImageData,FReadSurfaceDataFlags InFlags = FReadSurfaceDataFlags(), FIntRect InRect = FIntRect(0, 0, 0, 0) );
+	ENGINE_API bool ReadPixels(TArray< FColor >& OutImageData,FReadSurfaceDataFlags InFlags = FReadSurfaceDataFlags(RCM_UNorm, CubeFace_MAX), FIntRect InRect = FIntRect(0, 0, 0, 0) );
 
 	/**
 	* Reads the viewport's displayed pixels into a preallocated color buffer.
@@ -68,7 +68,7 @@ public:
 	* @param InSrcRect InSrcRect not specified means the whole rect
 	* @return True if the read succeeded.
 	*/
-	ENGINE_API bool ReadPixelsPtr(FColor* OutImageBytes, FReadSurfaceDataFlags InFlags = FReadSurfaceDataFlags(), FIntRect InSrcRect = FIntRect(0, 0, 0, 0));
+	ENGINE_API bool ReadPixelsPtr(FColor* OutImageBytes, FReadSurfaceDataFlags InFlags = FReadSurfaceDataFlags(RCM_UNorm, CubeFace_MAX), FIntRect InSrcRect = FIntRect(0, 0, 0, 0));
 
 	/**
 	 * Reads the viewport's displayed pixels into the given color buffer.
diff --git a/Engine/Source/Runtime/ShaderCore/Public/VertexFactory.h b/Engine/Source/Runtime/ShaderCore/Public/VertexFactory.h
--- a/Engine/Source/Runtime/ShaderCore/Public/VertexFactory.h
+++ b/Engine/Source/Runtime/ShaderCore/Public/VertexFactory.h
@@ -499,7 +499,7 @@ public:
 		delete Parameters;
 	}
 
-	void SetMesh(FRHICommandList& RHICmdList, FShader* Shader,const FVertexFactory* VertexFactory,const class FSceneView& View,const FMeshBatchElement& BatchElement,uint32 DataFlags) const
+	void SetMesh(FRHICommandList& RHICmdList, FShader* Shader,const FVertexFactory* VertexFactory,const class FSceneView& View,const struct FMeshBatchElement& BatchElement,uint32 DataFlags) const
 	{
 		if(Parameters)
 		{
