Fixed collision types not being used in collision checks. UE-4466

diff --git a/Engine/Source/Runtime/Engine/Classes/Particles/Collision/ParticleModuleCollision.h b/Engine/Source/Runtime/Engine/Classes/Particles/Collision/ParticleModuleCollision.h
--- a/Engine/Source/Runtime/Engine/Classes/Particles/Collision/ParticleModuleCollision.h
+++ b/Engine/Source/Runtime/Engine/Classes/Particles/Collision/ParticleModuleCollision.h
@@ -150,6 +150,7 @@ class UParticleModuleCollision : public UParticleModuleCollisionBase
 	virtual void PostEditChangeProperty(FPropertyChangedEvent& PropertyChangedEvent) override;
 #endif // WITH_EDITOR
 	virtual void PostInitProperties() override;
+	virtual void PostLoad() override;
 	//End UObject Interface
 
 	//Begin UParticleModule Interface
diff --git a/Engine/Source/Runtime/Engine/Private/Particles/ParticleModules_Collision.cpp b/Engine/Source/Runtime/Engine/Private/Particles/ParticleModules_Collision.cpp
--- a/Engine/Source/Runtime/Engine/Private/Particles/ParticleModules_Collision.cpp
+++ b/Engine/Source/Runtime/Engine/Private/Particles/ParticleModules_Collision.cpp
@@ -98,6 +98,13 @@ void UParticleModuleCollision::PostInitProperties()
 	}
 }
 
+void UParticleModuleCollision::PostLoad()
+{
+	Super::PostLoad();
+
+	ObjectParams = FCollisionObjectQueryParams(CollisionTypes);
+}
+
 #if WITH_EDITOR
 void UParticleModuleCollision::PostEditChangeProperty(FPropertyChangedEvent& PropertyChangedEvent)
 {
