Fix Angular drive mode greying out when velocity motor is not enabled

diff --git a/Engine/Source/Runtime/Engine/Classes/PhysicsEngine/ConstraintInstance.h b/Engine/Source/Runtime/Engine/Classes/PhysicsEngine/ConstraintInstance.h
--- a/Engine/Source/Runtime/Engine/Classes/PhysicsEngine/ConstraintInstance.h
+++ b/Engine/Source/Runtime/Engine/Classes/PhysicsEngine/ConstraintInstance.h
@@ -325,7 +325,7 @@ struct ENGINE_API FConstraintInstance
 	FQuat AngularPositionTarget_DEPRECATED;
 	
 	/** The way rotation paths are estimated */
-	UPROPERTY(EditAnywhere, BlueprintReadOnly, Category = AngularMotor, meta = (editcondition = "bAngularVelocityDrive"))
+	UPROPERTY(EditAnywhere, BlueprintReadOnly, Category = AngularMotor)
 	TEnumAsByte<enum EAngularDriveMode::Type> AngularDriveMode;
 
 	/** Target orientation for the angular drive. */
@@ -341,7 +341,7 @@ struct ENGINE_API FConstraintInstance
 	float AngularDriveSpring;
 
 	/** Damping value to apply to the for angular drive. */
-	UPROPERTY(EditAnywhere, BlueprintReadOnly, Category = AngularMotor, meta = (DisplayName = "Angular Velocity Strength", editcondition = "bAngularVelocityDrive"))
+	UPROPERTY(EditAnywhere, BlueprintReadOnly, Category = AngularMotor, meta = (DisplayName = "Angular Velocity Strength"))
 	float AngularDriveDamping;
 
 	/** Limit to the force the angular drive can apply. */
