Mark ClientStopCameraShake and ClientStopForceFeedback reliable

diff --git a/Engine/Source/Runtime/Engine/Classes/GameFramework/PlayerController.h b/Engine/Source/Runtime/Engine/Classes/GameFramework/PlayerController.h
--- a/Engine/Source/Runtime/Engine/Classes/GameFramework/PlayerController.h
+++ b/Engine/Source/Runtime/Engine/Classes/GameFramework/PlayerController.h
@@ -801,7 +801,7 @@ public:
 	void ClientStopCameraAnim(class UCameraAnim* AnimToStop);
 
 	/** Stop camera shake on client.  */
-	UFUNCTION(unreliable, client, BlueprintCallable, Category="Game|Feedback")
+	UFUNCTION(reliable, client, BlueprintCallable, Category="Game|Feedback")
 	void ClientStopCameraShake(TSubclassOf<class UCameraShake> Shake);
 
 	/** 
@@ -818,7 +818,7 @@ public:
 	 * @param	ForceFeedbackEffect		If set only patterns from that effect will be stopped
 	 * @param	Tag						If not none only the pattern with this tag will be stopped
 	 */
-	UFUNCTION(unreliable, client, BlueprintCallable, Category="Game|Feedback")
+	UFUNCTION(reliable, client, BlueprintCallable, Category="Game|Feedback")
 	void ClientStopForceFeedback(class UForceFeedbackEffect* ForceFeedbackEffect, FName Tag);
 
 	/** 
