Minimize the possibility that replicated components array can get out of sync by forcing an update inside PostInitializeComponents

diff --git a/Engine/Source/Runtime/Engine/Classes/GameFramework/Actor.h b/Engine/Source/Runtime/Engine/Classes/GameFramework/Actor.h
--- a/Engine/Source/Runtime/Engine/Classes/GameFramework/Actor.h
+++ b/Engine/Source/Runtime/Engine/Classes/GameFramework/Actor.h
@@ -2191,6 +2191,10 @@ public:
 	 */
 	void UpdateReplicatedComponent(UActorComponent* Component);
 
+	/** Completely synchronizes the replicated components array so that it contains exactly the number of replicated components currently owned
+	 */
+	void UpdateAllReplicatedComponents();
+
 	/** Returns a constant reference to the replicated components array
 	 */
 	const TArray<UActorComponent*>& GetReplicatedComponents() const;
diff --git a/Engine/Source/Runtime/Engine/Private/Actor.cpp b/Engine/Source/Runtime/Engine/Private/Actor.cpp
--- a/Engine/Source/Runtime/Engine/Private/Actor.cpp
+++ b/Engine/Source/Runtime/Engine/Private/Actor.cpp
@@ -2048,6 +2048,19 @@ void AActor::UpdateReplicatedComponent(UActorComponent* Component)
 	}
 }
 
+void AActor::UpdateAllReplicatedComponents()
+{
+	ReplicatedComponents.Empty();
+
+	for (UActorComponent* Component : OwnedComponents)
+	{
+		if (Component != NULL)
+		{
+			UpdateReplicatedComponent(Component);
+		}
+	}
+}
+
 UActorComponent* AActor::FindComponentByClass(const TSubclassOf<UActorComponent> ComponentClass) const
 {
 	UActorComponent* FoundComponent = NULL;
@@ -3426,6 +3439,8 @@ void AActor::PostInitializeComponents()
 		bActorInitialized = true;
 
 		UNavigationSystem::OnActorRegistered(this);
+		
+		UpdateAllReplicatedComponents();
 	}
 }
 
