#ue4 - Add more comments to AActor::GetComponents()

diff --git a/Engine/Source/Runtime/Engine/Classes/GameFramework/Actor.h b/Engine/Source/Runtime/Engine/Classes/GameFramework/Actor.h
--- a/Engine/Source/Runtime/Engine/Classes/GameFramework/Actor.h
+++ b/Engine/Source/Runtime/Engine/Classes/GameFramework/Actor.h
@@ -2116,7 +2116,7 @@ public:
 	template<class T>
 	T* FindComponentByClass() const
 	{
-		static_assert(CanConvertPointerFromTo<T, UActorComponent>::Result, "'T' template parameter to FindComponentByClass must be derived from ActorComponent");
+		static_assert(CanConvertPointerFromTo<T, UActorComponent>::Result, "'T' template parameter to FindComponentByClass must be derived from UActorComponent");
 
 		return (T*)FindComponentByClass(T::StaticClass());
 	}
@@ -2133,8 +2133,7 @@ public:
 	template<class T, class AllocatorType>
 	void GetComponents(TArray<T*, AllocatorType>& OutComponents) const
 	{
-		static_assert(CanConvertPointerFromTo<T, UActorComponent>::Result, "'T' template parameter to GetComponents must be derived from ActorComponent");
-
+		static_assert(CanConvertPointerFromTo<T, UActorComponent>::Result, "'T' template parameter to GetComponents must be derived from UActorComponent");
 		SCOPE_CYCLE_COUNTER(STAT_GetComponentsTime);
 
 		OutComponents.Reset(OwnedComponents.Num());
@@ -2149,7 +2148,15 @@ public:
 		}
 	}
 
-	/** UActorComponent specialization of GetComponents() to avoid unnecessary casts. */
+	/**
+	 * UActorComponent specialization of GetComponents() to avoid unnecessary casts.
+	 * It's recommended to use TArrays with a TInlineAllocator to potentially avoid memory allocation costs.
+	 * TInlineComponentArray is defined to make this easier, for example:
+	 * {
+	 * 	   TInlineComponentArray<UPrimitiveComponent*> PrimComponents;
+	 *     Actor->GetComponents(PrimComponents);
+	 * }
+	 */
 	template<class AllocatorType>
 	void GetComponents(TArray<UActorComponent*, AllocatorType>& OutComponents) const
 	{
