Move ecosystem settings into foliagetype #codereview(Jack.Porter) #codereview(Gil.Gribb)

diff --git a/Engine/Source/Editor/FoliageEdit/FoliageEdit.build.cs b/Engine/Source/Editor/FoliageEdit/FoliageEdit.build.cs
--- a/Engine/Source/Editor/FoliageEdit/FoliageEdit.build.cs
+++ b/Engine/Source/Editor/FoliageEdit/FoliageEdit.build.cs
@@ -19,6 +19,7 @@ public class FoliageEdit : ModuleRules
 				"RenderCore",
 				"LevelEditor",
 				"Landscape",
+                "PropertyEditor",
 			}
 		);
 	}
diff --git a/Engine/Source/Editor/FoliageEdit/Private/FoliageEditModule.cpp b/Engine/Source/Editor/FoliageEdit/Private/FoliageEditModule.cpp
--- a/Engine/Source/Editor/FoliageEdit/Private/FoliageEditModule.cpp
+++ b/Engine/Source/Editor/FoliageEdit/Private/FoliageEditModule.cpp
@@ -2,12 +2,13 @@
 
 
 #include "UnrealEd.h"
-
 #include "FoliageEditModule.h"
 
 const FName FoliageEditAppIdentifier = FName(TEXT("FoliageEdApp"));
 
 #include "FoliageEdMode.h"
+#include "PropertyEditing.h"
+#include "FoliageTypeDetails.h"
 
 /**
  * Foliage Edit Mode module
@@ -27,6 +28,10 @@ public:
 			FSlateIcon(FEditorStyle::GetStyleSetName(), "LevelEditor.FoliageMode", "LevelEditor.FoliageMode.Small"),
 			true, 400
 			);
+
+		// Register the details customizer
+		FPropertyEditorModule& PropertyModule = FModuleManager::LoadModuleChecked<FPropertyEditorModule>("PropertyEditor");
+		PropertyModule.RegisterCustomClassLayout("FoliageType", FOnGetDetailCustomizationInstance::CreateStatic(&FFoliageTypeDetails::MakeInstance));
 	}
 
 	/**
@@ -35,6 +40,19 @@ public:
 	virtual void ShutdownModule() override
 	{
 		FEditorModeRegistry::Get().UnregisterMode(FBuiltinEditorModes::EM_Foliage);
+
+		if (!UObjectInitialized())
+		{
+			return;
+		}
+
+		// Unregister the details customization
+		if (FModuleManager::Get().IsModuleLoaded("PropertyEditor"))
+		{
+			FPropertyEditorModule& PropertyModule = FModuleManager::LoadModuleChecked<FPropertyEditorModule>("PropertyEditor");
+			PropertyModule.UnregisterCustomClassLayout("FoliageType");
+			PropertyModule.NotifyCustomizationModuleChanged();
+		}
 	}
 };
 
diff --git a/Engine/Source/Editor/FoliageEdit/Private/FoliageTypeDetails.cpp b/Engine/Source/Editor/FoliageEdit/Private/FoliageTypeDetails.cpp
new file mode 100644
--- /dev/null
+++ b/Engine/Source/Editor/FoliageEdit/Private/FoliageTypeDetails.cpp
@@ -0,0 +1,17 @@
+// Copyright 1998-2015 Epic Games, Inc. All Rights Reserved.
+
+#include "UnrealEd.h"
+#include "FoliageEditModule.h"
+#include "FoliageTypeDetails.h"
+
+TSharedRef<IDetailCustomization> FFoliageTypeDetails::MakeInstance()
+{
+	return MakeShareable(new FFoliageTypeDetails());
+
+}
+
+
+void FFoliageTypeDetails::CustomizeDetails(IDetailLayoutBuilder& DetailBuilder)
+{
+	
+}
\ No newline at end of file
diff --git a/Engine/Source/Editor/FoliageEdit/Private/FoliageTypeDetails.h b/Engine/Source/Editor/FoliageEdit/Private/FoliageTypeDetails.h
new file mode 100644
--- /dev/null
+++ b/Engine/Source/Editor/FoliageEdit/Private/FoliageTypeDetails.h
@@ -0,0 +1,16 @@
+// Copyright 1998-2015 Epic Games, Inc. All Rights Reserved.
+
+#pragma once
+
+class FFoliageTypeDetails : public IDetailCustomization
+{
+public:
+	virtual ~FFoliageTypeDetails(){};
+
+	/** Makes a new instance of this detail layout class for a specific detail view requesting it */
+	static TSharedRef<IDetailCustomization> MakeInstance();
+
+	/** IDetailCustomization interface */
+	virtual void CustomizeDetails( IDetailLayoutBuilder& DetailBuilder );
+
+};
diff --git a/Engine/Source/Runtime/Engine/Classes/Foliage/FoliageType.h b/Engine/Source/Runtime/Engine/Classes/Foliage/FoliageType.h
--- a/Engine/Source/Runtime/Engine/Classes/Foliage/FoliageType.h
+++ b/Engine/Source/Runtime/Engine/Classes/Foliage/FoliageType.h
@@ -15,7 +15,7 @@ enum FoliageVertexColorMask
 };
 
 
-UCLASS(hidecategories=Object, editinlinenew, MinimalAPI)
+UCLASS(hidecategories = Object, editinlinenew, MinimalAPI, BlueprintType, Blueprintable)
 class UFoliageType : public UObject
 {
 	GENERATED_UCLASS_BODY()
@@ -53,28 +53,28 @@ class UFoliageType : public UObject
 	UPROPERTY(EditAnywhere, Category=Painting)
 	uint32 LockScaleZ:1;
 
-	UPROPERTY(EditAnywhere, Category=Painting)
+	UPROPERTY(EditAnywhere, Category=General, meta=(Subcategory = "Placement"))
 	float AlignMaxAngle;
 
-	UPROPERTY(EditAnywhere, Category=Painting)
+	UPROPERTY(EditAnywhere, Category = General, meta = (Subcategory = "Placement"))
 	float RandomPitchAngle;
 
-	UPROPERTY(EditAnywhere, Category=Painting)
+	UPROPERTY(EditAnywhere, Category = General, meta = (Subcategory = "Placement"))
 	float GroundSlope;
 
-	UPROPERTY(EditAnywhere, Category=Painting)
+	UPROPERTY(EditAnywhere, Category = General, meta = (Subcategory = "Placement"))
 	float HeightMin;
 
-	UPROPERTY(EditAnywhere, Category=Painting)
+	UPROPERTY(EditAnywhere, Category = General, meta = (Subcategory = "Placement"))
 	float HeightMax;
 
-	UPROPERTY(EditAnywhere, Category=Painting)
+	UPROPERTY(EditAnywhere, Category = General, meta = (Subcategory = "Placement"))
 	FName LandscapeLayer;
 
-	UPROPERTY(EditAnywhere, Category=Painting)
+	UPROPERTY(EditAnywhere, Category = General, meta = (Subcategory = "Placement"))
 	uint32 AlignToNormal:1;
 
-	UPROPERTY(EditAnywhere, Category=Painting)
+	UPROPERTY(EditAnywhere, Category = General, meta = (Subcategory = "Placement"))
 	uint32 RandomYaw:1;
 
 	UPROPERTY(EditAnywhere, Category=Painting)
@@ -154,14 +154,14 @@ class UFoliageType : public UObject
 	uint32 ReapplyVertexColorMask:1;
 
 	/* The distance where instances will begin to fade out if using a PerInstanceFadeAmount material node. 0 disables. */
-	UPROPERTY(EditAnywhere, Category=Culling)
+	UPROPERTY(EditAnywhere, Category = General, meta = (Subcategory = "Culling"))
 	int32 StartCullDistance;
 
 	/**
 	 * The distance where instances will have completely faded out when using a PerInstanceFadeAmount material node. 0 disables. 
 	 * When the entire cluster is beyond this distance, the cluster is completely culled and not rendered at all.
 	 */
-	UPROPERTY(EditAnywhere, Category = Culling)
+	UPROPERTY(EditAnywhere, Category = General, meta = (Subcategory = "Culling"))
 	int32 EndCullDistance;
 
 	UPROPERTY()
@@ -180,43 +180,43 @@ class UFoliageType : public UObject
 	uint32 ShowInstanceSettings:1;
 
 	/** Controls whether the foliage should cast a shadow or not. **/
-	UPROPERTY(EditAnywhere, BlueprintReadOnly, Category=Lighting)
+	UPROPERTY(EditAnywhere, BlueprintReadOnly, Category = General, meta = (Subcategory = "Lighting"))
 	uint32 CastShadow:1;
 
 	/** Controls whether the foliage should inject light into the Light Propagation Volume.  This flag is only used if CastShadow is true. **/
-	UPROPERTY(EditAnywhere, BlueprintReadOnly, Category=Lighting)
+	UPROPERTY(EditAnywhere, BlueprintReadOnly, Category = General, meta = (Subcategory = "Lighting"))
 	uint32 bAffectDynamicIndirectLighting:1;
 
 	/** Controls whether the primitive should affect dynamic distance field lighting methods.  This flag is only used if CastShadow is true. **/
-	UPROPERTY(EditAnywhere, BlueprintReadOnly, Category=Lighting, meta=(EditCondition="CastShadow"))
+	UPROPERTY(EditAnywhere, BlueprintReadOnly, Category = General, meta = (Subcategory = "Lighting", EditCondition = "CastShadow"))
 	uint32 bAffectDistanceFieldLighting:1;
 
 	/** Controls whether the foliage should cast shadows in the case of non precomputed shadowing.  This flag is only used if CastShadow is true. **/
-	UPROPERTY(EditAnywhere, BlueprintReadOnly, Category=Lighting)
+	UPROPERTY(EditAnywhere, BlueprintReadOnly, Category = General, meta = (Subcategory = "Lighting"))
 	uint32 bCastDynamicShadow:1;
 
 	/** Whether the foliage should cast a static shadow from shadow casting lights.  This flag is only used if CastShadow is true. */
-	UPROPERTY(EditAnywhere, BlueprintReadOnly, Category=Lighting)
+	UPROPERTY(EditAnywhere, BlueprintReadOnly, Category = General, meta = (Subcategory = "Lighting"))
 	uint32 bCastStaticShadow:1;
 
 	/** Whether this foliage should cast dynamic shadows as if it were a two sided material. */
-	UPROPERTY(EditAnywhere, AdvancedDisplay, BlueprintReadOnly, Category=Lighting)
+	UPROPERTY(EditAnywhere, AdvancedDisplay, BlueprintReadOnly, Category = General, meta = (Subcategory = "Lighting"))
 	uint32 bCastShadowAsTwoSided:1;
 
 	/** Whether the foliage receives decals. */
-	UPROPERTY(EditAnywhere, AdvancedDisplay, BlueprintReadOnly, Category=Rendering)
+	UPROPERTY(EditAnywhere, AdvancedDisplay, BlueprintReadOnly, Category = General, meta = (Subcategory = "Rendering"))
 	uint32 bReceivesDecals : 1;
 	
 	/** Whether to override the lightmap resolution defined in the static mesh. */
-	UPROPERTY(BlueprintReadOnly, Category=Lighting)
+	UPROPERTY(BlueprintReadOnly, Category = General, meta = (Subcategory = "Lighting"))
 	uint32 bOverrideLightMapRes:1;
 
 	/** Overrides the lightmap resolution defined in the static mesh. A value of 0 disables static lighting/shadowing */
-	UPROPERTY(EditAnywhere, BlueprintReadOnly, Category=Lighting, meta=(DisplayName="Light Map Resolution", EditCondition="bOverrideLightMapRes"))
+	UPROPERTY(EditAnywhere, BlueprintReadOnly, Category = General, meta = (Subcategory = "Lighting", DisplayName = "Light Map Resolution", EditCondition = "bOverrideLightMapRes"))
 	int32 OverriddenLightMapRes;
 
 	/** Custom collision for foliage */
-	UPROPERTY(EditAnywhere, Category=Collision, meta=(HideObjectType=true))
+	UPROPERTY(EditAnywhere, Category = Painting, meta = (Subcategory = "Collision", HideObjectType = true))
 	struct FBodyInstance BodyInstance;
 
 	/* Gets the mesh associated with this FoliageType */
@@ -225,6 +225,81 @@ class UFoliageType : public UObject
 #if WITH_EDITOR
 	virtual void PostEditChangeProperty(struct FPropertyChangedEvent& PropertyChangedEvent) override;
 #endif
+
+	/** Ecosystem specific parameters */
+
+	/** Specifies the number of seeds to populate along 10 meters. The number is implicitly squared to cover a 10m x 10m area*/
+	UPROPERTY(Category = Ecosystem, EditAnywhere, meta = (Subcategory = "Clustering", ClampMin = "0.0", UIMin = "0.0"))
+	float InitialSeedDensity;
+
+	/** The average distance between the spreading instance and its seeds. For example, a tree with an AverageSpreadDistance 10 will ensure the average distance between the tree and its seeds is 10cm */
+	UPROPERTY(Category = Ecosystem, EditAnywhere, meta = (Subcategory = "Clustering", ClampMin = "0.0", UIMin = "0.0"))
+	float AverageSpreadDistance;
+
+	/** Specifies how much seed distance varies from the average. For example, a tree with an AverageSpreadDistance 10 and a SpreadVariance 1 will produce seeds with an average distance of 10cm plus or minus 1cm */
+	UPROPERTY(Category = Ecosystem, EditAnywhere, meta = (Subcategory = "Clustering", ClampMin = "0.0", UIMin = "0.0"))
+	float SpreadVariance;
+
+	/** The CollisionRadius determines when two instances overlap. When two instances overlap a winner will be picked based on rules and priority. */
+	UPROPERTY(Category = Ecosystem, EditAnywhere, meta = (Subcategory = "Collision", ClampMin = "0.0", UIMin = "0.0"))
+	float CollisionRadius;
+
+	/** The ShadeRadius determines when two instances overlap. If an instance can grow in the shade this radius is ignored.*/
+	UPROPERTY(Category = Ecosystem, EditAnywhere, meta = (Subcategory = "Collision", ClampMin = "0.0", UIMin = "0.0"))
+	float ShadeRadius;
+
+	/** The number of times we age the species and spread its seeds. */
+	UPROPERTY(Category = Ecosystem, EditAnywhere, meta = (Subcategory = "Clustering", ClampMin = "0", UIMin = "0"))
+	int32 NumSteps;
+
+	/** The number of seeds an instance will spread in a single step of the simulation. */
+	UPROPERTY(Category = Ecosystem, EditAnywhere, meta = (Subcategory = "Clustering", ClampMin = "0", UIMin = "0"))
+	int32 SeedsPerStep;
+
+	/** Whether the species can grow in shade. If this is true shade radius is ignored during overlap tests*/
+	UPROPERTY(Category = Ecosystem, EditAnywhere, meta = (Subcategory = "Growth"))
+	bool bGrowsInShade;
+
+	/** The minimum scale that an instance will be. Corresponds to age = 0*/
+	UPROPERTY(Category = Ecosystem, EditAnywhere, meta = (Subcategory = "Growth", ClampMin = "0.0", UIMin = "0.0"))
+	float MinScale;
+
+	/** The maximum scale that a seed will grow to. Corresponds to age = 1*/
+	UPROPERTY(Category = Ecosystem, EditAnywhere, meta = (Subcategory = "Clustering", ClampMin = "0.0", UIMin = "0.0"))
+	float MaxScale;
+
+	/** Specifies the oldest a new seed can be. The new seed will have an age randomly distributed in [0,InitMaxAge] */
+	UPROPERTY(Category = Ecosystem, EditAnywhere, meta = (Subcategory = "Growth", ClampMin = "0.0", UIMin = "0.0"))
+	float InitialMaxAge;
+
+	/** Specifies the oldest a seed can be. After reaching this age the instance will still spread seeds, but will not get any older*/
+	UPROPERTY(Category = Ecosystem, EditAnywhere, meta = (Subcategory = "Growth", ClampMin = "0.0", UIMin = "0.0"))
+	float MaxAge;
+
+	/** When two instances overlap we must determine which instance to remove. The instance with a lower OverlapPriority will be removed. In the case where OverlapPriority is the same regular simulation rules apply.*/
+	UPROPERTY(Category = Ecosystem, EditAnywhere, meta = (Subcategory = "Growth"))
+	float OverlapPriority;
+
+	/** The set of meshes to be used for representing an instance placed in the world.*/
+	UPROPERTY(Category = Ecosystem, EditAnywhere, meta = (Subcategory = "Rendering"))
+	TArray<UStaticMesh*> StaticMeshes;
+
+
+	UPROPERTY()
+	int32 ChangeCount;
+
+	
+	ENGINE_API float GetSeedDensitySquared() const { return InitialSeedDensity * InitialSeedDensity; }
+	ENGINE_API float GetMaxRadius() const;
+	ENGINE_API float GetScaleForAge(const float Age) const;
+	ENGINE_API float GetInitAge(FRandomStream& RandomStream) const;
+	ENGINE_API float GetNextAge(const float CurrentAge, const int32 NumSteps) const;
+
+private:
+
+	/** The curve used to interpolate the instance scale.*/
+	UPROPERTY(Category = Ecosystem, EditAnywhere, meta = (Subcategory = "Rendering"))
+	FRuntimeFloatCurve ScaleCurve;
 };
 
 
diff --git a/Engine/Source/Runtime/Engine/Private/InstancedFoliage.cpp b/Engine/Source/Runtime/Engine/Private/InstancedFoliage.cpp
--- a/Engine/Source/Runtime/Engine/Private/InstancedFoliage.cpp
+++ b/Engine/Source/Runtime/Engine/Private/InstancedFoliage.cpp
@@ -205,14 +205,73 @@ UFoliageType::UFoliageType(const FObjectInitializer& ObjectInitializer)
 	OverriddenLightMapRes = 32;
 
 	BodyInstance.SetCollisionProfileName(UCollisionProfile::NoCollision_ProfileName);
+
+	/** Ecosystem settings*/
+	AverageSpreadDistance = 50;
+	SpreadVariance = 150;
+	bGrowsInShade = false;
+	SeedsPerStep = 3;
+	OverlapPriority = 0.f;
+	NumSteps = 3;
+	MinScale = 1.f;
+	MaxScale = 3.f;
+	ChangeCount = 0;
+	InitialSeedDensity = 1.f;
+	CollisionRadius = 100.f;
+	ShadeRadius = 100.f;
+	InitialMaxAge = 0.f;
+	MaxAge = 10.f;
+
+	FRichCurve* Curve = ScaleCurve.GetRichCurve();
+	Curve->AddKey(0.f, 0.f);
+	Curve->AddKey(1.f, 1.f);
 }
 
+
 UFoliageType_InstancedStaticMesh::UFoliageType_InstancedStaticMesh(const FObjectInitializer& ObjectInitializer)
 	: Super(ObjectInitializer)
 {
 	Mesh = nullptr;
 }
 
+
+float UFoliageType::GetMaxRadius() const
+{
+	return FMath::Max(CollisionRadius, ShadeRadius);
+}
+
+float UFoliageType::GetScaleForAge(const float Age) const
+{
+	const FRichCurve* Curve = ScaleCurve.GetRichCurveConst();
+	const float Time = FMath::Clamp(MaxAge == 0 ? 1.f : Age / MaxAge, 0.f, 1.f);
+	const float Scale = Curve->Eval(Time);
+	return MinScale + (MaxScale - MinScale) * Scale;
+}
+
+float UFoliageType::GetInitAge(FRandomStream& RandomStream) const
+{
+	return RandomStream.FRandRange(0, InitialMaxAge);
+}
+
+float UFoliageType::GetNextAge(const float CurrentAge, const int32 NumSteps) const
+{
+	float NewAge = CurrentAge;
+	for (int32 Count = 0; Count < NumSteps; ++Count)
+	{
+		const float GrowAge = NewAge + 1;
+		if (GrowAge <= MaxAge)
+		{
+			NewAge = GrowAge;
+		}
+		else
+		{
+			break;
+		}
+	}
+
+	return NewAge;
+}
+
 #if WITH_EDITOR
 void UFoliageType::PostEditChangeProperty(struct FPropertyChangedEvent& PropertyChangedEvent)
 {
@@ -220,6 +279,7 @@ void UFoliageType::PostEditChangeProperty(struct FPropertyChangedEvent& Property
 
 	// Ensure that OverriddenLightMapRes is a factor of 4
 	OverriddenLightMapRes = FMath::Max(OverriddenLightMapRes + 3 & ~3, 4);
+	++ChangeCount;
 }
 #endif
 
