#ue4 - PhysX locking improvements to scene queries to prevent potentially accessing scene data after the scene has been unlocked.

diff --git a/Engine/Source/Runtime/Engine/Private/Collision/PhysXCollision.cpp b/Engine/Source/Runtime/Engine/Private/Collision/PhysXCollision.cpp
--- a/Engine/Source/Runtime/Engine/Private/Collision/PhysXCollision.cpp
+++ b/Engine/Source/Runtime/Engine/Private/Collision/PhysXCollision.cpp
@@ -478,6 +478,53 @@ private:
 
 static float DebugLineLifetime = 2.f;
 
+/**
+ * Helper to lock/unlock multiple scenes that also makes sure to unlock everything when it goes out of scope.
+ * Multiple locks are NOT SAFE. You can't call LockRead() if already locked.
+ * Multiple unlocks are safe (repeated unlocks do nothing after the first successful unlock.).
+ */
+struct FScopedMultiSceneReadLock
+{
+	FScopedMultiSceneReadLock()
+	{
+		for (int32 i=0; i < ARRAY_COUNT(SceneLocks); i++)
+		{
+			SceneLocks[i] = nullptr;
+		}
+	}
+
+	~FScopedMultiSceneReadLock()
+	{
+		UnlockAll();
+	}
+
+	inline void LockRead(PxScene* Scene, int32 SceneIndex)
+	{
+		check(SceneLocks[SceneIndex] == nullptr); // no nested locks allowed.
+		SCENE_LOCK_READ(Scene);
+		SceneLocks[SceneIndex] = Scene;
+	}
+
+	inline void UnlockRead(PxScene* Scene, int32 SceneIndex)
+	{
+		check(SceneLocks[SceneIndex] == Scene || SceneLocks[SceneIndex] == nullptr);
+		SCENE_UNLOCK_READ(Scene);
+		SceneLocks[SceneIndex] = nullptr;
+	}
+
+	inline void UnlockAll()
+	{
+		for (int32 i=0; i < ARRAY_COUNT(SceneLocks); i++)
+		{
+			SCENE_UNLOCK_READ(SceneLocks[i]);
+			SceneLocks[i] = nullptr;
+		}
+	}
+
+	PxScene* SceneLocks[PST_MAX];
+};
+
+
 /** 
  * Type of query for object type or trace type
  * Trace queries correspond to trace functions with TravelChannel/ResponseParams
@@ -668,7 +715,7 @@ bool bSkipCapture = false;
 /** Util to extract type and dimensions from physx geom being swept, and pass info to CollisionAnalyzer, if its recording */
 void CaptureGeomSweep(const UWorld* World, const FVector& Start, const FVector& End, const PxQuat& PRot, ECAQueryType::Type QueryType, const PxGeometry& PGeom, ECollisionChannel TraceChannel, const struct FCollisionQueryParams& Params, const struct FCollisionResponseParams& ResponseParams, const struct FCollisionObjectQueryParams& ObjectParams, const TArray<FHitResult>& Results, double CPUTime)
 {
-	if(bSkipCapture || !IsInGameThread() || !GCollisionAnalyzerIsRecording)
+	if(bSkipCapture || !GCollisionAnalyzerIsRecording || !IsInGameThread())
 	{
 		return;
 	}
@@ -728,7 +775,7 @@ void CaptureGeomSweep(const UWorld* World, const FVector& Start, const FVector&
 /** Util to capture a raycast with the CollisionAnalyzer if recording */
 void CaptureRaycast(const UWorld* World, const FVector& Start, const FVector& End, ECAQueryType::Type QueryType, ECollisionChannel TraceChannel, const struct FCollisionQueryParams& Params, const struct FCollisionResponseParams& ResponseParams, const struct FCollisionObjectQueryParams& ObjectParams, const TArray<FHitResult>& Results, double CPUTime)
 {
-	if(bSkipCapture || !IsInGameThread() || !GCollisionAnalyzerIsRecording)
+	if(bSkipCapture || !GCollisionAnalyzerIsRecording || !IsInGameThread())
 	{
 		return;
 	}
@@ -743,8 +790,8 @@ void CaptureRaycast(const UWorld* World, const FVector& Start, const FVector& En
 }
 
 #define STARTQUERYTIMER() double StartTime = FPlatformTime::Seconds()
-#define CAPTUREGEOMSWEEP(World, Start, End, Rot, QueryType, PGeom, TraceChannel, Params, ResponseParam, ObjectParam, Results)	do { CaptureGeomSweep(World, Start, End, Rot, QueryType, PGeom, TraceChannel, Params, ResponseParam, ObjectParam, Results, FPlatformTime::Seconds() - StartTime); } while(0)
-#define CAPTURERAYCAST(World, Start, End, QueryType, TraceChannel, Params, ResponseParam, ObjectParam, Results)	do { CaptureRaycast(World, Start, End, QueryType, TraceChannel, Params, ResponseParam, ObjectParam, Results, FPlatformTime::Seconds() - StartTime); } while(0)
+#define CAPTUREGEOMSWEEP(World, Start, End, Rot, QueryType, PGeom, TraceChannel, Params, ResponseParam, ObjectParam, Results) if (GCollisionAnalyzerIsRecording) { CaptureGeomSweep(World, Start, End, Rot, QueryType, PGeom, TraceChannel, Params, ResponseParam, ObjectParam, Results, FPlatformTime::Seconds() - StartTime); }
+#define CAPTURERAYCAST(World, Start, End, QueryType, TraceChannel, Params, ResponseParam, ObjectParam, Results)	if (GCollisionAnalyzerIsRecording) { CaptureRaycast(World, Start, End, QueryType, TraceChannel, Params, ResponseParam, ObjectParam, Results, FPlatformTime::Seconds() - StartTime); }
 
 #else
 
@@ -777,6 +824,9 @@ bool RaycastTest(const UWorld* World, const FVector Start, const FVector End, EC
 	{
 #if WITH_PHYSX
 		{
+			const PxVec3 PDir = U2PVector(Delta / DeltaMag);
+			PxRaycastHit PHit;
+
 			// Create filter data used to filter collisions
 			PxFilterData PFilter = CreateQueryFilterData(TraceChannel, Params.bTraceComplex, ResponseParams.CollisionResponse, ObjectParams, false);
 			PxSceneQueryFilterData PQueryFilterData(PFilter, PxSceneQueryFilterFlag::eSTATIC | PxSceneQueryFilterFlag::eDYNAMIC | PxSceneQueryFilterFlag::ePREFILTER);
@@ -785,14 +835,12 @@ bool RaycastTest(const UWorld* World, const FVector Start, const FVector End, EC
 			PQueryCallback.bSingleQuery = true;
 
 			FPhysScene* PhysScene = World->GetPhysicsScene();
-
-			// Enable scene locks, in case they are required
-			PxScene* SyncScene = PhysScene->GetPhysXScene(PST_Sync);
-			SCOPED_SCENE_READ_LOCK(SyncScene);
-
-			const PxVec3 PDir = U2PVector(Delta / DeltaMag);
-			PxRaycastHit PHit;
-			bHaveBlockingHit = SyncScene->raycastSingle(U2PVector(Start), PDir, DeltaMag, POutputFlags, PHit, PQueryFilterData, &PQueryCallback);
+			{
+				// Enable scene locks, in case they are required
+				PxScene* SyncScene = PhysScene->GetPhysXScene(PST_Sync);
+				SCOPED_SCENE_READ_LOCK(SyncScene);
+				bHaveBlockingHit = SyncScene->raycastSingle(U2PVector(Start), PDir, DeltaMag, POutputFlags, PHit, PQueryFilterData, &PQueryCallback);
+			}
 
 			// Test async scene if we have no blocking hit, and async tests are requested
 			if (!bHaveBlockingHit && Params.bTraceAsyncScene && PhysScene->HasAsyncScene())
@@ -849,6 +897,8 @@ bool RaycastSingle(const UWorld* World, struct FHitResult& OutHit, const FVector
 	{
 #if WITH_PHYSX
 		{
+			FScopedMultiSceneReadLock SceneLocks;
+
 			// Create filter data used to filter collisions
 			PxFilterData PFilter = CreateQueryFilterData(TraceChannel, Params.bTraceComplex, ResponseParams.CollisionResponse, ObjectParams, false);
 			PxSceneQueryFilterData PQueryFilterData(PFilter, PxSceneQueryFilterFlag::eSTATIC | PxSceneQueryFilterFlag::eDYNAMIC | PxSceneQueryFilterFlag::ePREFILTER);
@@ -863,18 +913,23 @@ bool RaycastSingle(const UWorld* World, struct FHitResult& OutHit, const FVector
 			PxScene* SyncScene = PhysScene->GetPhysXScene(PST_Sync);
 
 			// Enable scene locks, in case they are required
-			SCOPED_SCENE_READ_LOCK(SyncScene);
+			SceneLocks.LockRead(SyncScene, PST_Sync);
+
 			PxRaycastHit PHit;
 			bHaveBlockingHit = SyncScene->raycastSingle(U2PVector(Start), PDir, DeltaMag, POutputFlags, PHit, PQueryFilterData, &PQueryCallback);
+			if (!bHaveBlockingHit)
+			{
+				// Not going to use anything from this scene, so unlock it now.
+				SceneLocks.UnlockRead(SyncScene, PST_Sync);
+			}
 
 			// Test async scene if async tests are requested
 			if (Params.bTraceAsyncScene && PhysScene->HasAsyncScene())
 			{
 				PxScene* AsyncScene = PhysScene->GetPhysXScene(PST_Async);
-				SCOPED_SCENE_READ_LOCK(AsyncScene);
-				bool bHaveBlockingHitAsync;
+				SceneLocks.LockRead(AsyncScene, PST_Async);
 				PxRaycastHit PHitAsync;
-				bHaveBlockingHitAsync = AsyncScene->raycastSingle(U2PVector(Start), PDir, DeltaMag, POutputFlags, PHitAsync, PQueryFilterData, &PQueryCallback);
+				const bool bHaveBlockingHitAsync = AsyncScene->raycastSingle(U2PVector(Start), PDir, DeltaMag, POutputFlags, PHitAsync, PQueryFilterData, &PQueryCallback);
 
 				// If we have a blocking hit in the async scene and there was no sync blocking hit, or if the async blocking hit came first,
 				// then this becomes the blocking hit.  We can test the PxSceneQueryImpactHit::distance since the DeltaMag is the same for both queries.
@@ -883,9 +938,14 @@ bool RaycastSingle(const UWorld* World, struct FHitResult& OutHit, const FVector
 					PHit = PHitAsync;
 					bHaveBlockingHit = true;
 				}
+				else
+				{
+					// Not going to use anything from this scene, so unlock it now.
+					SceneLocks.UnlockRead(AsyncScene, PST_Async);
+				}
 			}
 
-			if (bHaveBlockingHit) // If we got a hit, 
+			if (bHaveBlockingHit) // If we got a hit
 			{
 				PxTransform PStartTM(U2PVector(Start));
 				ConvertQueryImpactHit(PHit, OutHit, DeltaMag, PFilter, Start, End, NULL, PStartTM, Params.bReturnFaceIndex, Params.bReturnPhysicalMaterial);
@@ -970,7 +1030,8 @@ bool RaycastMulti(const UWorld* World, TArray<struct FHitResult>& OutHits, const
 		FPhysScene* PhysScene = World->GetPhysicsScene();
 		PxScene* SyncScene = PhysScene->GetPhysXScene(PST_Sync);
 
-		SCOPED_SCENE_READ_LOCK(SyncScene);
+		FScopedMultiSceneReadLock SceneLocks;
+		SceneLocks.LockRead(SyncScene, PST_Sync);
 
 		PxI32 NumHits = SyncScene->raycastMultiple(U2PVector(Start), PDir, DeltaMag, POutputFlags, PHits, HIT_BUFFER_MAX_SYNC_QUERIES, bBlockingHit, PQueryFilterData, &PQueryCallback);
 
@@ -983,12 +1044,17 @@ bool RaycastMulti(const UWorld* World, TArray<struct FHitResult>& OutHits, const
 			// we can use, though that is not the entire set
 			NumHits = HIT_BUFFER_MAX_SYNC_QUERIES;
 		}
+		else if (NumHits == 0)
+		{
+			// Not going to use anything from this scene, so unlock it now.
+			SceneLocks.UnlockRead(SyncScene, PST_Sync);
+		}
 
 		// Test async scene if async tests are requested and there was no overflow
 		if( Params.bTraceAsyncScene && PhysScene->HasAsyncScene())
 		{
 			PxScene* AsyncScene = PhysScene->GetPhysXScene(PST_Async);
-			SCOPED_SCENE_READ_LOCK(AsyncScene);
+			SceneLocks.LockRead(AsyncScene, PST_Async);
 
 			// Write into the same PHits buffer
 			PxRaycastHit* PHitsAsync = PHits+NumHits;
@@ -1014,12 +1080,17 @@ bool RaycastMulti(const UWorld* World, TArray<struct FHitResult>& OutHits, const
 				// we can use, though that is not the entire set
 				NumAsyncHits = AsyncBufferSize;
 			}
+			else if (NumAsyncHits == 0)
+			{
+				// Not going to use anything from this scene, so unlock it now.
+				SceneLocks.UnlockRead(AsyncScene, PST_Async);
+			}
 
 			PxI32 TotalNumHits = NumHits + NumAsyncHits;
 
 			// If there is a blocking hit in the sync scene, and it is closer than the blocking hit in the async scene (or there is no blocking hit in the async scene),
 			// then move it to the end of the array to get it out of the way.
-			if (bBlockingHit && NumAsyncHits > 0)
+			if (bBlockingHit)
 			{
 				if (!bBlockingHitAsync || PHits[NumHits-1].distance < PHits[TotalNumHits-1].distance)
 				{
@@ -1125,13 +1196,13 @@ bool GeomSweepTest(const UWorld* World, const struct FCollisionShape& CollisionS
 		PxVec3 PDir = U2PVector(Delta/DeltaMag);
 
 		FPhysScene* PhysScene = World->GetPhysicsScene();
-		PxScene* SyncScene = PhysScene->GetPhysXScene(PST_Sync);
-
-		// Enable scene locks, in case they are required
-		SCOPED_SCENE_READ_LOCK(SyncScene);
-
-		PxSweepHit PQueryHit;
-		bHaveBlockingHit = SyncScene->sweepSingle(PGeom, PStartTM, PDir, DeltaMag, POutputFlags, PQueryHit, PQueryFilterData, &PQueryCallbackSweep);
+		{
+			// Enable scene locks, in case they are required
+			PxScene* SyncScene = PhysScene->GetPhysXScene(PST_Sync);
+			SCOPED_SCENE_READ_LOCK(SyncScene);
+			PxSweepHit PQueryHit;
+			bHaveBlockingHit = SyncScene->sweepSingle(PGeom, PStartTM, PDir, DeltaMag, POutputFlags, PQueryHit, PQueryFilterData, &PQueryCallbackSweep);
+		}
 
 		// Test async scene if async tests are requested and there was no blocking hit was found in the sync scene (since no hit info other than a boolean yes/no is recorded)
 		if( !bHaveBlockingHit && Params.bTraceAsyncScene && PhysScene->HasAsyncScene())
@@ -1200,7 +1271,8 @@ bool GeomSweepSingle(const UWorld* World, const struct FCollisionShape& Collisio
 		PxScene* SyncScene = PhysScene->GetPhysXScene(PST_Sync);
 
 		// Enable scene locks, in case they are required
-		SCOPED_SCENE_READ_LOCK(SyncScene);
+		FScopedMultiSceneReadLock SceneLocks;
+		SceneLocks.LockRead(SyncScene, PST_Sync);
 
 		PxSweepHit PHit;
 		bHaveBlockingHit = SyncScene->sweepSingle(PGeom, PStartTM, PDir, DeltaMag, POutputFlags, PHit, PQueryFilterData, &PQueryCallbackSweep);
@@ -1225,11 +1297,17 @@ bool GeomSweepSingle(const UWorld* World, const struct FCollisionShape& Collisio
 		}
 #endif
 
+		if (!bHaveBlockingHit)
+		{
+			// Not using anything from this scene, so unlock it.
+			SceneLocks.UnlockRead(SyncScene, PST_Sync);
+		}
+
 		// Test async scene if async tests are requested
 		if( Params.bTraceAsyncScene && PhysScene->HasAsyncScene())
 		{
 			PxScene* AsyncScene = PhysScene->GetPhysXScene(PST_Async);
-			SCOPED_SCENE_READ_LOCK(AsyncScene);
+			SceneLocks.LockRead(AsyncScene, PST_Async);
 
 			bool bHaveBlockingHitAsync;
 			PxSweepHit PHitAsync;
@@ -1242,6 +1320,11 @@ bool GeomSweepSingle(const UWorld* World, const struct FCollisionShape& Collisio
 				PHit = PHitAsync;
 				bHaveBlockingHit = true;
 			}
+			else
+			{
+				// Not using anything from this scene, so unlock it.
+				SceneLocks.UnlockRead(AsyncScene, PST_Async);
+			}
 		}
 
 		if(bHaveBlockingHit) // If we got a hit, convert it to unreal type
@@ -1250,7 +1333,6 @@ bool GeomSweepSingle(const UWorld* World, const struct FCollisionShape& Collisio
 		}
 	}
 
-
 #if !(UE_BUILD_SHIPPING || UE_BUILD_TEST)
 	if((World->DebugDrawTraceTag != NAME_None) && (World->DebugDrawTraceTag == Params.TraceTag))
 	{
@@ -1302,7 +1384,9 @@ bool GeomSweepMulti_PhysX(const UWorld* World, const PxGeometry& PGeom, const Px
 		FPhysScene* PhysScene = World->GetPhysicsScene();
 		PxScene* SyncScene = PhysScene->GetPhysXScene(PST_Sync);
 
-		SCOPED_SCENE_READ_LOCK(SyncScene);
+		// Lock scene
+		FScopedMultiSceneReadLock SceneLocks;
+		SceneLocks.LockRead(SyncScene, PST_Sync);
 
 		const PxTransform PStartTM(U2PVector(Start), PGeomRot);
 		const PxVec3 PDir = U2PVector(Delta / DeltaMag);
@@ -1371,11 +1455,17 @@ bool GeomSweepMulti_PhysX(const UWorld* World, const PxGeometry& PGeom, const Px
 		}
 #endif
 
+		if (NumHits == 0)
+		{
+			// Not using anything from this scene, so unlock it.
+			SceneLocks.UnlockRead(SyncScene, PST_Sync);
+		}
+
 		// Test async scene if async tests are requested and there was no overflow
 		if (Params.bTraceAsyncScene && MinBlockDistance > SMALL_NUMBER && PhysScene->HasAsyncScene())
 		{
 			PxScene* AsyncScene = PhysScene->GetPhysXScene(PST_Async);
-			SCOPED_SCENE_READ_LOCK(AsyncScene);
+			SceneLocks.LockRead(AsyncScene, PST_Async);
 
 			// Write into the same PHits buffer
 			PxSweepHit* PHitsAsync = PHits + NumHits;
@@ -1391,6 +1481,11 @@ bool GeomSweepMulti_PhysX(const UWorld* World, const PxGeometry& PGeom, const Px
 				// sweepMultiple still fills in the buffer with the blocking hit and an arbitrary subset of nearer touches
 				NumAsyncHits = AsyncBufferSize;
 			}
+			else if (NumAsyncHits == 0)
+			{
+				// Not using anything from this scene, so unlock it.
+				SceneLocks.UnlockRead(AsyncScene, PST_Async);
+			}
 
 			NumHits += NumAsyncHits;
 
@@ -1475,22 +1570,21 @@ bool GeomOverlapTest(const UWorld* World, const struct FCollisionShape& Collisio
 	bool bHit = false;
 
 	FPhysScene* PhysScene = World->GetPhysicsScene();
-	PxScene* SyncScene = PhysScene->GetPhysXScene(PST_Sync);
-
 	{
 		// Enable scene locks, in case they are required
+		PxScene* SyncScene = PhysScene->GetPhysXScene(PST_Sync);
 		SCOPED_SCENE_READ_LOCK(SyncScene);
 
 		bHit = SyncScene->overlapAny(PGeom, PGeomPose, POverlapResult, PQueryFilterData, &PQueryCallback);
-	}
 
-	// if we got a hit, need to see if it was a touch or a block
-	// @todo UE4 james remove this once overlapAny only returns blocking hits
-	if(bHit)
-	{
-		FOverlapResult NewOverlap;
-		ConvertQueryOverlap(POverlapResult.shape, POverlapResult.actor, NewOverlap, PFilter);
-		bHaveBlockingHit = NewOverlap.bBlockingHit;
+		// if we got a hit, need to see if it was a touch or a block
+		// @todo UE4 james remove this once overlapAny only returns blocking hits
+		if (bHit)
+		{
+			FOverlapResult NewOverlap;
+			ConvertQueryOverlap(POverlapResult.shape, POverlapResult.actor, NewOverlap, PFilter);
+			bHaveBlockingHit = NewOverlap.bBlockingHit;
+		}
 	}
 
 	// Test async scene if async tests are requested and there was no blocking hit was found in the sync scene (since no hit info other than a boolean yes/no is recorded)
@@ -1534,6 +1628,7 @@ bool GeomOverlapSingle(const UWorld* World, const struct FCollisionShape& Collis
 
 	// Track if we get any 'blocking' hits
 	bool bHaveBlockingHit = false;
+	bool bHit = false;
 
 #if WITH_PHYSX
 	FPhysXShapeAdaptor ShapeAdaptor(Rot, CollisionShape);
@@ -1547,19 +1642,20 @@ bool GeomOverlapSingle(const UWorld* World, const struct FCollisionShape& Collis
 
 	// Enable scene locks, in case they are required
 	FPhysScene* PhysScene = World->GetPhysicsScene();
-	PxScene* SyncScene = PhysScene->GetPhysXScene(PST_Sync);
-
-	SCOPED_SCENE_READ_LOCK(SyncScene);
+	{
+		PxScene* SyncScene = PhysScene->GetPhysXScene(PST_Sync);
+		SCOPED_SCENE_READ_LOCK(SyncScene);
 
-	PxOverlapHit POverlapResult;
-	bool bHit = SyncScene->overlapAny(PGeom, PGeomPose, POverlapResult, PQueryFilterData, &PQueryCallback);
+		PxOverlapHit POverlapResult;
+		bHit = SyncScene->overlapAny(PGeom, PGeomPose, POverlapResult, PQueryFilterData, &PQueryCallback);
 
-	// if we got a hit, need to see if it was a touch or a block
-	// @todo UE4 james remove this once overlapAny only returns blocking hits
-	if (bHit)
-	{
-		ConvertQueryOverlap( POverlapResult.shape, POverlapResult.actor, OutOverlap, PFilter);
-		bHaveBlockingHit = OutOverlap.bBlockingHit;
+		// if we got a hit, need to see if it was a touch or a block
+		// @todo UE4 james remove this once overlapAny only returns blocking hits
+		if (bHit)
+		{
+			ConvertQueryOverlap( POverlapResult.shape, POverlapResult.actor, OutOverlap, PFilter);
+			bHaveBlockingHit = OutOverlap.bBlockingHit;
+		}
 	}
 
 	// Test async scene if async tests are requested and there was no blocking hit was found in the sync scene (since no hit info other than a boolean yes/no is recorded)
@@ -1603,7 +1699,6 @@ bool GeomOverlapMulti_PhysX(const UWorld* World, const PxGeometry& PGeom, const
 {
 	SCOPE_CYCLE_COUNTER(STAT_Collision_GeomOverlapMultiple);
 	bool bHaveBlockingHit = false;
-	PxScene* LockedScenes[] = { nullptr, nullptr, nullptr };
 
 	// overlapMultiple only supports sphere/capsule/box 
 	if (PGeom.getType()==PxGeometryType::eSPHERE || PGeom.getType()==PxGeometryType::eCAPSULE || PGeom.getType()==PxGeometryType::eBOX || PGeom.getType()==PxGeometryType::eCONVEXMESH )
@@ -1614,11 +1709,12 @@ bool GeomOverlapMulti_PhysX(const UWorld* World, const PxGeometry& PGeom, const
 		FPxQueryFilterCallback PQueryCallback(Params.IgnoreComponents);
 
 		// Enable scene locks, in case they are required
+		FScopedMultiSceneReadLock SceneLocks;
 		FPhysScene* PhysScene = World->GetPhysicsScene();
 		PxScene* SyncScene = PhysScene->GetPhysXScene(PST_Sync);
 
-		SCENE_LOCK_READ(SyncScene);		//we can't use scoped because we later do a conversion which depends on these results and it should all be atomic
-		LockedScenes[PST_Sync] = SyncScene;
+		// we can't use scoped because we later do a conversion which depends on these results and it should all be atomic
+		SceneLocks.LockRead(SyncScene, PST_Sync);
 
 		// Create buffer for hits. Note: memory is not initialized (for perf reasons), since API does not require it.
 		TTypeCompatibleBytes<PxOverlapHit> RawPOverlapArray[OVERLAP_BUFFER_SIZE];
@@ -1635,13 +1731,19 @@ bool GeomOverlapMulti_PhysX(const UWorld* World, const PxGeometry& PGeom, const
 			// we can use, though that is not the entire set
 			NumHits = OVERLAP_BUFFER_SIZE_MAX_SYNC_QUERIES;
 		}
+		else if (NumHits == 0)
+		{
+			// Not using anything from this scene, so unlock it.
+			SceneLocks.UnlockRead(SyncScene, PST_Sync);
+		}
 
 		// Test async scene if async tests are requested and there was no overflow
 		if (Params.bTraceAsyncScene && PhysScene->HasAsyncScene())
 		{		
 			PxScene* AsyncScene = PhysScene->GetPhysXScene(PST_Async);
-			SCENE_LOCK_READ(AsyncScene);			//we can't use scoped because we later do a conversion which depends on these results and it should all be atomic
-			LockedScenes[PST_Async] = AsyncScene;
+			
+			// we can't use scoped because we later do a conversion which depends on these results and it should all be atomic
+			SceneLocks.LockRead(AsyncScene, PST_Async);
 
 			// Write into the same PHits buffer
 			PxOverlapHit* PAsyncOverlapArray = POverlapArray + NumHits;
@@ -1658,6 +1760,11 @@ bool GeomOverlapMulti_PhysX(const UWorld* World, const PxGeometry& PGeom, const
 				// we can use, though that is not the entire set
 				NumAsyncHits = AsyncBufferSize;
 			}
+			else if (NumAsyncHits == 0)
+			{
+				// Not using anything from this scene, so unlock it.
+				SceneLocks.UnlockRead(AsyncScene, PST_Async);
+			}
 
 			NumHits += NumAsyncHits;
 		}
@@ -1665,12 +1772,7 @@ bool GeomOverlapMulti_PhysX(const UWorld* World, const PxGeometry& PGeom, const
 		if (NumHits > 0)
 		{
 			bHaveBlockingHit = ConvertOverlapResults(NumHits, POverlapArray, PFilter, OutOverlaps);
-		}
-
-		for (PxScene* LockedScene : LockedScenes)
-		{
-			SCENE_UNLOCK_READ(LockedScene);
-		}
+		}		
 		
 #if !(UE_BUILD_SHIPPING || UE_BUILD_TEST)
 		if ((World->DebugDrawTraceTag != NAME_None) && (World->DebugDrawTraceTag == Params.TraceTag))
