#ue4 - Avoid repeated Casts to get the anim update context UAnimBlueprintGeneratedClass within functions where it doesn't change.

diff --git a/Engine/Source/Runtime/Engine/Private/Animation/AnimNode_StateMachine.cpp b/Engine/Source/Runtime/Engine/Private/Animation/AnimNode_StateMachine.cpp
--- a/Engine/Source/Runtime/Engine/Private/Animation/AnimNode_StateMachine.cpp
+++ b/Engine/Source/Runtime/Engine/Private/Animation/AnimNode_StateMachine.cpp
@@ -81,13 +81,14 @@ void FAnimationActiveTransitionEntry::InitializeCustomGraphLinks(const FAnimatio
 {
 	if (TransitionRule.CustomResultNodeIndex != INDEX_NONE)
 	{
-		CustomTransitionGraph.LinkID = Context.GetAnimBlueprintClass()->AnimNodeProperties.Num() - 1 - TransitionRule.CustomResultNodeIndex; //@TODO: Crazysauce
+		const UAnimBlueprintGeneratedClass* AnimBlueprintClass = Context.GetAnimBlueprintClass();
+		CustomTransitionGraph.LinkID = AnimBlueprintClass->AnimNodeProperties.Num() - 1 - TransitionRule.CustomResultNodeIndex; //@TODO: Crazysauce
 		FAnimationInitializeContext InitContext(Context.AnimInstance);
 		CustomTransitionGraph.Initialize(InitContext);
 
 		for (int32 Index = 0; Index < TransitionRule.PoseEvaluatorLinks.Num(); ++Index)
 		{
-			FAnimNode_TransitionPoseEvaluator* PoseEvaluator = GetNodeFromPropertyIndex<FAnimNode_TransitionPoseEvaluator>(Context.AnimInstance, Context.GetAnimBlueprintClass(), TransitionRule.PoseEvaluatorLinks[Index]);
+			FAnimNode_TransitionPoseEvaluator* PoseEvaluator = GetNodeFromPropertyIndex<FAnimNode_TransitionPoseEvaluator>(Context.AnimInstance, AnimBlueprintClass, TransitionRule.PoseEvaluatorLinks[Index]);
 			PoseEvaluators.Add(PoseEvaluator);
 		}
 	}
@@ -414,9 +415,11 @@ bool FAnimNode_StateMachine::FindValidTransition(const FAnimationUpdateContext&
 	}
 	OutVisitedStateIndices.Add(CheckingStateIndex);
 
+	const UAnimBlueprintGeneratedClass* AnimBlueprintClass = Context.GetAnimBlueprintClass();
+
 	// Conduit 'states' have an additional entry rule which must be true to consider taking any transitions via the conduit
 	//@TODO: It would add flexibility to be able to define this on normal state nodes as well, assuming the dual-graph editing is sorted out
-	if (FAnimNode_TransitionResult* StateEntryRuleNode = GetNodeFromPropertyIndex<FAnimNode_TransitionResult>(Context.AnimInstance, Context.GetAnimBlueprintClass(), StateInfo.EntryRuleNodeIndex))
+	if (FAnimNode_TransitionResult* StateEntryRuleNode = GetNodeFromPropertyIndex<FAnimNode_TransitionResult>(Context.AnimInstance, AnimBlueprintClass, StateInfo.EntryRuleNodeIndex))
 	{
 		StateEntryRuleNode->EvaluateGraphExposedInputs.Execute(Context);
 
@@ -436,7 +439,7 @@ bool FAnimNode_StateMachine::FindValidTransition(const FAnimationUpdateContext&
 			continue;
 		}
 
-		FAnimNode_TransitionResult* ResultNode = GetNodeFromPropertyIndex<FAnimNode_TransitionResult>(Context.AnimInstance, Context.GetAnimBlueprintClass(), TransitionRule.CanTakeDelegateIndex);
+		FAnimNode_TransitionResult* ResultNode = GetNodeFromPropertyIndex<FAnimNode_TransitionResult>(Context.AnimInstance, AnimBlueprintClass, TransitionRule.CanTakeDelegateIndex);
 
 		// Execute it and see if we can take this rule
 		ResultNode->EvaluateGraphExposedInputs.Execute(Context);
