Fix non-unity build

diff --git a/Engine/Source/Runtime/Engine/Classes/Foliage/FoliageType.h b/Engine/Source/Runtime/Engine/Classes/Foliage/FoliageType.h
--- a/Engine/Source/Runtime/Engine/Classes/Foliage/FoliageType.h
+++ b/Engine/Source/Runtime/Engine/Classes/Foliage/FoliageType.h
@@ -1,7 +1,7 @@
 // Copyright 1998-2015 Epic Games, Inc. All Rights Reserved.
 
 #pragma once
-
+#include "Curves/CurveFloat.h"
 #include "FoliageType.generated.h"
 
 UENUM()
