Fix for UE-7108

diff --git a/Engine/Source/Runtime/Engine/Private/Slate/SceneViewport.cpp b/Engine/Source/Runtime/Engine/Private/Slate/SceneViewport.cpp
--- a/Engine/Source/Runtime/Engine/Private/Slate/SceneViewport.cpp
+++ b/Engine/Source/Runtime/Engine/Private/Slate/SceneViewport.cpp
@@ -1149,10 +1149,14 @@ void FSceneViewport::EnqueueBeginRenderFrame()
 		// Get the viewport for this window from the renderer so we can render directly to the backbuffer
 		TSharedPtr<FSlateRenderer> Renderer = FSlateApplication::Get().GetRenderer();
 		FWidgetPath WidgetPath;
-		void* ViewportResource = Renderer->GetViewportResource( *FSlateApplication::Get().FindWidgetWindow( ViewportWidget.Pin().ToSharedRef(), WidgetPath ) );
-		if( ViewportResource )
+		auto WidgetWindow = FSlateApplication::Get().FindWidgetWindow(ViewportWidget.Pin().ToSharedRef(), WidgetPath);
+		if (WidgetWindow.IsValid())
 		{
-			ViewportRHI = *((FViewportRHIRef*)ViewportResource);
+			void* ViewportResource = Renderer->GetViewportResource(*WidgetWindow);
+			if (ViewportResource)
+			{
+				ViewportRHI = *((FViewportRHIRef*)ViewportResource);
+			}
 		}
 	}
 
