Fix typo where unweld was iterating over asyncActor shapes but calling function on syncActor

diff --git a/Engine/Source/Runtime/Engine/Private/PhysicsEngine/BodyInstance.cpp b/Engine/Source/Runtime/Engine/Private/PhysicsEngine/BodyInstance.cpp
--- a/Engine/Source/Runtime/Engine/Private/PhysicsEngine/BodyInstance.cpp
+++ b/Engine/Source/Runtime/Engine/Private/PhysicsEngine/BodyInstance.cpp
@@ -2211,7 +2211,7 @@ void FBodyInstance::UnWeld(FBodyInstance* TheirBI)
 			if (TheirBI == BI)
 			{
 				PShape->userData = NULL;
-				RigidActorSync->detachShape(*PShape);
+				RigidActorAsync->detachShape(*PShape);
 				bShapesChanged = true;
 			}
 		}
