Expose Actor.SetOwner to Blueprints
#708

diff --git a/Engine/Source/Runtime/Engine/Classes/GameFramework/Actor.h b/Engine/Source/Runtime/Engine/Classes/GameFramework/Actor.h
--- a/Engine/Source/Runtime/Engine/Classes/GameFramework/Actor.h
+++ b/Engine/Source/Runtime/Engine/Classes/GameFramework/Actor.h
@@ -1540,7 +1540,11 @@ public:
 	/** Update and smooth simulated physic state, replaces PostNetReceiveLocation() and PostNetReceiveVelocity() */
 	virtual void PostNetReceivePhysicState();
 
-	/** Set the owner of this Actor, used primarily for network replication. */
+	/** 
+	 * Set the owner of this Actor, used primarily for network replication. 
+	 * @param NewOwner	The Actor whom takes over ownership of this Actor
+	 */
+	UFUNCTION(BlueprintCallable, Category=Actor)
 	void SetOwner( AActor* NewOwner );
 
 	/**
