more consistent stats rendering (same unit for easier comparison)

diff --git a/Engine/Source/Runtime/Engine/Private/StatsRender2.cpp b/Engine/Source/Runtime/Engine/Private/StatsRender2.cpp
--- a/Engine/Source/Runtime/Engine/Private/StatsRender2.cpp
+++ b/Engine/Source/Runtime/Engine/Private/StatsRender2.cpp
@@ -432,21 +432,27 @@ static int32 RenderMemoryHeadings(class FCanvas* Canvas,int32 X,int32 Y)
 	return GetStatRenderGlobals().GetFontHeight() + (GetStatRenderGlobals().GetFontHeight() / 3);
 }
 
-static FString GetMemoryString(double Value)
+// @param bAutoType true: automatically choose GB/MB/KB/... false: always use MB for easier comparisons
+static FString GetMemoryString(double Value, bool bAutoType = true)
 {
-	if (Value > 1024.0 * 1024.0 * 1024.0)
+	if(bAutoType)
 	{
-		return FString::Printf(TEXT("%.2f GB"), float(Value / (1024.0 * 1024.0 * 1024.0)));
-	}
-	if (Value > 1024.0 * 1024.0)
-	{
-		return FString::Printf(TEXT("%.2f MB"), float(Value / (1024.0 * 1024.0)));
-	}
-	if (Value > 1024.0)
-	{
-		return FString::Printf(TEXT("%.2f KB"), float(Value / (1024.0)));
+		if (Value > 1024.0 * 1024.0 * 1024.0)
+		{
+			return FString::Printf(TEXT("%.2f GB"), float(Value / (1024.0 * 1024.0 * 1024.0)));
+		}
+		if (Value > 1024.0 * 1024.0)
+		{
+			return FString::Printf(TEXT("%.2f MB"), float(Value / (1024.0 * 1024.0)));
+		}
+		if (Value > 1024.0)
+		{
+			return FString::Printf(TEXT("%.2f KB"), float(Value / (1024.0)));
+		}
+		return FString::Printf(TEXT("%.2f B"), float(Value));
 	}
-	return FString::Printf(TEXT("%.2f B"), float(Value));
+
+	return FString::Printf(TEXT("%.2f MB"), float(Value / (1024.0 * 1024.0)));
 }
 
 static int32 RenderMemoryCounter(const FGameThreadHudData& ViewData, const FComplexStatMessage& All,class FCanvas* Canvas,int32 X,int32 Y)
@@ -461,8 +467,11 @@ static int32 RenderMemoryCounter(const FGameThreadHudData& ViewData, const FComp
 	Canvas->DrawShadowedString(X,Y,*ShortenName(*All.GetDescription()),GetStatRenderGlobals().StatFont,GetStatRenderGlobals().StatColor);
 	int32 CurrX = X + GetStatRenderGlobals().AfterNameColumnOffset;
 
+	// always use MB for easier comparisons
+	const bool bAutoType = false;
+
 	// Now append the max value of the stat
-	RightJustify(Canvas,CurrX,Y,*GetMemoryString(MaxMemUsed),GetStatRenderGlobals().StatColor);
+	RightJustify(Canvas,CurrX,Y,*GetMemoryString(MaxMemUsed, bAutoType),GetStatRenderGlobals().StatColor);
 	CurrX += GetStatRenderGlobals().InterColumnOffset;
 	if (ViewData.PoolCapacity.Contains(Region))
 	{
@@ -476,7 +485,7 @@ static int32 RenderMemoryCounter(const FGameThreadHudData& ViewData, const FComp
 	CurrX += GetStatRenderGlobals().InterColumnOffset;
 	if (ViewData.PoolCapacity.Contains(Region))
 	{
-		RightJustify(Canvas,CurrX,Y,*GetMemoryString(double(ViewData.PoolCapacity[Region])),GetStatRenderGlobals().StatColor);
+		RightJustify(Canvas,CurrX,Y,*GetMemoryString(double(ViewData.PoolCapacity[Region]), bAutoType),GetStatRenderGlobals().StatColor);
 	}
 	CurrX += GetStatRenderGlobals().InterColumnOffset;
 
