fixed bad CVar usage, had no effect anyway

diff --git a/Engine/Source/Runtime/Engine/Private/ContentStreaming.cpp b/Engine/Source/Runtime/Engine/Private/ContentStreaming.cpp
--- a/Engine/Source/Runtime/Engine/Private/ContentStreaming.cpp
+++ b/Engine/Source/Runtime/Engine/Private/ContentStreaming.cpp
@@ -2143,8 +2143,6 @@ FStreamingManagerTexture::FStreamingManagerTexture()
 		GNeverStreamOutTextures = false;
 
 		// Use unlimited texture streaming in the editor. Quality is more important than stutter.
-		static IConsoleVariable* CVar = IConsoleManager::Get().FindConsoleVariable(TEXT("r.Streaming.PoolSize"));
-		CVar->Set(0, ECVF_SetByCode);
 		GTexturePoolSize = 0;
 	}
 	if (GNeverStreamOutTextures)
