Add more network logging to help with issues that involve channels opening

diff --git a/Engine/Source/Runtime/Engine/Private/NetConnection.cpp b/Engine/Source/Runtime/Engine/Private/NetConnection.cpp
--- a/Engine/Source/Runtime/Engine/Private/NetConnection.cpp
+++ b/Engine/Source/Runtime/Engine/Private/NetConnection.cpp
@@ -873,6 +873,12 @@ void UNetConnection::ReceivedPacket( FBitReader& Reader )
 			{
 				UE_LOG(LogNetTraffic, Verbose, TEXT("   Unreliable Bunch, Channel %i: Size %.1f+%.1f"), Bunch.ChIndex, (HeaderPos-IncomingStartPos)/8.f, (Reader.GetPosBits()-HeaderPos)/8.f );
 			}
+
+			if ( Bunch.bOpen )
+			{
+				UE_LOG(LogNetTraffic, Verbose, TEXT("   bOpen Bunch, Channel %i Sequence %i: Size %.1f+%.1f"), Bunch.ChIndex, Bunch.ChSequence, (HeaderPos-IncomingStartPos)/8.f, (Reader.GetPosBits()-HeaderPos)/8.f );
+			}
+
 			if ( Channels[Bunch.ChIndex] == NULL && ( Bunch.ChIndex != 0 || Bunch.ChType != CHTYPE_Control ) )
 			{
 				// Can't handle other channels until control channel exists.
@@ -951,7 +957,7 @@ void UNetConnection::ReceivedPacket( FBitReader& Reader )
 				}
 
 				// Reliable (either open or later), so create new channel.
-				UE_LOG(LogNetTraffic, Log, TEXT("      Bunch Create %i: ChType %i"), Bunch.ChIndex, Bunch.ChType );
+				UE_LOG(LogNetTraffic, Log, TEXT("      Bunch Create %i: ChType %i, bReliable: %i, bPartial: %i, bPartialInitial: %i, bPartialFinal: %i"), Bunch.ChIndex, Bunch.ChType, (int)Bunch.bReliable, (int)Bunch.bPartial, (int)Bunch.bPartialInitial, (int)Bunch.bPartialFinal );
 				Channel = CreateChannel( (EChannelType)Bunch.ChType, false, Bunch.ChIndex );
 
 				// Notify the server of the new channel.
