larger radius range for SSAO

diff --git a/Engine/Source/Runtime/Engine/Private/SceneView.cpp b/Engine/Source/Runtime/Engine/Private/SceneView.cpp
--- a/Engine/Source/Runtime/Engine/Private/SceneView.cpp
+++ b/Engine/Source/Runtime/Engine/Private/SceneView.cpp
@@ -1145,7 +1145,7 @@ void FSceneView::EndFinalPostprocessSettings(const FSceneViewInitOptions& ViewIn
 	{
 		static const auto CVar = IConsoleManager::Get().FindTConsoleVariableDataFloat(TEXT("r.AmbientOcclusionRadiusScale"));
 
-		float Scale = FMath::Clamp(CVar->GetValueOnGameThread(), 0.1f, 5.0f);
+		float Scale = FMath::Clamp(CVar->GetValueOnGameThread(), 0.1f, 15.0f);
 		
 		FinalPostProcessSettings.AmbientOcclusionRadius *= Scale;
 	}
