Improved comments

diff --git a/Engine/Source/Runtime/Engine/Classes/Engine/Channel.h b/Engine/Source/Runtime/Engine/Classes/Engine/Channel.h
--- a/Engine/Source/Runtime/Engine/Classes/Engine/Channel.h
+++ b/Engine/Source/Runtime/Engine/Classes/Engine/Channel.h
@@ -33,19 +33,19 @@ class ENGINE_API UChannel : public UObject
 	class UNetConnection*	Connection;		// Owner connection.
 
 	// Variables.
-	uint32			OpenAcked:1;		// Whether open has been acknowledged.
-	uint32			Closing:1;			// State of the channel.
-	uint32			Dormant:1;			// Channel is going dormant (it will close but the client will not destroy 
-	uint32			OpenTemporary:1;	// Opened temporarily.
-	uint32			Broken:1;			// Has encountered errors and is ignoring subsequent packets.
-	uint32			bTornOff:1;			// Actor associated with this channel was torn off
-	uint32			bPendingDormancy:1;	// Channel wants to go dormant (it will check during tick if it can go dormant)
-	int32					ChIndex;			// Index of this channel.
-	int32					OpenedLocally;		// Whether channel was opened locally or by remote.
-	FPacketIdRange		OpenPacketId;		// Packet the spawn message was sent in.
+	uint32				OpenAcked:1;		// If OpenedLocally is true, this means we have acknowledged the packet we sent the bOpen bunch on. Otherwise, it means we have received the bOpen bunch from the server.
+	uint32				Closing:1;			// State of the channel.
+	uint32				Dormant:1;			// Channel is going dormant (it will close but the client will not destroy 
+	uint32				OpenTemporary:1;	// Opened temporarily.
+	uint32				Broken:1;			// Has encountered errors and is ignoring subsequent packets.
+	uint32				bTornOff:1;			// Actor associated with this channel was torn off
+	uint32				bPendingDormancy:1;	// Channel wants to go dormant (it will check during tick if it can go dormant)
+	int32				ChIndex;			// Index of this channel.
+	int32				OpenedLocally;		// Whether channel was opened locally or by remote.
+	FPacketIdRange		OpenPacketId;		// If OpenedLocally is true, this is the packet we sent the bOpen bunch on. Otherwise, it's the packet we received the bOpen bunch on.
 	EChannelType		ChType;				// Type of this channel.
-	int32					NumInRec;			// Number of packets in InRec.
-	int32					NumOutRec;			// Number of packets in OutRec.
+	int32				NumInRec;			// Number of packets in InRec.
+	int32				NumOutRec;			// Number of packets in OutRec.
 	class FInBunch*		InRec;				// Incoming data with queued dependencies.
 	class FOutBunch*	OutRec;				// Outgoing reliable unacked data.
 	class FInBunch*		InPartialBunch;		// Partial bunch we are receiving (incoming partial bunches are appended to this)
