#ue4 - Add warning to AActor::GetComponents() about the possibility of the array changing while iterating its contents.

diff --git a/Engine/Source/Runtime/Engine/Classes/GameFramework/Actor.h b/Engine/Source/Runtime/Engine/Classes/GameFramework/Actor.h
--- a/Engine/Source/Runtime/Engine/Classes/GameFramework/Actor.h
+++ b/Engine/Source/Runtime/Engine/Classes/GameFramework/Actor.h
@@ -2174,8 +2174,11 @@ public:
 		}
 	}
 
-	// Get a direct reference to the Components array rather than a copy
-	// with the null pointers removed
+	/**
+	 * Get a direct reference to the Components array rather than a copy with the null pointers removed.
+	 * WARNING: anything that could cause the component to change ownership or be destroyed will invalidate
+	 * this array, so use caution when iterating this list!
+	 */
 	const TArray<UActorComponent*>& GetComponents() const
 	{
 		return OwnedComponents;
