#ue4 - Change a few GetComponent() calls in critical paths to use the direct access instead of the array copy. This is only possible where it's knowns that the components will not change ownership or be destroyed.

diff --git a/Engine/Source/Runtime/Engine/Private/Actor.cpp b/Engine/Source/Runtime/Engine/Private/Actor.cpp
--- a/Engine/Source/Runtime/Engine/Private/Actor.cpp
+++ b/Engine/Source/Runtime/Engine/Private/Actor.cpp
@@ -748,20 +748,19 @@ void AActor::GetComponentsBoundingCylinder(float& OutCollisionRadius, float& Out
 	float Radius = 0.f;
 	float HalfHeight = 0.f;
 
-	TInlineComponentArray<UPrimitiveComponent*> Components;
-	GetComponents(Components);
-
-	for (int32 CompIdx = 0; CompIdx < Components.Num(); CompIdx++)
+	for (const UActorComponent* ActorComponent : GetComponents())
 	{
-		UPrimitiveComponent* PrimComp = Components[CompIdx];
-
-		// Only use collidable components to find collision bounding box.
-		if ((bIgnoreRegistration || PrimComp->IsRegistered()) && (bNonColliding || PrimComp->IsCollisionEnabled()))
+		const UPrimitiveComponent* PrimComp = Cast<const UPrimitiveComponent>(ActorComponent);
+		if (PrimComp)
 		{
-			float TestRadius, TestHalfHeight;
-			PrimComp->CalcBoundingCylinder(TestRadius, TestHalfHeight);
-			Radius = FMath::Max(Radius, TestRadius);
-			HalfHeight = FMath::Max(HalfHeight, TestHalfHeight);
+			// Only use collidable components to find collision bounding box.
+			if ((bIgnoreRegistration || PrimComp->IsRegistered()) && (bNonColliding || PrimComp->IsCollisionEnabled()))
+			{
+				float TestRadius, TestHalfHeight;
+				PrimComp->CalcBoundingCylinder(TestRadius, TestHalfHeight);
+				Radius = FMath::Max(Radius, TestRadius);
+				HalfHeight = FMath::Max(HalfHeight, TestHalfHeight);
+			}
 		}
 	}
 
@@ -839,17 +838,16 @@ FBox AActor::GetComponentsBoundingBox(bool bNonColliding) const
 {
 	FBox Box(0);
 
-	TInlineComponentArray<UPrimitiveComponent*> Components;
-	GetComponents(Components);
-
-	for(int32 CompIdx = 0; CompIdx < Components.Num(); CompIdx++)
+	for (const UActorComponent* ActorComponent : GetComponents())
 	{
-		UPrimitiveComponent* PrimComp = Components[CompIdx];
-
-		// Only use collidable components to find collision bounding box.
-		if( PrimComp->IsRegistered() && (bNonColliding || PrimComp->IsCollisionEnabled()) )
+		const UPrimitiveComponent* PrimComp = Cast<const UPrimitiveComponent>(ActorComponent);
+		if (PrimComp)
 		{
-			Box += PrimComp->Bounds.GetBox();
+			// Only use collidable components to find collision bounding box.
+			if (PrimComp->IsRegistered() && (bNonColliding || PrimComp->IsCollisionEnabled()))
+			{
+				Box += PrimComp->Bounds.GetBox();
+			}
 		}
 	}
 
@@ -1062,14 +1060,11 @@ float AActor::GetLastRenderTime() const
 {
 	// return most recent of Components' LastRenderTime
 	// @todo UE4 maybe check base component and components attached to it instead?
-	TInlineComponentArray<UPrimitiveComponent*> Components;
-	GetComponents(Components);
-
 	float LastRenderTime = -1000.f;
-	for( int32 i=0; i<Components.Num(); i++ )
+	for (const UActorComponent* ActorComponent : GetComponents())
 	{
-		UPrimitiveComponent* PrimComp = Components[i];
-		if(PrimComp->IsRegistered())
+		const UPrimitiveComponent* PrimComp = Cast<const UPrimitiveComponent>(ActorComponent);
+		if (PrimComp && PrimComp->IsRegistered())
 		{
 			LastRenderTime = FMath::Max(LastRenderTime, PrimComp->LastRenderTime);
 		}
diff --git a/Engine/Source/Runtime/Engine/Private/Character.cpp b/Engine/Source/Runtime/Engine/Private/Character.cpp
--- a/Engine/Source/Runtime/Engine/Private/Character.cpp
+++ b/Engine/Source/Runtime/Engine/Private/Character.cpp
@@ -424,11 +424,9 @@ namespace MovementBaseUtility
 				}
 
 				// @TODO: We need to find a more efficient way of finding all ticking components in an actor.
-				TInlineComponentArray<UActorComponent*> Components;
-				NewBaseOwner->GetComponents(Components);
-				for (auto& Component : Components)
+				for (UActorComponent* Component : NewBaseOwner->GetComponents())
 				{
-					if (Component->PrimaryComponentTick.bCanEverTick)
+					if (Component && Component->PrimaryComponentTick.bCanEverTick)
 					{
 						BasedObjectTick.AddPrerequisite(Component, Component->PrimaryComponentTick);
 					}
@@ -454,11 +452,9 @@ namespace MovementBaseUtility
 				BasedObjectTick.RemovePrerequisite(OldBaseOwner, OldBaseOwner->PrimaryActorTick);
 
 				// @TODO: We need to find a more efficient way of finding all ticking components in an actor.
-				TInlineComponentArray<UActorComponent*> Components;
-				OldBaseOwner->GetComponents(Components);
-				for (auto& Component : Components)
+				for (UActorComponent* Component : OldBaseOwner->GetComponents())
 				{
-					if (Component->PrimaryComponentTick.bCanEverTick)
+					if (Component && Component->PrimaryComponentTick.bCanEverTick)
 					{
 						BasedObjectTick.RemovePrerequisite(Component, Component->PrimaryComponentTick);
 					}
diff --git a/Engine/Source/Runtime/Engine/Private/Collision/Collision.cpp b/Engine/Source/Runtime/Engine/Private/Collision/Collision.cpp
--- a/Engine/Source/Runtime/Engine/Private/Collision/Collision.cpp
+++ b/Engine/Source/Runtime/Engine/Private/Collision/Collision.cpp
@@ -133,22 +133,27 @@ void FCollisionQueryParams::AddIgnoredActor(const AActor* InIgnoreActor)
 {
 	if (InIgnoreActor)
 	{
-		TInlineComponentArray<UPrimitiveComponent*> PrimComponents;
-		InIgnoreActor->GetComponents(PrimComponents);
-
 		if (IgnoreComponents.Num() == 0)
 		{
 			// Don't need to AddUnique if we start empty, as GetComponents should never contain duplicate components.
-			for (UPrimitiveComponent* PrimComponent : PrimComponents)
+			for (const UActorComponent* ActorComponent : InIgnoreActor->GetComponents())
 			{
-				IgnoreComponents.Add(PrimComponent->GetUniqueID());
+				const UPrimitiveComponent* PrimComponent = Cast<const UPrimitiveComponent>(ActorComponent);
+				if (PrimComponent)
+				{
+					IgnoreComponents.Add(PrimComponent->GetUniqueID());
+				}
 			}
 		}
 		else
 		{
-			for (UPrimitiveComponent* PrimComponent : PrimComponents)
+			for (const UActorComponent* ActorComponent : InIgnoreActor->GetComponents())
 			{
-				IgnoreComponents.AddUnique(PrimComponent->GetUniqueID());
+				const UPrimitiveComponent* PrimComponent = Cast<const UPrimitiveComponent>(ActorComponent);
+				if (PrimComponent)
+				{
+					IgnoreComponents.AddUnique(PrimComponent->GetUniqueID());
+				}
 			}
 		}
 	}
diff --git a/Engine/Source/Runtime/Engine/Private/PlayerController.cpp b/Engine/Source/Runtime/Engine/Private/PlayerController.cpp
--- a/Engine/Source/Runtime/Engine/Private/PlayerController.cpp
+++ b/Engine/Source/Runtime/Engine/Private/PlayerController.cpp
@@ -2290,13 +2290,10 @@ void APlayerController::BuildInputStack(TArray<UInputComponent*>& InputStack)
 			}
 
 			// See if there is another InputComponent that was added to the Pawn's components array (possibly by script).
-			TInlineComponentArray<UInputComponent*> Components;
-			ControlledPawn->GetComponents(Components);
-
-			for (int32 i=0; i < Components.Num(); i++)
+			for (UActorComponent* ActorComponent : ControlledPawn->GetComponents())
 			{
-				UInputComponent* PawnInputComponent = Components[i];
-				if (PawnInputComponent != ControlledPawn->InputComponent)
+				UInputComponent* PawnInputComponent = Cast<UInputComponent>(ActorComponent);
+				if (PawnInputComponent && PawnInputComponent != ControlledPawn->InputComponent)
 				{
 					InputStack.Push(PawnInputComponent);
 				}
diff --git a/Engine/Source/Runtime/Landscape/Private/Landscape.cpp b/Engine/Source/Runtime/Landscape/Private/Landscape.cpp
--- a/Engine/Source/Runtime/Landscape/Private/Landscape.cpp
+++ b/Engine/Source/Runtime/Landscape/Private/Landscape.cpp
@@ -2683,26 +2683,19 @@ void ALandscapeProxy::Tick(float DeltaSeconds)
 
 	TSet<UHierarchicalInstancedStaticMeshComponent *> ToDestroy;
 
-	TInlineComponentArray<UHierarchicalInstancedStaticMeshComponent*> FoliageComps;
-	GetComponents(FoliageComps);
-	for (UHierarchicalInstancedStaticMeshComponent* Component : FoliageComps)
+	for (UActorComponent* ActorComponent : GetComponents())
 	{
-		if (!StillUsed.Contains(Component))
+		UHierarchicalInstancedStaticMeshComponent* HComponent = Cast<UHierarchicalInstancedStaticMeshComponent>(ActorComponent);
+		if (HComponent && !StillUsed.Contains(HComponent))
 		{
-			ToDestroy.Add(Component);
+			ToDestroy.Add(HComponent);
 		}
 	}
 
-	TArray<USceneComponent*> AttachedFoliageComponents = RootComponent->AttachChildren.FilterByPredicate(
-		[](USceneComponent* Component)
-	{
-		return Cast<UHierarchicalInstancedStaticMeshComponent>(Component);
-	});
-
 	// Destroy any attached but un-owned foliage components
-	for (USceneComponent* Component : AttachedFoliageComponents)
+	for (USceneComponent* SceneComponent : RootComponent->AttachChildren)
 	{
-		UHierarchicalInstancedStaticMeshComponent* HComponent = Cast<UHierarchicalInstancedStaticMeshComponent>(Component);
+		UHierarchicalInstancedStaticMeshComponent* HComponent = Cast<UHierarchicalInstancedStaticMeshComponent>(SceneComponent);
 		if (HComponent && !StillUsed.Contains(HComponent))
 		{
 			ToDestroy.Add(HComponent);
