Support multiple level-centric blueprints

diff --git a/Engine/Source/Runtime/Engine/Classes/Engine/Level.h b/Engine/Source/Runtime/Engine/Classes/Engine/Level.h
--- a/Engine/Source/Runtime/Engine/Classes/Engine/Level.h
+++ b/Engine/Source/Runtime/Engine/Classes/Engine/Level.h
@@ -666,6 +666,11 @@ public:
 	 */
 	ENGINE_API class ULevelScriptBlueprint* GetLevelScriptBlueprint(bool bDontCreate=false);
 
+	/**
+	 *  Returns a list of all blueprints contained within the level
+	 */
+	ENGINE_API TArray<class UBlueprint*> GetLevelBlueprints() const;
+
 	/**
 	 *  Called when the level script blueprint has been successfully changed and compiled.  Handles creating an instance of the blueprint class in LevelScriptActor
 	 */
diff --git a/Engine/Source/Runtime/Engine/Private/Level.cpp b/Engine/Source/Runtime/Engine/Private/Level.cpp
--- a/Engine/Source/Runtime/Engine/Private/Level.cpp
+++ b/Engine/Source/Runtime/Engine/Private/Level.cpp
@@ -1629,6 +1629,24 @@ bool ULevel::HasAnyActorsOfType(UClass *SearchType)
 
 #if WITH_EDITOR
 
+TArray<UBlueprint*> ULevel::GetLevelBlueprints() const
+{
+	TArray<UBlueprint*> LevelBlueprints;
+	TArray<UObject*> LevelChildren;
+	GetObjectsWithOuter(this, LevelChildren, false, RF_PendingKill);
+
+	for (UObject* LevelChild : LevelChildren)
+	{
+		UBlueprint* LevelChildBP = Cast<UBlueprint>(LevelChild);
+		if (LevelChildBP)
+		{
+			LevelBlueprints.Add(LevelChildBP);
+		}
+	}
+
+	return LevelBlueprints;
+}
+
 ULevelScriptBlueprint* ULevel::GetLevelScriptBlueprint(bool bDontCreate)
 {
 	const FString LevelScriptName = ULevelScriptBlueprint::CreateLevelScriptNameFromLevel(this);
diff --git a/Engine/Source/Runtime/Engine/Private/World.cpp b/Engine/Source/Runtime/Engine/Private/World.cpp
--- a/Engine/Source/Runtime/Engine/Private/World.cpp
+++ b/Engine/Source/Runtime/Engine/Private/World.cpp
@@ -591,10 +591,13 @@ bool UWorld::PreSaveRoot(const TCHAR* Filename, TArray<FString>& AdditionalPacka
 		}
 	}
 #if WITH_EDITOR
-	// If this level has a level script, rebuild it now to ensure no stale data is stored in the level script actor
-	if( !IsRunningCommandlet() && PersistentLevel->GetLevelScriptBlueprint(true) )
+	// Rebuild all level blueprints now to ensure no stale data is stored on the actors
+	if( !IsRunningCommandlet() )
 	{
-		FKismetEditorUtilities::CompileBlueprint(PersistentLevel->GetLevelScriptBlueprint(true), true, true);
+		for (UBlueprint* Blueprint : PersistentLevel->GetLevelBlueprints())
+		{
+			FKismetEditorUtilities::CompileBlueprint(Blueprint, false, true);
+		}
 	}
 #endif
 
@@ -5398,9 +5401,9 @@ void UWorld::GetAssetRegistryTags(TArray<FAssetRegistryTag>& OutTags) const
 {
 	if(PersistentLevel && PersistentLevel->OwningWorld)
 	{
-		if(ULevelScriptBlueprint* LevelBlueprint = PersistentLevel->GetLevelScriptBlueprint(true))
+		for (UBlueprint* Blueprint : PersistentLevel->GetLevelBlueprints())
 		{
-			LevelBlueprint->GetAssetRegistryTags(OutTags);
+			Blueprint->GetAssetRegistryTags(OutTags);
 		}
 	}
 }
