#ue4 - Remove declared function PlayerController::SetFocusAndLocking() without matching implementation.

Jira: UE-7120

diff --git a/Engine/Source/Runtime/Engine/Classes/GameFramework/PlayerController.h b/Engine/Source/Runtime/Engine/Classes/GameFramework/PlayerController.h
--- a/Engine/Source/Runtime/Engine/Classes/GameFramework/PlayerController.h
+++ b/Engine/Source/Runtime/Engine/Classes/GameFramework/PlayerController.h
@@ -1041,11 +1041,6 @@ public:
 	/** Setup an input mode. */
 	virtual void SetInputMode(const FInputModeDataBase& InData);
 
-protected:
-
-	/** Utility function for Input Mode functions. */
-	void SetFocusAndLocking(TSharedPtr<SWidget> InWidgetToFocus, bool bLockMouseToViewport, TSharedPtr<class SViewport> InViewportWidget);
-
 public:
 	/**
 	 * Change Camera mode
