Fixed index out of range error

diff --git a/Engine/Source/Runtime/Engine/Private/LightMap.cpp b/Engine/Source/Runtime/Engine/Private/LightMap.cpp
--- a/Engine/Source/Runtime/Engine/Private/LightMap.cpp
+++ b/Engine/Source/Runtime/Engine/Private/LightMap.cpp
@@ -1625,17 +1625,12 @@ void FLightMap2D::Serialize(FArchive& Ar)
 	{
 		for(uint32 CoefficientIndex = 0;CoefficientIndex < 3;CoefficientIndex++)
 		{
-			Ar << Textures[CoefficientIndex];
-			Ar << ScaleVectors[CoefficientIndex];
-			Ar << AddVectors[CoefficientIndex];
+			ULightMapTexture2D* Dummy = NULL;
+			Ar << Dummy;
+			FVector4 Dummy2;
+			Ar << Dummy2;
+			Ar << Dummy2;
 		}
-
-		ScaleVectors[0].W *= 11.5f;
-		AddVectors[0].W = ( AddVectors[0].W - 0.5f ) * 11.5f;
-
-		ScaleVectors[1] *= FVector4( -0.325735f, 0.325735f, -0.325735f, 0.0f );
-		AddVectors[1] *= FVector4( -0.325735f, 0.325735f, -0.325735f, 0.0f );
-		AddVectors[1].W = 0.282095f;
 	}
 	else if( Ar.IsLoading() && Ar.UE4Ver() < VER_UE4_COMBINED_LIGHTMAP_TEXTURES )
 	{
