Fixed crash after traveling to a world composition map during PIE
UE-6492

diff --git a/Engine/Source/Runtime/Engine/Private/WorldComposition.cpp b/Engine/Source/Runtime/Engine/Private/WorldComposition.cpp
--- a/Engine/Source/Runtime/Engine/Private/WorldComposition.cpp
+++ b/Engine/Source/Runtime/Engine/Private/WorldComposition.cpp
@@ -202,8 +202,9 @@ void UWorldComposition::Rescan()
 	FString WorldRootFilename = FPackageName::LongPackageNameToFilename(WorldRoot);
 	FPlatformFileManager::Get().GetPlatformFile().IterateDirectoryRecursively(*WorldRootFilename, Gatherer);
 
-	FString PersistentLevelPackageName = OwningWorld->GetOutermost()->GetName();
-	
+	// Make sure we have persistent level name without PIE prefix
+	FString PersistentLevelPackageName = UWorld::StripPIEPrefixFromPackageName(OwningWorld->GetOutermost()->GetName(), OwningWorld->StreamingLevelsPrefix);
+		
 	// Add found tiles to a world composition, except persistent level
 	for (const auto& TilePackageName : Gatherer.TilesCollection)
 	{
