Fix for crash with foliage instances but no component when resizing landscape.

diff --git a/Engine/Source/Runtime/Engine/Private/InstancedFoliage.cpp b/Engine/Source/Runtime/Engine/Private/InstancedFoliage.cpp
--- a/Engine/Source/Runtime/Engine/Private/InstancedFoliage.cpp
+++ b/Engine/Source/Runtime/Engine/Private/InstancedFoliage.cpp
@@ -1544,6 +1544,28 @@ void AInstancedFoliageActor::PostLoad()
 		{
 			// Find the per-mesh info matching the mesh.
 			FFoliageMeshInfo& MeshInfo = *MeshPair.Value;
+			UFoliageType* FoliageType = MeshPair.Key;
+
+			if (MeshInfo.Instances.Num() && MeshInfo.Component == nullptr)
+			{
+				const UStaticMesh* StaticMesh = FoliageType->GetStaticMesh();
+				FFormatNamedArguments Arguments;
+				if (StaticMesh)
+				{
+					Arguments.Add(TEXT("MeshName"), FText::FromString(StaticMesh->GetName()));
+				}
+				else
+				{
+					Arguments.Add(TEXT("MeshName"), FText::FromString(TEXT("None")));
+				}
+
+				FMessageLog("MapCheck").Warning()
+					->AddToken(FUObjectToken::Create(this))
+					->AddToken(FTextToken::Create(FText::Format(LOCTEXT("MapCheck_Message_FoliageMissingComponent", "Foliage in this map is missing a component for static mesh {MeshName}. This has been repaired."),Arguments)))
+					->AddToken(FMapErrorToken::Create(FMapErrors::FoliageMissingClusterComponent));
+
+				MeshInfo.ReallocateClusters(this, MeshPair.Key);
+			}
 
 			// Update the hash.
 			for (int32 InstanceIdx = 0; InstanceIdx < MeshInfo.Instances.Num(); InstanceIdx++)
