Fix for issue where a SceneProxy is not created for newly-converted foliage that is still building its tree for the first time

diff --git a/Engine/Source/Runtime/Engine/Private/HierarchicalInstancedStaticMesh.cpp b/Engine/Source/Runtime/Engine/Private/HierarchicalInstancedStaticMesh.cpp
--- a/Engine/Source/Runtime/Engine/Private/HierarchicalInstancedStaticMesh.cpp
+++ b/Engine/Source/Runtime/Engine/Private/HierarchicalInstancedStaticMesh.cpp
@@ -1799,7 +1799,6 @@ FPrimitiveSceneProxy* UHierarchicalInstancedStaticMeshComponent::CreateSceneProx
 	const bool bMeshIsValid = 
 		// make sure we have instances
 		PerInstanceSMData.Num() > 0 &&
-		ClusterTreePtr->Num() &&
 		// make sure we have an actual staticmesh
 		StaticMesh &&
 		StaticMesh->HasValidRenderData() &&
