fixed PIE warning for simple move functions

diff --git a/Engine/Source/Runtime/Engine/Private/Controller.cpp b/Engine/Source/Runtime/Engine/Private/Controller.cpp
--- a/Engine/Source/Runtime/Engine/Private/Controller.cpp
+++ b/Engine/Source/Runtime/Engine/Private/Controller.cpp
@@ -571,6 +571,7 @@ void AController::InitNavigationControl(UPathFollowingComponent*& PathFollowingC
 	{
 		PathFollowingComp = NewObject<UPathFollowingComponent>(this);
 		PathFollowingComp->RegisterComponentWithWorld(GetWorld());
+		PathFollowingComp->Initialize();
 	}
 }
 
