Made BlackboardComponent a blueprint spawnable component so that BP users can start using it outside of AI controllers #UE4

- also fixed FloatingPawnComponent not having a proper class group set

diff --git a/Engine/Source/Runtime/AIModule/Classes/BehaviorTree/BlackboardComponent.h b/Engine/Source/Runtime/AIModule/Classes/BehaviorTree/BlackboardComponent.h
--- a/Engine/Source/Runtime/AIModule/Classes/BehaviorTree/BlackboardComponent.h
+++ b/Engine/Source/Runtime/AIModule/Classes/BehaviorTree/BlackboardComponent.h
@@ -33,7 +33,7 @@ namespace EBlackboardDescription
 	};
 }
 
-UCLASS()
+UCLASS(ClassGroup = AI, meta = (BlueprintSpawnableComponent))
 class AIMODULE_API UBlackboardComponent : public UActorComponent
 {
 	GENERATED_BODY()
diff --git a/Engine/Source/Runtime/Engine/Classes/GameFramework/FloatingPawnMovement.h b/Engine/Source/Runtime/Engine/Classes/GameFramework/FloatingPawnMovement.h
--- a/Engine/Source/Runtime/Engine/Classes/GameFramework/FloatingPawnMovement.h
+++ b/Engine/Source/Runtime/Engine/Classes/GameFramework/FloatingPawnMovement.h
@@ -12,7 +12,7 @@
  * FloatingPawnMovement is a movement component that provides simple movement for any Pawn class.
  * Limits on speed and acceleration are provided, while gravity is not implemented.
  */
-UCLASS(meta = (BlueprintSpawnableComponent))
+UCLASS(ClassGroup = Movement, meta = (BlueprintSpawnableComponent))
 class ENGINE_API UFloatingPawnMovement : public UPawnMovementComponent
 {
 	GENERATED_UCLASS_BODY()
