missing file for CL 2399868

diff --git a/Engine/Source/Runtime/Engine/Private/SceneView.cpp b/Engine/Source/Runtime/Engine/Private/SceneView.cpp
--- a/Engine/Source/Runtime/Engine/Private/SceneView.cpp
+++ b/Engine/Source/Runtime/Engine/Private/SceneView.cpp
@@ -759,7 +759,7 @@ void FSceneView::OverridePostProcessSettings(const FPostProcessSettings& Src, fl
 	LERP_PP(AmbientOcclusionRadius);
 	LERP_PP(AmbientOcclusionFadeDistance);
 	LERP_PP(AmbientOcclusionFadeRadius);
-	LERP_PP(AmbientOcclusionDistance);
+	LERP_PP(AmbientOcclusionDistance_DEPRECATED);
 	LERP_PP(AmbientOcclusionPower);
 	LERP_PP(AmbientOcclusionBias);
 	LERP_PP(AmbientOcclusionQuality);
@@ -1153,7 +1153,7 @@ void FSceneView::EndFinalPostprocessSettings(const FSceneViewInitOptions& ViewIn
 	{
 		float Scale = FMath::Clamp(CVarSSAOFadeRadiusScale.GetValueOnGameThread(), 0.01f, 50.0f);
 
-		FinalPostProcessSettings.AmbientOcclusionDistance *= Scale;
+		FinalPostProcessSettings.AmbientOcclusionDistance_DEPRECATED *= Scale;
 	}
 
 	{
