Fix issues with copy/pasting Ambient Sound actors that have an internal Sound Cue.

diff --git a/Engine/Source/Editor/SoundCueEditor/Private/SoundCueEditorUtilities.cpp b/Engine/Source/Editor/SoundCueEditor/Private/SoundCueEditorUtilities.cpp
--- a/Engine/Source/Editor/SoundCueEditor/Private/SoundCueEditorUtilities.cpp
+++ b/Engine/Source/Editor/SoundCueEditor/Private/SoundCueEditorUtilities.cpp
@@ -43,7 +43,7 @@ void FSoundCueEditorUtilities::CreateWaveContainers(TArray<USoundWave*>& Selecte
 
 			WavePlayer->SoundWave = NewWave;
 
-			WavePlayer->GraphNode->NodePosX = Location.X - WavePlayer->GraphNode->EstimateNodeWidth();
+			WavePlayer->GraphNode->NodePosX = Location.X - WavePlayer->GetGraphNode()->EstimateNodeWidth();
 			WavePlayer->GraphNode->NodePosY = Location.Y + (NodeSpacing * WaveIndex);
 
 			OutPlayers.Add(WavePlayer);
diff --git a/Engine/Source/Editor/UnrealEd/Private/SoundCueGraphNode.cpp b/Engine/Source/Editor/UnrealEd/Private/SoundCueGraphNode.cpp
--- a/Engine/Source/Editor/UnrealEd/Private/SoundCueGraphNode.cpp
+++ b/Engine/Source/Editor/UnrealEd/Private/SoundCueGraphNode.cpp
@@ -188,8 +188,13 @@ void USoundCueGraphNode::ResetSoundNodeOwner()
 {
 	if (SoundNode)
 	{
-		// Ensures SoundNode is owned by the SoundCue
-		SoundNode->Rename(NULL, CastChecked<USoundCueGraph>(GetGraph())->GetSoundCue(), REN_DontCreateRedirectors);
+		USoundCue* SoundCue = CastChecked<USoundCueGraph>(GetGraph())->GetSoundCue();
+
+		if (SoundNode->GetOuter() != SoundCue)
+		{
+			// Ensures SoundNode is owned by the SoundCue
+			SoundNode->Rename(NULL, SoundCue, REN_DontCreateRedirectors);
+		}
 
 		// Set up the back pointer for newly created sound nodes
 		SoundNode->GraphNode = this;
diff --git a/Engine/Source/Editor/UnrealEd/Private/SoundCueGraphSchema.cpp b/Engine/Source/Editor/UnrealEd/Private/SoundCueGraphSchema.cpp
--- a/Engine/Source/Editor/UnrealEd/Private/SoundCueGraphSchema.cpp
+++ b/Engine/Source/Editor/UnrealEd/Private/SoundCueGraphSchema.cpp
@@ -123,7 +123,7 @@ UEdGraphNode* FSoundCueGraphSchemaAction_NewFromSelected::PerformAction(class UE
 	if (SoundNodeClass)
 	{
 		USoundNode* NewNode = SoundCue->ConstructSoundNode<USoundNode>(SoundNodeClass, bSelectNewNode);
-		USoundCueGraphNode* NewGraphNode = NewNode->GraphNode;
+		UEdGraphNode* NewGraphNode = NewNode->GraphNode;
 		const USoundCueGraphSchema* NewSchema = CastChecked<USoundCueGraphSchema>(NewGraphNode->GetSchema());
 
 		// If this node allows >0 children but by default has zero - create a connector for starters
@@ -230,11 +230,11 @@ void USoundCueGraphSchema::TryConnectNodes(const TArray<USoundNode*>& OutputNode
 	{
 		if ( Index < InputNode->GetMaxChildNodes() )
 		{
-			if (Index >= InputNode->GraphNode->GetInputCount())
+			if (Index >= InputNode->GetGraphNode()->GetInputCount())
 			{
-				InputNode->GraphNode->CreateInputPin();
+				InputNode->GetGraphNode()->CreateInputPin();
 			}
-			TryCreateConnection( InputNode->GraphNode->GetInputPin(Index), OutputNodes[Index]->GraphNode->GetOutputPin() );
+			TryCreateConnection( InputNode->GetGraphNode()->GetInputPin(Index), OutputNodes[Index]->GetGraphNode()->GetOutputPin() );
 		}
 	}
 }
diff --git a/Engine/Source/Runtime/Engine/Classes/Sound/SoundCue.h b/Engine/Source/Runtime/Engine/Classes/Sound/SoundCue.h
--- a/Engine/Source/Runtime/Engine/Classes/Sound/SoundCue.h
+++ b/Engine/Source/Runtime/Engine/Classes/Sound/SoundCue.h
@@ -61,8 +61,9 @@ class USoundCue : public USoundBase
 	UPROPERTY()
 	TArray<USoundNode*> AllNodes;
 
-	/** EdGraph based representation of the SoundCue */
-	class USoundCueGraph* SoundCueGraph;
+	UPROPERTY()
+	UEdGraph* SoundCueGraph;
+
 #endif
 
 private:
@@ -152,10 +153,7 @@ public:
 	ENGINE_API void CompileSoundNodesFromGraphNodes();
 
 	/** Get the EdGraph of SoundNodes */
-	USoundCueGraph* GetGraph()
-	{
-		return SoundCueGraph;
-	}
+	ENGINE_API class USoundCueGraph* GetGraph();
 #endif
 };
 
diff --git a/Engine/Source/Runtime/Engine/Classes/Sound/SoundNode.h b/Engine/Source/Runtime/Engine/Classes/Sound/SoundNode.h
--- a/Engine/Source/Runtime/Engine/Classes/Sound/SoundNode.h
+++ b/Engine/Source/Runtime/Engine/Classes/Sound/SoundNode.h
@@ -52,7 +52,10 @@ class ENGINE_API USoundNode : public UObject
 
 #if WITH_EDITORONLY_DATA
 	/** Node's Graph representation, used to get position. */
-	class USoundCueGraphNode*	GraphNode;
+	UPROPERTY()
+	UEdGraphNode*	GraphNode;
+
+	class USoundCueGraphNode* GetGraphNode() const;
 #endif
 
 public:
diff --git a/Engine/Source/Runtime/Engine/Private/SoundCue.cpp b/Engine/Source/Runtime/Engine/Private/SoundCue.cpp
--- a/Engine/Source/Runtime/Engine/Private/SoundCue.cpp
+++ b/Engine/Source/Runtime/Engine/Private/SoundCue.cpp
@@ -310,16 +310,22 @@ const FAttenuationSettings* USoundCue::GetAttenuationSettingsToApply() const
 }
 
 #if WITH_EDITOR
+USoundCueGraph* USoundCue::GetGraph()
+{ 
+	return CastChecked<USoundCueGraph>(SoundCueGraph);
+}
+
 void USoundCue::CreateGraph()
 {
-	check(SoundCueGraph == NULL);
-
-	SoundCueGraph = CastChecked<USoundCueGraph>(FBlueprintEditorUtils::CreateNewGraph(this, NAME_None, USoundCueGraph::StaticClass(), USoundCueGraphSchema::StaticClass()));
-	SoundCueGraph->bAllowDeletion = false;
+	if (SoundCueGraph == nullptr)
+	{
+		SoundCueGraph = CastChecked<USoundCueGraph>(FBlueprintEditorUtils::CreateNewGraph(this, NAME_None, USoundCueGraph::StaticClass(), USoundCueGraphSchema::StaticClass()));
+		SoundCueGraph->bAllowDeletion = false;
 
-	// Give the schema a chance to fill out any required nodes (like the results node)
-	const UEdGraphSchema* Schema = SoundCueGraph->GetSchema();
-	Schema->CreateDefaultNodesForGraph(*SoundCueGraph);
+		// Give the schema a chance to fill out any required nodes (like the results node)
+		const UEdGraphSchema* Schema = SoundCueGraph->GetSchema();
+		Schema->CreateDefaultNodesForGraph(*SoundCueGraph);
+	}
 }
 
 void USoundCue::ClearGraph()
@@ -355,7 +361,7 @@ void USoundCue::LinkGraphNodesFromSoundNodes()
 		check(RootNodeList.Num() == 1);
 
 		RootNodeList[0]->Pins[0]->BreakAllPinLinks();
-		RootNodeList[0]->Pins[0]->MakeLinkTo(FirstNode->GraphNode->GetOutputPin());
+		RootNodeList[0]->Pins[0]->MakeLinkTo(FirstNode->GetGraphNode()->GetOutputPin());
 	}
 
 	for(TArray<USoundNode*>::TConstIterator It(AllNodes); It; ++It)
@@ -364,7 +370,7 @@ void USoundCue::LinkGraphNodesFromSoundNodes()
 		if (SoundNode)
 		{
 			TArray<UEdGraphPin*> InputPins;
-			SoundNode->GraphNode->GetInputPins(/*out*/ InputPins);
+			SoundNode->GetGraphNode()->GetInputPins(/*out*/ InputPins);
 			check(InputPins.Num() == SoundNode->ChildNodes.Num());
 			for (int32 ChildIndex = 0; ChildIndex < SoundNode->ChildNodes.Num(); ChildIndex++)
 			{
@@ -372,7 +378,7 @@ void USoundCue::LinkGraphNodesFromSoundNodes()
 				if (ChildNode)
 				{
 					InputPins[ChildIndex]->BreakAllPinLinks();
-					InputPins[ChildIndex]->MakeLinkTo(ChildNode->GraphNode->GetOutputPin());
+					InputPins[ChildIndex]->MakeLinkTo(ChildNode->GetGraphNode()->GetOutputPin());
 				}
 			}
 		}
diff --git a/Engine/Source/Runtime/Engine/Private/SoundNode.cpp b/Engine/Source/Runtime/Engine/Private/SoundNode.cpp
--- a/Engine/Source/Runtime/Engine/Private/SoundNode.cpp
+++ b/Engine/Source/Runtime/Engine/Private/SoundNode.cpp
@@ -38,6 +38,13 @@ void USoundNode::AddReferencedObjects(UObject* InThis, FReferenceCollector& Coll
 }
 #endif //WITH_EDITOR
 
+#if WITH_EDITORONLY_DATA
+USoundCueGraphNode* USoundNode::GetGraphNode() const
+{
+	return CastChecked<USoundCueGraphNode>(GraphNode);
+}
+#endif
+
 UPTRINT USoundNode::GetNodeWaveInstanceHash(const UPTRINT ParentWaveInstanceHash, const USoundNode* ChildNode, const uint32 ChildIndex)
 {
 	checkf(ChildIndex < MAX_ALLOWED_CHILD_NODES, TEXT("Too many children (%d) in SoundCue '%s'"), ChildIndex, *CastChecked<USoundCue>(ChildNode->GetOuter())->GetFullName());
@@ -86,7 +93,7 @@ void USoundNode::InsertChildNode( int32 Index )
 	{
 		ChildNodes.InsertZeroed( Index );
 #if WITH_EDITOR
-		GraphNode->CreateInputPin();
+		GetGraphNode()->CreateInputPin();
 #endif //WITH_EDITORONLY_DATA
 	}
 }
diff --git a/Engine/Source/Runtime/Engine/Private/SoundNodeSwitch.cpp b/Engine/Source/Runtime/Engine/Private/SoundNodeSwitch.cpp
--- a/Engine/Source/Runtime/Engine/Private/SoundNodeSwitch.cpp
+++ b/Engine/Source/Runtime/Engine/Private/SoundNodeSwitch.cpp
@@ -50,7 +50,7 @@ void USoundNodeSwitch::RenamePins()
 	TArray<class UEdGraphPin*> InputPins;
 
 #if WITH_EDITORONLY_DATA
-	GraphNode->GetInputPins(InputPins);
+	GetGraphNode()->GetInputPins(InputPins);
 #endif
 
 	for (int32 i = 0; i < InputPins.Num(); i++)
