cleanup UI

diff --git a/Engine/Source/Runtime/Engine/Classes/Engine/Scene.h b/Engine/Source/Runtime/Engine/Classes/Engine/Scene.h
--- a/Engine/Source/Runtime/Engine/Classes/Engine/Scene.h
+++ b/Engine/Source/Runtime/Engine/Classes/Engine/Scene.h
@@ -634,7 +634,7 @@ struct FPostProcessSettings
 	float AmbientOcclusionFadeRadius;
 
 	/** >0, in unreal units, how wide the ambient occlusion effect should affect the geometry (in depth), will be removed - only used for non normal method which is not exposed */
-	UPROPERTY(interp, BlueprintReadWrite, Category=AmbientOcclusion, AdvancedDisplay, meta=(ClampMin = "0.0", UIMax = "2000.0", editcondition = "bOverride_AmbientOcclusionDistance_DEPRECATED", DisplayName = "Occlusion Distance"))
+	UPROPERTY()
 	float AmbientOcclusionDistance_DEPRECATED;
 
 	/** >0, in unreal units, bigger values means even distant surfaces affect the ambient occlusion */
