Don't process unreliable packets if the channel isn't fully open + general cleanup of how we open channel + more sanity checking

diff --git a/Engine/Source/Runtime/Engine/Private/DataChannel.cpp b/Engine/Source/Runtime/Engine/Private/DataChannel.cpp
--- a/Engine/Source/Runtime/Engine/Private/DataChannel.cpp
+++ b/Engine/Source/Runtime/Engine/Private/DataChannel.cpp
@@ -399,6 +399,9 @@ bool UChannel::ReceivedNextBunch( FInBunch & Bunch, bool & bOutSkipAck )
 	// Note this bunch's retirement.
 	if ( Bunch.bReliable )
 	{
+		// Reliables should be ordered properly at this point
+		check( Bunch.ChSequence == Connection->InReliable[Bunch.ChIndex] + 1 );
+
 		Connection->InReliable[Bunch.ChIndex] = Bunch.ChSequence;
 	}
 
@@ -465,6 +468,7 @@ bool UChannel::ReceivedNextBunch( FInBunch & Bunch, bool & bOutSkipAck )
 				// This is to ensure fast copies/appending of partial bunches. The final partial bunch may be non byte aligned.
 				check( Bunch.bHasGUIDs || Bunch.bPartialFinal || Bunch.GetBitsLeft() % 8 == 0 );
 
+				// Advance the sequence of the current partial bunch so we know what to expect next
 				InPartialBunch->ChSequence++;
 
 				if (Bunch.bPartialFinal)
@@ -482,16 +486,6 @@ bool UChannel::ReceivedNextBunch( FInBunch & Bunch, bool & bOutSkipAck )
 					InPartialBunch->bClose					= Bunch.bClose;
 					InPartialBunch->bDormant				= Bunch.bDormant;
 					InPartialBunch->bHasMustBeMappedGUIDs	= Bunch.bHasMustBeMappedGUIDs;
-
-					if (InPartialBunch->bOpen)
-					{
-						UE_LOG(LogNetPartialBunch, Verbose, TEXT("Channel: %d is now Open! [%d] - [%d]"), ChIndex, OpenPacketId.First, OpenPacketId.Last );
-
-						// If the open bunch was partial, set the ending range to the final partial bunch id
-						OpenPacketId.First = InPartialBunch->PacketId;
-						OpenPacketId.Last = Bunch.PacketId;
-						OpenAcked = true;
-					}
 				}
 				else
 				{
@@ -504,13 +498,13 @@ bool UChannel::ReceivedNextBunch( FInBunch & Bunch, bool & bOutSkipAck )
 			else
 			{
 				// Merge problem - delete InPartialBunch. This is mainly so that in the unlikely chance that ChSequence wraps around, we wont merge two completely separate partial bunches.
-				bOutSkipAck = true;
+				
+				bOutSkipAck = true;	// Don't ack the packet, since we didn't process the bunch
 
 				if ( InPartialBunch && InPartialBunch->bReliable )
 				{
 					check( !Bunch.bReliable );		// FIXME: Disconnect client in this case
 					UE_LOG( LogNetPartialBunch, Log, TEXT( "Unreliable partial trying to destroy reliable partial 2" ) );
-					bOutSkipAck = true;
 					return false;
 				}
 
@@ -537,9 +531,45 @@ bool UChannel::ReceivedNextBunch( FInBunch & Bunch, bool & bOutSkipAck )
 			return false;
 		}
 	}
-	
-	if (HandleBunch!=NULL)
+
+	if ( HandleBunch != NULL )
 	{
+		if ( HandleBunch->bOpen )
+		{
+			check( !OpenedLocally );					// If we opened the channel, we shouldn't be receiving bOpen commands from the other side
+			check( OpenPacketId.First == INDEX_NONE );	// This should be the first and only assignment of the packet range (we should only receive one bOpen bunch)
+			check( OpenPacketId.Last == INDEX_NONE );	// This should be the first and only assignment of the packet range (we should only receive one bOpen bunch)
+
+			// Remember the range.
+			// In the case of a non partial, HandleBunch == Bunch
+			// In the case of a partial, HandleBunch should == InPartialBunch, and Bunch should be the last bunch.
+			OpenPacketId.First = HandleBunch->PacketId;
+			OpenPacketId.Last = Bunch.PacketId;
+			OpenAcked = true;
+
+			UE_LOG( LogNetTraffic, Verbose, TEXT( "ReceivedNextBunch: Channel now fully open. ChIndex: %i, OpenPacketId.First: %i, OpenPacketId.Last: %i" ), ChIndex, OpenPacketId.First, OpenPacketId.Last );
+		}
+
+		// Don't process any packets until we've fully opened this channel 
+		// (unless we opened it locally, in which case it's safe to process packets)
+		if ( !OpenedLocally && !OpenAcked )
+		{
+			// If we receive a reliable at this point, this means reliables are out of order, which shouldn't be possible
+			check( !HandleBunch->bReliable );
+
+			// Don't ack this packet (since we won't process all of it)
+			bOutSkipAck = true;
+
+			UE_LOG( LogNetTraffic, Warning, TEXT( "ReceivedNextBunch: Skipping bunch since channel isn't fully open. ChIndex: %i" ), ChIndex );
+			return false;
+		}
+
+		// At this point, we should have the open packet range
+		// This is because if we opened the channel locally, we set it immediately when we sent the first bOpen bunch
+		// If we opened it from a remote connection, then we shouldn't be processing any packets until it's fully opened (which is handled above)
+		check( OpenPacketId.First != INDEX_NONE );
+		check( OpenPacketId.Last != INDEX_NONE );
+
 		// Receive it in sequence.
 		return ReceivedSequencedBunch( *HandleBunch );
 	}
@@ -725,7 +755,7 @@ FPacketIdRange UChannel::SendBunch( FOutBunch* Bunch, bool Merge )
 		}
 
 		FOutBunch *ThisOutBunch = PrepBunch(NextBunch, OutBunch, Merge); // This handles queuing reliable bunches into the ack list
-			
+
 		if (ThisOutBunch->bPartial && !ThisOutBunch->bReliable)
 		{
 			// If we are reliable, then partial bunches just use the reliable sequences
@@ -1483,7 +1513,7 @@ void UActorChannel::SetChannelActor( AActor* InActor )
 	// Set stuff.
 	Actor = InActor;
 
-	UE_LOG(LogNetTraffic, VeryVerbose, TEXT("SetChannelActor[%d]: Actor: %s"), ChIndex, Actor ? *Actor->GetFullName() : TEXT("NULL") );
+	UE_LOG(LogNetTraffic, VeryVerbose, TEXT("SetChannelActor: ChIndex: %i, Actor: %s, NetGUID: %s"), ChIndex, Actor ? *Actor->GetFullName() : TEXT("NULL"), *ActorNetGUID.ToString() );
 
 	if ( ChIndex >= 0 && Connection->PendingOutRec[ChIndex] > 0 )
 	{
@@ -1706,7 +1736,8 @@ void UActorChannel::ProcessBunch( FInBunch & Bunch )
 	{
 		if( !Bunch.bOpen )
 		{
-			UE_LOG(LogNetTraffic, Warning, TEXT("UActorChannel::ProcessBunch: New actor channel received non-open packet. bOpen: %i, bClose: %i, bReliable: %i"),Bunch.bOpen,Bunch.bClose,Bunch.bReliable);
+			// This absolutely shouldn't happen anymore, since we no longer process packets until channel is fully open early on
+			UE_LOG(LogNetTraffic, Error, TEXT( "UActorChannel::ProcessBunch: New actor channel received non-open packet. bOpen: %i, bClose: %i, bReliable: %i, bPartial: %i, bPartialInitial: %i, bPartialFinal: %i, ChType: %i, ChIndex: %i, Closing: %i, OpenedLocally: %i, OpenAcked: %i, NetGUID: %s" ), (int)Bunch.bOpen, (int)Bunch.bClose, (int)Bunch.bReliable, (int)Bunch.bPartial, (int)Bunch.bPartialInitial, (int)Bunch.bPartialFinal, (int)Bunch.ChType, ChIndex, (int)Closing, (int)OpenedLocally, (int)OpenAcked, *ActorNetGUID.ToString() );
 			return;
 		}
 
diff --git a/Engine/Source/Runtime/Engine/Private/NetConnection.cpp b/Engine/Source/Runtime/Engine/Private/NetConnection.cpp
--- a/Engine/Source/Runtime/Engine/Private/NetConnection.cpp
+++ b/Engine/Source/Runtime/Engine/Private/NetConnection.cpp
@@ -640,6 +640,12 @@ void UNetConnection::ReceivedPacket( FBitReader& Reader )
 	if( PacketId > InPacketId )
 	{
 		const int32 PacketsLost = PacketId - InPacketId - 1;
+		
+		if ( PacketsLost > 10 )
+		{
+			UE_LOG( LogNetTraffic, Warning, TEXT( "High single frame packet loss: %i" ), PacketsLost );
+		}
+
 		InPacketsLost += PacketsLost;
 		Driver->InPacketsLost += PacketsLost;
 		InPacketId = PacketId;
@@ -982,12 +988,6 @@ void UNetConnection::ReceivedPacket( FBitReader& Reader )
 				}
 			}
 
-			if( Bunch.bOpen && !Bunch.bPartial )
-			{
-				Channel->OpenAcked = 1;
-				Channel->OpenPacketId = FPacketIdRange(PacketId);
-			}
-
 			// Dispatch the raw, unsequenced bunch to the channel.
 			bool bLocalSkipAck = false;
 			Channel->ReceivedRawBunch( Bunch, bLocalSkipAck ); //warning: May destroy channel.
