Reuse packet sequence for unreliable partial sequence, fixes issue where partial sequence was overflowing due to using less bits for ChSequence

diff --git a/Engine/Source/Runtime/Engine/Classes/Engine/NetConnection.h b/Engine/Source/Runtime/Engine/Classes/Engine/NetConnection.h
--- a/Engine/Source/Runtime/Engine/Classes/Engine/NetConnection.h
+++ b/Engine/Source/Runtime/Engine/Classes/Engine/NetConnection.h
@@ -217,8 +217,6 @@ public:
 	int32			InPacketId;				// Full incoming packet index.
 	int32			OutPacketId;			// Most recently sent packet.
 	int32 			OutAckPacketId;			// Most recently acked outgoing packet.
-	int32			PartialPacketId;		// Sequencing number for partial packets
-	int32			LastPartialPacketId;	// Most recently received unrealiable partial packet id
 
 	uint32			PingAckDataCache[MAX_PACKETID/PING_ACK_PACKET_INTERVAL];	// Caches packet data on the server, for verifying pings
 	float			LastPingAck;												// The time of the most recent PingAck on the client
diff --git a/Engine/Source/Runtime/Engine/Private/DataChannel.cpp b/Engine/Source/Runtime/Engine/Private/DataChannel.cpp
--- a/Engine/Source/Runtime/Engine/Private/DataChannel.cpp
+++ b/Engine/Source/Runtime/Engine/Private/DataChannel.cpp
@@ -454,7 +454,15 @@ bool UChannel::ReceivedNextBunch( FInBunch & Bunch, bool & bOutSkipAck )
 			//  -ChSequence is next in partial sequence
 			//	-Reliability flag matches
 
-			if (InPartialBunch && !InPartialBunch->bPartialFinal && InPartialBunch->ChSequence+1 == Bunch.ChSequence && InPartialBunch->bReliable == Bunch.bReliable)
+			const bool bReliableSequencesMatches = Bunch.ChSequence == InPartialBunch->ChSequence + 1;
+			const bool bUnreliableSequenceMatches = bReliableSequencesMatches || ( Bunch.ChSequence == InPartialBunch->ChSequence );
+
+			// Unreliable partial bunches use the packet sequence, and since we can merge multiple bunches into a single packet,
+			// it's perfectly legal for the ChSequence to match in this case.
+			// Reliable partial bunches must be in consecutive order though
+			const bool bSequenceMatches = InPartialBunch->bReliable ? bReliableSequencesMatches : bUnreliableSequenceMatches;
+
+			if ( InPartialBunch && !InPartialBunch->bPartialFinal && bSequenceMatches && InPartialBunch->bReliable == Bunch.bReliable )
 			{
 				// Merge.
 				UE_LOG(LogNetPartialBunch, Verbose, TEXT("Merging Partial Bunch: %d Bytes"), Bunch.GetBytesLeft() );
@@ -469,7 +477,7 @@ bool UChannel::ReceivedNextBunch( FInBunch & Bunch, bool & bOutSkipAck )
 				check( Bunch.bHasGUIDs || Bunch.bPartialFinal || Bunch.GetBitsLeft() % 8 == 0 );
 
 				// Advance the sequence of the current partial bunch so we know what to expect next
-				InPartialBunch->ChSequence++;
+				InPartialBunch->ChSequence = Bunch.ChSequence;
 
 				if (Bunch.bPartialFinal)
 				{
@@ -756,13 +764,6 @@ FPacketIdRange UChannel::SendBunch( FOutBunch* Bunch, bool Merge )
 
 		FOutBunch *ThisOutBunch = PrepBunch(NextBunch, OutBunch, Merge); // This handles queuing reliable bunches into the ack list
 
-		if (ThisOutBunch->bPartial && !ThisOutBunch->bReliable)
-		{
-			// If we are reliable, then partial bunches just use the reliable sequences
-			// if not reliable, we hijack ChSequence and use Connection->PartialPacketId to sequence these packets
-			ThisOutBunch->ChSequence = Connection->PartialPacketId++;
-		}
-
 		if (UE_LOG_ACTIVE(LogNetPartialBunch,Verbose) && (OutgoingBunches.Num() > 1)) // Don't want to call appMemcrc unless we need to
 		{
 			UE_LOG(LogNetPartialBunch, Verbose, TEXT("	Bunch[%d]: Bytes: %d Bits: %d ChSequence: %d 0x%X"), PartialNum, ThisOutBunch->GetNumBytes(), ThisOutBunch->GetNumBits(), ThisOutBunch->ChSequence, FCrc::MemCrc_DEPRECATED(ThisOutBunch->GetData(), ThisOutBunch->GetNumBytes()));
diff --git a/Engine/Source/Runtime/Engine/Private/NetConnection.cpp b/Engine/Source/Runtime/Engine/Private/NetConnection.cpp
--- a/Engine/Source/Runtime/Engine/Private/NetConnection.cpp
+++ b/Engine/Source/Runtime/Engine/Private/NetConnection.cpp
@@ -57,8 +57,6 @@ UNetConnection::UNetConnection(const FObjectInitializer& ObjectInitializer)
 ,	InPacketId			( -1 )
 ,	OutPacketId			( 0 ) // must be initialized as OutAckPacketId + 1 so loss of first packet can be detected
 ,	OutAckPacketId		( -1 )
-,	PartialPacketId		( 0 )
-,	LastPartialPacketId	( -1 )
 ,	LastPingAck			( 0.f )
 ,	LastPingAckPacketId	( -1 )
 ,	ClientWorldPackageName( NAME_None )
@@ -815,19 +813,8 @@ void UNetConnection::ReceivedPacket( FBitReader& Reader )
 			} 
 			else if ( Bunch.bPartial )
 			{
-				// If this is a partial bunch, use the last processed partial sequence to read the new partial sequence
-				Bunch.ChSequence = MakeRelative( Reader.ReadInt( MAX_CHSEQUENCE ), LastPartialPacketId, MAX_CHSEQUENCE );
-				if ( Bunch.ChSequence > LastPartialPacketId )
-				{
-					LastPartialPacketId = Bunch.ChSequence;
-				}
-				else
-				{
-					// Well behaved clients should never hit this, since we've already thrown out packets that were older
-					// Since PartialPacketId should evolve in sync with PacketId, this should be impossible unless the 
-					// client is misbehaving...
-					UE_LOG( LogNetTraffic, Error, TEXT( "UNetConnection::ReceivedPacket: Invalid partial packet id." ) );
-				}
+				// If this is an unreliable partial bunch, we simply use packet sequence since we already have it
+				Bunch.ChSequence = PacketId;
 			}
 			else
 			{
@@ -1156,15 +1143,18 @@ int32 UNetConnection::SendRawBunch( FOutBunch& Bunch, bool InAllowMerge )
 	Header.WriteBit( Bunch.bHasGUIDs );
 	Header.WriteBit( Bunch.bHasMustBeMappedGUIDs );
 	Header.WriteBit( Bunch.bPartial );
-	if (Bunch.bReliable || Bunch.bPartial)
+
+	if ( Bunch.bReliable )
 	{
 		Header.WriteIntWrapped(Bunch.ChSequence, MAX_CHSEQUENCE);
-		if (Bunch.bPartial)
-		{
-			Header.WriteBit( Bunch.bPartialInitial );
-			Header.WriteBit( Bunch.bPartialFinal );
-		}
 	}
+
+	if (Bunch.bPartial)
+	{
+		Header.WriteBit( Bunch.bPartialInitial );
+		Header.WriteBit( Bunch.bPartialFinal );
+	}
+
 	if (Bunch.bReliable || Bunch.bOpen)
 	{
 		Header.WriteIntWrapped(Bunch.ChType, CHTYPE_MAX);
