#UE4 Fixed build error due to editable keyword on deprecated property.

diff --git a/Engine/Source/Runtime/Engine/Classes/Engine/Scene.h b/Engine/Source/Runtime/Engine/Classes/Engine/Scene.h
--- a/Engine/Source/Runtime/Engine/Classes/Engine/Scene.h
+++ b/Engine/Source/Runtime/Engine/Classes/Engine/Scene.h
@@ -205,7 +205,7 @@ struct FPostProcessSettings
 	UPROPERTY(BlueprintReadWrite, Category=Overrides, meta=(PinHiddenByDefault))
 	uint32 bOverride_AmbientOcclusionFadeRadius:1;
 
-	UPROPERTY(BlueprintReadWrite, Category=Overrides, meta=(PinHiddenByDefault))
+	UPROPERTY()
 	uint32 bOverride_AmbientOcclusionDistance_DEPRECATED:1;
 
 	UPROPERTY(BlueprintReadWrite, Category=Overrides, meta=(PinHiddenByDefault))
