Forward declare UEdGraph to make Mac compile happy

diff --git a/Engine/Source/Runtime/Engine/Classes/Sound/SoundCue.h b/Engine/Source/Runtime/Engine/Classes/Sound/SoundCue.h
--- a/Engine/Source/Runtime/Engine/Classes/Sound/SoundCue.h
+++ b/Engine/Source/Runtime/Engine/Classes/Sound/SoundCue.h
@@ -62,7 +62,7 @@ class USoundCue : public USoundBase
 	TArray<USoundNode*> AllNodes;
 
 	UPROPERTY()
-	UEdGraph* SoundCueGraph;
+	class UEdGraph* SoundCueGraph;
 
 #endif
 
