nicer name

diff --git a/Engine/Source/Runtime/Engine/Public/ShowFlagsValues.inl b/Engine/Source/Runtime/Engine/Public/ShowFlagsValues.inl
--- a/Engine/Source/Runtime/Engine/Public/ShowFlagsValues.inl
+++ b/Engine/Source/Runtime/Engine/Public/ShowFlagsValues.inl
@@ -31,7 +31,7 @@ SHOWFLAG_ALWAYS_ACCESSIBLE(AmbientCubemap, SFG_LightingFeatures, LOCTEXT("Ambien
 /** Human like eye simulation to adapt to the brightness of the view */
 SHOWFLAG_ALWAYS_ACCESSIBLE(EyeAdaptation, SFG_PostProcess, LOCTEXT("EyeAdaptationSF", "Eye Adaptation"))
 /** Display a histogram of the scene HDR color */
-SHOWFLAG_ALWAYS_ACCESSIBLE(VisualizeHDR, SFG_Visualize, LOCTEXT("VisualizeHDRSF", "HDR"))
+SHOWFLAG_ALWAYS_ACCESSIBLE(VisualizeHDR, SFG_Visualize, LOCTEXT("VisualizeHDRSF", "HDR (Eye Adaptation)"))
 /** Image based lens flares (Simulate artifact of reflections within a camera system) */
 SHOWFLAG_ALWAYS_ACCESSIBLE(LensFlares, SFG_PostProcess, LOCTEXT("LensFlaresSF", "Lens Flares"))
 /** show indirect lighting component */
