Fixed issue with Sound Cue Editor not showing execution path

Changed FindPathToNode function to start with Hash made from first node as the Parse function was changed to use this in CL#2347435

diff --git a/Engine/Source/Runtime/Engine/Private/SoundCue.cpp b/Engine/Source/Runtime/Engine/Private/SoundCue.cpp
--- a/Engine/Source/Runtime/Engine/Private/SoundCue.cpp
+++ b/Engine/Source/Runtime/Engine/Private/SoundCue.cpp
@@ -162,7 +162,7 @@ bool USoundCue::RecursiveFindPathToNode(USoundNode* CurrentNode, const UPTRINT C
 
 bool USoundCue::FindPathToNode(const UPTRINT NodeHashToFind, TArray<USoundNode*>& OutPath) const
 {
-	return RecursiveFindPathToNode(FirstNode, 0, NodeHashToFind, OutPath);
+	return RecursiveFindPathToNode(FirstNode, (UPTRINT)FirstNode, NodeHashToFind, OutPath);
 }
 
 
