Delegate comparisons deprecated, lots of other associated code deprecated, and lots of warning fixups:

diff --git a/Engine/Source/Runtime/Engine/Private/UnrealEngine.cpp b/Engine/Source/Runtime/Engine/Private/UnrealEngine.cpp
--- a/Engine/Source/Runtime/Engine/Private/UnrealEngine.cpp
+++ b/Engine/Source/Runtime/Engine/Private/UnrealEngine.cpp
@@ -407,9 +407,11 @@ void RefreshEngineSettings()
 	SystemResolutionSinkCallback();
 }
 
+FConsoleVariableSinkHandle GRefreshEngineSettingsSinkHandle;
+
 ENGINE_API void InitializeRenderingCVarsCaching()
 {
-	IConsoleManager::Get().RegisterConsoleVariableSink(FConsoleCommandDelegate::CreateStatic(&RefreshEngineSettings));
+	GRefreshEngineSettingsSinkHandle = IConsoleManager::Get().RegisterConsoleVariableSink_Handle(FConsoleCommandDelegate::CreateStatic(&RefreshEngineSettings));
 
 	// Initialise this to invalid
 	GCachedScalabilityCVars.MaterialQualityLevel = EMaterialQualityLevel::Num;
@@ -421,7 +423,7 @@ ENGINE_API void InitializeRenderingCVarsCaching()
 
 void ShutdownRenderingCVarsCaching()
 {
-	IConsoleManager::Get().UnregisterConsoleVariableSink(FConsoleCommandDelegate::CreateStatic(&RefreshEngineSettings));
+	IConsoleManager::Get().UnregisterConsoleVariableSink_Handle(GRefreshEngineSettingsSinkHandle);
 }
 
 namespace
@@ -833,7 +835,7 @@ void UEngine::Init(IEngineLoop* InEngineLoop)
 			FOnExternalUIChangeDelegate OnExternalUIChangeDelegate;
 			OnExternalUIChangeDelegate.BindUObject(this, &UEngine::OnExternalUIChange);
 
-			ExternalUI->AddOnExternalUIChangeDelegate(OnExternalUIChangeDelegate);
+			ExternalUI->AddOnExternalUIChangeDelegate_Handle(OnExternalUIChangeDelegate);
 		}
 	}
 
@@ -847,7 +849,7 @@ void UEngine::Init(IEngineLoop* InEngineLoop)
 	// register screenshot capture if we are dumping a movie
 	if(GIsDumpingMovie)
 	{
-		UGameViewportClient::OnScreenshotCaptured().AddUObject(this, &UEngine::HandleScreenshotCaptured);
+		HandleScreenshotCapturedDelegateHandle = UGameViewportClient::OnScreenshotCaptured().AddUObject(this, &UEngine::HandleScreenshotCaptured);
 	}
 #endif
 
@@ -931,7 +933,7 @@ void UEngine::PreExit()
 	FEngineAnalytics::Shutdown();
 
 #if WITH_EDITOR
-	UGameViewportClient::OnScreenshotCaptured().RemoveUObject(this, &UEngine::HandleScreenshotCaptured);
+	UGameViewportClient::OnScreenshotCaptured().Remove(HandleScreenshotCapturedDelegateHandle);
 #endif
 
 	if (ScreenSaverInhibitor)
diff --git a/Engine/Source/Runtime/Engine/Private/World.cpp b/Engine/Source/Runtime/Engine/Private/World.cpp
--- a/Engine/Source/Runtime/Engine/Private/World.cpp
+++ b/Engine/Source/Runtime/Engine/Private/World.cpp
@@ -3195,14 +3195,19 @@ void UWorld::RemoveNetworkActor( AActor* Actor )
 	NetworkActors.RemoveSingleSwap( Actor );
 }
 
-void UWorld::AddOnActorSpawnedHandler( const FOnActorSpawned::FDelegate& InHandler )
+FDelegateHandle UWorld::AddOnActorSpawnedHandler( const FOnActorSpawned::FDelegate& InHandler )
 {
-	OnActorSpawned.Add(InHandler);
+	return OnActorSpawned.Add(InHandler);
 }
 
 void UWorld::RemoveOnActorSpawnedHandler( const FOnActorSpawned::FDelegate& InHandler )
 {
-	OnActorSpawned.Remove(InHandler);
+	OnActorSpawned.DEPRECATED_Remove(InHandler);
+}
+
+void UWorld::RemoveOnActorSpawnedHandler( FDelegateHandle InHandle )
+{
+	OnActorSpawned.Remove(InHandle);
 }
 
 ABrush* UWorld::GetBrush() const
diff --git a/Engine/Source/Runtime/Engine/Public/DebugRenderSceneProxy.h b/Engine/Source/Runtime/Engine/Public/DebugRenderSceneProxy.h
--- a/Engine/Source/Runtime/Engine/Public/DebugRenderSceneProxy.h
+++ b/Engine/Source/Runtime/Engine/Public/DebugRenderSceneProxy.h
@@ -235,6 +235,7 @@ public:
 	FString ViewFlagName;
 	float TextWithoutShadowDistance;
 	FDebugDrawDelegate DebugTextDrawingDelegate;
+	FDelegateHandle DebugTextDrawingDelegateHandle;
 	EDrawType DrawType;
 };
 
diff --git a/Engine/Source/Runtime/Engine/Public/EngineUtils.h b/Engine/Source/Runtime/Engine/Public/EngineUtils.h
--- a/Engine/Source/Runtime/Engine/Public/EngineUtils.h
+++ b/Engine/Source/Runtime/Engine/Public/EngineUtils.h
@@ -165,6 +165,8 @@ protected:
 	TArray<AActor*> SpawnedActorArray;
 	/** The class type we are iterating, kept for filtering		*/
 	UClass* DesiredClass;
+	/** Handle to the registered OnActorSpawned delegate		*/
+	FDelegateHandle ActorSpawnedDelegateHandle;
 
 	/**
 	 * Default ctor, inits everything
@@ -183,12 +185,12 @@ protected:
 		GetObjectsOfClass(InClass, ObjectArray, true, ExcludeFlags);
 
 		ActorSpawnedDelegate = FOnActorSpawned::FDelegate::CreateRaw(this, &FActorIteratorBase::OnActorSpawned);
-		CurrentWorld->AddOnActorSpawnedHandler(ActorSpawnedDelegate);
+		ActorSpawnedDelegateHandle = CurrentWorld->AddOnActorSpawnedHandler(ActorSpawnedDelegate);
 	}
 
 	~FActorIteratorBase()
 	{
-		CurrentWorld->RemoveOnActorSpawnedHandler(ActorSpawnedDelegate);
+		CurrentWorld->RemoveOnActorSpawnedHandler(ActorSpawnedDelegateHandle);
 	}
 
 	void OnActorSpawned(AActor* InActor)
diff --git a/Engine/Source/Runtime/Engine/Public/TimerManager.h b/Engine/Source/Runtime/Engine/Public/TimerManager.h
--- a/Engine/Source/Runtime/Engine/Public/TimerManager.h
+++ b/Engine/Source/Runtime/Engine/Public/TimerManager.h
@@ -39,7 +39,7 @@ struct FTimerUnifiedDelegate
 		}
 		else if (FuncDynDelegate.IsBound())
 		{
-			FuncDynDelegate.ProcessDelegate<UObject>(NULL);
+			FuncDynDelegate.ProcessDelegate<UObject>(nullptr);
 		}
 	}
 
@@ -67,11 +67,6 @@ struct FTimerUnifiedDelegate
 		FuncDelegate.Unbind();
 		FuncDynDelegate.Unbind();
 	}
-
-	inline bool operator==(const FTimerUnifiedDelegate& Other) const
-	{
-		return ( FuncDelegate.IsBound() && (FuncDelegate == Other.FuncDelegate) ) || ( FuncDynDelegate.IsBound() && (FuncDynDelegate == Other.FuncDynDelegate) );
-	}
 };
 
 // Unique handle that can be used to distinguish timers that have identical delegates.
@@ -194,49 +189,53 @@ public:
 	 * @param InbLoop				true to keep firing at Rate intervals, false to fire only once.
 	 * @param InFirstDelay			The time for the first iteration of a looping timer. If < 0.f inRate will be used.
 	 */
-	template< class UserClass >	
+	template< class UserClass >
+	DELEGATE_DEPRECATED("This overload of SetTimer is deprecated, please pass a FTimerHandle as the first argument.")
 	FORCEINLINE void SetTimer(UserClass* InObj, typename FTimerDelegate::TUObjectMethodDelegate< UserClass >::FMethodPtr InTimerMethod, float InRate, bool InbLoop = false, float InFirstDelay = -1.f)
 	{
-		SetTimer( FTimerDelegate::CreateUObject(InObj, InTimerMethod), InRate, InbLoop, InFirstDelay );
+		DEPRECATED_InternalSetTimer( FTimerUnifiedDelegate( FTimerDelegate::CreateUObject(InObj, InTimerMethod) ), InRate, InbLoop, InFirstDelay );
 	}
-	template< class UserClass >	
+	template< class UserClass >
+	DELEGATE_DEPRECATED("This overload of SetTimer is deprecated, please pass a FTimerHandle as the first argument.")
 	FORCEINLINE void SetTimer(UserClass* InObj, typename FTimerDelegate::TUObjectMethodDelegate_Const< UserClass >::FMethodPtr InTimerMethod, float InRate, bool InbLoop = false, float InFirstDelay = -1.f)
 	{
-		SetTimer( FTimerDelegate::CreateUObject(InObj, InTimerMethod), InRate, InbLoop, InFirstDelay );
+		DEPRECATED_InternalSetTimer( FTimerUnifiedDelegate( FTimerDelegate::CreateUObject(InObj, InTimerMethod) ), InRate, InbLoop, InFirstDelay );
 	}
 
 	/** Version that takes any generic delegate. */
+	DELEGATE_DEPRECATED("This overload of SetTimer is deprecated, please pass a FTimerHandle as the first argument.")
 	FORCEINLINE void SetTimer(FTimerDelegate const& InDelegate, float InRate, bool InbLoop, float InFirstDelay = -1.f)
 	{
-		InternalSetTimer( FTimerUnifiedDelegate(InDelegate), InRate, InbLoop, InFirstDelay );
+		DEPRECATED_InternalSetTimer( FTimerUnifiedDelegate(InDelegate), InRate, InbLoop, InFirstDelay );
 	}
 	/** Version that takes a dynamic delegate (e.g. for UFunctions). */
+	DELEGATE_DEPRECATED("This overload of SetTimer is deprecated, please pass a FTimerHandle as the first argument.")
 	FORCEINLINE void SetTimer(FTimerDynamicDelegate const& InDynDelegate, float InRate, bool InbLoop, float InFirstDelay = -1.f)
 	{
-		InternalSetTimer( FTimerUnifiedDelegate(InDynDelegate), InRate, InbLoop, InFirstDelay );
+		DEPRECATED_InternalSetTimer( FTimerUnifiedDelegate(InDynDelegate), InRate, InbLoop, InFirstDelay );
 	}
 
 	/**
-	* Sets a timer to call the given native function at a set interval.  If a timer is already set
-	* for this delegate, it will update the current timer to the new parameters and reset its
-	* elapsed time to 0.
-	*
-	* @param InOutHandle			Handle to identify this timer. If it is invalid when passed in it will be made into a valid handle.
-	* @param InObj					Object to call the timer function on.
-	* @param InTimerMethod			Method to call when timer fires.
-	* @param InRate					The amount of time between set and firing.  If <= 0.f, clears existing timers.
-	* @param InbLoop				true to keep firing at Rate intervals, false to fire only once.
-	* @param InFirstDelay			The time for the first iteration of a looping timer. If < 0.f inRate will be used.
-	*/
+	 * Sets a timer to call the given native function at a set interval.  If a timer is already set
+	 * for this delegate, it will update the current timer to the new parameters and reset its
+	 * elapsed time to 0.
+	 *
+	 * @param InOutHandle			Handle to identify this timer. If it is invalid when passed in it will be made into a valid handle.
+	 * @param InObj					Object to call the timer function on.
+	 * @param InTimerMethod			Method to call when timer fires.
+	 * @param InRate					The amount of time between set and firing.  If <= 0.f, clears existing timers.
+	 * @param InbLoop				true to keep firing at Rate intervals, false to fire only once.
+	 * @param InFirstDelay			The time for the first iteration of a looping timer. If < 0.f inRate will be used.
+	 */
 	template< class UserClass >
 	FORCEINLINE void SetTimer(FTimerHandle& InOutHandle, UserClass* InObj, typename FTimerDelegate::TUObjectMethodDelegate< UserClass >::FMethodPtr InTimerMethod, float InRate, bool InbLoop = false, float InFirstDelay = -1.f)
 	{
-		SetTimer(InOutHandle, FTimerDelegate::CreateUObject(InObj, InTimerMethod), InRate, InbLoop, InFirstDelay);
+		InternalSetTimer(InOutHandle, FTimerUnifiedDelegate( FTimerDelegate::CreateUObject(InObj, InTimerMethod) ), InRate, InbLoop, InFirstDelay);
 	}
 	template< class UserClass >
 	FORCEINLINE void SetTimer(FTimerHandle& InOutHandle, UserClass* InObj, typename FTimerDelegate::TUObjectMethodDelegate_Const< UserClass >::FMethodPtr InTimerMethod, float InRate, bool InbLoop = false, float InFirstDelay = -1.f)
 	{
-		SetTimer(InOutHandle, FTimerDelegate::CreateUObject(InObj, InTimerMethod), InRate, InbLoop, InFirstDelay);
+		InternalSetTimer(InOutHandle, FTimerUnifiedDelegate( FTimerDelegate::CreateUObject(InObj, InTimerMethod) ), InRate, InbLoop, InFirstDelay);
 	}
 
 	/** Version that takes any generic delegate. */
@@ -256,20 +255,20 @@ public:
 	}
 
 	/**
-	* Sets a timer to call the given native function on the next tick.
-	*
-	* @param inObj					Object to call the timer function on.
-	* @param inTimerMethod			Method to call when timer fires.
-	*/
+	 * Sets a timer to call the given native function on the next tick.
+	 *
+	 * @param inObj					Object to call the timer function on.
+	 * @param inTimerMethod			Method to call when timer fires.
+	 */
 	template< class UserClass >
 	FORCEINLINE void SetTimerForNextTick(UserClass* inObj, typename FTimerDelegate::TUObjectMethodDelegate< UserClass >::FMethodPtr inTimerMethod)
 	{
-		SetTimerForNextTick(FTimerDelegate::CreateUObject(inObj, inTimerMethod));
+		InternalSetTimerForNextTick(FTimerUnifiedDelegate(FTimerDelegate::CreateUObject(inObj, inTimerMethod)));
 	}
 	template< class UserClass >
 	FORCEINLINE void SetTimerForNextTick(UserClass* inObj, typename FTimerDelegate::TUObjectMethodDelegate_Const< UserClass >::FMethodPtr inTimerMethod)
 	{
-		SetTimerForNextTick(FTimerDelegate::CreateUObject(inObj, inTimerMethod));
+		InternalSetTimerForNextTick(FTimerUnifiedDelegate(FTimerDelegate::CreateUObject(inObj, inTimerMethod)));
 	}
 
 	/** Version that takes any generic delegate. */
@@ -290,29 +289,33 @@ public:
 	 * @param inObj			Object to call the timer function on.
 	 * @param inTimerMethod Method to call when timer fires
 	 */
-	template< class UserClass >	
+	template< class UserClass >
+	DELEGATE_DEPRECATED("This overload of ClearTimer is deprecated, use ClearTimer(FTimerHandle InHandle) instead.")
 	FORCEINLINE void ClearTimer( UserClass* inObj, typename FTimerDelegate::TUObjectMethodDelegate< UserClass >::FMethodPtr inTimerMethod)
 	{
-		ClearTimer( FTimerDelegate::CreateUObject(inObj, inTimerMethod) );
+		DEPRECATED_InternalClearTimer( FTimerUnifiedDelegate( FTimerDelegate::CreateUObject(inObj, inTimerMethod) ) );
 	}
-	template< class UserClass >	
+	template< class UserClass >
+	DELEGATE_DEPRECATED("This overload of ClearTimer is deprecated, use ClearTimer(FTimerHandle InHandle) instead.")
 	FORCEINLINE void ClearTimer( UserClass* inObj, typename FTimerDelegate::TUObjectMethodDelegate_Const< UserClass >::FMethodPtr inTimerMethod)
 	{
-		ClearTimer( FTimerDelegate::CreateUObject(inObj, inTimerMethod) );
+		DEPRECATED_InternalClearTimer( FTimerUnifiedDelegate( FTimerDelegate::CreateUObject(inObj, inTimerMethod) ) );
 	}
 
 	/** Version that takes any generic delegate. */
+	DELEGATE_DEPRECATED("This overload of ClearTimer is deprecated, use ClearTimer(FTimerHandle InHandle) instead.")
 	FORCEINLINE void ClearTimer(FTimerDelegate const& InDelegate)
 	{
-		InternalClearTimer( FTimerUnifiedDelegate(InDelegate) );
+		DEPRECATED_InternalClearTimer( FTimerUnifiedDelegate(InDelegate) );
 	}
 	/** Version that takes a dynamic delegate (e.g. for UFunctions). */
+	DELEGATE_DEPRECATED("This overload of ClearTimer is deprecated, use ClearTimer(FTimerHandle InHandle) instead.")
 	FORCEINLINE void ClearTimer(FTimerDynamicDelegate const& InDynDelegate)
 	{
-		InternalClearTimer( FTimerUnifiedDelegate(InDynDelegate) );
+		DEPRECATED_InternalClearTimer( FTimerUnifiedDelegate(InDynDelegate) );
 	}
 	/** Version that takes a handle */
-	FORCEINLINE void ClearTimer(FTimerHandle const& InHandle)
+	FORCEINLINE void ClearTimer(FTimerHandle InHandle)
 	{
 		InternalClearTimer(InHandle);
 	}
@@ -333,33 +336,41 @@ public:
 	 * @param inObj				Object to call the timer function on.
 	 * @param inTimerMethod		Method to call when timer fires.
 	 */
-	template< class UserClass >	
+	template< class UserClass >
+	DELEGATE_DEPRECATED("This overload of PauseTimer is deprecated, use PauseTimer(FTimerHandle InHandle) instead.")
 	FORCEINLINE void PauseTimer(UserClass* inObj, typename FTimerDelegate::TUObjectMethodDelegate< UserClass >::FMethodPtr inTimerMethod)
 	{
-		PauseTimer( FTimerDelegate::CreateUObject(inObj, inTimerMethod) );
+		int32 TimerIdx;
+		FTimerData const* TimerToPause = DEPRECATED_FindTimer(FTimerUnifiedDelegate( FTimerDelegate::CreateUObject(inObj, inTimerMethod) ), &TimerIdx);
+		InternalPauseTimer(TimerToPause, TimerIdx);
 	}
-	template< class UserClass >	
+	template< class UserClass >
+	DELEGATE_DEPRECATED("This overload of PauseTimer is deprecated, use PauseTimer(FTimerHandle InHandle) instead.")
 	FORCEINLINE void PauseTimer(UserClass* inObj, typename FTimerDelegate::TUObjectMethodDelegate_Const< UserClass >::FMethodPtr inTimerMethod)
 	{
-		PauseTimer( FTimerDelegate::CreateUObject(inObj, inTimerMethod) );
+		int32 TimerIdx;
+		FTimerData const* TimerToPause = DEPRECATED_FindTimer(FTimerUnifiedDelegate( FTimerDelegate::CreateUObject(inObj, inTimerMethod) ), &TimerIdx);
+		InternalPauseTimer(TimerToPause, TimerIdx);
 	}
 
 	/** Version that takes any generic delegate. */
+	DELEGATE_DEPRECATED("This overload of PauseTimer is deprecated, use PauseTimer(FTimerHandle InHandle) instead.")
 	FORCEINLINE void PauseTimer(FTimerDelegate const& InDelegate)
 	{
 		int32 TimerIdx;
-		FTimerData const* TimerToPause = FindTimer(FTimerUnifiedDelegate(InDelegate), &TimerIdx);
+		FTimerData const* TimerToPause = DEPRECATED_FindTimer(FTimerUnifiedDelegate(InDelegate), &TimerIdx);
 		InternalPauseTimer(TimerToPause, TimerIdx);
 	}
 	/** Version that takes a dynamic delegate (e.g. for UFunctions). */
+	DELEGATE_DEPRECATED("This overload of PauseTimer is deprecated, use PauseTimer(FTimerHandle InHandle) instead.")
 	FORCEINLINE void PauseTimer(FTimerDynamicDelegate const& InDynDelegate)
 	{
 		int32 TimerIdx;
-		FTimerData const* TimerToPause = FindTimer(FTimerUnifiedDelegate(InDynDelegate), &TimerIdx);
+		FTimerData const* TimerToPause = DEPRECATED_FindTimer(FTimerUnifiedDelegate(InDynDelegate), &TimerIdx);
 		InternalPauseTimer(TimerToPause, TimerIdx);
 	}
 	/** Version that takes a handle */
-	FORCEINLINE void PauseTimer(FTimerHandle const& InHandle)
+	FORCEINLINE void PauseTimer(FTimerHandle InHandle)
 	{
 		int32 TimerIdx;
 		FTimerData const* TimerToPause = FindTimer(InHandle, &TimerIdx);
@@ -372,31 +383,37 @@ public:
 	 * @param inObj				Object to call the timer function on.
 	 * @param inTimerMethod		Method to call when timer fires.
 	 */
-	template< class UserClass >	
+	template< class UserClass >
+	DELEGATE_DEPRECATED("This overload of UnPauseTimer is deprecated, use UnPauseTimer(FTimerHandle InHandle) instead.")
 	FORCEINLINE void UnPauseTimer(UserClass* inObj, typename FTimerDelegate::TUObjectMethodDelegate< UserClass >::FMethodPtr inTimerMethod)
 	{
-		UnPauseTimer( FTimerDelegate::CreateUObject(inObj, inTimerMethod) );
+		int32 TimerIdx = DEPRECATED_FindTimerInList( PausedTimerList, FTimerUnifiedDelegate( FTimerDelegate::CreateUObject(inObj, inTimerMethod) ) );
+		InternalUnPauseTimer(TimerIdx);
 	}
-	template< class UserClass >	
+	template< class UserClass >
+	DELEGATE_DEPRECATED("This overload of UnPauseTimer is deprecated, use UnPauseTimer(FTimerHandle InHandle) instead.")
 	FORCEINLINE void UnPauseTimer(UserClass* inObj, typename FTimerDelegate::TUObjectMethodDelegate_Const< UserClass >::FMethodPtr inTimerMethod)
 	{
-		UnPauseTimer( FTimerDelegate::CreateUObject(inObj, inTimerMethod) );
+		int32 TimerIdx = DEPRECATED_FindTimerInList( PausedTimerList, FTimerUnifiedDelegate( FTimerDelegate::CreateUObject(inObj, inTimerMethod) ) );
+		InternalUnPauseTimer(TimerIdx);
 	}
 
 	/** Version that takes any generic delegate. */
+	DELEGATE_DEPRECATED("This overload of UnPauseTimer is deprecated, use UnPauseTimer(FTimerHandle InHandle) instead.")
 	FORCEINLINE void UnPauseTimer(FTimerDelegate const& InDelegate)
 	{
-		int32 TimerIdx = FindTimerInList( PausedTimerList, FTimerUnifiedDelegate(InDelegate) );
+		int32 TimerIdx = DEPRECATED_FindTimerInList( PausedTimerList, FTimerUnifiedDelegate(InDelegate) );
 		InternalUnPauseTimer(TimerIdx);
 	}
 	/** Version that takes a dynamic delegate (e.g. for UFunctions). */
+	DELEGATE_DEPRECATED("This overload of UnPauseTimer is deprecated, use UnPauseTimer(FTimerHandle InHandle) instead.")
 	FORCEINLINE void UnPauseTimer(FTimerDynamicDelegate const& InDynDelegate)
 	{
-		int32 TimerIdx = FindTimerInList( PausedTimerList, FTimerUnifiedDelegate(InDynDelegate) );
+		int32 TimerIdx = DEPRECATED_FindTimerInList( PausedTimerList, FTimerUnifiedDelegate(InDynDelegate) );
 		InternalUnPauseTimer(TimerIdx);
 	}
 	/** Version that takes a handle */
-	FORCEINLINE void UnPauseTimer(FTimerHandle const& InHandle)
+	FORCEINLINE void UnPauseTimer(FTimerHandle InHandle)
 	{
 		int32 TimerIdx = FindTimerInList(PausedTimerList, InHandle);
 		InternalUnPauseTimer(TimerIdx);
@@ -409,34 +426,40 @@ public:
 	 * @param inTimerMethod Method to call when timer fires.
 	 * @return				The current rate or -1.f if timer does not exist.
 	 */
-	template< class UserClass >	
+	template< class UserClass >
+	DELEGATE_DEPRECATED("This overload of GetTimerRate is deprecated, use GetTimerRate(FTimerHandle InHandle) instead.")
 	FORCEINLINE float GetTimerRate(UserClass* inObj, typename FTimerDelegate::TUObjectMethodDelegate< UserClass >::FMethodPtr inTimerMethod) const
 	{
-		return GetTimerRate( FTimerDelegate::CreateUObject(inObj, inTimerMethod) );
+		FTimerData const* const TimerData = DEPRECATED_FindTimer( FTimerUnifiedDelegate( FTimerDelegate::CreateUObject(inObj, inTimerMethod) ) );
+		return InternalGetTimerRate(TimerData);
 	}
 
-	template< class UserClass >	
+	template< class UserClass >
+	DELEGATE_DEPRECATED("This overload of GetTimerRate is deprecated, use GetTimerRate(FTimerHandle InHandle) instead.")
 	FORCEINLINE float GetTimerRate(UserClass* inObj, typename FTimerDelegate::TUObjectMethodDelegate_Const< UserClass >::FMethodPtr inTimerMethod) const
 	{
-		return GetTimerRate( FTimerDelegate::CreateUObject(inObj, inTimerMethod) );
+		FTimerData const* const TimerData = DEPRECATED_FindTimer( FTimerUnifiedDelegate( FTimerDelegate::CreateUObject(inObj, inTimerMethod) ) );
+		return InternalGetTimerRate(TimerData);
 	}
 
 	/** Version that takes any generic delegate. */
+	DELEGATE_DEPRECATED("This overload of GetTimerRate is deprecated, use GetTimerRate(FTimerHandle InHandle) instead.")
 	FORCEINLINE float GetTimerRate(FTimerDelegate const& InDelegate) const
 	{
-		FTimerData const* const TimerData = FindTimer( FTimerUnifiedDelegate(InDelegate) );
+		FTimerData const* const TimerData = DEPRECATED_FindTimer( FTimerUnifiedDelegate(InDelegate) );
 		return InternalGetTimerRate(TimerData);
 	}
 
 	/** Version that takes a dynamic delegate (e.g. for UFunctions). */
+	DELEGATE_DEPRECATED("This overload of GetTimerRate is deprecated, use GetTimerRate(FTimerHandle InHandle) instead.")
 	FORCEINLINE float GetTimerRate(FTimerDynamicDelegate const& InDynDelegate) const
 	{
-		FTimerData const* const TimerData = FindTimer( FTimerUnifiedDelegate(InDynDelegate) );
+		FTimerData const* const TimerData = DEPRECATED_FindTimer( FTimerUnifiedDelegate(InDynDelegate) );
 		return InternalGetTimerRate(TimerData);
 	}
 
 	/** Version that takes a handle */
-	FORCEINLINE float GetTimerRate(FTimerHandle const& InHandle) const
+	FORCEINLINE float GetTimerRate(FTimerHandle InHandle) const
 	{
 		FTimerData const* const TimerData = FindTimer(InHandle);
 		return InternalGetTimerRate(TimerData);
@@ -449,37 +472,43 @@ public:
 	 * @param inTimerMethod Method binding for query.
 	 * @return				true if the timer matching the given criteria exists and is active, false otherwise.
 	 */
-	template< class UserClass >	
+	template< class UserClass >
+	DELEGATE_DEPRECATED("This overload of IsTimerActive is deprecated, use IsTimerActive(FTimerHandle InHandle) instead.")
 	FORCEINLINE bool IsTimerActive(UserClass* inObj, typename FTimerDelegate::TUObjectMethodDelegate< UserClass >::FMethodPtr inTimerMethod) const
 	{
-		return IsTimerActive( FTimerDelegate::CreateUObject(inObj, inTimerMethod) );
+		FTimerData const* const TimerData = DEPRECATED_FindTimer( FTimerUnifiedDelegate( FTimerDelegate::CreateUObject(inObj, inTimerMethod) ) );
+		return TimerData && TimerData->Status != ETimerStatus::Paused;
 	}
 		
-	template< class UserClass >	
+	template< class UserClass >
+	DELEGATE_DEPRECATED("This overload of IsTimerActive is deprecated, use IsTimerActive(FTimerHandle InHandle) instead.")
 	FORCEINLINE bool IsTimerActive(UserClass* inObj, typename FTimerDelegate::TUObjectMethodDelegate_Const< UserClass >::FMethodPtr inTimerMethod) const
 	{
-		return IsTimerActive( FTimerDelegate::CreateUObject(inObj, inTimerMethod) );
+		FTimerData const* const TimerData = DEPRECATED_FindTimer( FTimerUnifiedDelegate( FTimerDelegate::CreateUObject(inObj, inTimerMethod) ) );
+		return TimerData && TimerData->Status != ETimerStatus::Paused;
 	}
 
 	/** Version that takes any generic delegate. */
+	DELEGATE_DEPRECATED("This overload of IsTimerActive is deprecated, use IsTimerActive(FTimerHandle InHandle) instead.")
 	FORCEINLINE bool IsTimerActive(FTimerDelegate const& InDelegate) const
 	{
-		FTimerData const* const TimerData = FindTimer( FTimerUnifiedDelegate(InDelegate) );
-		return ( (TimerData != NULL) && (TimerData->Status != ETimerStatus::Paused) );
+		FTimerData const* const TimerData = DEPRECATED_FindTimer( FTimerUnifiedDelegate(InDelegate) );
+		return TimerData && TimerData->Status != ETimerStatus::Paused;
 	}
 
 	/** Version that takes a dynamic delegate (e.g. for UFunctions). */
+	DELEGATE_DEPRECATED("This overload of IsTimerActive is deprecated, use IsTimerActive(FTimerHandle InHandle) instead.")
 	FORCEINLINE bool IsTimerActive(FTimerDynamicDelegate const& InDynDelegate) const
 	{
-		FTimerData const* const TimerData = FindTimer( FTimerUnifiedDelegate(InDynDelegate) );
-		return ( (TimerData != NULL) && (TimerData->Status != ETimerStatus::Paused) );
+		FTimerData const* const TimerData = DEPRECATED_FindTimer( FTimerUnifiedDelegate(InDynDelegate) );
+		return TimerData && TimerData->Status != ETimerStatus::Paused;
 	}
 
 	/** Version that takes a handle */
-	FORCEINLINE bool IsTimerActive(FTimerHandle const& InHandle) const
+	FORCEINLINE bool IsTimerActive(FTimerHandle InHandle) const
 	{
 		FTimerData const* const TimerData = FindTimer( InHandle );
-		return ( (TimerData != NULL) && (TimerData->Status != ETimerStatus::Paused) );
+		return TimerData && TimerData->Status != ETimerStatus::Paused;
 	}
 
 	/**
@@ -490,36 +519,42 @@ public:
 	* @return				true if the timer matching the given criteria exists and is paused, false otherwise.
 	*/
 	template< class UserClass >
+	DELEGATE_DEPRECATED("This overload of IsTimerPaused is deprecated, use IsTimerPaused(FTimerHandle InHandle) instead.")
 	FORCEINLINE bool IsTimerPaused(UserClass* inObj, typename FTimerDelegate::TUObjectMethodDelegate< UserClass >::FMethodPtr inTimerMethod) const
 	{
-		return IsTimerPaused(FTimerDelegate::CreateUObject(inObj, inTimerMethod));
+		FTimerData const* const TimerData = DEPRECATED_FindTimer(FTimerUnifiedDelegate(FTimerDelegate::CreateUObject(inObj, inTimerMethod)));
+		return TimerData && TimerData->Status == ETimerStatus::Paused;
 	}
 
 	template< class UserClass >
+	DELEGATE_DEPRECATED("This overload of IsTimerPaused is deprecated, use IsTimerPaused(FTimerHandle InHandle) instead.")
 	FORCEINLINE bool IsTimerPaused(UserClass* inObj, typename FTimerDelegate::TUObjectMethodDelegate_Const< UserClass >::FMethodPtr inTimerMethod)  const
 	{
-		return IsTimerPaused(FTimerDelegate::CreateUObject(inObj, inTimerMethod));
+		FTimerData const* const TimerData = DEPRECATED_FindTimer(FTimerUnifiedDelegate(FTimerDelegate::CreateUObject(inObj, inTimerMethod)));
+		return TimerData && TimerData->Status == ETimerStatus::Paused;
 	}
 
 	/** Version that takes any generic delegate. */
+	DELEGATE_DEPRECATED("This overload of IsTimerPaused is deprecated, use IsTimerPaused(FTimerHandle InHandle) instead.")
 	FORCEINLINE bool IsTimerPaused(FTimerDelegate const& InDelegate) const
 	{
-		FTimerData const* const TimerData = FindTimer(FTimerUnifiedDelegate(InDelegate));
-		return ((TimerData != NULL) && (TimerData->Status == ETimerStatus::Paused));
+		FTimerData const* const TimerData = DEPRECATED_FindTimer(FTimerUnifiedDelegate(InDelegate));
+		return TimerData && TimerData->Status == ETimerStatus::Paused;
 	}
 
 	/** Version that takes a dynamic delegate (e.g. for UFunctions). */
+	DELEGATE_DEPRECATED("This overload of IsTimerPaused is deprecated, use IsTimerPaused(FTimerHandle InHandle) instead.")
 	FORCEINLINE bool IsTimerPaused(FTimerDynamicDelegate const& InDynDelegate) const
 	{
-		FTimerData const* const TimerData = FindTimer(FTimerUnifiedDelegate(InDynDelegate));
-		return ((TimerData != NULL) && (TimerData->Status == ETimerStatus::Paused));
+		FTimerData const* const TimerData = DEPRECATED_FindTimer(FTimerUnifiedDelegate(InDynDelegate));
+		return TimerData && TimerData->Status == ETimerStatus::Paused;
 	}
 
 	/** Version that takes a handle */
-	FORCEINLINE bool IsTimerPaused(FTimerHandle const& InHandle)
+	FORCEINLINE bool IsTimerPaused(FTimerHandle InHandle)
 	{
 		FTimerData const* const TimerData = FindTimer(InHandle);
-		return ((TimerData != NULL) && (TimerData->Status == ETimerStatus::Paused));
+		return TimerData && TimerData->Status == ETimerStatus::Paused;
 	}
 
 	/**
@@ -530,36 +565,42 @@ public:
 	* @return				true if the timer matching the given criteria exists and is paused, false otherwise.
 	*/
 	template< class UserClass >
+	DELEGATE_DEPRECATED("This overload of IsTimerPending is deprecated, use IsTimerPending(FTimerHandle InHandle) instead.")
 	FORCEINLINE bool IsTimerPending(UserClass* inObj, typename FTimerDelegate::TUObjectMethodDelegate< UserClass >::FMethodPtr inTimerMethod) const
 	{
-		return IsTimerPending(FTimerDelegate::CreateUObject(inObj, inTimerMethod));
+		FTimerData const* const TimerData = DEPRECATED_FindTimer(FTimerUnifiedDelegate(FTimerDelegate::CreateUObject(inObj, inTimerMethod)));
+		return TimerData && TimerData->Status == ETimerStatus::Pending;
 	}
 
 	template< class UserClass >
+	DELEGATE_DEPRECATED("This overload of IsTimerPending is deprecated, use IsTimerPending(FTimerHandle InHandle) instead.")
 	FORCEINLINE bool IsTimerPending(UserClass* inObj, typename FTimerDelegate::TUObjectMethodDelegate_Const< UserClass >::FMethodPtr inTimerMethod)  const
 	{
-		return IsTimerPending(FTimerDelegate::CreateUObject(inObj, inTimerMethod));
+		FTimerData const* const TimerData = DEPRECATED_FindTimer(FTimerUnifiedDelegate(FTimerDelegate::CreateUObject(inObj, inTimerMethod)));
+		return TimerData && TimerData->Status == ETimerStatus::Pending;
 	}
 
 	/** Version that takes any generic delegate. */
+	DELEGATE_DEPRECATED("This overload of IsTimerPending is deprecated, use IsTimerPending(FTimerHandle InHandle) instead.")
 	FORCEINLINE bool IsTimerPending(FTimerDelegate const& InDelegate) const
 	{
-		FTimerData const* const TimerData = FindTimer(FTimerUnifiedDelegate(InDelegate));
-		return ((TimerData != NULL) && (TimerData->Status == ETimerStatus::Pending));
+		FTimerData const* const TimerData = DEPRECATED_FindTimer(FTimerUnifiedDelegate(InDelegate));
+		return TimerData && TimerData->Status == ETimerStatus::Pending;
 	}
 
 	/** Version that takes a dynamic delegate (e.g. for UFunctions). */
+	DELEGATE_DEPRECATED("This overload of IsTimerPending is deprecated, use IsTimerPending(FTimerHandle InHandle) instead.")
 	FORCEINLINE bool IsTimerPending(FTimerDynamicDelegate const& InDynDelegate) const
 	{
-		FTimerData const* const TimerData = FindTimer(FTimerUnifiedDelegate(InDynDelegate));
-		return ((TimerData != NULL) && (TimerData->Status == ETimerStatus::Pending));
+		FTimerData const* const TimerData = DEPRECATED_FindTimer(FTimerUnifiedDelegate(InDynDelegate));
+		return TimerData && TimerData->Status == ETimerStatus::Pending;
 	}
 
 	/** Version that takes a handle */
-	FORCEINLINE bool IsTimerPending(FTimerHandle const& InHandle)
+	FORCEINLINE bool IsTimerPending(FTimerHandle InHandle)
 	{
 		FTimerData const* const TimerData = FindTimer(InHandle);
-		return ((TimerData != NULL) && (TimerData->Status == ETimerStatus::Pending));
+		return TimerData && TimerData->Status == ETimerStatus::Pending;
 	}
 
 	/**
@@ -570,34 +611,38 @@ public:
 	* @return				true if the timer matching the given criteria exists, false otherwise.
 	*/
 	template< class UserClass >
+	DELEGATE_DEPRECATED("This overload of TimerExists is deprecated, use TimerExists(FTimerHandle InHandle) instead.")
 	FORCEINLINE bool TimerExists(UserClass* inObj, typename FTimerDelegate::TUObjectMethodDelegate< UserClass >::FMethodPtr inTimerMethod) const
 	{
-		return TimerExists(FTimerDelegate::CreateUObject(inObj, inTimerMethod));
+		return DEPRECATED_FindTimer(FTimerUnifiedDelegate(FTimerDelegate::CreateUObject(inObj, inTimerMethod))) != nullptr;
 	}
 
 	template< class UserClass >
+	DELEGATE_DEPRECATED("This overload of TimerExists is deprecated, use TimerExists(FTimerHandle InHandle) instead.")
 	FORCEINLINE bool TimerExists(UserClass* inObj, typename FTimerDelegate::TUObjectMethodDelegate_Const< UserClass >::FMethodPtr inTimerMethod) const
 	{
-		return TimerExists(FTimerDelegate::CreateUObject(inObj, inTimerMethod));
+		return DEPRECATED_FindTimer(FTimerUnifiedDelegate(FTimerDelegate::CreateUObject(inObj, inTimerMethod))) != nullptr;
 	}
 
 	/** Version that takes any generic delegate. */
+	DELEGATE_DEPRECATED("This overload of TimerExists is deprecated, use TimerExists(FTimerHandle InHandle) instead.")
 	FORCEINLINE bool TimerExists(FTimerDelegate const& InDelegate) const
 	{
-		return FindTimer(FTimerUnifiedDelegate(InDelegate)) != NULL;
+		return DEPRECATED_FindTimer(FTimerUnifiedDelegate(InDelegate)) != nullptr;
 		
 	}
 
 	/** Version that takes a dynamic delegate (e.g. for UFunctions). */
+	DELEGATE_DEPRECATED("This overload of TimerExists is deprecated, use TimerExists(FTimerHandle InHandle) instead.")
 	FORCEINLINE bool TimerExists(FTimerDynamicDelegate const& InDynDelegate) const
 	{
-		return FindTimer(FTimerUnifiedDelegate(InDynDelegate)) != NULL;
+		return DEPRECATED_FindTimer(FTimerUnifiedDelegate(InDynDelegate)) != nullptr;
 	}
 
 	/** Version that takes a handle */
-	FORCEINLINE bool TimerExists(FTimerHandle const& InHandle)
+	FORCEINLINE bool TimerExists(FTimerHandle InHandle)
 	{
-		return FindTimer(InHandle) != NULL;
+		return FindTimer(InHandle) != nullptr;
 	}
 
 	/**
@@ -607,31 +652,37 @@ public:
 	 * @param inTimerMethod Method to call when timer fires
 	 * @return				The current time elapsed or -1.f if timer does not exist.
 	 */
-	template< class UserClass >	
+	template< class UserClass >
+	DELEGATE_DEPRECATED("This overload of GetTimerElapsed is deprecated, use GetTimerElapsed(FTimerHandle InHandle) instead.")
 	FORCEINLINE float GetTimerElapsed(UserClass* inObj, typename FTimerDelegate::TUObjectMethodDelegate< UserClass >::FMethodPtr inTimerMethod) const
 	{
-		return GetTimerElapsed( FTimerDelegate::CreateUObject(inObj, inTimerMethod) );
+		FTimerData const* const TimerData = DEPRECATED_FindTimer( FTimerUnifiedDelegate( FTimerDelegate::CreateUObject(inObj, inTimerMethod) ) );
+		return InternalGetTimerElapsed(TimerData);
 	}
-	template< class UserClass >	
+	template< class UserClass >
+	DELEGATE_DEPRECATED("This overload of GetTimerElapsed is deprecated, use GetTimerElapsed(FTimerHandle InHandle) instead.")
 	FORCEINLINE float GetTimerElapsed(UserClass* inObj, typename FTimerDelegate::TUObjectMethodDelegate_Const< UserClass >::FMethodPtr inTimerMethod) const
 	{
-		return GetTimerElapsed( FTimerDelegate::CreateUObject(inObj, inTimerMethod) );
+		FTimerData const* const TimerData = DEPRECATED_FindTimer( FTimerUnifiedDelegate( FTimerDelegate::CreateUObject(inObj, inTimerMethod) ) );
+		return InternalGetTimerElapsed(TimerData);
 	}
 	/** Version that takes any generic delegate. */
+	DELEGATE_DEPRECATED("This overload of GetTimerElapsed is deprecated, use GetTimerElapsed(FTimerHandle InHandle) instead.")
 	FORCEINLINE float GetTimerElapsed(FTimerDelegate const& InDelegate) const
 	{
-		FTimerData const* const TimerData = FindTimer( FTimerUnifiedDelegate(InDelegate) );
+		FTimerData const* const TimerData = DEPRECATED_FindTimer( FTimerUnifiedDelegate(InDelegate) );
 		return InternalGetTimerElapsed(TimerData);
 	}
 	/** Version that takes a dynamic delegate (e.g. for UFunctions). */
+	DELEGATE_DEPRECATED("This overload of GetTimerElapsed is deprecated, use GetTimerElapsed(FTimerHandle InHandle) instead.")
 	FORCEINLINE float GetTimerElapsed(FTimerDynamicDelegate const& InDynDelegate) const
 	{
-		FTimerData const* const TimerData = FindTimer( FTimerUnifiedDelegate(InDynDelegate) );
+		FTimerData const* const TimerData = DEPRECATED_FindTimer( FTimerUnifiedDelegate(InDynDelegate) );
 		return InternalGetTimerElapsed(TimerData);
 	}
 
 	/** Version that takes a handle */
-	FORCEINLINE float GetTimerElapsed(FTimerHandle const& InHandle) const
+	FORCEINLINE float GetTimerElapsed(FTimerHandle InHandle) const
 	{
 		FTimerData const* const TimerData = FindTimer(InHandle);
 		return InternalGetTimerElapsed(TimerData);
@@ -644,30 +695,36 @@ public:
 	 * @param inTimerMethod Method to call when timer fires.
 	 * @return				The current time remaining, or -1.f if timer does not exist
 	 */
-	template< class UserClass >	
+	template< class UserClass >
+	DELEGATE_DEPRECATED("This overload of GetTimerRemaining is deprecated, use GetTimerRemaining(FTimerHandle InHandle) instead.")
 	FORCEINLINE float GetTimerRemaining(UserClass* inObj, typename FTimerDelegate::TUObjectMethodDelegate< UserClass >::FMethodPtr inTimerMethod) const
 	{
-		return GetTimerRemaining( FTimerDelegate::CreateUObject(inObj, inTimerMethod) );
+		FTimerData const* const TimerData = DEPRECATED_FindTimer( FTimerUnifiedDelegate( FTimerDelegate::CreateUObject(inObj, inTimerMethod) ) );
+		return InternalGetTimerRemaining(TimerData);
 	}
-	template< class UserClass >	
+	template< class UserClass >
+	DELEGATE_DEPRECATED("This overload of GetTimerRemaining is deprecated, use GetTimerRemaining(FTimerHandle InHandle) instead.")
 	FORCEINLINE float GetTimerRemaining(UserClass* inObj, typename FTimerDelegate::TUObjectMethodDelegate_Const< UserClass >::FMethodPtr inTimerMethod) const
 	{
-		return GetTimerRemaining( FTimerDelegate::CreateUObject(inObj, inTimerMethod) );
+		FTimerData const* const TimerData = DEPRECATED_FindTimer( FTimerUnifiedDelegate( FTimerDelegate::CreateUObject(inObj, inTimerMethod) ) );
+		return InternalGetTimerRemaining(TimerData);
 	}
 	/** Version that takes any generic delegate. */
+	DELEGATE_DEPRECATED("This overload of GetTimerRemaining is deprecated, use GetTimerRemaining(FTimerHandle InHandle) instead.")
 	FORCEINLINE float GetTimerRemaining(FTimerDelegate const& InDelegate) const
 	{
-		FTimerData const* const TimerData = FindTimer( FTimerUnifiedDelegate(InDelegate) );
+		FTimerData const* const TimerData = DEPRECATED_FindTimer( FTimerUnifiedDelegate(InDelegate) );
 		return InternalGetTimerRemaining(TimerData);
 	}
 	/** Version that takes a dynamic delegate (e.g. for UFunctions). */
+	DELEGATE_DEPRECATED("This overload of GetTimerRemaining is deprecated, use GetTimerRemaining(FTimerHandle InHandle) instead.")
 	FORCEINLINE float GetTimerRemaining(FTimerDynamicDelegate const& InDynDelegate) const
 	{
-		FTimerData const* const TimerData = FindTimer( FTimerUnifiedDelegate(InDynDelegate) );
+		FTimerData const* const TimerData = DEPRECATED_FindTimer( FTimerUnifiedDelegate(InDynDelegate) );
 		return InternalGetTimerRemaining(TimerData);
 	}
 	/** Version that takes a handle */
-	FORCEINLINE float GetTimerRemaining(FTimerHandle const& InHandle) const
+	FORCEINLINE float GetTimerRemaining(FTimerHandle InHandle) const
 	{
 		FTimerData const* const TimerData = FindTimer(InHandle);
 		return InternalGetTimerRemaining(TimerData);
@@ -679,19 +736,30 @@ public:
 		return (LastTickedFrame == GFrameCounter);
 	}
 
+	/**
+	 * Finds a handle to a timer bound to a particular dynamic delegate.
+	 * This function is intended to be used only by the K2 system.
+	 *
+	 * @param  InDynamicDelegate  The dynamic delegate to search for.
+	 *
+	 * @return A handle to the found timer - !IsValid() if no such timer was found.
+	 */
+	FTimerHandle K2_FindDynamicTimerHandle(FTimerDynamicDelegate InDynamicDelegate) const;
+
 private:
 
-	void InternalSetTimer( FTimerUnifiedDelegate const& InDelegate, float InRate, bool InbLoop, float InFirstDelay );
+	void DEPRECATED_InternalSetTimer( FTimerUnifiedDelegate const& InDelegate, float InRate, bool InbLoop, float InFirstDelay );
 	void InternalSetTimer( FTimerHandle& InOutHandle, FTimerUnifiedDelegate const& InDelegate, float InRate, bool InbLoop, float InFirstDelay );
 	void InternalSetTimer( FTimerData& NewTimerData, float InRate, bool InbLoop, float InFirstDelay );
 	void InternalSetTimerForNextTick( FTimerUnifiedDelegate const& InDelegate );
-	void InternalClearTimer( FTimerUnifiedDelegate const& InDelegate );
+	void DEPRECATED_InternalClearTimer( FTimerUnifiedDelegate const& InDelegate );
 	void InternalClearTimer( FTimerHandle const& InDelegate );
+	void InternalClearTimer( int32 TimerIdx, ETimerStatus::Type TimerStatus );
 	void InternalClearAllTimers( void const* Object );
 
 	/** Will find a timer in the active, paused, or pending list. */
-	FTimerData const* FindTimer( FTimerUnifiedDelegate const& InDelegate, int32* OutTimerIndex=NULL ) const;
-	FTimerData const* FindTimer( FTimerHandle const& InHandle, int32* OutTimerIndex = NULL ) const;
+	FTimerData const* DEPRECATED_FindTimer( FTimerUnifiedDelegate const& InDelegate, int32* OutTimerIndex=nullptr ) const;
+	FTimerData const* FindTimer( FTimerHandle const& InHandle, int32* OutTimerIndex = nullptr ) const;
 
 	void InternalPauseTimer( FTimerData const* TimerToPause, int32 TimerIdx );
 	void InternalUnPauseTimer( int32 PausedTimerIdx );
@@ -701,7 +769,7 @@ private:
 	float InternalGetTimerRemaining( FTimerData const* const TimerData ) const;
 
 	/** Will find the given timer in the given TArray and return its index. */ 
-	int32 FindTimerInList( const TArray<FTimerData> &SearchArray, FTimerUnifiedDelegate const& InDelegate ) const;
+	int32 DEPRECATED_FindTimerInList( const TArray<FTimerData> &SearchArray, FTimerUnifiedDelegate const& InDelegate ) const;
 	int32 FindTimerInList( const TArray<FTimerData> &SearchArray, FTimerHandle const& InHandle ) const;
 
 	/** Heap of actively running timers. */
diff --git a/Engine/Source/Runtime/Engine/Public/UnrealEngine.h b/Engine/Source/Runtime/Engine/Public/UnrealEngine.h
--- a/Engine/Source/Runtime/Engine/Public/UnrealEngine.h
+++ b/Engine/Source/Runtime/Engine/Public/UnrealEngine.h
@@ -252,6 +252,8 @@ private:
 	float TimesToReport;
 	/** Name of operation performed that is being tracked. */
 	FString OperationPerformed;
+	/** Handle to the registered OnPreGarbageCollect delegate */
+	FDelegateHandle OnPreGarbageCollectDelegateHandle;
 };
 
 /** Scoped helper structure for capturing tick time. */
diff --git a/Engine/Source/Runtime/GameplayAbilities/Private/Abilities/Tasks/AbilityTask_WaitAbilityActivate.cpp b/Engine/Source/Runtime/GameplayAbilities/Private/Abilities/Tasks/AbilityTask_WaitAbilityActivate.cpp
--- a/Engine/Source/Runtime/GameplayAbilities/Private/Abilities/Tasks/AbilityTask_WaitAbilityActivate.cpp
+++ b/Engine/Source/Runtime/GameplayAbilities/Private/Abilities/Tasks/AbilityTask_WaitAbilityActivate.cpp
@@ -26,7 +26,7 @@ void UAbilityTask_WaitAbilityActivate::Activate()
 {
 	if (AbilitySystemComponent.IsValid())
 	{
-		AbilitySystemComponent->AbilityActivatedCallbacks.AddUObject(this, &UAbilityTask_WaitAbilityActivate::OnAbilityActivate);
+		OnAbilityActivateDelegateHandle = AbilitySystemComponent->AbilityActivatedCallbacks.AddUObject(this, &UAbilityTask_WaitAbilityActivate::OnAbilityActivate);
 	}
 }
 
@@ -53,7 +53,7 @@ void UAbilityTask_WaitAbilityActivate::OnDestroy(bool AbilityEnded)
 {
 	if (AbilitySystemComponent.IsValid())
 	{
-		AbilitySystemComponent->AbilityActivatedCallbacks.RemoveUObject(this, &UAbilityTask_WaitAbilityActivate::OnAbilityActivate);
+		AbilitySystemComponent->AbilityActivatedCallbacks.Remove(OnAbilityActivateDelegateHandle);
 	}
 
 	Super::OnDestroy(AbilityEnded);
diff --git a/Engine/Source/Runtime/GameplayAbilities/Private/Abilities/Tasks/AbilityTask_WaitAbilityCommit.cpp b/Engine/Source/Runtime/GameplayAbilities/Private/Abilities/Tasks/AbilityTask_WaitAbilityCommit.cpp
--- a/Engine/Source/Runtime/GameplayAbilities/Private/Abilities/Tasks/AbilityTask_WaitAbilityCommit.cpp
+++ b/Engine/Source/Runtime/GameplayAbilities/Private/Abilities/Tasks/AbilityTask_WaitAbilityCommit.cpp
@@ -26,7 +26,7 @@ void UAbilityTask_WaitAbilityCommit::Activate()
 {
 	if (AbilitySystemComponent.IsValid())	
 	{		
-		AbilitySystemComponent->AbilityCommitedCallbacks.AddUObject(this, &UAbilityTask_WaitAbilityCommit::OnAbilityCommit);
+		OnAbilityCommitDelegateHandle = AbilitySystemComponent->AbilityCommitedCallbacks.AddUObject(this, &UAbilityTask_WaitAbilityCommit::OnAbilityCommit);
 	}
 }
 
@@ -34,7 +34,7 @@ void UAbilityTask_WaitAbilityCommit::OnDestroy(bool AbilityEnded)
 {
 	if (AbilitySystemComponent.IsValid())
 	{
-		AbilitySystemComponent->AbilityCommitedCallbacks.RemoveUObject(this, &UAbilityTask_WaitAbilityCommit::OnAbilityCommit);
+		AbilitySystemComponent->AbilityCommitedCallbacks.Remove(OnAbilityCommitDelegateHandle);
 	}
 
 	Super::OnDestroy(AbilityEnded);
diff --git a/Engine/Source/Runtime/GameplayAbilities/Private/Abilities/Tasks/AbilityTask_WaitAttributeChange.cpp b/Engine/Source/Runtime/GameplayAbilities/Private/Abilities/Tasks/AbilityTask_WaitAttributeChange.cpp
--- a/Engine/Source/Runtime/GameplayAbilities/Private/Abilities/Tasks/AbilityTask_WaitAttributeChange.cpp
+++ b/Engine/Source/Runtime/GameplayAbilities/Private/Abilities/Tasks/AbilityTask_WaitAttributeChange.cpp
@@ -40,7 +40,7 @@ void UAbilityTask_WaitAttributeChange::Activate()
 {
 	if (AbilitySystemComponent.IsValid())
 	{
-		AbilitySystemComponent->RegisterGameplayAttributeEvent(Attribute).AddUObject(this, &UAbilityTask_WaitAttributeChange::OnAttributeChange);
+		OnAttributeChangeDelegateHandle = AbilitySystemComponent->RegisterGameplayAttributeEvent(Attribute).AddUObject(this, &UAbilityTask_WaitAttributeChange::OnAttributeChange);
 	}
 }
 
@@ -100,7 +100,7 @@ void UAbilityTask_WaitAttributeChange::OnDestroy(bool AbilityEnded)
 {
 	if (AbilitySystemComponent.IsValid())
 	{
-		AbilitySystemComponent->RegisterGameplayAttributeEvent(Attribute).RemoveUObject(this, &UAbilityTask_WaitAttributeChange::OnAttributeChange);
+		AbilitySystemComponent->RegisterGameplayAttributeEvent(Attribute).Remove(OnAttributeChangeDelegateHandle);
 	}
 
 	Super::OnDestroy(AbilityEnded);
diff --git a/Engine/Source/Runtime/GameplayAbilities/Private/Abilities/Tasks/AbilityTask_WaitConfirm.cpp b/Engine/Source/Runtime/GameplayAbilities/Private/Abilities/Tasks/AbilityTask_WaitConfirm.cpp
--- a/Engine/Source/Runtime/GameplayAbilities/Private/Abilities/Tasks/AbilityTask_WaitConfirm.cpp
+++ b/Engine/Source/Runtime/GameplayAbilities/Private/Abilities/Tasks/AbilityTask_WaitConfirm.cpp
@@ -30,7 +30,7 @@ void UAbilityTask_WaitConfirm::Activate()
 		{
 			// Register delegate so that when we are confirmed by the server, we call OnConfirmCallback	
 			
-			Ability->OnConfirmDelegate.AddUObject(this, &UAbilityTask_WaitConfirm::OnConfirmCallback);
+			OnConfirmCallbackDelegateHandle = Ability->OnConfirmDelegate.AddUObject(this, &UAbilityTask_WaitConfirm::OnConfirmCallback);
 			RegisteredCallback = true;
 		}
 		else
@@ -45,7 +45,7 @@ void UAbilityTask_WaitConfirm::OnDestroy(bool AbilityEnded)
 {
 	if (RegisteredCallback && Ability.IsValid())
 	{
-		Ability->OnConfirmDelegate.RemoveUObject(this, &UAbilityTask_WaitConfirm::OnConfirmCallback);
+		Ability->OnConfirmDelegate.Remove(OnConfirmCallbackDelegateHandle);
 	}
 
 	Super::OnDestroy(AbilityEnded);
diff --git a/Engine/Source/Runtime/GameplayAbilities/Private/Abilities/Tasks/AbilityTask_WaitGameplayEffectApplied_Self.cpp b/Engine/Source/Runtime/GameplayAbilities/Private/Abilities/Tasks/AbilityTask_WaitGameplayEffectApplied_Self.cpp
--- a/Engine/Source/Runtime/GameplayAbilities/Private/Abilities/Tasks/AbilityTask_WaitGameplayEffectApplied_Self.cpp
+++ b/Engine/Source/Runtime/GameplayAbilities/Private/Abilities/Tasks/AbilityTask_WaitGameplayEffectApplied_Self.cpp
@@ -30,11 +30,10 @@ void UAbilityTask_WaitGameplayEffectApplied_Self::BroadcastDelegate(AActor* Avat
 
 void UAbilityTask_WaitGameplayEffectApplied_Self::RegisterDelegate()
 {
-	AbilitySystemComponent->OnGameplayEffectAppliedDelegateToSelf.AddUObject(this, &UAbilityTask_WaitGameplayEffectApplied::OnApplyGameplayEffectCallback);
+	OnApplyGameplayEffectCallbackDelegateHandle = AbilitySystemComponent->OnGameplayEffectAppliedDelegateToSelf.AddUObject(this, &UAbilityTask_WaitGameplayEffectApplied::OnApplyGameplayEffectCallback);
 }
 
 void UAbilityTask_WaitGameplayEffectApplied_Self::RemoveDelegate()
 {
-	AbilitySystemComponent->OnGameplayEffectAppliedDelegateToSelf.RemoveUObject(this, &UAbilityTask_WaitGameplayEffectApplied::OnApplyGameplayEffectCallback);
-
+	AbilitySystemComponent->OnGameplayEffectAppliedDelegateToSelf.Remove(OnApplyGameplayEffectCallbackDelegateHandle);
 }
\ No newline at end of file
diff --git a/Engine/Source/Runtime/GameplayAbilities/Private/Abilities/Tasks/AbilityTask_WaitGameplayEffectApplied_Target.cpp b/Engine/Source/Runtime/GameplayAbilities/Private/Abilities/Tasks/AbilityTask_WaitGameplayEffectApplied_Target.cpp
--- a/Engine/Source/Runtime/GameplayAbilities/Private/Abilities/Tasks/AbilityTask_WaitGameplayEffectApplied_Target.cpp
+++ b/Engine/Source/Runtime/GameplayAbilities/Private/Abilities/Tasks/AbilityTask_WaitGameplayEffectApplied_Target.cpp
@@ -30,11 +30,10 @@ void UAbilityTask_WaitGameplayEffectApplied_Target::BroadcastDelegate(AActor* Av
 
 void UAbilityTask_WaitGameplayEffectApplied_Target::RegisterDelegate()
 {
-	AbilitySystemComponent->OnGameplayEffectAppliedDelegateToTarget.AddUObject(this, &UAbilityTask_WaitGameplayEffectApplied::OnApplyGameplayEffectCallback);
+	OnApplyGameplayEffectCallbackDelegateHandle = AbilitySystemComponent->OnGameplayEffectAppliedDelegateToTarget.AddUObject(this, &UAbilityTask_WaitGameplayEffectApplied::OnApplyGameplayEffectCallback);
 }
 
 void UAbilityTask_WaitGameplayEffectApplied_Target::RemoveDelegate()
 {
-	AbilitySystemComponent->OnGameplayEffectAppliedDelegateToTarget.RemoveUObject(this, &UAbilityTask_WaitGameplayEffectApplied::OnApplyGameplayEffectCallback);
-
+	AbilitySystemComponent->OnGameplayEffectAppliedDelegateToTarget.Remove(OnApplyGameplayEffectCallbackDelegateHandle);
 }
\ No newline at end of file
diff --git a/Engine/Source/Runtime/GameplayAbilities/Private/Abilities/Tasks/AbilityTask_WaitGameplayEffectRemoved.cpp b/Engine/Source/Runtime/GameplayAbilities/Private/Abilities/Tasks/AbilityTask_WaitGameplayEffectRemoved.cpp
--- a/Engine/Source/Runtime/GameplayAbilities/Private/Abilities/Tasks/AbilityTask_WaitGameplayEffectRemoved.cpp
+++ b/Engine/Source/Runtime/GameplayAbilities/Private/Abilities/Tasks/AbilityTask_WaitGameplayEffectRemoved.cpp
@@ -37,7 +37,7 @@ void UAbilityTask_WaitGameplayEffectRemoved::Activate()
 		FOnActiveGameplayEffectRemoved* DelPtr = EffectOwningAbilitySystemComponent->OnGameplayEffectRemovedDelegate(Handle);
 		if (DelPtr)
 		{
-			DelPtr->AddUObject(this, &UAbilityTask_WaitGameplayEffectRemoved::OnGameplayEffectRemoved);
+			OnGameplayEffectRemovedDelegateHandle = DelPtr->AddUObject(this, &UAbilityTask_WaitGameplayEffectRemoved::OnGameplayEffectRemoved);
 			Registered = true;
 		}
 	}
@@ -57,7 +57,7 @@ void UAbilityTask_WaitGameplayEffectRemoved::OnDestroy(bool AbilityIsEnding)
 		FOnActiveGameplayEffectRemoved* DelPtr = EffectOwningAbilitySystemComponent->OnGameplayEffectRemovedDelegate(Handle);
 		if (DelPtr)
 		{
-			DelPtr->RemoveUObject(this, &UAbilityTask_WaitGameplayEffectRemoved::OnGameplayEffectRemoved);
+			DelPtr->Remove(OnGameplayEffectRemovedDelegateHandle);
 		}
 	}
 
diff --git a/Engine/Source/Runtime/GameplayAbilities/Private/AbilitySystemComponent_Abilities.cpp b/Engine/Source/Runtime/GameplayAbilities/Private/AbilitySystemComponent_Abilities.cpp
--- a/Engine/Source/Runtime/GameplayAbilities/Private/AbilitySystemComponent_Abilities.cpp
+++ b/Engine/Source/Runtime/GameplayAbilities/Private/AbilitySystemComponent_Abilities.cpp
@@ -321,7 +321,7 @@ void UAbilitySystemComponent::OnGiveAbility(FGameplayAbilitySpec& Spec)
 
 			if (!CountChangedEvent.IsBoundToObject(this))
 			{
-				CountChangedEvent.AddUObject(this, &UAbilitySystemComponent::MonitoredTagChanged);
+				MonitoredTagChangedDelegatHandle = CountChangedEvent.AddUObject(this, &UAbilitySystemComponent::MonitoredTagChanged);
 			}
 		}
 	}
@@ -402,7 +402,7 @@ void UAbilitySystemComponent::CheckForClearedAbilities()
 		
 			if (CountChangedEvent.IsBoundToObject(this))
 			{
-				CountChangedEvent.RemoveUObject(this, &UAbilitySystemComponent::MonitoredTagChanged);
+				CountChangedEvent.Remove(MonitoredTagChangedDelegatHandle);
 			}
 		}
 
@@ -632,7 +632,7 @@ void UAbilitySystemComponent::OnRep_ActivateAbilities()
 		if (!SpecAbility)
 		{
 			// Queue up another call to make sure this gets run again, as our abilities haven't replicated yet
-			GetWorld()->GetTimerManager().SetTimer(this, &UAbilitySystemComponent::OnRep_ActivateAbilities, 0.5);
+			GetWorld()->GetTimerManager().SetTimer(OnRep_ActivateAbilitiesTimerHandle, this, &UAbilitySystemComponent::OnRep_ActivateAbilities, 0.5);
 			return;
 		}
 	}
diff --git a/Engine/Source/Runtime/GameplayAbilities/Private/AbilitySystemDebugHUD.cpp b/Engine/Source/Runtime/GameplayAbilities/Private/AbilitySystemDebugHUD.cpp
--- a/Engine/Source/Runtime/GameplayAbilities/Private/AbilitySystemDebugHUD.cpp
+++ b/Engine/Source/Runtime/GameplayAbilities/Private/AbilitySystemDebugHUD.cpp
@@ -137,17 +137,19 @@ static void	ToggleDebugHUD(const TArray<FString>& Args, UWorld* InWorld)
 		break;
 	}
 
+	static FDelegateHandle DrawDebugDelegateHandle;
+
 	if (!HUD)
 	{
 		HUD = InWorld->SpawnActor<AAbilitySystemDebugHUD>();
 		
 		FDebugDrawDelegate DrawDebugDelegate = FDebugDrawDelegate::CreateUObject(HUD, &AAbilitySystemDebugHUD::DrawDebugHUD);
-		UDebugDrawService::Register(TEXT("GameplayDebug"), DrawDebugDelegate);
+		DrawDebugDelegateHandle = UDebugDrawService::Register(TEXT("GameplayDebug"), DrawDebugDelegate);
 	}
 	else
 	{
 		FDebugDrawDelegate DrawDebugDelegate = FDebugDrawDelegate::CreateUObject(HUD, &AAbilitySystemDebugHUD::DrawDebugHUD);
-		UDebugDrawService::Unregister(DrawDebugDelegate);
+		UDebugDrawService::Unregister(DrawDebugDelegateHandle);
 		HUD->Destroy();
 	}
 }
diff --git a/Engine/Source/Runtime/GameplayAbilities/Public/Abilities/Tasks/AbilityTask_WaitAbilityActivate.h b/Engine/Source/Runtime/GameplayAbilities/Public/Abilities/Tasks/AbilityTask_WaitAbilityActivate.h
--- a/Engine/Source/Runtime/GameplayAbilities/Public/Abilities/Tasks/AbilityTask_WaitAbilityActivate.h
+++ b/Engine/Source/Runtime/GameplayAbilities/Public/Abilities/Tasks/AbilityTask_WaitAbilityActivate.h
@@ -36,4 +36,5 @@ protected:
 
 	virtual void OnDestroy(bool AbilityEnded) override;
 
+	FDelegateHandle OnAbilityActivateDelegateHandle;
 };
\ No newline at end of file
diff --git a/Engine/Source/Runtime/GameplayAbilities/Public/Abilities/Tasks/AbilityTask_WaitAbilityCommit.h b/Engine/Source/Runtime/GameplayAbilities/Public/Abilities/Tasks/AbilityTask_WaitAbilityCommit.h
--- a/Engine/Source/Runtime/GameplayAbilities/Public/Abilities/Tasks/AbilityTask_WaitAbilityCommit.h
+++ b/Engine/Source/Runtime/GameplayAbilities/Public/Abilities/Tasks/AbilityTask_WaitAbilityCommit.h
@@ -35,4 +35,6 @@ class UAbilityTask_WaitAbilityCommit : public UAbilityTask
 protected:
 
 	virtual void OnDestroy(bool AbilityEnded) override;
+
+	FDelegateHandle OnAbilityCommitDelegateHandle;
 };
\ No newline at end of file
diff --git a/Engine/Source/Runtime/GameplayAbilities/Public/Abilities/Tasks/AbilityTask_WaitAttributeChange.h b/Engine/Source/Runtime/GameplayAbilities/Public/Abilities/Tasks/AbilityTask_WaitAttributeChange.h
--- a/Engine/Source/Runtime/GameplayAbilities/Public/Abilities/Tasks/AbilityTask_WaitAttributeChange.h
+++ b/Engine/Source/Runtime/GameplayAbilities/Public/Abilities/Tasks/AbilityTask_WaitAttributeChange.h
@@ -56,6 +56,7 @@ class UAbilityTask_WaitAttributeChange : public UAbilityTask
 	FGameplayAttribute	Attribute;
 	TEnumAsByte<EWaitAttributeChangeComparison::Type> ComparisonType;
 	float ComparisonValue;
+	FDelegateHandle OnAttributeChangeDelegateHandle;
 
 protected:
 
diff --git a/Engine/Source/Runtime/GameplayAbilities/Public/Abilities/Tasks/AbilityTask_WaitConfirm.h b/Engine/Source/Runtime/GameplayAbilities/Public/Abilities/Tasks/AbilityTask_WaitConfirm.h
--- a/Engine/Source/Runtime/GameplayAbilities/Public/Abilities/Tasks/AbilityTask_WaitConfirm.h
+++ b/Engine/Source/Runtime/GameplayAbilities/Public/Abilities/Tasks/AbilityTask_WaitConfirm.h
@@ -26,4 +26,6 @@ protected:
 	virtual void OnDestroy(bool AbilityEnded) override;
 
 	bool RegisteredCallback;
+
+	FDelegateHandle OnConfirmCallbackDelegateHandle;
 };
\ No newline at end of file
diff --git a/Engine/Source/Runtime/GameplayAbilities/Public/Abilities/Tasks/AbilityTask_WaitGameplayEffectApplied_Self.h b/Engine/Source/Runtime/GameplayAbilities/Public/Abilities/Tasks/AbilityTask_WaitGameplayEffectApplied_Self.h
--- a/Engine/Source/Runtime/GameplayAbilities/Public/Abilities/Tasks/AbilityTask_WaitGameplayEffectApplied_Self.h
+++ b/Engine/Source/Runtime/GameplayAbilities/Public/Abilities/Tasks/AbilityTask_WaitGameplayEffectApplied_Self.h
@@ -24,4 +24,7 @@ protected:
 	virtual void BroadcastDelegate(AActor* Avatar, FGameplayEffectSpecHandle SpecHandle, FActiveGameplayEffectHandle ActiveHandle) override;
 	virtual void RegisterDelegate() override;
 	virtual void RemoveDelegate() override;
+
+private:
+	FDelegateHandle OnApplyGameplayEffectCallbackDelegateHandle;
 };
\ No newline at end of file
diff --git a/Engine/Source/Runtime/GameplayAbilities/Public/Abilities/Tasks/AbilityTask_WaitGameplayEffectApplied_Target.h b/Engine/Source/Runtime/GameplayAbilities/Public/Abilities/Tasks/AbilityTask_WaitGameplayEffectApplied_Target.h
--- a/Engine/Source/Runtime/GameplayAbilities/Public/Abilities/Tasks/AbilityTask_WaitGameplayEffectApplied_Target.h
+++ b/Engine/Source/Runtime/GameplayAbilities/Public/Abilities/Tasks/AbilityTask_WaitGameplayEffectApplied_Target.h
@@ -24,4 +24,7 @@ protected:
 	virtual void BroadcastDelegate(AActor* Avatar, FGameplayEffectSpecHandle SpecHandle, FActiveGameplayEffectHandle ActiveHandle) override;
 	virtual void RegisterDelegate() override;
 	virtual void RemoveDelegate() override;
+
+private:
+	FDelegateHandle OnApplyGameplayEffectCallbackDelegateHandle;
 };
\ No newline at end of file
diff --git a/Engine/Source/Runtime/GameplayAbilities/Public/Abilities/Tasks/AbilityTask_WaitGameplayEffectRemoved.h b/Engine/Source/Runtime/GameplayAbilities/Public/Abilities/Tasks/AbilityTask_WaitGameplayEffectRemoved.h
--- a/Engine/Source/Runtime/GameplayAbilities/Public/Abilities/Tasks/AbilityTask_WaitGameplayEffectRemoved.h
+++ b/Engine/Source/Runtime/GameplayAbilities/Public/Abilities/Tasks/AbilityTask_WaitGameplayEffectRemoved.h
@@ -40,4 +40,6 @@ protected:
 
 	virtual void OnDestroy(bool AbilityIsEnding) override;
 	bool Registered;
+
+	FDelegateHandle OnGameplayEffectRemovedDelegateHandle;
 };
\ No newline at end of file
diff --git a/Engine/Source/Runtime/GameplayAbilities/Public/Abilities/Tasks/AbilityTask_WaitTargetData.h b/Engine/Source/Runtime/GameplayAbilities/Public/Abilities/Tasks/AbilityTask_WaitTargetData.h
--- a/Engine/Source/Runtime/GameplayAbilities/Public/Abilities/Tasks/AbilityTask_WaitTargetData.h
+++ b/Engine/Source/Runtime/GameplayAbilities/Public/Abilities/Tasks/AbilityTask_WaitTargetData.h
@@ -61,6 +61,8 @@ protected:
 	TWeakObjectPtr<AGameplayAbilityTargetActor>	MyTargetActor;
 
 	TEnumAsByte<EGameplayTargetingConfirmation::Type> ConfirmationType;
+
+	FDelegateHandle OnTargetDataReplicatedCallbackDelegateHandle;
 };
 
 
diff --git a/Engine/Source/Runtime/GameplayAbilities/Public/AbilitySystemComponent.h b/Engine/Source/Runtime/GameplayAbilities/Public/AbilitySystemComponent.h
--- a/Engine/Source/Runtime/GameplayAbilities/Public/AbilitySystemComponent.h
+++ b/Engine/Source/Runtime/GameplayAbilities/Public/AbilitySystemComponent.h
@@ -898,4 +898,8 @@ protected:
 	friend struct FActiveGameplayCueContainer;
 	friend struct FGameplayAbilitySpec;
 	friend struct FGameplayAbilitySpecContainer;
+
+private:
+	FDelegateHandle MonitoredTagChangedDelegatHandle;
+	FTimerHandle    OnRep_ActivateAbilitiesTimerHandle;
 };
diff --git a/Engine/Source/Runtime/Online/BuildPatchServices/Private/BuildPatchHTTP.cpp b/Engine/Source/Runtime/Online/BuildPatchServices/Private/BuildPatchHTTP.cpp
--- a/Engine/Source/Runtime/Online/BuildPatchServices/Private/BuildPatchHTTP.cpp
+++ b/Engine/Source/Runtime/Online/BuildPatchServices/Private/BuildPatchHTTP.cpp
@@ -134,7 +134,7 @@ FBuildPatchHTTP& FBuildPatchHTTP::Get()
 	if( SingletonInstance == NULL )
 	{
 		SingletonInstance = new FBuildPatchHTTP();
-		FTicker::GetCoreTicker().AddTicker( FTickerDelegate::CreateRaw( SingletonInstance, &FBuildPatchHTTP::Tick ) );
+		TickDelegateHandle = FTicker::GetCoreTicker().AddTicker( FTickerDelegate::CreateRaw( SingletonInstance, &FBuildPatchHTTP::Tick ) );
 	}
 	return *SingletonInstance;
 }
@@ -143,8 +143,9 @@ void FBuildPatchHTTP::OnShutdown()
 	if( SingletonInstance != NULL )
 	{
 		SingletonInstance->Tick( 0.0f );
-		FTicker::GetCoreTicker().RemoveTicker( FTickerDelegate::CreateRaw( SingletonInstance, &FBuildPatchHTTP::Tick ) );
+		FTicker::GetCoreTicker().RemoveTicker( TickDelegateHandle );
 		delete SingletonInstance;
 		SingletonInstance = NULL;
 	}
 }
+FDelegateHandle FBuildPatchHTTP::TickDelegateHandle;
diff --git a/Engine/Source/Runtime/Online/BuildPatchServices/Private/BuildPatchHTTP.h b/Engine/Source/Runtime/Online/BuildPatchServices/Private/BuildPatchHTTP.h
--- a/Engine/Source/Runtime/Online/BuildPatchServices/Private/BuildPatchHTTP.h
+++ b/Engine/Source/Runtime/Online/BuildPatchServices/Private/BuildPatchHTTP.h
@@ -119,6 +119,7 @@ private:
 private:
 	static FBuildPatchHTTP& Get();
 	static FBuildPatchHTTP* SingletonInstance;
+	static FDelegateHandle TickDelegateHandle;
 };
 
 #endif // __BuildPatchHTTP_h__
diff --git a/Engine/Source/Runtime/Online/BuildPatchServices/Private/BuildPatchServicesModule.cpp b/Engine/Source/Runtime/Online/BuildPatchServices/Private/BuildPatchServicesModule.cpp
--- a/Engine/Source/Runtime/Online/BuildPatchServices/Private/BuildPatchServicesModule.cpp
+++ b/Engine/Source/Runtime/Online/BuildPatchServices/Private/BuildPatchServicesModule.cpp
@@ -63,7 +63,7 @@ void FBuildPatchServicesModule::StartupModule()
 	FRollingHashConst::Init();
 
 	// Add our ticker
-	FTicker::GetCoreTicker().AddTicker( FTickerDelegate::CreateRaw( this, &FBuildPatchServicesModule::Tick ) );
+	TickDelegateHandle = FTicker::GetCoreTicker().AddTicker( FTickerDelegate::CreateRaw( this, &FBuildPatchServicesModule::Tick ) );
 
 	// Register core PreExit
 	FCoreDelegates::OnPreExit.AddRaw(this, &FBuildPatchServicesModule::PreExit);
@@ -79,7 +79,7 @@ void FBuildPatchServicesModule::ShutdownModule()
 	checkf(BuildPatchInstallers.Num() == 0, TEXT("BuildPatchServicesModule: FATAL ERROR: Core PreExit not called, or installer created during shutdown!"));
 
 	// Remove our ticker
-	FTicker::GetCoreTicker().RemoveTicker( FTickerDelegate::CreateRaw( this, &FBuildPatchServicesModule::Tick ) );
+	FTicker::GetCoreTicker().RemoveTicker( TickDelegateHandle );
 
 	FBuildPatchHTTP::OnShutdown();
 }
diff --git a/Engine/Source/Runtime/Online/BuildPatchServices/Private/BuildPatchServicesModule.h b/Engine/Source/Runtime/Online/BuildPatchServices/Private/BuildPatchServicesModule.h
--- a/Engine/Source/Runtime/Online/BuildPatchServices/Private/BuildPatchServicesModule.h
+++ b/Engine/Source/Runtime/Online/BuildPatchServices/Private/BuildPatchServicesModule.h
@@ -135,4 +135,7 @@ private:
 
 	// The system that holds available installations used for recycling install data
 	FBuildPatchInstallationInfo InstallationInfo;
+
+	// Handle to the registered Tick delegate
+	FDelegateHandle TickDelegateHandle;
 };
diff --git a/Engine/Source/Runtime/Online/OnlineSubsystem/Private/LANBeacon.cpp b/Engine/Source/Runtime/Online/OnlineSubsystem/Private/LANBeacon.cpp
--- a/Engine/Source/Runtime/Online/OnlineSubsystem/Private/LANBeacon.cpp
+++ b/Engine/Source/Runtime/Online/OnlineSubsystem/Private/LANBeacon.cpp
@@ -130,7 +130,7 @@ bool FLANSession::Host(FOnValidQueryPacketDelegate& QueryDelegate)
 	LanBeacon = new FLanBeacon();
 	if (LanBeacon->Init(LanAnnouncePort))
 	{
-		AddOnValidQueryPacketDelegate(QueryDelegate);
+		AddOnValidQueryPacketDelegate_Handle(QueryDelegate);
 		// We successfully created everything so mark the socket as needing polling
 		LanBeaconState = ELanBeaconState::Hosting;
 		bSuccess = true;
@@ -178,8 +178,8 @@ bool FLANSession::Search(FNboSerializeToBuffer& Packet, FOnValidResponsePacketDe
 			// Set the timestamp for timing out a search
 			LanQueryTimeLeft = LanQueryTimeout;
 
-			AddOnValidResponsePacketDelegate(ResponseDelegate);
-			AddOnSearchingTimeoutDelegate(TimeoutDelegate);
+			AddOnValidResponsePacketDelegate_Handle(ResponseDelegate);
+			AddOnSearchingTimeoutDelegate_Handle(TimeoutDelegate);
 		}
 		else
 		{
diff --git a/Engine/Source/Runtime/Online/OnlineSubsystem/Private/OnlineNotificationHandler.cpp b/Engine/Source/Runtime/Online/OnlineSubsystem/Private/OnlineNotificationHandler.cpp
--- a/Engine/Source/Runtime/Online/OnlineSubsystem/Private/OnlineNotificationHandler.cpp
+++ b/Engine/Source/Runtime/Online/OnlineSubsystem/Private/OnlineNotificationHandler.cpp
@@ -9,17 +9,16 @@
 
 void FOnlineNotificationHandler::AddSystemNotificationBinding(FString NotificationType, const FOnlineNotificationBinding& NewBinding)
 {
-	TArray<FOnlineNotificationBinding>& FoundBindings = SystemBindingMap.FindOrAdd(NotificationType);
-
 	if (!NewBinding.NotificationDelegate.IsBound())
 	{
 		UE_LOG(LogOnline, Error, TEXT("Adding empty notification binding for type %s"), *NotificationType);
 		return;
 	}
 
+	TArray<FOnlineNotificationBinding>& FoundBindings = SystemBindingMap.FindOrAdd(NotificationType);
 	for (int32 i = 0; i < FoundBindings.Num(); i++)
 	{
-		if (FoundBindings[i] == NewBinding)
+		if (FoundBindings[i].NotificationDelegate.DEPRECATED_Compare(NewBinding.NotificationDelegate))
 		{
 			UE_LOG(LogOnline, Error, TEXT("Adding identical notification binding for type %s"), *NotificationType);
 			return;
@@ -29,6 +28,19 @@ void FOnlineNotificationHandler::AddSystemNotificationBinding(FString Notificati
 	FoundBindings.Add(NewBinding);
 }
 
+FDelegateHandle FOnlineNotificationHandler::AddSystemNotificationBinding_Handle(FString NotificationType, const FOnlineNotificationBinding& NewBinding)
+{
+	if (!NewBinding.NotificationDelegate.IsBound())
+	{
+		UE_LOG(LogOnline, Error, TEXT("Adding empty notification binding for type %s"), *NotificationType);
+		return FDelegateHandle();
+	}
+
+	TArray<FOnlineNotificationBinding>& FoundBindings = SystemBindingMap.FindOrAdd(NotificationType);
+	FoundBindings.Add(NewBinding);
+	return FoundBindings.Last().NotificationDelegate.GetHandle();
+}
+
 void FOnlineNotificationHandler::RemoveSystemNotificationBinding(FString NotificationType, const FOnlineNotificationBinding& RemoveBinding)
 {
 	TArray<FOnlineNotificationBinding>* FoundBindings = SystemBindingMap.Find(NotificationType);
@@ -36,7 +48,23 @@ void FOnlineNotificationHandler::RemoveSystemNotificationBinding(FString Notific
 	int32 BindingsRemoved = 0;
 	if (FoundBindings)
 	{
-		BindingsRemoved = FoundBindings->Remove(RemoveBinding);
+		BindingsRemoved = FoundBindings->RemoveAll([&](const FOnlineNotificationBinding& Binding) { return Binding.NotificationDelegate.DEPRECATED_Compare(RemoveBinding.NotificationDelegate); });
+	}
+
+	if (BindingsRemoved == 0)
+	{
+		UE_LOG_ONLINE(Error, TEXT("Attempted to remove binding that could not be found for type %s"), *NotificationType);
+	}
+}
+
+void FOnlineNotificationHandler::RemoveSystemNotificationBinding(FString NotificationType, FDelegateHandle RemoveHandle)
+{
+	TArray<FOnlineNotificationBinding>* FoundBindings = SystemBindingMap.Find(NotificationType);
+
+	int32 BindingsRemoved = 0;
+	if (FoundBindings)
+	{
+		BindingsRemoved = FoundBindings->RemoveAll([=](const FOnlineNotificationBinding& Binding) { return Binding.NotificationDelegate.GetHandle() == RemoveHandle; });
 	}
 
 	if (BindingsRemoved == 0)
@@ -69,7 +97,7 @@ void FOnlineNotificationHandler::AddPlayerNotificationBinding(const FUniqueNetId
 
 	for (int32 i = 0; i < FoundPlayerTypeBindings.Num(); i++)
 	{
-		if (FoundPlayerTypeBindings[i] == NewBinding)
+		if (FoundPlayerTypeBindings[i].NotificationDelegate.DEPRECATED_Compare(NewBinding.NotificationDelegate))
 		{
 			UE_LOG(LogOnline, Error, TEXT("Adding identical notification binding for type %s"), *NotificationType);
 			return;
@@ -79,6 +107,20 @@ void FOnlineNotificationHandler::AddPlayerNotificationBinding(const FUniqueNetId
 	FoundPlayerTypeBindings.Add(NewBinding);
 }
 
+FDelegateHandle FOnlineNotificationHandler::AddPlayerNotificationBinding_Handle(const FUniqueNetId& PlayerId, FString NotificationType, const FOnlineNotificationBinding& NewBinding)
+{
+	if (!NewBinding.NotificationDelegate.IsBound())
+	{
+		UE_LOG(LogOnline, Error, TEXT("Adding empty notification binding for type %s"), *NotificationType);
+		return FDelegateHandle();
+	}
+
+	NotificationTypeBindingsMap& FoundPlayerBindings = PlayerBindingMap.FindOrAdd(PlayerId.ToString());
+	TArray<FOnlineNotificationBinding>& FoundPlayerTypeBindings = FoundPlayerBindings.FindOrAdd(NotificationType);
+	FoundPlayerTypeBindings.Add(NewBinding);
+	return FoundPlayerTypeBindings.Last().NotificationDelegate.GetHandle();
+}
+
 /** Remove the player notification handler for a type */
 void FOnlineNotificationHandler::RemovePlayerNotificationBinding(const FUniqueNetId& PlayerId, FString NotificationType, const FOnlineNotificationBinding& RemoveBinding)
 {
@@ -92,7 +134,30 @@ void FOnlineNotificationHandler::RemovePlayerNotificationBinding(const FUniqueNe
 
 		if (FoundPlayerTypeBindings)
 		{
-			BindingsRemoved = FoundPlayerTypeBindings->Remove(RemoveBinding);
+			BindingsRemoved = FoundPlayerTypeBindings->RemoveAll([&](const FOnlineNotificationBinding& Binding) { return Binding.NotificationDelegate.DEPRECATED_Compare(RemoveBinding.NotificationDelegate); });
+		}
+	}
+
+	if (BindingsRemoved == 0)
+	{
+		UE_LOG_ONLINE(Error, TEXT("Attempted to remove binding that could not be found for player %s type %s"), *PlayerId.ToDebugString(), *NotificationType);
+	}
+}
+
+/** Remove the player notification handler for a type */
+void FOnlineNotificationHandler::RemovePlayerNotificationBinding(const FUniqueNetId& PlayerId, FString NotificationType, FDelegateHandle RemoveHandle)
+{
+	int32 BindingsRemoved = 0;
+
+	NotificationTypeBindingsMap* FoundPlayerBindings = PlayerBindingMap.Find(PlayerId.ToString());
+
+	if (FoundPlayerBindings)
+	{
+		TArray<FOnlineNotificationBinding>* FoundPlayerTypeBindings = FoundPlayerBindings->Find(NotificationType);
+
+		if (FoundPlayerTypeBindings)
+		{
+			BindingsRemoved = FoundPlayerTypeBindings->RemoveAll([=](const FOnlineNotificationBinding& Binding) { return Binding.NotificationDelegate.GetHandle() == RemoveHandle; });
 		}
 	}
 
diff --git a/Engine/Source/Runtime/Online/OnlineSubsystem/Public/OnlineDelegateMacros.h b/Engine/Source/Runtime/Online/OnlineSubsystem/Public/OnlineDelegateMacros.h
--- a/Engine/Source/Runtime/Online/OnlineSubsystem/Public/OnlineDelegateMacros.h
+++ b/Engine/Source/Runtime/Online/OnlineSubsystem/Public/OnlineDelegateMacros.h
@@ -6,14 +6,25 @@
 public: \
 	F##DelegateName DelegateName##Delegates; \
 public: \
+	DELEGATE_DEPRECATED("The Add" #DelegateName "Delegate function is deprecated - please use the Add" #DelegateName "Delegate_Handle function instead which returns a handle.") \
 	virtual void Add##DelegateName##Delegate(const F##DelegateName##Delegate& Delegate) \
 	{ \
 		DelegateName##Delegates.Add(Delegate); \
 	} \
+	DELEGATE_DEPRECATED("The Clear" #DelegateName "Delegate function is deprecated - please use the Add" #DelegateName "Delegate_Handle function and pass its return value to Clear" #DelegateName "Delegate_Handle instead.") \
 	virtual void Clear##DelegateName##Delegate(const F##DelegateName##Delegate& Delegate) \
 	{ \
-		DelegateName##Delegates.Remove(Delegate); \
-    }
+		DelegateName##Delegates.DEPRECATED_Remove(Delegate); \
+	} \
+	virtual FDelegateHandle Add##DelegateName##Delegate_Handle(const F##DelegateName##Delegate& Delegate) \
+	{ \
+		DelegateName##Delegates.Add(Delegate); \
+		return Delegate.GetHandle(); \
+	} \
+	virtual void Clear##DelegateName##Delegate_Handle(FDelegateHandle Handle) \
+	{ \
+		DelegateName##Delegates.Remove(Handle); \
+	}
 
 #define DEFINE_ONLINE_DELEGATE(DelegateName) \
 	DEFINE_ONLINE_DELEGATE_BASE(DelegateName) \
@@ -76,6 +87,7 @@ public: \
 public: \
 	F##DelegateName DelegateName##Delegates[MaxPlayers]; \
 public: \
+	DELEGATE_DEPRECATED("The Add" #DelegateName "Delegate function is deprecated - please use the Add" #DelegateName "Delegate_Handle function instead which returns a handle.") \
 	virtual void Add##DelegateName##Delegate(int32 LocalUserNum, const F##DelegateName##Delegate& Delegate) \
 	{ \
 		if (LocalUserNum >= 0 && LocalUserNum < MaxPlayers) \
@@ -83,11 +95,30 @@ public: \
 			DelegateName##Delegates[LocalUserNum].Add(Delegate); \
 		} \
 	} \
+	DELEGATE_DEPRECATED("The Clear" #DelegateName "Delegate function is deprecated - please use the Clear" #DelegateName "Delegate_Handle and use its return value to pass to Clear" #DelegateName "Delegate_Handle instead.") \
 	virtual void Clear##DelegateName##Delegate(int32 LocalUserNum, const F##DelegateName##Delegate& Delegate) \
 	{ \
 		if (LocalUserNum >= 0 && LocalUserNum < MaxPlayers) \
 		{ \
+			PRAGMA_DISABLE_DEPRECATION_WARNINGS \
 			DelegateName##Delegates[LocalUserNum].Remove(Delegate); \
+			PRAGMA_ENABLE_DEPRECATION_WARNINGS \
+		} \
+	} \
+	virtual FDelegateHandle Add##DelegateName##Delegate_Handle(int32 LocalUserNum, const F##DelegateName##Delegate& Delegate) \
+	{ \
+		FDelegateHandle Result; \
+		if (LocalUserNum >= 0 && LocalUserNum < MaxPlayers) \
+		{ \
+			Result = DelegateName##Delegates[LocalUserNum].Add(Delegate); \
+		} \
+		return Result; \
+	} \
+	virtual void Clear##DelegateName##Delegate_Handle(int32 LocalUserNum, FDelegateHandle Handle) \
+	{ \
+		if (LocalUserNum >= 0 && LocalUserNum < MaxPlayers) \
+		{ \
+			DelegateName##Delegates[LocalUserNum].Remove(Handle); \
 		} \
 	}
 
diff --git a/Engine/Source/Runtime/Online/OnlineSubsystem/Public/OnlineNotificationHandler.h b/Engine/Source/Runtime/Online/OnlineSubsystem/Public/OnlineNotificationHandler.h
--- a/Engine/Source/Runtime/Online/OnlineSubsystem/Public/OnlineNotificationHandler.h
+++ b/Engine/Source/Runtime/Online/OnlineSubsystem/Public/OnlineNotificationHandler.h
@@ -43,11 +43,6 @@ struct FOnlineNotificationBinding
 	/** Delegate to call when this binding is activated */
 	FHandleOnlineNotificationSignature NotificationDelegate;
 
-	bool operator==(const FOnlineNotificationBinding& Other) const
-	{
-		return NotificationDelegate == Other.NotificationDelegate;
-	}
-
 	FOnlineNotificationBinding()
 	{}
 
@@ -80,10 +75,18 @@ public:
 	// SYSTEM NOTIFICATION BINDINGS
 
 	/** Add a notification binding for a type */
+	DELEGATE_DEPRECATED("AddSystemNotificationBinding is deprecated because it relies on deprecated delegate comparison to ignore multiple registrations, instead use AddSystemNotificationBinding_Handle and avoid multiple registrations.")
 	void AddSystemNotificationBinding(FString NotificationType, const FOnlineNotificationBinding& NewBinding);
 
+	/** Add a notification binding for a type */
+	FDelegateHandle AddSystemNotificationBinding_Handle(FString NotificationType, const FOnlineNotificationBinding& NewBinding);
+
 	/** Remove the notification handler for a type */
-	void RemoveSystemNotificationBinding(FString NotificationType, const FOnlineNotificationBinding& NewBinding);
+	DELEGATE_DEPRECATED("This overload of RemoveSystemNotificationBinding is deprecated, instead use RemoveSystemNotificationBinding passing the result of AddSystemNotificationBinding_Handle.")
+	void RemoveSystemNotificationBinding(FString NotificationType, const FOnlineNotificationBinding& RemoveBinding);
+
+	/** Remove the notification handler for a type */
+	void RemoveSystemNotificationBinding(FString NotificationType, FDelegateHandle RemoveHandle);
 
 	/** Resets all system notification handlers */
 	void ResetSystemNotificationBindings();
@@ -91,10 +94,18 @@ public:
 	// PLAYER NOTIFICATION BINDINGS
 
 	/** Add a notification binding for a type */
+	DELEGATE_DEPRECATED("AddPlayerNotificationBinding is deprecated because it relies on deprecated delegate comparison to ignore multiple registrations, instead use AddPlayerNotificationBinding_Handle and avoid multiple registrations.")
 	void AddPlayerNotificationBinding(const FUniqueNetId& PlayerId, FString NotificationType, const FOnlineNotificationBinding& NewBinding);
 
+	/** Add a notification binding for a type */
+	FDelegateHandle AddPlayerNotificationBinding_Handle(const FUniqueNetId& PlayerId, FString NotificationType, const FOnlineNotificationBinding& NewBinding);
+
+	/** Remove the player notification handler for a type */
+	DELEGATE_DEPRECATED("This overload of RemovePlayerNotificationBinding is deprecated, instead use RemovePlayerNotificationBinding passing the result of AddPlayerNotificationBinding_Handle.")
+	void RemovePlayerNotificationBinding(const FUniqueNetId& PlayerId, FString NotificationType, const FOnlineNotificationBinding& RemoveBinding);
+
 	/** Remove the player notification handler for a type */
-	void RemovePlayerNotificationBinding(const FUniqueNetId& PlayerId, FString NotificationType, const FOnlineNotificationBinding& NewBinding);
+	void RemovePlayerNotificationBinding(const FUniqueNetId& PlayerId, FString NotificationType, FDelegateHandle RemoveHandle);
 
 	/** Resets a player's notification handlers */
 	void ResetPlayerNotificationBindings(const FUniqueNetId& PlayerId);
@@ -108,4 +119,4 @@ public:
 	void DeliverNotification(const FOnlineNotification& Notification);
 };
 
-typedef TSharedPtr<FOnlineNotificationHandler, ESPMode::ThreadSafe> FOnlineNotificationHandlerPtr;
\ No newline at end of file
+typedef TSharedPtr<FOnlineNotificationHandler, ESPMode::ThreadSafe> FOnlineNotificationHandlerPtr;
diff --git a/Engine/Source/Runtime/Online/OnlineSubsystemSteam/Private/OnlineSubsystemSteam.cpp b/Engine/Source/Runtime/Online/OnlineSubsystemSteam/Private/OnlineSubsystemSteam.cpp
--- a/Engine/Source/Runtime/Online/OnlineSubsystemSteam/Private/OnlineSubsystemSteam.cpp
+++ b/Engine/Source/Runtime/Online/OnlineSubsystemSteam/Private/OnlineSubsystemSteam.cpp
@@ -666,6 +666,10 @@ void FOnlineSubsystemSteam::ClearUserCloudFiles()
 	UserCloudData.Empty();
 }
 
+static FDelegateHandle GOnEnumerateUserFilesCompleteDelegateHandle;
+
+TMap<IOnlineUserCloud*, FDelegateHandle> GPerCloudDeleteFromEnumerateUserFilesCompleteDelegateHandles;
+
 static void DeleteFromEnumerateUserFilesComplete(bool bWasSuccessful, const FUniqueNetId& UserId)
 {
 	IOnlineSubsystem* OnlineSub = IOnlineSubsystem::Get();
@@ -673,8 +677,8 @@ static void DeleteFromEnumerateUserFilesComplete(bool bWasSuccessful, const FUni
 
 	IOnlineUserCloudPtr UserCloud = OnlineSub->GetUserCloudInterface();
 
-	FOnEnumerateUserFilesCompleteDelegate Delegate = FOnEnumerateUserFilesCompleteDelegate::CreateStatic(&DeleteFromEnumerateUserFilesComplete);
-	UserCloud->ClearOnEnumerateUserFilesCompleteDelegate(Delegate);
+	UserCloud->ClearOnEnumerateUserFilesCompleteDelegate_Handle(GOnEnumerateUserFilesCompleteDelegateHandle);
+	GPerCloudDeleteFromEnumerateUserFilesCompleteDelegateHandles.Remove(UserCloud.Get());
 	if (bWasSuccessful)
 	{
 		TArray<FCloudFileHeader> UserFiles;
@@ -696,7 +700,7 @@ bool FOnlineSubsystemSteam::Exec(UWorld* InWorld, const TCHAR* Cmd, FOutputDevic
 		FUniqueNetIdSteam SteamId(SteamUser()->GetSteamID());
 
 		FOnEnumerateUserFilesCompleteDelegate Delegate = FOnEnumerateUserFilesCompleteDelegate::CreateStatic(&DeleteFromEnumerateUserFilesComplete);
-		UserCloud->AddOnEnumerateUserFilesCompleteDelegate(Delegate);
+		GPerCloudDeleteFromEnumerateUserFilesCompleteDelegateHandles.Add(UserCloud.Get(), UserCloud->AddOnEnumerateUserFilesCompleteDelegate_Handle(Delegate));
 		UserCloud->EnumerateUserFiles(SteamId);
 		return true;
 	}
diff --git a/Engine/Source/Runtime/Online/OnlineSubsystemUtils/Classes/ConnectionCallbackProxy.h b/Engine/Source/Runtime/Online/OnlineSubsystemUtils/Classes/ConnectionCallbackProxy.h
--- a/Engine/Source/Runtime/Online/OnlineSubsystemUtils/Classes/ConnectionCallbackProxy.h
+++ b/Engine/Source/Runtime/Online/OnlineSubsystemUtils/Classes/ConnectionCallbackProxy.h
@@ -40,6 +40,9 @@ private:
 	// Connection callback delegate
 	FOnLoginCompleteDelegate OnLoginCompleteDelegate;
 
+	// OnLoginComplete delegate handle
+	FDelegateHandle OnLoginCompleteDelegateHandle;
+
 	// The world context object in which this call is taking place
 	UObject* WorldContextObject;
 };
\ No newline at end of file
diff --git a/Engine/Source/Runtime/Online/OnlineSubsystemUtils/Classes/CreateSessionCallbackProxy.h b/Engine/Source/Runtime/Online/OnlineSubsystemUtils/Classes/CreateSessionCallbackProxy.h
--- a/Engine/Source/Runtime/Online/OnlineSubsystemUtils/Classes/CreateSessionCallbackProxy.h
+++ b/Engine/Source/Runtime/Online/OnlineSubsystemUtils/Classes/CreateSessionCallbackProxy.h
@@ -41,6 +41,10 @@ private:
 	// The delegate executed by the online subsystem
 	FOnStartSessionCompleteDelegate StartCompleteDelegate;
 
+	// Handles to the registered delegates above
+	FDelegateHandle CreateCompleteDelegateHandle;
+	FDelegateHandle StartCompleteDelegateHandle;
+
 	// Number of public connections
 	int NumPublicConnections;
 
diff --git a/Engine/Source/Runtime/Online/OnlineSubsystemUtils/Classes/DestroySessionCallbackProxy.h b/Engine/Source/Runtime/Online/OnlineSubsystemUtils/Classes/DestroySessionCallbackProxy.h
--- a/Engine/Source/Runtime/Online/OnlineSubsystemUtils/Classes/DestroySessionCallbackProxy.h
+++ b/Engine/Source/Runtime/Online/OnlineSubsystemUtils/Classes/DestroySessionCallbackProxy.h
@@ -36,6 +36,9 @@ private:
 	// The delegate executed by the online subsystem
 	FOnDestroySessionCompleteDelegate Delegate;
 
+	// Handle to the registered OnDestroySessionComplete delegate
+	FDelegateHandle DelegateHandle;
+
 	// The world context object in which this call is taking place
 	UObject* WorldContextObject;
 };
diff --git a/Engine/Source/Runtime/Online/OnlineSubsystemUtils/Classes/FindSessionsCallbackProxy.h b/Engine/Source/Runtime/Online/OnlineSubsystemUtils/Classes/FindSessionsCallbackProxy.h
--- a/Engine/Source/Runtime/Online/OnlineSubsystemUtils/Classes/FindSessionsCallbackProxy.h
+++ b/Engine/Source/Runtime/Online/OnlineSubsystemUtils/Classes/FindSessionsCallbackProxy.h
@@ -58,6 +58,9 @@ private:
 	// The delegate executed by the online subsystem
 	FOnFindSessionsCompleteDelegate Delegate;
 
+	// Handle to the registered OnFindSessionsComplete delegate
+	FDelegateHandle DelegateHandle;
+
 	// Object to track search results
 	TSharedPtr<FOnlineSessionSearch> SearchObject;
 
diff --git a/Engine/Source/Runtime/Online/OnlineSubsystemUtils/Classes/InAppPurchaseCallbackProxy.h b/Engine/Source/Runtime/Online/OnlineSubsystemUtils/Classes/InAppPurchaseCallbackProxy.h
--- a/Engine/Source/Runtime/Online/OnlineSubsystemUtils/Classes/InAppPurchaseCallbackProxy.h
+++ b/Engine/Source/Runtime/Online/OnlineSubsystemUtils/Classes/InAppPurchaseCallbackProxy.h
@@ -46,6 +46,9 @@ private:
 	/** Delegate called when a InAppPurchase has been successfully read */
 	FOnInAppPurchaseCompleteDelegate InAppPurchaseCompleteDelegate;
 
+	/** Handle to the registered InAppPurchaseCompleteDelegate */
+	FDelegateHandle InAppPurchaseCompleteDelegateHandle;
+
 	/** The InAppPurchase read request */
 	FOnlineInAppPurchaseTransactionPtr PurchaseRequest;
 
@@ -57,4 +60,6 @@ private:
 
 	/** Did the purchase succeed? */
 	EInAppPurchaseState::Type SavedPurchaseState;
+
+	FTimerHandle OnInAppPurchaseComplete_DelayedTimerHandle;
 };
diff --git a/Engine/Source/Runtime/Online/OnlineSubsystemUtils/Classes/InAppPurchaseQueryCallbackProxy.h b/Engine/Source/Runtime/Online/OnlineSubsystemUtils/Classes/InAppPurchaseQueryCallbackProxy.h
--- a/Engine/Source/Runtime/Online/OnlineSubsystemUtils/Classes/InAppPurchaseQueryCallbackProxy.h
+++ b/Engine/Source/Runtime/Online/OnlineSubsystemUtils/Classes/InAppPurchaseQueryCallbackProxy.h
@@ -43,6 +43,9 @@ private:
 	/** Delegate called when a InAppPurchase has been successfully read */
 	FOnQueryForAvailablePurchasesCompleteDelegate InAppPurchaseReadCompleteDelegate;
 
+	/** InAppPurchaseReadComplete delegate handle */
+	FDelegateHandle InAppPurchaseReadCompleteDelegateHandle;
+
 	/** The InAppPurchase read request */
 	FOnlineProductInformationReadPtr ReadObject;
 
@@ -55,4 +58,6 @@ private:
 	// Did the read succeed?
 	bool bSavedWasSuccessful;
 	TArray<FInAppPurchaseProductInfo> SavedProductInformation;
+
+	FTimerHandle OnInAppPurchaseRead_DelayedTimerHandle;
 };
diff --git a/Engine/Source/Runtime/Online/OnlineSubsystemUtils/Classes/JoinSessionCallbackProxy.h b/Engine/Source/Runtime/Online/OnlineSubsystemUtils/Classes/JoinSessionCallbackProxy.h
--- a/Engine/Source/Runtime/Online/OnlineSubsystemUtils/Classes/JoinSessionCallbackProxy.h
+++ b/Engine/Source/Runtime/Online/OnlineSubsystemUtils/Classes/JoinSessionCallbackProxy.h
@@ -39,6 +39,9 @@ private:
 	// The delegate executed by the online subsystem
 	FOnJoinSessionCompleteDelegate Delegate;
 
+	// Handle to the registered FOnJoinSessionComplete delegate
+	FDelegateHandle DelegateHandle;
+
 	// The world context object in which this call is taking place
 	UObject* WorldContextObject;
 };
diff --git a/Engine/Source/Runtime/Online/OnlineSubsystemUtils/Classes/LeaderboardFlushCallbackProxy.cpp b/Engine/Source/Runtime/Online/OnlineSubsystemUtils/Classes/LeaderboardFlushCallbackProxy.cpp
--- a/Engine/Source/Runtime/Online/OnlineSubsystemUtils/Classes/LeaderboardFlushCallbackProxy.cpp
+++ b/Engine/Source/Runtime/Online/OnlineSubsystemUtils/Classes/LeaderboardFlushCallbackProxy.cpp
@@ -25,8 +25,8 @@ void ULeaderboardFlushCallbackProxy::TriggerFlush(APlayerController* PlayerContr
 				if (Leaderboards.IsValid())
 				{
 					// Register the completion callback
-					LeaderboardFlushCompleteDelegate = FOnLeaderboardFlushCompleteDelegate::CreateUObject(this, &ULeaderboardFlushCallbackProxy::OnFlushCompleted);
-					Leaderboards->AddOnLeaderboardFlushCompleteDelegate(LeaderboardFlushCompleteDelegate);
+					LeaderboardFlushCompleteDelegate       = FOnLeaderboardFlushCompleteDelegate::CreateUObject(this, &ULeaderboardFlushCallbackProxy::OnFlushCompleted);
+					LeaderboardFlushCompleteDelegateHandle = Leaderboards->AddOnLeaderboardFlushCompleteDelegate_Handle(LeaderboardFlushCompleteDelegate);
 
 					// Flush the leaderboard
 					Leaderboards->FlushLeaderboards(InSessionName);
@@ -81,7 +81,7 @@ void ULeaderboardFlushCallbackProxy::RemoveDelegate()
 			IOnlineLeaderboardsPtr Leaderboards = OnlineSub->GetLeaderboardsInterface();
 			if (Leaderboards.IsValid())
 			{
-				Leaderboards->ClearOnLeaderboardFlushCompleteDelegate(LeaderboardFlushCompleteDelegate);
+				Leaderboards->ClearOnLeaderboardFlushCompleteDelegate_Handle(LeaderboardFlushCompleteDelegateHandle);
 			}
 		}
 	}
diff --git a/Engine/Source/Runtime/Online/OnlineSubsystemUtils/Classes/LeaderboardFlushCallbackProxy.h b/Engine/Source/Runtime/Online/OnlineSubsystemUtils/Classes/LeaderboardFlushCallbackProxy.h
--- a/Engine/Source/Runtime/Online/OnlineSubsystemUtils/Classes/LeaderboardFlushCallbackProxy.h
+++ b/Engine/Source/Runtime/Online/OnlineSubsystemUtils/Classes/LeaderboardFlushCallbackProxy.h
@@ -42,6 +42,9 @@ private:
 	/** Delegate called when a leaderboard has been successfully flushed */
 	FOnLeaderboardFlushCompleteDelegate LeaderboardFlushCompleteDelegate;
 
+	/** OnLeaderboardFlushComplete delegate handle */
+	FDelegateHandle LeaderboardFlushCompleteDelegateHandle;
+
 	/** Did we fail immediately? */
 	bool bFailedToEvenSubmit;
 };
diff --git a/Engine/Source/Runtime/Online/OnlineSubsystemUtils/Classes/LeaderboardQueryCallbackProxy.cpp b/Engine/Source/Runtime/Online/OnlineSubsystemUtils/Classes/LeaderboardQueryCallbackProxy.cpp
--- a/Engine/Source/Runtime/Online/OnlineSubsystemUtils/Classes/LeaderboardQueryCallbackProxy.cpp
+++ b/Engine/Source/Runtime/Online/OnlineSubsystemUtils/Classes/LeaderboardQueryCallbackProxy.cpp
@@ -34,8 +34,8 @@ void ULeaderboardQueryCallbackProxy::TriggerQuery(APlayerController* PlayerContr
 					new (ReadObject->ColumnMetadata) FColumnMetaData(StatName, StatType);
 
 					// Register the completion callback
-					LeaderboardReadCompleteDelegate = FOnLeaderboardReadCompleteDelegate::CreateUObject(this, &ULeaderboardQueryCallbackProxy::OnStatsRead);
-					Leaderboards->AddOnLeaderboardReadCompleteDelegate(LeaderboardReadCompleteDelegate);
+					LeaderboardReadCompleteDelegate       = FOnLeaderboardReadCompleteDelegate::CreateUObject(this, &ULeaderboardQueryCallbackProxy::OnStatsRead);
+					LeaderboardReadCompleteDelegateHandle = Leaderboards->AddOnLeaderboardReadCompleteDelegate_Handle(LeaderboardReadCompleteDelegate);
 
 					TArray< TSharedRef<FUniqueNetId> > ListOfIDs;
 					ListOfIDs.Add(UserID.ToSharedRef());
@@ -108,7 +108,7 @@ void ULeaderboardQueryCallbackProxy::OnStatsRead(bool bWasSuccessful)
 	{
 		// Use a dummy timer handle as we don't need to store it for later but we don't need to look for something to clear
 		FTimerHandle TimerHandle;
-		World->GetTimerManager().SetTimer(this, &ULeaderboardQueryCallbackProxy::OnStatsRead_Delayed, 0.001f, false);
+		World->GetTimerManager().SetTimer(OnStatsRead_DelayedTimerHandle, this, &ULeaderboardQueryCallbackProxy::OnStatsRead_Delayed, 0.001f, false);
 	}
 	ReadObject = NULL;
 }
@@ -134,7 +134,7 @@ void ULeaderboardQueryCallbackProxy::RemoveDelegate()
 			IOnlineLeaderboardsPtr Leaderboards = OnlineSub->GetLeaderboardsInterface();
 			if (Leaderboards.IsValid())
 			{
-				Leaderboards->ClearOnLeaderboardReadCompleteDelegate(LeaderboardReadCompleteDelegate);
+				Leaderboards->ClearOnLeaderboardReadCompleteDelegate_Handle(LeaderboardReadCompleteDelegateHandle);
 			}
 		}
 	}
diff --git a/Engine/Source/Runtime/Online/OnlineSubsystemUtils/Classes/LeaderboardQueryCallbackProxy.h b/Engine/Source/Runtime/Online/OnlineSubsystemUtils/Classes/LeaderboardQueryCallbackProxy.h
--- a/Engine/Source/Runtime/Online/OnlineSubsystemUtils/Classes/LeaderboardQueryCallbackProxy.h
+++ b/Engine/Source/Runtime/Online/OnlineSubsystemUtils/Classes/LeaderboardQueryCallbackProxy.h
@@ -43,6 +43,9 @@ private:
 	/** Delegate called when a leaderboard has been successfully read */
 	FOnLeaderboardReadCompleteDelegate LeaderboardReadCompleteDelegate;
 
+	/** LeaderboardReadComplete delegate handle */
+	FDelegateHandle LeaderboardReadCompleteDelegateHandle;
+
 	/** The leaderboard read request */
 	TSharedPtr<class FOnlineLeaderboardRead, ESPMode::ThreadSafe> ReadObject;
 
@@ -58,4 +61,6 @@ private:
 	// Did the read succeed?
 	bool bSavedWasSuccessful;
 	int32 SavedValue;
+
+	FTimerHandle OnStatsRead_DelayedTimerHandle;
 };
diff --git a/Engine/Source/Runtime/Online/OnlineSubsystemUtils/Classes/OnlineSessionClient.h b/Engine/Source/Runtime/Online/OnlineSubsystemUtils/Classes/OnlineSessionClient.h
--- a/Engine/Source/Runtime/Online/OnlineSubsystemUtils/Classes/OnlineSessionClient.h
+++ b/Engine/Source/Runtime/Online/OnlineSubsystemUtils/Classes/OnlineSessionClient.h
@@ -29,6 +29,13 @@ protected:
 	/** Delegate after joining a session */
 	FOnJoinSessionCompleteDelegate OnJoinSessionCompleteDelegate;
 
+	// Handles to the above delegates
+	FDelegateHandle OnSessionInviteAcceptedDelegateHandle;
+	FDelegateHandle OnEndForJoinSessionCompleteDelegateHandle;
+	FDelegateHandle OnDestroyForJoinSessionCompleteDelegateHandle;
+	FDelegateHandle OnDestroyForMainMenuCompleteDelegateHandle;
+	FDelegateHandle OnJoinSessionCompleteDelegateHandle;
+
 	/** Cached invite/search result while in the process of tearing down an existing session */
 	FOnlineSessionSearchResult CachedSessionResult;
 	/** Is this join from an invite */
@@ -76,6 +83,16 @@ protected:
 	 */
 	void EndExistingSession(FName SessionName, FOnEndSessionCompleteDelegate& Delegate);
 
+private:
+	/**
+	 * Implementation of EndExistingSession
+	 *
+	 * @param SessionName name of session to end
+	 * @param Delegate delegate to call at session end
+	 */
+	FDelegateHandle EndExistingSession_Impl(FName SessionName, FOnEndSessionCompleteDelegate& Delegate);
+
+protected:
 	/**
 	 * Transition from destroying a session to joining a new one of the same name
 	 *
@@ -100,6 +117,17 @@ protected:
 	 */
 	void DestroyExistingSession(FName SessionName, FOnDestroySessionCompleteDelegate& Delegate);
 
+private:
+	/**
+	 * Implementation of DestroyExistingSession
+	 *
+	 * @param SessionName name of session to destroy
+	 * @param Delegate delegate to call at session destruction
+	 * @return Handle to the added delegate.
+	 */
+	FDelegateHandle DestroyExistingSession_Impl(FName SessionName, FOnDestroySessionCompleteDelegate& Delegate);
+
+protected:
 	/**
 	 * Delegate fired when the joining process for an online session has completed
 	 *
diff --git a/Engine/Source/Runtime/Online/OnlineSubsystemUtils/Private/ConnectionCallbackProxy.cpp b/Engine/Source/Runtime/Online/OnlineSubsystemUtils/Private/ConnectionCallbackProxy.cpp
--- a/Engine/Source/Runtime/Online/OnlineSubsystemUtils/Private/ConnectionCallbackProxy.cpp
+++ b/Engine/Source/Runtime/Online/OnlineSubsystemUtils/Private/ConnectionCallbackProxy.cpp
@@ -36,7 +36,7 @@ void UConnectionCallbackProxy::Activate()
 
 			if (!OnlineIdentity->OnLoginCompleteDelegates[ControllerId].IsBoundToObject(this))
 			{
-				OnlineIdentity->AddOnLoginCompleteDelegate(ControllerId, OnLoginCompleteDelegate);
+				OnLoginCompleteDelegateHandle = OnlineIdentity->AddOnLoginCompleteDelegate_Handle(ControllerId, OnLoginCompleteDelegate);
 				OnlineIdentity->Login(ControllerId, FOnlineAccountCredentials()); /// Probably need to supply real creds here somehow... doesn't apply for all imple however.
 			}
 			else
@@ -67,7 +67,7 @@ void UConnectionCallbackProxy::OnLoginCompleted(int32 LocalUserNum, bool bWasSuc
 		IOnlineIdentityPtr OnlineIdentity = Helper.OnlineSub->GetIdentityInterface();
 		if (OnlineIdentity.IsValid())
 		{
-			OnlineIdentity->ClearOnLoginCompleteDelegate(LocalUserNum, OnLoginCompleteDelegate);
+			OnlineIdentity->ClearOnLoginCompleteDelegate_Handle(LocalUserNum, OnLoginCompleteDelegateHandle);
 		}
 	}
 
diff --git a/Engine/Source/Runtime/Online/OnlineSubsystemUtils/Private/CreateSessionCallbackProxy.cpp b/Engine/Source/Runtime/Online/OnlineSubsystemUtils/Private/CreateSessionCallbackProxy.cpp
--- a/Engine/Source/Runtime/Online/OnlineSubsystemUtils/Private/CreateSessionCallbackProxy.cpp
+++ b/Engine/Source/Runtime/Online/OnlineSubsystemUtils/Private/CreateSessionCallbackProxy.cpp
@@ -34,7 +34,7 @@ void UCreateSessionCallbackProxy::Activate()
 		auto Sessions = Helper.OnlineSub->GetSessionInterface();
 		if (Sessions.IsValid())
 		{
-			Sessions->AddOnCreateSessionCompleteDelegate(CreateCompleteDelegate);
+			CreateCompleteDelegateHandle = Sessions->AddOnCreateSessionCompleteDelegate_Handle(CreateCompleteDelegate);
 
 			FOnlineSessionSettings Settings;
 			Settings.NumPublicConnections = NumPublicConnections;
@@ -69,11 +69,11 @@ void UCreateSessionCallbackProxy::OnCreateCompleted(FName SessionName, bool bWas
 		auto Sessions = Helper.OnlineSub->GetSessionInterface();
 		if (Sessions.IsValid())
 		{
-			Sessions->ClearOnCreateSessionCompleteDelegate(CreateCompleteDelegate);
+			Sessions->ClearOnCreateSessionCompleteDelegate_Handle(CreateCompleteDelegateHandle);
 			
 			if (bWasSuccessful)
 			{
-				Sessions->AddOnStartSessionCompleteDelegate(StartCompleteDelegate);
+				StartCompleteDelegateHandle = Sessions->AddOnStartSessionCompleteDelegate_Handle(StartCompleteDelegate);
 				Sessions->StartSession(GameSessionName);
 
 				// OnStartCompleted will get called, nothing more to do now
@@ -98,7 +98,7 @@ void UCreateSessionCallbackProxy::OnStartCompleted(FName SessionName, bool bWasS
 		auto Sessions = Helper.OnlineSub->GetSessionInterface();
 		if (Sessions.IsValid())
 		{
-			Sessions->ClearOnStartSessionCompleteDelegate(StartCompleteDelegate);
+			Sessions->ClearOnStartSessionCompleteDelegate_Handle(StartCompleteDelegateHandle);
 		}
 	}
 
diff --git a/Engine/Source/Runtime/Online/OnlineSubsystemUtils/Private/DestroySessionCallbackProxy.cpp b/Engine/Source/Runtime/Online/OnlineSubsystemUtils/Private/DestroySessionCallbackProxy.cpp
--- a/Engine/Source/Runtime/Online/OnlineSubsystemUtils/Private/DestroySessionCallbackProxy.cpp
+++ b/Engine/Source/Runtime/Online/OnlineSubsystemUtils/Private/DestroySessionCallbackProxy.cpp
@@ -30,7 +30,7 @@ void UDestroySessionCallbackProxy::Activate()
 		auto Sessions = Helper.OnlineSub->GetSessionInterface();
 		if (Sessions.IsValid())
 		{
-			Sessions->AddOnDestroySessionCompleteDelegate(Delegate);
+			DelegateHandle = Sessions->AddOnDestroySessionCompleteDelegate_Handle(Delegate);
 			Sessions->DestroySession(GameSessionName);
 
 			// OnCompleted will get called, nothing more to do now
@@ -56,7 +56,7 @@ void UDestroySessionCallbackProxy::OnCompleted(FName SessionName, bool bWasSucce
 		auto Sessions = Helper.OnlineSub->GetSessionInterface();
 		if (Sessions.IsValid())
 		{
-			Sessions->ClearOnDestroySessionCompleteDelegate(Delegate);
+			Sessions->ClearOnDestroySessionCompleteDelegate_Handle(DelegateHandle);
 		}
 	}
 
diff --git a/Engine/Source/Runtime/Online/OnlineSubsystemUtils/Private/FindSessionsCallbackProxy.cpp b/Engine/Source/Runtime/Online/OnlineSubsystemUtils/Private/FindSessionsCallbackProxy.cpp
--- a/Engine/Source/Runtime/Online/OnlineSubsystemUtils/Private/FindSessionsCallbackProxy.cpp
+++ b/Engine/Source/Runtime/Online/OnlineSubsystemUtils/Private/FindSessionsCallbackProxy.cpp
@@ -33,7 +33,7 @@ void UFindSessionsCallbackProxy::Activate()
 		auto Sessions = Helper.OnlineSub->GetSessionInterface();
 		if (Sessions.IsValid())
 		{
-			Sessions->AddOnFindSessionsCompleteDelegate(Delegate);
+			DelegateHandle = Sessions->AddOnFindSessionsCompleteDelegate_Handle(Delegate);
 			
 			SearchObject = MakeShareable(new FOnlineSessionSearch);
 			SearchObject->MaxSearchResults = MaxResults;
@@ -66,7 +66,7 @@ void UFindSessionsCallbackProxy::OnCompleted(bool bSuccess)
 		auto Sessions = Helper.OnlineSub->GetSessionInterface();
 		if (Sessions.IsValid())
 		{
-			Sessions->ClearOnFindSessionsCompleteDelegate(Delegate);
+			Sessions->ClearOnFindSessionsCompleteDelegate_Handle(DelegateHandle);
 		}
 	}
 
diff --git a/Engine/Source/Runtime/Online/OnlineSubsystemUtils/Private/InAppPurchaseCallbackProxy.cpp b/Engine/Source/Runtime/Online/OnlineSubsystemUtils/Private/InAppPurchaseCallbackProxy.cpp
--- a/Engine/Source/Runtime/Online/OnlineSubsystemUtils/Private/InAppPurchaseCallbackProxy.cpp
+++ b/Engine/Source/Runtime/Online/OnlineSubsystemUtils/Private/InAppPurchaseCallbackProxy.cpp
@@ -29,7 +29,7 @@ void UInAppPurchaseCallbackProxy::Trigger(APlayerController* PlayerController, c
 
 				// Register the completion callback
 				InAppPurchaseCompleteDelegate = FOnInAppPurchaseCompleteDelegate::CreateUObject(this, &UInAppPurchaseCallbackProxy::OnInAppPurchaseComplete);
-				StoreInterface->AddOnInAppPurchaseCompleteDelegate(InAppPurchaseCompleteDelegate);
+				InAppPurchaseCompleteDelegateHandle = StoreInterface->AddOnInAppPurchaseCompleteDelegate_Handle(InAppPurchaseCompleteDelegate);
 
 				// Set-up, and trigger the transaction through the store interface
 				PurchaseRequest = MakeShareable(new FOnlineInAppPurchaseTransaction());
@@ -67,7 +67,7 @@ void UInAppPurchaseCallbackProxy::OnInAppPurchaseComplete(EInAppPurchaseState::T
 	{
 		// Use a local timer handle as we don't need to store it for later but we don't need to look for something to clear
 		FTimerHandle TimerHandle;
-		World->GetTimerManager().SetTimer(this, &UInAppPurchaseCallbackProxy::OnInAppPurchaseComplete_Delayed, 0.001f, false);
+		World->GetTimerManager().SetTimer(OnInAppPurchaseComplete_DelayedTimerHandle, this, &UInAppPurchaseCallbackProxy::OnInAppPurchaseComplete_Delayed, 0.001f, false);
     }
     else
     {
@@ -108,7 +108,7 @@ void UInAppPurchaseCallbackProxy::RemoveDelegate()
 			IOnlineStorePtr InAppPurchases = OnlineSub->GetStoreInterface();
 			if (InAppPurchases.IsValid())
 			{
-				InAppPurchases->ClearOnInAppPurchaseCompleteDelegate(InAppPurchaseCompleteDelegate);
+				InAppPurchases->ClearOnInAppPurchaseCompleteDelegate_Handle(InAppPurchaseCompleteDelegateHandle);
 			}
 		}
 	}
