Texture popping fix for when the pool is full and we're streaming out mip levels that are in use by the GPU.

diff --git a/Engine/Source/Runtime/Engine/Private/Texture2D.cpp b/Engine/Source/Runtime/Engine/Private/Texture2D.cpp
--- a/Engine/Source/Runtime/Engine/Private/Texture2D.cpp
+++ b/Engine/Source/Runtime/Engine/Private/Texture2D.cpp
@@ -1612,7 +1612,16 @@ void FTexture2DResource::UpdateMipCount()
 		PendingFirstMip	= OwnerMips.Num() - Owner->RequestedMips;
 		check(PendingFirstMip>=0);
 
-		RHIVirtualTextureSetFirstMipInMemory(Texture2DRHI, PendingFirstMip);
+		// When removing levels we need to start the process of the mip not being visible to the GPU
+		// Whereas when we are streaming in levels we need to allocate physical backing for the mip
+		if ( Owner->RequestedMips < Owner->ResidentMips )
+		{
+			RHIVirtualTextureSetFirstMipVisible(Texture2DRHI, PendingFirstMip);
+		}
+		else
+		{
+			RHIVirtualTextureSetFirstMipInMemory(Texture2DRHI, PendingFirstMip);
+		}
 
 		bUsingAsyncCreation = false;
 
@@ -2053,6 +2062,7 @@ void FTexture2DResource::FinalizeMipCount()
 		{
 			// Show the streamed in mip levels
 			RHIVirtualTextureSetFirstMipVisible(Texture2DRHI, PendingFirstMip);
+			RHIVirtualTextureSetFirstMipInMemory(Texture2DRHI, PendingFirstMip);
 			bSuccess		= true;
 			CurrentFirstMip = PendingFirstMip;
 
