#ue4 - Fix NavWalking behavior not getting restored if a Character is set to falling and then lands again.

diff --git a/Engine/Source/Runtime/Engine/Private/CharacterMovementComponent.cpp b/Engine/Source/Runtime/Engine/Private/CharacterMovementComponent.cpp
--- a/Engine/Source/Runtime/Engine/Private/CharacterMovementComponent.cpp
+++ b/Engine/Source/Runtime/Engine/Private/CharacterMovementComponent.cpp
@@ -666,6 +666,7 @@ void UCharacterMovementComponent::OnMovementModeChanged(EMovementMode PreviousMo
 	if (MovementMode == MOVE_NavWalking)
 	{
 		SetNavWalkingPhysics(true);
+		GroundMovementMode = MovementMode;
 	}
 	else if (PreviousMovementMode == MOVE_NavWalking)
 	{
