Only instanced meshes with bAffectDistanceFieldLighting will keep their CPU instance data copy - saves a bunch of memory on grass

diff --git a/Engine/Source/Runtime/Engine/Private/InstancedStaticMesh.cpp b/Engine/Source/Runtime/Engine/Private/InstancedStaticMesh.cpp
--- a/Engine/Source/Runtime/Engine/Private/InstancedStaticMesh.cpp
+++ b/Engine/Source/Runtime/Engine/Private/InstancedStaticMesh.cpp
@@ -69,6 +69,14 @@ void FStaticMeshInstanceBuffer::Init(UInstancedStaticMeshComponent* InComponent,
 	// Allocate the vertex data storage type.
 	AllocateData();
 
+	static const auto CVar = IConsoleManager::Get().FindTConsoleVariableDataInt(TEXT("r.GenerateMeshDistanceFields"));
+
+	const bool bNeedsCPUAccess = InComponent->CastShadow && InComponent->bAffectDistanceFieldLighting 
+		// Distance field algorithms need access to instance data on the CPU
+		&& CVar->GetValueOnGameThread() != 0;
+
+	InstanceData->SetAllowCPUAccess(InstanceData->GetAllowCPUAccess() || bNeedsCPUAccess);
+
 	// We cannot write directly to the data on all platforms,
 	// so we make a TArray of the right type, then assign it
 	check( GetStride() % sizeof(FVector4) == 0 );
@@ -110,7 +118,7 @@ void FStaticMeshInstanceBuffer::Init(UInstancedStaticMeshComponent* InComponent,
 		const float RandomInstanceID = RandomInstanceIDBase + RandomStream.GetFraction() * RandomInstanceIDRange;
 						
 		FillInstanceRenderData(Instance, InstanceRenderData, RandomInstanceID, Z, W);
-		}
+	}
 
 	// Hide any removed instances
 	if (NumRemoved)
@@ -179,11 +187,9 @@ void FStaticMeshInstanceBuffer::AllocateData()
 	CleanUp();
 
 	check(HasValidFeatureLevel());
-	static const auto CVar = IConsoleManager::Get().FindTConsoleVariableDataInt(TEXT("r.GenerateMeshDistanceFields"));
+	
 	const bool bInstanced = RHISupportsInstancing(GetFeatureLevelShaderPlatform(GetFeatureLevel()));
-	const bool bNeedsCPUAccess = !bInstanced 
-		// Distance field algorithms need access to instance data on the CPU
-		|| CVar->GetValueOnGameThread() != 0;
+	const bool bNeedsCPUAccess = !bInstanced;
 	InstanceData = new FStaticMeshInstanceData(bNeedsCPUAccess);
 	// Calculate the vertex stride.
 	Stride = InstanceData->GetStride();
