#ANIM: Fixed issue where animation change doesn't update in viewport.

https://jira.ol.epicgames.net/browse/UE-7157

diff --git a/Engine/Source/Runtime/Engine/Private/SkeletalMeshComponent.cpp b/Engine/Source/Runtime/Engine/Private/SkeletalMeshComponent.cpp
--- a/Engine/Source/Runtime/Engine/Private/SkeletalMeshComponent.cpp
+++ b/Engine/Source/Runtime/Engine/Private/SkeletalMeshComponent.cpp
@@ -399,6 +399,10 @@ void USkeletalMeshComponent::PostEditChangeProperty(FPropertyChangedEvent& Prope
 				UE_LOG(LogAnimation, Warning, TEXT("Invalid animation"));
 				AnimationData.AnimToPlay = NULL;
 			}
+			else
+			{
+				PlayAnimation(AnimationData.AnimToPlay, false);
+			}
 		}
 
 		if ( PropertyThatChanged->GetFName() == GET_MEMBER_NAME_CHECKED( FSingleAnimationPlayData, SavedPosition ))
