#ue4 - Added stat collision counters for PreFilter and PostFilter.

diff --git a/Engine/Source/Runtime/Engine/Private/Collision/PhysXCollision.cpp b/Engine/Source/Runtime/Engine/Private/Collision/PhysXCollision.cpp
--- a/Engine/Source/Runtime/Engine/Private/Collision/PhysXCollision.cpp
+++ b/Engine/Source/Runtime/Engine/Private/Collision/PhysXCollision.cpp
@@ -612,6 +612,8 @@ PxSceneQueryHitType::Enum FPxQueryFilterCallback::CalcQueryHitType(const PxFilte
 
 PxSceneQueryHitType::Enum FPxQueryFilterCallback::preFilter(const PxFilterData& filterData, const PxShape* shape, const PxRigidActor* actor, PxSceneQueryFlags& queryFlags)
 {
+	SCOPE_CYCLE_COUNTER(STAT_Collision_PreFilter);
+
 	// Check if the shape is the right complexity for the trace 
 	PxFilterData ShapeFilter = shape->getQueryFilterData();
 #define ENABLE_PREFILTER_LOGGING 0
@@ -679,6 +681,8 @@ PxSceneQueryHitType::Enum FPxQueryFilterCallback::preFilter(const PxFilterData&
 
 PxSceneQueryHitType::Enum FPxQueryFilterCallbackSweep::postFilter(const PxFilterData& filterData, const PxSceneQueryHit& hit)
 {
+	SCOPE_CYCLE_COUNTER(STAT_Collision_PostFilter);
+
 	PxSweepHit& SweepHit = (PxSweepHit&)hit;
 	const bool bIsOverlap = SweepHit.hadInitialOverlap();
 
diff --git a/Engine/Source/Runtime/Engine/Private/Collision/WorldCollision.cpp b/Engine/Source/Runtime/Engine/Private/Collision/WorldCollision.cpp
--- a/Engine/Source/Runtime/Engine/Private/Collision/WorldCollision.cpp
+++ b/Engine/Source/Runtime/Engine/Private/Collision/WorldCollision.cpp
@@ -37,6 +37,8 @@ DEFINE_STAT(STAT_Collision_GeomOverlapAny);
 DEFINE_STAT(STAT_Collision_GeomOverlapSingle);
 DEFINE_STAT(STAT_Collision_GeomOverlapMultiple);
 DEFINE_STAT(STAT_Collision_GeomComputePenetration);
+DEFINE_STAT(STAT_Collision_PreFilter);
+DEFINE_STAT(STAT_Collision_PostFilter);
 
 /** default collision response container - to be used without reconstructing every time**/
 FCollisionResponseContainer FCollisionResponseContainer::DefaultResponseContainer(ECR_Block);
diff --git a/Engine/Source/Runtime/Engine/Public/Collision.h b/Engine/Source/Runtime/Engine/Public/Collision.h
--- a/Engine/Source/Runtime/Engine/Public/Collision.h
+++ b/Engine/Source/Runtime/Engine/Public/Collision.h
@@ -19,6 +19,8 @@ DECLARE_CYCLE_STAT_EXTERN(TEXT("GeomOverlapAny"),STAT_Collision_GeomOverlapAny,S
 DECLARE_CYCLE_STAT_EXTERN(TEXT("GeomOverlapSingle"),STAT_Collision_GeomOverlapSingle,STATGROUP_Collision, );
 DECLARE_CYCLE_STAT_EXTERN(TEXT("GeomOverlapMultiple"),STAT_Collision_GeomOverlapMultiple,STATGROUP_Collision, );
 DECLARE_CYCLE_STAT_EXTERN(TEXT("GeomComputePenetration"), STAT_Collision_GeomComputePenetration, STATGROUP_Collision, );
+DECLARE_CYCLE_STAT_EXTERN(TEXT("PreFilter"), STAT_Collision_PreFilter, STATGROUP_Collision, );
+DECLARE_CYCLE_STAT_EXTERN(TEXT("PostFilter"), STAT_Collision_PostFilter, STATGROUP_Collision, );
 
 /** Enable collision analyzer support */
 #if (1 && !(UE_BUILD_SHIPPING || UE_BUILD_TEST) && WITH_EDITOR && WITH_UNREAL_DEVELOPER_TOOLS && WITH_PHYSX)
