KismetMathLibrary - Added the Vector2D Interpolation functions, Fixed Comments

From github PR#730 from user EverNewJoy. Thanks Rama!

diff --git a/Engine/Source/Runtime/Engine/Classes/Kismet/KismetMathLibrary.h b/Engine/Source/Runtime/Engine/Classes/Kismet/KismetMathLibrary.h
--- a/Engine/Source/Runtime/Engine/Classes/Kismet/KismetMathLibrary.h
+++ b/Engine/Source/Runtime/Engine/Classes/Kismet/KismetMathLibrary.h
@@ -1357,7 +1357,7 @@ class UKismetMathLibrary : public UBlueprintFunctionLibrary
 	static float FInterpTo(float Current, float Target, float DeltaTime, float InterpSpeed);
 
 	/**
-	 * Tries to reach Target based on distance from Current position, giving a nice smooth feeling when tracking a position.
+	 * Tries to reach Target at a constant rate.
 	 *
 	 * @param		Current			Actual position
 	 * @param		Target			Target position
@@ -1402,14 +1402,38 @@ class UKismetMathLibrary : public UBlueprintFunctionLibrary
 	 * @param		InterpSpeed		Interpolation speed
 	 * @return		New interpolated position
 	 */
+	UFUNCTION(BlueprintPure, Category="Math|Interpolation", meta=(Keywords="position"))
+	static FVector2D Vector2DInterpTo(FVector2D Current, FVector2D Target, float DeltaTime, float InterpSpeed);
+	
+	/**
+	 * Tries to reach Target at a constant rate.
+	 *
+	 * @param		Current			Actual position
+	 * @param		Target			Target position
+	 * @param		DeltaTime		Time since last tick
+	 * @param		InterpSpeed		Interpolation speed
+	 * @return		New interpolated position
+	 */
+	UFUNCTION(BlueprintPure, Category="Math|Interpolation", meta=(Keywords="position"))
+	static FVector2D Vector2DInterpTo_Constant(FVector2D Current, FVector2D Target, float DeltaTime, float InterpSpeed);
+	
+	/**
+	 * Tries to reach Target based on distance from Current position, giving a nice smooth feeling when tracking a position.
+	 *
+	 * @param		Current			Actual rotation
+	 * @param		Target			Target rotation
+	 * @param		DeltaTime		Time since last tick
+	 * @param		InterpSpeed		Interpolation speed
+	 * @return		New interpolated position
+	 */
 	UFUNCTION(BlueprintPure, Category="Math|Interpolation", meta=(Keywords="rotation rotate"))
 	static FRotator RInterpTo(FRotator Current, FRotator Target, float DeltaTime, float InterpSpeed);
 
 	/**
-	 * Tries to reach Target based on distance from Current position, giving a nice smooth feeling when tracking a position.
+	 * Tries to reach Target at a constant rate.
 	 *
-	 * @param		Current			Actual position
-	 * @param		Target			Target position
+	 * @param		Current			Actual rotation
+	 * @param		Target			Target rotation
 	 * @param		DeltaTime		Time since last tick
 	 * @param		InterpSpeed		Interpolation speed
 	 * @return		New interpolated position
diff --git a/Engine/Source/Runtime/Engine/Private/KismetMathLibrary.cpp b/Engine/Source/Runtime/Engine/Private/KismetMathLibrary.cpp
--- a/Engine/Source/Runtime/Engine/Private/KismetMathLibrary.cpp
+++ b/Engine/Source/Runtime/Engine/Private/KismetMathLibrary.cpp
@@ -773,6 +773,14 @@ FVector UKismetMathLibrary::VInterpTo_Constant(FVector Current, FVector Target,
 	return FMath::VInterpConstantTo(Current, Target, DeltaTime, InterpSpeed);
 }
 
+FVector2D UKismetMathLibrary::Vector2DInterpTo(FVector2D Current, FVector2D Target, float DeltaTime, float InterpSpeed)
+{
+	return FMath::Vector2DInterpTo( Current, Target, DeltaTime, InterpSpeed );
+}
+FVector2D UKismetMathLibrary::Vector2DInterpTo_Constant(FVector2D Current, FVector2D Target, float DeltaTime, float InterpSpeed)
+{
+	return FMath::Vector2DInterpConstantTo( Current, Target, DeltaTime, InterpSpeed );
+}
 
 FVector UKismetMathLibrary::RandomUnitVector()
 {
