Fix ecosystem ignoring cull distance

diff --git a/Engine/Source/Runtime/Engine/Classes/Engine/StaticMesh.h b/Engine/Source/Runtime/Engine/Classes/Engine/StaticMesh.h
--- a/Engine/Source/Runtime/Engine/Classes/Engine/StaticMesh.h
+++ b/Engine/Source/Runtime/Engine/Classes/Engine/StaticMesh.h
@@ -373,7 +373,7 @@ class UStaticMesh : public UObject, public IInterface_CollisionDataProvider, pub
 #if WITH_EDITORONLY_DATA
 	/** Default settings when using this mesh for instanced foliage */
 	UPROPERTY(EditAnywhere, AdvancedDisplay, Instanced, Category=StaticMesh)
-	UFoliageType_InstancedStaticMesh* FoliageDefaultSettings;
+	const UFoliageType_InstancedStaticMesh* FoliageDefaultSettings;
 
 	/** Importing data and options used for this mesh */
 	UPROPERTY(VisibleAnywhere, Instanced, Category=ImportSettings)
