Input: When customizing default console key setting based on keyboard language, add to the list instead of replacing, to avoid removing options for certain IME & physical keyboard setups.

diff --git a/Engine/Source/Runtime/Engine/Private/UserInterface/InputSettings.cpp b/Engine/Source/Runtime/Engine/Private/UserInterface/InputSettings.cpp
--- a/Engine/Source/Runtime/Engine/Private/UserInterface/InputSettings.cpp
+++ b/Engine/Source/Runtime/Engine/Private/UserInterface/InputSettings.cpp
@@ -42,7 +42,7 @@ void UInputSettings::PostInitProperties()
 	}
 
 #if PLATFORM_WINDOWS
-	// If the console key is set to the default we'll see about switching it to the keyboard default
+	// If the console key is set to the default we'll see about adding the keyboard default
 	// If they've mapped any additional keys, we'll just assume they've set it up in a way they desire
 	if (ConsoleKeys.Num() == 1 && ConsoleKeys[0] == EKeys::Tilde)
 	{
@@ -66,7 +66,7 @@ void UInputSettings::PostInitProperties()
 			break;
 		}
 
-		ConsoleKeys[0] = DefaultConsoleKey;
+		ConsoleKeys.Add(DefaultConsoleKey);
 	}
 #endif
 }
