UE-7230 - Move the functionality introduced in CL2356824 behind a bool. bDisabledLODsKeepEmitterAlive. Defaults to false.

diff --git a/Engine/Source/Runtime/Engine/Classes/Particles/ParticleEmitter.h b/Engine/Source/Runtime/Engine/Classes/Particles/ParticleEmitter.h
--- a/Engine/Source/Runtime/Engine/Classes/Particles/ParticleEmitter.h
+++ b/Engine/Source/Runtime/Engine/Classes/Particles/ParticleEmitter.h
@@ -169,6 +169,9 @@ class UParticleEmitter : public UObject
 	UPROPERTY()
 	uint32 bCookedOut:1;
 
+	/** When true, if the current LOD is disabled the emitter will be kept alive. Otherwise, the emitter will be considered complete if the current LOD is disabled. */
+	UPROPERTY(EditAnywhere, Category = Particle)
+	uint32 bDisabledLODsKeepEmitterAlive : 1;
 
 	// Begin UObject Interface
 #if WITH_EDITOR
diff --git a/Engine/Source/Runtime/Engine/Private/Particles/ParticleComponents.cpp b/Engine/Source/Runtime/Engine/Private/Particles/ParticleComponents.cpp
--- a/Engine/Source/Runtime/Engine/Private/Particles/ParticleComponents.cpp
+++ b/Engine/Source/Runtime/Engine/Private/Particles/ParticleComponents.cpp
@@ -955,8 +955,9 @@ bool UParticleLODLevel::IsModuleEditable(UParticleModule* InModule)
 	UParticleEmitter implementation.
 -----------------------------------------------------------------------------*/
 UParticleEmitter::UParticleEmitter(const FObjectInitializer& ObjectInitializer)
-	: Super(ObjectInitializer),
-	QualityLevelSpawnRateScale(1.0f)
+	: Super(ObjectInitializer)
+	, QualityLevelSpawnRateScale(1.0f)
+	, bDisabledLODsKeepEmitterAlive(false)
 {
 	// Structure to hold one-time initialization
 	struct FConstructorStatics
@@ -4899,9 +4900,11 @@ bool UParticleSystemComponent::HasCompleted()
 			}
 			else
 			{
-				//Don't assume emitter is complete if the current lod is disabled.
-				//Any emitters with all their lods disabled aren't created in the first place.
-				bHasCompleted = false;
+				UParticleEmitter* Em = CastChecked<UParticleEmitter>(Instance->CurrentLODLevel->GetOuter());
+				if (Em && Em->bDisabledLODsKeepEmitterAlive)
+				{
+					bHasCompleted = false;
+				}
 			}
 		}
 	}
