Added a sanity-check to avoid NavigationTestingActor crashing on empty navmeshes #UE4

diff --git a/Engine/Source/Runtime/Engine/Private/AI/Navigation/NavigationTestingActor.cpp b/Engine/Source/Runtime/Engine/Private/AI/Navigation/NavigationTestingActor.cpp
--- a/Engine/Source/Runtime/Engine/Private/AI/Navigation/NavigationTestingActor.cpp
+++ b/Engine/Source/Runtime/Engine/Private/AI/Navigation/NavigationTestingActor.cpp
@@ -362,7 +362,7 @@ void ANavigationTestingActor::SearchPathTo(ANavigationTestingActor* Goal)
 	if (bGatherDetailedInfo && !bUseHierarchicalPathfinding)
 	{
 		ARecastNavMesh* RecastNavMesh = Cast<ARecastNavMesh>(MyNavData);
-		if (RecastNavMesh)
+		if (RecastNavMesh && RecastNavMesh->HasValidNavmesh())
 		{
 			PathfindingSteps = RecastNavMesh->DebugPathfinding(Query, DebugSteps);
 		}
