Fixed Navmesh generation not working properly with World Composition #UE4

- addresses UE-5116

diff --git a/Engine/Source/Runtime/Engine/Private/AI/Navigation/NavigationData.cpp b/Engine/Source/Runtime/Engine/Private/AI/Navigation/NavigationData.cpp
--- a/Engine/Source/Runtime/Engine/Private/AI/Navigation/NavigationData.cpp
+++ b/Engine/Source/Runtime/Engine/Private/AI/Navigation/NavigationData.cpp
@@ -153,10 +153,17 @@ void ANavigationData::PostInitializeComponents()
 {
 	Super::PostInitializeComponents();
 	
-	if (GetWorld() == NULL || GetWorld()->GetNavigationSystem() == NULL)
+	UWorld* WorldOuter = GetWorld();
+
+	if (WorldOuter == NULL || WorldOuter->GetNavigationSystem() == NULL)
 	{
 		CleanUpAndMarkPendingKill();
-		return;
+	}
+	else
+	{
+		// note: this is not a final fix for world composition's issues with navmesh generation
+		// but it's good for now, and navmesh creation is going to get a face-lift soon anyway
+		bWantsUpdate |= WorldOuter->GetWorldSettings()->bEnableWorldComposition;
 	}
 }
 
diff --git a/Engine/Source/Runtime/Engine/Private/Actor.cpp b/Engine/Source/Runtime/Engine/Private/Actor.cpp
--- a/Engine/Source/Runtime/Engine/Private/Actor.cpp
+++ b/Engine/Source/Runtime/Engine/Private/Actor.cpp
@@ -524,6 +524,9 @@ void AActor::ApplyWorldOffset(const FVector& InOffset, bool bWorldShift)
 	if (RootComponent != NULL && RootComponent->AttachParent == NULL)
 	{
 		RootComponent->ApplyWorldOffset(InOffset, bWorldShift);
+
+		UNavigationSystem::UpdateNavOctreeBounds(this);
+		UNavigationSystem::UpdateNavOctreeAll(this);
 	}
 }
 
