Fixed engine crashing when user modifies navmesh generation properties via Projects Settings #UE4

- reported as UE-7232

diff --git a/Engine/Source/Runtime/Engine/Private/AI/Navigation/RecastNavMesh.cpp b/Engine/Source/Runtime/Engine/Private/AI/Navigation/RecastNavMesh.cpp
--- a/Engine/Source/Runtime/Engine/Private/AI/Navigation/RecastNavMesh.cpp
+++ b/Engine/Source/Runtime/Engine/Private/AI/Navigation/RecastNavMesh.cpp
@@ -1761,9 +1761,17 @@ void ARecastNavMesh::PostEditChangeProperty(FPropertyChangedEvent& PropertyChang
 				{
 					TileSizeUU = MAX_uint16 * CellSize;
 				}
+				// also it can't be 0, and if it's 1 then we should make sure tile size is equal to cell size
+				else if (DimensionVX <= 1)
+				{
+					TileSizeUU = CellSize;
+				}
 			}
 
-			RebuildAll();
+			if (HasAnyFlags(RF_ClassDefaultObject) == false)
+			{
+				RebuildAll();
+			}
 		}
 		else if (FObjectEditorUtils::GetCategoryFName(PropertyChangedEvent.Property) == NAME_Display)
 		{
