Fixed UNavigationSystem::GetRandomPointInRadius always returning true when a NavigationData instance have been found regardless of the actual result of the query #UE4

- github PR#735 by dzeligman

diff --git a/Engine/Source/Runtime/Engine/Private/AI/Navigation/NavigationSystem.cpp b/Engine/Source/Runtime/Engine/Private/AI/Navigation/NavigationSystem.cpp
--- a/Engine/Source/Runtime/Engine/Private/AI/Navigation/NavigationSystem.cpp
+++ b/Engine/Source/Runtime/Engine/Private/AI/Navigation/NavigationSystem.cpp
@@ -915,18 +915,12 @@ bool UNavigationSystem::GetRandomPointInRadius(const FVector& Origin, float Radi
 {
 	SCOPE_CYCLE_COUNTER(STAT_Navigation_QueriesTimeSync);
 
-	if (NavData == NULL)
+	if (NavData == nullptr)
 	{
 		NavData = MainNavData;
 	}
 
-	if (NavData != NULL)
-	{
-		NavData->GetRandomPointInRadius(Origin, Radius, ResultLocation, QueryFilter);
-		return true;
-	}
-
-	return false;
+	return NavData != nullptr && NavData->GetRandomPointInRadius(Origin, Radius, ResultLocation, QueryFilter);
 }
 
 ENavigationQueryResult::Type UNavigationSystem::GetPathCost(const FVector& PathStart, const FVector& PathEnd, float& OutPathCost, const ANavigationData* NavData, TSharedPtr<const FNavigationQueryFilter> QueryFilter) const
