Fix for quad tree not doing a final AABB check on the elements it returns.

diff --git a/Engine/Source/Runtime/Engine/Public/GenericQuadTree.h b/Engine/Source/Runtime/Engine/Public/GenericQuadTree.h
--- a/Engine/Source/Runtime/Engine/Public/GenericQuadTree.h
+++ b/Engine/Source/Runtime/Engine/Public/GenericQuadTree.h
@@ -196,7 +196,10 @@ void TQuadTree<ElementType, NodeCapacity>::GetElements(const FBox2D& Box, TArray
 		ElementsOut.Reserve(ElementsOut.Num() + Leaf->Nodes.Num());
 		for (const FNode& Node : Leaf->Nodes)
 		{
-			ElementsOut.Add(Node.Element);
+			if (Box.Intersect(Node.Box))
+			{
+				ElementsOut.Add(Node.Element);
+			}
 		};
 	}
 }
@@ -211,7 +214,10 @@ void TQuadTree<ElementType, NodeCapacity>::GetElementsUnique(const FBox2D& Box,
 	{
 		for (const FNode& Node : Leaf->Nodes)
 		{
-			Unique.Add(Node.Element);
+			if (Box.Intersect(Node.Box))
+			{
+				Unique.Add(Node.Element);
+			}
 		};
 	}
 
