UE4 - improve memory accounting of instanced static meshes

diff --git a/Engine/Source/Runtime/Engine/Classes/Components/HierarchicalInstancedStaticMeshComponent.h b/Engine/Source/Runtime/Engine/Classes/Components/HierarchicalInstancedStaticMeshComponent.h
--- a/Engine/Source/Runtime/Engine/Classes/Components/HierarchicalInstancedStaticMeshComponent.h
+++ b/Engine/Source/Runtime/Engine/Classes/Components/HierarchicalInstancedStaticMeshComponent.h
@@ -80,7 +80,11 @@ class ENGINE_API UHierarchicalInstancedStaticMeshComponent : public UInstancedSt
 	void ApplyBuildTreeAsync(ENamedThreads::Type CurrentThread, const FGraphEventRef& MyCompletionGraphEvent, TSharedRef<FClusterBuilder, ESPMode::ThreadSafe> Builder, double StartTime);
 
 public:
+
+	//Begin UObject Interface
 	virtual void Serialize(FArchive& Ar) override;
+	virtual SIZE_T GetResourceSize(EResourceSizeMode::Type Mode) override;
+
 	virtual FPrimitiveSceneProxy* CreateSceneProxy() override;
 	virtual FBoxSphereBounds CalcBounds(const FTransform& BoundTransform) const override;
 
diff --git a/Engine/Source/Runtime/Engine/Private/HierarchicalInstancedStaticMesh.cpp b/Engine/Source/Runtime/Engine/Private/HierarchicalInstancedStaticMesh.cpp
--- a/Engine/Source/Runtime/Engine/Private/HierarchicalInstancedStaticMesh.cpp
+++ b/Engine/Source/Runtime/Engine/Private/HierarchicalInstancedStaticMesh.cpp
@@ -1343,6 +1343,15 @@ void UHierarchicalInstancedStaticMeshComponent::Serialize(FArchive& Ar)
 	ClusterTree.BulkSerialize(Ar);
 }
 
+SIZE_T UHierarchicalInstancedStaticMeshComponent::GetResourceSize( EResourceSizeMode::Type Mode )
+{
+	SIZE_T ResSize = Super::GetResourceSize(Mode);
+
+	ResSize += SortedInstances.GetAllocatedSize();
+
+	return ResSize;
+}
+
 bool UHierarchicalInstancedStaticMeshComponent::RemoveInstance(int32 InstanceIndex)
 {
 	if (!PerInstanceSMData.IsValidIndex(InstanceIndex))
diff --git a/Engine/Source/Runtime/Engine/Private/InstancedStaticMesh.cpp b/Engine/Source/Runtime/Engine/Private/InstancedStaticMesh.cpp
--- a/Engine/Source/Runtime/Engine/Private/InstancedStaticMesh.cpp
+++ b/Engine/Source/Runtime/Engine/Private/InstancedStaticMesh.cpp
@@ -1406,8 +1406,13 @@ void UInstancedStaticMeshComponent::GetNavigationPerInstanceTransforms(const FBo
 
 SIZE_T UInstancedStaticMeshComponent::GetResourceSize( EResourceSizeMode::Type Mode )
 {
-	SIZE_T ResSize = 0;
+	SIZE_T ResSize = Super::GetResourceSize(Mode);
 
+	// proxy stuff
+	ResSize += FStaticMeshInstanceData::GetResourceSize(PerInstanceSMData.Num());
+
+	// component stuff
+	ResSize += InstanceBodies.GetAllocatedSize();
 	for (int32 i=0; i < InstanceBodies.Num(); ++i)
 	{
 		if (InstanceBodies[i] != NULL && InstanceBodies[i]->IsValidBodyInstance())
@@ -1415,6 +1420,16 @@ SIZE_T UInstancedStaticMeshComponent::GetResourceSize( EResourceSizeMode::Type M
 			ResSize += InstanceBodies[i]->GetBodyInstanceResourceSize(Mode);
 		}
 	}
+	ResSize += InstanceReorderTable.GetAllocatedSize();
+	ResSize += RemovedInstances.GetAllocatedSize();
+
+#if WITH_EDITOR
+	ResSize += SelectedInstances.GetAllocatedSize();
+#endif
+
+#if WITH_PHYSX
+	ResSize += Aggregates.GetAllocatedSize();
+#endif
 
 	return ResSize;
 }
diff --git a/Engine/Source/Runtime/Engine/Private/InstancedStaticMesh.h b/Engine/Source/Runtime/Engine/Private/InstancedStaticMesh.h
--- a/Engine/Source/Runtime/Engine/Private/InstancedStaticMesh.h
+++ b/Engine/Source/Runtime/Engine/Private/InstancedStaticMesh.h
@@ -76,6 +76,10 @@ class FStaticMeshInstanceData :
 {
 public:
 
+	enum 
+	{
+		VectorsPerInstance = 7
+	};
 	typedef TResourceArray<FVector4,VERTEXBUFFER_ALIGNMENT> ArrayType;
 
 	/**
@@ -87,6 +91,16 @@ public:
 	{
 	}
 
+	static uint32 StaticGetStride()
+	{
+		return sizeof(FVector4) * VectorsPerInstance;
+	}
+
+	static SIZE_T GetResourceSize(uint32 NumInstances)
+	{
+		return SIZE_T(NumInstances) * SIZE_T(StaticGetStride());
+	}
+
 	/**
 	 * Resizes the vertex data buffer, discarding any data which no longer fits.
 	 * @param NumVertices - The number of vertices to allocate the buffer for.
@@ -98,8 +112,7 @@ public:
 
 	virtual uint32 GetStride() const
 	{
-		const uint32 VectorsPerInstance = 7;
-		return sizeof(FVector4) * VectorsPerInstance;
+		return StaticGetStride();
 	}
 	virtual uint8* GetDataPointer()
 	{
