Deprecated several non-pure functions in Actor and replaced them with renamed pure versions (GetDistanceTo, GetDotProductTo, et al). (merge of CL 2401444 from branch)

diff --git a/Engine/Source/Runtime/Engine/Classes/GameFramework/Actor.h b/Engine/Source/Runtime/Engine/Classes/GameFramework/Actor.h
--- a/Engine/Source/Runtime/Engine/Classes/GameFramework/Actor.h
+++ b/Engine/Source/Runtime/Engine/Classes/GameFramework/Actor.h
@@ -687,25 +687,37 @@ public:
 	UFUNCTION(BlueprintCallable, Category="Utilities|Orientation")
 	FVector GetActorScale3D() const;
 
-	/** Returns the distance from this Actor to OtherActor. */
-	UFUNCTION(BlueprintCallable, Category="Utilities|Transformation")
+	/** Deprecated in favor of const/pure versions below. */
+	UFUNCTION(BlueprintCallable, meta=(DeprecatedFunction, DeprecationMessage = "Use the pure node GetDistanceToActor instead"), Category = "Utilities|Transformation")
 	float GetDistanceTo(AActor* OtherActor);
+	UFUNCTION(BlueprintCallable, meta = (DeprecatedFunction, DeprecationMessage = "Use the pure node GetHorizontalDistanceToActor instead"), Category = "Utilities|Transformation")
+	float GetHorizontalDistanceTo(AActor* OtherActor);
+	UFUNCTION(BlueprintCallable, meta = (DeprecatedFunction, DeprecationMessage = "Use the pure node GetVerticalDistanceToActor instead"), Category = "Utilities|Transformation")
+	float GetVerticalDistanceTo(AActor* OtherActor);
+	UFUNCTION(BlueprintCallable, meta = (DeprecatedFunction, DeprecationMessage = "Use the pure node GetDistanceToActor instead"), Category = "Utilities|Transformation")
+	float GetDotProductTo(AActor* OtherActor);
+	UFUNCTION(BlueprintCallable, meta = (DeprecatedFunction, DeprecationMessage = "Use the pure node GetDistanceToActor instead"), Category = "Utilities|Transformation")
+	float GetHorizontalDotProductTo(AActor* OtherActor);
+
+	/** Returns the distance from this Actor to OtherActor. */
+	UFUNCTION(BlueprintCallable, Category = "Utilities|Transformation")
+	float GetDistanceToActor(AActor* OtherActor) const;
 
 	/** Returns the distance from this Actor to OtherActor, ignoring Z. */
-	UFUNCTION(BlueprintCallable, Category="Utilities|Transformation")
-	float GetHorizontalDistanceTo(AActor* OtherActor);
+	UFUNCTION(BlueprintCallable, Category = "Utilities|Transformation")
+	float GetHorizontalDistanceToActor(AActor* OtherActor) const;
 
 	/** Returns the distance from this Actor to OtherActor, ignoring XY. */
-	UFUNCTION(BlueprintCallable, Category="Utilities|Transformation")
-	float GetVerticalDistanceTo(AActor* OtherActor);
+	UFUNCTION(BlueprintCallable, Category = "Utilities|Transformation")
+	float GetVerticalDistanceToActor(AActor* OtherActor) const;
 
 	/** Returns the dot product from this Actor to OtherActor. Returns -2.0 on failure. Returns 0.0 for coincidental actors. */
 	UFUNCTION(BlueprintCallable, Category = "Utilities|Transformation")
-	float GetDotProductTo(AActor* OtherActor);
+	float GetDotProductToActor(AActor* OtherActor) const;
 
 	/** Returns the dot product from this Actor to OtherActor, ignoring Z. Returns -2.0 on failure. Returns 0.0 for coincidental actors. */
 	UFUNCTION(BlueprintCallable, Category = "Utilities|Transformation")
-	float GetHorizontalDotProductTo(AActor* OtherActor);
+	float GetHorizontalDotProductToActor(AActor* OtherActor) const;
 
 
 	/**
@@ -2335,4 +2347,12 @@ FORCEINLINE FVector AActor::GetSimpleCollisionCylinderExtent() const
 	bool SetActorLocationAndRotation(FVector NewLocation, FRotator NewRotation, bool bSweep=false, FHitResult* OutSweepHitResult=nullptr) { return Super::SetActorLocationAndRotation(NewLocation, NewRotation, bSweep, OutSweepHitResult); } \
 	virtual bool TeleportTo( const FVector& DestLocation, const FRotator& DestRotation, bool bIsATest, bool bNoCheck ) override { return Super::TeleportTo(DestLocation, DestRotation, bIsATest, bNoCheck); } \
 	virtual FVector GetVelocity() const override { return Super::GetVelocity(); } \
-	float GetDistanceTo(AActor* OtherActor) { return Super::GetDistanceTo(OtherActor); }
+	float GetHorizontalDistanceToActor(AActor* OtherActor)  const { return Super::GetHorizontalDistanceToActor(OtherActor); } \
+	float GetVerticalDistanceToActor(AActor* OtherActor)  const { return Super::GetVerticalDistanceToActor(OtherActor); } \
+	float GetDotProductToActor(AActor* OtherActor) const { return Super::GetDotProductToActor(OtherActor); } \
+	float GetHorizontalDotProductToActor(AActor* OtherActor) const { return Super::GetHorizontalDotProductToActor(OtherActor); } \
+	float GetDistanceToActor(AActor* OtherActor) const { return Super::GetDistanceToActor(OtherActor); }
+
+
+
+
diff --git a/Engine/Source/Runtime/Engine/Private/Actor.cpp b/Engine/Source/Runtime/Engine/Private/Actor.cpp
--- a/Engine/Source/Runtime/Engine/Private/Actor.cpp
+++ b/Engine/Source/Runtime/Engine/Private/Actor.cpp
@@ -3478,21 +3478,25 @@ UMaterialInstanceDynamic* AActor::MakeMIDForMaterial(class UMaterialInterface* P
 	return NULL;
 }
 
+// deprecated
 float AActor::GetDistanceTo(AActor* OtherActor)
 {
 	return OtherActor ? (GetActorLocation() - OtherActor->GetActorLocation()).Size() : 0.f;
 }
 
+// deprecated
 float AActor::GetHorizontalDistanceTo(AActor* OtherActor)
 {
 	return OtherActor ? (GetActorLocation() - OtherActor->GetActorLocation()).Size2D() : 0.f;
 }
 
+// deprecated
 float AActor::GetVerticalDistanceTo(AActor* OtherActor)
 {
 	return OtherActor ? FMath::Abs((GetActorLocation().Z - OtherActor->GetActorLocation().Z)) : 0.f;
 }
 
+// deprecated
 float AActor::GetDotProductTo(AActor* OtherActor)
 {
 	if (OtherActor)
@@ -3505,6 +3509,7 @@ float AActor::GetDotProductTo(AActor* OtherActor)
 	return -2.0;
 }
 
+// deprecated
 float AActor::GetHorizontalDotProductTo(AActor* OtherActor)
 {
 	if (OtherActor)
@@ -3517,6 +3522,47 @@ float AActor::GetHorizontalDotProductTo(AActor* OtherActor)
 	return -2.0;
 }
 
+
+
+float AActor::GetDistanceToActor(AActor* OtherActor) const
+{
+	return OtherActor ? (GetActorLocation() - OtherActor->GetActorLocation()).Size() : 0.f;
+}
+
+float AActor::GetHorizontalDistanceToActor(AActor* OtherActor) const
+{
+	return OtherActor ? (GetActorLocation() - OtherActor->GetActorLocation()).Size2D() : 0.f;
+}
+
+float AActor::GetVerticalDistanceToActor(AActor* OtherActor) const
+{
+	return OtherActor ? FMath::Abs((GetActorLocation().Z - OtherActor->GetActorLocation().Z)) : 0.f;
+}
+
+float AActor::GetDotProductToActor(AActor* OtherActor) const
+{
+	if (OtherActor)
+	{
+		FVector Dir = GetActorForwardVector();
+		FVector Offset = OtherActor->GetActorLocation() - GetActorLocation();
+		Offset = Offset.GetSafeNormal();
+		return FVector::DotProduct(Dir, Offset);
+	}
+	return -2.0;
+}
+
+float AActor::GetHorizontalDotProductToActor(AActor* OtherActor) const
+{
+	if (OtherActor)
+	{
+		FVector Dir = GetActorForwardVector();
+		FVector Offset = OtherActor->GetActorLocation() - GetActorLocation();
+		Offset = Offset.GetSafeNormal2D();
+		return FVector::DotProduct(Dir, Offset);
+	}
+	return -2.0;
+}
+
 #if WITH_EDITOR
 const int32 AActor::GetNumUncachedStaticLightingInteractions() const
 {
