Fix crash when using parallel eval on a native anim instance

diff --git a/Engine/Source/Runtime/Engine/Classes/Animation/AnimInstance.h b/Engine/Source/Runtime/Engine/Classes/Animation/AnimInstance.h
--- a/Engine/Source/Runtime/Engine/Classes/Animation/AnimInstance.h
+++ b/Engine/Source/Runtime/Engine/Classes/Animation/AnimInstance.h
@@ -549,6 +549,7 @@ public:
 
 	//@TODO: Better comments
 	virtual void EvaluateAnimation(struct FPoseContext& Output);
+	virtual void PostAnimEvaluation() {}
 
 	void InitializeAnimation();
 	void UpdateAnimation(float DeltaSeconds);
diff --git a/Engine/Source/Runtime/Engine/Classes/Animation/AnimSingleNodeInstance.h b/Engine/Source/Runtime/Engine/Classes/Animation/AnimSingleNodeInstance.h
--- a/Engine/Source/Runtime/Engine/Classes/Animation/AnimSingleNodeInstance.h
+++ b/Engine/Source/Runtime/Engine/Classes/Animation/AnimSingleNodeInstance.h
@@ -59,6 +59,8 @@ class ENGINE_API UAnimSingleNodeInstance : public UAnimInstance
 	virtual void NativeInitializeAnimation() override;
 	virtual void NativeUpdateAnimation(float DeltaTimeX) override;
 	virtual bool NativeEvaluateAnimation(FPoseContext& Output) override;
+	virtual void PostAnimEvaluation() override;
+
 protected:
 	virtual void Montage_Advance(float DeltaTime) override;
 	void InternalBlendSpaceEvaluatePose(class UBlendSpaceBase* BlendSpace, TArray<FBlendSampleData>& BlendSampleDataCache, struct FA2Pose& Pose);
diff --git a/Engine/Source/Runtime/Engine/Private/Animation/AnimSingleNodeInstance.cpp b/Engine/Source/Runtime/Engine/Private/Animation/AnimSingleNodeInstance.cpp
--- a/Engine/Source/Runtime/Engine/Private/Animation/AnimSingleNodeInstance.cpp
+++ b/Engine/Source/Runtime/Engine/Private/Animation/AnimSingleNodeInstance.cpp
@@ -414,11 +414,15 @@ bool UAnimSingleNodeInstance::NativeEvaluateAnimation(FPoseContext& Output)
 		VertexAnims.Add(FActiveVertexAnim(CurrentVertexAnim, 1.f, CurrentTime));
 	}
 
-	PostEvaluateAnimEvent.ExecuteIfBound();
-
 	return true;
 }
 
+void UAnimSingleNodeInstance::PostAnimEvaluation()
+{
+	PostEvaluateAnimEvent.ExecuteIfBound();
+
+	Super::PostAnimEvaluation();
+}
 
 void UAnimSingleNodeInstance::Montage_Advance(float DeltaTime)
 {
diff --git a/Engine/Source/Runtime/Engine/Private/SkeletalMeshComponent.cpp b/Engine/Source/Runtime/Engine/Private/SkeletalMeshComponent.cpp
--- a/Engine/Source/Runtime/Engine/Private/SkeletalMeshComponent.cpp
+++ b/Engine/Source/Runtime/Engine/Private/SkeletalMeshComponent.cpp
@@ -1008,6 +1008,11 @@ void USkeletalMeshComponent::RefreshBoneTransforms(FActorComponentTickFunction*
 
 void USkeletalMeshComponent::PostAnimEvaluation(FAnimationEvaluationContext& EvaluationContext)
 {
+	if (AnimScriptInstance)
+	{
+		AnimScriptInstance->PostAnimEvaluation();
+	}
+
 	AnimEvaluationContext.Clear();
 
 	SCOPE_CYCLE_COUNTER(STAT_PostAnimEvaluation);
