Updated deprecation of AActor::GetDistanceTo, et al to warn about C++ usage as well. (merge CL 2401998 from branch)

diff --git a/Engine/Source/Runtime/Engine/Classes/GameFramework/Actor.h b/Engine/Source/Runtime/Engine/Classes/GameFramework/Actor.h
--- a/Engine/Source/Runtime/Engine/Classes/GameFramework/Actor.h
+++ b/Engine/Source/Runtime/Engine/Classes/GameFramework/Actor.h
@@ -688,15 +688,24 @@ public:
 	FVector GetActorScale3D() const;
 
 	/** Deprecated in favor of const/pure versions below. */
-	UFUNCTION(BlueprintCallable, meta=(DeprecatedFunction, DeprecationMessage = "Use the pure node GetDistanceToActor instead"), Category = "Utilities|Transformation")
+	DEPRECATED(4.8, "AActor::GetDistanceTo will be removed. Use AActor::GetDistanceToActor instead.")
+	UFUNCTION(BlueprintCallable, meta = (DeprecatedFunction, DeprecationMessage = "Use the pure node GetDistanceToActor instead"), Category = "Utilities|Transformation")
 	float GetDistanceTo(AActor* OtherActor);
+	
+	DEPRECATED(4.8, "AActor::GetHorizontalDistanceTo will be removed. Use AActor::GetHorizontalDistanceToActor instead.")
 	UFUNCTION(BlueprintCallable, meta = (DeprecatedFunction, DeprecationMessage = "Use the pure node GetHorizontalDistanceToActor instead"), Category = "Utilities|Transformation")
 	float GetHorizontalDistanceTo(AActor* OtherActor);
+	
+	DEPRECATED(4.8, "AActor::GetVerticalDistanceTo will be removed. Use AActor::GetVerticalDistanceToActor instead.")
 	UFUNCTION(BlueprintCallable, meta = (DeprecatedFunction, DeprecationMessage = "Use the pure node GetVerticalDistanceToActor instead"), Category = "Utilities|Transformation")
 	float GetVerticalDistanceTo(AActor* OtherActor);
-	UFUNCTION(BlueprintCallable, meta = (DeprecatedFunction, DeprecationMessage = "Use the pure node GetDistanceToActor instead"), Category = "Utilities|Transformation")
+	
+	DEPRECATED(4.8, "AActor::GetDotProductTo will be removed. Use AActor::GetDotProductToActor instead.")
+	UFUNCTION(BlueprintCallable, meta = (DeprecatedFunction, DeprecationMessage = "Use the pure node GetDotProductToActor instead"), Category = "Utilities|Transformation")
 	float GetDotProductTo(AActor* OtherActor);
-	UFUNCTION(BlueprintCallable, meta = (DeprecatedFunction, DeprecationMessage = "Use the pure node GetDistanceToActor instead"), Category = "Utilities|Transformation")
+	
+	DEPRECATED(4.8, "AActor::GetHorizontalDotProductTo will be removed. Use AActor::GetHorizontalDotProductToActor instead.")
+	UFUNCTION(BlueprintCallable, meta = (DeprecatedFunction, DeprecationMessage = "Use the pure node GetHorizontalDotProductToActor instead"), Category = "Utilities|Transformation")
 	float GetHorizontalDotProductTo(AActor* OtherActor);
 
 	/** Returns the distance from this Actor to OtherActor. */
