UE4 - Fix for race condition UE-7358

diff --git a/Engine/Source/Runtime/Engine/Private/SceneManagement.cpp b/Engine/Source/Runtime/Engine/Private/SceneManagement.cpp
--- a/Engine/Source/Runtime/Engine/Private/SceneManagement.cpp
+++ b/Engine/Source/Runtime/Engine/Private/SceneManagement.cpp
@@ -48,6 +48,13 @@ void FTemporalLODState::UpdateTemporalLODTransition(const FViewInfo& View, float
 
 
 
+FSimpleElementCollector::FSimpleElementCollector() :
+	FPrimitiveDrawInterface(nullptr)
+{
+	static auto* MobileHDRCvar = IConsoleManager::Get().FindTConsoleVariableDataInt(TEXT("r.MobileHDR"));
+	bIsMobileHDR = (MobileHDRCvar->GetValueOnAnyThread() == 1);
+}
+
 FSimpleElementCollector::~FSimpleElementCollector()
 {
 	// Cleanup the dynamic resources.
@@ -146,9 +153,6 @@ void FSimpleElementCollector::RegisterDynamicResource(FDynamicPrimitiveResource*
 
 void FSimpleElementCollector::DrawBatchedElements(FRHICommandList& RHICmdList, const FSceneView& View, FTexture2DRHIRef DepthTexture, EBlendModeFilter::Type Filter) const
 {
-	static auto* MobileHDRCvar = IConsoleManager::Get().FindTConsoleVariableDataInt(TEXT("r.MobileHDR"));
-	const bool bIsMobileHDR = MobileHDRCvar->GetValueOnRenderThread() == 1;
-
 	// Mobile HDR does not execute post process, so does not need to render flipped
 	const bool bNeedToSwitchVerticalAxis = RHINeedsToSwitchVerticalAxis(View.GetShaderPlatform()) && !bIsMobileHDR;
 
diff --git a/Engine/Source/Runtime/Engine/Public/SceneManagement.h b/Engine/Source/Runtime/Engine/Public/SceneManagement.h
--- a/Engine/Source/Runtime/Engine/Public/SceneManagement.h
+++ b/Engine/Source/Runtime/Engine/Public/SceneManagement.h
@@ -1143,9 +1143,7 @@ class ENGINE_API FSimpleElementCollector : public FPrimitiveDrawInterface
 {
 public:
 
-	FSimpleElementCollector() : FPrimitiveDrawInterface(NULL)
-	{}
-
+	FSimpleElementCollector();
 	~FSimpleElementCollector();
 
 	virtual void SetHitProxy(HHitProxy* HitProxy);
@@ -1221,6 +1219,8 @@ private:
 
 	FHitProxyId HitProxyId;
 
+	bool bIsMobileHDR;
+
 	/** The dynamic resources which have been registered with this drawer. */
 	TArray<FDynamicPrimitiveResource*,SceneRenderingAllocator> DynamicResources;
 
