UE4 - Fix cvar reading

diff --git a/Engine/Source/Runtime/Engine/Private/SceneManagement.cpp b/Engine/Source/Runtime/Engine/Private/SceneManagement.cpp
--- a/Engine/Source/Runtime/Engine/Private/SceneManagement.cpp
+++ b/Engine/Source/Runtime/Engine/Private/SceneManagement.cpp
@@ -51,8 +51,9 @@ void FTemporalLODState::UpdateTemporalLODTransition(const FViewInfo& View, float
 FSimpleElementCollector::FSimpleElementCollector() :
 	FPrimitiveDrawInterface(nullptr)
 {
+	check(IsInRenderingThread());
 	static auto* MobileHDRCvar = IConsoleManager::Get().FindTConsoleVariableDataInt(TEXT("r.MobileHDR"));
-	bIsMobileHDR = (MobileHDRCvar->GetValueOnAnyThread() == 1);
+	bIsMobileHDR = (MobileHDRCvar->GetValueOnRenderThread() == 1);
 }
 
 FSimpleElementCollector::~FSimpleElementCollector()
