Changed NavMovementComponent::UpdateNavAgent functions to take references rather than pointers, for API readablity's sake #UE4

diff --git a/Engine/Source/Runtime/Engine/Classes/GameFramework/NavMovementComponent.h b/Engine/Source/Runtime/Engine/Classes/GameFramework/NavMovementComponent.h
--- a/Engine/Source/Runtime/Engine/Classes/GameFramework/NavMovementComponent.h
+++ b/Engine/Source/Runtime/Engine/Classes/GameFramework/NavMovementComponent.h
@@ -13,6 +13,9 @@
 #include "Components/PrimitiveComponent.h"
 #include "NavMovementComponent.generated.h"
 
+class AActor;
+class UCapsuleComponent;
+
 /**
  * NavMovementComponent defines base functionality for MovementComponents that move any 'agent' that may be involved in AI pathfinding.
  */
@@ -56,8 +59,8 @@ public:
 	void SetUpdateNavAgentWithOwnersCollisions(bool bUpdateWithOwner);
 	FORCEINLINE bool ShouldUpdateNavAgentWithOwnersCollision() const { return bUpdateNavAgentWithOwnersCollision != 0; }
 	
-	void UpdateNavAgent(class AActor* Owner);
-	void UpdateNavAgent(class UCapsuleComponent* CapsuleComponent);
+	void UpdateNavAgent(const AActor& Owner);
+	void UpdateNavAgent(const UCapsuleComponent& CapsuleComponent);
 
 	/** @returns location of controlled actor - meaning center of collision bounding box */
 	FORCEINLINE FVector GetActorLocation() const { return UpdatedComponent ? UpdatedComponent->GetComponentLocation() : FVector(FLT_MAX); }
diff --git a/Engine/Source/Runtime/Engine/Private/Character.cpp b/Engine/Source/Runtime/Engine/Private/Character.cpp
--- a/Engine/Source/Runtime/Engine/Private/Character.cpp
+++ b/Engine/Source/Runtime/Engine/Private/Character.cpp
@@ -120,7 +120,7 @@ void ACharacter::PostInitializeComponents()
 
 		if (CharacterMovement && CapsuleComponent)
 		{
-			CharacterMovement->UpdateNavAgent(CapsuleComponent);
+			CharacterMovement->UpdateNavAgent(*CapsuleComponent);
 		}
 
 		if (Controller == NULL && GetNetMode() != NM_Client)
diff --git a/Engine/Source/Runtime/Engine/Private/NavMovementComponent.cpp b/Engine/Source/Runtime/Engine/Private/NavMovementComponent.cpp
--- a/Engine/Source/Runtime/Engine/Private/NavMovementComponent.cpp
+++ b/Engine/Source/Runtime/Engine/Private/NavMovementComponent.cpp
@@ -49,30 +49,30 @@ void UNavMovementComponent::StopActiveMovement()
 	}
 }
 
-void UNavMovementComponent::UpdateNavAgent(AActor* Owner)
+void UNavMovementComponent::UpdateNavAgent(const AActor& Owner)
 {
-	ensure(Owner == NULL || Owner == GetOwner());
-	if (Owner == NULL || ShouldUpdateNavAgentWithOwnersCollision() == false)
+	ensure(&Owner == GetOwner());
+	if (ShouldUpdateNavAgentWithOwnersCollision() == false)
 	{
 		return;
 	}
 
 	// Can't call GetSimpleCollisionCylinder(), because no components will be registered.
 	float BoundRadius, BoundHalfHeight;	
-	Owner->GetSimpleCollisionCylinder(BoundRadius, BoundHalfHeight);
+	Owner.GetSimpleCollisionCylinder(BoundRadius, BoundHalfHeight);
 	NavAgentProps.AgentRadius = BoundRadius;
 	NavAgentProps.AgentHeight = BoundHalfHeight * 2.f;
 }
 
-void UNavMovementComponent::UpdateNavAgent(class UCapsuleComponent* CapsuleComponent)
+void UNavMovementComponent::UpdateNavAgent(const UCapsuleComponent& CapsuleComponent)
 {
-	if (CapsuleComponent == NULL || ShouldUpdateNavAgentWithOwnersCollision() == false)
+	if (ShouldUpdateNavAgentWithOwnersCollision() == false)
 	{
 		return;
 	}
 
-	NavAgentProps.AgentRadius = CapsuleComponent->GetScaledCapsuleRadius();
-	NavAgentProps.AgentHeight = CapsuleComponent->GetScaledCapsuleHalfHeight() * 2.f;
+	NavAgentProps.AgentRadius = CapsuleComponent.GetScaledCapsuleRadius();
+	NavAgentProps.AgentHeight = CapsuleComponent.GetScaledCapsuleHalfHeight() * 2.f;
 }
 
 void UNavMovementComponent::SetUpdateNavAgentWithOwnersCollisions(bool bUpdateWithOwner)
diff --git a/Engine/Source/Runtime/Engine/Private/Pawn.cpp b/Engine/Source/Runtime/Engine/Private/Pawn.cpp
--- a/Engine/Source/Runtime/Engine/Private/Pawn.cpp
+++ b/Engine/Source/Runtime/Engine/Private/Pawn.cpp
@@ -141,7 +141,7 @@ void APawn::UpdateNavAgent()
 	if (RootComponent != NULL && MovementComponent != NULL && MovementComponent->ShouldUpdateNavAgentWithOwnersCollision())
 	{
 		RootComponent->UpdateBounds();
-		MovementComponent->UpdateNavAgent(this);
+		MovementComponent->UpdateNavAgent(*this);
 	}
 }
 
