Backout change to display nicer text when picking matinee properties as it breaks the ability for the system to actually parse the property path. (GDC blocker)

diff --git a/Engine/Source/Runtime/Engine/Private/MatineeUtils.cpp b/Engine/Source/Runtime/Engine/Private/MatineeUtils.cpp
--- a/Engine/Source/Runtime/Engine/Private/MatineeUtils.cpp
+++ b/Engine/Source/Runtime/Engine/Private/MatineeUtils.cpp
@@ -44,12 +44,8 @@ protected:
 				// Is this property the desired type?
 				if (IsDesiredProperty(ClassMemberProperty))
 				{
-#if WITH_EDITOR
 					const FText DisplayName = ClassMemberProperty->GetDisplayNameText();
-					const FString QualifiedFullPath = FString::Printf(TEXT("%s (%s%s)"), *DisplayName.ToString(), *Prefix, *ClassMemberProperty->GetName());
-#else
 					const FString QualifiedFullPath = FString::Printf(TEXT("%s%s"), *Prefix, *ClassMemberProperty->GetName());
-#endif
 					GatheredPropertyPaths.Add(*QualifiedFullPath);
 				}
 
@@ -61,12 +57,8 @@ protected:
 						UProperty* StructMemberProperty = *StructFieldIt;
 						if (StructMemberProperty->HasAnyPropertyFlags(CPF_Interp) && IsDesiredProperty(StructMemberProperty))
 						{
-#if WITH_EDITOR
 							const FText DisplayName = StructMemberProperty->GetDisplayNameText();
-							const FString QualifiedFullPath = FString::Printf(TEXT("%s (%s%s.%s)"), *DisplayName.ToString(), *Prefix, *OuterStructProperty->GetName(), *StructMemberProperty->GetName());
-#else
 							const FString QualifiedFullPath = FString::Printf(TEXT("%s%s.%s"), *Prefix, *OuterStructProperty->GetName(), *StructMemberProperty->GetName());
-#endif
 							GatheredPropertyPaths.Add(*QualifiedFullPath);
 						}
 					}
