Fix android compile error

diff --git a/Engine/Source/Runtime/Engine/Private/MatineeUtils.cpp b/Engine/Source/Runtime/Engine/Private/MatineeUtils.cpp
--- a/Engine/Source/Runtime/Engine/Private/MatineeUtils.cpp
+++ b/Engine/Source/Runtime/Engine/Private/MatineeUtils.cpp
@@ -44,7 +44,6 @@ protected:
 				// Is this property the desired type?
 				if (IsDesiredProperty(ClassMemberProperty))
 				{
-					const FText DisplayName = ClassMemberProperty->GetDisplayNameText();
 					const FString QualifiedFullPath = FString::Printf(TEXT("%s%s"), *Prefix, *ClassMemberProperty->GetName());
 					GatheredPropertyPaths.Add(*QualifiedFullPath);
 				}
@@ -57,7 +56,6 @@ protected:
 						UProperty* StructMemberProperty = *StructFieldIt;
 						if (StructMemberProperty->HasAnyPropertyFlags(CPF_Interp) && IsDesiredProperty(StructMemberProperty))
 						{
-							const FText DisplayName = StructMemberProperty->GetDisplayNameText();
 							const FString QualifiedFullPath = FString::Printf(TEXT("%s%s.%s"), *Prefix, *OuterStructProperty->GetName(), *StructMemberProperty->GetName());
 							GatheredPropertyPaths.Add(*QualifiedFullPath);
 						}
