Change possible benign error into a normal log. The calling code should be better at handling whether or not it's a real error, it's not always obvious at this level.

diff --git a/Engine/Source/Runtime/Engine/Private/PackageMapClient.cpp b/Engine/Source/Runtime/Engine/Private/PackageMapClient.cpp
--- a/Engine/Source/Runtime/Engine/Private/PackageMapClient.cpp
+++ b/Engine/Source/Runtime/Engine/Private/PackageMapClient.cpp
@@ -1666,18 +1666,23 @@ UObject* FNetGUIDCache::GetObjectFromNetGUID( const FNetworkGUID& NetGUID, const
 	if ( CacheObjectPtr->bIsBroken )
 	{
 		// This object is broken, we know it won't load
+		// At this stage, any warnings should have already been logged, so we just need to ignore from this point forward
 		return NULL;
 	}
 
 	if ( CacheObjectPtr->bIsPending )
 	{
-		// We're not done loading yet
+		// We're not done loading yet (and no error has been reported yet)
 		return NULL;
 	}
 
 	if ( IsNetGUIDAuthority() )
 	{
-		UE_LOG( LogNetPackageMap, Error, TEXT( "GetObjectFromNetGUID: Guid with no object on server. FullNetGUIDPath: %s" ), *FullNetGUIDPath( NetGUID ) );
+		// The client should never force the server to load objects
+		// In this case, make sure to log it, but we're relying on the calling code to know what's best to do.
+		// As of now, clients only send references to objects to the server via RPC, so this is usually just a NULL parameter that the game code
+		// needs to handle.
+		UE_LOG( LogNetPackageMap, Log, TEXT( "GetObjectFromNetGUID: Guid with no object on server. FullNetGUIDPath: %s" ), *FullNetGUIDPath( NetGUID ) );
 		return NULL;
 	}
 
