help/comment update

diff --git a/Engine/Source/Runtime/Engine/Classes/Materials/MaterialExpressionSphereMask.h b/Engine/Source/Runtime/Engine/Classes/Materials/MaterialExpressionSphereMask.h
--- a/Engine/Source/Runtime/Engine/Classes/Materials/MaterialExpressionSphereMask.h
+++ b/Engine/Source/Runtime/Engine/Classes/Materials/MaterialExpressionSphereMask.h
@@ -26,12 +26,12 @@ class UMaterialExpressionSphereMask : public UMaterialExpression
 	UPROPERTY(meta = (RequiredInput = "false", ToolTip = "Defaults to 'HardnessPercent' if not specified"))
 	FExpressionInput Hardness;
 
-	/** in the units that A and B are measured */
-	UPROPERTY(EditAnywhere, Category=MaterialExpressionSphereMask, meta=(OverridingInputProperty = "Radius"))
+	/** in the unit that A and B are measured */
+	UPROPERTY(EditAnywhere, Category=MaterialExpressionSphereMask, meta=(OverridingInputProperty = "Radius"), DisplayName = "Radius")
 	float AttenuationRadius;
 
-	/** in percent 0%..100% */
-	UPROPERTY(EditAnywhere, Category=MaterialExpressionSphereMask, meta=(UIMin = "0.0", UIMax = "100.0", ClampMin = "0.0", ClampMax = "100.0", OverridingInputProperty = "Hardness"))
+	/** in percent 0%=soft .. 100%=hard */
+	UPROPERTY(EditAnywhere, Category=MaterialExpressionSphereMask, meta=(UIMin = "0.0", UIMax = "100.0", ClampMin = "0.0", ClampMax = "100.0", OverridingInputProperty = "Hardness", DisplayName = "Hardness"))
 	float HardnessPercent;
 
 
