Added new option for navmesh runtime generation - DynamicModifeirsOnly UE-5673
New option supports runtime navmesh rebuilding only for navigation mesh area modifiers (dynamic obstacles). Has a little memory and performance overhead over completely static navmesh.
NavOctree allocates on demand when first navigation volume is placed in the map.

diff --git a/Engine/Source/Runtime/Core/Public/UObject/ObjectVersion.h b/Engine/Source/Runtime/Core/Public/UObject/ObjectVersion.h
--- a/Engine/Source/Runtime/Core/Public/UObject/ObjectVersion.h
+++ b/Engine/Source/Runtime/Core/Public/UObject/ObjectVersion.h
@@ -452,6 +452,8 @@ enum EUnrealEngineObjectUE4Version
 	VER_UE4_FIX_SKEL_VERT_ORIENT_MESH_PARTICLES,
 	// Do not change landscape section offset on load under world composition
 	VER_UE4_LANDSCAPE_STATIC_SECTION_OFFSET,
+	// New options for navigation data runtime generation (static, modifiers only, dynamic)
+	VER_UE4_ADD_MODIFIERS_RUNTIME_GENERATION,
 
 	// -----<new versions can be added before this line>-------------------------------------------------
 	// - this needs to be the last line (see note below)
diff --git a/Engine/Source/Runtime/Engine/Classes/AI/Navigation/NavigationData.h b/Engine/Source/Runtime/Engine/Classes/AI/Navigation/NavigationData.h
--- a/Engine/Source/Runtime/Engine/Classes/AI/Navigation/NavigationData.h
+++ b/Engine/Source/Runtime/Engine/Classes/AI/Navigation/NavigationData.h
@@ -374,6 +374,22 @@ private:
 	FVector GoalActorLastLocation;
 };
 
+/** 
+ *  Supported options for runtime navigation data generation
+ */
+UENUM()
+enum class ERuntimeGenerationType : uint8
+{
+	// No runtime generation, fully static navigation data
+	Static,				
+	// Supports only navigation modifiers updates
+	DynamicModifiersOnly,	
+	// Fully dynamic, supports geometry changes along with navigation modifiers
+	Dynamic,
+	// Only for legacy loading don't use it!
+	LegacyGeneration UMETA(Hidden)
+};
+
 /** 
  *	Represents abstract Navigation Data (sub-classed as NavMesh, NavGraph, etc)
  *	Used as a common interface for all navigation types handled by NavigationSystem
@@ -396,10 +412,14 @@ class ENGINE_API ANavigationData : public AActor
 	//----------------------------------------------------------------------//
 	// game-time config
 	//----------------------------------------------------------------------//
-
+	
 	/** If true, the NavMesh can be dynamically rebuilt at runtime. */
+	UPROPERTY(config)
+	uint32 bRebuildAtRuntime_DEPRECATED:1;
+
+	/** Navigation data runtime generation options */
 	UPROPERTY(EditAnywhere, Category = Runtime, config)
-	uint32 bRebuildAtRuntime:1;
+	ERuntimeGenerationType RuntimeGeneration;
 
 	/** By default navigation will skip the first update after being successfully loaded
 	 *  setting bForceRebuildOnLoad to false can override this behavior */
@@ -442,6 +462,7 @@ class ENGINE_API ANavigationData : public AActor
 
 	virtual bool NeedsRebuild() const { return false; }
 	virtual bool SupportsRuntimeGeneration() const;
+	virtual bool SupportsStreaming() const;
 	virtual void OnNavigationBoundsChanged();
 	virtual void OnStreamingLevelAdded(ULevel* InLevel) {};
 	virtual void OnStreamingLevelRemoved(ULevel* InLevel) {};
@@ -465,7 +486,7 @@ protected:
 
 public:
 	/** Creates new generator in case navigation supports it */
-	virtual void ConstructGenerator();
+	virtual void ConditionalConstructGenerator();
 	
 	/** Triggers rebuild in case navigation supports it */
 	virtual void RebuildAll();
@@ -545,7 +566,7 @@ public:
 	TArray<FBox> GetNavigableBounds() const;
 	
 	/** Returns list of navigable bounds that belongs to specific level */
-	TArray<FBox> GetNavigableBoundsInLevel(const FName& InLevelPackageName) const;
+	TArray<FBox> GetNavigableBoundsInLevel(ULevel* InLevel) const;
 	
 	//----------------------------------------------------------------------//
 	// Debug                                                                
diff --git a/Engine/Source/Runtime/Engine/Classes/AI/Navigation/NavigationSystem.h b/Engine/Source/Runtime/Engine/Classes/AI/Navigation/NavigationSystem.h
--- a/Engine/Source/Runtime/Engine/Classes/AI/Navigation/NavigationSystem.h
+++ b/Engine/Source/Runtime/Engine/Classes/AI/Navigation/NavigationSystem.h
@@ -465,9 +465,9 @@ protected:
 	void RegisterNavigationDataInstances();
 
 	/** called in places where we need to spawn the NavOctree, but is checking additional conditions if we really want to do that
-	 *	depending on project setup among others 
+	 *	depending on navigation data setup among others 
 	 *	@return true if NavOctree instance has been created, or if one is already present */
-	virtual bool ConditionallyCreateNavOctree();
+	bool ConditionalPopulateNavOctree();
 
 	/** Processes registration of candidates queues via RequestRegistration and stored in NavDataRegistrationQueue */
 	void ProcessRegistrationCandidates();
@@ -510,10 +510,6 @@ public:
 
 	const FNavigationOctree* GetNavOctree() const { return NavOctree; }
 
-	/** called to gather navigation relevant actors that have been created while
-	 *	Navigation System was not present */
-	void PopulateNavOctree();
-
 	FORCEINLINE void SetObjectsNavOctreeId(UObject* Object, FOctreeElementId Id) { ObjectToOctreeId.Add(Object, Id); }
 	FORCEINLINE const FOctreeElementId* GetObjectsNavOctreeId(const UObject* Object) const { return ObjectToOctreeId.Find(Object); }
 	FORCEINLINE	void RemoveObjectsNavOctreeId(UObject* Object) { ObjectToOctreeId.Remove(Object); }
@@ -823,5 +819,18 @@ private:
 
 	/** Processes pathfinding requests given in PathFindingQueries.*/
 	void PerformAsyncQueries(TArray<FAsyncPathFindingQuery> PathFindingQueries);
+
+	/** */
+	void DestroyNavOctree();
+
+	/** Whether Navigation system needs to populate nav octree. 
+	 *  Depends on runtime generation settings of each navigation data, always true in the editor
+	 */
+	bool RequiresNavOctree() const;
+
+	/** Return "Strongest" runtime generation type required by registered navigation data objects
+	 *  Depends on runtime generation settings of each navigation data, always ERuntimeGenerationType::Dynamic in the editor world
+	 */
+	ERuntimeGenerationType GetRuntimeGenerationType() const;
 };
 
diff --git a/Engine/Source/Runtime/Engine/Classes/AI/Navigation/NavigationTypes.h b/Engine/Source/Runtime/Engine/Classes/AI/Navigation/NavigationTypes.h
--- a/Engine/Source/Runtime/Engine/Classes/AI/Navigation/NavigationTypes.h
+++ b/Engine/Source/Runtime/Engine/Classes/AI/Navigation/NavigationTypes.h
@@ -86,9 +86,9 @@ struct FNavigationDirtyArea
 
 struct FNavigationBounds
 {
-	uint32	UniqueID;
-	FBox	AreaBox;	
-	FName	PackageName;
+	uint32						UniqueID;
+	FBox						AreaBox;	 
+	TWeakObjectPtr<ULevel>		Level;		// The level this bounds belongs to
 
 	bool operator==(const FNavigationBounds& Other) const 
 	{ 
diff --git a/Engine/Source/Runtime/Engine/Classes/AI/Navigation/RecastNavMesh.h b/Engine/Source/Runtime/Engine/Classes/AI/Navigation/RecastNavMesh.h
--- a/Engine/Source/Runtime/Engine/Classes/AI/Navigation/RecastNavMesh.h
+++ b/Engine/Source/Runtime/Engine/Classes/AI/Navigation/RecastNavMesh.h
@@ -363,7 +363,7 @@ struct FNavMeshTileData
 		return NavData->RawNavData;
 	}
 
-	FORCEINLINE const uint8* GetDataSafe() const
+	FORCEINLINE uint8* GetDataSafe()
 	{
 		return NavData.IsValid() ? NavData->RawNavData : NULL;
 	}
@@ -376,6 +376,9 @@ struct FNavMeshTileData
 	FORCEINLINE bool IsValid() const { return NavData.IsValid() && GetData() != nullptr && DataSize > 0; }
 
 	uint8* Release();
+
+	// Duplicate shared state so we will have own copy of the data
+	void MakeUnique();
 };
 
 DECLARE_MULTICAST_DELEGATE(FOnNavMeshUpdate);
@@ -697,6 +700,15 @@ public:
 	/** Retrieves number of tiles in this navmesh */
 	int32 GetNavMeshTilesCount() const;
 
+	/**  */
+	void RemoveTileCacheLayers(int32 TileX, int32 TileY);
+	
+	/**  */
+	void AddTileCacheLayers(int32 TileX, int32 TileY, const TArray<FNavMeshTileData>& Layers);
+	
+	/**  */
+	TArray<FNavMeshTileData> GetTileCacheLayers(int32 TileX, int32 TileY) const;
+	
 	void GetEdgesForPathCorridor(const TArray<NavNodeRef>* PathCorridor, TArray<struct FNavigationPortalEdge>* PathCorridorEdges) const;
 
 	// @todo docuement
@@ -864,7 +876,8 @@ public:
 	
 	virtual bool NeedsRebuild() const override;
 	virtual bool SupportsRuntimeGeneration() const override;
-	virtual void ConstructGenerator() override;
+	virtual bool SupportsStreaming() const override;
+	virtual void ConditionalConstructGenerator() override;
 
 protected:
 	// @todo docuement
diff --git a/Engine/Source/Runtime/Engine/Classes/AI/Navigation/RecastNavMeshDataChunk.h b/Engine/Source/Runtime/Engine/Classes/AI/Navigation/RecastNavMeshDataChunk.h
--- a/Engine/Source/Runtime/Engine/Classes/AI/Navigation/RecastNavMeshDataChunk.h
+++ b/Engine/Source/Runtime/Engine/Classes/AI/Navigation/RecastNavMeshDataChunk.h
@@ -15,11 +15,17 @@ struct FRecastTileData
 	};
 
 	FRecastTileData();
-	FRecastTileData(NavNodeRef Ref, int32 DataSize, uint8* RawData);
+	FRecastTileData(NavNodeRef Ref, int32 TileDataSize, uint8* TileRawData, int32 TileCacheDataSize, uint8* TileCacheRawData);
 	
 	NavNodeRef				TileRef;
+
+	// Tile data
 	int32					TileDataSize;
 	TSharedPtr<FRawData>	TileRawData;
+
+	// Compressed tile cache layer 
+	int32					TileCacheDataSize;
+	TSharedPtr<FRawData>	TileCacheRawData;
 };
 
 class dtNavMesh;
@@ -37,10 +43,10 @@ class ENGINE_API URecastNavMeshDataChunk : public UNavigationDataChunk
 	// End UObject Interface
 
 	/** Attaches tiles to specified navmesh, transferring tile ownership to navmesh */
-	TArray<uint32> AttachTiles(dtNavMesh* NavMesh);
+	TArray<uint32> AttachTiles(FPImplRecastNavMesh* NavMeshImpl);
 	
 	/** Detaches tiles from specified navmesh, taking tile ownership */
-	TArray<uint32> DetachTiles(dtNavMesh* NavMesh);
+	TArray<uint32> DetachTiles(FPImplRecastNavMesh* NavMeshImpl);
 	
 	/** Number of tiles in this chunk */
 	int32 GetNumTiles() const;
@@ -50,7 +56,7 @@ class ENGINE_API URecastNavMeshDataChunk : public UNavigationDataChunk
 
 public:
 #if WITH_EDITOR
-	void GatherTiles(const dtNavMesh* NavMesh, const TArray<int32>& TileIndices);
+	void GatherTiles(const FPImplRecastNavMesh* NavMeshImpl, const TArray<int32>& TileIndices);
 #endif// WITH_EDITOR
 
 private:
diff --git a/Engine/Source/Runtime/Engine/Classes/Components/InstancedStaticMeshComponent.h b/Engine/Source/Runtime/Engine/Classes/Components/InstancedStaticMeshComponent.h
--- a/Engine/Source/Runtime/Engine/Classes/Components/InstancedStaticMeshComponent.h
+++ b/Engine/Source/Runtime/Engine/Classes/Components/InstancedStaticMeshComponent.h
@@ -157,7 +157,11 @@ public:
 	
 	virtual bool DoCustomNavigableGeometryExport(struct FNavigableGeometryExport* GeomExport) const override;
 	// End UPrimitiveComponent Interface
-
+	
+	// Begin UNavRelevantInterface interface.
+	virtual void GetNavigationData(FNavigationRelevantData& Data) const override;
+	// End UPrimitiveComponent interface.
+	
 	//Begin UObject Interface
 	virtual void Serialize(FArchive& Ar) override;
 	virtual SIZE_T GetResourceSize(EResourceSizeMode::Type Mode) override;
diff --git a/Engine/Source/Runtime/Engine/Classes/Components/StaticMeshComponent.h b/Engine/Source/Runtime/Engine/Classes/Components/StaticMeshComponent.h
--- a/Engine/Source/Runtime/Engine/Classes/Components/StaticMeshComponent.h
+++ b/Engine/Source/Runtime/Engine/Classes/Components/StaticMeshComponent.h
@@ -296,6 +296,10 @@ public:
 #endif
 	// End UPrimitiveComponent interface.
 
+	// Begin UNavRelevantInterface interface.
+	virtual void GetNavigationData(FNavigationRelevantData& Data) const override;
+	// End UPrimitiveComponent interface.
+	
 	// Begin UMeshComponent interface
 	virtual TArray<class UMaterialInterface*> GetMaterials() const override;
 	// End UMeshComponent interface
diff --git a/Engine/Source/Runtime/Engine/Private/AI/Navigation/NavigationData.cpp b/Engine/Source/Runtime/Engine/Private/AI/Navigation/NavigationData.cpp
--- a/Engine/Source/Runtime/Engine/Private/AI/Navigation/NavigationData.cpp
+++ b/Engine/Source/Runtime/Engine/Private/AI/Navigation/NavigationData.cpp
@@ -99,7 +99,7 @@ FSupportedAreaData::FSupportedAreaData(TSubclassOf<UNavArea> NavAreaClass, int32
 ANavigationData::ANavigationData(const FObjectInitializer& ObjectInitializer)
 	: Super(ObjectInitializer)
 	, bEnableDrawing(false)
-	, bRebuildAtRuntime(false)
+	, RuntimeGeneration(ERuntimeGenerationType::LegacyGeneration) //TODO: set to a valid value once bRebuildAtRuntime_DEPRECATED is removed
 	, bForceRebuildOnLoad(false)
 	, DataVersion(NAVMESHVER_LATEST)
 	, FindPathImplementation(NULL)
@@ -135,7 +135,15 @@ void ANavigationData::PostInitProperties()
 	{
 		return;
 	}
-	
+
+	if (HasAnyFlags(RF_ClassDefaultObject))
+	{
+		if (RuntimeGeneration == ERuntimeGenerationType::LegacyGeneration)
+		{
+			RuntimeGeneration = bRebuildAtRuntime_DEPRECATED ? ERuntimeGenerationType::Dynamic : ERuntimeGenerationType::Static;
+		}
+	}
+		
 	if (HasAnyFlags(RF_ClassDefaultObject) == false)
 	{
 		UWorld* WorldOuter = GetWorld();
@@ -171,6 +179,11 @@ void ANavigationData::PostLoad()
 {
 	Super::PostLoad();
 
+	if (GetLinkerUE4Version() < VER_UE4_ADD_MODIFIERS_RUNTIME_GENERATION)
+	{
+		RuntimeGeneration = bRebuildAtRuntime_DEPRECATED ? ERuntimeGenerationType::Dynamic : ERuntimeGenerationType::Static;
+	}
+
 	InstantiateAndRegisterRenderingComponent();
 }
 
@@ -277,7 +290,7 @@ void ANavigationData::OnRegistered()
 	InstantiateAndRegisterRenderingComponent();
 
 	bRegistered = true;
-	ConstructGenerator();
+	ConditionalConstructGenerator();
 }
 
 void ANavigationData::OnUnregistered()
@@ -380,13 +393,18 @@ bool ANavigationData::SupportsRuntimeGeneration() const
 	return false;
 }
 
-void ANavigationData::ConstructGenerator()
+bool ANavigationData::SupportsStreaming() const
+{
+	return false;
+}
+
+void ANavigationData::ConditionalConstructGenerator()
 {
 }
 
 void ANavigationData::RebuildAll()
 {
-	ConstructGenerator(); //recreate generator
+	ConditionalConstructGenerator(); //recreate generator
 	
 	if (NavDataGenerator)
 	{
@@ -413,9 +431,9 @@ void ANavigationData::CancelBuild()
 void ANavigationData::OnNavigationBoundsChanged()
 {
 	// Create generator if it wasn't yet
-	if (SupportsRuntimeGeneration() && NavDataGenerator.Get() == nullptr)
+	if (NavDataGenerator.Get() == nullptr)
 	{
-		ConstructGenerator();
+		ConditionalConstructGenerator();
 	}
 	
 	if (NavDataGenerator)
@@ -460,7 +478,7 @@ TArray<FBox> ANavigationData::GetNavigableBounds() const
 	return Result;
 }
 	
-TArray<FBox> ANavigationData::GetNavigableBoundsInLevel(const FName& InLevelPackageName) const
+TArray<FBox> ANavigationData::GetNavigableBoundsInLevel(ULevel* InLevel) const
 {
 	TArray<FBox> Result;
 	const UNavigationSystem* NavSys = GetWorld()->GetNavigationSystem();
@@ -472,7 +490,7 @@ TArray<FBox> ANavigationData::GetNavigableBoundsInLevel(const FName& InLevelPack
 
 		for (const auto& Bounds : NavigationBounds)
 		{
-			if (Bounds.PackageName == InLevelPackageName)
+			if (Bounds.Level == InLevel)
 			{
 				Result.Add(Bounds.AreaBox);
 			}
diff --git a/Engine/Source/Runtime/Engine/Private/AI/Navigation/NavigationSystem.cpp b/Engine/Source/Runtime/Engine/Private/AI/Navigation/NavigationSystem.cpp
--- a/Engine/Source/Runtime/Engine/Private/AI/Navigation/NavigationSystem.cpp
+++ b/Engine/Source/Runtime/Engine/Private/AI/Navigation/NavigationSystem.cpp
@@ -338,11 +338,8 @@ void UNavigationSystem::PostInitProperties()
 			// setting AWorldSettings.bEnableNavigationSystem to false
 			UE_LOG(LogNavigation, Error, TEXT("No navigation data types found while setting up required navigation types!"));
 		}
-
-		ConditionallyCreateNavOctree();
 	
 		bInitialBuildingLockActive = bInitialBuildingLocked;
-		InitializeLevelCollisions();
 
 		// register for any actor move change
 #if WITH_EDITOR
@@ -365,35 +362,56 @@ void UNavigationSystem::PostInitProperties()
 	}
 }
 
-bool UNavigationSystem::ConditionallyCreateNavOctree()
+bool UNavigationSystem::ConditionalPopulateNavOctree()
 {
-	ensure(NavOctree == nullptr);
-	if (NavOctree != nullptr)
-	{
-		return true;
-	}
-
-	bSupportRebuilding = GetWorld()->IsGameWorld() == false;
-	for (int32 NavDataClassIndex = 0; NavDataClassIndex < NavDataClasses.Num() && bSupportRebuilding == false; ++NavDataClassIndex)
-	{
-		const ANavigationData* NavDataCDO = GetDefault<ANavigationData>(NavDataClasses[NavDataClassIndex]);
-		check(NavDataCDO);
-		bSupportRebuilding = NavDataCDO->bRebuildAtRuntime;
-	}
+	// We are going to fully re-populate NavOctree so all pending update request are outdated
+	PendingOctreeUpdates.Empty(32);
+	
+	// Discard current octree
+	DestroyNavOctree();
+	
+	// See if any of registered navigation data need navoctree
+	bSupportRebuilding = RequiresNavOctree();
 
 	if (bSupportRebuilding)
 	{
-		NavOctree = new FNavigationOctree(FVector(0, 0, 0), 64000);
+		NavOctree = new FNavigationOctree(FVector(0,0,0), 64000);
+		
+		const ERuntimeGenerationType RuntimeGenerationType = GetRuntimeGenerationType();
+		const bool bStoreNavGeometry = (RuntimeGenerationType == ERuntimeGenerationType::Dynamic);
+		NavOctree->SetNavigableGeometryStoringMode(bStoreNavGeometry ? FNavigationOctree::StoreNavGeometry : FNavigationOctree::SkipNavGeometry);
+		if (bStoreNavGeometry)
+		{
 #if WITH_RECAST
-		NavOctree->ComponentExportDelegate = FNavigationOctree::FNavigableGeometryComponentExportDelegate::CreateStatic(&FRecastNavMeshGenerator::ExportComponentGeometry);
+			NavOctree->ComponentExportDelegate = FNavigationOctree::FNavigableGeometryComponentExportDelegate::CreateStatic(&FRecastNavMeshGenerator::ExportComponentGeometry);
 #endif // WITH_RECAST
-		NavOctree->SetNavigableGeometryStoringMode(FNavigationOctree::StoreNavGeometry);
+		}
+
+		UWorld* World = GetWorld();
+		check(World);
+		
+		// now process all actors on all levels
+		for (int32 LevelIndex = 0; LevelIndex < World->GetNumLevels(); ++LevelIndex) 
+		{
+			ULevel* Level = World->GetLevel(LevelIndex);
+			AddLevelCollisionToOctree(Level);
+
+			for (int32 ActorIndex=0; ActorIndex<Level->Actors.Num(); ActorIndex++)
+			{
+				AActor* Actor = Level->Actors[ActorIndex];
+
+				const bool bLegalActor = Actor && !Actor->IsPendingKill();
+				if (bLegalActor)
+				{
+					UpdateNavOctreeAll(Actor);
+				}
+			}
+		}
 	}
 
-	return NavOctree != nullptr;
+	return bSupportRebuilding;
 }
 
-
 #if WITH_EDITOR
 void UNavigationSystem::PostEditChangeProperty(FPropertyChangedEvent& PropertyChangedEvent)
 {
@@ -434,115 +452,106 @@ void UNavigationSystem::OnWorldInitDone(FNavigationSystem::EMode Mode)
 {
 	OperationMode = Mode;
 	
-	if (bSupportRebuilding)
-	{
-		UWorld* World = GetWorld();
+	UWorld* World = GetWorld();
 
-		if (IsThereAnywhereToBuildNavigation() == false
-			// Simulation mode is a special case - better not do it in this case
-			&& OperationMode != FNavigationSystem::SimulationMode)
+	if (IsThereAnywhereToBuildNavigation() == false
+		// Simulation mode is a special case - better not do it in this case
+		&& OperationMode != FNavigationSystem::SimulationMode)
+	{
+		// remove all navigation data instances
+		for (TActorIterator<ANavigationData> It(World); It; ++It)
 		{
-			// remove all navigation data instances
-			for (TActorIterator<ANavigationData> It(World); It; ++It)
+			ANavigationData* Nav = (*It);
+			if (Nav != NULL && Nav->IsPendingKill() == false)
 			{
-				ANavigationData* Nav = (*It);
-				if (Nav != NULL && Nav->IsPendingKill() == false)
-				{
-					UnregisterNavData(Nav);
-					Nav->CleanUpAndMarkPendingKill();
-				}
+				UnregisterNavData(Nav);
+				Nav->CleanUpAndMarkPendingKill();
+				bNavDataRemovedDueToMissingNavBounds = true;
 			}
-
-			if (OperationMode == FNavigationSystem::EditorMode)
-			{
-				bInitialBuildingLockActive = false;
-			}
-
-			bNavDataRemovedDueToMissingNavBounds = true;
 		}
-		else
-		{
-			InitializeLevelCollisions();
-			PopulateNavOctree();
 
-			// gather navigable bounds
-			GatherNavigationBounds();
+		if (OperationMode == FNavigationSystem::EditorMode)
+		{
+			bInitialBuildingLockActive = false;
+		}
+	}
+	else
+	{
+		// gather navigable bounds
+		GatherNavigationBounds();
 
-			// gather all navigation data instances and register all not-yet-registered
-			// (since it's quite possible navigation system was not ready by the time
-			// those instances were serialized-in or spawned)
-			RegisterNavigationDataInstances();
+		// gather all navigation data instances and register all not-yet-registered
+		// (since it's quite possible navigation system was not ready by the time
+		// those instances were serialized-in or spawned)
+		RegisterNavigationDataInstances();
 
-			if (OperationMode == FNavigationSystem::EditorMode)
-			{
-				// don't lock navigation building in editor
-				bInitialBuildingLockActive = false;
-			}
+		if (OperationMode == FNavigationSystem::EditorMode)
+		{
+			// don't lock navigation building in editor
+			bInitialBuildingLockActive = false;
+		}
 
-			if (bAutoCreateNavigationData == true)
-			{
-				SpawnMissingNavigationData();
-				// in case anything spawned has registered
-				ProcessRegistrationCandidates();
-			}
-			else
+		if (bAutoCreateNavigationData == true)
+		{
+			SpawnMissingNavigationData();
+			// in case anything spawned has registered
+			ProcessRegistrationCandidates();
+		}
+		else
+		{
+			if (GetMainNavData(FNavigationSystem::DontCreate) != NULL)
 			{
-				if (GetMainNavData(FNavigationSystem::DontCreate) != NULL)
+				// trigger navmesh update
+				for (TActorIterator<ANavigationData> It(World); It; ++It)
 				{
-					// trigger navmesh update
-					for (TActorIterator<ANavigationData> It(World); It; ++It)
+					ANavigationData* NavData = (*It);
+					if (NavData != NULL)
 					{
-						ANavigationData* NavData = (*It);
-						if (NavData != NULL)
-						{
-							ERegistrationResult Result = RegisterNavData(NavData);
+						ERegistrationResult Result = RegisterNavData(NavData);
 
-							if (Result == RegistrationSuccessful)
-							{
+						if (Result == RegistrationSuccessful)
+						{
 #if WITH_RECAST
-								if (Cast<ARecastNavMesh>(NavData) != NULL)
+							if (Cast<ARecastNavMesh>(NavData) != NULL)
+							{
+								if (bInitialBuildingLockActive == false && bNavigationAutoUpdateEnabled)
 								{
-									if (bInitialBuildingLockActive == false && bNavigationAutoUpdateEnabled)
-									{
-										NavData->RebuildAll();
-									}
+									NavData->RebuildAll();
 								}
-#endif // WITH_RECAST
-							}
-							else if (Result != RegistrationFailed_DataPendingKill
-								&& Result != RegistrationFailed_AgentNotValid
-								)
-							{
-								NavData->CleanUpAndMarkPendingKill();
 							}
+#endif // WITH_RECAST
+						}
+						else if (Result != RegistrationFailed_DataPendingKill
+							&& Result != RegistrationFailed_AgentNotValid
+							)
+						{
+							NavData->CleanUpAndMarkPendingKill();
 						}
 					}
 				}
 			}
+		}
 
-			// All navigation actors are registered
-			// Add NavMesh parts from all sub-levels that were streamed in prior NavMesh registration
-			if (World->IsGameWorld())
+		// See if any of registered navigation data needs NavOctree
+		ConditionalPopulateNavOctree();
+
+		// All navigation actors are registered
+		// Add NavMesh parts from all sub-levels that were streamed in prior NavMesh registration
+		if (World->IsGameWorld())
+		{
+			const auto& Levels = World->GetLevels();
+			for (ULevel* Level : Levels)
 			{
-				const auto& Levels = World->GetLevels();
-				for (ULevel* Level : Levels)
+				if (!Level->IsPersistentLevel() && Level->bIsVisible)
 				{
-					if (!Level->IsPersistentLevel() && Level->bIsVisible)
+					for (ANavigationData* NavData : NavDataSet)
 					{
-						for (ANavigationData* NavData : NavDataSet)
-						{
-							NavData->OnStreamingLevelAdded(Level);
-						}
+						NavData->OnStreamingLevelAdded(Level);
 					}
 				}
 			}
 		}
 	}
-	else // just register data already present
-	{
-		RegisterNavigationDataInstances();
-		UpdateAbstractNavData();
-	}
 }
 
 void UNavigationSystem::RegisterNavigationDataInstances()
@@ -1331,7 +1340,7 @@ bool UNavigationSystem::IsNavigationBuilt(const AWorldSettings* Settings) const
 		if (NavData != NULL && NavData->GetWorldSettings() == Settings)
 		{
 			FNavDataGenerator* Generator = NavData->GetGenerator();
-			if ((NavData->bRebuildAtRuntime == true 
+			if ((NavData->RuntimeGeneration != ERuntimeGenerationType::Static 
 #if WITH_EDITOR
 				|| GEditor != NULL
 #endif // WITH_EDITOR
@@ -2487,33 +2496,14 @@ void UNavigationSystem::OnActorUnregistered(AActor* Actor)
 	}
 }
 
-void UNavigationSystem::PopulateNavOctree()
+void UNavigationSystem::FindElementsInNavOctree(const FBox& QueryBox, const FNavigationOctreeFilter& Filter, TArray<FNavigationOctreeElement>& Elements)
 {
-	UWorld* World = GetWorld();
-
-	check(World && NavOctree);
-
-	// now process all actors on all levels
-	for (int32 LevelIndex = 0; LevelIndex < World->GetNumLevels(); ++LevelIndex) 
+	if (NavOctree == nullptr)
 	{
-		ULevel* Level = World->GetLevel(LevelIndex);
-		AddLevelCollisionToOctree(Level);
-
-		for (int32 ActorIndex=0; ActorIndex<Level->Actors.Num(); ActorIndex++)
-		{
-			AActor* Actor = Level->Actors[ActorIndex];
-
-			const bool bLegalActor = Actor && !Actor->IsPendingKill();
-			if (bLegalActor)
-			{
-				UpdateNavOctreeAll(Actor);
-			}
-		}
+		UE_LOG(LogNavigation, Warning, TEXT("UNavigationSystem::FindElementsInNavOctree gets called while NavOctree is null"));
+		return;
 	}
-}
-
-void UNavigationSystem::FindElementsInNavOctree(const FBox& QueryBox, const FNavigationOctreeFilter& Filter, TArray<FNavigationOctreeElement>& Elements)
-{
+	
 	for (FNavigationOctree::TConstElementBoxIterator<> It(*NavOctree, QueryBox); It.HasPendingElements(); It.Advance())
 	{
 		const FNavigationOctreeElement& Element = It.GetCurrentElement();
@@ -2626,7 +2616,7 @@ void UNavigationSystem::OnNavigationBoundsUpdated(ANavMeshBoundsVolume* NavVolum
 	FNavigationBoundsUpdateRequest UpdateRequest;
 	UpdateRequest.NavBounds.UniqueID	= NavVolume->GetUniqueID();
 	UpdateRequest.NavBounds.AreaBox		= NavVolume->GetComponentsBoundingBox(true);
-	UpdateRequest.NavBounds.PackageName	= NavVolume->GetOutermost()->GetFName();
+	UpdateRequest.NavBounds.Level		= NavVolume->GetLevel();
 	UpdateRequest.UpdateRequest			= FNavigationBoundsUpdateRequest::Updated;
 	AddNavigationBoundsUpdateRequest(UpdateRequest);
 }
@@ -2641,7 +2631,7 @@ void UNavigationSystem::OnNavigationBoundsAdded(ANavMeshBoundsVolume* NavVolume)
 	FNavigationBoundsUpdateRequest UpdateRequest;
 	UpdateRequest.NavBounds.UniqueID	= NavVolume->GetUniqueID();
 	UpdateRequest.NavBounds.AreaBox		= NavVolume->GetComponentsBoundingBox(true);
-	UpdateRequest.NavBounds.PackageName	= NavVolume->GetOutermost()->GetFName();
+	UpdateRequest.NavBounds.Level		= NavVolume->GetLevel();
 	UpdateRequest.UpdateRequest	= FNavigationBoundsUpdateRequest::Added;
 	AddNavigationBoundsUpdateRequest(UpdateRequest);
 }
@@ -2656,7 +2646,7 @@ void UNavigationSystem::OnNavigationBoundsRemoved(ANavMeshBoundsVolume* NavVolum
 	FNavigationBoundsUpdateRequest UpdateRequest;
 	UpdateRequest.NavBounds.UniqueID	= NavVolume->GetUniqueID();
 	UpdateRequest.NavBounds.AreaBox		= NavVolume->GetComponentsBoundingBox(true);
-	UpdateRequest.NavBounds.PackageName	= NavVolume->GetOutermost()->GetFName();
+	UpdateRequest.NavBounds.Level		= NavVolume->GetLevel();
 	UpdateRequest.UpdateRequest	= FNavigationBoundsUpdateRequest::Removed;
 	AddNavigationBoundsUpdateRequest(UpdateRequest);
 }
@@ -2680,14 +2670,10 @@ void UNavigationSystem::AddNavigationBoundsUpdateRequest(const FNavigationBounds
 
 void UNavigationSystem::PerformNavigationBoundsUpdate(const TArray<FNavigationBoundsUpdateRequest>& UpdateRequests)
 {
-	if (bNavDataRemovedDueToMissingNavBounds)
-	{
-		PopulateNavOctree();
-		bNavDataRemovedDueToMissingNavBounds = false;
-	}
-	
 	if (NavDataSet.Num() == 0)
 	{
+		//TODO: will hitch when user places first navigation volume in the world 
+		
 		if (NavDataRegistrationQueue.Num() > 0)
 		{
 			ProcessRegistrationCandidates();
@@ -2698,6 +2684,8 @@ void UNavigationSystem::PerformNavigationBoundsUpdate(const TArray<FNavigationBo
 			SpawnMissingNavigationData();
 			ProcessRegistrationCandidates();
 		}
+				
+		ConditionalPopulateNavOctree();
 	}
 	
 	// Create list of areas that needs to be updated
@@ -2775,7 +2763,7 @@ void UNavigationSystem::GatherNavigationBounds()
 			FNavigationBounds NavBounds;
 			NavBounds.UniqueID		= V->GetUniqueID();
 			NavBounds.AreaBox		= V->GetComponentsBoundingBox(true);
-			NavBounds.PackageName	= V->GetOutermost()->GetFName();
+			NavBounds.Level			= V->GetLevel();
 			RegisteredNavBounds.Add(NavBounds);
 		}
 	}
@@ -2997,8 +2985,8 @@ void UNavigationSystem::RebuildAll()
 	for (int32 NavDataIndex = 0; NavDataIndex < NavDataSet.Num(); ++NavDataIndex)
 	{
 		ANavigationData* NavData = NavDataSet[NavDataIndex];
-
-		if (NavData && (NavData->bRebuildAtRuntime || (GIsEditor && !bIsInGame)))
+				
+		if (NavData && (!bIsInGame || NavData->SupportsRuntimeGeneration()))
 		{
 			NavData->RebuildAll();
 		}
@@ -3099,7 +3087,8 @@ void UNavigationSystem::OnLevelRemovedFromWorld(ULevel* InLevel, UWorld* InWorld
 void UNavigationSystem::AddLevelCollisionToOctree(ULevel* Level)
 {
 #if WITH_RECAST
-	if (Level && NavOctree)
+	if (Level && NavOctree && 
+		NavOctree->GetNavGeometryStoringMode() == FNavigationOctree::StoreNavGeometry)
 	{
 		const TArray<FVector>* LevelGeom = Level->GetStaticNavigableGeometry();
 		const FOctreeElementId* ElementId = GetObjectsNavOctreeId(Level);
@@ -3124,20 +3113,23 @@ void UNavigationSystem::AddLevelCollisionToOctree(ULevel* Level)
 
 void UNavigationSystem::RemoveLevelCollisionFromOctree(ULevel* Level)
 {
-	const FOctreeElementId* ElementId = GetObjectsNavOctreeId(Level);
-	UE_LOG(LogNavOctree, Log, TEXT("UNREG %s %s"), *GetNameSafe(Level), ElementId ? TEXT("[exists]") : TEXT(""));
-
-	if (ElementId != NULL)
+	if (Level && NavOctree)
 	{
-		if (NavOctree->IsValidElementId(*ElementId))
+		const FOctreeElementId* ElementId = GetObjectsNavOctreeId(Level);
+		UE_LOG(LogNavOctree, Log, TEXT("UNREG %s %s"), *GetNameSafe(Level), ElementId ? TEXT("[exists]") : TEXT(""));
+
+		if (ElementId != NULL)
 		{
-			// mark area occupied by given actor as dirty
-			FNavigationOctreeElement& ElementData = NavOctree->GetElementById(*ElementId);
-			AddDirtyArea(ElementData.Bounds.GetBox(), ENavigationDirtyFlag::All);
-		}
+			if (NavOctree->IsValidElementId(*ElementId))
+			{
+				// mark area occupied by given actor as dirty
+				FNavigationOctreeElement& ElementData = NavOctree->GetElementById(*ElementId);
+				AddDirtyArea(ElementData.Bounds.GetBox(), ENavigationDirtyFlag::All);
+			}
 
-		NavOctree->RemoveNode(*ElementId);
-		RemoveObjectsNavOctreeId(Level);
+			NavOctree->RemoveNode(*ElementId);
+			RemoveObjectsNavOctreeId(Level);
+		}
 	}
 }
 
@@ -3187,15 +3179,8 @@ void UNavigationSystem::CleanUp(ECleanupMode Mode)
 	FWorldDelegates::LevelAddedToWorld.RemoveAll(this);
 	FWorldDelegates::LevelRemovedFromWorld.RemoveAll(this);
 
-	if (NavOctree != NULL)
-	{
-		NavOctree->Destroy();
-		delete NavOctree;
-		NavOctree = NULL;
-	}
-
-	ObjectToOctreeId.Empty();
-
+	DestroyNavOctree();
+	
 	for (int32 NavDataIndex = 0; NavDataIndex < NavDataSet.Num(); ++NavDataIndex)
 	{
 		if (NavDataSet[NavDataIndex] != NULL)
@@ -3216,6 +3201,64 @@ void UNavigationSystem::CleanUp(ECleanupMode Mode)
 	}
 }
 
+void UNavigationSystem::DestroyNavOctree()
+{
+	if (NavOctree != NULL)
+	{
+		NavOctree->Destroy();
+		delete NavOctree;
+		NavOctree = NULL;
+	}
+
+	ObjectToOctreeId.Empty();
+}
+
+bool UNavigationSystem::RequiresNavOctree() const
+{
+	UWorld* World = GetWorld();
+	check(World);
+	
+	// We always require navoctree in editor worlds
+	if (!World->IsGameWorld())
+	{
+		return true;
+	}
+		
+	for (ANavigationData* NavData : NavDataSet)
+	{
+		if (NavData && NavData->SupportsRuntimeGeneration())
+		{
+			return true;
+		}
+	}
+	
+	return false;
+}
+
+ERuntimeGenerationType UNavigationSystem::GetRuntimeGenerationType() const
+{
+	UWorld* World = GetWorld();
+	check(World);
+	
+	// We always use ERuntimeGenerationType::Dynamic in editor worlds
+	if (!World->IsGameWorld())
+	{
+		return ERuntimeGenerationType::Dynamic;
+	}
+	
+	ERuntimeGenerationType RuntimeGenerationType = ERuntimeGenerationType::Static;
+
+	for (ANavigationData* NavData : NavDataSet)
+	{
+		if (NavData && NavData->RuntimeGeneration > RuntimeGenerationType)
+		{
+			RuntimeGenerationType = NavData->RuntimeGeneration;
+		}
+	}
+	
+	return RuntimeGenerationType;
+}
+
 //----------------------------------------------------------------------//
 // Blueprint functions
 //----------------------------------------------------------------------//
@@ -3413,10 +3456,12 @@ bool UNavigationSystem::IsNavigationDirty() const
 
 bool UNavigationSystem::CanRebuildDirtyNavigation() const
 {
+	const bool bIsInGame = GetWorld()->IsGameWorld();
+
 	for (int32 NavDataIndex = 0; NavDataIndex < NavDataSet.Num(); ++NavDataIndex)
 	{
 		const bool bIsDirty = NavDataSet[NavDataIndex]->NeedsRebuild();
-		const bool bCanRebuild = NavDataSet[NavDataIndex]->SupportsRuntimeGeneration();
+		const bool bCanRebuild = !bIsInGame || NavDataSet[NavDataIndex]->SupportsRuntimeGeneration();
 
 		if (bIsDirty && !bCanRebuild)
 		{
diff --git a/Engine/Source/Runtime/Engine/Private/AI/Navigation/PImplRecastNavMesh.cpp b/Engine/Source/Runtime/Engine/Private/AI/Navigation/PImplRecastNavMesh.cpp
--- a/Engine/Source/Runtime/Engine/Private/AI/Navigation/PImplRecastNavMesh.cpp
+++ b/Engine/Source/Runtime/Engine/Private/AI/Navigation/PImplRecastNavMesh.cpp
@@ -261,6 +261,9 @@ void FPImplRecastNavMesh::ReleaseDetourNavMesh()
 		dtFreeNavMesh(DetourNavMesh);
 	}
 	DetourNavMesh = nullptr;
+	
+	//
+	CompressedTileCacheLayers.Empty();
 }
 
 /**
@@ -298,12 +301,11 @@ void FPImplRecastNavMesh::Serialize( FArchive& Ar )
 	if (Ar.IsSaving())
 	{
 		TilesToSave.Reserve(DetourNavMesh->getMaxTiles());
-
-		if (!NavMeshOwner->bRebuildAtRuntime && !IsRunningCommandlet())
+		
+		if (NavMeshOwner->SupportsStreaming() && !IsRunningCommandlet())
 		{
-			// For static navmeshes we save only tiles that belongs to this level
-			FName LevelPackageName = NavMeshOwner->GetOutermost()->GetFName();
-			GetNavMeshTilesIn(NavMeshOwner->GetNavigableBoundsInLevel(LevelPackageName), TilesToSave);
+			// We save only tiles that belongs to this level
+			GetNavMeshTilesIn(NavMeshOwner->GetNavigableBoundsInLevel(NavMeshOwner->GetLevel()), TilesToSave);
 		}
 		else
 		{
@@ -361,6 +363,15 @@ void FPImplRecastNavMesh::Serialize( FArchive& Ar )
 				{
 					dtMeshHeader* const TileHeader = (dtMeshHeader*)TileData;
 					dtFree(TileHeader);
+
+					//
+					if (Ar.UE4Ver() >= VER_UE4_ADD_MODIFIERS_RUNTIME_GENERATION)
+					{
+						unsigned char* ComressedTileData = NULL;
+						int32 CompressedTileDataSize = 0;
+						SerializeCompressedTileCacheData(Ar, ComressedTileData, CompressedTileDataSize);
+						dtFree(ComressedTileData);
+					}
 				}
 			}
 		}
@@ -392,13 +403,29 @@ void FPImplRecastNavMesh::Serialize( FArchive& Ar )
 				{
 					dtMeshHeader* const TileHeader = (dtMeshHeader*)TileData;
 					DetourNavMesh->addTile(TileData, TileDataSize, DT_TILE_FREE_DATA, TileRef, NULL);
+
+					// Serialize compressed tile cache layer
+					if (Ar.UE4Ver() >= VER_UE4_ADD_MODIFIERS_RUNTIME_GENERATION)
+					{
+						uint8* ComressedTileData = nullptr;
+						int32 CompressedTileDataSize = 0;
+						SerializeCompressedTileCacheData(Ar, ComressedTileData, CompressedTileDataSize);
+						
+						if (CompressedTileDataSize > 0)
+						{
+							AddTileCacheLayer(TileHeader->x, TileHeader->y, TileHeader->layer,
+								FNavMeshTileData(ComressedTileData, CompressedTileDataSize, TileHeader->layer, Recast2UnrealBox(TileHeader->bmin, TileHeader->bmax)));
+						}
+					}
 				}
 			}
 		}
 	}
 	else if (Ar.IsSaving())
 	{
+		const bool bSupportsRuntimeGeneration = NavMeshOwner->SupportsRuntimeGeneration();
 		dtNavMesh const* ConstNavMesh = DetourNavMesh;
+		
 		for (int TileIndex : TilesToSave)
 		{
 			const dtMeshTile* Tile = ConstNavMesh->getTile(TileIndex);
@@ -408,6 +435,20 @@ void FPImplRecastNavMesh::Serialize( FArchive& Ar )
 
 			unsigned char* TileData = Tile->data;
 			SerializeRecastMeshTile(Ar, TileData, TileDataSize);
+
+			// Serialize compressed tile cache layer only if navmesh requires it
+			{
+				uint8* CompressedData = nullptr;
+				int32 CompressedDataSize = 0;
+				if (bSupportsRuntimeGeneration)
+				{
+					FNavMeshTileData TileCacheLayer = GetTileCacheLayer(Tile->header->x, Tile->header->y, Tile->header->layer);
+					CompressedData = TileCacheLayer.GetDataSafe();
+					CompressedDataSize = TileCacheLayer.DataSize;
+				}
+				
+				SerializeCompressedTileCacheData(Ar, CompressedData, CompressedDataSize);
+			}
 		}
 	}
 }
@@ -624,6 +665,26 @@ void FPImplRecastNavMesh::SerializeRecastMeshTile(FArchive& Ar, unsigned char*&
 	}
 }
 
+void FPImplRecastNavMesh::SerializeCompressedTileCacheData(FArchive& Ar, unsigned char*& CompressedData, int32& CompressedDataSize)
+{
+	Ar << CompressedDataSize;
+
+	if (CompressedDataSize > 0)
+	{
+		if (Ar.IsLoading())
+		{
+			CompressedData = (unsigned char*)dtAlloc(sizeof(unsigned char)*CompressedDataSize, DT_ALLOC_PERM);
+			if (!CompressedData)
+			{
+				UE_LOG(LogNavigation, Error, TEXT("Failed to alloc tile compressed data"));
+			}
+			FMemory::Memset(CompressedData, 0, CompressedDataSize);
+		}
+
+		Ar.Serialize(CompressedData, CompressedDataSize);
+	}
+}
+
 void FPImplRecastNavMesh::SetRecastMesh(dtNavMesh* NavMesh)
 {
 	if (NavMesh == DetourNavMesh)
@@ -2346,6 +2407,67 @@ void FPImplRecastNavMesh::SetFilterForbiddenFlags(FRecastQueryFilter* Filter, ui
 	// include-exclude don't need to be symmetrical, filter will check both conditions
 }
 
+void FPImplRecastNavMesh::RemoveTileCacheLayers(int32 TileX, int32 TileY)
+{
+	CompressedTileCacheLayers.Remove(FIntPoint(TileX, TileY));
+}
+
+void FPImplRecastNavMesh::RemoveTileCacheLayer(int32 TileX, int32 TileY, int32 LayerIdx)
+{
+	TArray<FNavMeshTileData>* ExistingLayersList = CompressedTileCacheLayers.Find(FIntPoint(TileX, TileY));
+	if (ExistingLayersList)
+	{
+		if (ExistingLayersList->IsValidIndex(LayerIdx))
+		{
+			ExistingLayersList->RemoveAt(LayerIdx);
+		}
+		
+		if (ExistingLayersList->Num() == 0)
+		{
+			CompressedTileCacheLayers.Remove(FIntPoint(TileX, TileY));
+		}
+	}
+}
+
+void FPImplRecastNavMesh::AddTileCacheLayers(int32 TileX, int32 TileY, const TArray<FNavMeshTileData>& Layers)
+{
+	CompressedTileCacheLayers.Add(FIntPoint(TileX, TileY), Layers);
+}
+
+void FPImplRecastNavMesh::AddTileCacheLayer(int32 TileX, int32 TileY, int32 LayerIdx, const FNavMeshTileData& LayerData)
+{
+	TArray<FNavMeshTileData>* ExistingLayersList = CompressedTileCacheLayers.Find(FIntPoint(TileX, TileY));
+	
+	if (ExistingLayersList)
+	{
+		ExistingLayersList->SetNum(FMath::Max(ExistingLayersList->Num(), LayerIdx + 1));
+		(*ExistingLayersList)[LayerIdx] = LayerData;
+	}
+	else
+	{
+		TArray<FNavMeshTileData> LayersList;
+		LayersList.SetNum(FMath::Max(LayersList.Num(), LayerIdx + 1));
+		LayersList[LayerIdx] = LayerData;
+		CompressedTileCacheLayers.Add(FIntPoint(TileX, TileY), LayersList);
+	}
+}
+
+FNavMeshTileData FPImplRecastNavMesh::GetTileCacheLayer(int32 TileX, int32 TileY, int32 LayerIdx) const
+{
+	const TArray<FNavMeshTileData>* LayersList = CompressedTileCacheLayers.Find(FIntPoint(TileX, TileY));
+	if (LayersList && LayersList->IsValidIndex(LayerIdx))
+	{
+		return (*LayersList)[LayerIdx];
+	}
+
+	return FNavMeshTileData();
+}
+
+TArray<FNavMeshTileData> FPImplRecastNavMesh::GetTileCacheLayers(int32 TileX, int32 TileY) const
+{
+	return CompressedTileCacheLayers.FindRef(FIntPoint(TileX, TileY));
+}
+
 #undef INITIALIZE_NAVQUERY
 
 #endif // WITH_RECAST
diff --git a/Engine/Source/Runtime/Engine/Private/AI/Navigation/RecastNavMesh.cpp b/Engine/Source/Runtime/Engine/Private/AI/Navigation/RecastNavMesh.cpp
--- a/Engine/Source/Runtime/Engine/Private/AI/Navigation/RecastNavMesh.cpp
+++ b/Engine/Source/Runtime/Engine/Private/AI/Navigation/RecastNavMesh.cpp
@@ -79,6 +79,17 @@ uint8* FNavMeshTileData::Release()
 	return RawData;
 }
 
+void FNavMeshTileData::MakeUnique()
+{
+	if (DataSize > 0 && !NavData.IsUnique())
+	{
+		INC_MEMORY_STAT_BY(STAT_Navigation_TileCacheMemory, DataSize);
+		uint8* UniqueRawData = (uint8*)dtAlloc(sizeof(uint8)*DataSize, DT_ALLOC_PERM);
+		FMemory::Memcpy(UniqueRawData, NavData->RawNavData, DataSize);
+		NavData = MakeShareable(new FNavData(UniqueRawData));
+	}
+}
+
 float ARecastNavMesh::DrawDistanceSq = 0.0f;
 #if !WITH_RECAST
 
@@ -822,6 +833,32 @@ int32 ARecastNavMesh::GetNavMeshTilesCount() const
 	return NumTiles;
 }
 
+void ARecastNavMesh::RemoveTileCacheLayers(int32 TileX, int32 TileY)
+{
+	if (RecastNavMeshImpl)
+	{
+		RecastNavMeshImpl->RemoveTileCacheLayers(TileX, TileY);
+	}
+}
+	
+void ARecastNavMesh::AddTileCacheLayers(int32 TileX, int32 TileY, const TArray<FNavMeshTileData>& Layers)
+{
+	if (RecastNavMeshImpl)
+	{
+		RecastNavMeshImpl->AddTileCacheLayers(TileX, TileY, Layers);
+	}
+}
+	
+TArray<FNavMeshTileData> ARecastNavMesh::GetTileCacheLayers(int32 TileX, int32 TileY) const
+{
+	if (RecastNavMeshImpl)
+	{
+		return RecastNavMeshImpl->GetTileCacheLayers(TileX, TileY);
+	}
+	
+	return TArray<FNavMeshTileData>();
+}
+
 bool ARecastNavMesh::IsResizable() const
 {
 	return !bFixedTilePoolSize;
@@ -1317,7 +1354,7 @@ void ARecastNavMesh::OnNavMeshGenerationFinished()
 	if (World != nullptr && World->IsPendingKill() == false)
 	{
 #if WITH_EDITOR	
-		// For static navmeshes create navigation data holders in each streaming level
+		// For navmeshes that support streaming create navigation data holders in each streaming level
 		// so parts of navmesh can be streamed in/out with those levels
 		if (!World->IsGameWorld())
 		{
@@ -1331,11 +1368,11 @@ void ARecastNavMesh::OnNavMeshGenerationFinished()
 
 				URecastNavMeshDataChunk* NavDataChunk = GetNavigationDataChunk(Level);
 
-				if (!bRebuildAtRuntime)
+				if (SupportsStreaming())
 				{
-					// We use nav volumes that belongs to this streaming level to find tiles we want to save
+					// We use navigation volumes that belongs to this streaming level to find tiles we want to save
 					TArray<int32> LevelTiles;
-					TArray<FBox> LevelNavBounds = GetNavigableBoundsInLevel(Level->GetOutermost()->GetFName());
+					TArray<FBox> LevelNavBounds = GetNavigableBoundsInLevel(Level);
 					RecastNavMeshImpl->GetNavMeshTilesIn(LevelNavBounds, LevelTiles);
 
 					if (LevelTiles.Num())
@@ -1348,7 +1385,7 @@ void ARecastNavMesh::OnNavMeshGenerationFinished()
 							Level->NavDataChunks.Add(NavDataChunk);
 						}
 
-						NavDataChunk->GatherTiles(RecastNavMeshImpl->DetourNavMesh, LevelTiles);
+						NavDataChunk->GatherTiles(RecastNavMeshImpl, LevelTiles);
 						NavDataChunk->MarkPackageDirty();
 						continue;
 					}
@@ -1427,12 +1464,12 @@ void ARecastNavMesh::OnStreamingLevelAdded(ULevel* InLevel)
 {
 	QUICK_SCOPE_CYCLE_COUNTER(STAT_RecastNavMesh_OnStreamingLevelAdded);
 	
-	if (!bRebuildAtRuntime && GetWorld()->IsGameWorld())
+	if (SupportsStreaming() && GetWorld()->IsGameWorld())
 	{
 		URecastNavMeshDataChunk* NavDataChunk = GetNavigationDataChunk(InLevel);
 		if (NavDataChunk)
 		{
-			TArray<uint32> AttachedIndices = NavDataChunk->AttachTiles(RecastNavMeshImpl->DetourNavMesh);
+			TArray<uint32> AttachedIndices = NavDataChunk->AttachTiles(RecastNavMeshImpl);
 			if (AttachedIndices.Num() > 0)
 			{
 				InvalidateAffectedPaths(AttachedIndices);
@@ -1446,12 +1483,12 @@ void ARecastNavMesh::OnStreamingLevelRemoved(ULevel* InLevel)
 {
 	QUICK_SCOPE_CYCLE_COUNTER(STAT_RecastNavMesh_OnStreamingLevelRemoved);
 	
-	if (!bRebuildAtRuntime && GetWorld()->IsGameWorld())
+	if (SupportsStreaming() && GetWorld()->IsGameWorld())
 	{
 		URecastNavMeshDataChunk* NavDataChunk = GetNavigationDataChunk(InLevel);
 		if (NavDataChunk)
 		{
-			TArray<uint32> DetachedIndices = NavDataChunk->DetachTiles(RecastNavMeshImpl->DetourNavMesh);
+			TArray<uint32> DetachedIndices = NavDataChunk->DetachTiles(RecastNavMeshImpl);
 			if (DetachedIndices.Num() > 0)
 			{
 				InvalidateAffectedPaths(DetachedIndices);
@@ -1795,14 +1832,23 @@ bool ARecastNavMesh::NeedsRebuild() const
 
 bool ARecastNavMesh::SupportsRuntimeGeneration() const
 {
-	// Generator should be enabled in the editor and if navmesh supports runtime generation
-	return (bRebuildAtRuntime || (GetWorld() && !GetWorld()->IsGameWorld()));
+	// Generator should be disabled for Static navmesh
+	return (RuntimeGeneration != ERuntimeGenerationType::Static);
 }
 
-void ARecastNavMesh::ConstructGenerator()
+bool ARecastNavMesh::SupportsStreaming() const
+{
+	// Actually nothing prevents us to support streaming with dynamic generation
+	// Right now streaming in sub-level causes navmesh to build itself, so no point to stream tiles in
+	return (RuntimeGeneration != ERuntimeGenerationType::Dynamic);
+}
+
+void ARecastNavMesh::ConditionalConstructGenerator()
 {
 	NavDataGenerator.Reset();
-	if (SupportsRuntimeGeneration())
+	
+	const bool bRequiresGenerator = SupportsRuntimeGeneration() || !GetWorld()->IsGameWorld();
+	if (bRequiresGenerator)
 	{
 		NavDataGenerator.Reset(new FRecastNavMeshGenerator(*this));
 	}
diff --git a/Engine/Source/Runtime/Engine/Private/AI/Navigation/RecastNavMeshDataChunk.cpp b/Engine/Source/Runtime/Engine/Private/AI/Navigation/RecastNavMeshDataChunk.cpp
--- a/Engine/Source/Runtime/Engine/Private/AI/Navigation/RecastNavMeshDataChunk.cpp
+++ b/Engine/Source/Runtime/Engine/Private/AI/Navigation/RecastNavMeshDataChunk.cpp
@@ -3,6 +3,7 @@
 #include "EnginePrivate.h"
 #include "AI/Navigation/RecastNavMeshDataChunk.h"
 #include "AI/Navigation/PImplRecastNavMesh.h"
+#include "AI/Navigation/RecastHelpers.h"
 
 //----------------------------------------------------------------------//
 // FRecastTileData                                                                
@@ -24,14 +25,17 @@ FRecastTileData::FRawData::~FRawData()
 FRecastTileData::FRecastTileData()
 	: TileRef(0)
 	, TileDataSize(0)
+	, TileCacheDataSize(0)
 {
 }
-	
-FRecastTileData::FRecastTileData(NavNodeRef Ref, int32 DataSize, uint8* RawData)
+
+FRecastTileData::FRecastTileData(NavNodeRef Ref, int32 DataSize, uint8* RawData, int32 CacheDataSize, uint8* CacheRawData)
 	: TileRef(Ref)
 	, TileDataSize(DataSize)
+	, TileCacheDataSize(CacheDataSize)
 {
 	TileRawData = MakeShareable(new FRawData(RawData));
+	TileCacheRawData = MakeShareable(new FRawData(CacheRawData));
 }
 
 //----------------------------------------------------------------------//
@@ -104,12 +108,22 @@ void URecastNavMeshDataChunk::SerializeRecastData(FArchive& Ar)
 			int32 TileDataSize = 0;
 			Ar << TileDataSize;
 
+			// Load tile data 
 			uint8* TileRawData = nullptr;
 			FPImplRecastNavMesh::SerializeRecastMeshTile(Ar, TileRawData, TileDataSize);
 			
 			if (TileRawData != nullptr)
 			{
-				FRecastTileData TileData(0, TileDataSize, TileRawData);
+				// Load compressed tile cache layer
+				int32 TileCacheDataSize = 0;
+				uint8* TileCacheRawData = nullptr;
+				if (Ar.UE4Ver() >= VER_UE4_ADD_MODIFIERS_RUNTIME_GENERATION)
+				{
+					FPImplRecastNavMesh::SerializeCompressedTileCacheData(Ar, TileCacheRawData, TileCacheDataSize);
+				}
+				
+				// We are owner of tile raw data
+				FRecastTileData TileData(0, TileDataSize, TileRawData, TileCacheDataSize, TileCacheRawData);
 				Tiles.Add(TileData);
 			}
 		}
@@ -120,24 +134,30 @@ void URecastNavMeshDataChunk::SerializeRecastData(FArchive& Ar)
 		{
 			if (TileData.TileRawData.IsValid())
 			{
+				// Save tile itself
 				Ar << TileData.TileDataSize;
 				FPImplRecastNavMesh::SerializeRecastMeshTile(Ar, TileData.TileRawData->RawData, TileData.TileDataSize);
+				// Save compressed tile cache layer
+				FPImplRecastNavMesh::SerializeCompressedTileCacheData(Ar, TileData.TileCacheRawData->RawData, TileData.TileCacheDataSize);
 			}
 		}
 	}
 }
 #endif// WITH_RECAST
 
-TArray<uint32> URecastNavMeshDataChunk::AttachTiles(dtNavMesh* NavMesh)
+TArray<uint32> URecastNavMeshDataChunk::AttachTiles(FPImplRecastNavMesh* NavMeshImpl)
 {
 	TArray<uint32> Result;
 	Result.Reserve(Tiles.Num());
 
 #if WITH_RECAST	
+	dtNavMesh* NavMesh = NavMeshImpl->DetourNavMesh;
+
 	for (FRecastTileData& TileData : Tiles)
 	{
 		if (TileData.TileRawData.IsValid())
 		{
+			// Attach mesh tile to target nav mesh 
 			dtStatus status = NavMesh->addTile(TileData.TileRawData->RawData, TileData.TileDataSize, DT_TILE_FREE_DATA, 0, &TileData.TileRef);
 			if (dtStatusFailed(status))
 			{
@@ -145,10 +165,28 @@ TArray<uint32> URecastNavMeshDataChunk::AttachTiles(dtNavMesh* NavMesh)
 				continue;
 			}
 			
-			// We don't own data anymore
+			// We don't own tile data anymore
 			TileData.TileDataSize = 0;
 			TileData.TileRawData->RawData = nullptr;
+			
+			// Attach tile cache layer to target nav mesh
+			if (TileData.TileCacheDataSize > 0)
+			{
+				const dtMeshTile* MeshTile = NavMesh->getTileByRef(TileData.TileRef);
+				check(MeshTile);
+				int32 TileX = MeshTile->header->x;
+				int32 TileY = MeshTile->header->y;
+				int32 TileLayerIdx = MeshTile->header->layer;
+				FBox TileBBox = Recast2UnrealBox(MeshTile->header->bmin, MeshTile->header->bmax);
+				
+				FNavMeshTileData LayerData(TileData.TileCacheRawData->RawData, TileData.TileCacheDataSize, TileLayerIdx, TileBBox);
+				NavMeshImpl->AddTileCacheLayer(TileX, TileY, TileLayerIdx, LayerData);
 
+				// We don't own tile cache data anymore
+				TileData.TileCacheDataSize = 0;
+				TileData.TileCacheRawData->RawData = nullptr;
+			}
+			
 			Result.Add(NavMesh->decodePolyIdTile(TileData.TileRef));
 		}
 	}
@@ -158,17 +196,34 @@ TArray<uint32> URecastNavMeshDataChunk::AttachTiles(dtNavMesh* NavMesh)
 	return Result;
 }
 
-TArray<uint32> URecastNavMeshDataChunk::DetachTiles(dtNavMesh* NavMesh)
+TArray<uint32> URecastNavMeshDataChunk::DetachTiles(FPImplRecastNavMesh* NavMeshImpl)
 {
 	TArray<uint32> Result;
 	Result.Reserve(Tiles.Num());
 
 #if WITH_RECAST	
+	dtNavMesh* NavMesh = NavMeshImpl->DetourNavMesh;
+
 	for (FRecastTileData& TileData : Tiles)
 	{
 		if (TileData.TileRef != 0)
 		{
+			// Detach tile cache layer and take ownership over compressed data
+			{
+				const dtMeshTile* MeshTile = NavMesh->getTileByRef(TileData.TileRef);
+				FNavMeshTileData TileCacheData = NavMeshImpl->GetTileCacheLayer(MeshTile->header->x, MeshTile->header->y, MeshTile->header->layer);
+				if (TileCacheData.IsValid())
+				{
+					TileData.TileCacheDataSize = TileCacheData.DataSize;
+					TileData.TileCacheRawData->RawData = TileCacheData.Release();
+				}
+				
+				NavMeshImpl->RemoveTileCacheLayer(MeshTile->header->x, MeshTile->header->y, MeshTile->header->layer);
+			}
+						
+			// Remove tile from navmesh and take ownership of tile raw data
 			NavMesh->removeTile(TileData.TileRef, &TileData.TileRawData->RawData, &TileData.TileDataSize);
+						
 			Result.Add(NavMesh->decodePolyIdTile(TileData.TileRef));
 		}
 	}
@@ -189,18 +244,31 @@ void URecastNavMeshDataChunk::ReleaseTiles()
 }
 
 #if WITH_EDITOR
-void URecastNavMeshDataChunk::GatherTiles(const dtNavMesh* NavMesh, const TArray<int32>& TileIndices)
+void URecastNavMeshDataChunk::GatherTiles(const FPImplRecastNavMesh* NavMeshImpl, const TArray<int32>& TileIndices)
 {
 	Tiles.Empty(TileIndices.Num());
+
+	const dtNavMesh* NavMesh = NavMeshImpl->DetourNavMesh;
 	
 	for (int32 TileIdx : TileIndices)
 	{
 		const dtMeshTile* Tile = NavMesh->getTile(TileIdx);
 		if (Tile && Tile->header)
 		{
+			// Make our own copy of tile data
 			uint8* RawTileData = (uint8*)dtAlloc(Tile->dataSize, DT_ALLOC_PERM);
 			FMemory::Memcpy(RawTileData, Tile->data, Tile->dataSize);
-			FRecastTileData RecastTileData(NavMesh->getTileRef(Tile), Tile->dataSize, RawTileData);
+
+			// Make our own copy of tile cache data
+			uint8* RawTileCacheData = nullptr;
+			FNavMeshTileData TileCacheData = NavMeshImpl->GetTileCacheLayer(Tile->header->x, Tile->header->y, Tile->header->layer);
+			if (TileCacheData.IsValid())
+			{
+				RawTileCacheData = (uint8*)dtAlloc(TileCacheData.DataSize, DT_ALLOC_PERM);
+				FMemory::Memcpy(RawTileCacheData, TileCacheData.GetData(), TileCacheData.DataSize);
+			}
+
+			FRecastTileData RecastTileData(NavMesh->getTileRef(Tile), Tile->dataSize, RawTileData, TileCacheData.DataSize, RawTileCacheData);
 			Tiles.Add(RecastTileData);
 		}
 	}
diff --git a/Engine/Source/Runtime/Engine/Private/AI/Navigation/RecastNavMeshGenerator.cpp b/Engine/Source/Runtime/Engine/Private/AI/Navigation/RecastNavMeshGenerator.cpp
--- a/Engine/Source/Runtime/Engine/Private/AI/Navigation/RecastNavMeshGenerator.cpp
+++ b/Engine/Source/Runtime/Engine/Private/AI/Navigation/RecastNavMeshGenerator.cpp
@@ -1477,11 +1477,20 @@ void FRecastTileGenerator::Setup(const FRecastNavMeshGenerator& ParentGenerator,
 		}
 	}
 
-	// Take ownership of tile cache data if it exist 
-	CompressedLayers = ParentGenerator.TakeIntermediateLayersData(FIntPoint(TileX, TileY));
+	const bool bGeometryChanged = (DirtyAreas.Num() == 0);
+	if (!bGeometryChanged)
+	{
+		// Get compressed tile cache layers if they exist for this location
+		CompressedLayers = ParentGenerator.GetOwner()->GetTileCacheLayers(TileX, TileY);
+		for (FNavMeshTileData& LayerData : CompressedLayers)
+		{
+			// we don't want to modify shared state inside async task, so make sure we are unique owner
+			LayerData.MakeUnique();
+		}
+	}
 
 	// We have to regenerate layers data in case geometry is changed or tile cache is missing
-	bRegenerateCompressedLayers = (DirtyAreas.Num() == 0 || CompressedLayers.Num() == 0);
+	bRegenerateCompressedLayers = (bGeometryChanged || CompressedLayers.Num() == 0);
 	
 	// Gather geometry for tile if it inside navigable bounds
 	if (InclusionBounds.Num())
@@ -2973,7 +2982,6 @@ void FRecastNavMeshGenerator::CancelBuild()
 	DiscardCurrentBuildingTasks();
 	RunningDirtyTiles.Empty();
 	PendingDirtyTiles.Empty();
-	IntermediateLayerDataMap.Empty();
 
 #if	WITH_EDITOR	
 	RecentlyBuiltTiles.Empty();
@@ -3108,8 +3116,8 @@ TArray<uint32> FRecastNavMeshGenerator::RemoveTileLayers(const int32 TileX, cons
 		}
 	}
 
-	// Remove intermediate layers data at this grid location
-	IntermediateLayerDataMap.Remove(FIntPoint(TileX, TileY));
+	// Remove compressed tile cache layers
+	DestNavMesh->RemoveTileCacheLayers(TileX, TileY);
 
 	return ResultTileIndices;
 }
@@ -3223,13 +3231,6 @@ FBox FRecastNavMeshGenerator::GrowBoundingBox(const FBox& BBox, bool bIncludeAge
 	return FBox(BBox.Min - BBoxGrowOffsetBoth - BBoxGrowOffsetMin, BBox.Max + BBoxGrowOffsetBoth);
 }
 
-TArray<FNavMeshTileData> FRecastNavMeshGenerator::TakeIntermediateLayersData(FIntPoint GridCoord) const
-{
-	TArray<FNavMeshTileData> Result;
-	IntermediateLayerDataMap.RemoveAndCopyValue(GridCoord, Result);
-	return Result;
-}
-
 static bool IntercestBounds(const FBox& TestBox, const TNavStatArray<FBox>& Bounds)
 {
 	for (const FBox& Box : Bounds)
@@ -3514,12 +3515,11 @@ TArray<uint32> FRecastNavMeshGenerator::ProcessTileTasks(const int32 NumTasksToS
 				TArray<uint32> UpdatedTileIndices = AddGeneratedTiles(TileGenerator);
 				UpdatedTiles.Append(UpdatedTileIndices);
 			
-				// Store intermediate layers data, so it can be reused later
-				// TODO: make this optional?
-				TArray<FNavMeshTileData> ComressedLayers = TileGenerator.GetCompressedLayers();
-				if (ComressedLayers.Num())
+				// Store compressed tile cache layers so it can be reused later
+				TArray<FNavMeshTileData> CompressedLayers = TileGenerator.GetCompressedLayers();
+				if (CompressedLayers.Num())
 				{
-					IntermediateLayerDataMap.Add(Element.Coord, ComressedLayers);
+					DestNavMesh->AddTileCacheLayers(Element.Coord.X, Element.Coord.Y, CompressedLayers);
 				}
 			}
 
diff --git a/Engine/Source/Runtime/Engine/Private/InstancedStaticMesh.cpp b/Engine/Source/Runtime/Engine/Private/InstancedStaticMesh.cpp
--- a/Engine/Source/Runtime/Engine/Private/InstancedStaticMesh.cpp
+++ b/Engine/Source/Runtime/Engine/Private/InstancedStaticMesh.cpp
@@ -1360,18 +1360,17 @@ bool UInstancedStaticMeshComponent::DoCustomNavigableGeometryExport(struct FNavi
 	if (StaticMesh && StaticMesh->NavCollision)
 	{
 		UNavCollision* NavCollision = StaticMesh->NavCollision;
-		if (NavCollision->bIsDynamicObstacle)
-				{
-					FCompositeNavModifier Modifiers;
-			NavCollision->GetNavigationModifier(Modifiers, FTransform::Identity);
-					GeomExport->AddNavModifiers(Modifiers);
-				}
-		else if (NavCollision->bHasConvexGeometry)
-				{
-					GeomExport->ExportCustomMesh(NavCollision->ConvexCollision.VertexBuffer.GetData(), NavCollision->ConvexCollision.VertexBuffer.Num(),
+		if (StaticMesh->NavCollision->bIsDynamicObstacle)
+		{
+			return false;
+		}
+		
+		if (NavCollision->bHasConvexGeometry)
+		{
+			GeomExport->ExportCustomMesh(NavCollision->ConvexCollision.VertexBuffer.GetData(), NavCollision->ConvexCollision.VertexBuffer.Num(),
 				NavCollision->ConvexCollision.IndexBuffer.GetData(), NavCollision->ConvexCollision.IndexBuffer.Num(), FTransform::Identity);
 
-					GeomExport->ExportCustomMesh(NavCollision->TriMeshCollision.VertexBuffer.GetData(), NavCollision->TriMeshCollision.VertexBuffer.Num(),
+			GeomExport->ExportCustomMesh(NavCollision->TriMeshCollision.VertexBuffer.GetData(), NavCollision->TriMeshCollision.VertexBuffer.Num(),
 				NavCollision->TriMeshCollision.IndexBuffer.GetData(), NavCollision->TriMeshCollision.IndexBuffer.Num(), FTransform::Identity);
 		}
 		else
@@ -1380,8 +1379,8 @@ bool UInstancedStaticMeshComponent::DoCustomNavigableGeometryExport(struct FNavi
 			if (BodySetup)
 			{
 				GeomExport->ExportRigidBodySetup(*BodySetup, FTransform::Identity);
-					}
-				}
+			}
+		}
 
 		// Hook per instance transform delegate
 		GeomExport->SetNavDataPerInstanceTransformDelegate(FNavDataPerInstanceTransformDelegate::CreateUObject(this, &UInstancedStaticMeshComponent::GetNavigationPerInstanceTransforms));
@@ -1391,6 +1390,14 @@ bool UInstancedStaticMeshComponent::DoCustomNavigableGeometryExport(struct FNavi
 	return false;
 }
 
+void UInstancedStaticMeshComponent::GetNavigationData(FNavigationRelevantData& Data) const
+{
+	Super::GetNavigationData(Data);
+
+	// Hook per instance transform delegate
+	Data.NavDataPerInstanceTransformDelegate = FNavDataPerInstanceTransformDelegate::CreateUObject(this, &UInstancedStaticMeshComponent::GetNavigationPerInstanceTransforms);
+}
+
 void UInstancedStaticMeshComponent::GetNavigationPerInstanceTransforms(const FBox& AreaBox, TArray<FTransform>& InstanceData) const
 {
 	for (const auto& InstancedData : PerInstanceSMData)
diff --git a/Engine/Source/Runtime/Engine/Private/StaticMeshComponent.cpp b/Engine/Source/Runtime/Engine/Private/StaticMeshComponent.cpp
--- a/Engine/Source/Runtime/Engine/Private/StaticMeshComponent.cpp
+++ b/Engine/Source/Runtime/Engine/Private/StaticMeshComponent.cpp
@@ -1607,12 +1607,9 @@ bool UStaticMeshComponent::DoCustomNavigableGeometryExport(struct FNavigableGeom
 		UNavCollision* NavCollision = StaticMesh->NavCollision;
 		if (NavCollision->bIsDynamicObstacle)
 		{
-			FCompositeNavModifier Modifiers;
-			NavCollision->GetNavigationModifier(Modifiers, ComponentToWorld);
-			GeomExport->AddNavModifiers(Modifiers);
 			return false;
 		}
-				
+		
 		if (NavCollision->bHasConvexGeometry)
 		{
 			const FVector Scale3D = ComponentToWorld.GetScale3D();
@@ -1634,6 +1631,20 @@ bool UStaticMeshComponent::DoCustomNavigableGeometryExport(struct FNavigableGeom
 	return true;
 }
 
+void UStaticMeshComponent::GetNavigationData(FNavigationRelevantData& Data) const
+{
+	Super::GetNavigationData(Data);
+
+	if (StaticMesh && StaticMesh->NavCollision)	
+	{
+		UNavCollision* NavCollision = StaticMesh->NavCollision;
+		if (NavCollision->bIsDynamicObstacle)
+		{
+			NavCollision->GetNavigationModifier(Data.Modifiers, ComponentToWorld);
+		}
+	}
+}
+
 #if WITH_EDITOR
 bool UStaticMeshComponent::ComponentIsTouchingSelectionBox(const FBox& InSelBBox, const FEngineShowFlags& ShowFlags, const bool bConsiderOnlyBSP, const bool bMustEncompassEntireComponent) const
 {
diff --git a/Engine/Source/Runtime/Engine/Public/AI/Navigation/PImplRecastNavMesh.h b/Engine/Source/Runtime/Engine/Public/AI/Navigation/PImplRecastNavMesh.h
--- a/Engine/Source/Runtime/Engine/Public/AI/Navigation/PImplRecastNavMesh.h
+++ b/Engine/Source/Runtime/Engine/Public/AI/Navigation/PImplRecastNavMesh.h
@@ -190,6 +190,13 @@ public:
 	dtNavMesh* GetRecastMesh() { return DetourNavMesh; };
 	void ReleaseDetourNavMesh();
 
+	void RemoveTileCacheLayers(int32 TileX, int32 TileY);
+	void RemoveTileCacheLayer(int32 TileX, int32 TileY, int32 LayerIdx);
+	void AddTileCacheLayers(int32 TileX, int32 TileY, const TArray<FNavMeshTileData>& Layers);
+	void AddTileCacheLayer(int32 TileX, int32 TileY, int32 LayerIdx, const FNavMeshTileData& LayerData);
+	FNavMeshTileData GetTileCacheLayer(int32 TileX, int32 TileY, int32 LayerIdx) const;
+	TArray<FNavMeshTileData> GetTileCacheLayers(int32 TileX, int32 TileY) const;
+	
 	/** Assigns recast generated navmesh to this instance.
 	 *	@param bOwnData if true from now on this FPImplRecastNavMesh instance will be responsible for this piece 
 	 *		of memory
@@ -211,12 +218,18 @@ public:
 	/** Recast's runtime navmesh data that we can query against */
 	dtNavMesh* DetourNavMesh;
 
+	/** Compressed layers data, can be reused for tiles generation */
+	TMap<FIntPoint, TArray<FNavMeshTileData>> CompressedTileCacheLayers;
+
 	/** query used for searching data on game thread */
 	mutable dtNavMeshQuery SharedNavQuery;
 
 	/** Helper function to serialize a single Recast tile. */
 	static void SerializeRecastMeshTile(FArchive& Ar, unsigned char*& TileData, int32& TileDataSize);
 
+	/** Helper function to serialize a Recast tile compressed data. */
+	static void SerializeCompressedTileCacheData(FArchive& Ar, unsigned char*& CompressedData, int32& CompressedDataSize);
+
 	/** Initialize data for pathfinding */
 	bool InitPathfinding(const FVector& UnrealStart, const FVector& UnrealEnd, 
 		const dtNavMeshQuery& Query, const dtQueryFilter* Filter,
diff --git a/Engine/Source/Runtime/Engine/Public/AI/Navigation/RecastNavMeshGenerator.h b/Engine/Source/Runtime/Engine/Public/AI/Navigation/RecastNavMeshGenerator.h
--- a/Engine/Source/Runtime/Engine/Public/AI/Navigation/RecastNavMeshGenerator.h
+++ b/Engine/Source/Runtime/Engine/Public/AI/Navigation/RecastNavMeshGenerator.h
@@ -129,7 +129,8 @@ struct FRecastAreaNavModifierElement
 /**
  * Class handling generation of a single tile, caching data that can speed up subsequent tile generations
  */
-class ENGINE_API FRecastTileGenerator : public FNonAbandonableTask
+class ENGINE_API FRecastTileGenerator 
+	: public FNoncopyable
 {
 	friend FRecastNavMeshGenerator;
 
@@ -155,7 +156,7 @@ public:
 
 	// Memory amount used to construct generator 
 	uint32 UsedMemoryOnStartup;
-	
+		
 protected:
 	/** Does the actual tile generations. 
 	 *	@note always trigger tile generation only via TriggerAsyncBuild. This is a worker function
@@ -380,9 +381,6 @@ public:
 	bool HasDirtyTiles() const;
 
 	FBox GrowBoundingBox(const FBox& BBox, bool bIncludeAgentHeight) const;
-
-	/** Transfers ownership if tile cache data to the caller */
-	TArray<FNavMeshTileData> TakeIntermediateLayersData(FIntPoint GridCoord) const;
 	
 #if !(UE_BUILD_SHIPPING || UE_BUILD_TEST)
 	virtual void ExportNavigationData(const FString& FileName) const;
@@ -467,9 +465,6 @@ private:
 	/** */
 	FRecastNavMeshCachedData AdditionalCachedData;
 
-	/** Compressed layers data, can be reused for tiles generation */
-	mutable TMap<FIntPoint, TArray<FNavMeshTileData>> IntermediateLayerDataMap;
-
 	/** */
 	TMapBase<const AActor*, FBox, false> ActorToAreaMap;
 
diff --git a/Engine/Source/Runtime/Engine/Public/AI/NavigationOctree.h b/Engine/Source/Runtime/Engine/Public/AI/NavigationOctree.h
--- a/Engine/Source/Runtime/Engine/Public/AI/NavigationOctree.h
+++ b/Engine/Source/Runtime/Engine/Public/AI/NavigationOctree.h
@@ -161,6 +161,11 @@ public:
 
 	const FNavigationRelevantData* GetDataForID(const FOctreeElementId& Id) const;
 
+	ENavGeometryStoringMode GetNavGeometryStoringMode() const
+	{
+		return bGatherGeometry ? StoreNavGeometry : SkipNavGeometry;
+	}
+
 protected:
 
 	uint32 bGatherGeometry : 1;
