Fixed crash on PIE exit UE-7411

diff --git a/Engine/Source/Runtime/Engine/Private/InstancedStaticMesh.cpp b/Engine/Source/Runtime/Engine/Private/InstancedStaticMesh.cpp
--- a/Engine/Source/Runtime/Engine/Private/InstancedStaticMesh.cpp
+++ b/Engine/Source/Runtime/Engine/Private/InstancedStaticMesh.cpp
@@ -896,6 +896,7 @@ void UInstancedStaticMeshComponent::CreateAllInstanceBodies()
 
 		Instance->CopyBodyInstancePropertiesFrom(&BodyInstance);
 		Instance->InstanceBodyIndex = i; // Set body index 
+		Instance->bAutoWeld = false;
 
 		// make sure we never enable bSimulatePhysics for ISMComps
 		Instance->bSimulatePhysics = false;
