Added streaming simulation for world composition maps while "Settings->Previewing->Enable Automatic Level Streaming" option is enabled
Only simulates sub-levels visibility without loading/unloading them.
UETOOL-180

diff --git a/Engine/Source/Editor/UnrealEd/Private/Editor.cpp b/Engine/Source/Editor/UnrealEd/Private/Editor.cpp
--- a/Engine/Source/Editor/UnrealEd/Private/Editor.cpp
+++ b/Engine/Source/Editor/UnrealEd/Private/Editor.cpp
@@ -1204,7 +1204,16 @@ void UEditorEngine::Tick( float DeltaSeconds, bool bIdleMode )
 						}
 					}
 				}
-
+				
+				// Simulate world composition streaming while in editor world
+				if (EditorContext.World()->WorldComposition)
+				{
+					if (EditorContext.World()->WorldComposition->UpdateEditorStreamingState(ViewLocation))
+					{
+						bProcessViewer = true;
+					}
+				}
+				
 				// Call UpdateLevelStreaming if the visibility of any streaming levels was modified.
 				if ( bProcessViewer )
 				{
diff --git a/Engine/Source/Runtime/Engine/Classes/Engine/WorldComposition.h b/Engine/Source/Runtime/Engine/Classes/Engine/WorldComposition.h
--- a/Engine/Source/Runtime/Engine/Classes/Engine/WorldComposition.h
+++ b/Engine/Source/Runtime/Engine/Classes/Engine/WorldComposition.h
@@ -74,6 +74,13 @@ class ENGINE_API UWorldComposition : public UObject
 	/** Adds or removes level streaming objects to world based on distance settings from current view point */
 	void UpdateStreamingState(const FVector& InLocation);
 
+#if WITH_EDITOR
+	/** Simulates streaming in editor world, only visibility, no loading/unloading, no LOD sub-levels 
+	 *  @returns Whether streaming levels state was updated by this call
+	 */
+	bool UpdateEditorStreamingState(const FVector& InLocation);
+#endif// WITH_EDITOR
+
 	/**
 	 * Evaluates current world origin location against provided view location
 	 * Issues request for world origin rebasing in case location is far enough from current origin
diff --git a/Engine/Source/Runtime/Engine/Private/WorldComposition.cpp b/Engine/Source/Runtime/Engine/Private/WorldComposition.cpp
--- a/Engine/Source/Runtime/Engine/Private/WorldComposition.cpp
+++ b/Engine/Source/Runtime/Engine/Private/WorldComposition.cpp
@@ -606,6 +606,52 @@ void UWorldComposition::UpdateStreamingState()
 	}
 }
 
+#if WITH_EDITOR
+bool UWorldComposition::UpdateEditorStreamingState(const FVector& InLocation)
+{
+	UWorld* OwningWorld = GetWorld();
+	bool bStateChanged = false;
+	
+	// Handle only editor worlds
+	if (!OwningWorld->IsGameWorld() && 
+		!OwningWorld->IsVisibilityRequestPending())
+	{
+		// Get the list of visible and hidden levels from current view point
+		TArray<FDistanceVisibleLevel> DistanceVisibleLevels;
+		TArray<FDistanceVisibleLevel> DistanceHiddenLevels;
+		GetDistanceVisibleLevels(InLocation, DistanceVisibleLevels, DistanceHiddenLevels);
+
+		// Hidden levels
+		for (const auto& Level : DistanceHiddenLevels)
+		{
+			ULevelStreaming* EditorStreamingLevel = OwningWorld->GetLevelStreamingForPackageName(TilesStreaming[Level.TileIdx]->GetWorldAssetPackageFName());
+			if (EditorStreamingLevel && 
+				EditorStreamingLevel->IsLevelLoaded() &&
+				EditorStreamingLevel->bShouldBeVisibleInEditor != false)
+			{
+				EditorStreamingLevel->bShouldBeVisibleInEditor = false;
+				bStateChanged = true;
+			}
+		}
+
+		// Visible levels
+		for (const auto& Level : DistanceVisibleLevels)
+		{
+			ULevelStreaming* EditorStreamingLevel = OwningWorld->GetLevelStreamingForPackageName(TilesStreaming[Level.TileIdx]->GetWorldAssetPackageFName());
+			if (EditorStreamingLevel &&
+				EditorStreamingLevel->IsLevelLoaded() &&
+				EditorStreamingLevel->bShouldBeVisibleInEditor != true)
+			{
+				EditorStreamingLevel->bShouldBeVisibleInEditor = true;
+				bStateChanged = true;
+			}
+		}
+	}
+
+	return bStateChanged;
+}
+#endif// WITH_EDITOR
+
 void UWorldComposition::EvaluateWorldOriginLocation(const FVector& ViewLocation)
 {
 	UWorld* OwningWorld = GetWorld();
