CIS fix

diff --git a/Engine/Source/Runtime/Engine/Private/AI/Navigation/RecastNavMesh.cpp b/Engine/Source/Runtime/Engine/Private/AI/Navigation/RecastNavMesh.cpp
--- a/Engine/Source/Runtime/Engine/Private/AI/Navigation/RecastNavMesh.cpp
+++ b/Engine/Source/Runtime/Engine/Private/AI/Navigation/RecastNavMesh.cpp
@@ -42,7 +42,7 @@ FNavMeshTileData::FNavData::~FNavData()
 #if WITH_RECAST
 	dtFree(RawNavData);
 #else
-	delete RawNavData;
+	FMemory::Free(RawNavData);
 #endif
 }
 
@@ -84,7 +84,11 @@ void FNavMeshTileData::MakeUnique()
 	if (DataSize > 0 && !NavData.IsUnique())
 	{
 		INC_MEMORY_STAT_BY(STAT_Navigation_TileCacheMemory, DataSize);
+#if WITH_RECAST
 		uint8* UniqueRawData = (uint8*)dtAlloc(sizeof(uint8)*DataSize, DT_ALLOC_PERM);
+#else
+		uint8* UniqueRawData = (uint8*)FMemory::Malloc(sizeof(uint8)*DataSize);
+#endif
 		FMemory::Memcpy(UniqueRawData, NavData->RawNavData, DataSize);
 		NavData = MakeShareable(new FNavData(UniqueRawData));
 	}
diff --git a/Engine/Source/Runtime/Engine/Private/AI/Navigation/RecastNavMeshDataChunk.cpp b/Engine/Source/Runtime/Engine/Private/AI/Navigation/RecastNavMeshDataChunk.cpp
--- a/Engine/Source/Runtime/Engine/Private/AI/Navigation/RecastNavMeshDataChunk.cpp
+++ b/Engine/Source/Runtime/Engine/Private/AI/Navigation/RecastNavMeshDataChunk.cpp
@@ -18,7 +18,7 @@ FRecastTileData::FRawData::~FRawData()
 #if WITH_RECAST
 	dtFree(RawData);
 #else
-	delete RawData;
+	FMemory::Free(RawData);
 #endif
 }
 
