Fixed NavMesh movement so velocity is reported correctly.

Fixes agents running on the spot when blocked.

diff --git a/Engine/Source/Runtime/Engine/Private/CharacterMovementComponent.cpp b/Engine/Source/Runtime/Engine/Private/CharacterMovementComponent.cpp
--- a/Engine/Source/Runtime/Engine/Private/CharacterMovementComponent.cpp
+++ b/Engine/Source/Runtime/Engine/Private/CharacterMovementComponent.cpp
@@ -3836,6 +3836,14 @@ void UCharacterMovementComponent::PhysNavWalking(float deltaTime, int32 Iteratio
 		{
 			MoveUpdatedComponent(AdjustedDelta, CharacterOwner->GetActorRotation(), true);
 		}
+
+		// Update velocity to reflect actual move
+		if (!bJustTeleported && !HasRootMotion())
+		{
+			Velocity = (GetActorFeetLocation() - OldLocation) / deltaTime;
+		}
+
+		bJustTeleported = false;
 	}
 	else
 	{
