UE-7070 Fix warning on datatable import

diff --git a/Engine/Source/Runtime/Engine/Private/DataTable.cpp b/Engine/Source/Runtime/Engine/Private/DataTable.cpp
--- a/Engine/Source/Runtime/Engine/Private/DataTable.cpp
+++ b/Engine/Source/Runtime/Engine/Private/DataTable.cpp
@@ -454,7 +454,7 @@ TArray<UProperty*> UDataTable::GetTablePropertyArray(const TArray<const TCHAR*>&
 			FName PropName = MakeValidName(ColumnValue);
 			if(PropName == NAME_None)
 			{
-				OutProblems.Add(FString(TEXT("Missing name for column %d."), ColIdx));
+				OutProblems.Add(FString::Printf(TEXT("Missing name for column %d."), ColIdx));
 			}
 			else
 			{
