UE-7426 Export AudioVolume class

diff --git a/Engine/Source/Runtime/Engine/Classes/Sound/AudioVolume.h b/Engine/Source/Runtime/Engine/Classes/Sound/AudioVolume.h
--- a/Engine/Source/Runtime/Engine/Classes/Sound/AudioVolume.h
+++ b/Engine/Source/Runtime/Engine/Classes/Sound/AudioVolume.h
@@ -94,7 +94,6 @@ struct TStructOpsTypeTraits<FReverbSettings> : public TStructOpsTypeTraitsBase
 };
 
 /** Struct encapsulating settings for interior areas. */
-//@warning: manually mirrored in Components.h
 USTRUCT(BlueprintType)
 struct FInteriorSettings
 {
@@ -153,7 +152,7 @@ struct FInteriorSettings
 	bool operator!=(const FInteriorSettings& Other) const;
 };
 
-UCLASS(hidecategories=(Advanced, Attachment, Collision, Volume))
+UCLASS(hidecategories=(Advanced, Attachment, Collision, Volume), MinimalAPI)
 class AAudioVolume : public AVolume
 {
 	GENERATED_UCLASS_BODY()
