Make it so GetInstancedFoliageActorForLevel will create a foliage actor if it's needed.

diff --git a/Engine/Source/Runtime/Engine/Classes/Foliage/InstancedFoliageActor.h b/Engine/Source/Runtime/Engine/Classes/Foliage/InstancedFoliageActor.h
--- a/Engine/Source/Runtime/Engine/Classes/Foliage/InstancedFoliageActor.h
+++ b/Engine/Source/Runtime/Engine/Classes/Foliage/InstancedFoliageActor.h
@@ -106,7 +106,11 @@ public:
 	static ENGINE_API AInstancedFoliageActor* GetInstancedFoliageActorForCurrentLevel(UWorld* InWorld, bool bCreateIfNone = true);
 
 
-	// Get the instanced foliage actor for the specified streaming level. Never creates a new IFA.
-	static ENGINE_API AInstancedFoliageActor* GetInstancedFoliageActorForLevel(ULevel* Level);
+	/**
+	* Get the instanced foliage actor for the specified streaming level.
+	* @param bCreateIfNone					Create if doesnt already exist
+	* returns								pointer to foliage object instance
+	*/
+	static ENGINE_API AInstancedFoliageActor* GetInstancedFoliageActorForLevel(ULevel* Level, bool bCreateIfNone = true);
 #endif	//WITH_EDITOR
 };
diff --git a/Engine/Source/Runtime/Engine/Private/InstancedFoliage.cpp b/Engine/Source/Runtime/Engine/Private/InstancedFoliage.cpp
--- a/Engine/Source/Runtime/Engine/Private/InstancedFoliage.cpp
+++ b/Engine/Source/Runtime/Engine/Private/InstancedFoliage.cpp
@@ -790,7 +790,7 @@ AInstancedFoliageActor* AInstancedFoliageActor::GetInstancedFoliageActorForCurre
 }
 
 
-AInstancedFoliageActor* AInstancedFoliageActor::GetInstancedFoliageActorForLevel(ULevel* InLevel)
+AInstancedFoliageActor* AInstancedFoliageActor::GetInstancedFoliageActorForLevel(ULevel* InLevel, bool bCreateIfNone /* = true */)
 {
 	if (!InLevel)
 	{
@@ -805,6 +805,15 @@ AInstancedFoliageActor* AInstancedFoliageActor::GetInstancedFoliageActorForLevel
 			return InstancedFoliageActor;
 		}
 	}
+
+	if (bCreateIfNone)
+	{
+		FActorSpawnParameters SpawnParams;
+		SpawnParams.OverrideLevel = InLevel;
+
+		return InLevel->GetWorld()->SpawnActor<AInstancedFoliageActor>(SpawnParams);
+	}
+
 	return nullptr;
 }
 
