Make bNetAddressable a UPROPERTY so that it gets copied during PIE world duplication. Fixes UE-7382.

diff --git a/Engine/Source/Runtime/Engine/Classes/Components/ActorComponent.h b/Engine/Source/Runtime/Engine/Classes/Components/ActorComponent.h
--- a/Engine/Source/Runtime/Engine/Classes/Components/ActorComponent.h
+++ b/Engine/Source/Runtime/Engine/Classes/Components/ActorComponent.h
@@ -59,6 +59,7 @@ protected:
 	uint32 bReplicates:1;
 
 	/** Is this component safe to ID over the network by name?  */
+	UPROPERTY()
 	uint32 bNetAddressable:1;
 
 private:
