Applied some navmesh generation property validation #UE4

Made sure TileSizeUU, CellSize and CellHeight are reasonable.

diff --git a/Engine/Source/Runtime/Engine/Classes/AI/Navigation/RecastNavMesh.h b/Engine/Source/Runtime/Engine/Classes/AI/Navigation/RecastNavMesh.h
--- a/Engine/Source/Runtime/Engine/Classes/AI/Navigation/RecastNavMesh.h
+++ b/Engine/Source/Runtime/Engine/Classes/AI/Navigation/RecastNavMesh.h
@@ -473,11 +473,11 @@ class ENGINE_API ARecastNavMesh : public ANavigationData
 	float TileSizeUU;
 
 	/** horizontal size of voxelization cell */
-	UPROPERTY(EditAnywhere, Category=Generation, config, meta=(ClampMin = "1.0"))
+	UPROPERTY(EditAnywhere, Category = Generation, config, meta = (ClampMin = "1.0", ClampMax = "1024.0"))
 	float CellSize;
 
 	/** vertical size of voxelization cell */
-	UPROPERTY(EditAnywhere, Category=Generation, config, meta=(ClampMin = "1.0"))
+	UPROPERTY(EditAnywhere, Category = Generation, config, meta = (ClampMin = "1.0", ClampMax = "1024.0"))
 	float CellHeight;
 
 	/** Radius of smallest agent to traverse this navmesh */
@@ -507,7 +507,7 @@ class ENGINE_API ARecastNavMesh : public ANavigationData
 	float MergeRegionSize;
 
 	/** How much navigable shapes can get simplified - the higher the value the more freedom */
-	UPROPERTY(EditAnywhere, Category=Generation, config)
+	UPROPERTY(EditAnywhere, Category = Generation, config, meta = (ClampMin = "0.0"))
 	float MaxSimplificationError;
 
 	/** navmesh draw distance in game (always visible in editor) */
@@ -637,6 +637,7 @@ public:
 
 	// Begin UObject Interface
 	virtual void PostInitProperties() override;
+	virtual void PostLoad() override;
 	virtual SIZE_T GetResourceSize(EResourceSizeMode::Type Mode) override;	
 
 #if WITH_EDITOR
diff --git a/Engine/Source/Runtime/Engine/Private/AI/Navigation/RecastNavMesh.cpp b/Engine/Source/Runtime/Engine/Private/AI/Navigation/RecastNavMesh.cpp
--- a/Engine/Source/Runtime/Engine/Private/AI/Navigation/RecastNavMesh.cpp
+++ b/Engine/Source/Runtime/Engine/Private/AI/Navigation/RecastNavMesh.cpp
@@ -37,6 +37,8 @@
 
 #endif // WITH_RECAST
 
+static const int32 ArbitraryMaxVoxelTileSize = 1024;
+
 FNavMeshTileData::FNavData::~FNavData()
 {
 #if WITH_RECAST
@@ -468,6 +470,13 @@ void ARecastNavMesh::CleanUp()
 	DestroyRecastPImpl();
 }
 
+void ARecastNavMesh::PostLoad()
+{
+	Super::PostLoad();
+	// tilesize validation. This is temporary and should get removed by 4.9
+	TileSizeUU = FMath::Clamp(TileSizeUU, CellSize, ArbitraryMaxVoxelTileSize * CellSize);
+}
+
 void ARecastNavMesh::PostInitProperties()
 {
 	if (HasAnyFlags(RF_ClassDefaultObject) == true)
@@ -493,6 +502,8 @@ void ARecastNavMesh::PostInitProperties()
 	
 	Super::PostInitProperties();
 
+	TileSizeUU = FMath::Clamp(TileSizeUU, CellSize, ArbitraryMaxVoxelTileSize * CellSize);
+
 	if (HasAnyFlags(RF_ClassDefaultObject) == false)
 	{
 		DefaultQueryFilter->SetFilterType<FRecastQueryFilter>();
@@ -1796,17 +1807,9 @@ void ARecastNavMesh::PostEditChangeProperty(FPropertyChangedEvent& PropertyChang
 					TileSizeUU = FMath::Max(16.f * AgentRadius, RECAST_MIN_TILE_SIZE);
 				}
 
-				// tile's dimension can't exceed 2^16 x cell size, as it's being stored on 2 bytes
-				const int32 DimensionVX = FMath::CeilToInt(TileSizeUU / CellSize);
-				if (DimensionVX > MAX_uint16)
-				{
-					TileSizeUU = MAX_uint16 * CellSize;
-				}
-				// also it can't be 0, and if it's 1 then we should make sure tile size is equal to cell size
-				else if (DimensionVX <= 1)
-				{
-					TileSizeUU = CellSize;
-				}
+				// tile's can't be too big, otherwise we'll crash while tryng to allocate
+				// memory during navmesh generation
+				TileSizeUU = FMath::Clamp(TileSizeUU, CellSize, ArbitraryMaxVoxelTileSize * CellSize);
 			}
 
 			if (HasAnyFlags(RF_ClassDefaultObject) == false)
