CIS fix #UE4

diff --git a/Engine/Source/Runtime/Engine/Classes/GameFramework/NavMovementComponent.h b/Engine/Source/Runtime/Engine/Classes/GameFramework/NavMovementComponent.h
--- a/Engine/Source/Runtime/Engine/Classes/GameFramework/NavMovementComponent.h
+++ b/Engine/Source/Runtime/Engine/Classes/GameFramework/NavMovementComponent.h
@@ -179,15 +179,3 @@ inline void UNavMovementComponent::StopMovementImmediately()
 	Super::StopMovementImmediately();
 	StopActiveMovement();
 }
-
-//----------------------------------------------------------------------//
-// DEPRECATED
-//----------------------------------------------------------------------//
-const FNavAgentProperties* UNavMovementComponent::GetNavAgentProperties() const 
-{ 
-	return &GetNavAgentPropertiesRef();
-}
-FNavAgentProperties* UNavMovementComponent::GetNavAgentProperties() 
-{ 
-	return &GetNavAgentPropertiesRef();
-}
\ No newline at end of file
diff --git a/Engine/Source/Runtime/Engine/Private/NavMovementComponent.cpp b/Engine/Source/Runtime/Engine/Private/NavMovementComponent.cpp
--- a/Engine/Source/Runtime/Engine/Private/NavMovementComponent.cpp
+++ b/Engine/Source/Runtime/Engine/Private/NavMovementComponent.cpp
@@ -79,3 +79,15 @@ void UNavMovementComponent::SetUpdateNavAgentWithOwnersCollisions(bool bUpdateWi
 {
 	bUpdateNavAgentWithOwnersCollision = bUpdateWithOwner;
 }
+
+//----------------------------------------------------------------------//
+// DEPRECATED
+//----------------------------------------------------------------------//
+const FNavAgentProperties* UNavMovementComponent::GetNavAgentProperties() const
+{
+	return &GetNavAgentPropertiesRef();
+}
+FNavAgentProperties* UNavMovementComponent::GetNavAgentProperties()
+{
+	return &GetNavAgentPropertiesRef();
+}
\ No newline at end of file
