Merging

//depot/UE4-Orion/Engine/Source/Runtime/Engine/...

to //depot/UE4/Engine/Source/Runtime/Engine/...

- Adding a "IsFromLocalPlayer" function to FLocalPlayerContext

diff --git a/Engine/Source/Runtime/Engine/Classes/Engine/LocalPlayer.h b/Engine/Source/Runtime/Engine/Classes/Engine/LocalPlayer.h
--- a/Engine/Source/Runtime/Engine/Classes/Engine/LocalPlayer.h
+++ b/Engine/Source/Runtime/Engine/Classes/Engine/LocalPlayer.h
@@ -22,6 +22,10 @@ struct ENGINE_API FLocalPlayerContext
 	/** Is this context initialized */
 	bool IsInitialized() const;
 
+	/** This function tests if the given Actor is connected to the Local Player in any way. 
+		It tests against the APlayerController, APlayerState, and APawn. */
+	bool IsFromLocalPlayer(const AActor* ActorToTest) const;
+
 	/* Returns the world context. */
 	UWorld* GetWorld() const;
 
diff --git a/Engine/Source/Runtime/Engine/Private/LocalPlayer.cpp b/Engine/Source/Runtime/Engine/Private/LocalPlayer.cpp
--- a/Engine/Source/Runtime/Engine/Private/LocalPlayer.cpp
+++ b/Engine/Source/Runtime/Engine/Private/LocalPlayer.cpp
@@ -21,6 +21,7 @@
 #include "GameFramework/OnlineSession.h"
 #include "GameFramework/PlayerInput.h"
 #include "GameFramework/GameMode.h"
+#include "GameFramework/PlayerState.h"
 
 DEFINE_LOG_CATEGORY(LogPlayerManagement);
 DEFINE_LOG_CATEGORY_STATIC(LogEngine, Log, All);
@@ -136,6 +137,16 @@ void FLocalPlayerContext::SetPlayerController( const APlayerController* InPlayer
 	LocalPlayer = CastChecked<ULocalPlayer>(InPlayerController->Player);
 }
 
+bool FLocalPlayerContext::IsFromLocalPlayer(const AActor* ActorToTest) const
+{
+	return (ActorToTest != nullptr) &&
+		IsValid() &&
+		(ActorToTest == GetPlayerController() ||
+		ActorToTest == GetPlayerState() ||
+		ActorToTest == GetPawn());
+}
+
+
 //////////////////////////////////////////////////////////////////////////
 // ULocalPlayer
 
