Move bShouldUpdatePhysicsVolume to a more reasonable category

diff --git a/Engine/Source/Runtime/Engine/Classes/Components/SceneComponent.h b/Engine/Source/Runtime/Engine/Classes/Components/SceneComponent.h
--- a/Engine/Source/Runtime/Engine/Classes/Components/SceneComponent.h
+++ b/Engine/Source/Runtime/Engine/Classes/Components/SceneComponent.h
@@ -130,7 +130,7 @@ public:
 	 * Whether or not the cached PhysicsVolume this component overlaps should be updated when the component is moved.
 	 * @see GetPhysicsVolume()
 	 */
-	UPROPERTY(EditAnywhere, BlueprintReadWrite, Category=PhysicsVolume)
+	UPROPERTY(EditAnywhere, BlueprintReadWrite, AdvancedDisplay, Category=Physics)
 	uint32 bShouldUpdatePhysicsVolume:1;
 
 	/** If true, a change in the bounds of this component will call trigger a streaming data rebuild */
