UE4 - Materials using WPO with period functions now properly compute their velocity

diff --git a/Engine/Shaders/BasePassForForwardShadingVertexShader.usf b/Engine/Shaders/BasePassForForwardShadingVertexShader.usf
--- a/Engine/Shaders/BasePassForForwardShadingVertexShader.usf
+++ b/Engine/Shaders/BasePassForForwardShadingVertexShader.usf
@@ -45,7 +45,7 @@ void Main(
 	float4 WorldPosition = WorldPositionExcludingWPO;
 
 	half3x3 TangentToLocal = VertexFactoryGetTangentToLocal(Input, VFIntermediates);	
-	FMaterialVertexParameters VertexParameters = GetMaterialVertexParameters(Input, VFIntermediates, WorldPosition.xyz, TangentToLocal, View.GameTime, View.RealTime);
+	FMaterialVertexParameters VertexParameters = GetMaterialVertexParameters(Input, VFIntermediates, WorldPosition.xyz, TangentToLocal);
 
 	half3 WorldPositionOffset = GetMaterialWorldPositionOffset(VertexParameters);
 	
diff --git a/Engine/Shaders/BasePassVertexShader.usf b/Engine/Shaders/BasePassVertexShader.usf
--- a/Engine/Shaders/BasePassVertexShader.usf
+++ b/Engine/Shaders/BasePassVertexShader.usf
@@ -18,7 +18,7 @@ void Main(
 	float4 WorldPosition = WorldPositionExcludingWPO;
 
 	float3x3 TangentToLocal = VertexFactoryGetTangentToLocal(Input, VFIntermediates);	
-	FMaterialVertexParameters VertexParameters = GetMaterialVertexParameters(Input, VFIntermediates, WorldPosition.xyz, TangentToLocal, View.GameTime, View.RealTime);
+	FMaterialVertexParameters VertexParameters = GetMaterialVertexParameters(Input, VFIntermediates, WorldPosition.xyz, TangentToLocal);
 
 	// Isolate instructions used for world position offset
 	// As these cause the optimizer to generate different position calculating instructions in each pass, resulting in self-z-fighting.
diff --git a/Engine/Shaders/DepthOnlyVertexShader.usf b/Engine/Shaders/DepthOnlyVertexShader.usf
--- a/Engine/Shaders/DepthOnlyVertexShader.usf
+++ b/Engine/Shaders/DepthOnlyVertexShader.usf
@@ -83,7 +83,7 @@ void Main(
 	float4 WorldPos = VertexFactoryGetWorldPosition(Input, VFIntermediates);
 
 	float3x3 TangentToLocal = VertexFactoryGetTangentToLocal(Input, VFIntermediates);
-	FMaterialVertexParameters VertexParameters = GetMaterialVertexParameters(Input, VFIntermediates, WorldPos.xyz, TangentToLocal, View.GameTime, View.RealTime);
+	FMaterialVertexParameters VertexParameters = GetMaterialVertexParameters(Input, VFIntermediates, WorldPos.xyz, TangentToLocal);
 
 	// Isolate instructions used for world position offset
 	// As these cause the optimizer to generate different position calculating instructions in each pass, resulting in self-z-fighting.
diff --git a/Engine/Shaders/DistortAccumulateVS.usf b/Engine/Shaders/DistortAccumulateVS.usf
--- a/Engine/Shaders/DistortAccumulateVS.usf
+++ b/Engine/Shaders/DistortAccumulateVS.usf
@@ -74,7 +74,7 @@ void Main(
 	float4	WorldPosition = VertexFactoryGetWorldPosition(Input, VFIntermediates);
 	float3x3 TangentToLocal = VertexFactoryGetTangentToLocal(Input, VFIntermediates);
 
-	FMaterialVertexParameters VertexParameters = GetMaterialVertexParameters(Input, VFIntermediates, WorldPosition.xyz, TangentToLocal, View.GameTime, View.RealTime);
+	FMaterialVertexParameters VertexParameters = GetMaterialVertexParameters(Input, VFIntermediates, WorldPosition.xyz, TangentToLocal);
 	// Isolate instructions used for world position offset on xbox 360, 
 	// As these cause the optimizer to generate different position calculating instructions in each pass, resulting in self-z-fighting.
 	// This is only necessary for shaders used in passes that have depth testing enabled.
diff --git a/Engine/Shaders/GpuSkinVertexFactory.usf b/Engine/Shaders/GpuSkinVertexFactory.usf
--- a/Engine/Shaders/GpuSkinVertexFactory.usf
+++ b/Engine/Shaders/GpuSkinVertexFactory.usf
@@ -203,14 +203,12 @@ struct FVertexFactoryIntermediates
 };
 
 /** Converts from vertex factory specific input to a FMaterialVertexParameters, which is used by vertex shader material inputs. */
-FMaterialVertexParameters GetMaterialVertexParameters( FVertexFactoryInput Input, FVertexFactoryIntermediates Intermediates, float3 WorldPosition, float3x3 TangentToLocal, float GameTime, float RealTime)
+FMaterialVertexParameters GetMaterialVertexParameters(FVertexFactoryInput Input, FVertexFactoryIntermediates Intermediates, float3 WorldPosition, float3x3 TangentToLocal)
 {
 	FMaterialVertexParameters Result = (FMaterialVertexParameters)0;
 	Result.WorldPosition = WorldPosition - View.PreViewTranslation.xyz;
 	Result.VertexColor = Intermediates.Color;
 	Result.TangentToWorld = mul(TangentToLocal, GetLocalToWorld3x3());
-	Result.GameTime = GameTime;
-	Result.RealTime = RealTime;
 
 #if NUM_MATERIAL_TEXCOORDS_VERTEX
 	for(int CoordinateIndex = 0; CoordinateIndex < NUM_MATERIAL_TEXCOORDS_VERTEX; CoordinateIndex++)
diff --git a/Engine/Shaders/HitProxyVertexShader.usf b/Engine/Shaders/HitProxyVertexShader.usf
--- a/Engine/Shaders/HitProxyVertexShader.usf
+++ b/Engine/Shaders/HitProxyVertexShader.usf
@@ -93,7 +93,7 @@ void Main(
 	float4	WorldPosition = VertexFactoryGetWorldPosition(Input, VFIntermediates);
 	float3x3 TangentToLocal = VertexFactoryGetTangentToLocal(Input, VFIntermediates);
 
-	FMaterialVertexParameters VertexParameters = GetMaterialVertexParameters(Input, VFIntermediates, WorldPosition.xyz, TangentToLocal, View.GameTime, View.RealTime);
+	FMaterialVertexParameters VertexParameters = GetMaterialVertexParameters(Input, VFIntermediates, WorldPosition.xyz, TangentToLocal);
 	WorldPosition.xyz += GetMaterialWorldPositionOffset(VertexParameters);
 
 #if USE_INSTANCING
diff --git a/Engine/Shaders/LandscapeVertexFactory.usf b/Engine/Shaders/LandscapeVertexFactory.usf
--- a/Engine/Shaders/LandscapeVertexFactory.usf
+++ b/Engine/Shaders/LandscapeVertexFactory.usf
@@ -359,14 +359,12 @@ FMaterialPixelParameters GetMaterialPixelParameters(FVertexFactoryInterpolantsVS
 }
 
 /** Converts from vertex factory specific input to a FMaterialVertexParameters, which is used by vertex shader material inputs. */
-FMaterialVertexParameters GetMaterialVertexParameters( FVertexFactoryInput Input, FVertexFactoryIntermediates Intermediates, float3 WorldPosition, float3x3 TangentToLocal, float GameTime, float RealTime)
+FMaterialVertexParameters GetMaterialVertexParameters(FVertexFactoryInput Input, FVertexFactoryIntermediates Intermediates, float3 WorldPosition, float3x3 TangentToLocal)
 {
 	FMaterialVertexParameters Result = (FMaterialVertexParameters)0;
 	Result.WorldPosition = WorldPosition - View.PreViewTranslation;
 	Result.VertexColor = float4(1,1,1,1);
 	Result.TangentToWorld = mul(TangentToLocal, GetLocalToWorld3x3()); 
-	Result.GameTime = GameTime;
-	Result.RealTime = RealTime;
 
 	FLandscapeTexCoords LandscapeTexCoords = GetLandscapeTexCoords(Input, Intermediates);
 
diff --git a/Engine/Shaders/LightMapDensityShader.usf b/Engine/Shaders/LightMapDensityShader.usf
--- a/Engine/Shaders/LightMapDensityShader.usf
+++ b/Engine/Shaders/LightMapDensityShader.usf
@@ -78,7 +78,7 @@ void MainVertexShader(
 	float4 WorldPosition = VertexFactoryGetWorldPosition(Input, VFIntermediates);
 	float3x3 TangentToLocal = VertexFactoryGetTangentToLocal(Input, VFIntermediates);
 
-	FMaterialVertexParameters VertexParameters = GetMaterialVertexParameters(Input, VFIntermediates, WorldPosition.xyz, TangentToLocal, View.GameTime, View.RealTime);
+	FMaterialVertexParameters VertexParameters = GetMaterialVertexParameters(Input, VFIntermediates, WorldPosition.xyz, TangentToLocal);
 	WorldPosition.xyz += GetMaterialWorldPositionOffset(VertexParameters);
 
 #if USING_TESSELLATION
diff --git a/Engine/Shaders/LocalVertexFactory.usf b/Engine/Shaders/LocalVertexFactory.usf
--- a/Engine/Shaders/LocalVertexFactory.usf
+++ b/Engine/Shaders/LocalVertexFactory.usf
@@ -250,13 +250,11 @@ half3x3 CalcTangentToWorldNoScale(in half3x3 TangentToLocal)
 }
 
 /** Converts from vertex factory specific input to a FMaterialVertexParameters, which is used by vertex shader material inputs. */
-FMaterialVertexParameters GetMaterialVertexParameters( FVertexFactoryInput Input, FVertexFactoryIntermediates Intermediates, float3 WorldPosition, half3x3 TangentToLocal, float GameTime, float RealTime)
+FMaterialVertexParameters GetMaterialVertexParameters(FVertexFactoryInput Input, FVertexFactoryIntermediates Intermediates, float3 WorldPosition, half3x3 TangentToLocal)
 {
 	FMaterialVertexParameters Result = (FMaterialVertexParameters)0;
 	Result.WorldPosition = WorldPosition - View.PreViewTranslation;
 	Result.VertexColor = Intermediates.Color;
-	Result.GameTime = GameTime;
-	Result.RealTime = RealTime;
 
 	// does not handle instancing!
 	Result.TangentToWorld = Intermediates.TangentToWorld;
diff --git a/Engine/Shaders/MaterialTemplate.usf b/Engine/Shaders/MaterialTemplate.usf
--- a/Engine/Shaders/MaterialTemplate.usf
+++ b/Engine/Shaders/MaterialTemplate.usf
@@ -247,9 +247,6 @@ struct FMaterialVertexParameters
 	#endif
 #endif
 
-	float GameTime;
-	float RealTime;
-
 	/** Per-particle properties. Only valid for particle vertex factories. */
 	FMaterialParticleParameters Particle;
 };
@@ -1181,6 +1178,11 @@ float3 GetMaterialWorldPositionOffset(FMaterialVertexParameters Parameters)
 %s;
 }
 
+float3 GetMaterialPreviousWorldPositionOffset(FMaterialVertexParameters Parameters)
+{
+%s;
+}
+
 half3 GetMaterialWorldDisplacement(FMaterialTessellationParameters Parameters)
 {
 %s;
diff --git a/Engine/Shaders/MeshParticleVertexFactory.usf b/Engine/Shaders/MeshParticleVertexFactory.usf
--- a/Engine/Shaders/MeshParticleVertexFactory.usf
+++ b/Engine/Shaders/MeshParticleVertexFactory.usf
@@ -215,13 +215,11 @@ FMaterialPixelParameters GetMaterialPixelParameters(FVertexFactoryInterpolantsVS
 }
 
 /** Converts from vertex factory specific input to a FMaterialVertexParameters, which is used by vertex shader material inputs. */
-FMaterialVertexParameters GetMaterialVertexParameters( FVertexFactoryInput Input, FVertexFactoryIntermediates Intermediates, float3 WorldPosition, float3x3 TangentToLocal, float GameTime, float RealTime)
+FMaterialVertexParameters GetMaterialVertexParameters(FVertexFactoryInput Input, FVertexFactoryIntermediates Intermediates, float3 WorldPosition, float3x3 TangentToLocal)
 {
 	FMaterialVertexParameters Result = (FMaterialVertexParameters)0;
 	Result.WorldPosition = WorldPosition - View.PreViewTranslation;
 	Result.VertexColor = Intermediates.VertexColor;
-	Result.GameTime = GameTime;
-	Result.RealTime = RealTime;
 
 	Result.Particle.PositionAndSize.xyz = Intermediates.ParticleWorldPosition.xyz;
 	Result.Particle.PositionAndSize.w = 1;
diff --git a/Engine/Shaders/ParticleBeamTrailVertexFactory.usf b/Engine/Shaders/ParticleBeamTrailVertexFactory.usf
--- a/Engine/Shaders/ParticleBeamTrailVertexFactory.usf
+++ b/Engine/Shaders/ParticleBeamTrailVertexFactory.usf
@@ -88,15 +88,13 @@ FMaterialPixelParameters GetMaterialPixelParameters(FVertexFactoryInterpolantsVS
 }
 
 /** Converts from vertex factory specific input to a FMaterialVertexParameters, which is used by vertex shader material inputs. */
-FMaterialVertexParameters GetMaterialVertexParameters( FVertexFactoryInput Input, FVertexFactoryIntermediates Intermediates, float3 WorldPosition, float3x3 TangentToLocal, float GameTime, float RealTime)
+FMaterialVertexParameters GetMaterialVertexParameters(FVertexFactoryInput Input, FVertexFactoryIntermediates Intermediates, float3 WorldPosition, float3x3 TangentToLocal)
 {
 	FMaterialVertexParameters Result = (FMaterialVertexParameters)0;
 	Result.WorldPosition = WorldPosition - View.PreViewTranslation.xyz;
 	Result.VertexColor = Input.Color;
 	Result.TangentToWorld = mul(TangentToLocal, GetLocalToWorld3x3()); 
 	Result.Particle.Color = Intermediates.Color;
-	Result.GameTime = GameTime;
-	Result.RealTime = RealTime;
 
 #if NUM_MATERIAL_TEXCOORDS_VERTEX
 	for(int CoordinateIndex = 0; CoordinateIndex < NUM_MATERIAL_TEXCOORDS_VERTEX; CoordinateIndex++)
diff --git a/Engine/Shaders/ParticleGPUSpriteVertexFactory.usf b/Engine/Shaders/ParticleGPUSpriteVertexFactory.usf
--- a/Engine/Shaders/ParticleGPUSpriteVertexFactory.usf
+++ b/Engine/Shaders/ParticleGPUSpriteVertexFactory.usf
@@ -407,14 +407,12 @@ FMaterialPixelParameters GetMaterialPixelParameters( FVertexFactoryInterpolantsV
  * @param TangentToLocal - The tangent basis for this vertex.
  * @returns per-vertex material parameters.
  */
-FMaterialVertexParameters GetMaterialVertexParameters( FVertexFactoryInput Input, FVertexFactoryIntermediates Intermediates, float3 WorldPosition, float3x3 TangentToLocal, float GameTime, float RealTime)
+FMaterialVertexParameters GetMaterialVertexParameters(FVertexFactoryInput Input, FVertexFactoryIntermediates Intermediates, float3 WorldPosition, float3x3 TangentToLocal)
 {
 	FMaterialVertexParameters Result = (FMaterialVertexParameters)0;
 	Result.WorldPosition = WorldPosition - View.PreViewTranslation.xyz;
 	Result.VertexColor = 1;
 	Result.TangentToWorld = mul(TangentToLocal, GetLocalToWorld3x3()); 
-	Result.GameTime = GameTime;
-	Result.RealTime = RealTime;
 
 #if USE_PARTICLE_TIME
 	Result.Particle.RelativeTime = Intermediates.RelativeTime;
diff --git a/Engine/Shaders/ParticleSpriteVertexFactory.usf b/Engine/Shaders/ParticleSpriteVertexFactory.usf
--- a/Engine/Shaders/ParticleSpriteVertexFactory.usf
+++ b/Engine/Shaders/ParticleSpriteVertexFactory.usf
@@ -227,7 +227,7 @@ FMaterialPixelParameters GetMaterialPixelParameters(FVertexFactoryInterpolantsVS
 }
 
 /** Converts from vertex factory specific input to a FMaterialVertexParameters, which is used by vertex shader material inputs. */
-FMaterialVertexParameters GetMaterialVertexParameters(FVertexFactoryInput Input, FVertexFactoryIntermediates Intermediates, float3 WorldPosition, float3x3 TangentToLocal, float GameTime, float RealTime)
+FMaterialVertexParameters GetMaterialVertexParameters(FVertexFactoryInput Input, FVertexFactoryIntermediates Intermediates, float3 WorldPosition, float3x3 TangentToLocal)
 {
 	FMaterialVertexParameters Result = (FMaterialVertexParameters)0;
 	Result.WorldPosition = WorldPosition - View.PreViewTranslation.xyz;
@@ -236,8 +236,6 @@ FMaterialVertexParameters GetMaterialVertexParameters(FVertexFactoryInput Input,
 	Result.Particle.MacroUV = SpriteVF.MacroUVParameters;
 	Result.Particle.Color = Intermediates.Color;
 	Result.Particle.MotionBlurFade = 1.0f;
-	Result.GameTime = GameTime;
-	Result.RealTime = RealTime;
 
 #if USE_DYNAMIC_PARAMETERS
 	Result.Particle.DynamicParameter = Intermediates.DynamicParameter;
diff --git a/Engine/Shaders/ShaderVersion.usf b/Engine/Shaders/ShaderVersion.usf
--- a/Engine/Shaders/ShaderVersion.usf
+++ b/Engine/Shaders/ShaderVersion.usf
@@ -2,4 +2,4 @@
 // Each time the console command is executed it generates a new GUID. As this file is included
 // in common.usf (which should be included in any shader) it allows to invalidate the shader DDC.
 // 
-// GUID = {CAC8D6F0-2E83-4C3D-A767-584D1E3DC80t}
+// GUID = {2490AEE3-37E5-4684-93D8-F170F910138B}
diff --git a/Engine/Shaders/ShadowDepthVertexShader.usf b/Engine/Shaders/ShadowDepthVertexShader.usf
--- a/Engine/Shaders/ShadowDepthVertexShader.usf
+++ b/Engine/Shaders/ShadowDepthVertexShader.usf
@@ -188,7 +188,7 @@ void Main(
 	float4 WorldPos = VertexFactoryGetWorldPosition(Input, VFIntermediates);
 	float3x3 TangentToLocal = VertexFactoryGetTangentToLocal(Input, VFIntermediates);
 
-	FMaterialVertexParameters VertexParameters = GetMaterialVertexParameters(Input, VFIntermediates, WorldPos.xyz, TangentToLocal, View.GameTime, View.RealTime);
+	FMaterialVertexParameters VertexParameters = GetMaterialVertexParameters(Input, VFIntermediates, WorldPos.xyz, TangentToLocal);
 	WorldPos.xyz += GetMaterialWorldPositionOffset(VertexParameters);
 
 #if USING_TESSELLATION
diff --git a/Engine/Shaders/TranslucentShadowDepthShaders.usf b/Engine/Shaders/TranslucentShadowDepthShaders.usf
--- a/Engine/Shaders/TranslucentShadowDepthShaders.usf
+++ b/Engine/Shaders/TranslucentShadowDepthShaders.usf
@@ -53,7 +53,7 @@ void MainVS(
 	float4 WorldPos = VertexFactoryGetWorldPosition(Input, VFIntermediates);
 	float3x3 TangentToLocal = VertexFactoryGetTangentToLocal(Input, VFIntermediates);
 
-	FMaterialVertexParameters VertexParameters = GetMaterialVertexParameters(Input, VFIntermediates, WorldPos.xyz, TangentToLocal, View.GameTime, View.RealTime);
+	FMaterialVertexParameters VertexParameters = GetMaterialVertexParameters(Input, VFIntermediates, WorldPos.xyz, TangentToLocal);
 	WorldPos.xyz += GetMaterialWorldPositionOffset(VertexParameters);
 
 	SetShadowDepthOutputs(
diff --git a/Engine/Shaders/VectorFieldVisualizationVertexFactory.usf b/Engine/Shaders/VectorFieldVisualizationVertexFactory.usf
--- a/Engine/Shaders/VectorFieldVisualizationVertexFactory.usf
+++ b/Engine/Shaders/VectorFieldVisualizationVertexFactory.usf
@@ -106,14 +106,12 @@ FMaterialPixelParameters GetMaterialPixelParameters( FVertexFactoryInterpolantsV
  * @param TangentToLocal - The tangent basis for this vertex.
  * @returns per-vertex material parameters.
  */
-FMaterialVertexParameters GetMaterialVertexParameters( FVertexFactoryInput Input, FVertexFactoryIntermediates Intermediates, float3 WorldPosition, float3x3 TangentToLocal, float GameTime, float RealTime)
+FMaterialVertexParameters GetMaterialVertexParameters(FVertexFactoryInput Input, FVertexFactoryIntermediates Intermediates, float3 WorldPosition, float3x3 TangentToLocal)
 {
 	FMaterialVertexParameters Result = (FMaterialVertexParameters)0;
 	Result.WorldPosition = WorldPosition - View.PreViewTranslation.xyz;
 	Result.VertexColor = float4(1,1,1,1);
 	Result.TangentToWorld = mul(TangentToLocal, GetLocalToWorld3x3()); 
-	Result.GameTime = GameTime;
-	Result.RealTime = RealTime;
 	return Result;
 }
 
diff --git a/Engine/Shaders/VelocityShader.usf b/Engine/Shaders/VelocityShader.usf
--- a/Engine/Shaders/VelocityShader.usf
+++ b/Engine/Shaders/VelocityShader.usf
@@ -126,7 +126,7 @@ void MainVertexShader(
 	float4 TranslatedWorldPosition = VertexFactoryGetWorldPosition( Input, VFIntermediates );
 	float3x3 TangentToLocal = VertexFactoryGetTangentToLocal( Input, VFIntermediates );
 
-	FMaterialVertexParameters VertexParameters = GetMaterialVertexParameters(Input, VFIntermediates, TranslatedWorldPosition.xyz, TangentToLocal, View.GameTime, View.RealTime);
+	FMaterialVertexParameters VertexParameters = GetMaterialVertexParameters(Input, VFIntermediates, TranslatedWorldPosition.xyz, TangentToLocal);
 	// Isolate instructions used for world position offset on xbox 360, 
 	// As these cause the optimizer to generate different position calculating instructions in each pass, resulting in self-z-fighting.
 	// This is only necessary for shaders used in passes that have depth testing enabled.
@@ -142,8 +142,8 @@ void MainVertexShader(
 #endif // #if USING_TESSELLATION
 	float4 PrevTranslatedWorldPosition = VertexFactoryGetPreviousWorldPosition( Input, VFIntermediates );	
 
-	VertexParameters = GetMaterialVertexParameters(Input, VFIntermediates, PrevTranslatedWorldPosition.xyz, TangentToLocal, View.PrevFrameGameTime, View.PrevFrameRealTime);
-	PrevTranslatedWorldPosition.xyz += GetMaterialWorldPositionOffset(VertexParameters);
+	VertexParameters = GetMaterialVertexParameters(Input, VFIntermediates, PrevTranslatedWorldPosition.xyz, TangentToLocal);
+	PrevTranslatedWorldPosition.xyz += GetMaterialPreviousWorldPositionOffset(VertexParameters);
 	
 #if USING_TESSELLATION
 	Output.Position = TranslatedWorldPosition;
diff --git a/Engine/Source/Editor/MaterialEditor/Private/MaterialEditor.cpp b/Engine/Source/Editor/MaterialEditor/Private/MaterialEditor.cpp
--- a/Engine/Source/Editor/MaterialEditor/Private/MaterialEditor.cpp
+++ b/Engine/Source/Editor/MaterialEditor/Private/MaterialEditor.cpp
@@ -124,10 +124,10 @@ bool FMatExpressionPreview::ShouldCache(EShaderPlatform Platform, const FShaderT
 	return false;
 }
 
-int32 FMatExpressionPreview::CompilePropertyAndSetMaterialProperty(EMaterialProperty Property, FMaterialCompiler* Compiler, EShaderFrequency OverrideShaderFrequency) const
+int32 FMatExpressionPreview::CompilePropertyAndSetMaterialProperty(EMaterialProperty Property, FMaterialCompiler* Compiler, EShaderFrequency OverrideShaderFrequency, bool bUsePreviousFrameTime) const
 {
 	// needs to be called in this function!!
-	Compiler->SetMaterialProperty(Property, OverrideShaderFrequency);
+	Compiler->SetMaterialProperty(Property, OverrideShaderFrequency, bUsePreviousFrameTime);
 
 	int32 Ret = INDEX_NONE;
 
@@ -138,7 +138,7 @@ int32 FMatExpressionPreview::CompilePropertyAndSetMaterialProperty(EMaterialProp
 		// Get back into gamma corrected space, as DrawTile does not do this adjustment.
 		Ret = Compiler->Power(Compiler->Max(Expression->CompilePreview(Compiler, OutputIndex, -1), Compiler->Constant(0)), Compiler->Constant(1.f / 2.2f));
 	}
-	else if ( Property == MP_WorldPositionOffset)
+	else if (Property == MP_WorldPositionOffset)
 	{
 		//set to 0 to prevent off by 1 pixel errors
 		Ret = Compiler->Constant(0.0f);
diff --git a/Engine/Source/Editor/MaterialEditor/Private/MaterialEditor.h b/Engine/Source/Editor/MaterialEditor/Private/MaterialEditor.h
--- a/Engine/Source/Editor/MaterialEditor/Private/MaterialEditor.h
+++ b/Engine/Source/Editor/MaterialEditor/Private/MaterialEditor.h
@@ -108,7 +108,7 @@ public:
 	
 	// Material properties.
 	/** Entry point for compiling a specific material property.  This must call SetMaterialProperty. */
-	virtual int32 CompilePropertyAndSetMaterialProperty(EMaterialProperty Property, FMaterialCompiler* Compiler, EShaderFrequency OverrideShaderFrequency) const;
+	virtual int32 CompilePropertyAndSetMaterialProperty(EMaterialProperty Property, FMaterialCompiler* Compiler, EShaderFrequency OverrideShaderFrequency, bool bUsePreviousFrameTime) const override;
 
 	virtual int32 GetMaterialDomain() const override { return MD_Surface; }
 	virtual FString GetMaterialUsageDescription() const { return FString::Printf(TEXT("FMatExpressionPreview %s"), Expression.IsValid() ? *Expression->GetName() : TEXT("NULL")); }
diff --git a/Engine/Source/Editor/UnrealEd/Private/EditorExporters.cpp b/Engine/Source/Editor/UnrealEd/Private/EditorExporters.cpp
--- a/Engine/Source/Editor/UnrealEd/Private/EditorExporters.cpp
+++ b/Engine/Source/Editor/UnrealEd/Private/EditorExporters.cpp
@@ -802,10 +802,10 @@ public:
 
 	// Material properties.
 	/** Entry point for compiling a specific material property.  This must call SetMaterialProperty. */
-	virtual int32 CompilePropertyAndSetMaterialProperty(EMaterialProperty Property, FMaterialCompiler* Compiler, EShaderFrequency OverrideShaderFrequency) const
+	virtual int32 CompilePropertyAndSetMaterialProperty(EMaterialProperty Property, FMaterialCompiler* Compiler, EShaderFrequency OverrideShaderFrequency, bool bUsePreviousFrameTime) const
 	{
 		// needs to be called in this function!!
-		Compiler->SetMaterialProperty(Property, OverrideShaderFrequency);
+		Compiler->SetMaterialProperty(Property, OverrideShaderFrequency, bUsePreviousFrameTime);
 
 		int32 Ret = CompilePropertyAndSetMaterialPropertyWithoutCast(Property, Compiler);
 
@@ -905,7 +905,7 @@ public:
 	
 			return Compiler->Constant(0.0f);
 		}
-		else if( Property == MP_WorldPositionOffset)
+		else if (Property == MP_WorldPositionOffset)
 		{
 			//This property MUST return 0 as a default or during the process of rendering textures out for lightmass to use, pixels will be off by 1.
 			return Compiler->Constant(0.0f);
diff --git a/Engine/Source/Editor/UnrealEd/Private/Lightmass/LightmassRender.cpp b/Engine/Source/Editor/UnrealEd/Private/Lightmass/LightmassRender.cpp
--- a/Engine/Source/Editor/UnrealEd/Private/Lightmass/LightmassRender.cpp
+++ b/Engine/Source/Editor/UnrealEd/Private/Lightmass/LightmassRender.cpp
@@ -266,10 +266,10 @@ public:
 	// Material properties.
 
 	/** Entry point for compiling a specific material property.  This must call SetMaterialProperty. */
-	virtual int32 CompilePropertyAndSetMaterialProperty(EMaterialProperty Property, FMaterialCompiler* Compiler, EShaderFrequency OverrideShaderFrequency) const
+	virtual int32 CompilePropertyAndSetMaterialProperty(EMaterialProperty Property, FMaterialCompiler* Compiler, EShaderFrequency OverrideShaderFrequency, bool bUsePreviousFrameTime) const override
 	{
 		// needs to be called in this function!!
-		Compiler->SetMaterialProperty(Property, OverrideShaderFrequency);
+		Compiler->SetMaterialProperty(Property, OverrideShaderFrequency, bUsePreviousFrameTime);
 
 		int32 Ret = CompilePropertyAndSetMaterialPropertyWithoutCast(Property, Compiler);
 
diff --git a/Engine/Source/Runtime/Engine/Private/Materials/HLSLMaterialTranslator.h b/Engine/Source/Runtime/Engine/Private/Materials/HLSLMaterialTranslator.h
--- a/Engine/Source/Runtime/Engine/Private/Materials/HLSLMaterialTranslator.h
+++ b/Engine/Source/Runtime/Engine/Private/Materials/HLSLMaterialTranslator.h
@@ -191,6 +191,8 @@ protected:
 	/** true if the material reads particle color in the pixel shader. */
 	uint32 bUsesParticleColor : 1;
 	uint32 bUsesTransformVector : 1;
+	// True if the current property requires last frame's time
+	uint32 bUsesPreviousFrameTime : 1;
 	/** Tracks the number of texture coordinates used by this material. */
 	uint32 NumUserTexCoords;
 	/** Tracks the number of texture coordinates used by the vertex shader in this material. */
@@ -232,6 +234,7 @@ public:
 	,	bUsesVertexColor(false)
 	,	bUsesParticleColor(false)
 	,	bUsesTransformVector(false)
+	,	bUsesPreviousFrameTime(false)
 	,	NumUserTexCoords(0)
 	,	NumUserVertexTexCoords(0)
 	{}
@@ -321,6 +324,12 @@ public:
 				Chunk[CompiledMP_EmissiveColorCS]		= Material->CompilePropertyAndSetMaterialProperty(MP_EmissiveColor,this, SF_Compute);
 			}
 
+			if (Chunk[MP_WorldPositionOffset] != -1)
+			{
+				// Only calculate previous WPO if there is a current WPO
+				Chunk[CompiledMP_PrevWorldPositionOffset] = Material->CompilePropertyAndSetMaterialProperty(MP_WorldPositionOffset, this, SF_Vertex, true);
+			}
+
 			for (uint32 CustomUVIndex = MP_CustomizedUVs0; CustomUVIndex <= MP_CustomizedUVs7; CustomUVIndex++)
 			{
 				// Only compile custom UV inputs for UV channels requested by the pixel shader inputs
@@ -342,7 +351,6 @@ public:
 			{
 				int32 WPOFreq = (int32)GetMaterialPropertyShaderFrequency(MP_WorldPositionOffset);
 				FShaderCodeChunk& WPOChunk = CodeChunks[MP_WorldPositionOffset][WPOFreq][Chunk[MP_WorldPositionOffset]];
-				FMaterialRenderContext DummyContext( NULL, *Material, NULL );
 
 				// Determine whether the world position offset is used. 
 				// If the output chunk has a uniform expression, it is constant, and GetNumberValue returns the default property value then WPO isn't used.
@@ -350,6 +358,7 @@ public:
 				if( WPOChunk.UniformExpression && WPOChunk.UniformExpression->IsConstant() )
 				{
 					FLinearColor WPOValue;
+					FMaterialRenderContext DummyContext(nullptr, *Material, nullptr);
 					WPOChunk.UniformExpression->GetNumberValue(DummyContext, WPOValue);
 					if (FVector(WPOValue) == FVector::ZeroVector)
 					{
@@ -440,6 +449,16 @@ public:
 				{
 					ResourcesString += FString::Printf(TEXT("float4 UE_Material_PerFrameVectorExpression%u;"), Index) + "\r\n\r\n";
 				}
+
+				for (int32 Index = 0, Num = MaterialCompilationOutput.UniformExpressionSet.PerFramePrevUniformScalarExpressions.Num(); Index < Num; ++Index)
+				{
+					ResourcesString += FString::Printf(TEXT("float UE_Material_PerFramePrevScalarExpression%u;"), Index) + "\r\n\r\n";
+				}
+
+				for (int32 Index = 0, Num = MaterialCompilationOutput.UniformExpressionSet.PerFramePrevUniformVectorExpressions.Num(); Index < Num; ++Index)
+				{
+					ResourcesString += FString::Printf(TEXT("float4 UE_Material_PerFramePrevVectorExpression%u;"), Index) + "\r\n\r\n";
+				}
 			}
 
 			for(uint32 PropertyId = 0; PropertyId < MP_MAX; ++PropertyId)
@@ -456,11 +475,21 @@ public:
 					TranslatedCodeChunks[PropertyId]);
 			}
 
-			if (bCompileForComputeShader)
+			for(uint32 PropertyId = MP_MAX; PropertyId < CompiledMP_MAX; ++PropertyId)
 			{
-				for(uint32 PropertyId = MP_MAX; PropertyId < CompiledMP_MAX; ++PropertyId)
+				switch(PropertyId)
 				{
-					GetFixedParameterCode(Chunk[PropertyId], MP_EmissiveColor, SF_Compute, TranslatedCodeChunkDefinitions[PropertyId], TranslatedCodeChunks[PropertyId]);
+				case CompiledMP_EmissiveColorCS:
+					if (bCompileForComputeShader)
+					{
+						GetFixedParameterCode(Chunk[PropertyId], MP_EmissiveColor, SF_Compute, TranslatedCodeChunkDefinitions[PropertyId], TranslatedCodeChunks[PropertyId]);
+					}
+					break;
+				case CompiledMP_PrevWorldPositionOffset:
+					GetFixedParameterCode(Chunk[PropertyId], MP_WorldPositionOffset, SF_Vertex, TranslatedCodeChunkDefinitions[PropertyId], TranslatedCodeChunks[PropertyId]);
+					break;
+				default: check(0);
+					break;
 				}
 			}
 
@@ -623,6 +652,7 @@ public:
 		LazyPrintf.PushParam(*GenerateFunctionCode(MP_Opacity));
 		LazyPrintf.PushParam(*GenerateFunctionCode(MP_OpacityMask));
 		LazyPrintf.PushParam(*GenerateFunctionCode(MP_WorldPositionOffset));
+		LazyPrintf.PushParam(*GenerateFunctionCode(CompiledMP_PrevWorldPositionOffset));
 		LazyPrintf.PushParam(*GenerateFunctionCode(MP_WorldDisplacement));
 		LazyPrintf.PushParam(*GenerateFunctionCode(MP_TessellationMultiplier));
 		LazyPrintf.PushParam(*GenerateFunctionCode(MP_SubsurfaceColor));
@@ -1004,8 +1034,16 @@ protected:
 		{
 			if (CodeChunk.UniformExpression->IsChangingPerFrame())
 			{
-				const int32 ScalarInputIndex = MaterialCompilationOutput.UniformExpressionSet.PerFrameUniformScalarExpressions.AddUnique(CodeChunk.UniformExpression);
-				FCString::Sprintf(FormattedCode, TEXT("UE_Material_PerFrameScalarExpression%u"), ScalarInputIndex);
+				if (bUsesPreviousFrameTime)
+				{
+					const int32 ScalarInputIndex = MaterialCompilationOutput.UniformExpressionSet.PerFramePrevUniformScalarExpressions.AddUnique(CodeChunk.UniformExpression);
+					FCString::Sprintf(FormattedCode, TEXT("UE_Material_PerFramePrevScalarExpression%u"), ScalarInputIndex);
+				}
+				else
+				{
+					const int32 ScalarInputIndex = MaterialCompilationOutput.UniformExpressionSet.PerFrameUniformScalarExpressions.AddUnique(CodeChunk.UniformExpression);
+					FCString::Sprintf(FormattedCode, TEXT("UE_Material_PerFrameScalarExpression%u"), ScalarInputIndex);
+				}
 			}
 			else
 			{
@@ -1029,8 +1067,16 @@ protected:
 
 			if (CodeChunk.UniformExpression->IsChangingPerFrame())
 			{
-				const int32 VectorInputIndex = MaterialCompilationOutput.UniformExpressionSet.PerFrameUniformVectorExpressions.AddUnique(CodeChunk.UniformExpression);
-				FCString::Sprintf(FormattedCode, TEXT("UE_Material_PerFrameVectorExpression%u%s"), VectorInputIndex, Mask);
+				if (bUsesPreviousFrameTime)
+				{
+					const int32 VectorInputIndex = MaterialCompilationOutput.UniformExpressionSet.PerFramePrevUniformVectorExpressions.AddUnique(CodeChunk.UniformExpression);
+					FCString::Sprintf(FormattedCode, TEXT("UE_Material_PerFramePrevVectorExpression%u%s"), VectorInputIndex, Mask);
+				}
+				else
+				{
+					const int32 VectorInputIndex = MaterialCompilationOutput.UniformExpressionSet.PerFrameUniformVectorExpressions.AddUnique(CodeChunk.UniformExpression);
+					FCString::Sprintf(FormattedCode, TEXT("UE_Material_PerFrameVectorExpression%u%s"), VectorInputIndex, Mask);
+				}
 			}
 			else
 			{
@@ -1167,7 +1213,7 @@ protected:
 	 * This affects the internal state of the compiler and the results of all functions except GetFixedParameterCode.
 	 * @param OverrideShaderFrequency SF_NumFrequencies to not override
 	 */
-	virtual void SetMaterialProperty(EMaterialProperty InProperty, EShaderFrequency OverrideShaderFrequency = SF_NumFrequencies)
+	virtual void SetMaterialProperty(EMaterialProperty InProperty, EShaderFrequency OverrideShaderFrequency = SF_NumFrequencies, bool bUsePreviousFrameTime = false)
 	{
 		MaterialProperty = InProperty;
 
@@ -1179,6 +1225,8 @@ protected:
 		{
 			ShaderFrequency = GetMaterialPropertyShaderFrequency(InProperty);
 		}
+
+		bUsesPreviousFrameTime = bUsePreviousFrameTime;
 	}
 	virtual EShaderFrequency GetCurrentShaderFrequency() const
 	{
@@ -1515,14 +1563,12 @@ protected:
 	{
 		if (!bPeriodic)
 		{
-			if (ShaderFrequency == SF_Vertex)
+			if (bUsesPreviousFrameTime)
 			{
-				return AddInlinedCodeChunk(MCT_Float, TEXT("Parameters.GameTime"));
-			}
-			else
-			{
-				return AddInlinedCodeChunk(MCT_Float, TEXT("View.GameTime"));
+				return AddInlinedCodeChunk(MCT_Float, TEXT("View.PrevFrameGameTime"));
 			}
+
+			return AddInlinedCodeChunk(MCT_Float, TEXT("View.GameTime"));
 		}
 		else if (Period == 0.0f)
 		{
@@ -1542,14 +1588,12 @@ protected:
 	{
 		if (!bPeriodic)
 		{
-			if (ShaderFrequency == SF_Vertex)
+			if (bUsesPreviousFrameTime)
 			{
-				return AddInlinedCodeChunk(MCT_Float, TEXT("Parameters.RealTime"));
-			}
-			else
-			{
-				return AddInlinedCodeChunk(MCT_Float, TEXT("View.RealTime"));
+				return AddInlinedCodeChunk(MCT_Float, TEXT("View.PrevFrameRealTime"));
 			}
+
+			return AddInlinedCodeChunk(MCT_Float, TEXT("View.RealTime"));
 		}
 		else if (Period == 0.0f)
 		{
diff --git a/Engine/Source/Runtime/Engine/Private/Materials/Material.cpp b/Engine/Source/Runtime/Engine/Private/Materials/Material.cpp
--- a/Engine/Source/Runtime/Engine/Private/Materials/Material.cpp
+++ b/Engine/Source/Runtime/Engine/Private/Materials/Material.cpp
@@ -67,11 +67,11 @@ FMaterialResource::FMaterialResource()
 {
 }
 
-int32 FMaterialResource::CompilePropertyAndSetMaterialProperty(EMaterialProperty Property, FMaterialCompiler* Compiler, EShaderFrequency OverrideShaderFrequency) const
+int32 FMaterialResource::CompilePropertyAndSetMaterialProperty(EMaterialProperty Property, FMaterialCompiler* Compiler, EShaderFrequency OverrideShaderFrequency, bool bUsePreviousFrameTime) const
 {
 	// needs to be called in this function!!
 	// sets CurrentShaderFrequency
-	Compiler->SetMaterialProperty(Property, OverrideShaderFrequency);
+	Compiler->SetMaterialProperty(Property, OverrideShaderFrequency, bUsePreviousFrameTime);
 
 	EShaderFrequency ShaderFrequency = Compiler->GetCurrentShaderFrequency();
 
@@ -3057,7 +3057,7 @@ FExpressionInput* UMaterial::GetExpressionInputForProperty(EMaterialProperty InP
 		return &CustomizedUVs[InProperty - MP_CustomizedUVs0];
 	}
 
-	return NULL;
+	return nullptr;
 }
 
 
diff --git a/Engine/Source/Runtime/Engine/Private/Materials/MaterialUniformExpressions.cpp b/Engine/Source/Runtime/Engine/Private/Materials/MaterialUniformExpressions.cpp
--- a/Engine/Source/Runtime/Engine/Private/Materials/MaterialUniformExpressions.cpp
+++ b/Engine/Source/Runtime/Engine/Private/Materials/MaterialUniformExpressions.cpp
@@ -93,6 +93,8 @@ void FUniformExpressionSet::Serialize(FArchive& Ar)
 
 	Ar << PerFrameUniformScalarExpressions;
 	Ar << PerFrameUniformVectorExpressions;
+	Ar << PerFramePrevUniformScalarExpressions;
+	Ar << PerFramePrevUniformVectorExpressions;
 
 	// Recreate the uniform buffer struct after loading.
 	if(Ar.IsLoading())
@@ -109,6 +111,8 @@ bool FUniformExpressionSet::IsEmpty() const
 		&& UniformCubeTextureExpressions.Num() == 0
 		&& PerFrameUniformScalarExpressions.Num() == 0
 		&& PerFrameUniformVectorExpressions.Num() == 0
+		&& PerFramePrevUniformScalarExpressions.Num() == 0
+		&& PerFramePrevUniformVectorExpressions.Num() == 0
 		&& ParameterCollections.Num() == 0;
 }
 
@@ -120,6 +124,8 @@ bool FUniformExpressionSet::operator==(const FUniformExpressionSet& ReferenceSet
 	|| UniformCubeTextureExpressions.Num() != ReferenceSet.UniformCubeTextureExpressions.Num()
 	|| PerFrameUniformScalarExpressions.Num() != ReferenceSet.PerFrameUniformScalarExpressions.Num()
 	|| PerFrameUniformVectorExpressions.Num() != ReferenceSet.PerFrameUniformVectorExpressions.Num()
+	|| PerFramePrevUniformScalarExpressions.Num() != ReferenceSet.PerFramePrevUniformScalarExpressions.Num()
+	|| PerFramePrevUniformVectorExpressions.Num() != ReferenceSet.PerFramePrevUniformVectorExpressions.Num()
 	|| ParameterCollections.Num() != ReferenceSet.ParameterCollections.Num())
 	{
 		return false;
@@ -173,6 +179,22 @@ bool FUniformExpressionSet::operator==(const FUniformExpressionSet& ReferenceSet
 		}
 	}
 
+	for (int32 i = 0; i < PerFramePrevUniformScalarExpressions.Num(); i++)
+	{
+		if (!PerFramePrevUniformScalarExpressions[i]->IsIdentical(ReferenceSet.PerFramePrevUniformScalarExpressions[i]))
+		{
+			return false;
+		}
+	}
+
+	for (int32 i = 0; i < PerFramePrevUniformVectorExpressions.Num(); i++)
+	{
+		if (!PerFramePrevUniformVectorExpressions[i]->IsIdentical(ReferenceSet.PerFramePrevUniformVectorExpressions[i]))
+		{
+			return false;
+		}
+	}
+
 	for (int32 i = 0; i < ParameterCollections.Num(); i++)
 	{
 		if (ParameterCollections[i] != ReferenceSet.ParameterCollections[i])
diff --git a/Engine/Source/Runtime/Engine/Private/ShaderDerivedDataVersion.h b/Engine/Source/Runtime/Engine/Private/ShaderDerivedDataVersion.h
--- a/Engine/Source/Runtime/Engine/Private/ShaderDerivedDataVersion.h
+++ b/Engine/Source/Runtime/Engine/Private/ShaderDerivedDataVersion.h
@@ -10,4 +10,4 @@
 // guid as version
 
 #define GLOBALSHADERMAP_DERIVEDDATA_VER			TEXT("af1606fcec114373844d2b6cf9df71c8")
-#define MATERIALSHADERMAP_DERIVEDDATA_VER		TEXT("85dd00daf10d45c79d4ed0306af2628b")
+#define MATERIALSHADERMAP_DERIVEDDATA_VER		TEXT("bfbd0e1a93dd4481a991041cedfae4e9")
diff --git a/Engine/Source/Runtime/Engine/Public/MaterialCompiler.h b/Engine/Source/Runtime/Engine/Public/MaterialCompiler.h
--- a/Engine/Source/Runtime/Engine/Public/MaterialCompiler.h
+++ b/Engine/Source/Runtime/Engine/Public/MaterialCompiler.h
@@ -20,7 +20,7 @@ class FMaterialCompiler
 public:
 	// sets internal state CurrentShaderFrequency 
 	// @param OverrideShaderFrequency SF_NumFrequencies to not override
-	virtual void SetMaterialProperty(EMaterialProperty InProperty, EShaderFrequency OverrideShaderFrequency = SF_NumFrequencies) = 0;
+	virtual void SetMaterialProperty(EMaterialProperty InProperty, EShaderFrequency OverrideShaderFrequency = SF_NumFrequencies, bool bUsePreviousFrameTime = false) = 0;
 	// gets value stored by SetMaterialProperty()
 	virtual EShaderFrequency GetCurrentShaderFrequency() const = 0;
 	//
@@ -191,7 +191,7 @@ public:
 
 	// Simple pass through all other material operations unmodified.
 
-	virtual void SetMaterialProperty(EMaterialProperty InProperty, EShaderFrequency OverrideShaderFrequency) override { Compiler->SetMaterialProperty(InProperty, OverrideShaderFrequency); }
+	virtual void SetMaterialProperty(EMaterialProperty InProperty, EShaderFrequency OverrideShaderFrequency, bool bUsePreviousFrameTime) override { Compiler->SetMaterialProperty(InProperty, OverrideShaderFrequency, bUsePreviousFrameTime); }
 	virtual EShaderFrequency GetCurrentShaderFrequency() const	{ return Compiler->GetCurrentShaderFrequency(); }
 	virtual int32 Error(const TCHAR* Text) override { return Compiler->Error(Text); }
 
diff --git a/Engine/Source/Runtime/Engine/Public/MaterialShared.h b/Engine/Source/Runtime/Engine/Public/MaterialShared.h
--- a/Engine/Source/Runtime/Engine/Public/MaterialShared.h
+++ b/Engine/Source/Runtime/Engine/Public/MaterialShared.h
@@ -240,6 +240,8 @@ public:
 			+ UniformCubeTextureExpressions.GetAllocatedSize()
 			+ PerFrameUniformScalarExpressions.GetAllocatedSize()
 			+ PerFrameUniformVectorExpressions.GetAllocatedSize()
+			+ PerFramePrevUniformScalarExpressions.GetAllocatedSize()
+			+ PerFramePrevUniformVectorExpressions.GetAllocatedSize()
 			+ ParameterCollections.GetAllocatedSize()
 			+ (UniformBufferStruct ? (sizeof(FUniformBufferStruct) + UniformBufferStruct->GetMembers().GetAllocatedSize()) : 0);
 	}
@@ -252,6 +254,8 @@ protected:
 	TArray<TRefCountPtr<FMaterialUniformExpressionTexture> > UniformCubeTextureExpressions;
 	TArray<TRefCountPtr<FMaterialUniformExpression> > PerFrameUniformScalarExpressions;
 	TArray<TRefCountPtr<FMaterialUniformExpression> > PerFrameUniformVectorExpressions;
+	TArray<TRefCountPtr<FMaterialUniformExpression> > PerFramePrevUniformScalarExpressions;
+	TArray<TRefCountPtr<FMaterialUniformExpression> > PerFramePrevUniformVectorExpressions;
 
 	/** Ids of parameter collections referenced by the material that was translated. */
 	TArray<FGuid> ParameterCollections;
@@ -747,6 +751,7 @@ private:
 enum ECompiledMaterialProperty
 {
 	CompiledMP_EmissiveColorCS = MP_MAX,
+	CompiledMP_PrevWorldPositionOffset,
 	CompiledMP_MAX
 };
 
@@ -1094,7 +1099,7 @@ protected:
 	 * @param OverrideShaderFrequency SF_NumFrequencies to not override
 	 * @return cases to the proper type e.g. Compiler->ForceCast(Ret, GetMaterialPropertyType(Property));
 	 */
-	virtual int32 CompilePropertyAndSetMaterialProperty(EMaterialProperty Property, class FMaterialCompiler* Compiler, EShaderFrequency OverrideShaderFrequency = SF_NumFrequencies) const = 0;
+	virtual int32 CompilePropertyAndSetMaterialProperty(EMaterialProperty Property, class FMaterialCompiler* Compiler, EShaderFrequency OverrideShaderFrequency = SF_NumFrequencies, bool bUsePreviousFrameTime = false) const = 0;
 
 	/** Returns the index to the Expression in the Expressions array, or -1 if not found. */
 	int32 FindExpression(const TArray<TRefCountPtr<FMaterialUniformExpressionTexture> >&Expressions, const FMaterialUniformExpressionTexture &Expression);
@@ -1543,7 +1548,7 @@ protected:
 	UMaterialInstance* MaterialInstance;
 
 	/** Entry point for compiling a specific material property.  This must call SetMaterialProperty. */
-	ENGINE_API virtual int32 CompilePropertyAndSetMaterialProperty(EMaterialProperty Property, class FMaterialCompiler* Compiler, EShaderFrequency OverrideShaderFrequency) const;
+	ENGINE_API virtual int32 CompilePropertyAndSetMaterialProperty(EMaterialProperty Property, class FMaterialCompiler* Compiler, EShaderFrequency OverrideShaderFrequency, bool bUsePreviousFrameTime) const override;
 	ENGINE_API virtual bool HasVertexPositionOffsetConnected() const;
 	ENGINE_API virtual bool HasMaterialAttributesConnected() const;
 	/** Useful for debugging. */
diff --git a/Engine/Source/Runtime/Renderer/Private/ShaderBaseClasses.cpp b/Engine/Source/Runtime/Renderer/Private/ShaderBaseClasses.cpp
--- a/Engine/Source/Runtime/Renderer/Private/ShaderBaseClasses.cpp
+++ b/Engine/Source/Runtime/Renderer/Private/ShaderBaseClasses.cpp
@@ -48,6 +48,20 @@ FMaterialShader::FMaterialShader(const FMaterialShaderType::CompiledShaderInitia
 		PerFrameVectorExpressions.Add(Parameter);
 	}
 
+	for (int32 Index = 0; Index < Initializer.UniformExpressionSet.PerFramePrevUniformScalarExpressions.Num(); Index++)
+	{
+		FShaderParameter Parameter;
+		Parameter.Bind(Initializer.ParameterMap, *FString::Printf(TEXT("UE_Material_PerFramePrevScalarExpression%u"), Index));
+		PerFramePrevScalarExpressions.Add(Parameter);
+	}
+
+	for (int32 Index = 0; Index < Initializer.UniformExpressionSet.PerFramePrevUniformVectorExpressions.Num(); Index++)
+	{
+		FShaderParameter Parameter;
+		Parameter.Bind(Initializer.ParameterMap, *FString::Printf(TEXT("UE_Material_PerFramePrevVectorExpression%u"), Index));
+		PerFramePrevVectorExpressions.Add(Parameter);
+	}
+
 	DeferredParameters.Bind(Initializer.ParameterMap);
 	LightAttenuation.Bind(Initializer.ParameterMap, TEXT("LightAttenuationTexture"));
 	LightAttenuationSampler.Bind(Initializer.ParameterMap, TEXT("LightAttenuationTextureSampler"));
@@ -200,6 +214,37 @@ void FMaterialShader::SetParameters(
 					SetShaderValue(RHICmdList, ShaderRHI, Parameter, TempValue);
 				}
 			}
+
+			// Now previous frame's expressions
+			const int32 NumPrevScalarExpressions = PerFramePrevScalarExpressions.Num();
+			const int32 NumPrevVectorExpressions = PerFramePrevVectorExpressions.Num();
+			if (NumPrevScalarExpressions > 0 || NumPrevVectorExpressions > 0)
+			{
+				MaterialRenderContext.Time = View.Family->CurrentWorldTime - View.Family->DeltaWorldTime;
+				MaterialRenderContext.RealTime = View.Family->CurrentRealTime - View.Family->DeltaWorldTime;
+
+				for (int32 Index = 0; Index < NumPrevScalarExpressions; ++Index)
+				{
+					auto& Parameter = PerFramePrevScalarExpressions[Index];
+					if (Parameter.IsBound())
+					{
+						FLinearColor TempValue;
+						MaterialUniformExpressionSet.PerFramePrevUniformScalarExpressions[Index]->GetNumberValue(MaterialRenderContext, TempValue);
+						SetShaderValue(RHICmdList, ShaderRHI, Parameter, TempValue.R);
+					}
+				}
+
+				for (int32 Index = 0; Index < NumPrevVectorExpressions; ++Index)
+				{
+					auto& Parameter = PerFramePrevVectorExpressions[Index];
+					if (Parameter.IsBound())
+					{
+						FLinearColor TempValue;
+						MaterialUniformExpressionSet.PerFramePrevUniformVectorExpressions[Index]->GetNumberValue(MaterialRenderContext, TempValue);
+						SetShaderValue(RHICmdList, ShaderRHI, Parameter, TempValue);
+					}
+				}
+			}
 		}
 	}
 
@@ -312,6 +357,8 @@ bool FMaterialShader::Serialize(FArchive& Ar)
 
 	Ar << PerFrameScalarExpressions;
 	Ar << PerFrameVectorExpressions;
+	Ar << PerFramePrevScalarExpressions;
+	Ar << PerFramePrevVectorExpressions;
 
 	return bShaderHasOutdatedParameters;
 }
diff --git a/Engine/Source/Runtime/Renderer/Public/MaterialShader.h b/Engine/Source/Runtime/Renderer/Public/MaterialShader.h
--- a/Engine/Source/Runtime/Renderer/Public/MaterialShader.h
+++ b/Engine/Source/Runtime/Renderer/Public/MaterialShader.h
@@ -140,6 +140,8 @@ private:
 	TArray<FShaderUniformBufferParameter> ParameterCollectionUniformBuffers;
 	TArray<FShaderParameter> PerFrameScalarExpressions;
 	TArray<FShaderParameter> PerFrameVectorExpressions;
+	TArray<FShaderParameter> PerFramePrevScalarExpressions;
+	TArray<FShaderParameter> PerFramePrevVectorExpressions;
 	FDeferredPixelShaderParameters DeferredParameters;
 	FShaderResourceParameter LightAttenuation;
 	FShaderResourceParameter LightAttenuationSampler;
