Fixed: Level streaming does not work with camera location, which breaks cinematic
UE-7538

diff --git a/Engine/Source/Runtime/Engine/Private/LevelTick.cpp b/Engine/Source/Runtime/Engine/Private/LevelTick.cpp
--- a/Engine/Source/Runtime/Engine/Private/LevelTick.cpp
+++ b/Engine/Source/Runtime/Engine/Private/LevelTick.cpp
@@ -550,7 +550,7 @@ void UWorld::ProcessLevelStreamingVolumes(FVector* OverrideViewLocation)
 			else
 			{
 				FRotator ViewRotation(0,0,0);
-				PlayerActor->GetActorEyesViewPoint( ViewLocation, ViewRotation );
+				PlayerActor->GetPlayerViewPoint( ViewLocation, ViewRotation );
 			}
 
 			TMap<AVolume*,bool> VolumeMap;
diff --git a/Engine/Source/Runtime/Engine/Private/WorldComposition.cpp b/Engine/Source/Runtime/Engine/Private/WorldComposition.cpp
--- a/Engine/Source/Runtime/Engine/Private/WorldComposition.cpp
+++ b/Engine/Source/Runtime/Engine/Private/WorldComposition.cpp
@@ -587,7 +587,7 @@ void UWorldComposition::UpdateStreamingState()
 		{
 			FVector ViewLocation;
 			FRotator ViewRotation;
-			Player->PlayerController->GetActorEyesViewPoint(ViewLocation, ViewRotation);
+			Player->PlayerController->GetPlayerViewPoint(ViewLocation, ViewRotation);
 			CentroidLocation+= ViewLocation;
 			NumViews++;
 		}
