Core - More stats for memory profiling

diff --git a/Engine/Source/Runtime/Engine/Private/AI/Navigation/NavigationSystem.cpp b/Engine/Source/Runtime/Engine/Private/AI/Navigation/NavigationSystem.cpp
--- a/Engine/Source/Runtime/Engine/Private/AI/Navigation/NavigationSystem.cpp
+++ b/Engine/Source/Runtime/Engine/Private/AI/Navigation/NavigationSystem.cpp
@@ -1947,6 +1947,7 @@ UNavigationSystem* UNavigationSystem::CreateNavigationSystem(UWorld* WorldOwner)
 
 void UNavigationSystem::InitializeForWorld(UWorld* World, FNavigationSystem::EMode Mode)
 {
+	DECLARE_SCOPE_CYCLE_COUNTER(TEXT("UNavigationSystem::InitializeForWorld"), STAT_NavSys_InitializeForWorld, STATGROUP_Navigation);
 	if (World)
 	{
 		UNavigationSystem* NavSys = World->GetNavigationSystem();
