Fix multiplication order of root motion

Revert change made in CL 2349909

diff --git a/Engine/Source/Runtime/Engine/Private/Animation/AnimSequence.cpp b/Engine/Source/Runtime/Engine/Private/Animation/AnimSequence.cpp
--- a/Engine/Source/Runtime/Engine/Private/Animation/AnimSequence.cpp
+++ b/Engine/Source/Runtime/Engine/Private/Animation/AnimSequence.cpp
@@ -806,8 +806,8 @@ FTransform UAnimSequence::ExtractRootMotionFromRange(float StartTrackPosition, f
 
 	// Transform to Component Space Rotation (inverse root transform from first frame)
 	const FTransform RootToComponentRot = FTransform(InitialTransform.GetRotation().Inverse());
-	StartTransform = StartTransform * RootToComponentRot;
-	EndTransform = EndTransform * RootToComponentRot;
+	StartTransform = RootToComponentRot * StartTransform;
+	EndTransform = RootToComponentRot * EndTransform;
 
 	return EndTransform.GetRelativeTransform(StartTransform);
 }
