Fix inability to disable foliage lightmaps and reduce default override to something more sane

diff --git a/Engine/Source/Runtime/Engine/Classes/Foliage/FoliageType.h b/Engine/Source/Runtime/Engine/Classes/Foliage/FoliageType.h
--- a/Engine/Source/Runtime/Engine/Classes/Foliage/FoliageType.h
+++ b/Engine/Source/Runtime/Engine/Classes/Foliage/FoliageType.h
@@ -179,19 +179,23 @@ class UFoliageType : public UObject
 	UPROPERTY()
 	uint32 ShowInstanceSettings:1;
 
-	/** Controls whether the foliage should cast a shadow or not. **/
+	/** Controls whether the foliage should take part in static lighting/shadowing. If false, mesh will not receive or cast static lighting or shadows, regardless of other settings. */
+	UPROPERTY(EditAnywhere, BlueprintReadOnly, Category = General, meta = (Subcategory = "Lighting"))
+	uint32 bEnableStaticLighting : 1;
+
+	/** Controls whether the foliage should cast a shadow or not. */
 	UPROPERTY(EditAnywhere, BlueprintReadOnly, Category = General, meta = (Subcategory = "Lighting"))
 	uint32 CastShadow:1;
 
-	/** Controls whether the foliage should inject light into the Light Propagation Volume.  This flag is only used if CastShadow is true. **/
+	/** Controls whether the foliage should inject light into the Light Propagation Volume.  This flag is only used if CastShadow is true. */
 	UPROPERTY(EditAnywhere, BlueprintReadOnly, Category = General, meta = (Subcategory = "Lighting"))
 	uint32 bAffectDynamicIndirectLighting:1;
 
-	/** Controls whether the primitive should affect dynamic distance field lighting methods.  This flag is only used if CastShadow is true. **/
+	/** Controls whether the primitive should affect dynamic distance field lighting methods.  This flag is only used if CastShadow is true. */
 	UPROPERTY(EditAnywhere, BlueprintReadOnly, Category = General, meta = (Subcategory = "Lighting", EditCondition = "CastShadow"))
 	uint32 bAffectDistanceFieldLighting:1;
 
-	/** Controls whether the foliage should cast shadows in the case of non precomputed shadowing.  This flag is only used if CastShadow is true. **/
+	/** Controls whether the foliage should cast shadows in the case of non precomputed shadowing.  This flag is only used if CastShadow is true. */
 	UPROPERTY(EditAnywhere, BlueprintReadOnly, Category = General, meta = (Subcategory = "Lighting"))
 	uint32 bCastDynamicShadow:1;
 
@@ -211,7 +215,7 @@ class UFoliageType : public UObject
 	UPROPERTY(BlueprintReadOnly, Category = General, meta = (Subcategory = "Lighting"))
 	uint32 bOverrideLightMapRes:1;
 
-	/** Overrides the lightmap resolution defined in the static mesh. A value of 0 disables static lighting/shadowing */
+	/** Overrides the lightmap resolution defined in the static mesh */
 	UPROPERTY(EditAnywhere, BlueprintReadOnly, Category = General, meta = (Subcategory = "Lighting", DisplayName = "Light Map Resolution", EditCondition = "bOverrideLightMapRes"))
 	int32 OverriddenLightMapRes;
 
diff --git a/Engine/Source/Runtime/Engine/Private/InstancedFoliage.cpp b/Engine/Source/Runtime/Engine/Private/InstancedFoliage.cpp
--- a/Engine/Source/Runtime/Engine/Private/InstancedFoliage.cpp
+++ b/Engine/Source/Runtime/Engine/Private/InstancedFoliage.cpp
@@ -192,6 +192,7 @@ UFoliageType::UFoliageType(const FObjectInitializer& ObjectInitializer)
 	VertexColorMask = FOLIAGEVERTEXCOLORMASK_Disabled;
 	VertexColorMaskThreshold = 0.5f;
 
+	bEnableStaticLighting = true;
 	CastShadow = true;
 	bCastDynamicShadow = true;
 	bCastStaticShadow = true;
@@ -202,7 +203,7 @@ UFoliageType::UFoliageType(const FObjectInitializer& ObjectInitializer)
 	bReceivesDecals = false;
 
 	bOverrideLightMapRes = false;
-	OverriddenLightMapRes = 32;
+	OverriddenLightMapRes = 8;
 
 	BodyInstance.SetCollisionProfileName(UCollisionProfile::NoCollision_ProfileName);
 
@@ -278,7 +279,7 @@ void UFoliageType::PostEditChangeProperty(struct FPropertyChangedEvent& Property
 	Super::PostEditChangeProperty(PropertyChangedEvent);
 
 	// Ensure that OverriddenLightMapRes is a factor of 4
-	OverriddenLightMapRes = FMath::Max(OverriddenLightMapRes + 3 & ~3, 4);
+	OverriddenLightMapRes = OverriddenLightMapRes > 4 ? OverriddenLightMapRes + 3 & ~3 : 4;
 	++ChangeCount;
 }
 #endif
@@ -376,7 +377,7 @@ void FFoliageMeshInfo::AddInstance(AInstancedFoliageActor* InIFA, UFoliageType*
 	{
 		Component = ConstructObject<UHierarchicalInstancedStaticMeshComponent>(UHierarchicalInstancedStaticMeshComponent::StaticClass(), InIFA, NAME_None, RF_Transactional);
 
-		Component->Mobility = EComponentMobility::Static;
+		Component->Mobility = InSettings->bEnableStaticLighting ? EComponentMobility::Static : EComponentMobility::Movable;
 
 		Component->StaticMesh = InSettings->GetStaticMesh();
 		Component->bSelectable = true;
