Inheritable Component - Very eraly version.

TODO:
- Propagation on changes (from parent to children)
- Handling the component renaming.

diff --git a/Engine/Config/BaseEngine.ini b/Engine/Config/BaseEngine.ini
--- a/Engine/Config/BaseEngine.ini
+++ b/Engine/Config/BaseEngine.ini
@@ -1475,6 +1475,7 @@ bDisplaysLayout=false
 bNativeCodeGenerationTool=false
 bOptimizeExecutionFlowStack=true
 bOptimizeAdjacentStates=true
+bEnableInheritableComponents=false
 
 [/Script/Engine.Blueprint]
 bRecompileOnLoad=true
diff --git a/Engine/Source/Editor/Kismet/Private/BlueprintEditor.cpp b/Engine/Source/Editor/Kismet/Private/BlueprintEditor.cpp
--- a/Engine/Source/Editor/Kismet/Private/BlueprintEditor.cpp
+++ b/Engine/Source/Editor/Kismet/Private/BlueprintEditor.cpp
@@ -677,10 +677,23 @@ void FBlueprintEditor::OnSCSEditorUpdateSelectionFromNodes(const TArray<FSCSEdit
 	for (auto NodeIt = SelectedNodes.CreateConstIterator(); NodeIt; ++NodeIt)
 	{
 		auto NodePtr = *NodeIt;
-		if(NodePtr.IsValid() && NodePtr->CanEditDefaults())
+		if (NodePtr.IsValid())
+		{
+			UActorComponent* EditableComponent = nullptr;
+			if (NodePtr->CanEditDefaults())
+			{
+				EditableComponent = NodePtr->GetComponentTemplate();
+			}
+			else if (!NodePtr->IsNative() && NodePtr->IsInherited())
+			{
+				EditableComponent = NodePtr->GetOverridenComponentTemplate(GetBlueprintObj(), true);
+			}
+
+			if (EditableComponent)
 		{
 			InspectorTitle = FText::FromString(NodePtr->GetDisplayString());
-			InspectorObjects.Add(NodePtr->GetComponentTemplate());
+				InspectorObjects.Add(EditableComponent);
+			}
 		}
 	}
 
diff --git a/Engine/Source/Editor/Kismet/Private/SSCSEditor.cpp b/Engine/Source/Editor/Kismet/Private/SSCSEditor.cpp
--- a/Engine/Source/Editor/Kismet/Private/SSCSEditor.cpp
+++ b/Engine/Source/Editor/Kismet/Private/SSCSEditor.cpp
@@ -33,6 +33,8 @@
 #include "Engine/SimpleConstructionScript.h"
 #include "Engine/Selection.h"
 
+#include "Engine/InheritableComponentHandler.h"
+
 #define LOCTEXT_NAMESPACE "SSCSEditor"
 
 static const FName SCS_ColumnName_ComponentClass( "ComponentClass" );
@@ -707,6 +709,39 @@ void FSCSEditorTreeNode::OnCompleteRename(const FText& InNewName)
 	}
 }
 
+UActorComponent* FSCSEditorTreeNode::GetOverridenComponentTemplate(UBlueprint* Blueprint, bool bCreateIfNecessary) const
+{
+	UActorComponent* OverridenComponent = NULL;
+
+	FComponentKey Key(GetSCSNode());
+	if (!Key.IsValid())
+	{
+		UActorComponent* ComponentTemplate = GetComponentTemplate();
+		auto OriginalComponentOwner = ComponentTemplate ? ComponentTemplate->GetOwner() : nullptr;
+		auto ComponentOwnerBPGC = OriginalComponentOwner ? Cast<UBlueprintGeneratedClass>(OriginalComponentOwner->GetClass()) : nullptr;
+		Key = FComponentKey(ComponentOwnerBPGC, GetVariableName());
+	}
+
+	const bool BlueprintCanOverrideComponentFormKey = Key.IsValid() 
+		&& Blueprint 
+		&& Blueprint->ParentClass 
+		&& Blueprint->ParentClass->IsChildOf(Key.OwnerClass);
+
+	if (BlueprintCanOverrideComponentFormKey)
+	{
+		UInheritableComponentHandler* InheritableComponentHandler = Blueprint->GetInheritableComponentHandler(bCreateIfNecessary);
+		if (InheritableComponentHandler)
+		{
+			OverridenComponent = InheritableComponentHandler->GetOverridenComponentTemplate(Key);
+			if (!OverridenComponent && bCreateIfNecessary)
+			{
+				OverridenComponent = InheritableComponentHandler->CreateOverridenComponentTemplate(Key);
+			}
+		}
+	}
+	return OverridenComponent;
+}
+
 //////////////////////////////////////////////////////////////////////////
 // FSCSEditorComponentObjectTextFactory
 
diff --git a/Engine/Source/Editor/Kismet/Public/SSCSEditor.h b/Engine/Source/Editor/Kismet/Public/SSCSEditor.h
--- a/Engine/Source/Editor/Kismet/Public/SSCSEditor.h
+++ b/Engine/Source/Editor/Kismet/Public/SSCSEditor.h
@@ -187,6 +187,9 @@ public:
 	/** Sets up the delegate for renaming a component */
 	void SetRenameRequestedDelegate(FOnRenameRequested InRenameRequested) { RenameRequestedDelegate = InRenameRequested; }
 
+	/** Get overridden template component, specialized in given blueprint */
+	UActorComponent* GetOverridenComponentTemplate(UBlueprint* Blueprint, bool bCreateIfNecessary) const;
+
 private:
 	bool bIsInherited;
 	TWeakObjectPtr<class USCS_Node> SCSNodePtr;
diff --git a/Engine/Source/Editor/KismetCompiler/Private/KismetCompiler.cpp b/Engine/Source/Editor/KismetCompiler/Private/KismetCompiler.cpp
--- a/Engine/Source/Editor/KismetCompiler/Private/KismetCompiler.cpp
+++ b/Engine/Source/Editor/KismetCompiler/Private/KismetCompiler.cpp
@@ -23,6 +23,7 @@
 #include "EdGraph/EdGraphNode_Documentation.h"
 #include "Engine/DynamicBlueprintBinding.h"
 #include "Kismet/KismetMathLibrary.h"
+#include "Engine/InheritableComponentHandler.h"
 
 static bool bDebugPropertyPropagation = false;
 
@@ -258,6 +259,14 @@ void FKismetCompilerContext::SaveSubObjectsFromCleanAndSanitizeClass(FSubobjectC
 			SubObjectsToSave.AddObject(Curve);
 		}
 	}
+
+	if (Blueprint->InheritableComponentHandler)
+	{
+		SubObjectsToSave.AddObject(Blueprint->InheritableComponentHandler);
+		TArray<UActorComponent*> AllTemplates;
+		Blueprint->InheritableComponentHandler->GetAllTemplates(AllTemplates);
+		SubObjectsToSave.AddObjects(AllTemplates);
+	}
 }
 
 void FKismetCompilerContext::PostCreateSchema()
@@ -1768,10 +1777,12 @@ void FKismetCompilerContext::FinishCompilingClass(UClass* Class)
 		BPGClass->ComponentTemplates.Empty();
 		BPGClass->Timelines.Empty();
 		BPGClass->SimpleConstructionScript = NULL;
+		BPGClass->InheritableComponentHandler = NULL;
 
 		BPGClass->ComponentTemplates = Blueprint->ComponentTemplates;
 		BPGClass->Timelines = Blueprint->Timelines;
 		BPGClass->SimpleConstructionScript = Blueprint->SimpleConstructionScript;
+		BPGClass->InheritableComponentHandler = Blueprint->InheritableComponentHandler;
 	}
 
 	//@TODO: Not sure if doing this again is actually necessary
@@ -3264,6 +3275,14 @@ void FKismetCompilerContext::Compile()
 		++TimelineIndex;
 	}
 
+	if (CompileOptions.CompileType == EKismetCompileType::Full)
+	{
+		if (Blueprint->InheritableComponentHandler)
+		{
+			Blueprint->InheritableComponentHandler->RemoveInvalidAndUnnecessaryTemplates();
+		}
+	}
+
 	CleanAndSanitizeClass(TargetClass, OldCDO);
 
 	FKismetCompilerVMBackend Backend_VM(Blueprint, Schema, *this);
diff --git a/Engine/Source/Editor/UnrealEd/Private/Kismet2/BlueprintEditorUtils.cpp b/Engine/Source/Editor/UnrealEd/Private/Kismet2/BlueprintEditorUtils.cpp
--- a/Engine/Source/Editor/UnrealEd/Private/Kismet2/BlueprintEditorUtils.cpp
+++ b/Engine/Source/Editor/UnrealEd/Private/Kismet2/BlueprintEditorUtils.cpp
@@ -46,6 +46,7 @@
 #include "NotificationManager.h"
 #include "Editor/Blutility/Public/IBlutilityModule.h"
 
+#include "Engine/InheritableComponentHandler.h"
 #define LOCTEXT_NAMESPACE "Blueprint"
 
 DEFINE_LOG_CATEGORY(LogBlueprintDebug);
@@ -1135,6 +1136,21 @@ UClass* FBlueprintEditorUtils::RegenerateBlueprintClass(UBlueprint* Blueprint, U
 		// Make sure the simple construction script is loaded, since the outer hierarchy isn't compatible with PreloadMembers past the root node
 		FBlueprintEditorUtils::PreloadConstructionScript(Blueprint);
 
+		// Preload Overriden Components
+		if (Blueprint->InheritableComponentHandler)
+		{
+			if (Blueprint->InheritableComponentHandler->HasAllFlags(RF_NeedLoad))
+			{
+				auto Linker = Blueprint->InheritableComponentHandler->GetLinker();
+				if (Linker)
+				{
+					Linker->Preload(Blueprint->InheritableComponentHandler);
+				}
+			}
+
+			Blueprint->InheritableComponentHandler->PreloadAllTempates();
+		}
+
 		// Purge any NULL graphs
 		FBlueprintEditorUtils::PurgeNullGraphs(Blueprint);
 
@@ -1411,6 +1427,9 @@ void FBlueprintEditorUtils::PostDuplicateBlueprint(UBlueprint* Blueprint)
 			USimpleConstructionScript* SCSRootNode = Blueprint->SimpleConstructionScript;
 			Blueprint->SimpleConstructionScript = NULL;
 
+			UInheritableComponentHandler* InheritableComponentHandler = Blueprint->InheritableComponentHandler;
+			Blueprint->InheritableComponentHandler = NULL;
+
 			TArray<UActorComponent*> Templates = Blueprint->ComponentTemplates;
 			Blueprint->ComponentTemplates.Empty();
 
@@ -1493,9 +1512,19 @@ void FBlueprintEditorUtils::PostDuplicateBlueprint(UBlueprint* Blueprint)
 				OldToNewMap.Add(OldTimeline, NewTimeline);
 			}
 
+			if (InheritableComponentHandler)
+			{
+				NewBPGC->InheritableComponentHandler = Cast<UInheritableComponentHandler>(StaticDuplicateObject(InheritableComponentHandler, NewBPGC, *InheritableComponentHandler->GetName()));
+				if (NewBPGC->InheritableComponentHandler)
+				{
+					NewBPGC->InheritableComponentHandler->UpdateOwnerClass(NewBPGC);
+				}
+			}
+
 			Blueprint->SimpleConstructionScript = NewBPGC->SimpleConstructionScript;
 			Blueprint->ComponentTemplates = NewBPGC->ComponentTemplates;
 			Blueprint->Timelines = NewBPGC->Timelines;
+			Blueprint->InheritableComponentHandler = NewBPGC->InheritableComponentHandler;
 
 			Compiler.CompileBlueprint(Blueprint, CompileOptions, Results);
 
@@ -2383,6 +2412,11 @@ bool FBlueprintEditorUtils::IsDataOnlyBlueprint(const UBlueprint* Blueprint)
 		return false;
 	}
 
+	if (Blueprint->InheritableComponentHandler && !Blueprint->InheritableComponentHandler->IsEmpty())
+	{
+		return false;
+	}
+
 	return true;
 }
 
diff --git a/Engine/Source/Runtime/Engine/Classes/Engine/Blueprint.h b/Engine/Source/Runtime/Engine/Classes/Engine/Blueprint.h
--- a/Engine/Source/Runtime/Engine/Classes/Engine/Blueprint.h
+++ b/Engine/Source/Runtime/Engine/Classes/Engine/Blueprint.h
@@ -379,6 +379,10 @@ class ENGINE_API UBlueprint : public UBlueprintCore
 	UPROPERTY()
 	TArray<class UTimelineTemplate*> Timelines;
 
+	/** Stores data to override (in children classes) components (created by SCS) from parent classes */
+	UPROPERTY()
+	class UInheritableComponentHandler* InheritableComponentHandler;
+
 	/** The type of this blueprint */
 	UPROPERTY(AssetRegistrySearchable)
 	TEnumAsByte<enum EBlueprintType> BlueprintType;
@@ -550,6 +554,8 @@ public:
 
 	bool ChangeOwnerOfTemplates();
 
+	UInheritableComponentHandler* GetInheritableComponentHandler(bool bCreateIfNecessary);
+
 #endif	//#if WITH_EDITOR
 
 	// Begin UObject interface
diff --git a/Engine/Source/Runtime/Engine/Classes/Engine/BlueprintGeneratedClass.h b/Engine/Source/Runtime/Engine/Classes/Engine/BlueprintGeneratedClass.h
--- a/Engine/Source/Runtime/Engine/Classes/Engine/BlueprintGeneratedClass.h
+++ b/Engine/Source/Runtime/Engine/Classes/Engine/BlueprintGeneratedClass.h
@@ -434,6 +434,10 @@ public:
 	UPROPERTY()
 	class USimpleConstructionScript* SimpleConstructionScript;
 
+	/** Stores data to override (in children classes) components (created by SCS) from parent classes */
+	UPROPERTY()
+	class UInheritableComponentHandler* InheritableComponentHandler;
+
 	UPROPERTY()
 	UStructProperty* UberGraphFramePointerProperty;
 
diff --git a/Engine/Source/Runtime/Engine/Classes/Engine/InheritableComponentHandler.h b/Engine/Source/Runtime/Engine/Classes/Engine/InheritableComponentHandler.h
new file mode 100644
--- /dev/null
+++ b/Engine/Source/Runtime/Engine/Classes/Engine/InheritableComponentHandler.h
@@ -0,0 +1,103 @@
+// Copyright 1998-2014 Epic Games, Inc. All Rights Reserved.
+
+
+#pragma once
+#include "InheritableComponentHandler.generated.h"
+
+class USCS_Node;
+class UActorComponent;
+
+USTRUCT()
+struct ENGINE_API FComponentKey
+{
+	GENERATED_USTRUCT_BODY()
+
+	UPROPERTY()
+	UBlueprintGeneratedClass* OwnerClass;
+
+	UPROPERTY()
+	FName VariableName;
+
+	// TODO: a GUID should be added
+
+	FComponentKey()
+		: OwnerClass(nullptr)
+	{}
+
+	FComponentKey(UBlueprintGeneratedClass* InOwnerClass, FName InVariableNam)
+		: OwnerClass(nullptr)
+		, VariableName(InVariableNam)
+	{}
+
+	FComponentKey(USCS_Node* ParentNode);
+
+	bool Match(FComponentKey OtherKey) const;
+
+	bool IsValid() const
+	{
+		return OwnerClass && (VariableName != NAME_None);
+	}
+
+	UActorComponent* GetOriginalTemplate() const;
+};
+
+USTRUCT()
+struct FComponentOverrideRecord
+{
+	GENERATED_USTRUCT_BODY()
+
+	UPROPERTY()
+	UActorComponent* ComponentTemplate;
+
+	UPROPERTY()
+	FComponentKey ComponentKey;
+
+	FComponentOverrideRecord()
+		: ComponentTemplate(nullptr)
+	{}
+};
+
+UCLASS()
+class ENGINE_API UInheritableComponentHandler : public UObject
+{
+	GENERATED_BODY()
+
+#if WITH_EDITOR
+private:
+	bool IsRecordValid(const FComponentOverrideRecord& Record) const;
+	bool IsRecordNecessary(const FComponentOverrideRecord& Record) const;
+
+public:
+	UActorComponent* CreateOverridenComponentTemplate(FComponentKey Key);
+	void UpdateOverridenComponentTemplate(FComponentKey Key);
+	void UpdateOwnerClass(UBlueprintGeneratedClass* OwnerClass);
+	void RemoveInvalidAndUnnecessaryTemplates();
+	bool IsValid() const;
+	UActorComponent* FindBestArchetype(FComponentKey Key) const;
+
+	void GetAllTemplates(TArray<UActorComponent*>& OutArray) const
+	{
+		for (auto Record : Records)
+		{
+			OutArray.Add(Record.ComponentTemplate);
+		}
+	}
+
+	bool IsEmpty() const
+	{
+		return 0 == Records.Num();
+	}
+
+	void PreloadAllTempates();
+
+#endif
+
+public:
+	UActorComponent* GetOverridenComponentTemplate(FComponentKey Key) const;
+
+private:
+	const FComponentOverrideRecord* FindRecord(FComponentKey Key) const;
+	
+	UPROPERTY()
+	TArray<FComponentOverrideRecord> Records;
+};
diff --git a/Engine/Source/Runtime/Engine/Private/Blueprint.cpp b/Engine/Source/Runtime/Engine/Private/Blueprint.cpp
--- a/Engine/Source/Runtime/Engine/Private/Blueprint.cpp
+++ b/Engine/Source/Runtime/Engine/Private/Blueprint.cpp
@@ -22,6 +22,7 @@
 #include "Engine/SCS_Node.h"
 #endif
 #include "Components/TimelineComponent.h"
+#include "Engine/InheritableComponentHandler.h"
 
 DEFINE_LOG_CATEGORY(LogBlueprint);
 
@@ -533,6 +534,11 @@ void UBlueprint::PostLoad()
 	}
 
 	FStructureEditorUtils::RemoveInvalidStructureMemberVariableFromBlueprint(this);
+
+	if (InheritableComponentHandler)
+	{
+		InheritableComponentHandler->RemoveInvalidAndUnnecessaryTemplates();
+	}
 }
 
 void UBlueprint::DebuggingWorldRegistrationHelper(UObject* ObjectProvidingWorld, UObject* ValueToRegister)
@@ -844,6 +850,22 @@ bool UBlueprint::ValidateGeneratedClass(const UClass* InClass)
 		}
 	}
 
+	if (const UInheritableComponentHandler* InheritableComponentHandler = Blueprint->InheritableComponentHandler)
+	{
+		if (!ensure(InheritableComponentHandler->GetOuter() == GeneratedClass))
+		{
+			return false;
+		}
+	}
+
+	if (const UInheritableComponentHandler* InheritableComponentHandler = GeneratedClass->InheritableComponentHandler)
+	{
+		if (!ensure(InheritableComponentHandler->GetOuter() == GeneratedClass))
+		{
+			return false;
+		}
+	}
+
 	return true;
 }
 
@@ -1186,6 +1208,24 @@ bool UBlueprint::Modify(bool bAlwaysMarkDirty)
 	return Super::Modify(bAlwaysMarkDirty);
 }
 
+UInheritableComponentHandler* UBlueprint::GetInheritableComponentHandler(bool bCreateIfNecessary)
+{
+	static const FBoolConfigValueHelper EnableInheritableComponents(TEXT("Kismet"), TEXT("bEnableInheritableComponents"), GEngineIni);
+	if (!EnableInheritableComponents)
+	{
+		return NULL;
+	}
+
+	if (!InheritableComponentHandler && bCreateIfNecessary)
+	{
+		auto BPGC = CastChecked<UBlueprintGeneratedClass>(GeneratedClass);
+		ensure(!BPGC->InheritableComponentHandler);
+		InheritableComponentHandler = NewNamedObject<UInheritableComponentHandler>(BPGC, FName(TEXT("InheritableComponentHandler")));
+		BPGC->InheritableComponentHandler = InheritableComponentHandler;
+	}
+	return InheritableComponentHandler;
+}
+
 #endif
 
 #if WITH_EDITOR
diff --git a/Engine/Source/Runtime/Engine/Private/InheritableComponentHandler.cpp b/Engine/Source/Runtime/Engine/Private/InheritableComponentHandler.cpp
new file mode 100644
--- /dev/null
+++ b/Engine/Source/Runtime/Engine/Private/InheritableComponentHandler.cpp
@@ -0,0 +1,281 @@
+// Copyright 1998-2015 Epic Games, Inc. All Rights Reserved.
+
+#include "EnginePrivate.h"
+#include "Engine/SCS_Node.h"
+#include "Engine/InheritableComponentHandler.h"
+
+// UInheritableComponentHandler
+
+#if WITH_EDITOR
+UActorComponent* UInheritableComponentHandler::CreateOverridenComponentTemplate(FComponentKey Key)
+{
+	for (int32 Index = 0; Index < Records.Num(); ++Index)
+	{
+		if (Records[Index].ComponentKey.Match(Key))
+		{
+			if (Records[Index].ComponentTemplate)
+			{
+				return Records[Index].ComponentTemplate;
+			}
+			Records.RemoveAtSwap(Index);
+			break;
+		}
+	}
+
+	auto BestArchetype = FindBestArchetype(Key);
+	if (!BestArchetype)
+	{
+		UE_LOG(LogBlueprint, Warning, TEXT("CreateOverridenComponentTemplate '%s': cannot find archetype for component '%s' from '%s'"),
+			*GetPathNameSafe(this), *Key.VariableName.ToString(), *GetPathNameSafe(Key.OwnerClass));
+		return NULL;
+	}
+	ensure(Cast<UBlueprintGeneratedClass>(GetOuter()));
+	auto NewComponentTemplate = ConstructObject<UActorComponent>(
+		BestArchetype->GetClass(), GetOuter(), BestArchetype->GetFName(), RF_ArchetypeObject, BestArchetype);
+
+	FComponentOverrideRecord NewRecord;
+	NewRecord.ComponentKey = Key;
+	NewRecord.ComponentTemplate = NewComponentTemplate;
+	Records.Add(NewRecord);
+
+	return NewComponentTemplate;
+}
+
+void UInheritableComponentHandler::UpdateOverridenComponentTemplate(FComponentKey Key)
+{
+	auto FoundRecord = FindRecord(Key);
+	if (!FoundRecord)
+	{
+		UE_LOG(LogBlueprint, Warning, TEXT("UpdateOverridenComponentTemplate '%s': cannot find overriden template for component '%s' from '%s'"),
+			*GetPathNameSafe(this), *Key.VariableName.ToString(), *GetPathNameSafe(Key.OwnerClass));
+		return;
+	}
+
+	check(FoundRecord->ComponentTemplate);
+
+	// TODO: save diff against current (previous) archetype. This part is tricky:
+	// the diff should be made before the archetype is updated
+	// OR the diff should be made from saved archetype
+
+	// TODO: copy data from best archetype (it should be new/updated one)
+
+	// TODO: restore custom data from diff
+}
+
+void UInheritableComponentHandler::UpdateOwnerClass(UBlueprintGeneratedClass* OwnerClass)
+{
+	for (auto& Record : Records)
+	{
+		auto OldComponentTemplate = Record.ComponentTemplate;
+		if (OldComponentTemplate && (OwnerClass != OldComponentTemplate->GetOuter()))
+		{
+			Record.ComponentTemplate = DuplicateObject(OldComponentTemplate, OwnerClass, *OldComponentTemplate->GetName());
+		}
+	}
+}
+
+void UInheritableComponentHandler::RemoveInvalidAndUnnecessaryTemplates()
+{
+	for (int32 Index = 0; Index < Records.Num();)
+	{
+		bool bIsValidAndNecessary = false;
+		{
+			const FComponentOverrideRecord& Record = Records[Index];
+			if (IsRecordValid(Record))
+			{
+				if (IsRecordNecessary(Record))
+				{
+					bIsValidAndNecessary = true;
+				}
+				else
+				{
+					UE_LOG(LogBlueprint, Log, TEXT("RemoveInvalidAndUnnecessaryTemplates '%s': overriden template is unnecessary - component '%s' from '%s'"),
+						*GetPathNameSafe(this), *Record.ComponentKey.VariableName.ToString(), *GetPathNameSafe(Record.ComponentKey.OwnerClass));
+				}
+			}
+			else
+			{
+				UE_LOG(LogBlueprint, Warning, TEXT("RemoveInvalidAndUnnecessaryTemplates '%s': overriden template is invalid - component '%s' from '%s'"),
+					*GetPathNameSafe(this), *Record.ComponentKey.VariableName.ToString(), *GetPathNameSafe(Record.ComponentKey.OwnerClass));
+			}
+		}
+
+		if (bIsValidAndNecessary)
+		{
+			++Index;
+		}
+		else
+		{
+			Records.RemoveAtSwap(Index);
+		}
+	}
+}
+
+bool UInheritableComponentHandler::IsValid() const
+{
+	for (auto& Record : Records)
+	{
+		if (!IsRecordValid(Record))
+		{
+			return false;
+		}
+	}
+	return true;
+}
+
+bool UInheritableComponentHandler::IsRecordValid(const FComponentOverrideRecord& Record) const
+{
+	UClass* OwnerClass = Cast<UClass>(GetOuter());
+	ensure(OwnerClass);
+
+	if (!Record.ComponentTemplate)
+	{
+		return false;
+	}
+
+	if (Record.ComponentTemplate->GetOuter() != OwnerClass)
+	{
+		return false;
+	}
+
+	if (!Record.ComponentKey.IsValid())
+	{
+		return false;
+	}
+
+	auto OriginalTemplate = Record.ComponentKey.GetOriginalTemplate();
+	if (!OriginalTemplate)
+	{
+		return false;
+	}
+
+	if (OriginalTemplate->GetClass() != Record.ComponentTemplate->GetClass())
+	{
+		return false;
+	}
+	
+	return true;
+}
+
+struct FComponentComparisonHelper
+{
+	static bool AreIdentical(UObject* ObjectA, UObject* ObjectB)
+	{
+		if (!ObjectA || !ObjectB || (ObjectA->GetClass() != ObjectB->GetClass()))
+		{
+			return false;
+		}
+
+		bool Result = true;
+		for (UProperty* Prop = ObjectA->GetClass()->PropertyLink; Prop && Result; Prop = Prop->PropertyLinkNext)
+		{
+			bool bConsiderProperty = Prop->ShouldDuplicateValue(); //Should the property be compared at all?
+			if (bConsiderProperty)
+			{
+				for (int32 Idx = 0; (Idx < Prop->ArrayDim) && Result; Idx++)
+				{
+					if (!Prop->Identical_InContainer(ObjectA, ObjectB, Idx, PPF_DeepComparison))
+					{
+						Result = false;
+					}
+				}
+			}
+		}
+		if (Result)
+		{
+			// Allow the component to compare its native/ intrinsic properties.
+			Result = ObjectA->AreNativePropertiesIdenticalTo(ObjectB);
+		}
+		return Result;
+	}
+};
+
+bool UInheritableComponentHandler::IsRecordNecessary(const FComponentOverrideRecord& Record) const
+{
+	auto ChildComponentTemplate = Record.ComponentTemplate;
+	auto ParentComponentTemplate = FindBestArchetype(Record.ComponentKey);
+	check(ChildComponentTemplate && ParentComponentTemplate && (ParentComponentTemplate != ChildComponentTemplate));
+	return !FComponentComparisonHelper::AreIdentical(ChildComponentTemplate, ParentComponentTemplate);
+}
+
+void UInheritableComponentHandler::PreloadAllTempates()
+{
+	for (auto Record : Records)
+	{
+		if (Record.ComponentTemplate && Record.ComponentTemplate->HasAllFlags(RF_NeedLoad))
+		{
+			auto Linker = Record.ComponentTemplate->GetLinker();
+			if (Linker)
+			{
+				Linker->Preload(Record.ComponentTemplate);
+			}
+		}
+	}
+}
+
+UActorComponent* UInheritableComponentHandler::FindBestArchetype(FComponentKey Key) const
+{
+	UActorComponent* ClosestArchetype = nullptr;
+
+	auto ActualBPGC = Cast<UBlueprintGeneratedClass>(GetOuter());
+	if (ActualBPGC && Key.OwnerClass && (ActualBPGC != Key.OwnerClass))
+	{
+		ActualBPGC = Cast<UBlueprintGeneratedClass>(ActualBPGC->GetSuperClass());
+		while (!ClosestArchetype && ActualBPGC)
+		{
+			if (ActualBPGC->InheritableComponentHandler)
+			{
+				ClosestArchetype = ActualBPGC->InheritableComponentHandler->GetOverridenComponentTemplate(Key);
+			}
+			ActualBPGC = Cast<UBlueprintGeneratedClass>(ActualBPGC->GetSuperClass());
+		}
+
+		if (!ClosestArchetype)
+		{
+			ClosestArchetype = Key.GetOriginalTemplate();
+		}
+	}
+
+	return ClosestArchetype;
+}
+
+#endif
+
+UActorComponent* UInheritableComponentHandler::GetOverridenComponentTemplate(FComponentKey Key) const
+{
+	auto Record = FindRecord(Key);
+	return Record ? Record->ComponentTemplate : nullptr;
+}
+
+const FComponentOverrideRecord* UInheritableComponentHandler::FindRecord(FComponentKey Key) const
+{
+	for (auto& Record : Records)
+	{
+		if (Record.ComponentKey.Match(Key))
+		{
+			return &Record;
+		}
+	}
+	return nullptr;
+}
+
+// FComponentOverrideRecord
+
+FComponentKey::FComponentKey(USCS_Node* ParentNode) : OwnerClass(nullptr)
+{
+	auto ParentSCS = ParentNode ? ParentNode->GetSCS() : nullptr;
+	OwnerClass = ParentSCS ? Cast<UBlueprintGeneratedClass>(ParentSCS->GetOwnerClass()) : nullptr;
+	VariableName = (ParentNode && OwnerClass) ? ParentNode->GetVariableName() : NAME_None;
+}
+
+bool FComponentKey::Match(FComponentKey OtherKey) const
+{
+	return (OwnerClass == OtherKey.OwnerClass) && (VariableName == OtherKey.VariableName);
+}
+
+UActorComponent* FComponentKey::GetOriginalTemplate() const
+{
+	auto ParentSCS = OwnerClass ? OwnerClass->SimpleConstructionScript : nullptr;
+	auto SCSNode = ParentSCS ? ParentSCS->FindSCSNode(VariableName) : nullptr;
+	return SCSNode ? SCSNode->ComponentTemplate : nullptr;
+}
\ No newline at end of file
diff --git a/Engine/Source/Runtime/Engine/Private/SCS_Node.cpp b/Engine/Source/Runtime/Engine/Private/SCS_Node.cpp
--- a/Engine/Source/Runtime/Engine/Private/SCS_Node.cpp
+++ b/Engine/Source/Runtime/Engine/Private/SCS_Node.cpp
@@ -8,6 +8,7 @@
 #endif
 #include "Engine/SCS_Node.h"
 #include "Engine/SimpleConstructionScript.h"
+#include "Engine/InheritableComponentHandler.h"
 
 //////////////////////////////////////////////////////////////////////////
 // USCS_Node
@@ -30,8 +31,22 @@ UActorComponent* USCS_Node::ExecuteNodeOnActor(AActor* Actor, USceneComponent* P
 	check(Actor != nullptr);
 	check((ParentComponent != nullptr && !ParentComponent->IsPendingKill()) || (RootTransform != nullptr)); // must specify either a parent component or a world transform
 
+	UActorComponent* OverridenComponentTemplate = nullptr;
+	{
+		auto ActualBPGC = Cast<UBlueprintGeneratedClass>(Actor->GetClass());
+		while (!OverridenComponentTemplate && ActualBPGC)
+		{
+			if (ActualBPGC->InheritableComponentHandler)
+			{
+				OverridenComponentTemplate = ActualBPGC->InheritableComponentHandler->GetOverridenComponentTemplate(this);
+			}
+			ActualBPGC = Cast<UBlueprintGeneratedClass>(ActualBPGC->GetSuperClass());
+		}
+	}
+	UActorComponent* ActualComponentTemplate = OverridenComponentTemplate ? OverridenComponentTemplate : ComponentTemplate;
+
 	// Create a new component instance based on the template
-	UActorComponent* NewActorComp = Actor->CreateComponentFromTemplate(ComponentTemplate, VariableName.ToString());
+	UActorComponent* NewActorComp = Actor->CreateComponentFromTemplate(ActualComponentTemplate, VariableName.ToString());
 	if(NewActorComp != nullptr)
 	{
 		// SCS created components are net addressable
