Fixed navigation system not using properly handling navigation-relevant actors in "always loaded" sublevels #UE4

And as a minor API improvement refactored UNavigationSystem::RegisterCustomLink and UNavigationSystem::UnregisterCustomLink to take references rather than pointers

diff --git a/Engine/Plugins/Experimental/LiveEditor/Source/LiveEditor/Private/LiveEditorManager.cpp b/Engine/Plugins/Experimental/LiveEditor/Source/LiveEditor/Private/LiveEditorManager.cpp
--- a/Engine/Plugins/Experimental/LiveEditor/Source/LiveEditor/Private/LiveEditorManager.cpp
+++ b/Engine/Plugins/Experimental/LiveEditor/Source/LiveEditor/Private/LiveEditorManager.cpp
@@ -624,6 +624,7 @@ void FLiveEditorManager::CreateLiveEditorWorld()
 										.CreatePhysicsScene(false)
 										.RequiresHitProxies(false)
 										.CreateNavigation(false)
+										.CreateAISystem(false)
 										.ShouldSimulatePhysics(false)
 										.SetTransactional(bTransactional));
 	LiveEditorWorld->InitializeActorsForPlay(FURL());
diff --git a/Engine/Source/Editor/UnrealEd/Private/Commandlets/CookCommandlet.cpp b/Engine/Source/Editor/UnrealEd/Private/Commandlets/CookCommandlet.cpp
--- a/Engine/Source/Editor/UnrealEd/Private/Commandlets/CookCommandlet.cpp
+++ b/Engine/Source/Editor/UnrealEd/Private/Commandlets/CookCommandlet.cpp
@@ -440,7 +440,7 @@ bool UCookCommandlet::SaveCookedPackage( UPackage* Package, uint32 SaveFlags, bo
 					if ( !World->bIsWorldInitialized)
 					{
 						// we need to initialize the world - at least need physics scene since BP construction script runs during cooking, otherwise trace won't work
-						World->InitWorld(UWorld::InitializationValues().RequiresHitProxies(false).ShouldSimulatePhysics(false).EnableTraceCollision(false).CreateNavigation(false).AllowAudioPlayback(false).CreatePhysicsScene(true));
+						World->InitWorld(UWorld::InitializationValues().RequiresHitProxies(false).ShouldSimulatePhysics(false).EnableTraceCollision(false).CreateNavigation(false).CreateAISystem(false).AllowAudioPlayback(false).CreatePhysicsScene(true));
 					}
 				}
 
diff --git a/Engine/Source/Editor/UnrealEd/Private/CookOnTheFlyServer.cpp b/Engine/Source/Editor/UnrealEd/Private/CookOnTheFlyServer.cpp
--- a/Engine/Source/Editor/UnrealEd/Private/CookOnTheFlyServer.cpp
+++ b/Engine/Source/Editor/UnrealEd/Private/CookOnTheFlyServer.cpp
@@ -1720,7 +1720,7 @@ bool UCookOnTheFlyServer::SaveCookedPackage( UPackage* Package, uint32 SaveFlags
 					if ( !World->bIsWorldInitialized)
 					{
 						// we need to initialize the world - at least need physics scene since BP construction script runs during cooking, otherwise trace won't work
-						World->InitWorld(UWorld::InitializationValues().RequiresHitProxies(false).ShouldSimulatePhysics(false).EnableTraceCollision(false).CreateNavigation(false).AllowAudioPlayback(false).CreatePhysicsScene(true));
+						World->InitWorld(UWorld::InitializationValues().RequiresHitProxies(false).ShouldSimulatePhysics(false).EnableTraceCollision(false).CreateNavigation(false).CreateAISystem(false).AllowAudioPlayback(false).CreatePhysicsScene(true));
 						bDidInitializeWorld = true;
 					}
 				}
diff --git a/Engine/Source/Runtime/Engine/Classes/AI/Navigation/NavigationSystem.h b/Engine/Source/Runtime/Engine/Classes/AI/Navigation/NavigationSystem.h
--- a/Engine/Source/Runtime/Engine/Classes/AI/Navigation/NavigationSystem.h
+++ b/Engine/Source/Runtime/Engine/Classes/AI/Navigation/NavigationSystem.h
@@ -532,8 +532,8 @@ public:
 	//----------------------------------------------------------------------//
 	// Custom navigation links
 	//----------------------------------------------------------------------//
-	void RegisterCustomLink(INavLinkCustomInterface* CustomLink);
-	void UnregisterCustomLink(INavLinkCustomInterface* CustomLink);
+	void RegisterCustomLink(INavLinkCustomInterface& CustomLink);
+	void UnregisterCustomLink(INavLinkCustomInterface& CustomLink);
 	
 	/** find custom link by unique ID */
 	INavLinkCustomInterface* GetCustomLink(uint32 UniqueLinkId) const;
@@ -832,5 +832,14 @@ private:
 	 *  Depends on runtime generation settings of each navigation data, always ERuntimeGenerationType::Dynamic in the editor world
 	 */
 	ERuntimeGenerationType GetRuntimeGenerationType() const;
+
+	//----------------------------------------------------------------------//
+	// DEPRECATED
+	//----------------------------------------------------------------------//
+public:
+	DEPRECATED(4.8, "This version is deprecated. Please use the one taking reference to INavLinkCustomInterface rather than a pointer instead.")
+	void RegisterCustomLink(INavLinkCustomInterface* CustomLink);
+	DEPRECATED(4.8, "This version is deprecated. Please use the one taking reference to INavLinkCustomInterface rather than a pointer instead.")
+	void UnregisterCustomLink(INavLinkCustomInterface* CustomLink);
 };
 
diff --git a/Engine/Source/Runtime/Engine/Classes/Engine/World.h b/Engine/Source/Runtime/Engine/Classes/Engine/World.h
--- a/Engine/Source/Runtime/Engine/Classes/Engine/World.h
+++ b/Engine/Source/Runtime/Engine/Classes/Engine/World.h
@@ -1917,8 +1917,8 @@ public:
 			, bAllowAudioPlayback(true)
 			, bRequiresHitProxies(true)
 			, bCreatePhysicsScene(true)
-			, bCreateNavigation(false)
-			, bCreateAISystem(false)
+			, bCreateNavigation(true)
+			, bCreateAISystem(true)
 			, bShouldSimulatePhysics(true)
 			, bEnableTraceCollision(false)
 			, bTransactional(true)
diff --git a/Engine/Source/Runtime/Engine/Private/AI/Navigation/NavLinkCustomComponent.cpp b/Engine/Source/Runtime/Engine/Private/AI/Navigation/NavLinkCustomComponent.cpp
--- a/Engine/Source/Runtime/Engine/Private/AI/Navigation/NavLinkCustomComponent.cpp
+++ b/Engine/Source/Runtime/Engine/Private/AI/Navigation/NavLinkCustomComponent.cpp
@@ -115,18 +115,18 @@ void UNavLinkCustomComponent::OnRegister()
 {
 	Super::OnRegister();
 
-	UNavigationSystem* NavSys = UNavigationSystem::GetCurrent(GetWorld());
-	if (NavSys == NULL)
+	if (NavLinkUserId == 0)
 	{
-		return;
+		NavLinkUserId = INavLinkCustomInterface::GetUniqueId();
 	}
 
-	if (NavLinkUserId == 0)
+	UNavigationSystem* NavSys = UNavigationSystem::GetCurrent(GetWorld());
+	if (NavSys == NULL)
 	{
-		NavLinkUserId = INavLinkCustomInterface::GetUniqueId();
+		return;
 	}
 
-	NavSys->RegisterCustomLink(this);
+	NavSys->RegisterCustomLink(*this);
 }
 
 void UNavLinkCustomComponent::OnUnregister()
@@ -140,7 +140,7 @@ void UNavLinkCustomComponent::OnUnregister()
 	}
 
 	// always try to unregister, even if not relevant right now
-	NavSys->UnregisterCustomLink(this);
+	NavSys->UnregisterCustomLink(*this);
 }
 
 void UNavLinkCustomComponent::SetLinkData(const FVector& RelativeStart, const FVector& RelativeEnd, ENavLinkDirection::Type Direction)
diff --git a/Engine/Source/Runtime/Engine/Private/AI/Navigation/NavigationSystem.cpp b/Engine/Source/Runtime/Engine/Private/AI/Navigation/NavigationSystem.cpp
--- a/Engine/Source/Runtime/Engine/Private/AI/Navigation/NavigationSystem.cpp
+++ b/Engine/Source/Runtime/Engine/Private/AI/Navigation/NavigationSystem.cpp
@@ -551,6 +551,17 @@ void UNavigationSystem::OnWorldInitDone(FNavigationSystem::EMode Mode)
 				}
 			}
 		}
+		else if (World->GetLevels().Num() > 1)
+		{
+			// this is a bit of a @hack for the time being until we reorganize navigation data instances (4.9?)
+			// the point of this hack is to force editor-time navmesh rebuilding if we have any sublevels
+			// since currently there's no information regarding whether we already have all navigation built 
+			// for all the actors in sublevels
+			for (ANavigationData* NavData : NavDataSet)
+			{
+				NavData->MarkAsNeedingUpdate();
+			}
+		}
 	}
 }
 
@@ -1647,14 +1658,14 @@ void UNavigationSystem::UnregisterNavData(ANavigationData* NavData)
 	NavData->OnUnregistered();
 }
 
-void UNavigationSystem::RegisterCustomLink(INavLinkCustomInterface* CustomLink)
+void UNavigationSystem::RegisterCustomLink(INavLinkCustomInterface& CustomLink)
 {
-	CustomLinksMap.Add(CustomLink->GetLinkId(), CustomLink);
+	CustomLinksMap.Add(CustomLink.GetLinkId(), &CustomLink);
 }
 
-void UNavigationSystem::UnregisterCustomLink(INavLinkCustomInterface* CustomLink)
+void UNavigationSystem::UnregisterCustomLink(INavLinkCustomInterface& CustomLink)
 {
-	CustomLinksMap.Remove(CustomLink->GetLinkId());
+	CustomLinksMap.Remove(CustomLink.GetLinkId());
 }
 
 INavLinkCustomInterface* UNavigationSystem::GetCustomLink(uint32 UniqueLinkId) const
@@ -1947,7 +1958,6 @@ UNavigationSystem* UNavigationSystem::CreateNavigationSystem(UWorld* WorldOwner)
 
 void UNavigationSystem::InitializeForWorld(UWorld* World, FNavigationSystem::EMode Mode)
 {
-	DECLARE_SCOPE_CYCLE_COUNTER(TEXT("UNavigationSystem::InitializeForWorld"), STAT_NavSys_InitializeForWorld, STATGROUP_Navigation);
 	if (World)
 	{
 		UNavigationSystem* NavSys = World->GetNavigationSystem();
@@ -3482,3 +3492,22 @@ bool UNavigationSystem::DoesPathIntersectBox(const FNavigationPath* Path, const
 {
 	return Path != NULL && Path->DoesIntersectBox(Box, AgentLocation, StartingIndex);
 }
+
+//----------------------------------------------------------------------//
+// DEPRECATED
+//----------------------------------------------------------------------//
+void UNavigationSystem::RegisterCustomLink(INavLinkCustomInterface* CustomLink)
+{
+	if (CustomLink)
+	{
+		RegisterCustomLink(*CustomLink);
+	}
+}
+
+void UNavigationSystem::UnregisterCustomLink(INavLinkCustomInterface* CustomLink)
+{
+	if (CustomLink)
+	{
+		UnregisterCustomLink(*CustomLink);
+	}
+}
diff --git a/Engine/Source/Runtime/Engine/Private/AI/Navigation/RecastNavMesh.cpp b/Engine/Source/Runtime/Engine/Private/AI/Navigation/RecastNavMesh.cpp
--- a/Engine/Source/Runtime/Engine/Private/AI/Navigation/RecastNavMesh.cpp
+++ b/Engine/Source/Runtime/Engine/Private/AI/Navigation/RecastNavMesh.cpp
@@ -1479,6 +1479,8 @@ void ARecastNavMesh::OnStreamingLevelAdded(ULevel* InLevel)
 {
 	QUICK_SCOPE_CYCLE_COUNTER(STAT_RecastNavMesh_OnStreamingLevelAdded);
 	
+	bWantsUpdate = true;
+
 	if (SupportsStreaming() && GetWorld()->IsGameWorld())
 	{
 		URecastNavMeshDataChunk* NavDataChunk = GetNavigationDataChunk(InLevel);
diff --git a/Engine/Source/Runtime/Engine/Private/World.cpp b/Engine/Source/Runtime/Engine/Private/World.cpp
--- a/Engine/Source/Runtime/Engine/Private/World.cpp
+++ b/Engine/Source/Runtime/Engine/Private/World.cpp
@@ -763,7 +763,11 @@ void UWorld::InitWorld(const InitializationValues IVS)
 	// Prepare AI systems
 	if (IVS.bCreateNavigation)
 	{
-		UNavigationSystem::CreateNavigationSystem(this);
+		AWorldSettings* WorldSettings = GetWorldSettings();
+		if (WorldSettings->bEnableNavigationSystem)
+		{
+			UNavigationSystem::CreateNavigationSystem(this);
+		}
 	}
 
 	if (IVS.bCreateAISystem)
