Fix crash in SetMaterial on destructible mesh (PR-692)

diff --git a/Engine/Source/Runtime/Engine/Private/MeshComponent.cpp b/Engine/Source/Runtime/Engine/Private/MeshComponent.cpp
--- a/Engine/Source/Runtime/Engine/Private/MeshComponent.cpp
+++ b/Engine/Source/Runtime/Engine/Private/MeshComponent.cpp
@@ -51,9 +51,10 @@ void UMeshComponent::SetMaterial(int32 ElementIndex, UMaterialInterface* Materia
 			OverrideMaterials[ElementIndex] = Material;
 			MarkRenderStateDirty();
 
-			if (BodyInstance.IsValidBodyInstance())
+			FBodyInstance* BodyInst = GetBodyInstance();
+			if (BodyInst && BodyInst->IsValidBodyInstance())
 			{
-				BodyInstance.UpdatePhysicalMaterials();
+				BodyInst->UpdatePhysicalMaterials();
 			}
 		}
 	}
