fixed UE-7498 Launching or Stopping PIE when Running 'Effects Cave' and 'Elemental' Demos

diff --git a/Engine/Source/Runtime/Engine/Classes/Components/SkinnedMeshComponent.h b/Engine/Source/Runtime/Engine/Classes/Components/SkinnedMeshComponent.h
--- a/Engine/Source/Runtime/Engine/Classes/Components/SkinnedMeshComponent.h
+++ b/Engine/Source/Runtime/Engine/Classes/Components/SkinnedMeshComponent.h
@@ -264,6 +264,7 @@ public:
 	 * Allows adjusting the desired streaming distance of streaming textures that uses UV 0.
 	 * 1.0 is the default, whereas a higher value makes the textures stream in sooner from far away.
 	 * A lower value (0.0-1.0) makes the textures stream in later (you have to be closer).
+	 * Value can be < 0 (from legcay content, or code changes)
 	 */
 	UPROPERTY(EditAnywhere, AdvancedDisplay, BlueprintReadWrite, Category=SkeletalMesh)
 	float StreamingDistanceMultiplier;
diff --git a/Engine/Source/Runtime/Engine/Classes/Components/StaticMeshComponent.h b/Engine/Source/Runtime/Engine/Classes/Components/StaticMeshComponent.h
--- a/Engine/Source/Runtime/Engine/Classes/Components/StaticMeshComponent.h
+++ b/Engine/Source/Runtime/Engine/Classes/Components/StaticMeshComponent.h
@@ -170,8 +170,9 @@ class ENGINE_API UStaticMeshComponent : public UMeshComponent
 	 * Allows adjusting the desired streaming distance of streaming textures that uses UV 0.
 	 * 1.0 is the default, whereas a higher value makes the textures stream in sooner from far away.
 	 * A lower value (0.0-1.0) makes the textures stream in later (you have to be closer).
+	 * Value can be < 0 (from legcay content, or code changes)
 	 */
-	UPROPERTY(EditAnywhere, AdvancedDisplay, BlueprintReadWrite, Category=TextureStreaming, meta=(ToolTip="Allows adjusting the desired resolution of streaming textures that uses UV 0.  1.0 is the default, whereas a higher value increases the streamed-in resolution."))
+	UPROPERTY(EditAnywhere, AdvancedDisplay, BlueprintReadWrite, Category=TextureStreaming, meta=(ClampMin = 0, ToolTip="Allows adjusting the desired resolution of streaming textures that uses UV 0.  1.0 is the default, whereas a higher value increases the streamed-in resolution."))
 	float StreamingDistanceMultiplier;
 
 	/** Subdivision step size for static vertex lighting.				*/
diff --git a/Engine/Source/Runtime/Engine/Classes/Engine/SkeletalMesh.h b/Engine/Source/Runtime/Engine/Classes/Engine/SkeletalMesh.h
--- a/Engine/Source/Runtime/Engine/Classes/Engine/SkeletalMesh.h
+++ b/Engine/Source/Runtime/Engine/Classes/Engine/SkeletalMesh.h
@@ -627,8 +627,9 @@ public:
 	/**
 	 * Allows artists to adjust the distance where textures using UV 0 are streamed in/out.
 	 * 1.0 is the default, whereas a higher value increases the streamed-in resolution.
+	 * Value can be < 0 (from legcay content, or code changes)
 	 */
-	UPROPERTY(EditAnywhere, BlueprintReadOnly, Category=TextureStreaming)
+	UPROPERTY(EditAnywhere, BlueprintReadOnly, Category=TextureStreaming, meta=(ClampMin = 0))
 	float StreamingDistanceMultiplier;
 
 	UPROPERTY(Category=Mesh, BlueprintReadWrite)
diff --git a/Engine/Source/Runtime/Engine/Classes/Engine/StaticMesh.h b/Engine/Source/Runtime/Engine/Classes/Engine/StaticMesh.h
--- a/Engine/Source/Runtime/Engine/Classes/Engine/StaticMesh.h
+++ b/Engine/Source/Runtime/Engine/Classes/Engine/StaticMesh.h
@@ -353,8 +353,9 @@ class UStaticMesh : public UObject, public IInterface_CollisionDataProvider, pub
 	/**
 	 * Allows artists to adjust the distance where textures using UV 0 are streamed in/out.
 	 * 1.0 is the default, whereas a higher value increases the streamed-in resolution.
+	 * Value can be < 0 (from legcay content, or code changes)
 	 */
-	UPROPERTY(EditAnywhere, AdvancedDisplay, Category=StaticMesh)
+	UPROPERTY(EditAnywhere, AdvancedDisplay, Category=StaticMesh, meta=(ClampMin = 0))
 	float StreamingDistanceMultiplier;
 
 	/** Bias multiplier for Light Propagation Volume lighting */
@@ -371,8 +372,6 @@ class UStaticMesh : public UObject, public IInterface_CollisionDataProvider, pub
 	FString HighResSourceMeshName;
 
 #if WITH_EDITORONLY_DATA
-
-
 	/** Importing data and options used for this mesh */
 	UPROPERTY(VisibleAnywhere, Instanced, Category=ImportSettings)
 	class UAssetImportData* AssetImportData;
diff --git a/Engine/Source/Runtime/Engine/Private/SkeletalMesh.cpp b/Engine/Source/Runtime/Engine/Private/SkeletalMesh.cpp
--- a/Engine/Source/Runtime/Engine/Private/SkeletalMesh.cpp
+++ b/Engine/Source/Runtime/Engine/Private/SkeletalMesh.cpp
@@ -1932,8 +1932,9 @@ void USkeletalMesh::ReleaseResources()
 	ReleaseResourcesFence.BeginFence();
 }
 
+// @param InStreamingDistanceMultiplier should be >= 0
 template <bool bExtraBoneInfluencesT>
-static void GetStreamingTextureFactorForLOD(FStaticLODModel& LODModel, TArray<float>& CachedStreamingTextureFactors, float StreamingDistanceMultiplier)
+static void GetStreamingTextureFactorForLOD(FStaticLODModel& LODModel, TArray<float>& CachedStreamingTextureFactors, float InStreamingDistanceMultiplier)
 {
 	int32 NumTotalTriangles = LODModel.GetTotalFaces();
 
@@ -2003,7 +2004,7 @@ static void GetStreamingTextureFactorForLOD(FStaticLODModel& LODModel, TArray<fl
 					float TexelRatio = TexelRatios[UVIndex][ FMath::TruncToInt(TexelRatios[UVIndex].Num() * 0.75f) ];
 					if ( UVIndex == 0 )
 					{
-						TexelRatio *= StreamingDistanceMultiplier;
+						TexelRatio *= InStreamingDistanceMultiplier;
 					}
 					CachedStreamingTextureFactors[UVIndex] = TexelRatio;
 				}
@@ -2035,11 +2036,11 @@ float USkeletalMesh::GetStreamingTextureFactor( int32 RequestedUVIndex )
 			FStaticLODModel& LODModel = Resource->LODModels[0];
 			if (LODModel.DoesVertexBufferHaveExtraBoneInfluences())
 			{
-				GetStreamingTextureFactorForLOD<true>(LODModel, CachedStreamingTextureFactors, StreamingDistanceMultiplier);
+				GetStreamingTextureFactorForLOD<true>(LODModel, CachedStreamingTextureFactors, FMath::Max(0.0f, StreamingDistanceMultiplier));
 			}
 			else
 			{
-				GetStreamingTextureFactorForLOD<false>(LODModel, CachedStreamingTextureFactors, StreamingDistanceMultiplier);
+				GetStreamingTextureFactorForLOD<false>(LODModel, CachedStreamingTextureFactors, FMath::Max(0.0f, StreamingDistanceMultiplier));
 			}
 		}
 		else
diff --git a/Engine/Source/Runtime/Engine/Private/SkinnedMeshComponent.cpp b/Engine/Source/Runtime/Engine/Private/SkinnedMeshComponent.cpp
--- a/Engine/Source/Runtime/Engine/Private/SkinnedMeshComponent.cpp
+++ b/Engine/Source/Runtime/Engine/Private/SkinnedMeshComponent.cpp
@@ -364,7 +364,7 @@ void USkinnedMeshComponent::GetStreamingTextureInfo(TArray<FStreamingTexturePrim
 {
 	if( SkeletalMesh )
 	{
-		const float LocalTexelFactor = SkeletalMesh->GetStreamingTextureFactor(0) * StreamingDistanceMultiplier;
+		const float LocalTexelFactor = SkeletalMesh->GetStreamingTextureFactor(0) * FMath::Max(0.0f, StreamingDistanceMultiplier);
 		const float WorldTexelFactor = LocalTexelFactor * ComponentToWorld.GetMaximumAxisScale();
 		const int32 NumMaterials = FMath::Max(SkeletalMesh->Materials.Num(), OverrideMaterials.Num());
 		for( int32 MatIndex = 0; MatIndex < NumMaterials; MatIndex++ )
diff --git a/Engine/Source/Runtime/Engine/Private/StaticMesh.cpp b/Engine/Source/Runtime/Engine/Private/StaticMesh.cpp
--- a/Engine/Source/Runtime/Engine/Private/StaticMesh.cpp
+++ b/Engine/Source/Runtime/Engine/Private/StaticMesh.cpp
@@ -1337,11 +1337,11 @@ float UStaticMesh::GetStreamingTextureFactor(int32 RequestedUVIndex)
 	{
 		if( bUseMaximumStreamingTexelRatio )
 		{
-			StreamingTextureFactor = RenderData->MaxStreamingTextureFactor * StreamingDistanceMultiplier;
+			StreamingTextureFactor = RenderData->MaxStreamingTextureFactor * FMath::Max(0.0f, StreamingDistanceMultiplier);
 		}
 		else if( RequestedUVIndex == 0 )
 		{
-			StreamingTextureFactor = RenderData->StreamingTextureFactors[RequestedUVIndex] * StreamingDistanceMultiplier;
+			StreamingTextureFactor = RenderData->StreamingTextureFactors[RequestedUVIndex] * FMath::Max(0.0f, StreamingDistanceMultiplier);
 		}
 		else
 		{
diff --git a/Engine/Source/Runtime/Engine/Private/StaticMeshComponent.cpp b/Engine/Source/Runtime/Engine/Private/StaticMeshComponent.cpp
--- a/Engine/Source/Runtime/Engine/Private/StaticMeshComponent.cpp
+++ b/Engine/Source/Runtime/Engine/Private/StaticMeshComponent.cpp
@@ -516,7 +516,7 @@ void UStaticMeshComponent::GetStreamingTextureInfo(TArray<FStreamingTexturePrimi
 		}
 
 		const FSphere BoundingSphere	= Bounds.GetSphere();
-		const float LocalTexelFactor	= StaticMesh->GetStreamingTextureFactor(0) * StreamingDistanceMultiplier;
+		const float LocalTexelFactor	= StaticMesh->GetStreamingTextureFactor(0) * FMath::Max(0.0f, StreamingDistanceMultiplier);
 		const float LocalLightmapFactor	= bHasValidLightmapCoordinates ? StaticMesh->GetStreamingTextureFactor(StaticMesh->LightMapCoordinateIndex) : 1.0f;
 		const float WorldTexelFactor	= SplineDeformFactor * LocalTexelFactor * ComponentToWorld.GetMaximumAxisScale();
 		const float WorldLightmapFactor	= SplineDeformFactor * LocalLightmapFactor * ComponentToWorld.GetMaximumAxisScale();
diff --git a/Engine/Source/Runtime/Landscape/Classes/LandscapeComponent.h b/Engine/Source/Runtime/Landscape/Classes/LandscapeComponent.h
--- a/Engine/Source/Runtime/Landscape/Classes/LandscapeComponent.h
+++ b/Engine/Source/Runtime/Landscape/Classes/LandscapeComponent.h
@@ -131,7 +131,7 @@ class ULandscapeComponent : public UPrimitiveComponent
 	UPROPERTY(VisibleAnywhere, BlueprintReadOnly, Category=LandscapeComponent)
 	int32 SectionBaseY;
 
-	/** Total number of quads for this component */
+	/** Total number of quads for this component, has to be >0 */
 	UPROPERTY()
 	int32 ComponentSizeQuads;
 
diff --git a/Engine/Source/Runtime/Landscape/Classes/LandscapeProxy.h b/Engine/Source/Runtime/Landscape/Classes/LandscapeProxy.h
--- a/Engine/Source/Runtime/Landscape/Classes/LandscapeProxy.h
+++ b/Engine/Source/Runtime/Landscape/Classes/LandscapeProxy.h
@@ -313,8 +313,9 @@ public:
 	/**
 	 * Allows artists to adjust the distance where textures using UV 0 are streamed in/out.
 	 * 1.0 is the default, whereas a higher value increases the streamed-in resolution.
+	 * Value can be < 0 (from legcay content, or code changes)
 	 */
-	UPROPERTY(EditAnywhere, Category=Landscape, meta=(ClampMin=0))
+	UPROPERTY(EditAnywhere, Category=Landscape)
 	float StreamingDistanceMultiplier;
 
 	/** Combined material used to render the landscape */
diff --git a/Engine/Source/Runtime/Landscape/Private/LandscapeRender.cpp b/Engine/Source/Runtime/Landscape/Private/LandscapeRender.cpp
--- a/Engine/Source/Runtime/Landscape/Private/LandscapeRender.cpp
+++ b/Engine/Source/Runtime/Landscape/Private/LandscapeRender.cpp
@@ -2318,12 +2318,12 @@ void ULandscapeComponent::GetStreamingTextureInfo(TArray<FStreamingTexturePrimit
 {
 	ALandscapeProxy* Proxy = Cast<ALandscapeProxy>(GetOuter());
 	FSphere BoundingSphere = Bounds.GetSphere();
-	float StreamingDistanceMultiplier = 1.f;
+	float LocalStreamingDistanceMultiplier = 1.f;
 	if (Proxy)
 	{
-		StreamingDistanceMultiplier = Proxy->StreamingDistanceMultiplier;
+		LocalStreamingDistanceMultiplier = FMath::Max(0.0f, Proxy->StreamingDistanceMultiplier);
 	}
-	const float TexelFactor = 0.75f * StreamingDistanceMultiplier * ComponentSizeQuads * FMath::Abs(Proxy->GetRootComponent()->RelativeScale3D.X);
+	const float TexelFactor = 0.75f * LocalStreamingDistanceMultiplier * ComponentSizeQuads * FMath::Abs(Proxy->GetRootComponent()->RelativeScale3D.X);
 
 	// Normal usage...
 	// Enumerate the textures used by the material.
