Cleanup logging

diff --git a/Engine/Source/Runtime/Engine/Private/PackageMapClient.cpp b/Engine/Source/Runtime/Engine/Private/PackageMapClient.cpp
--- a/Engine/Source/Runtime/Engine/Private/PackageMapClient.cpp
+++ b/Engine/Source/Runtime/Engine/Private/PackageMapClient.cpp
@@ -880,21 +880,15 @@ bool UPackageMapClient::AppendExportBunches(TArray<FOutBunch *>& OutgoingBunches
 	// Append the bunches we've made to the passed in list reference
 	if (ExportBunches.Num() > 0)
 	{
-		if (UE_LOG_ACTIVE(LogNetPackageMap,Log))
+		if (UE_LOG_ACTIVE(LogNetPackageMap,Verbose))
 		{
-			UE_LOG(LogNetPackageMap, Log, TEXT("AppendExportBunches. %d Bunches  %d"), ExportBunches.Num(), ExportNetGUIDCount);
+			UE_LOG(LogNetPackageMap, Verbose, TEXT("AppendExportBunches. ExportBunches: %d, ExportNetGUIDCount: %d"), ExportBunches.Num(), ExportNetGUIDCount);
 			for (auto It=ExportBunches.CreateIterator(); It; ++It)
 			{
-				UE_LOG(LogNetPackageMap, Log, TEXT("   Bunch[%d]. NumBytes: %d NumBits: %d"), It.GetIndex(), (*It)->GetNumBytes(), (*It)->GetNumBits() );
+				UE_LOG(LogNetPackageMap, Verbose, TEXT("   BunchIndex: %d, ExportNetGUIDs: %d, NumBytes: %d, NumBits: %d"), It.GetIndex(), (*It)->ExportNetGUIDs.Num(), (*It)->GetNumBytes(), (*It)->GetNumBits() );
 			}
 		}
 
-		for (auto It = ExportBunches.CreateIterator(); It; ++It)
-		{
-			UE_LOG(LogNetPackageMap, Log, TEXT("Bunch[%d] Num ExportNetGUIDs: %d"), It.GetIndex(), (*It)->ExportNetGUIDs.Num() );
-		}
-
-
 		OutgoingBunches.Append(ExportBunches);
 		ExportBunches.Empty();
 		return true;
