UE4 - bypass almost everything for grass and send the final instance buffer, write once, directly where it has to go.

diff --git a/Engine/Source/Runtime/Engine/Classes/Components/HierarchicalInstancedStaticMeshComponent.h b/Engine/Source/Runtime/Engine/Classes/Components/HierarchicalInstancedStaticMeshComponent.h
--- a/Engine/Source/Runtime/Engine/Classes/Components/HierarchicalInstancedStaticMeshComponent.h
+++ b/Engine/Source/Runtime/Engine/Classes/Components/HierarchicalInstancedStaticMeshComponent.h
@@ -3,6 +3,7 @@
 #pragma once
 
 #include "Components/InstancedStaticMeshComponent.h"
+#include "StaticMeshResources.h"
 
 #include "HierarchicalInstancedStaticMeshComponent.generated.h"
 
@@ -57,6 +58,8 @@ class ENGINE_API UHierarchicalInstancedStaticMeshComponent : public UInstancedSt
 	GENERATED_UCLASS_BODY()
 
 	TSharedPtr<TArray<FClusterNode>, ESPMode::ThreadSafe> ClusterTreePtr;
+
+	FStaticMeshInstanceData WriteOncePrebuiltInstanceBuffer;
 	
 	// Table for remaping instances from cluster tree to PerInstanceSMData order
 	UPROPERTY()
@@ -99,18 +102,14 @@ public:
 	void BuildTree();
 	void BuildTreeAsync();
 	static void BuildTreeAnyThread(
-		const TArray<FInstancedStaticMeshInstanceData>& PerInstanceSMData, 
+		TArray<FMatrix>& InstanceTransforms, 
 		const FBox& MeshBox,
 		TArray<FClusterNode>& OutClusterTree,
 		TArray<int32>& OutSortedInstances,
 		TArray<int32>& OutInstanceReorderTable,
 		int32 MaxInstancesPerLeaf
 		);
-	void AcceptPrebuiltTree(
-		TArray<FClusterNode>& InClusterTree,
-		TArray<int32>& InSortedInstances,
-		TArray<int32>& InInstanceReorderTable
-		);
+	void AcceptPrebuiltTree(TArray<FClusterNode>& InClusterTree);
 	void BuildFlatTree(const TArray<int32>& LeafInstanceCounts);
 	bool IsAsyncBuilding() const { return bIsAsyncBuilding; }
 	bool IsTreeFullyBuilt() const { return NumBuiltInstances == PerInstanceSMData.Num() && RemovedInstances.Num() == 0; }
diff --git a/Engine/Source/Runtime/Engine/Classes/Components/InstancedStaticMeshComponent.h b/Engine/Source/Runtime/Engine/Classes/Components/InstancedStaticMeshComponent.h
--- a/Engine/Source/Runtime/Engine/Classes/Components/InstancedStaticMeshComponent.h
+++ b/Engine/Source/Runtime/Engine/Classes/Components/InstancedStaticMeshComponent.h
@@ -90,6 +90,9 @@ class ENGINE_API UInstancedStaticMeshComponent : public UStaticMeshComponent
 	UPROPERTY()
 	TArray<int32> RemovedInstances;
 
+	/** Tracks outstanding proxysize, as this is a bit hard to do with the fire-and-forget grass. */
+	SIZE_T ProxySize;
+
 	/** Add an instance to this component. Transform is given in local space of this component.  */
 	UFUNCTION(BlueprintCallable, Category="Components|InstancedStaticMesh")
 	virtual int32 AddInstance(const FTransform& InstanceTransform);
diff --git a/Engine/Source/Runtime/Engine/Private/HierarchicalInstancedStaticMesh.cpp b/Engine/Source/Runtime/Engine/Private/HierarchicalInstancedStaticMesh.cpp
--- a/Engine/Source/Runtime/Engine/Private/HierarchicalInstancedStaticMesh.cpp
+++ b/Engine/Source/Runtime/Engine/Private/HierarchicalInstancedStaticMesh.cpp
@@ -599,6 +599,19 @@ public:
 		check(InComponent->InstanceReorderTable.Num() == InComponent->PerInstanceSMData.Num());
 	}
 
+	FHierarchicalStaticMeshSceneProxy(UHierarchicalInstancedStaticMeshComponent* InComponent, ERHIFeatureLevel::Type InFeatureLevel, FStaticMeshInstanceData& Other)
+		: FInstancedStaticMeshSceneProxy(InComponent, InFeatureLevel, Other)
+		, ClusterTreePtr(InComponent->ClusterTreePtr)
+		, ClusterTree(*InComponent->ClusterTreePtr)
+		, FirstUnbuiltIndex(InComponent->NumBuiltInstances)
+		, LastUnbuiltIndex(InstancedRenderData.NumInstances+InComponent->RemovedInstances.Num()-1)
+#if !(UE_BUILD_SHIPPING || UE_BUILD_TEST)
+		, CaptureTag(0)
+#endif
+	{
+		check(!InComponent->InstanceReorderTable.Num() && !InComponent->PerInstanceSMData.Num());
+	}
+
 	// FPrimitiveSceneProxy interface.
 	
 	virtual FPrimitiveViewRelevance GetViewRelevance(const FSceneView* View) override
@@ -1284,19 +1297,8 @@ void FHierarchicalStaticMeshSceneProxy::GetDynamicMeshElements(const TArray<cons
 
 FBoxSphereBounds UHierarchicalInstancedStaticMeshComponent::CalcBounds(const FTransform& BoundTransform) const
 {
-	const bool bMeshIsValid =
-		// make sure we have instances
-		PerInstanceSMData.Num() > 0 &&
-		// make sure we have an actual staticmesh
-		StaticMesh &&
-		StaticMesh->HasValidRenderData() &&
-		// You really can't use hardware instancing on the consoles with multiple elements because they share the same index buffer. 
-		// @todo: Level error or something to let LDs know this
-		1;//StaticMesh->LODModels(0).Elements.Num() == 1;
-
 	const TArray<FClusterNode>& ClusterTree = *ClusterTreePtr;
-
-	if ((ClusterTree.Num() || UnbuiltInstanceBounds.IsValid) && bMeshIsValid)
+	if (ClusterTree.Num() || UnbuiltInstanceBounds.IsValid)
 	{
 		FBoxSphereBounds Result;
 		if (ClusterTree.Num())
@@ -1536,7 +1538,7 @@ void UHierarchicalInstancedStaticMeshComponent::BuildTree()
 }
 
 void UHierarchicalInstancedStaticMeshComponent::BuildTreeAnyThread(
-	const TArray<FInstancedStaticMeshInstanceData>& PerInstanceSMData, 
+	TArray<FMatrix>& InstanceTransforms, 
 	const FBox& MeshBox,
 	TArray<FClusterNode>& OutClusterTree,
 	TArray<int32>& OutSortedInstances,
@@ -1544,13 +1546,6 @@ void UHierarchicalInstancedStaticMeshComponent::BuildTreeAnyThread(
 	int32 MaxInstancesPerLeaf
 	)
 {
-	TArray<FMatrix> InstanceTransforms;
-	InstanceTransforms.AddUninitialized(PerInstanceSMData.Num());
-	for (int32 Index = 0; Index < PerInstanceSMData.Num(); Index++)
-	{
-		InstanceTransforms[Index] = PerInstanceSMData[Index].Transform;
-	}
-
 	TUniquePtr<FClusterBuilder> Builder(new FClusterBuilder(InstanceTransforms, MeshBox, MaxInstancesPerLeaf));
 	Builder->Build();
 
@@ -1560,13 +1555,12 @@ void UHierarchicalInstancedStaticMeshComponent::BuildTreeAnyThread(
 
 }
 
-void UHierarchicalInstancedStaticMeshComponent::AcceptPrebuiltTree(
-	TArray<FClusterNode>& InClusterTree,
-	TArray<int32>& InSortedInstances,
-	TArray<int32>& InInstanceReorderTable
-	)
+void UHierarchicalInstancedStaticMeshComponent::AcceptPrebuiltTree(TArray<FClusterNode>& InClusterTree)
 {
-	NumBuiltInstances = PerInstanceSMData.Num();
+	// this is only for prebuild data, already in the correct order
+	check(!PerInstanceSMData.Num());
+	NumBuiltInstances = WriteOncePrebuiltInstanceBuffer.GetNumInstances();
+	check(NumBuiltInstances);
 	UnbuiltInstanceBounds.Init();
 	RemovedInstances.Empty();
 	ClusterTreePtr = MakeShareable(new TArray<FClusterNode>);
@@ -1576,7 +1570,7 @@ void UHierarchicalInstancedStaticMeshComponent::AcceptPrebuiltTree(
 	// Verify that the mesh is valid before using it.
 	const bool bMeshIsValid = 
 		// make sure we have instances
-		PerInstanceSMData.Num() > 0 &&
+		NumBuiltInstances > 0 &&
 		// make sure we have an actual staticmesh
 		StaticMesh &&
 		StaticMesh->HasValidRenderData() &&
@@ -1594,13 +1588,7 @@ void UHierarchicalInstancedStaticMeshComponent::AcceptPrebuiltTree(
 		{
 			InstancingRandomSeed = FMath::Rand();
 		}
-
 		Exchange(*ClusterTreePtr, InClusterTree);
-		Exchange(InstanceReorderTable, InInstanceReorderTable);
-		Exchange(SortedInstances, InSortedInstances);
-
-		FlushAccumulatedNavigationUpdates();
-
 	}
 	MarkRenderStateDirty();
 }
@@ -1807,10 +1795,12 @@ FPrimitiveSceneProxy* UHierarchicalInstancedStaticMeshComponent::CreateSceneProx
 {
 	QUICK_SCOPE_CYCLE_COUNTER(STAT_HierarchicalInstancedStaticMeshComponent_CreateSceneProxy);
 
+	ProxySize = 0;
+
 	// Verify that the mesh is valid before using it.
 	const bool bMeshIsValid = 
 		// make sure we have instances
-		PerInstanceSMData.Num() > 0 &&
+		(PerInstanceSMData.Num() > 0 || WriteOncePrebuiltInstanceBuffer.GetNumInstances() > 0)&&
 		// make sure we have an actual staticmesh
 		StaticMesh &&
 		StaticMesh->HasValidRenderData() &&
@@ -1828,6 +1818,15 @@ FPrimitiveSceneProxy* UHierarchicalInstancedStaticMeshComponent::CreateSceneProx
 			InstancingRandomSeed = FMath::Rand();
 		}
 
+		if (WriteOncePrebuiltInstanceBuffer.GetNumInstances() > 0)
+		{
+			ProxySize = FStaticMeshInstanceData::GetResourceSize(WriteOncePrebuiltInstanceBuffer.GetNumInstances());
+
+			//@todo - this is sketchy, what you really want to do is flush the grass cache instead of doing a render thread update
+			bNeverNeedsRenderUpdate = true; // we can never update it because the data is gone
+			return ::new FHierarchicalStaticMeshSceneProxy(this, GetWorld()->FeatureLevel, WriteOncePrebuiltInstanceBuffer);
+		}
+		ProxySize = FStaticMeshInstanceData::GetResourceSize(PerInstanceSMData.Num());
 		return ::new FHierarchicalStaticMeshSceneProxy(this, GetWorld()->FeatureLevel);
 	}
 	return NULL;
diff --git a/Engine/Source/Runtime/Engine/Private/InstancedStaticMesh.cpp b/Engine/Source/Runtime/Engine/Private/InstancedStaticMesh.cpp
--- a/Engine/Source/Runtime/Engine/Private/InstancedStaticMesh.cpp
+++ b/Engine/Source/Runtime/Engine/Private/InstancedStaticMesh.cpp
@@ -23,6 +23,7 @@ void HInstancedStaticMeshInstance::AddReferencedObjects(FReferenceCollector& Col
 
 FStaticMeshInstanceBuffer::FStaticMeshInstanceBuffer(ERHIFeatureLevel::Type InFeatureLevel):
 	InstanceData(NULL)
+	, Stride(0)
 {
 	SetFeatureLevel(InFeatureLevel);
 }
@@ -62,33 +63,30 @@ static void FillInstanceRenderData(const FInstancedStaticMeshInstanceData& Insta
 	RenderData[6] = FVector4(Inverse.M[0][2], Inverse.M[1][2], Inverse.M[2][2], RandomInstanceID);
 }
 
-/**
- * Initializes the buffer with the component's data.
- * @param InComponent - The owning component
- */
+void FStaticMeshInstanceBuffer::SetupCPUAccess(UInstancedStaticMeshComponent* InComponent)
+{
+	static const auto CVar = IConsoleManager::Get().FindTConsoleVariableDataInt(TEXT("r.GenerateMeshDistanceFields"));
+
+	const bool bNeedsCPUAccess = InComponent->CastShadow && InComponent->bAffectDistanceFieldLighting 
+		// Distance field algorithms need access to instance data on the CPU
+		&& CVar->GetValueOnGameThread() != 0;
+
+	InstanceData->SetAllowCPUAccess(InstanceData->GetAllowCPUAccess() || bNeedsCPUAccess);
+}
+
 void FStaticMeshInstanceBuffer::Init(UInstancedStaticMeshComponent* InComponent, const TArray<TRefCountPtr<HHitProxy> >& InHitProxies)
 {
 	bool bUseRemapTable = InComponent->PerInstanceSMData.Num() == InComponent->InstanceReorderTable.Num();
 
-	NumInstances = InComponent->PerInstanceSMData.Num();
+	int32 NumInstances = InComponent->PerInstanceSMData.Num();
 	int32 NumRemoved = InComponent->RemovedInstances.Num();
 
 	// Allocate the vertex data storage type.
 	AllocateData();
 
-	static const auto CVar = IConsoleManager::Get().FindTConsoleVariableDataInt(TEXT("r.GenerateMeshDistanceFields"));
+	SetupCPUAccess(InComponent);
 
-	const bool bNeedsCPUAccess = InComponent->CastShadow && InComponent->bAffectDistanceFieldLighting 
-		// Distance field algorithms need access to instance data on the CPU
-		&& CVar->GetValueOnGameThread() != 0;
-
-	InstanceData->SetAllowCPUAccess(InstanceData->GetAllowCPUAccess() || bNeedsCPUAccess);
-
-	// We cannot write directly to the data on all platforms,
-	// so we make a TArray of the right type, then assign it
-	check( GetStride() % sizeof(FVector4) == 0 );
-	InstanceData->Empty((NumInstances + NumRemoved) * GetStride() / sizeof(FVector4));
-	InstanceData->AddUninitialized((NumInstances + NumRemoved) * GetStride() / sizeof(FVector4));
+	InstanceData->AllocateInstances(NumInstances + NumRemoved);
 
 	// @todo: Make LD-customizable per component?
 	const float RandomInstanceIDBase = 0.0f;
@@ -102,7 +100,7 @@ void FStaticMeshInstanceBuffer::Init(UInstancedStaticMeshComponent* InComponent,
 	{
 		const FInstancedStaticMeshInstanceData& Instance = InComponent->PerInstanceSMData[InstanceIndex];
 
-		FVector4* InstanceRenderData = InstanceData->GetData() + (bUseRemapTable ? InComponent->InstanceReorderTable[InstanceIndex] : InstanceIndex) * (GetStride() / sizeof(FVector4));
+		FVector4* InstanceRenderData = InstanceData->GetInstanceWriteAddress(bUseRemapTable ? InComponent->InstanceReorderTable[InstanceIndex] : InstanceIndex);
 
 		// X, Y	: Shadow map UV bias
 		// Z, W : Encoded HitProxy ID and selection mask.
@@ -151,10 +149,18 @@ void FStaticMeshInstanceBuffer::Init(UInstancedStaticMeshComponent* InComponent,
 	}
 }
 
+void FStaticMeshInstanceBuffer::InitFromPreallocatedData(UInstancedStaticMeshComponent* InComponent, FStaticMeshInstanceData& Other)
+{
+	AllocateData(Other);
+	SetupCPUAccess(InComponent);
+}
+
+
 /** Serializer. */
 FArchive& operator<<(FArchive& Ar,FStaticMeshInstanceBuffer& InstanceBuffer)
 {
-	Ar << InstanceBuffer.Stride << InstanceBuffer.NumInstances;
+	int32 NumInstances = InstanceBuffer.GetNumInstances();
+	Ar << InstanceBuffer.Stride << NumInstances;
 
 	if(Ar.IsLoading())
 	{
@@ -202,6 +208,12 @@ void FStaticMeshInstanceBuffer::AllocateData()
 	Stride = InstanceData->GetStride();
 }
 
+void FStaticMeshInstanceBuffer::AllocateData(FStaticMeshInstanceData& Other)
+{
+	AllocateData();
+	Other.SetAllowCPUAccess(InstanceData->GetAllowCPUAccess());
+	Exchange(Other, *InstanceData);
+}
 
 
 
@@ -824,6 +836,7 @@ void UInstancedStaticMeshComponent::ApplyComponentInstanceData(FComponentInstanc
 
 FPrimitiveSceneProxy* UInstancedStaticMeshComponent::CreateSceneProxy()
 {
+	ProxySize = 0;
 	// Verify that the mesh is valid before using it.
 	const bool bMeshIsValid = 
 		// make sure we have instances
@@ -845,6 +858,7 @@ FPrimitiveSceneProxy* UInstancedStaticMeshComponent::CreateSceneProxy()
 			InstancingRandomSeed = FMath::Rand();
 		}
 
+		ProxySize = FStaticMeshInstanceData::GetResourceSize(PerInstanceSMData.Num());
 		return ::new FInstancedStaticMeshSceneProxy(this, GetWorld()->FeatureLevel);
 	}
 	else
@@ -1300,6 +1314,8 @@ void UInstancedStaticMeshComponent::ClearInstances()
 	InstanceReorderTable.Empty();
 	RemovedInstances.Empty();
 
+	ProxySize = 0;
+
 	// Release any physics representations
 	ClearAllInstanceBodies();
 
@@ -1424,7 +1440,7 @@ SIZE_T UInstancedStaticMeshComponent::GetResourceSize( EResourceSizeMode::Type M
 	SIZE_T ResSize = Super::GetResourceSize(Mode);
 
 	// proxy stuff
-	ResSize += FStaticMeshInstanceData::GetResourceSize(PerInstanceSMData.Num());
+	ResSize += ProxySize; 
 
 	// component stuff
 	ResSize += InstanceBodies.GetAllocatedSize();
diff --git a/Engine/Source/Runtime/Engine/Private/InstancedStaticMesh.h b/Engine/Source/Runtime/Engine/Private/InstancedStaticMesh.h
--- a/Engine/Source/Runtime/Engine/Private/InstancedStaticMesh.h
+++ b/Engine/Source/Runtime/Engine/Private/InstancedStaticMesh.h
@@ -27,7 +27,6 @@
 #include "Components/TimelineComponent.h"
 #include "SlateBasics.h"
 #include "NavDataGenerator.h"
-#include "OnlineSubsystemUtils.h"
 #include "AI/Navigation/RecastHelpers.h"
 
 #include "StaticMeshResources.h"
@@ -61,80 +60,10 @@ extern TAutoConsoleVariable<float> CVarFoliageLODDistanceScale;
 // to look up a texture coordinate we didn't provide an element for.
 extern const int32 InstancedStaticMeshMaxTexCoord;
 
-/*-----------------------------------------------------------------------------
-	FStaticMeshInstanceData
------------------------------------------------------------------------------*/
-
-/** The implementation of the static mesh instance data storage type. */
-class FStaticMeshInstanceData :
-	public FStaticMeshVertexDataInterface,
-	public TResourceArray<FVector4,VERTEXBUFFER_ALIGNMENT>
-{
-public:
-
-	enum 
-	{
-		VectorsPerInstance = 7
-	};
-	typedef TResourceArray<FVector4,VERTEXBUFFER_ALIGNMENT> ArrayType;
-
-	/**
-	 * Constructor
-	 * @param InNeedsCPUAccess - true if resource array data should be CPU accessible
-	 */
-	FStaticMeshInstanceData(bool InNeedsCPUAccess=false)
-		:	TResourceArray<FVector4,VERTEXBUFFER_ALIGNMENT>(InNeedsCPUAccess)
-	{
-	}
-
-	static uint32 StaticGetStride()
-	{
-		return sizeof(FVector4) * VectorsPerInstance;
-	}
-
-	static SIZE_T GetResourceSize(uint32 NumInstances)
-	{
-		return SIZE_T(NumInstances) * SIZE_T(StaticGetStride());
-	}
-
-	/**
-	 * Resizes the vertex data buffer, discarding any data which no longer fits.
-	 * @param NumVertices - The number of vertices to allocate the buffer for.
-	 */
-	virtual void ResizeBuffer(uint32 NumInstances)
-	{
-		checkf(0, TEXT("ArrayType::Add is not supported on all platforms"));
-	}
-
-	virtual uint32 GetStride() const
-	{
-		return StaticGetStride();
-	}
-	virtual uint8* GetDataPointer()
-	{
-		return (uint8*)&(*this)[0];
-	}
-	virtual FResourceArrayInterface* GetResourceArray()
-	{
-		return this;
-	}
-	virtual void Serialize(FArchive& Ar)
-	{
-		TResourceArray<FVector4,VERTEXBUFFER_ALIGNMENT>::BulkSerialize(Ar);
-	}
-
-	void Set(const TArray<FVector4>& RawData)
-	{
-		*((ArrayType*)this) = TArray<FVector4,TAlignedHeapAllocator<VERTEXBUFFER_ALIGNMENT> >(RawData);
-	}
-};
-
-
 /*-----------------------------------------------------------------------------
 	FStaticMeshInstanceBuffer
 -----------------------------------------------------------------------------*/
 
-
 /** A vertex buffer of positions. */
 class FStaticMeshInstanceBuffer : public FVertexBuffer
 {
@@ -156,6 +85,13 @@ public:
 	 */
 	void Init(UInstancedStaticMeshComponent* InComponent, const TArray<TRefCountPtr<HHitProxy> >& InHitProxies);
 
+	/**
+	 * Initializes the buffer with the component's data.
+	 * @param InComponent - The owning component; this need not have PerInstanceSMData, as we are taking a prebuilt instance buffer
+	 * @param Other - instance data, this call assumes the memory, so this will be empty after the call
+	 */
+	void InitFromPreallocatedData(UInstancedStaticMeshComponent* InComponent, FStaticMeshInstanceData& Other);
+
 	/** Serializer. */
 	friend FArchive& operator<<(FArchive& Ar, FStaticMeshInstanceBuffer& VertexBuffer);
 
@@ -171,7 +107,7 @@ public:
 	}
 	FORCEINLINE uint32 GetNumInstances() const
 	{
-		return NumInstances;
+		return InstanceData ? InstanceData->GetNumInstances() : 0;
 	}
 
 	const void* GetRawData() const
@@ -191,11 +127,13 @@ private:
 	/** The cached vertex stride. */
 	uint32 Stride;
 
-	/** The cached number of instances. */
-	uint32 NumInstances;
-
 	/** Allocates the vertex data storage type. */
 	void AllocateData();
+
+	/** Accepts preallocated data; Other is left empty after the call because no memory is copied. */
+	void AllocateData(FStaticMeshInstanceData& Other);
+
+	void SetupCPUAccess(UInstancedStaticMeshComponent* InComponent);
 };
 
 /*-----------------------------------------------------------------------------
@@ -365,11 +303,7 @@ public:
 	  , FeatureLevel(InFeatureLevel)
 	{
 		// Allocate the vertex factories for each LOD
-		for( int32 LODIndex=0;LODIndex<LODModels.Num();LODIndex++ )
-		{
-			FInstancedStaticMeshVertexFactory* VertexFactory = new(VertexFactories)FInstancedStaticMeshVertexFactory;
-			VertexFactory->SetFeatureLevel(InFeatureLevel);
-		}
+		InitVertexFactories();
 
 		// Create hit proxies for each instance if the component wants
 		if( GIsEditor && InComponent->bHasPerInstanceHitProxies )
@@ -383,6 +317,20 @@ public:
 
 		// initialize the instance buffer from the component's instances
 		InstanceBuffer.Init(Component, HitProxies);
+		NumInstances = InstanceBuffer.GetNumInstances();
+		InitResources();
+	}
+
+	FInstancedStaticMeshRenderData(UInstancedStaticMeshComponent* InComponent, ERHIFeatureLevel::Type InFeatureLevel, FStaticMeshInstanceData& Other)
+		: Component(InComponent)
+		, InstanceBuffer(InFeatureLevel)
+		, LODModels(Component->StaticMesh->RenderData->LODResources)
+		, FeatureLevel(InFeatureLevel)
+	{
+		InitVertexFactories();
+		// initialize the instance buffer from the component's instances
+		InstanceBuffer.InitFromPreallocatedData(Component, Other);
+		NumInstances = InstanceBuffer.GetNumInstances();
 		InitResources();
 	}
 
@@ -459,6 +407,22 @@ public:
 
 	/** Feature level used when creating instance data */
 	ERHIFeatureLevel::Type FeatureLevel;
+
+	/** Feature level used when creating instance data */
+	int32 NumInstances;
+
+private:
+
+	void InitVertexFactories()
+	{
+		// Allocate the vertex factories for each LOD
+		for( int32 LODIndex=0;LODIndex<LODModels.Num();LODIndex++ )
+		{
+			FInstancedStaticMeshVertexFactory* VertexFactory = new(VertexFactories)FInstancedStaticMeshVertexFactory;
+			VertexFactory->SetFeatureLevel(FeatureLevel);
+		}
+	}
+
 };
 
 
@@ -477,48 +441,20 @@ public:
 	,	bHasSelectedInstances(InComponent->SelectedInstances.Num() > 0)
 #endif
 	{
+		SetupProxy(InComponent);
+	}
+
+	FInstancedStaticMeshSceneProxy(UInstancedStaticMeshComponent* InComponent, ERHIFeatureLevel::Type InFeatureLevel, FStaticMeshInstanceData& Other)
+		:	FStaticMeshSceneProxy(InComponent)
+		,	InstancedRenderData(InComponent, InFeatureLevel, Other)
 #if WITH_EDITOR
-		if( bHasSelectedInstances )
-		{
-			// if we have selected indices, mark scene proxy as selected.
-			SetSelection_GameThread(true);
-		}
+		,	bHasSelectedInstances(InComponent->SelectedInstances.Num() > 0)
 #endif
-		// Make sure all the materials are okay to be rendered as an instanced mesh.
-		for (int32 LODIndex = 0; LODIndex < LODs.Num(); LODIndex++)
-		{
-			FStaticMeshSceneProxy::FLODInfo& LODInfo = LODs[LODIndex];
-			for (int32 SectionIndex = 0; SectionIndex < LODInfo.Sections.Num(); SectionIndex++)
-			{
-				FStaticMeshSceneProxy::FLODInfo::FSectionInfo& Section = LODInfo.Sections[SectionIndex];
-				if (!Section.Material->CheckMaterialUsage_Concurrent(MATUSAGE_InstancedStaticMeshes))
-				{
-					Section.Material = UMaterial::GetDefaultMaterial(MD_Surface);
-				}
-			}
-		}
-
-		check(InstancedRenderData.InstanceBuffer.GetStride() == sizeof(FInstancingUserData::FInstanceStream));
-
-		const bool bInstanced = RHISupportsInstancing(GetFeatureLevelShaderPlatform(InFeatureLevel));
-
-		// Copy the parameters for LOD - all instances
-		UserData_AllInstances.MeshRenderData = InComponent->StaticMesh->RenderData;
-		UserData_AllInstances.StartCullDistance = InComponent->InstanceStartCullDistance;
-		UserData_AllInstances.EndCullDistance = InComponent->InstanceEndCullDistance;
-		UserData_AllInstances.bRenderSelected = true;
-		UserData_AllInstances.bRenderUnselected = true;
-		UserData_AllInstances.RenderData = bInstanced ? nullptr : &InstancedRenderData;
-
-		// selected only
-		UserData_SelectedInstances = UserData_AllInstances;
-		UserData_SelectedInstances.bRenderUnselected = false;
-
-		// unselected only
-		UserData_DeselectedInstances = UserData_AllInstances;
-		UserData_DeselectedInstances.bRenderSelected = false;
+	{
+		SetupProxy(InComponent);
 	}
 
+
 	~FInstancedStaticMeshSceneProxy()
 	{
 		InstancedRenderData.ReleaseResources(&GetScene( ), StaticMesh);
@@ -603,6 +539,52 @@ protected:
 
 	/** Common path for the Get*MeshElement functions */
 	void SetupInstancedMeshBatch(int32 LODIndex, int32 BatchIndex, FMeshBatch& OutMeshBatch) const;
+
+private:
+
+	void SetupProxy(UInstancedStaticMeshComponent* InComponent)
+	{
+#if WITH_EDITOR
+		if( bHasSelectedInstances )
+		{
+			// if we have selected indices, mark scene proxy as selected.
+			SetSelection_GameThread(true);
+		}
+#endif
+		// Make sure all the materials are okay to be rendered as an instanced mesh.
+		for (int32 LODIndex = 0; LODIndex < LODs.Num(); LODIndex++)
+		{
+			FStaticMeshSceneProxy::FLODInfo& LODInfo = LODs[LODIndex];
+			for (int32 SectionIndex = 0; SectionIndex < LODInfo.Sections.Num(); SectionIndex++)
+			{
+				FStaticMeshSceneProxy::FLODInfo::FSectionInfo& Section = LODInfo.Sections[SectionIndex];
+				if (!Section.Material->CheckMaterialUsage_Concurrent(MATUSAGE_InstancedStaticMeshes))
+				{
+					Section.Material = UMaterial::GetDefaultMaterial(MD_Surface);
+				}
+			}
+		}
+
+		check(InstancedRenderData.InstanceBuffer.GetStride() == sizeof(FInstancingUserData::FInstanceStream));
+
+		const bool bInstanced = RHISupportsInstancing(GetFeatureLevelShaderPlatform(InstancedRenderData.FeatureLevel));
+
+		// Copy the parameters for LOD - all instances
+		UserData_AllInstances.MeshRenderData = InComponent->StaticMesh->RenderData;
+		UserData_AllInstances.StartCullDistance = InComponent->InstanceStartCullDistance;
+		UserData_AllInstances.EndCullDistance = InComponent->InstanceEndCullDistance;
+		UserData_AllInstances.bRenderSelected = true;
+		UserData_AllInstances.bRenderUnselected = true;
+		UserData_AllInstances.RenderData = bInstanced ? nullptr : &InstancedRenderData;
+
+		// selected only
+		UserData_SelectedInstances = UserData_AllInstances;
+		UserData_SelectedInstances.bRenderUnselected = false;
+
+		// unselected only
+		UserData_DeselectedInstances = UserData_AllInstances;
+		UserData_DeselectedInstances.bRenderSelected = false;
+	}
 };
 
 #if WITH_EDITOR
diff --git a/Engine/Source/Runtime/Engine/Public/StaticMeshResources.h b/Engine/Source/Runtime/Engine/Public/StaticMeshResources.h
--- a/Engine/Source/Runtime/Engine/Public/StaticMeshResources.h
+++ b/Engine/Source/Runtime/Engine/Public/StaticMeshResources.h
@@ -892,6 +892,95 @@ protected:
 	float GetScreenSize(int32 LODIndex) const;
 };
 
+/*-----------------------------------------------------------------------------
+	FStaticMeshInstanceData
+-----------------------------------------------------------------------------*/
+
+/** The implementation of the static mesh instance data storage type. */
+class FStaticMeshInstanceData :
+	public FStaticMeshVertexDataInterface,
+	public TResourceArray<FVector4,VERTEXBUFFER_ALIGNMENT>
+{
+public:
+
+	enum 
+	{
+		VectorsPerInstance = 7
+	};
+	typedef TResourceArray<FVector4,VERTEXBUFFER_ALIGNMENT> ArrayType;
+
+	/**
+	 * Constructor
+	 * @param InNeedsCPUAccess - true if resource array data should be CPU accessible
+	 */
+	FStaticMeshInstanceData(bool InNeedsCPUAccess=false)
+		:	TResourceArray<FVector4,VERTEXBUFFER_ALIGNMENT>(InNeedsCPUAccess)
+	{
+	}
+
+	static uint32 StaticGetStride()
+	{
+		return sizeof(FVector4) * VectorsPerInstance;
+	}
+
+	static SIZE_T GetResourceSize(uint32 NumInstances)
+	{
+		return SIZE_T(NumInstances) * SIZE_T(StaticGetStride());
+	}
+
+	int32 GetNumInstances()
+	{
+		check(Num() % VectorsPerInstance == 0);
+		return Num() / VectorsPerInstance;
+	}
+
+	/**
+	 * Resizes the vertex data buffer, discarding any data which no longer fits.
+	 * @param NumVertices - The number of vertices to allocate the buffer for.
+	 */
+	virtual void ResizeBuffer(uint32 NumInstances)
+	{
+		checkf(0, TEXT("ArrayType::Add is not supported on all platforms"));
+	}
+
+	virtual uint32 GetStride() const
+	{
+		return StaticGetStride();
+	}
+	virtual uint8* GetDataPointer()
+	{
+		return (uint8*)&(*this)[0];
+	}
+	virtual FResourceArrayInterface* GetResourceArray()
+	{
+		return this;
+	}
+	virtual void Serialize(FArchive& Ar)
+	{
+		TResourceArray<FVector4,VERTEXBUFFER_ALIGNMENT>::BulkSerialize(Ar);
+	}
+
+	void Set(const TArray<FVector4>& RawData)
+	{
+		*((ArrayType*)this) = TArray<FVector4,TAlignedHeapAllocator<VERTEXBUFFER_ALIGNMENT> >(RawData);
+	}
+
+	void AllocateInstances(int32 NumInstances)
+	{
+		// We cannot write directly to the data on all platforms,
+		// so we make a TArray of the right type, then assign it
+		check( GetStride() % sizeof(FVector4) == 0 );
+		Empty((NumInstances * GetStride()) / sizeof(FVector4));
+		AddUninitialized((NumInstances * GetStride()) / sizeof(FVector4));
+	}
+	FORCEINLINE FVector4* GetInstanceWriteAddress(int32 InstanceIndex)
+	{
+		return GetData() + InstanceIndex * VectorsPerInstance;
+	}
+
+};
+
+
 #if WITH_EDITOR
 /**
  * Remaps painted vertex colors when the renderable mesh has changed.
diff --git a/Engine/Source/Runtime/Landscape/Private/Landscape.cpp b/Engine/Source/Runtime/Landscape/Private/Landscape.cpp
--- a/Engine/Source/Runtime/Landscape/Private/Landscape.cpp
+++ b/Engine/Source/Runtime/Landscape/Private/Landscape.cpp
@@ -2409,10 +2409,8 @@ struct FAsyncGrassBuilder : public FGrassBuilderBase
 	int32 TotalInstances;
 
 	// output
-	TArray<FInstancedStaticMeshInstanceData> PerInstanceSMData;
+	FStaticMeshInstanceData InstanceBuffer;
 	TArray<FClusterNode> ClusterTree;
-	TArray<int32> SortedInstances;
-	TArray<int32> InstanceReorderTable;
 
 	FAsyncGrassBuilder(ALandscapeProxy* Landscape, ULandscapeComponent* Component, ULandscapeLayerInfoObject* LayerInfo, UHierarchicalInstancedStaticMeshComponent* HierarchicalInstancedStaticMeshComponent, int32 SqrtSubsections, int32 SubX, int32 SubY, FAsyncGrassBuilder* CopyFrom)
 		: FGrassBuilderBase(Landscape, Component, LayerInfo, SqrtSubsections, SubX, SubY)
@@ -2477,16 +2475,19 @@ struct FAsyncGrassBuilder : public FGrassBuilderBase
 
 		if (NumKept)
 		{
+			TArray<FMatrix> InstanceTransforms;
+			InstanceTransforms.AddUninitialized(NumKept);
 			TotalInstances += NumKept;
 
-
-			FInstancedStaticMeshInstanceData Inst;
-			Inst.LightmapUVBias = FVector2D( -1.0f, -1.0f );
-			Inst.ShadowmapUVBias = FVector2D( -1.0f, -1.0f );
 			{
+				// @todo: Make LD-customizable per component?
+				const float RandomInstanceIDBase = 0.0f;
+				const float RandomInstanceIDRange = 1.0f;
+
 				InstanceTime -= FPlatformTime::Seconds();
-				PerInstanceSMData.Reserve(NumKept);
+				InstanceBuffer.AllocateInstances(NumKept);
 				int32 InstanceIndex = 0;
+				int32 OutInstanceIndex = 0;
 				for (int32 xStart = 0; xStart < SqrtMaxInstances; xStart++ )
 				{
 					for (int32 yStart = 0; yStart < SqrtMaxInstances; yStart++ )
@@ -2495,6 +2496,7 @@ struct FAsyncGrassBuilder : public FGrassBuilderBase
 						if (Instance.bKeep)
 						{
 							float Rot = RandomRotation ? RandomStream.GetFraction() * 360.0f : 0.0f;
+							FMatrix OutXForm;
 							if (AlignToSurface)
 							{
 								FVector PosX1 = xStart ? Instances[InstanceIndex - SqrtMaxInstances].Pos : Instance.Pos;
@@ -2511,18 +2513,32 @@ struct FAsyncGrassBuilder : public FGrassBuilderBase
 									const FVector NewY = NewZ ^ NewX;
 
 									FMatrix Align = FMatrix(NewX, NewY, NewZ, FVector::ZeroVector);
-									Inst.Transform = FRotationMatrix(FRotator(0.0f, Rot, 0.0f)) * Align * FTranslationMatrix(Instance.Pos) * XForm;
+									OutXForm = FRotationMatrix(FRotator(0.0f, Rot, 0.0f)) * Align * FTranslationMatrix(Instance.Pos) * XForm;
 								}
 								else
 								{
-									Inst.Transform = FRotationMatrix(FRotator(0.0f, Rot, 0.0f)) * FTranslationMatrix(Instance.Pos) * XForm;
+									OutXForm = FRotationMatrix(FRotator(0.0f, Rot, 0.0f)) * FTranslationMatrix(Instance.Pos) * XForm;
 								}
 							}
 							else
 							{
-								Inst.Transform = FRotationMatrix(FRotator(0.0f, Rot, 0.0f)) * FTranslationMatrix(Instance.Pos) * XForm;
+								OutXForm = FRotationMatrix(FRotator(0.0f, Rot, 0.0f)) * FTranslationMatrix(Instance.Pos) * XForm;
 							}
-							PerInstanceSMData.Add(Inst);
+							InstanceTransforms[OutInstanceIndex] = OutXForm;
+							FVector4* RenderData = InstanceBuffer.GetInstanceWriteAddress(OutInstanceIndex++);
+
+							const float RandomInstanceID = RandomInstanceIDBase + RandomStream.GetFraction() * RandomInstanceIDRange;
+							RenderData[0] = FVector4(-1.0f, -1.0f, 0.0f, 0.0f);
+
+							// Instance -> local matrix.  Every mesh instance has it's own transformation into
+							// the actor's coordinate space.
+							RenderData[1] = FVector4(OutXForm.M[0][0], OutXForm.M[1][0], OutXForm.M[2][0], OutXForm.M[3][0]);
+							RenderData[2] = FVector4(OutXForm.M[0][1], OutXForm.M[1][1], OutXForm.M[2][1], OutXForm.M[3][1]);
+							RenderData[3] = FVector4(OutXForm.M[0][2], OutXForm.M[1][2], OutXForm.M[2][2], OutXForm.M[3][2]);
+
+							RenderData[4] = FVector4(0.0f, 0.0f, 0.0f, -1.0f);
+							RenderData[5] = FVector4(0.0f, 0.0f, 0.0f, -1.0f);
+							RenderData[6] = FVector4(0.0f, 0.0f, 0.0f, RandomInstanceID);
 						}
 						InstanceIndex++;
 					}
@@ -2531,7 +2547,35 @@ struct FAsyncGrassBuilder : public FGrassBuilderBase
 			}
 			{
 				BuildTime -= FPlatformTime::Seconds();
-				UHierarchicalInstancedStaticMeshComponent::BuildTreeAnyThread(PerInstanceSMData, MeshBox, ClusterTree, SortedInstances, InstanceReorderTable, FMath::Clamp<int32>(CVarMinInstancesPerLeaf.GetValueOnAnyThread(), 16, 256 * 128));
+				TArray<int32> SortedInstances;
+				TArray<int32> InstanceReorderTable;
+				UHierarchicalInstancedStaticMeshComponent::BuildTreeAnyThread(InstanceTransforms, MeshBox, ClusterTree, SortedInstances, InstanceReorderTable, FMath::Clamp<int32>(CVarMinInstancesPerLeaf.GetValueOnAnyThread(), 16, 256 * 128));
+
+				// in-place sort the instances
+				FVector4 SwapBuffer[FStaticMeshInstanceData::VectorsPerInstance];
+				int32 FirstUnfixedIndex = 0;
+				int32 Stride = FStaticMeshInstanceData::VectorsPerInstance * sizeof(FVector4);
+				for (int32 FirstUnfixedIndex = 0; FirstUnfixedIndex < NumKept; FirstUnfixedIndex++)
+				{
+					int32 LoadFrom = SortedInstances[FirstUnfixedIndex];
+					if (LoadFrom != FirstUnfixedIndex)
+					{
+						check(LoadFrom > FirstUnfixedIndex);
+						FMemory::Memcpy(SwapBuffer, InstanceBuffer.GetInstanceWriteAddress(FirstUnfixedIndex), Stride);
+						FMemory::Memcpy(InstanceBuffer.GetInstanceWriteAddress(FirstUnfixedIndex), InstanceBuffer.GetInstanceWriteAddress(LoadFrom), Stride);
+						FMemory::Memcpy(InstanceBuffer.GetInstanceWriteAddress(LoadFrom), SwapBuffer, Stride);
+
+						int32 SwapGoesTo = InstanceReorderTable[FirstUnfixedIndex];
+						check(SwapGoesTo > FirstUnfixedIndex);
+						check(SortedInstances[SwapGoesTo] == FirstUnfixedIndex);
+						SortedInstances[SwapGoesTo] = LoadFrom;
+						InstanceReorderTable[LoadFrom] = SwapGoesTo;
+
+						InstanceReorderTable[FirstUnfixedIndex] = FirstUnfixedIndex;
+						SortedInstances[FirstUnfixedIndex] = FirstUnfixedIndex;
+					}
+				}
+
 				BuildTime += FPlatformTime::Seconds();
 			}
 		}
@@ -2930,11 +2974,8 @@ void ALandscapeProxy::UpdateFoliage(const TArray<FVector>& Cameras, ULandscapeCo
 				UHierarchicalInstancedStaticMeshComponent* HierarchicalInstancedStaticMeshComponent = Inner.Foliage.Get();
 				if (HierarchicalInstancedStaticMeshComponent && StillUsed.Contains(HierarchicalInstancedStaticMeshComponent))
 				{
-					Exchange(HierarchicalInstancedStaticMeshComponent->PerInstanceSMData, Inner.Builder->PerInstanceSMData);
-					HierarchicalInstancedStaticMeshComponent->AcceptPrebuiltTree(Inner.Builder->ClusterTree, Inner.Builder->SortedInstances, Inner.Builder->InstanceReorderTable);
-
-					//@todo - this is sketchy, what you really want to do is flush the grass cache instead of doing a render thread update
-					HierarchicalInstancedStaticMeshComponent->bNeverNeedsRenderUpdate = true;
+					Exchange(HierarchicalInstancedStaticMeshComponent->WriteOncePrebuiltInstanceBuffer, Inner.Builder->InstanceBuffer);
+					HierarchicalInstancedStaticMeshComponent->AcceptPrebuiltTree(Inner.Builder->ClusterTree);
 
 					if (bForceSync)
 					{
