Engine: code and documentation cleanup pass

diff --git a/Engine/Source/Runtime/Engine/Classes/Engine/ActorChannel.h b/Engine/Source/Runtime/Engine/Classes/Engine/ActorChannel.h
--- a/Engine/Source/Runtime/Engine/Classes/Engine/ActorChannel.h
+++ b/Engine/Source/Runtime/Engine/Classes/Engine/ActorChannel.h
@@ -1,12 +1,18 @@
 // Copyright 1998-2015 Epic Games, Inc. All Rights Reserved.
 
-/**	ActorChannel.h
- *	A channel for exchanging actor and its subobject's properties and RPCs.
+#pragma once
+
+#include "Engine/Channel.h"
+#include "ActorChannel.generated.h"
+
+
+/**
+ * A channel for exchanging actor and its subobject's properties and RPCs.
  *
- *	ActorChannel manages the creation and lifetime of a replicated actor. Actual replication of properties and RPCs 
- *	actually happens in FObjectReplicator now (see DataReplication.h).
+ * ActorChannel manages the creation and lifetime of a replicated actor. Actual replication of properties and RPCs
+ * actually happens in FObjectReplicator now (see DataReplication.h).
  *
- *	An ActorChannel bunch looks like this:
+ * An ActorChannel bunch looks like this:
  *
  *		|----------------------|---------------------------------------------------------------------------|
  *		| SpawnInfo		       | (Spawn Info) Initial bunch only                                           |
@@ -29,16 +35,10 @@
  *      |----------------------|---------------------------------------------------------------------------|
  *		| </End Tag>           |                                                                           |
  *		|----------------------|---------------------------------------------------------------------------|
- * 
- *
  */
-
-#pragma once
-#include "Engine/Channel.h"
-#include "ActorChannel.generated.h"
-
 UCLASS(transient, customConstructor)
-class ENGINE_API UActorChannel : public UChannel
+class ENGINE_API UActorChannel
+	: public UChannel
 {
 	GENERATED_UCLASS_BODY()
 
@@ -59,7 +59,7 @@ class ENGINE_API UActorChannel : public UChannel
 	uint32  bActorStillInitial:1;	// Not all properties sent while bNetInitial, so still bNetInitial next tick
 	uint32  bIsReplicatingActor:1;	// true when in this channel's ReplicateActor() to avoid recursion as that can cause invalid data to be sent
 	
-	/** whether we should NULL references to this channel's Actor in other channels' Recent data when this channel is closed
+	/** whether we should nullptr references to this channel's Actor in other channels' Recent data when this channel is closed
 	 * set to false in cases where the Actor can't become relevant again (e.g. destruction) as it's unnecessary in that case
 	 */
 	uint32 bClearRecentActorRefs:1;
@@ -86,25 +86,23 @@ class ENGINE_API UActorChannel : public UChannel
 		bClearRecentActorRefs = true;
 	}
 
-	// Begin UChannel interface.
+public:
 
-	virtual void Init( UNetConnection* InConnection, int32 InChIndex, bool InOpenedLocally ) override;
+	// UChannel interface.
 
+	virtual void Init( UNetConnection* InConnection, int32 InChIndex, bool InOpenedLocally ) override;
 	virtual void SetClosingFlag() override;
-
 	virtual void ReceivedBunch( FInBunch& Bunch ) override;
-
 	virtual void Tick() override;
+
 	void ProcessBunch( FInBunch & Bunch );
 	bool ProcessQueuedBunches();
 
 	virtual void ReceivedNak( int32 NakPacketId ) override;
-	
 	virtual void Close() override;
-
 	virtual FString Describe() override;
 
-	// End UChannel interface.
+public:
 
 	/** UActorChannel interface and accessors. */
 	AActor* GetActor() {return Actor;}
@@ -140,7 +138,7 @@ class ENGINE_API UActorChannel : public UChannel
 	void BeginContentBlockForSubObjectDelete( FOutBunch & Bunch, FNetworkGUID & GuidToDelete );
 
 	/** Writes the footer for a content block of proeprties / RPCs for the given object (either the actor a subobject of the actor) */
-	void EndContentBlock( UObject *Obj, FOutBunch &Bunch, FClassNetCache* ClassCache = NULL );
+	void EndContentBlock( UObject *Obj, FOutBunch &Bunch, FClassNetCache* ClassCache = nullptr );
 
 	/** Reads the header of the content block and instantiates the subobject if necessary */
 	UObject* ReadContentBlockHeader(FInBunch& Bunch, bool& bObjectDeleted);
diff --git a/Engine/Source/Runtime/Engine/Classes/Engine/BlockingVolume.h b/Engine/Source/Runtime/Engine/Classes/Engine/BlockingVolume.h
--- a/Engine/Source/Runtime/Engine/Classes/Engine/BlockingVolume.h
+++ b/Engine/Source/Runtime/Engine/Classes/Engine/BlockingVolume.h
@@ -1,23 +1,24 @@
 // Copyright 1998-2015 Epic Games, Inc. All Rights Reserved.
 
-
 #pragma once
+
 #include "GameFramework/Volume.h"
 #include "BlockingVolume.generated.h"
 
+
 /** An invisible volume used to block other actors. */
 UCLASS(MinimalAPI)
-class ABlockingVolume : public AVolume
+class ABlockingVolume
+	: public AVolume
 {
 	GENERATED_UCLASS_BODY()
 
-	// Begin UObject interface.
+public:
+
+	// UObject interface.
+
 #if WITH_EDITOR
 	virtual void LoadedFromAnotherClass(const FName& OldClassName) override;
 	virtual void PostEditChangeChainProperty(FPropertyChangedChainEvent& PropertyChangedEvent) override;
 #endif // WITH_EDITOR	
-	// End UObject interface.
 };
-
-
-
diff --git a/Engine/Source/Runtime/Engine/Classes/Engine/Blueprint.h b/Engine/Source/Runtime/Engine/Classes/Engine/Blueprint.h
--- a/Engine/Source/Runtime/Engine/Classes/Engine/Blueprint.h
+++ b/Engine/Source/Runtime/Engine/Classes/Engine/Blueprint.h
@@ -1,12 +1,16 @@
 // Copyright 1998-2015 Epic Games, Inc. All Rights Reserved.
 
 #pragma once
+
 #include "EdGraph/EdGraphPin.h"
 #include "EdGraph/EdGraphNode.h"
 #include "BlueprintCore.h"
 #include "Blueprint.generated.h"
 
-/** States a blueprint can be in */
+
+/**
+ * Enumerates states a blueprint can be in.
+ */
 UENUM()
 enum EBlueprintStatus
 {
@@ -25,7 +29,10 @@ enum EBlueprintStatus
 	BS_MAX,
 };
 
-/** Types of blueprints */
+
+/**
+ * Enumerates types of blueprints.
+ */
 UENUM()
 enum EBlueprintType
 {
@@ -59,6 +66,7 @@ namespace EKismetCompileType
 	};
 };
 
+
 struct FKismetCompilerOptions
 {
 public:
@@ -125,6 +133,7 @@ struct FBPVariableMetaDataEntry
 	
 };
 
+
 /** Struct indicating a variable in the generated class */
 USTRUCT()
 struct FBPVariableDescription
@@ -184,6 +193,7 @@ struct FBPVariableDescription
 	
 };
 
+
 /** Struct containing information about what interfaces are implemented in this blueprint */
 USTRUCT()
 struct FBPInterfaceDescription
@@ -200,12 +210,11 @@ struct FBPInterfaceDescription
 
 
 	FBPInterfaceDescription()
-		: Interface(NULL)
-	{
-	}
-
+		: Interface(nullptr)
+	{ }
 };
 
+
 USTRUCT()
 struct FEditedDocumentInfo
 {
@@ -223,28 +232,22 @@ struct FEditedDocumentInfo
 	float SavedZoomAmount;
 
 	FEditedDocumentInfo()
-		: EditedObject(NULL)
+		: EditedObject(nullptr)
 		, SavedViewOffset(0.0f, 0.0f)
 		, SavedZoomAmount(-1.0f)
-	{
-
-	}
+	{ }
 
 	FEditedDocumentInfo(UObject* InEditedObject)
 		: EditedObject(InEditedObject)
 		, SavedViewOffset(0.0f, 0.0f)
 		, SavedZoomAmount(-1.0f)
-	{
-
-	}
+	{ }
 
 	FEditedDocumentInfo(UObject* InEditedObject, FVector2D& InSavedViewOffset, float InSavedZoomAmount)
 		: EditedObject(InEditedObject)
 		, SavedViewOffset(InSavedViewOffset)
 		, SavedZoomAmount(InSavedZoomAmount)
-	{
-
-	}
+	{ }
 
 	friend bool operator==( const FEditedDocumentInfo& LHS, const FEditedDocumentInfo& RHS )
 	{
@@ -252,7 +255,6 @@ struct FEditedDocumentInfo
 	}
 };
 
-///////////////
 
 /**
  * Blueprints are special assets that provide an intuitive, node-based interface that can be used to create new types of Actors
@@ -439,7 +441,9 @@ protected:
 
 	/** Current world being debugged for this blueprint */
 	TWeakObjectPtr< class UWorld > CurrentWorldBeingDebugged;
+
 public:
+
 	/** Information for thumbnail rendering */
 	UPROPERTY(VisibleAnywhere, Instanced, Category=Thumbnail)
 	class UThumbnailInfo* ThumbnailInfo;
@@ -517,22 +521,23 @@ public:
 	virtual void GetReparentingRules(TSet< const UClass* >& AllowedChildrenOfClasses, TSet< const UClass* >& DisallowedChildrenOfClasses) const;
 
 private:
+
 	/** Sets the current object being debugged */
 	void DebuggingWorldRegistrationHelper(UObject* ObjectProvidingWorld, UObject* ValueToRegister);
-
 	
 public:
-	/** @return the current object being debugged, which can be NULL */
+
+	/** @return the current object being debugged, which can be nullptr */
 	virtual UObject* GetObjectBeingDebugged();
 
 	virtual class UWorld* GetWorldBeingDebugged();
 
 	/** Renames only the generated classes. Should only be used internally or when testing for rename. */
-	virtual bool RenameGeneratedClasses(const TCHAR* NewName = NULL, UObject* NewOuter = NULL, ERenameFlags Flags = REN_None);
+	virtual bool RenameGeneratedClasses(const TCHAR* NewName = nullptr, UObject* NewOuter = nullptr, ERenameFlags Flags = REN_None);
 
 	// Begin UObject interface (WITH_EDITOR)
 	virtual void PostDuplicate(bool bDuplicateForPIE) override;
-	virtual bool Rename(const TCHAR* NewName = NULL, UObject* NewOuter = NULL, ERenameFlags Flags = REN_None) override;
+	virtual bool Rename(const TCHAR* NewName = nullptr, UObject* NewOuter = nullptr, ERenameFlags Flags = REN_None) override;
 	virtual UClass* RegenerateClass(UClass* ClassToRegenerate, UObject* PreviousCDO, TArray<UObject*>& ObjLoaded) override;
 	virtual void PostLoad() override;
 	virtual void PostLoadSubobjects( FObjectInstancingGraph* OuterInstanceGraph ) override;
@@ -666,6 +671,7 @@ public:
 
 };
 
+
 #if WITH_EDITOR
 template<>
 inline FName UBlueprint::GetFieldNameFromClassByGuid<UFunction>(const UClass* InClass, const FGuid FunctionGuid)
@@ -678,4 +684,5 @@ inline bool UBlueprint::GetGuidFromClassByFieldName<UFunction>(const UClass* InC
 {
 	return GetFunctionGuidFromClassByFieldName(InClass, FunctionName, FunctionGuid);
 }
+
 #endif // #if WITH_EDITOR
diff --git a/Engine/Source/Runtime/Engine/Classes/Engine/BlueprintGeneratedClass.h b/Engine/Source/Runtime/Engine/Classes/Engine/BlueprintGeneratedClass.h
--- a/Engine/Source/Runtime/Engine/Classes/Engine/BlueprintGeneratedClass.h
+++ b/Engine/Source/Runtime/Engine/Classes/Engine/BlueprintGeneratedClass.h
@@ -1,11 +1,14 @@
 // Copyright 1998-2015 Epic Games, Inc. All Rights Reserved.
 
 #pragma once
+
 #include "MeshBatch.h"
 #include "BlueprintGeneratedClass.generated.h"
 
+
 class UEdGraphPin;
 
+
 USTRUCT()
 struct FNodeToCodeAssociation
 {
@@ -17,8 +20,8 @@ public:
 	int32 Offset;
 public:
 	FNodeToCodeAssociation()
-	: Node(NULL)
-	, Scope(NULL)
+	: Node(nullptr)
+	, Scope(nullptr)
 	, Offset(0)
 	{
 	}
@@ -31,6 +34,7 @@ public:
 	}
 };
 
+
 USTRUCT()
 struct FDebuggingInfoForSingleFunction
 {
@@ -55,6 +59,7 @@ public:
 	}
 };
 
+
 USTRUCT()
 struct FPointerToUberGraphFrame
 {
@@ -64,7 +69,7 @@ public:
 	uint8* RawPointer;
 
 	FPointerToUberGraphFrame() 
-		: RawPointer(NULL)
+		: RawPointer(nullptr)
 	{}
 
 	~FPointerToUberGraphFrame()
@@ -73,6 +78,7 @@ public:
 	}
 };
 
+
 template<>
 struct TStructOpsTypeTraits<FPointerToUberGraphFrame> : public TStructOpsTypeTraitsBase
 {
@@ -83,6 +89,7 @@ struct TStructOpsTypeTraits<FPointerToUberGraphFrame> : public TStructOpsTypeTra
 	};
 };
 
+
 //////////////////////////////////////////////////////////////////////////
 // TSimpleRingBuffer
 
@@ -158,6 +165,7 @@ private:
 	int32 WriteIndex;
 };
 
+
 struct FStructUtils
 {
 private:
@@ -168,6 +176,7 @@ public:
 	ENGINE_API static bool TheSameLayout(const UStruct* StructA, const UStruct* StructB, bool bCheckPropertiesNames = false);
 };
 
+
 USTRUCT()
 struct ENGINE_API FBlueprintDebugData
 {
@@ -207,7 +216,7 @@ public:
 		DebugPinToPropertyMap.Empty();
 	}
 
-	// Returns the UEdGraphNode associated with the UUID, or NULL if there isn't one.
+	// Returns the UEdGraphNode associated with the UUID, or nullptr if there isn't one.
 	UEdGraphNode* FindNodeFromUUID(int32 UUID) const
 	{
 		if (const TWeakObjectPtr<UEdGraphNode>* pParentNode = DebugNodesAllocatedUniqueIDsMap.Find(UUID))
@@ -215,7 +224,7 @@ public:
 			return pParentNode->Get();
 		}
 	
-		return NULL;
+		return nullptr;
 	}
 
 	bool IsValid() const
@@ -223,16 +232,16 @@ public:
 		return DebugNodeLineNumbers.Num() > 0;
 	}
 
-	// Finds the UEdGraphNode associated with the code location Function+CodeOffset, or NULL if there isn't one
+	// Finds the UEdGraphNode associated with the code location Function+CodeOffset, or nullptr if there isn't one
 	UEdGraphNode* FindSourceNodeFromCodeLocation(UFunction* Function, int32 CodeOffset, bool bAllowImpreciseHit) const
 	{
 		if (const FDebuggingInfoForSingleFunction* pFuncInfo = PerFunctionLineNumbers.Find(Function))
 		{
 			UEdGraphNode* Result = pFuncInfo->LineNumberToSourceNodeMap.FindRef(CodeOffset).Get();
 
-			if ((Result == NULL) && bAllowImpreciseHit)
+			if ((Result == nullptr) && bAllowImpreciseHit)
 			{
-				for (int32 TrialOffset = CodeOffset + 1; (Result == NULL) && (TrialOffset < Function->Script.Num()); ++TrialOffset)
+				for (int32 TrialOffset = CodeOffset + 1; (Result == nullptr) && (TrialOffset < Function->Script.Num()); ++TrialOffset)
 				{
 					Result = pFuncInfo->LineNumberToSourceNodeMap.FindRef(TrialOffset).Get();
 				}
@@ -241,10 +250,10 @@ public:
 			return Result;
 		}
 
-		return NULL;
+		return nullptr;
 	}
 
-	// Finds the macro source node associated with the code location Function+CodeOffset, or NULL if there isn't one
+	// Finds the macro source node associated with the code location Function+CodeOffset, or nullptr if there isn't one
 	UEdGraphNode* FindMacroSourceNodeFromCodeLocation(UFunction* Function, int32 CodeOffset) const
 	{
 		if (const FDebuggingInfoForSingleFunction* pFuncInfo = PerFunctionLineNumbers.Find(Function))
@@ -252,7 +261,7 @@ public:
 			return pFuncInfo->LineNumberToMacroSourceNodeMap.FindRef(CodeOffset).Get();
 		}
 
-		return NULL;
+		return nullptr;
 	}
 
 	// Finds the macro instance node(s) associated with the code location Function+CodeOffset. The returned set can be empty.
@@ -274,7 +283,7 @@ public:
 		}
 	}
 
-	// Finds the macro source node associated with the code location Function+CodeOffset, or NULL if there isn't one
+	// Finds the macro source node associated with the code location Function+CodeOffset, or nullptr if there isn't one
 	UEdGraphPin* FindExecPinFromCodeLocation(UFunction* Function, int32 CodeOffset) const
 	{
 		if (const FDebuggingInfoForSingleFunction* pFuncInfo = PerFunctionLineNumbers.Find(Function))
@@ -282,7 +291,7 @@ public:
 			return pFuncInfo->LineNumberToPinMap.FindRef(CodeOffset).Get();
 		}
 
-		return NULL;
+		return nullptr;
 	}
 
 	// Finds the breakpoint injection site(s) in bytecode if any were associated with the given node
@@ -304,7 +313,7 @@ public:
 	UProperty* FindClassPropertyForPin(const UEdGraphPin* Pin) const
 	{
 		UProperty* PropertyPtr = DebugPinToPropertyMap.FindRef(Pin);
-		if ((PropertyPtr == NULL) && (Pin->Direction == EGPD_Input) && (Pin->LinkedTo.Num() > 0))
+		if ((PropertyPtr == nullptr) && (Pin->Direction == EGPD_Input) && (Pin->LinkedTo.Num() > 0))
 		{
 			// Try checking the other side of the connection
 			PropertyPtr = DebugPinToPropertyMap.FindRef(Pin->LinkedTo[0]);
@@ -361,7 +370,7 @@ public:
 		DebugNodesAllocatedUniqueIDsMap.Add(UUID, TrueSourceNode);
 	}
 
-	// Returns the object that caused the specified property to be created (can return NULL if the association is unknown)
+	// Returns the object that caused the specified property to be created (can return nullptr if the association is unknown)
 	UObject* FindObjectThatCreatedProperty(class UProperty* AssociatedProperty) const
 	{
 		if (const TWeakObjectPtr<UObject>* pValue = DebugPinToPropertyMap.FindKey(AssociatedProperty))
@@ -370,7 +379,7 @@ public:
 		}
 		else
 		{
-			return NULL;
+			return nullptr;
 		}
 	}
 
@@ -389,11 +398,11 @@ public:
 	 * Retrieves the most recently executed node (created to help pinpoint 
 	 * runtime errors).
 	 * 
-	 * @returns The most recently executed node (NULL if there is not one).
+	 * @returns The most recently executed node (nullptr if there is not one).
 	 */
 	UEdGraphNode* GetLastExecutedNode()
 	{
-		UEdGraphNode* LastExecutedNode = NULL;
+		UEdGraphNode* LastExecutedNode = nullptr;
 		if (IsValid())
 		{
 			TWeakObjectPtr<UEdGraphNode> NodePtr = DebugNodeLineNumbers.Top().Node;
@@ -408,6 +417,7 @@ public:
 #endif
 };
 
+
 UCLASS()
 class ENGINE_API UBlueprintGeneratedClass : public UClass
 {
diff --git a/Engine/Source/Runtime/Engine/Classes/Engine/BookMark.h b/Engine/Source/Runtime/Engine/Classes/Engine/BookMark.h
--- a/Engine/Source/Runtime/Engine/Classes/Engine/BookMark.h
+++ b/Engine/Source/Runtime/Engine/Classes/Engine/BookMark.h
@@ -1,14 +1,13 @@
 // Copyright 1998-2015 Epic Games, Inc. All Rights Reserved.
 
-/* epic ===============================================
-* class BookMark
-*
-* A camera position the current level.
- */
-
 #pragma once
+
 #include "BookMark.generated.h"
 
+
+/**
+ * A camera position the current level.
+ */
 UCLASS(hidecategories=Object, MinimalAPI)
 class UBookMark : public UObject
 {
@@ -25,6 +24,4 @@ class UBookMark : public UObject
 	/** Array of levels that are hidden */
 	UPROPERTY(EditAnywhere, Category=BookMark)
 	TArray<FString> HiddenLevels;
-
 };
-
diff --git a/Engine/Source/Runtime/Engine/Classes/Engine/BookMark2D.h b/Engine/Source/Runtime/Engine/Classes/Engine/BookMark2D.h
--- a/Engine/Source/Runtime/Engine/Classes/Engine/BookMark2D.h
+++ b/Engine/Source/Runtime/Engine/Classes/Engine/BookMark2D.h
@@ -1,15 +1,16 @@
 // Copyright 1998-2015 Epic Games, Inc. All Rights Reserved.
 
-/**
- *
- * Simple class to store 2D camera information.
- */
-
 #pragma once
+
 #include "BookMark2D.generated.h"
 
+
+/**
+ * Simple class to store 2D camera information.
+ */
 UCLASS(hidecategories=Object)
-class UBookMark2D : public UObject
+class UBookMark2D
+	: public UObject
 {
 	GENERATED_UCLASS_BODY()
 
@@ -20,6 +21,4 @@ class UBookMark2D : public UObject
 	/** Location of the camera */
 	UPROPERTY(EditAnywhere, Category=BookMark2D)
 	FIntPoint Location;
-
 };
-
diff --git a/Engine/Source/Runtime/Engine/Classes/Engine/BoxReflectionCapture.h b/Engine/Source/Runtime/Engine/Classes/Engine/BoxReflectionCapture.h
--- a/Engine/Source/Runtime/Engine/Classes/Engine/BoxReflectionCapture.h
+++ b/Engine/Source/Runtime/Engine/Classes/Engine/BoxReflectionCapture.h
@@ -1,21 +1,18 @@
 // Copyright 1998-2015 Epic Games, Inc. All Rights Reserved.
 
 #pragma once
+
 #include "ReflectionCapture.h"
 #include "BoxReflectionCapture.generated.h"
 
+
 /** 
- *	Actor used to capture the scene for reflection in a box shape 
- *	@see https://docs.unrealengine.com/latest/INT/Resources/ContentExamples/Reflections/1_3/index.html
+ * Actor used to capture the scene for reflection in a box shape 
+ * @see https://docs.unrealengine.com/latest/INT/Resources/ContentExamples/Reflections/1_3/index.html
  */
 UCLASS(hidecategories=(Collision, Attachment, Actor), MinimalAPI)
-class ABoxReflectionCapture : public AReflectionCapture
+class ABoxReflectionCapture
+	: public AReflectionCapture
 {
 	GENERATED_UCLASS_BODY()
-
-public:
-
 };
-
-
-
diff --git a/Engine/Source/Runtime/Engine/Classes/Engine/Breakpoint.h b/Engine/Source/Runtime/Engine/Classes/Engine/Breakpoint.h
--- a/Engine/Source/Runtime/Engine/Classes/Engine/Breakpoint.h
+++ b/Engine/Source/Runtime/Engine/Classes/Engine/Breakpoint.h
@@ -1,11 +1,13 @@
 // Copyright 1998-2015 Epic Games, Inc. All Rights Reserved.
 
-
 #pragma once
+
 #include "Breakpoint.generated.h"
 
+
 UCLASS()
-class ENGINE_API UBreakpoint : public UObject
+class ENGINE_API UBreakpoint
+	: public UObject
 {
 	GENERATED_UCLASS_BODY()
 
@@ -30,7 +32,6 @@ private:
 
 public:
 
-public:
 	/** Get the target node for the breakpoint */
 	UEdGraphNode* GetLocation() const
 	{
@@ -54,6 +55,3 @@ public:
 
 	friend class FKismetDebugUtilities;
 };
-
-
-
diff --git a/Engine/Source/Runtime/Engine/Classes/Engine/Brush.h b/Engine/Source/Runtime/Engine/Classes/Engine/Brush.h
--- a/Engine/Source/Runtime/Engine/Classes/Engine/Brush.h
+++ b/Engine/Source/Runtime/Engine/Classes/Engine/Brush.h
@@ -1,15 +1,13 @@
 // Copyright 1998-2015 Epic Games, Inc. All Rights Reserved.
 
-//=============================================================================
-// The brush class.
-//=============================================================================
-
 #pragma once
+
 #include "Model.h"
 #include "Level.h"
 #include "GameFramework/Actor.h"
 #include "Brush.generated.h"
 
+
 //-----------------------------------------------------------------------------
 // Variables.
 UENUM()
@@ -42,6 +40,7 @@ enum EBrushType
 	Brush_MAX,
 };
 
+
 // Selection information for geometry mode
 USTRUCT()
 struct FGeomSelection
@@ -62,9 +61,7 @@ struct FGeomSelection
 		: Type(0)
 		, Index(0)
 		, SelectionIndex(0)
-	{
-	}
-
+	{ }
 
 	friend ENGINE_API FArchive& operator<<(FArchive& Ar,FGeomSelection& MyGeomSelection)
 	{
@@ -72,8 +69,10 @@ struct FGeomSelection
 	}
 };
 
+
 UCLASS(hidecategories=(Object, Collision, Display, Rendering, Physics, Input, Blueprint), showcategories=("Input|MouseInput", "Input|TouchInput"), NotBlueprintable, ConversionRoot)
-class ENGINE_API ABrush : public AActor
+class ENGINE_API ABrush
+	: public AActor
 {
 	GENERATED_UCLASS_BODY()
 
@@ -137,60 +136,79 @@ public:
 	}
 #endif
 
+public:
 	
-	// Begin UObject interface.
+	// UObject interface.
+
 	virtual void PostLoad() override;
+
 #if WITH_EDITOR
 	virtual void PostEditMove(bool bFinished) override;
 	virtual void PostEditChangeProperty(FPropertyChangedEvent& PropertyChangedEvent) override;
 #endif // WITH_EDITOR
+
 	virtual bool Modify(bool bAlwaysMarkDirty = false) override;
 
 	virtual bool NeedsLoadForClient() const override
 	{ 
 		return !IsNotForClientOrServer(); 
 	}
+
 	virtual bool NeedsLoadForServer() const override
 	{ 
 		return !IsNotForClientOrServer(); 
 	}
-	// End UObject interface.
 
+public:
 	
-	// Begin AActor interface
+	// AActor interface
 	virtual void Destroyed() override;
 	virtual void PostRegisterAllComponents() override;
-	
 	virtual bool IsLevelBoundsRelevant() const override;
 
 #if WITH_EDITOR
 	virtual void CheckForErrors() override;
-
 	virtual void SetIsTemporarilyHiddenInEditor( bool bIsHidden ) override;
 
-	// End AActor interface
+public:
 
 	virtual void InitPosRotScale();
 	virtual void CopyPosRotScaleFrom( ABrush* Other );
 
-private: 
+private:
+
 	/** An array to keep track of all the levels that need rebuilding. This is checked via NeedsRebuild() in the editor tick and triggers a csg rebuild. */
 	static TArray< TWeakObjectPtr< ULevel > > LevelsToRebuild;
 
 	/** Delegate called when PostRegisterAllComponents is called for a Brush */
 	static FOnBrushRegistered OnBrushRegistered;
+
 public:
+
 	/**
 	 * Called to see if any of the levels need rebuilding
 	 *
 	 * @param	OutLevels if specified, provides a copy of the levels array
 	 * @return	true if the csg needs to be rebuilt on the next editor tick.	
 	 */
-	static bool NeedsRebuild(TArray< TWeakObjectPtr< ULevel > >* OutLevels = NULL){if(OutLevels){*OutLevels = LevelsToRebuild;} return(LevelsToRebuild.Num() > 0);}
+	static bool NeedsRebuild(TArray< TWeakObjectPtr< ULevel > >* OutLevels = nullptr)
+	{
+		if (OutLevels)
+		{
+			*OutLevels = LevelsToRebuild;
+		}
+		
+		return(LevelsToRebuild.Num() > 0);
+	}
+
 	/**
 	 * Called upon finishing the csg rebuild to clear the rebuild bool.
 	 */
-	static void OnRebuildDone(){LevelsToRebuild.Empty();}
+	static void OnRebuildDone()
+	{
+		LevelsToRebuild.Empty();
+	}
+
 	/**
 	 * Called to make not of the level that needs rebuilding
 	 *
@@ -198,7 +216,7 @@ public:
 	 */
 	static void SetNeedRebuild(ULevel* InLevel){if(InLevel){LevelsToRebuild.AddUnique(InLevel);}}
 #else
-	static bool NeedsRebuild(TArray< TWeakObjectPtr< ULevel > >* OutLevels = NULL){return false;}
+	static bool NeedsRebuild(TArray< TWeakObjectPtr< ULevel > >* OutLevels = nullptr){return false;}
 	static void OnRebuildDone(){}
 	static void SetNeedRebuild(ULevel* InLevel){}
 #endif//WITH_EDITOR
@@ -223,7 +241,6 @@ public:
 	/** Figures out the best color to use for this brushes wireframe drawing.	*/
 	virtual FColor GetWireColor() const;
 
-
 	/**
 	 * Return if true if this brush is not used for gameplay (i.e. builder brush)
 	 * 
@@ -255,6 +272,3 @@ public:
 	/** Returns BrushComponent subobject **/
 	class UBrushComponent* GetBrushComponent() const;
 };
-
-
-
diff --git a/Engine/Source/Runtime/Engine/Classes/Engine/BrushBuilder.h b/Engine/Source/Runtime/Engine/Classes/Engine/BrushBuilder.h
--- a/Engine/Source/Runtime/Engine/Classes/Engine/BrushBuilder.h
+++ b/Engine/Source/Runtime/Engine/Classes/Engine/BrushBuilder.h
@@ -1,28 +1,10 @@
 // Copyright 1998-2015 Epic Games, Inc. All Rights Reserved.
 
-/**
- *
- */
-
-
-//=============================================================================
-// BrushBuilder: Base class of UnrealEd brush builders.
-//
-// Tips for writing brush builders:
-//
-// * Always validate the user-specified and call BadParameters function
-//   if anything is wrong, instead of actually building geometry.
-//   If you build an invalid brush due to bad user parameters, you'll
-//   cause an extraordinary amount of pain for the poor user.
-//
-// * When generating polygons with more than 3 vertices, BE SURE all the
-//   polygon's vertices are coplanar!  Out-of-plane polygons will cause
-//   geometry to be corrupted.
-//=============================================================================
-
 #pragma once
+
 #include "BrushBuilder.generated.h"
 
+
 // Internal state, not accessible to script.
 USTRUCT()
 struct FBuilderPoly
@@ -46,8 +28,25 @@ struct FBuilderPoly
 	{}
 };
 
+
+/**
+ * Base class of UnrealEd brush builders.
+ *
+ *
+ * Tips for writing brush builders:
+ *
+ * - Always validate the user-specified and call BadParameters function
+ *   if anything is wrong, instead of actually building geometry.
+ *   If you build an invalid brush due to bad user parameters, you'll
+ *   cause an extraordinary amount of pain for the poor user.
+ *
+ * - When generating polygons with more than 3 vertices, BE SURE all the
+ *   polygon's vertices are coplanar!  Out-of-plane polygons will cause
+ *   geometry to be corrupted.
+ */
 UCLASS(abstract, hidecategories=(Object), MinimalAPI)
-class UBrushBuilder : public UObject
+class UBrushBuilder
+	: public UObject
 {
 	GENERATED_UCLASS_BODY()
 
@@ -63,6 +62,7 @@ class UBrushBuilder : public UObject
 	uint32 NotifyBadParams:1;
 
 protected:
+
 	UPROPERTY()
 	TArray<FVector> Vertices;
 
@@ -76,6 +76,7 @@ protected:
 	uint32 MergeCoplanars:1;
 
 public:
+
 	// @ todo document all below
 	virtual void BeginBrush( bool InMergeCoplanars, FName InLayer ) {}
 	virtual bool EndBrush( UWorld* InWorld, ABrush* InBrush ) { return false; }
@@ -92,13 +93,13 @@ public:
 	virtual void PolyEnd() {}
 
 	/**
-	 * Builds the brush shape for the specified ABrush or, if Brush==NULL, the builder brush.
+	 * Builds the brush shape for the specified ABrush or the builder brush if Brush is nullptr.
 	 * @param InWorld The world to operate in
-	 * @param InBrush The brush to change shape, or NULL to specify the builder brush
+	 * @param InBrush The brush to change shape, or nullptr to specify the builder brush
 	 * @return true if the brush shape was updated.
 	 */
-	ENGINE_API virtual bool Build( UWorld* InWorld, ABrush* InBrush = NULL ) { return false; }
+	ENGINE_API virtual bool Build( UWorld* InWorld, ABrush* InBrush = nullptr )
+	{
+		return false;
+	}
 };
-
-
-
diff --git a/Engine/Source/Runtime/Engine/Classes/Engine/BrushShape.h b/Engine/Source/Runtime/Engine/Classes/Engine/BrushShape.h
--- a/Engine/Source/Runtime/Engine/Classes/Engine/BrushShape.h
+++ b/Engine/Source/Runtime/Engine/Classes/Engine/BrushShape.h
@@ -1,22 +1,30 @@
 // Copyright 1998-2015 Epic Games, Inc. All Rights Reserved.
 
-//=============================================================================
-// BrushShape: A brush that acts as a template for geometry mode modifiers like "Lathe"
-//=============================================================================
-
 #pragma once
+
 #include "BrushShape.generated.h"
 
+
+/**
+ * A brush that acts as a template for geometry mode modifiers like "Lathe".
+ */
 UCLASS(MinimalAPI)
-class ABrushShape : public ABrush
+class ABrushShape
+	: public ABrush
 {
 	GENERATED_UCLASS_BODY()
 
+public:
 
-	virtual bool IsStaticBrush() const override {return false; }
-
-	virtual bool IsBrushShape() const override {return true; }
-};
-
+	// ABrush overrides
 
+	virtual bool IsStaticBrush() const override
+	{
+		return false;
+	}
 
+	virtual bool IsBrushShape() const override
+	{
+		return true;
+	}
+};
diff --git a/Engine/Source/Runtime/Engine/Classes/Engine/Canvas.h b/Engine/Source/Runtime/Engine/Classes/Engine/Canvas.h
--- a/Engine/Source/Runtime/Engine/Classes/Engine/Canvas.h
+++ b/Engine/Source/Runtime/Engine/Classes/Engine/Canvas.h
@@ -1,17 +1,15 @@
 // Copyright 1998-2015 Epic Games, Inc. All Rights Reserved.
 
-//=============================================================================
-// Canvas: A drawing canvas.
-//=============================================================================
-
 #pragma once
 
 #include "CanvasTypes.h"
 #include "Canvas.generated.h"
 
+
 class FSceneView;
 class UFont;
 
+
 /**
  * Holds texture information with UV coordinates as well.
  */
@@ -37,20 +35,23 @@ struct FCanvasIcon
 	UPROPERTY(EditAnywhere, BlueprintReadWrite, Category=CanvasIcon)
 	float VL;
 
-
 	FCanvasIcon()
-		: Texture(NULL)
+		: Texture(nullptr)
 		, U(0)
 		, V(0)
 		, UL(0)
 		, VL(0)
-	{
-	}
+	{ }
 
 };
 
+
+/**
+ * A drawing canvas.
+ */
 UCLASS(transient)
-class ENGINE_API UCanvas : public UObject
+class ENGINE_API UCanvas
+	: public UObject
 {
 	GENERATED_UCLASS_BODY()
 
@@ -110,6 +111,7 @@ class ENGINE_API UCanvas : public UObject
 	int32 CachedDisplayHeight;
 
 public:
+
 	FCanvas* Canvas;
 	FSceneView* SceneView;
 	FMatrix	ViewProjectionMatrix;
@@ -117,55 +119,56 @@ public:
 
 	// UCanvas interface.
 	
-	/** Initialise Canvas */
+	/** Initializes the canvas. */
 	void Init(int32 InSizeX, int32 InSizeY, FSceneView* InSceneView);
 
-	/** Changes the view for the canvas */
+	/** Changes the view for the canvas. */
 	void SetView(FSceneView* InView);
 
-	/** Update Canvas */
+	/** Updates the canvas. */
 	void Update();
 
-	/* Applies the current Platform's safezone to the current Canvas position.  Automatically called by Update */
+	/* Applies the current Platform's safe zone to the current Canvas position. Automatically called by Update. */
 	void ApplySafeZoneTransform();
 	void PopSafeZoneTransform();
 
-	/* Updates cached SafeZone data from the device.  Call when main device is resized */
+	/* Updates cached SafeZone data from the device.  Call when main device is resized. */
 	void UpdateSafeZoneData();
 
-	/* Function to go through all constructed canvas items and update their safezone data */
+	/* Function to go through all constructed canvas items and update their safe zone data. */
 	static void UpdateAllCanvasSafeZoneData();
 	
-
 	/** 
 	 * Draw arbitrary aligned rectangle.
-	 * @param Tex - Texture to draw
-	 * @param X - Position to draw X
-	 * @param Y - Position to draw Y
-	 * @param XL - Width of tile
-	 * @param YL - Height of tile
-	 * @param U - Horizontal position of the upper left corner of the portion of the texture to be shown(texels)
-	 * @param V - Vertical position of the upper left corner of the portion of the texture to be shown(texels)
-	 * @param UL - The width of the portion of the texture to be drawn(texels). 
-	 * @param VL - The height of the portion of the texture to be drawn(texels). 
-	 * @param ClipTile - true to clip tile
-	 * @param BlendMode - blending mode of texture
+	 *
+	 * @param Tex Texture to draw.
+	 * @param X Position to draw X.
+	 * @param Y Position to draw Y.
+	 * @param XL Width of tile.
+	 * @param YL Height of tile.
+	 * @param U Horizontal position of the upper left corner of the portion of the texture to be shown(texels).
+	 * @param V Vertical position of the upper left corner of the portion of the texture to be shown(texels).
+	 * @param UL The width of the portion of the texture to be drawn(texels).
+	 * @param VL The height of the portion of the texture to be drawn(texels).
+	 * @param ClipTile true to clip tile.
+	 * @param BlendMode Blending mode of texture.
 	 */
 	void DrawTile(UTexture* Tex, float X, float Y, float XL, float YL, float U, float V, float UL, float VL, EBlendMode BlendMode=BLEND_Translucent);
 
-	/** calculate the length of a string 
-	 * @param Font - font used
-	 * @param ScaleX - scale in X axis
-	 * @param ScaleY - scale in Y axis
-	 * @param XL out - horizontal length of string
-	 * @param YL out - vertical length of string
-	 * @param Text - string to calculate for
+	/**
+	 * Calculate the length of a string.
+	 *
+	 * @param Font The font used.
+	 * @param ScaleX Scale in X axis.
+	 * @param ScaleY Scale in Y axis.
+	 * @param XL out Horizontal length of string.
+	 * @param YL out Vertical length of string.
+	 * @param Text String to calculate for.
 	 */
 	static void ClippedStrLen(const UFont* Font, float ScaleX, float ScaleY, int32& XL, int32& YL, const TCHAR* Text);
 
 	/**	
-	 * Calculate the size of a string built from a font, word wrapped
-	 * to a specified region.
+	 * Calculate the size of a string built from a font, word wrapped to a specified region.
 	 */
 	void VARARGS WrappedStrLenf(const UFont* Font, float ScaleX, float ScaleY, int32& XL, int32& YL, const TCHAR* Fmt, ...);
 		
@@ -174,16 +177,17 @@ public:
 	 */
 	int32 WrappedPrint(bool Draw, float X, float Y, int32& out_XL, int32& out_YL, const UFont* Font, float ScaleX, float ScaleY, bool bCenterTextX, bool bCenterTextY, const TCHAR* Text, const FFontRenderInfo& RenderInfo) ;
 	
-	/** Draws a string of text to the screen 
-	 * @param InFont - font to draw with
-	 * @param InText - string to be drawn
-	 * @param X - Position to draw X
-	 * @param Y - Position to draw Y
-	 * @param XScale - Optional. The horizontal scaling to apply to the text. 
-	 * @param YScale - Optional. The vertical scaling to apply to the text. 
-	 * @param RenderInfo - Optional. The FontRenderInfo to use when drawing the text.
+	/**
+	 * Draws a string of text to the screen.
 	 *
-	 * returns the Y extent of the rendered text
+	 * @param InFont The font to draw with.
+	 * @param InText The string to be drawn.
+	 * @param X Position to draw X.
+	 * @param Y Position to draw Y.
+	 * @param XScale Optional. The horizontal scaling to apply to the text. 
+	 * @param YScale Optional. The vertical scaling to apply to the text. 
+	 * @param RenderInfo Optional. The FontRenderInfo to use when drawing the text.
+	 * @return The Y extent of the rendered text.
 	 */
 	float DrawText(const UFont* InFont, const FString& InText, float X, float Y, float XScale = 1.f, float YScale = 1.f, const FFontRenderInfo& RenderInfo = FFontRenderInfo());
 
@@ -228,7 +232,6 @@ public:
 	 */
 	static void CanvasStringSize( FTextSizingParameters& Parameters, const TCHAR* pText );
 
-
 	/**
 	 * Parses a single string into an array of strings that will fit inside the specified bounding region.
 	 *
@@ -254,35 +257,42 @@ public:
 
 	void WrapString( FTextSizingParameters& Parameters, const float InCurX, const TCHAR* const pText, TArray<FWrappedStringElement>& out_Lines, FCanvasWordWrapper::FWrappedLineData* const OutWrappedLineData = nullptr);
 
-
-	/** Transforms a 3D world-space vector into 2D screen coordinates. */
+	/**
+	 * Transforms a 3D world-space vector into 2D screen coordinates.
+	 *
+	 * @param Location The vector to transform.
+	 * @return The transformed vector.
+	 */
 	FVector Project(FVector Location) const;
 
 	/** 
 	 * Transforms 2D screen coordinates into a 3D world-space origin and direction.
-	 * @param ScreenPos - screen coordinates in pixels
-	 * @param WorldOrigin (out) - world-space origin vector
-	 * @param WorldDirection (out) - world-space direction vector
+	 *
+	 * @param ScreenPos Screen coordinates in pixels.
+	 * @param WorldOrigin (out) World-space origin vector.
+	 * @param WorldDirection (out) World-space direction vector.
 	 */
 	void Deproject(FVector2D ScreenPos, /*out*/ FVector& WorldOrigin, /*out*/ FVector& WorldDirection) const;
 
 	/**
 	 * Calculate the length of a string, taking text wrapping into account.
-	 * @param InFont - The Font use
-	 * @param Text	 - string to calculate for
-	 * @param XL out - horizontal length of string
-	 * @param YL out - vertical length of string
+	 *
+	 * @param InFont The Font use.
+	 * @param Text The string to calculate for.
+	 * @param XL out Horizontal length of string.
+	 * @param YL out Vertical length of string.
 	 */
 	void StrLen(const UFont* InFont, const FString& InText, float& XL, float& YL);
 
 	/** 
 	 * Calculates the horizontal and vertical size of a given string. This is used for clipped text as it does not take wrapping into account.
-	 * @param Font	 - Font 
-	 * @param Text	 - string to calculate for
-	 * @param XL out - horizontal length of string
-	 * @param YL out - vertical length of string
-	 * @param ScaleX - scale that the string is expected to draw at horizontally
-	 * @param ScaleY - scale that the string is expected to draw at vertically
+	 *
+	 * @param Font The font to use.
+	 * @param Text String to calculate for.
+	 * @param XL out Horizontal length of string.
+	 * @param YL out Vertical length of string.
+	 * @param ScaleX Scale that the string is expected to draw at horizontally.
+	 * @param ScaleY Scale that the string is expected to draw at vertically.
 	 */
 	void TextSize( const UFont* InFont, const FString& InText, float& XL, float& YL, float ScaleX=1.f, float ScaleY=1.f);
 	
diff --git a/Engine/Source/Runtime/Engine/Classes/Engine/Channel.h b/Engine/Source/Runtime/Engine/Classes/Engine/Channel.h
--- a/Engine/Source/Runtime/Engine/Classes/Engine/Channel.h
+++ b/Engine/Source/Runtime/Engine/Classes/Engine/Channel.h
@@ -1,16 +1,16 @@
 // Copyright 1998-2015 Epic Games, Inc. All Rights Reserved.
 
-//
-// Base class of communication channels.
-//
-
 #pragma once
+
 #include "Channel.generated.h"
 
+
 // Constant for all buffers that are reading from the network
 const int MAX_STRING_SERIALIZE_SIZE	= NAME_SIZE;
 
-// Types of channels.
+/**
+ * Enumerates channel types.
+ */
 enum EChannelType
 {
 	CHTYPE_None			= 0,  // Invalid type.
@@ -21,11 +21,17 @@ enum EChannelType
 	CHTYPE_MAX          = 8,  // Maximum.
 };
 
+
 // The channel index to use for voice
 #define VOICE_CHANNEL_INDEX 1
 
+
+/**
+ * Base class of communication channels.
+ */
 UCLASS(transient)
-class ENGINE_API UChannel : public UObject
+class ENGINE_API UChannel
+	: public UObject
 {
 	GENERATED_UCLASS_BODY()
 	
@@ -56,15 +62,13 @@ class ENGINE_API UChannel : public UObject
 	/** @return true if the specified channel type exists. */
 	static bool IsKnownChannelType( int32 Type );
 
-	
-	// Begin UObject Interface
+public:
 
+	// UObject overrides
 
 	virtual void BeginDestroy() override;
 
-
-	// End UObject Interface
-	
+public:	
 
 	/** UChannel interface. */
 	virtual void Init( UNetConnection* InConnection, int32 InChIndex, bool InOpenedLocally );
@@ -125,7 +129,7 @@ protected:
 	/** Closes the actor channel but with a 'dormant' flag set so it can be reopened */
 	virtual void BecomeDormant() { }
 
-	/** cleans up channel structures and NULLs references to the channel */
+	/** cleans up channel structures and nulls references to the channel */
 	virtual bool CleanUp( const bool bForDestroy );
 
 private:
@@ -133,7 +137,7 @@ private:
 	/** Just sends the bunch out on the connection */
 	int32 SendRawBunch(FOutBunch* Bunch, bool Merge);
 
-	/** Final step to prepare bunch to be sent. If reliable, adds to ack list. */
+	/** Final step to prepare bunch to be sent. If reliable, adds to acknowldege list. */
 	FOutBunch* PrepBunch(FOutBunch* Bunch, FOutBunch* OutBunch, bool Merge);
 
 	/** Received next bunch to process. This handles partial bunches */
diff --git a/Engine/Source/Runtime/Engine/Classes/Engine/ChildConnection.h b/Engine/Source/Runtime/Engine/Classes/Engine/ChildConnection.h
--- a/Engine/Source/Runtime/Engine/Classes/Engine/ChildConnection.h
+++ b/Engine/Source/Runtime/Engine/Classes/Engine/ChildConnection.h
@@ -1,19 +1,26 @@
 // Copyright 1998-2015 Epic Games, Inc. All Rights Reserved.
 
-/** represents a secondary splitscreen connection that reroutes calls to the parent connection */
-
 #pragma once
+
 #include "ChildConnection.generated.h"
 
+
+/**
+ * Represents a secondary split screen connection that reroutes calls to the parent connection.
+ */
 UCLASS(MinimalAPI,transient,config=Engine)
-class UChildConnection : public UNetConnection
+class UChildConnection
+	: public UNetConnection
 {
 	GENERATED_UCLASS_BODY()
 
 	UPROPERTY(transient)
 	class UNetConnection* Parent;
 
-	// Begin UNetConnection interface.
+public:
+
+	// UNetConnection interface.
+
 	virtual UChildConnection* GetUChildConnection() override
 	{
 		return this;
@@ -29,7 +36,7 @@ class UChildConnection : public UNetConnection
 		return Parent->LowLevelDescribe();
 	}
 
-		virtual void LowLevelSend( void* Data, int32 Count ) override
+	virtual void LowLevelSend( void* Data, int32 Count ) override
 	{
 	}
 
@@ -63,9 +70,5 @@ class UChildConnection : public UNetConnection
 	}
 
 	virtual void HandleClientPlayer(class APlayerController* PC, class UNetConnection* NetConnection) override;
-	
 	virtual void CleanUp() override;
-
-	// End UNetConnection interface.
-
 };
diff --git a/Engine/Source/Runtime/Engine/Classes/Engine/CloudStorageBase.h b/Engine/Source/Runtime/Engine/Classes/Engine/CloudStorageBase.h
--- a/Engine/Source/Runtime/Engine/Classes/Engine/CloudStorageBase.h
+++ b/Engine/Source/Runtime/Engine/Classes/Engine/CloudStorageBase.h
@@ -1,13 +1,10 @@
 // Copyright 1998-2015 Epic Games, Inc. All Rights Reserved.
 
-/**
- *
- * This is the base class for the various platform interface classes
- */
-
 #pragma once
+
 #include "CloudStorageBase.generated.h"
 
+
 /** All the types of delegate callbacks that a CloudStorage subclass may receive from C++ */
 UENUM()
 enum ECloudStorageDelegate
@@ -28,9 +25,13 @@ enum ECloudStorageDelegate
 	CSD_MAX,
 };
 
-	// Full API exported so that platform modules can extend it (e.g. FacebookWindows in WinDrv)
+
+/**
+ * Base class for the various platform interface classes.
+ */
 UCLASS()
-class UCloudStorageBase : public UPlatformInterfaceBase
+class UCloudStorageBase
+	: public UPlatformInterfaceBase
 {
 	GENERATED_UCLASS_BODY()
 
@@ -42,32 +43,27 @@ class UCloudStorageBase : public UPlatformInterfaceBase
 	UPROPERTY()
 	uint32 bSuppressDelegateCalls:1;
 
-
-	/**
-	 * Perform any initialization
-	 */
+	/** Performs any initialization. */
 	virtual void Init();
 
 	/**
 	 * Initiate reading a key/value pair from cloud storage. A CSD_KeyValueReadComplete
-	 * delegate will be called when it completes (if this function returns true)
+	 * delegate will be called when it completes (if this function returns true).
 	 * 
-	 * @param KeyName String name of the key to retrieve
-	 * @param Type Type of data to retrieve 
-	 * @param SerializedObj If you are reading an object (PIDT_Object), it will de-serialize the binary blob into this object
-	 *
-	 * @return True if successful
+	 * @param KeyName String name of the key to retrieve.
+	 * @param Type Type of data to retrieve .
+	 * @param SerializedObj If you are reading an object (PIDT_Object), it will de-serialize the binary blob into this object.
+	 * @return true on success, false otherwise.
 	 */
-	virtual bool ReadKeyValue(const FString& KeyName, EPlatformInterfaceDataType Type, class UObject* SerializedObj = NULL);
+	virtual bool ReadKeyValue(const FString& KeyName, EPlatformInterfaceDataType Type, class UObject* SerializedObj = nullptr);
 
 	/**
 	 * Write a key/value pair to the cloud. A CSD_KeyValueWriteComplete
-	 * delegate will be called when it completes (if this function returns true)
+	 * delegate will be called when it completes (if this function returns true).
 	 *
-	 * @param KeyName String name of the key to write
-	 * @param Value The type and value to write
-	 *
-	 * @return True if successful
+	 * @param KeyName String name of the key to write.
+	 * @param Value The type and value to write.
+	 * @return true on success, false otherwise.
 	 */
 	virtual bool WriteKeyValue(const FString& KeyName, const FPlatformInterfaceData& Value);
 
@@ -78,7 +74,7 @@ class UCloudStorageBase : public UPlatformInterfaceBase
 	 * function returns true). Then use GetNumCloudDocuments() and GetCloudDocumentName() 
 	 * to get the information about any existing documents.
 	 *
-	 * @return True if successful
+	 * @return true on success, false otherwise.
 	 */
 	virtual bool QueryForCloudDocuments();
 
@@ -95,9 +91,8 @@ class UCloudStorageBase : public UPlatformInterfaceBase
 	/**
 	 * Create a new document in the cloud (uninitialized, unsaved, use the Write/Save functions)
 	 *
-	 * @param Filename Filename for the cloud document (with any extension you want)
-	 * 
-	 * @return Index of the new document, or -1 on failure
+	 * @param Filename Filename for the cloud document (with any extension you want).
+	 * @return Index of the new document, or -1 on failure.
 	 */
 	virtual int32 CreateCloudDocument(const FString& Filename);
 
@@ -107,7 +102,6 @@ class UCloudStorageBase : public UPlatformInterfaceBase
 	 * will be called (if this function returns true).
 	 *
 	 * @param Index index of the document to read
-	 *
 	 * @param True if successful
 	 */
 	virtual bool ReadCloudDocument(int32 Index, bool bIsForConflict = false);
@@ -118,31 +112,28 @@ class UCloudStorageBase : public UPlatformInterfaceBase
 	 * generate the document.
 	 *
 	 * @param Index index of the document to read
-	 *
 	 * @return The document as a string. It will be empty if anything went wrong.
 	 */
 	virtual FString ParseDocumentAsString(int32 Index, bool bIsForConflict = false);
 
 	/**
-	 * Once a document has been read in, use this to return a string representing the 
-	 * entire document. This should only be used if SaveDocumentWithString was used to
-	 * generate the document.
+	 * Once a document has been read in, use this to return a string representing the.
+	 * entire document. This should only be used if SaveDocumentWithString was used to generate the document.
 	 *
-	 * @param Index index of the document to read
-	 * @param ByteData The array of bytes to be filled out. It will be empty if anything went wrong
+	 * @param Index index of the document to read.
+	 * @param ByteData The array of bytes to be filled out. It will be empty if anything went wrong.
 	 */
 	virtual void ParseDocumentAsBytes(int32 Index, TArray<uint8>& ByteData, bool bIsForConflict = false);
 
 	/**
-	 * Once a document has been read in, use this to return a string representing the 
+	 * Once a document has been read in, use this to return a string representing the
 	 * entire document. This should only be used if SaveDocumentWithString was used to
 	 * generate the document.
 	 *
-	 * @param Index index of the document to read
-	 * @param ExpectedVersion Version number expected to be in the save data. If this doesn't match what's there, this function will return NULL
-	 * @param ObjectClass The class of the object to create
-	 *
-	 * @return The object deserialized from the document. It will be none if anything went wrongs
+	 * @param Index index of the document to read.
+	 * @param ExpectedVersion Version number expected to be in the save data. If this doesn't match what's there, this function will return nullptr.
+	 * @param ObjectClass The class of the object to create.
+	 * @return The object deserialized from the document. It will be none if anything went wrongs.
 	 */
 	virtual class UObject* ParseDocumentAsObject(int32 Index, TSubclassOf<class UObject> ObjectClass, int32 ExpectedVersion, bool bIsForConflict = false);
 
@@ -150,58 +141,50 @@ class UCloudStorageBase : public UPlatformInterfaceBase
 	 * Writes a document that has been already "saved" using the SaveDocumentWith* functions.
 	 * A CSD_DocumentWriteComplete delegate will be called (if this function returns true).
 	 *
-	 * @param Index index of the document to read
-	 *
-	 * @param True if successful
+	 * @param Index index of the document to read.
+	 * @param True if successful.
 	 */
 	virtual bool WriteCloudDocument(int32 Index);
 
 	/**
-	 * Prepare a document for writing to the cloud with a string as input data. This is
-	 * synchronous
-	 *
-	 * @param Index index of the document to save
-	 * @param StringData The data to put into the document
+	 * Prepare a document for writing to the cloud with a string as input data. This is synchronous.
 	 *
+	 * @param Index index of the document to save.
+	 * @param StringData The data to put into the document.
 	 * @param True if successful
 	 */
 	virtual bool SaveDocumentWithString(int32 Index, const FString& StringData);
 
 	/**
-	 * Prepare a document for writing to the cloud with an array of bytes as input data. This is
-	 * synchronous
-	 *
-	 * @param Index index of the document to save
-	 * @param ByteData The array of generic bytes to put into the document
+	 * Prepare a document for writing to the cloud with an array of bytes as input data. This is synchronous.
 	 *
-	 * @param True if successful
+	 * @param Index index of the document to save.
+	 * @param ByteData The array of generic bytes to put into the document.
+	 * @param True if successful.
 	 */
 	virtual bool SaveDocumentWithBytes(int32 Index, const TArray<uint8>& ByteData);
 
 	/**
-	 * Prepare a document for writing to the cloud with an object as input data. This is
-	 * synchronous
-	 *
-	 * @param Index index of the document to save
-	 * @param ObjectData The object to serialize to bytes and put into the document
+	 * Prepare a document for writing to the cloud with an object as input data. This is synchronous.
 	 *
-	 * @param True if successful
+	 * @param Index index of the document to save.
+	 * @param ObjectData The object to serialize to bytes and put into the document.
+	 * @param true on success, false otherwise.
 	 */
 	virtual bool SaveDocumentWithObject(int32 Index, class UObject* ObjectData, int32 SaveVersion);
 
 	/**
 	 * If there was a conflict notification, this will simply tell the cloud interface
-	 * to choose the most recently modified version, and toss any others
+	 * to choose the most recently modified version, and toss any others.
 	 */
 	virtual bool ResolveConflictWithNewestDocument();
 
 	/**
 	 * If there was a conflict notification, this will tell the cloud interface
 	 * to choose the version with a given Index to be the master version, and to
-	 * toss any others
+	 * toss any others.
 	 *
-	 * @param Index Conflict version index
+	 * @param Index Conflict version index.
 	 */
 	virtual bool ResolveConflictWithVersionIndex(int32 Index);
 };
-
diff --git a/Engine/Source/Runtime/Engine/Classes/Engine/Console.h b/Engine/Source/Runtime/Engine/Classes/Engine/Console.h
--- a/Engine/Source/Runtime/Engine/Classes/Engine/Console.h
+++ b/Engine/Source/Runtime/Engine/Classes/Engine/Console.h
@@ -1,9 +1,7 @@
 // Copyright 1998-2015 Epic Games, Inc. All Rights Reserved.
 
-//=============================================================================
-// Console - A quick little command line console that accepts most commands.
-//=============================================================================
 #pragma once
+
 #include "ConsoleSettings.h"
 #include "Console.generated.h"
 
@@ -49,6 +47,9 @@ struct FAutoCompleteNode
 };
 
 
+/**
+ * A basic command line console that accepts most commands.
+ */
 UCLASS(Within=GameViewportClient, config=Input, transient)
 class ENGINE_API UConsole
 	: public UObject
@@ -56,7 +57,7 @@ class ENGINE_API UConsole
 {
 	GENERATED_UCLASS_BODY()
 
-	/** The player which the next console command should be executed in the context of.  If NULL, execute in the viewport. */
+	/** The player which the next console command should be executed in the context of.  If nullptr, execute in the viewport. */
 	UPROPERTY()
 	class ULocalPlayer* ConsoleTargetPlayer;
 
@@ -131,7 +132,7 @@ class ENGINE_API UConsole
 	
 	/**
 	 * Executes a console command.
-	 * @param Command - The command to execute.
+	 * @param Command The command to execute.
 	 */
 	virtual void ConsoleCommand(const FString& Command);
 	
@@ -143,13 +144,15 @@ class ENGINE_API UConsole
 	/**
 	 * Prints a (potentially multi-line) FString of text to the console.
 	 * The text is split into separate lines and passed to OutputTextLine.
-	 * @param Text - Text to display on the console.
+	 *
+	 * @param Text Text to display on the console.
 	 */
 	virtual void OutputText(const FString& Text);
 	
 	/**
 	 * Opens the typing bar with text already entered.
-	 * @param Text - The text to enter in the typing bar.
+	 *
+	 * @param Text The text to enter in the typing bar.
 	 */
 	virtual void StartTyping(const FString& Text);
 	
@@ -245,7 +248,7 @@ private:
 	 * @param Text - A line of text to display on the console.
 	 */
 	void OutputTextLine(const FString& Text);
-	
+
 	void PostRender_InputLine(class UCanvas* Canvas, FIntPoint UserInputLinePos);
 
 	void SetAutoCompleteFromHistory();
diff --git a/Engine/Source/Runtime/Engine/Classes/Engine/ControlChannel.h b/Engine/Source/Runtime/Engine/Classes/Engine/ControlChannel.h
--- a/Engine/Source/Runtime/Engine/Classes/Engine/ControlChannel.h
+++ b/Engine/Source/Runtime/Engine/Classes/Engine/ControlChannel.h
@@ -1,14 +1,16 @@
 // Copyright 1998-2015 Epic Games, Inc. All Rights Reserved.
 
 #pragma once
+
 #include "ControlChannel.generated.h"
 
-//
-// A channel for exchanging connection control messages
-//
 
+/**
+ * A channel for exchanging connection control messages.
+ */
 UCLASS(transient, customConstructor)
-class UControlChannel : public UChannel
+class UControlChannel
+	: public UChannel
 {
 	GENERATED_UCLASS_BODY()
 
diff --git a/Engine/Source/Runtime/Engine/Classes/Engine/CullDistanceVolume.h b/Engine/Source/Runtime/Engine/Classes/Engine/CullDistanceVolume.h
--- a/Engine/Source/Runtime/Engine/Classes/Engine/CullDistanceVolume.h
+++ b/Engine/Source/Runtime/Engine/Classes/Engine/CullDistanceVolume.h
@@ -1,10 +1,11 @@
 // Copyright 1998-2015 Epic Games, Inc. All Rights Reserved.
 
-
 #pragma once
+
 #include "GameFramework/Volume.h"
 #include "CullDistanceVolume.generated.h"
 
+
 /**
  * Helper structure containing size and cull distance pair.
  */
@@ -25,20 +26,18 @@ struct FCullDistanceSizePair
 	FCullDistanceSizePair()
 		: Size(0)
 		, CullDistance(0)
-	{
-	}
+	{ }
 
 	FCullDistanceSizePair(float InSize, float InCullDistance)
 		: Size(InSize)
 		, CullDistance(InCullDistance)
-	{
-	}
-
-
+	{ }
 };
 
+
 UCLASS(hidecategories=(Advanced, Attachment, Collision, Volume))
-class ACullDistanceVolume : public AVolume
+class ACullDistanceVolume
+	: public AVolume
 {
 	GENERATED_UCLASS_BODY()
 
@@ -55,14 +54,16 @@ class ACullDistanceVolume : public AVolume
 	UPROPERTY(EditAnywhere, BlueprintReadWrite, Category=CullDistanceVolume)
 	uint32 bEnabled:1;
 
+public:
 
+	// UObject Interface
 
-	// Begin UObject Interface
 #if WITH_EDITOR
 	virtual void PostEditChangeProperty(FPropertyChangedEvent& PropertyChangedEvent) override;
 	virtual void PostEditMove(bool bFinished) override;
 #endif // WITH_EDITOR
-	// End UObject Interface
+
+public:
 
 	/** 
 	 * Override Destroyed so that we can re-calculate primitive draw distances after this volume has been deleted.
@@ -80,6 +81,3 @@ class ACullDistanceVolume : public AVolume
 	/** Get the set of primitives and new max draw distances defined by this volume. */
 	void GetPrimitiveMaxDrawDistances(TMap<UPrimitiveComponent*,float>& OutCullDistances);
 };
-
-
-
diff --git a/Engine/Source/Runtime/Engine/Classes/Engine/CurveTable.h b/Engine/Source/Runtime/Engine/Classes/Engine/CurveTable.h
--- a/Engine/Source/Runtime/Engine/Classes/Engine/CurveTable.h
+++ b/Engine/Source/Runtime/Engine/Classes/Engine/CurveTable.h
@@ -1,19 +1,26 @@
 // Copyright 1998-2015 Epic Games, Inc. All Rights Reserved.
 
 #pragma once
+
 #include "CurveTable.generated.h"
 
+
 ENGINE_API DECLARE_LOG_CATEGORY_EXTERN(LogCurveTable, Log, All);
 
+
 // forward declare JSON writer
 template <class CharType>
 struct TPrettyJsonPrintPolicy;
 template <class CharType, class PrintPolicy>
 class TJsonWriter;
 
-/** Imported spreadsheet table as curves */
+
+/**
+ * Imported spreadsheet table as curves.
+ */
 UCLASS(MinimalAPI)
-class UCurveTable : public UObject
+class UCurveTable
+	: public UObject
 {
 	GENERATED_UCLASS_BODY()
 
@@ -38,15 +45,15 @@ class UCurveTable : public UObject
 		if(RowName == NAME_None)
 		{
 			UE_CLOG(WarnIfNotFound, LogCurveTable, Warning, TEXT("UCurveTable::FindRow : NAME_None is invalid row name for CurveTable '%s' (%s)."), *GetPathName(), *ContextString);
-			return NULL;
+			return nullptr;
 		}
 
 		FRichCurve* const* FoundCurve = RowMap.Find(RowName);
 
-		if(FoundCurve == NULL)
+		if(FoundCurve == nullptr)
 		{
 			UE_CLOG(WarnIfNotFound, LogCurveTable, Warning, TEXT("UCurveTable::FindRow : Row '%s' not found in CurveTable '%s' (%s)."), *RowName.ToString(), *GetPathName(), *ContextString);
-			return NULL;
+			return nullptr;
 		}
 
 		return (FRichCurve*)*FoundCurve;
@@ -82,18 +89,19 @@ protected:
 
 };
 
-/** Handle to a particular row in a table*/
+
+/**
+ * Handle to a particular row in a table.
+ */
 USTRUCT(BlueprintType)
 struct ENGINE_API FCurveTableRowHandle
 {
 	GENERATED_USTRUCT_BODY()
 
 	FCurveTableRowHandle()
-		: CurveTable(NULL)
+		: CurveTable(nullptr)
 		, RowName(NAME_None)
-	{
-
-	}
+	{ }
 
 	/** Pointer to table we want a row from */
 	UPROPERTY(EditAnywhere, BlueprintReadWrite, Category=CurveTableRowHandle)
@@ -106,7 +114,7 @@ struct ENGINE_API FCurveTableRowHandle
 	/** Returns true if the curve is valid */
 	bool IsValid() const
 	{
-		return (GetCurve() != NULL);
+		return (GetCurve() != nullptr);
 	}
 
 	static const FString Unknown;
@@ -131,5 +139,6 @@ struct ENGINE_API FCurveTableRowHandle
 	bool operator!=(const FCurveTableRowHandle& Other) const;
 };
 
+
 /** Macro to call GetCurve with a correct error info. Assumed to be called within a UObject */
 #define GETCURVE_REPORTERROR(Handle) Handle.GetCurve(FString::Printf(TEXT("%s.%s"), *GetPathName(), TEXT(#Handle)))
diff --git a/Engine/Source/Runtime/Engine/Classes/Engine/DataAsset.h b/Engine/Source/Runtime/Engine/Classes/Engine/DataAsset.h
--- a/Engine/Source/Runtime/Engine/Classes/Engine/DataAsset.h
+++ b/Engine/Source/Runtime/Engine/Classes/Engine/DataAsset.h
@@ -1,18 +1,24 @@
 // Copyright 1998-2015 Epic Games, Inc. All Rights Reserved.
 
 #pragma once
+
 #include "DataAsset.generated.h"
 
 
-/** Base class for a simple asset containing data */
+/**
+ * Base class for a simple asset containing data.
+ */
 UCLASS(abstract,MinimalAPI)
-class UDataAsset : public UObject
+class UDataAsset
+	: public UObject
 {
 	GENERATED_UCLASS_BODY()
 
+public:
+
 	// UObject interface
+
 #if WITH_EDITORONLY_DATA
 	ENGINE_API virtual void Serialize(FArchive& Ar) override;
 #endif
-	// End of UObject interface
-};
\ No newline at end of file
+};
diff --git a/Engine/Source/Runtime/Engine/Classes/Engine/DataTable.h b/Engine/Source/Runtime/Engine/Classes/Engine/DataTable.h
--- a/Engine/Source/Runtime/Engine/Classes/Engine/DataTable.h
+++ b/Engine/Source/Runtime/Engine/Classes/Engine/DataTable.h
@@ -4,28 +4,35 @@
 
 #include "DataTable.generated.h"
 
+
 ENGINE_API DECLARE_LOG_CATEGORY_EXTERN(LogDataTable, Log, All);
 
+
 // forward declare JSON writer
 template <class CharType>
 struct TPrettyJsonPrintPolicy;
 template <class CharType, class PrintPolicy>
 class TJsonWriter;
 
-/** Base class for all table row structs to inherit from */
+
+/**
+ * Base class for all table row structs to inherit from.
+ */
 USTRUCT()
 struct FTableRowBase
 {
 	GENERATED_USTRUCT_BODY()
 
-	FTableRowBase()
-	{
-	}
+	FTableRowBase() { }
 };
 
-/** Imported spreadsheet table */
+
+/**
+ * Imported spreadsheet table.
+ */
 UCLASS(MinimalAPI)
-class UDataTable : public UObject
+class UDataTable
+	: public UObject
 {
 	GENERATED_UCLASS_BODY()
 
@@ -52,32 +59,32 @@ class UDataTable : public UObject
 	template <class T>
 	T* FindRow(FName RowName, const FString& ContextString, bool bWarnIfRowMissing = true) const
 	{
-		if(RowStruct == NULL)
+		if(RowStruct == nullptr)
 		{
 			UE_LOG(LogDataTable, Error, TEXT("UDataTable::FindRow : DataTable '%s' has no RowStruct specified (%s)."), *GetPathName(), *ContextString);
-			return NULL;
+			return nullptr;
 		}
 
 		if(!RowStruct->IsChildOf(T::StaticStruct()))
 		{
 			UE_CLOG(bWarnIfRowMissing, LogDataTable, Error, TEXT("UDataTable::FindRow : Incorrect type specified for DataTable '%s' (%s)."), *GetPathName(), *ContextString);
-			return NULL;
+			return nullptr;
 		}
 
 		if(RowName == NAME_None)
 		{
 			UE_CLOG(bWarnIfRowMissing, LogDataTable, Warning, TEXT("UDataTable::FindRow : NAME_None is invalid row name for DataTable '%s' (%s)."), *GetPathName(), *ContextString);
-			return NULL;
+			return nullptr;
 		}
 
 		uint8* const* RowDataPtr = RowMap.Find(RowName);
-		if (RowDataPtr == NULL)
+		if (RowDataPtr == nullptr)
 		{
 			if (bWarnIfRowMissing)
 			{
 				UE_LOG(LogDataTable, Warning, TEXT("UDataTable::FindRow : Row '%s' not found in DataTable '%s' (%s)."), *RowName.ToString(), *GetPathName(), *ContextString);
 			}
-			return NULL;
+			return nullptr;
 		}
 
 		uint8* RowData = *RowDataPtr;
@@ -86,28 +93,28 @@ class UDataTable : public UObject
 		return (T*)RowData;
 	}
 
-	/** Returns the column property where PropertyName matches the name of the column property. Returns NULL if no match is found or the match is not a supported table property */
+	/** Returns the column property where PropertyName matches the name of the column property. Returns nullptr if no match is found or the match is not a supported table property */
 	ENGINE_API UProperty* FindTableProperty(const FName& PropertyName) const;
 
 	uint8* FindRowUnchecked(FName RowName, bool MustExist=false) const
 	{
-		if(RowStruct == NULL)
+		if(RowStruct == nullptr)
 		{
 			//UE_CLOG(MustExist, LogDataTable, Error, TEXT("UDataTable::FindRow : DataTable '%s' has no RowStruct specified (%s)."), *GetPathName(), *ContextString);
-			return NULL;
+			return nullptr;
 		}
 
 		if(RowName == NAME_None)
 		{
 			//UE_CLOG(MustExist, LogDataTable, Warning, TEXT("UDataTable::FindRow : NAME_None is invalid row name for DataTable '%s' (%s)."), *GetPathName(), *ContextString);
-			return NULL;
+			return nullptr;
 		}
 
 		uint8* const* RowDataPtr = RowMap.Find(RowName);
 
-		if(RowDataPtr == NULL)
+		if(RowDataPtr == nullptr)
 		{
-			return NULL;
+			return nullptr;
 		}
 
 		uint8* RowData = *RowDataPtr;
@@ -124,10 +131,13 @@ class UDataTable : public UObject
 #if WITH_EDITOR
 
 private:
+
 	//when RowStruct is being modified, row data is stored serialized with tags
 	TArray<uint8> RowsSerializedWithTags;
 	TSet<UObject*> TemporarilyReferencedObjects;
+
 public:
+
 	ENGINE_API void CleanBeforeStructChange();
 	ENGINE_API void RestoreAfterStructChange();
 
@@ -180,6 +190,7 @@ private:
 	void LoadStructData(FArchive& Ar);
 };
 
+
 /** Handle to a particular row in a table*/
 USTRUCT(BlueprintType)
 struct ENGINE_API FDataTableRowHandle
@@ -187,7 +198,7 @@ struct ENGINE_API FDataTableRowHandle
 	GENERATED_USTRUCT_BODY()
 
 	FDataTableRowHandle()
-		: DataTable(NULL)
+		: DataTable(nullptr)
 		, RowName(NAME_None)
 	{
 
@@ -207,19 +218,20 @@ struct ENGINE_API FDataTableRowHandle
 	template <class T>
 	T* GetRow(const FString& ContextString = Unknown) const
 	{
-		if(DataTable == NULL)
+		if(DataTable == nullptr)
 		{
 			if (RowName != NAME_None)
 			{
 				UE_LOG(LogDataTable, Warning, TEXT("FDataTableRowHandle::FindRow : No DataTable for row %s (%s)."), *RowName.ToString(), *ContextString);
 			}
-			return NULL;
+			return nullptr;
 		}
 
 		return DataTable->FindRow<T>(RowName, ContextString);
 	}
 };
 
+
 /** Handle to a particular row in a table*/
 USTRUCT()
 struct ENGINE_API FDataTableCategoryHandle
@@ -245,7 +257,7 @@ struct ENGINE_API FDataTableCategoryHandle
 	void GetRows(TArray<T*>& OutRows, const FString& ContextString = Unknown) const
 	{
 		OutRows.Empty();
-		if (DataTable == NULL)
+		if (DataTable == nullptr)
 		{
 			if (RowContents != NAME_None)
 			{
@@ -267,7 +279,7 @@ struct ENGINE_API FDataTableCategoryHandle
 
 		// Find the property that matches the desired column (ColumnName)
 		UProperty* Property = DataTable->FindTableProperty(ColumnName);
-		if (Property == NULL)
+		if (Property == nullptr)
 		{
 			return;
 		}
@@ -281,7 +293,7 @@ struct ENGINE_API FDataTableCategoryHandle
 
 			FString PropertyValue(TEXT(""));
 
-			Property->ExportText_InContainer(0, PropertyValue, RowData, RowData, NULL, PPF_None);
+			Property->ExportText_InContainer(0, PropertyValue, RowData, RowData, nullptr, PPF_None);
 
 			if (RowContentsAsString == PropertyValue)
 			{
@@ -293,6 +305,7 @@ struct ENGINE_API FDataTableCategoryHandle
 	}
 };
 
+
 /** Macro to call GetRow with a correct error info. Assumed to be called from within a UObject */
 #define GETROW_REPORTERROR(Handle, Template) Handle.GetRow<Template>(FString::Printf(TEXT("%s.%s"), *GetPathName(), TEXT(#Handle)))
 #define GETROWOBJECT_REPORTERROR(Object, Handle, Template) Handle.GetRow<Template>(FString::Printf(TEXT("%s.%s"), *Object->GetPathName(), TEXT(#Handle)))
diff --git a/Engine/Source/Runtime/Engine/Classes/Engine/DebugCameraController.h b/Engine/Source/Runtime/Engine/Classes/Engine/DebugCameraController.h
--- a/Engine/Source/Runtime/Engine/Classes/Engine/DebugCameraController.h
+++ b/Engine/Source/Runtime/Engine/Classes/Engine/DebugCameraController.h
@@ -1,18 +1,20 @@
 // Copyright 1998-2015 Epic Games, Inc. All Rights Reserved.
 
-/**
- * Camera controller that allows you to fly around a level mostly unrestricted by normal movement rules.
- * 
- * To turn it on, please press Alt+C or both (left and right) analogs on xbox pad, or use the "ToggleDebugCamera" console command.
- * 
- * Check the debug camera bindings in DefaultPawn.cpp for the camera controls.
- */
-
 #pragma once
+
 #include "DebugCameraController.generated.h"
 
+
+/**
+* Camera controller that allows you to fly around a level mostly unrestricted by normal movement rules.
+*
+* To turn it on, please press Alt+C or both (left and right) analogs on XBox pad,
+* or use the "ToggleDebugCamera" console command. Check the debug camera bindings
+* in DefaultPawn.cpp for the camera controls.
+*/
 UCLASS(config=Game, hidedropdown)
-class ENGINE_API ADebugCameraController : public APlayerController
+class ENGINE_API ADebugCameraController
+	: public APlayerController
 {
 	GENERATED_UCLASS_BODY()
 
@@ -27,7 +29,6 @@ class ENGINE_API ADebugCameraController : public APlayerController
 	/** @todo document */
 	UPROPERTY()
 	class UDrawFrustumComponent* DrawFrustum;
-
 	
 	/** @todo document */
 	UFUNCTION(exec)
@@ -36,7 +37,7 @@ class ENGINE_API ADebugCameraController : public APlayerController
 	/** Selects the object the camera is aiming at. */
 	void SelectTargetedObject();
 
-	/** Called when the user pressed the unselect key, just before the selected actor is cleared. */
+	/** Called when the user pressed the deselect key, just before the selected actor is cleared. */
 	void Unselect();
 
 	/** @todo document */
@@ -60,8 +61,8 @@ class ENGINE_API ADebugCameraController : public APlayerController
 	 */
 	virtual void ToggleFreezeRendering();
 
-
 public:
+
 	/** @todo document */
 	class AActor* SelectedActor;
 
@@ -93,13 +94,13 @@ protected:
 	virtual void SetupInputComponent() override;
 
 public:
+
 	/** Function called on activation debug camera controller */
 	void OnActivate(class APlayerController* OriginalPC);
 
 	/** Function called on deactivation debug camera controller */
 	void OnDeactivate(class APlayerController* RestoredPC);
 
-
 	/**
 	 * Builds a list of components that are hidden based upon gameplay
 	 *
@@ -108,15 +109,17 @@ public:
 	 */
 	virtual void UpdateHiddenComponents(const FVector& ViewLocation,TSet<FPrimitiveComponentId>& HiddenComponents);
 
+public:
+
+	// APlayerController Interface
 
-	// Begin APlayerController Interface
 	virtual void PostInitializeComponents() override;
 	virtual FString ConsoleCommand(const FString& Command, bool bWriteToLog = true) override;
 	virtual void AddCheats(bool bForce) override;
 	virtual void EndSpectatingState() override;
-	// End APlayerController Interface
 
 protected:
+
 	/**
 	 * Called when an actor has been selected with the primary key (e.g. left mouse button).
 	 * @param Hit	Info struct for the selection point.
@@ -126,6 +129,7 @@ protected:
 	virtual void SetSpectatorPawn(class ASpectatorPawn* NewSpectatorPawn) override;
 
 private:
+
 	/** The normalized screen location when a drag starts */
 	FVector2D LastTouchDragLocation;
 
@@ -133,6 +137,3 @@ private:
 	void OnTouchEnd(ETouchIndex::Type FingerIndex, FVector Location);
 	void OnFingerMove(ETouchIndex::Type FingerIndex, FVector Location);
 };
-
-
-
diff --git a/Engine/Source/Runtime/Engine/Classes/Engine/DebugDisplayProperty.h b/Engine/Source/Runtime/Engine/Classes/Engine/DebugDisplayProperty.h
--- a/Engine/Source/Runtime/Engine/Classes/Engine/DebugDisplayProperty.h
+++ b/Engine/Source/Runtime/Engine/Classes/Engine/DebugDisplayProperty.h
@@ -4,9 +4,9 @@
 
 #include "DebugDisplayProperty.generated.h"
 
+
 /** 
- * Debug property display functionality
- * to interact with this, use "display", "displayall", "displayclear"
+ * Debug property display functionality to interact with this, use "display", "displayall", "displayclear"
  *
  * @see UGameViewportClient
  * @see FDebugDisplayProperty
@@ -21,7 +21,7 @@ struct FDebugDisplayProperty
 	UPROPERTY()
 	class UObject* Obj;
 
-	/** if Obj is a class and WithinClass is not NULL, further limit the display to objects that have an Outer of WithinClass */
+	/** if Obj is a class and WithinClass is not nullptr, further limit the display to objects that have an Outer of WithinClass */
 	UPROPERTY()
 	TSubclassOf<class UObject>  WithinClass;
 
@@ -31,12 +31,10 @@ struct FDebugDisplayProperty
 	/** whether PropertyName is a "special" value not directly mapping to a real property (e.g. state name) */
 	uint32 bSpecialProperty:1;
 
-
+	/** Default constructor. */
 	FDebugDisplayProperty()
-		: Obj(NULL)
-		, WithinClass(NULL)
+		: Obj(nullptr)
+		, WithinClass(nullptr)
 		, bSpecialProperty(false)
-	{
-	}
-
+	{ }
 };
diff --git a/Engine/Source/Runtime/Engine/Classes/Engine/DecalActor.h b/Engine/Source/Runtime/Engine/Classes/Engine/DecalActor.h
--- a/Engine/Source/Runtime/Engine/Classes/Engine/DecalActor.h
+++ b/Engine/Source/Runtime/Engine/Classes/Engine/DecalActor.h
@@ -1,20 +1,19 @@
 // Copyright 1998-2015 Epic Games, Inc. All Rights Reserved.
 
-/**
- * Decal to add texture/material on existing geometry.
- *
- */
-
 #pragma once
+
 #include "DecalActor.generated.h"
 
-//
-// Forward declarations.
-//
+
 class UBoxComponent;
 
+
+/**
+* Decal to add texture/material on existing geometry.
+*/
 UCLASS(hidecategories=(Collision, Attachment, Actor, Input, Replication), showcategories=("Input|MouseInput", "Input|TouchInput"), MinimalAPI)
-class ADecalActor : public AActor
+class ADecalActor
+	: public AActor
 {
 	GENERATED_UCLASS_BODY()
 
@@ -65,6 +64,7 @@ public:
 #endif // WITH_EDITOR
 
 public:
+
 	/** Returns Decal subobject **/
 	ENGINE_API class UDecalComponent* GetDecal() const;
 #if WITH_EDITORONLY_DATA
@@ -76,6 +76,3 @@ public:
 	ENGINE_API UBoxComponent* GetBoxComponent() const;
 #endif
 };
-
-
-
diff --git a/Engine/Source/Runtime/Engine/Classes/Engine/DemoNetConnection.h b/Engine/Source/Runtime/Engine/Classes/Engine/DemoNetConnection.h
--- a/Engine/Source/Runtime/Engine/Classes/Engine/DemoNetConnection.h
+++ b/Engine/Source/Runtime/Engine/Classes/Engine/DemoNetConnection.h
@@ -1,40 +1,51 @@
 // Copyright 1998-2015 Epic Games, Inc. All Rights Reserved.
 
-//
-// Simulated network connection for recording and playing back game sessions.
-//
-
 #pragma once
+
 #include "DemoNetConnection.generated.h"
 
+
 struct FQueuedDemoPacket
 {
 	TArray< uint8 > Data;
 };
 
+
+/**
+ * Simulated network connection for recording and playing back game sessions.
+ */
 UCLASS(transient, config=Engine)
-class UDemoNetConnection : public UNetConnection
+class UDemoNetConnection
+	: public UNetConnection
 {
 	GENERATED_UCLASS_BODY()
 
-	// Begin UNetConnection interface.
-	virtual void		InitConnection( class UNetDriver* InDriver, EConnectionState InState, const FURL& InURL, int32 InConnectionSpeed = 0 ) override;
-	virtual FString		LowLevelGetRemoteAddress( bool bAppendPort=false ) override;
-	virtual FString		LowLevelDescribe() override;
-	virtual void		LowLevelSend( void* Data, int32 Count ) override;
-	virtual int32		IsNetReady( bool Saturate ) override;
-	virtual void		FlushNet( bool bIgnoreSimulation = false ) override;
-	virtual void		HandleClientPlayer( APlayerController* PC, class UNetConnection* NetConnection ) override;
-	virtual bool		ClientHasInitializedLevelFor( const UObject* TestObject ) const override;
-	// End UNetConnection interface.
+public:
+
+	// UNetConnection interface.
+
+	virtual void InitConnection( class UNetDriver* InDriver, EConnectionState InState, const FURL& InURL, int32 InConnectionSpeed = 0 ) override;
+	virtual FString LowLevelGetRemoteAddress( bool bAppendPort = false ) override;
+	virtual FString LowLevelDescribe() override;
+	virtual void LowLevelSend( void* Data, int32 Count ) override;
+	virtual int32 IsNetReady( bool Saturate ) override;
+	virtual void FlushNet( bool bIgnoreSimulation = false ) override;
+	virtual void HandleClientPlayer( APlayerController* PC, class UNetConnection* NetConnection ) override;
+	virtual bool ClientHasInitializedLevelFor( const UObject* TestObject ) const override;
+
+public:
 
 	/** @return The DemoRecording driver object */
-	FORCEINLINE class UDemoNetDriver* GetDriver() { return (UDemoNetDriver*)Driver; }
+	FORCEINLINE class UDemoNetDriver* GetDriver()
+	{
+		return (UDemoNetDriver*)Driver;
+	}
 
 	/** @return The DemoRecording driver object */
-	FORCEINLINE const class UDemoNetDriver* GetDriver() const { return (UDemoNetDriver*)Driver; }
+	FORCEINLINE const class UDemoNetDriver* GetDriver() const
+	{
+		return (UDemoNetDriver*)Driver;
+	}
 
-	TArray< FQueuedDemoPacket > QueuedDemoPackets;
+	TArray<FQueuedDemoPacket> QueuedDemoPackets;
 };
-
-
diff --git a/Engine/Source/Runtime/Engine/Classes/Engine/DemoNetDriver.h b/Engine/Source/Runtime/Engine/Classes/Engine/DemoNetDriver.h
--- a/Engine/Source/Runtime/Engine/Classes/Engine/DemoNetDriver.h
+++ b/Engine/Source/Runtime/Engine/Classes/Engine/DemoNetDriver.h
@@ -1,57 +1,62 @@
 // Copyright 1998-2015 Epic Games, Inc. All Rights Reserved.
 
-//
-// Simulated network driver for recording and playing back game sessions.
 #pragma once
+
 #include "DemoNetDriver.generated.h"
 
+
+/**
+ * Simulated network driver for recording and playing back game sessions.
+ */
 UCLASS(transient, config=Engine)
-class UDemoNetDriver : public UNetDriver
+class UDemoNetDriver
+	: public UNetDriver
 {
 	GENERATED_UCLASS_BODY()
 
-	// Variables.
-
 	/** Name of the file to read/write from */
-	FString				DemoFilename;
+	FString DemoFilename;
 
 	/** Handle to the archive that will read/write network packets */
-	FArchive*			FileAr;
+	FArchive* FileAr;
 
 	/** Current record/playback frame number */
-	int32				DemoFrameNum;
+	int32 DemoFrameNum;
 
 	/** Last time (in real seconds) that we recorded a frame */
-	double				LastRecordTime;
+	double LastRecordTime;
 
 	/** Time (in game seconds) that have elapsed between recorded frames */
-	float				DemoDeltaTime;
+	float DemoDeltaTime;
 
 	/** Total time of demo in seconds */
-	float				DemoTotalTime;
+	float DemoTotalTime;
 
 	/** Current record/playback position in seconds */
-	float				DemoCurrentTime;
+	float DemoCurrentTime;
 
 	/** Total number of frames in the demo */
-	int32				DemoTotalFrames;
+	int32 DemoTotalFrames;
 
 	/** during playback, set to offset of where the stream ends (we don't want to continue reading into the meta data section) */
-	int32				EndOfStreamOffset;
+	int32 EndOfStreamOffset;
 
 	/** True if we're in the middle of recording a frame */
-	bool				bIsRecordingDemoFrame;
+	bool bIsRecordingDemoFrame;
 
 	/** True if we are at the end of playing a demo */
-	bool				bDemoPlaybackDone;
+	bool bDemoPlaybackDone;
 
 	/** This is our spectator controller that is used to view the demo world from */
-	APlayerController*	SpectatorController;
+	APlayerController* SpectatorController;
+
+	UPROPERTY(config)
+	FString DemoSpectatorClass;
 
-	UPROPERTY( config )
-	FString				DemoSpectatorClass;
+public:
+
+	// UNetDriver interface.
 
-	// Begin UNetDriver interface.
 	virtual bool InitBase( bool bInitAsClient, FNetworkNotify* InNotify, const FURL& URL, bool bReuseAddressAndPort, FString& Error ) override;
 	virtual void FinishDestroy() override;
 	virtual FString LowLevelGetNetworkNumber() override;
@@ -59,13 +64,16 @@ class UDemoNetDriver : public UNetDriver
 	virtual bool InitListen( FNetworkNotify* InNotify, FURL& ListenURL, bool bReuseAddressAndPort, FString& Error ) override;
 	virtual void TickDispatch( float DeltaSeconds ) override;
 	virtual void TickFlush( float DeltaSeconds ) override;
-	virtual void ProcessRemoteFunction( class AActor* Actor, class UFunction* Function, void* Parameters, struct FOutParmRec* OutParms, struct FFrame* Stack, class UObject* SubObject = NULL );
+	virtual void ProcessRemoteFunction( class AActor* Actor, class UFunction* Function, void* Parameters, struct FOutParmRec* OutParms, struct FFrame* Stack, class UObject* SubObject = nullptr );
 	virtual bool IsAvailable() const override { return true; }
-	// End UNetDriver interface.
+	
+public:
+
+	// FExec interface
 
-	// Begin FExec interface.
 	virtual bool Exec( UWorld* InWorld, const TCHAR* Cmd, FOutputDevice& Ar ) override;
-	// End FExec interface.
+
+public:
 
 	/** @todo document */
 	bool UpdateDemoTime( float* DeltaTime, float TimeDilation );
diff --git a/Engine/Source/Runtime/Engine/Classes/Engine/DestructibleMesh.h b/Engine/Source/Runtime/Engine/Classes/Engine/DestructibleMesh.h
--- a/Engine/Source/Runtime/Engine/Classes/Engine/DestructibleMesh.h
+++ b/Engine/Source/Runtime/Engine/Classes/Engine/DestructibleMesh.h
@@ -2,11 +2,10 @@
 
 #pragma once
 
-/** This class holds an APEX destructible asset as well as an associated USkeletalMesh */
-
 #include "Engine/SkeletalMesh.h"
 #include "DestructibleMesh.generated.h"
 
+
 #if WITH_APEX
 namespace NxParameterized
 {
@@ -22,6 +21,7 @@ namespace physx
 }
 #endif
 
+
 /**
 	Chunks up to the depth DefaultImpactDamageDepth will take impact damage, unless IDO_On or IDO_Off is chosen.
 */
@@ -34,6 +34,7 @@ enum EImpactDamageOverride
 	IDO_MAX,
 };
 
+
 /** Properties that may be set for all chunks at a particular depth in the fracture hierarchy. */
 USTRUCT()
 struct FDestructibleDepthParameters
@@ -56,6 +57,7 @@ struct FDestructibleDepthParameters
 
 };
 
+
 /** Flags that apply to a destructible actor. */
 USTRUCT()
 struct FDestructibleParametersFlag
@@ -152,8 +154,7 @@ struct FDestructibleParametersFlag
 		, bAccurateRaycasts(false)
 		, bUseValidBounds(false)
 		, bFormExtendedStructures(false)
-	{
-	}
+	{ }
 
 #if WITH_APEX
 	void FillDestructibleActorDesc(NxParameterized::Interface* Params) const;
@@ -162,6 +163,7 @@ struct FDestructibleParametersFlag
 
 };
 
+
 /** Parameters that pertain to chunk damage. */
 USTRUCT()
 struct FDestructibleDamageParameters
@@ -215,16 +217,14 @@ struct FDestructibleDamageParameters
 
 
 	FDestructibleDamageParameters()
-	: DamageThreshold(1.0f)
-	, DamageSpread(0.1f)
-	, bEnableImpactDamage(false)
-	, ImpactDamage(0.0f)
-	, DefaultImpactDamageDepth(-1)
-	, bCustomImpactResistance(false)
-	, ImpactResistance(0.0f)
-	
-	{
-	}
+		: DamageThreshold(1.0f)
+		, DamageSpread(0.1f)
+		, bEnableImpactDamage(false)
+		, ImpactDamage(0.0f)
+		, DefaultImpactDamageDepth(-1)
+		, bCustomImpactResistance(false)
+		, ImpactResistance(0.0f)
+	{ }
 
 #if WITH_APEX
 	void FillDestructibleActorDesc(NxParameterized::Interface* Params, UPhysicalMaterial* PhysMat) const;
@@ -233,6 +233,7 @@ struct FDestructibleDamageParameters
 	
 };
 
+
 /** Parameters that pertain to chunk debris-level settings. */
 USTRUCT()
 struct FDestructibleDebrisParameters
@@ -293,6 +294,7 @@ struct FDestructibleDebrisParameters
 	
 };
 
+
 /** Parameters that are less-often used. */
 USTRUCT()
 struct FDestructibleAdvancedParameters
@@ -345,6 +347,7 @@ struct FDestructibleAdvancedParameters
 
 };
 
+
 /** Special hierarchy depths for various behaviors. */
 USTRUCT()
 struct FDestructibleSpecialHierarchyDepths
@@ -388,16 +391,16 @@ struct FDestructibleSpecialHierarchyDepths
 	UPROPERTY(EditAnywhere, Category=DestructibleSpecialHierarchyDepths, meta=(DisplayName = "Essential LOD Depth"))
 	int32 EssentialDepth;
 
-
-
+	/**
+	 * Default constructor.
+	 */
 	FDestructibleSpecialHierarchyDepths()
-	: SupportDepth(0)
-	, MinimumFractureDepth(0)
-	, bEnableDebris(false)
-	, DebrisDepth(-1)
-	, EssentialDepth(0)
-	{
-	}
+		: SupportDepth(0)
+		, MinimumFractureDepth(0)
+		, bEnableDebris(false)
+		, DebrisDepth(-1)
+		, EssentialDepth(0)
+	{ }
 	
 #if WITH_APEX
 	void FillDestructibleActorDesc(NxParameterized::Interface* Params) const;
@@ -405,6 +408,7 @@ struct FDestructibleSpecialHierarchyDepths
 #endif
 };
 
+
 /** Parameters that apply to a destructible actor. */
 USTRUCT()
 struct FDestructibleParameters
@@ -440,8 +444,13 @@ struct FDestructibleParameters
 	struct FDestructibleParametersFlag Flags;
 };
 
+
+/**
+ * Holds an APEX destructible asset as well as an associated USkeletalMesh.
+ */
 UCLASS(hidecategories=(Object, Mesh, LevelOfDetail, Mirroring, Physics, Reimport, Clothing), MinimalAPI)
-class UDestructibleMesh : public USkeletalMesh
+class UDestructibleMesh
+	: public USkeletalMesh
 {
 	GENERATED_UCLASS_BODY()
 
@@ -458,7 +467,7 @@ class UDestructibleMesh : public USkeletalMesh
 	UPROPERTY(instanced)
 	class UDestructibleFractureSettings* FractureSettings;
 	
-	/** Static mesh this destructible mesh is created from. Is NULL if not created from a static mesh */
+	/** Static mesh this destructible mesh is created from. Is nullptr if not created from a static mesh */
 	UPROPERTY()
 	UStaticMesh* SourceStaticMesh;
 	
@@ -478,15 +487,19 @@ public:
 	physx::apex::NxDestructibleAsset* ApexDestructibleAsset;
 #endif // WITH_APEX
 
-	// Begin UObject interface.
-	virtual void				PostLoad() override;
-	virtual void				Serialize(FArchive& Ar) override;
-	virtual void				FinishDestroy() override;
+public:
+
+	// UObject interface.
+
+	virtual void PostLoad() override;
+	virtual void Serialize(FArchive& Ar) override;
+	virtual void FinishDestroy() override;
 #if WITH_EDITOR
-	virtual void				PreEditChange(UProperty* PropertyAboutToChange) override;
+	virtual void PreEditChange(UProperty* PropertyAboutToChange) override;
 #endif
-	// End UObject interface.
 
+public:
+		
 #if WITH_APEX
 	/** 
 	 * Retrieve a default actor descriptor.
@@ -506,40 +519,35 @@ public:
 #endif // WITH_APEX
 
 	/** Fills DefaultDestructibleParameters with parameters from the NxDestructibleAsset. */
-	ENGINE_API void				LoadDefaultDestructibleParametersFromApexAsset();
+	ENGINE_API void LoadDefaultDestructibleParametersFromApexAsset();
 
 	/**
-	 * Create DestructibleFractureSettings for this DestructibleMesh if it doesn't have one
+	 * Create DestructibleFractureSettings for this DestructibleMesh if it doesn't have one.
 	 */
-	ENGINE_API void				CreateFractureSettings();
+	ENGINE_API void CreateFractureSettings();
 
 	/**
 	 * Imports FractureSettings data from a StaticMesh
 	 *
-	 * @param StaticMesh - the StaticMesh to import
-	 *
-	 * @return True if successful
+	 * @param StaticMesh The StaticMesh to import
+	 * @return true on success, false otherwise.
 	 */
-	ENGINE_API bool				BuildFractureSettingsFromStaticMesh(UStaticMesh* StaticMesh);
+	ENGINE_API bool BuildFractureSettingsFromStaticMesh(UStaticMesh* StaticMesh);
 
 	/** 
 	 * Initializes this DestructibleMesh from a StaticMesh.
 	 *
-	 * @param StaticMesh - the StaticMesh to import
-	 *
-	 * @return True if successful.
-	 **/
-	ENGINE_API bool				BuildFromStaticMesh(UStaticMesh& StaticMesh);
+	 * @param StaticMesh The StaticMesh to import
+	 * @return true on success, false otherwise.
+	 */
+	ENGINE_API bool BuildFromStaticMesh(UStaticMesh& StaticMesh);
 
 	/** 
 	 * Initialized this DestructibleMesh from the StaticMesh it was created from and the passed
 	 * in ChunkMeshes to build the level 1 chunks from.
 	 *
-	 * @param ChunkMeshes		Meshes to build the level1 chunks from
-	 * @return					True if successful 
-	 **/
-	ENGINE_API bool				SetupChunksFromStaticMeshes(const TArray<UStaticMesh*>& ChunkMeshes);
+	 * @param ChunkMeshes Meshes to build the level1 chunks from
+	 * @return true on success, false otherwise.
+	 */
+	ENGINE_API bool SetupChunksFromStaticMeshes(const TArray<UStaticMesh*>& ChunkMeshes);
 };
-
-
-
diff --git a/Engine/Source/Runtime/Engine/Classes/Engine/DirectionalLight.h b/Engine/Source/Runtime/Engine/Classes/Engine/DirectionalLight.h
--- a/Engine/Source/Runtime/Engine/Classes/Engine/DirectionalLight.h
+++ b/Engine/Source/Runtime/Engine/Classes/Engine/DirectionalLight.h
@@ -1,13 +1,20 @@
 // Copyright 1998-2015 Epic Games, Inc. All Rights Reserved.
 
 #pragma once
+
 #include "Engine/Light.h"
 #include "DirectionalLight.generated.h"
 
+
 class UArrowComponent;
 
+
+/**
+ * Implements a directional light actor.
+ */
 UCLASS(ClassGroup=(Lights, DirectionalLights), MinimalAPI, meta=(ChildCanTick))
-class ADirectionalLight : public ALight
+class ADirectionalLight
+	: public ALight
 {
 	GENERATED_UCLASS_BODY()
 
@@ -21,20 +28,17 @@ private_subobject:
 
 public:
 
-	// Begin UObject Interface
+	// UObject Interface
 	virtual void PostLoad() override;
 #if WITH_EDITOR
 	virtual void LoadedFromAnotherClass(const FName& OldClassName) override;
 	virtual void PostEditChangeProperty(FPropertyChangedEvent& PropertyChangedEvent) override;
 #endif
-	// End UObject Interface
 
 public:
+
 #if WITH_EDITORONLY_DATA
 	/** Returns ArrowComponent subobject **/
 	ENGINE_API UArrowComponent* GetArrowComponent() const;
 #endif
 };
-
-
-
diff --git a/Engine/Source/Runtime/Engine/Classes/Engine/DocumentationActor.h b/Engine/Source/Runtime/Engine/Classes/Engine/DocumentationActor.h
--- a/Engine/Source/Runtime/Engine/Classes/Engine/DocumentationActor.h
+++ b/Engine/Source/Runtime/Engine/Classes/Engine/DocumentationActor.h
@@ -1,11 +1,10 @@
 // Copyright 1998-2015 Epic Games, Inc. All Rights Reserved.
 
-//=============================================================================
-//=============================================================================
-
 #pragma once
+
 #include "DocumentationActor.generated.h"
 
+
 namespace EDocumentationActorType
 {
 	enum Type
@@ -21,8 +20,10 @@ namespace EDocumentationActorType
 	};
 }
 
+
 UCLASS(hidecategories = (Sprite, MaterialSprite, Actor, Transform, Tags, Materials, Rendering, Components, Blueprint, bject, Collision, Display, Rendering, Physics, Input, Lighting, Layers))
-class ENGINE_API ADocumentationActor : public AActor
+class ENGINE_API ADocumentationActor
+	: public AActor
 {
 	GENERATED_UCLASS_BODY()
 
@@ -50,6 +51,7 @@ private_subobject:
 	UPROPERTY(Category = Sprite, VisibleAnywhere, BlueprintReadOnly, meta = (AllowPrivateAccess = "true"))
 	class UMaterialBillboardComponent* Billboard;
 #endif
+
 public:
 
 	/** Determine the type of link contained in the DocumentLink member */
@@ -62,13 +64,10 @@ private:
 
 	EDocumentationActorType::Type	LinkType;
 
-
 public:
+
 #if WITH_EDITORONLY_DATA
 	/** Returns Billboard subobject **/
 	class UMaterialBillboardComponent* GetBillboard() const;
 #endif
 };
-
-
-
diff --git a/Engine/Source/Runtime/Engine/Classes/Engine/Engine.h b/Engine/Source/Runtime/Engine/Classes/Engine/Engine.h
--- a/Engine/Source/Runtime/Engine/Classes/Engine/Engine.h
+++ b/Engine/Source/Runtime/Engine/Classes/Engine/Engine.h
@@ -1,13 +1,11 @@
 // Copyright 1998-2015 Epic Games, Inc. All Rights Reserved.
 
-//=============================================================================
-// Engine: The base class of the global application object classes.
-//=============================================================================
-
 #pragma once
+
 #include "World.h"
 #include "Engine.generated.h"
 
+
 class FScreenSaverInhibitor;
 class UDeviceProfileManager;
 class FViewport;
@@ -15,6 +13,9 @@ class FCommonViewportClient;
 class FCanvas;
 
 
+/**
+ * Enumerats types of fully loaded packages.
+ */
 UENUM()
 enum EFullyLoadPackageType
 {
@@ -31,8 +32,9 @@ enum EFullyLoadPackageType
 	FULLYLOAD_MAX,
 };
 
+
 /**
- * A transition type.
+ * Enumerates transition types.
  */
 UENUM()
 enum ETransitionType
@@ -47,6 +49,7 @@ enum ETransitionType
 	TT_MAX,
 };
 
+
 UENUM()
 enum EConsoleType
 {
@@ -55,6 +58,7 @@ enum EConsoleType
 	CONSOLE_MAX,
 };
 
+
 /** Struct to help hold information about packages needing to be fully-loaded for DLC, etc */
 USTRUCT()
 struct FFullyLoadedPackagesInfo
@@ -85,6 +89,7 @@ struct FFullyLoadedPackagesInfo
 
 };
 
+
 /** level streaming updates that should be applied immediately after committing the map change */
 USTRUCT()
 struct FLevelStreamingStatus
@@ -116,6 +121,7 @@ struct FLevelStreamingStatus
 	
 };
 
+
 /**
  * Container for describing various types of netdrivers available to the engine
  * The engine will try to construct a netdriver of a given type and, failing that,
@@ -146,6 +152,7 @@ struct FNetDriverDefinition
 	}
 };
 
+
 /**
  * Active and named net drivers instantiated from an FNetDriverDefinition
  * The net driver will remain instantiated on this struct until it is destroyed
@@ -175,6 +182,7 @@ struct FNamedNetDriver
 	~FNamedNetDriver() {}
 };
 
+
 /** FWorldContext
  *	A context for dealing with UWorlds at the engine level. As the engine brings up and destroys world, we need a way to keep straight
  *	what world belongs to what.
@@ -329,8 +337,7 @@ struct FWorldContext
 		, RunAsDedicated(false)
 		, bWaitingOnOnlineSubsystem(false)
 		, ThisCurrentWorld(nullptr)
-	{
-	}
+	{ }
 
 private:
 
@@ -358,6 +365,7 @@ struct FStatColorMapEntry
 
 };
 
+
 USTRUCT()
 struct FStatColorMapping
 {
@@ -380,6 +388,7 @@ struct FStatColorMapping
 
 };
 
+
 /** Info about one note dropped in the map during PIE. */
 USTRUCT()
 struct FDropNoteInfo
@@ -407,7 +416,6 @@ struct FDropNoteInfo
 
 };
 
-/************************************/
 
 /** On-screen debug message handling */
 /** Helper struct for tracking on screen messages. */
@@ -448,6 +456,7 @@ struct FScreenMessageString
 	
 };
 
+
 UENUM()
 namespace EMatineeCaptureType
 {
@@ -460,6 +469,7 @@ namespace EMatineeCaptureType
 	};
 }
 
+
 USTRUCT()
 struct FGameNameRedirect
 {
@@ -497,6 +507,7 @@ struct FClassRedirect
 	bool InstanceOnly;
 };
 
+
 USTRUCT()
 struct FStructRedirect
 {
@@ -509,6 +520,7 @@ struct FStructRedirect
 	FString NewStructName;
 };
 
+
 USTRUCT()
 struct FPluginRedirect
 {
@@ -521,11 +533,13 @@ struct FPluginRedirect
 	FString NewPluginName;
 };
 
+
 class IAnalyticsProvider;
 
 DECLARE_DELEGATE_OneParam(FBeginStreamingPauseDelegate, FViewport*);
 DECLARE_DELEGATE(FEndStreamingPauseDelegate);
 
+
 /**
  * Abstract base class of all Engine classes, responsible for management of systems critical to editor or game systems.
  * Also defines default classes for certain engine systems.
@@ -2005,7 +2019,7 @@ public:
 	/** @return whether we're currently running in split screen (more than one local player) */
 	bool IsSplitScreen(UWorld *InWorld);
 
-	/** @return whether we're currently running with stereoscopic 3D enabled for the specified viewport (or globally, if viewport is NULL) */
+	/** @return whether we're currently running with stereoscopic 3D enabled for the specified viewport (or globally, if viewport is nullptr) */
 	bool IsStereoscopic3D(FViewport* InViewport = nullptr);
 
 	/**
@@ -2561,87 +2575,87 @@ private:
 
 public:
 	/**
-	 * Delegate we fire every time a new stat has been registered
+	 * Delegate we fire every time a new stat has been registered.
 	 *
-	 * @param FName	- The name of the new stat
-	 * @param FName - The category of the new stat
-	 * @param FText - The description of the new stat
+	 * @param FName The name of the new stat.
+	 * @param FName The category of the new stat.
+	 * @param FText The description of the new stat.
 	 */
 	DECLARE_EVENT_ThreeParams(UEngine, FOnNewStatRegistered, const FName&, const FName&, const FText&);
 	static FOnNewStatRegistered NewStatDelegate;
 	
 	/**
-	 * Wrapper for firing a simple stat exec
+	 * Wrapper for firing a simple stat exec.
 	 *
-	 * @param World	- The world to apply the exec to
-	 * @param ViewportClient - The viewport to apply the exec to
-	 * @param InName - The exec string
+	 * @param World	The world to apply the exec to.
+	 * @param ViewportClient The viewport to apply the exec to.
+	 * @param InName The exec string.
 	 */
 	void ExecEngineStat(UWorld* World, FCommonViewportClient* ViewportClient, const TCHAR* InName);
 
 	/**
-	 * Check to see if the specified stat name is a simple stat
+	 * Check to see if the specified stat name is a simple stat.
 	 *
-	 * @param InName - The name of the stat we're checking
-	 * @returns true if the stat is a registered simple stat
+	 * @param InName The name of the stat we're checking.
+	 * @returns true if the stat is a registered simple stat.
 	 */
 	bool IsEngineStat(const FString& InName);
 
 	/**
-	 * Set the state of the specified stat
+	 * Set the state of the specified stat.
 	 *
-	 * @param World	- The world to apply the exec to
-	 * @param ViewportClient - The viewport to apply the exec to
-	 * @param InName - The stat name
-	 * @param bShow - The state we would like the stat to be in
+	 * @param World	The world to apply the exec to.
+	 * @param ViewportClient The viewport to apply the exec to.
+	 * @param InName The stat name.
+	 * @param bShow The state we would like the stat to be in.
 	 */
 	void SetEngineStat(UWorld* World, FCommonViewportClient* ViewportClient, const FString& InName, const bool bShow);
 
 	/**
-	 * Set the state of the specified stats (note: array processed in reverse order when !bShow)
+	 * Set the state of the specified stats (note: array processed in reverse order when !bShow).
 	 *
-	 * @param World	- The world to apply the exec to
-	 * @param ViewportClient - The viewport to apply the exec to
-	 * @param InNames - The stat names
-	 * @param bShow - The state we would like the stat to be in
+	 * @param World	The world to apply the exec to.
+	 * @param ViewportClient The viewport to apply the exec to.
+	 * @param InNames The stat names.
+	 * @param bShow The state we would like the stat to be in.
 	 */
 	void SetEngineStats(UWorld* World, FCommonViewportClient* ViewportClient, const TArray<FString>& InNames, const bool bShow);
 
 	/**
 	 * Function to render all the simple stats
 	 *
-	 * @param World	- The world being drawn to
-	 * @param ViewportClient - The viewport being drawn to
-	 * @param Canvas - The canvas to use when drawing
-	 * @param LHSX - The left hand side X position to start drawing from
-	 * @param InOutLHSY - The left hand side Y position to start drawing from
-	 * @param RHSX - The right hand side X position to start drawing from
-	 * @param InOutRHSY - The right hand side Y position to start drawing from
-	 * @param ViewLocation - The world space view location
-	 * @param ViewRotation - The world space view rotation
+	 * @param World	The world being drawn to.
+	 * @param ViewportClient The viewport being drawn to.
+	 * @param Canvas The canvas to use when drawing.
+	 * @param LHSX The left hand side X position to start drawing from.
+	 * @param InOutLHSY The left hand side Y position to start drawing from.
+	 * @param RHSX The right hand side X position to start drawing from.
+	 * @param InOutRHSY The right hand side Y position to start drawing from.
+	 * @param ViewLocation The world space view location.
+	 * @param ViewRotation The world space view rotation.
 	 */
 	void RenderEngineStats(UWorld* World, FViewport* Viewport, FCanvas* Canvas, int32 LHSX, int32& InOutLHSY, int32 RHSX, int32& InOutRHSY, const FVector* ViewLocation, const FRotator* ViewRotation);
 
 private:
 	/**
-	 * Function definition for those stats which have their own render funcsions (or affect another render functions)
+	 * Function definition for those stats which have their own render funcsions (or affect another render functions).
 	 *
-	 * @param World	- The world being drawn to
-	 * @param ViewportClient - The viewport being drawn to
-	 * @param Canvas - The canvas to use when drawing
-	 * @param X - The X position to draw to
-	 * @param Y - The Y position to draw to
-	 * @param ViewLocation - The world space view location
-	 * @param ViewRotation - The world space view rotation
+	 * @param World	The world being drawn to.
+	 * @param ViewportClient The viewport being drawn to.
+	 * @param Canvas The canvas to use when drawing.
+	 * @param X The X position to draw to.
+	 * @param Y The Y position to draw to.
+	 * @param ViewLocation The world space view location.
+	 * @param ViewRotation The world space view rotation.
 	 */
 	typedef int32 (UEngine::*EngineStatRender)(UWorld* World, FViewport* Viewport, FCanvas* Canvas, int32 X, int32 Y, const FVector* ViewLocation, const FRotator* ViewRotation);
 
 	/**
-	 * Function definition for those stats which have their own toggle funcsions (or toggle other stats)
+	 * Function definition for those stats which have their own toggle funcsions (or toggle other stats).
 	 *
-	 * @param World	- The world being drawn to
-	 * @param ViewportClient - The viewport being drawn to
-	 * @param Stream - The remaining characters from the Exec call
+	 * @param World	The world being drawn to.
+	 * @param ViewportClient The viewport being drawn to.
+	 * @param Stream The remaining characters from the Exec call.
 	 */
 	typedef bool (UEngine::*EngineStatToggle)(UWorld* World, FCommonViewportClient* ViewportClient, const TCHAR* Stream);
 
@@ -2684,12 +2698,13 @@ private:
 	TArray<FEngineStatFuncs> EngineStats;
 
 private:
+
 	/**
-	 * Functions for performing other actions when the stat is toggled, should only be used when registering with EngineStats
+	 * Functions for performing other actions when the stat is toggled, should only be used when registering with EngineStats.
 	 *
-	 * @param World	- The world being drawn to
-	 * @param ViewportClient - The viewport being drawn to
-	 * @param Stream - The remaining characters from the Exec call (optional)
+	 * @param World	The world being drawn to.
+	 * @param ViewportClient The viewport being drawn to.
+	 * @param Stream The remaining characters from the Exec call (optional).
 	 */
 	bool ToggleStatFPS(UWorld* World, FCommonViewportClient* ViewportClient, const TCHAR* Stream = nullptr);
 	bool ToggleStatDetailed(UWorld* World, FCommonViewportClient* ViewportClient, const TCHAR* Stream = nullptr);
@@ -2705,15 +2720,15 @@ private:
 	bool ToggleStatSounds(UWorld* World, FCommonViewportClient* ViewportClient, const TCHAR* Stream = nullptr);
 
 	/**
-	 * Functions for rendering the various simple stats, should only be used when registering with EngineStats
+	 * Functions for rendering the various simple stats, should only be used when registering with EngineStats.
 	 *
-	 * @param World	- The world being drawn to
-	 * @param ViewportClient - The viewport being drawn to
-	 * @param Canvas - The canvas to use when drawing
-	 * @param X - The X position to draw to
-	 * @param Y - The Y position to draw to
-	 * @param ViewLocation - The world space view location
-	 * @param ViewRotation - The world space view rotation
+	 * @param World	The world being drawn to.
+	 * @param ViewportClient The viewport being drawn to.
+	 * @param Canvas The canvas to use when drawing.
+	 * @param X The X position to draw to.
+	 * @param Y The Y position to draw to.
+	 * @param ViewLocation The world space view location.
+	 * @param ViewRotation The world space view rotation.
 	 */
 #if !(UE_BUILD_SHIPPING || UE_BUILD_TEST)
 	int32 RenderStatVersion(UWorld* World, FViewport* Viewport, FCanvas* Canvas, int32 X, int32 Y, const FVector* ViewLocation = nullptr, const FRotator* ViewRotation = nullptr);
diff --git a/Engine/Source/Runtime/Engine/Classes/Engine/EngineBaseTypes.h b/Engine/Source/Runtime/Engine/Classes/Engine/EngineBaseTypes.h
--- a/Engine/Source/Runtime/Engine/Classes/Engine/EngineBaseTypes.h
+++ b/Engine/Source/Runtime/Engine/Classes/Engine/EngineBaseTypes.h
@@ -88,7 +88,7 @@ struct FTickPrerequisite
 
 		/** Noop constructor **/
 		FTickPrerequisite()
-		: PrerequisiteTickFunction(NULL)
+		: PrerequisiteTickFunction(nullptr)
 		{
 		}
 		/** 
@@ -115,7 +115,7 @@ struct FTickPrerequisite
 			{
 				return PrerequisiteTickFunction;
 			}
-			return NULL;
+			return nullptr;
 		}
 	
 };
@@ -181,7 +181,7 @@ protected:
 	TEnumAsByte<enum ETickingGroup> ActualTickGroup;
 
 private:
-	/** Completion handle for the task that will run this tick. Caution, this is no reset to NULL until an unspecified future time **/
+	/** Completion handle for the task that will run this tick. Caution, this is no reset to nullptr until an unspecified future time **/
 	FGraphEventRef CompletionHandle;
 
 	/** Prerequisites for this tick function **/
@@ -615,7 +615,7 @@ struct ENGINE_API FURL
 	// Constructors.
 	/* FURL() prevent default from being generated */
 	explicit FURL( ENoInit ) { }
-	FURL( const TCHAR* Filename=NULL );
+	FURL( const TCHAR* Filename=nullptr );
 	FURL( FURL* Base, const TCHAR* TextURL, ETravelType Type );
 	static void StaticInit();
 	static void StaticExit();
@@ -635,7 +635,7 @@ struct ENGINE_API FURL
 	void LoadURLConfig( const TCHAR* Section, const FString& Filename=GGameIni );
 	void SaveURLConfig( const TCHAR* Section, const TCHAR* Item, const FString& Filename=GGameIni ) const;
 	void AddOption( const TCHAR* Str );
-	void RemoveOption( const TCHAR* Key, const TCHAR* Section = NULL, const FString& Filename = GGameIni);
+	void RemoveOption( const TCHAR* Key, const TCHAR* Section = nullptr, const FString& Filename = GGameIni);
 	FString ToString( bool FullyQualified=0 ) const;
 	ENGINE_API friend FArchive& operator<<( FArchive& Ar, FURL& U );
 
diff --git a/Engine/Source/Runtime/Engine/Classes/Engine/EngineTypes.h b/Engine/Source/Runtime/Engine/Classes/Engine/EngineTypes.h
--- a/Engine/Source/Runtime/Engine/Classes/Engine/EngineTypes.h
+++ b/Engine/Source/Runtime/Engine/Classes/Engine/EngineTypes.h
@@ -11,13 +11,13 @@
 #include "EngineTypes.generated.h"
 
 
-
 /**
  * Default number of components to expect in TInlineAllocators used with AActor component arrays.
  * Used by engine code to try to avoid allocations in AActor::GetComponents(), among others.
  */
 enum { NumInlinedActorComponents = 24 };
 
+
 /**
  * TInlineComponentArray is simply a TArray that reserves a fixed amount of space on the stack
  * to try to avoid heap allocation when there are fewer than a specified number of elements expected in the result.
@@ -37,6 +37,7 @@ enum EAspectRatioAxisConstraint
 	AspectRatio_MAX,
 };
 
+
 /** The type of metric we want about the actor **/
 UENUM()
 enum EActorMetricsType
@@ -47,6 +48,7 @@ enum EActorMetricsType
 	METRICS_MAX,
 };
 
+
 // Return values for UEngine::Browse
 namespace EBrowseReturnVal
 {
@@ -58,6 +60,7 @@ namespace EBrowseReturnVal
 	};
 }
 
+
 UENUM()
 namespace EAttachLocation
 {
@@ -69,6 +72,7 @@ namespace EAttachLocation
 	};
 }
 
+
 /**
  * A priority for sorting scene elements by depth.
  * Elements with higher priority occlude elements with lower priority, disregarding distance.
@@ -83,6 +87,7 @@ enum ESceneDepthPriorityGroup
 	SDPG_MAX,
 };
 
+
 UENUM()
 enum EIndirectLightingCacheQuality
 {
@@ -94,6 +99,7 @@ enum EIndirectLightingCacheQuality
 	ILCQ_Volume
 };
 
+
 /** Note: This is mirrored in Lightmass, be sure to update the blend mode structure and logic there if this changes. */
 // Note: Check UMaterialInstance::Serialize if changed!!
 UENUM(BlueprintType)
@@ -107,6 +113,7 @@ enum EBlendMode
 	BLEND_MAX,
 };
 
+
 UENUM()
 enum ESamplerSourceMode
 {
@@ -118,6 +125,7 @@ enum ESamplerSourceMode
 	SSM_Clamp_WorldGroupSettings UMETA(DisplayName="Shared: Clamp")
 };
 
+
 UENUM()
 enum ETranslucencyLightingMode
 {
@@ -149,6 +157,7 @@ enum ETranslucencyLightingMode
 	TLM_MAX,
 };
 
+
 /**
  * Enumerates available options for the translucency sort policy.
  */
@@ -168,6 +177,7 @@ namespace ETranslucentSortPolicy
 	};
 }
 
+
 /** Controls the way that the width scale property affects anim trails */
 UENUM()
 enum ETrailWidthMode
@@ -177,6 +187,7 @@ enum ETrailWidthMode
 	ETrailWidthMode_FromSecond UMETA(DisplayName = "From Second Socket"),
 };
 
+
 // Note: Check UMaterialInstance::Serialize if changed!!
 UENUM()
 enum EMaterialShadingModel
@@ -191,6 +202,7 @@ enum EMaterialShadingModel
 	MSM_MAX,
 };
 
+
 // This is used by the drawing passes to determine tessellation policy, so changes here need to be supported in native code.
 UENUM()
 enum EMaterialTessellationMode
@@ -204,6 +216,7 @@ enum EMaterialTessellationMode
 	MTM_MAX,
 };
 
+
 UENUM()
 enum EMaterialSamplerType
 {
@@ -218,6 +231,7 @@ enum EMaterialSamplerType
 	SAMPLERTYPE_MAX,
 };
 
+
 /**	Lighting build quality enumeration */
 UENUM()
 enum ELightingBuildQuality
@@ -248,6 +262,7 @@ enum ETriangleSortOption
 	TRISORT_MAX,
 };
 
+
 /** Enum to specify which axis to use for the forward vector when using TRISORT_CustomLeftRight sort mode */
 UENUM()
 enum ETriangleSortAxis
@@ -258,6 +273,7 @@ enum ETriangleSortAxis
 	TSA_MAX,
 };
 
+
 /** Movement modes for Characters. */
 UENUM(BlueprintType)
 enum EMovementMode
@@ -383,6 +399,7 @@ enum EOverlapFilterOption
 	OverlapFilter_StaticOnly UMETA(DisplayName="AllStaticObjects"),
 };
 
+
 UENUM(BlueprintType)
 enum EObjectTypeQuery
 {
@@ -422,6 +439,7 @@ enum EObjectTypeQuery
 	ObjectTypeQuery_MAX	UMETA(Hidden)
 };
 
+
 UENUM(BlueprintType)
 enum ETraceTypeQuery
 {
@@ -461,6 +479,7 @@ enum ETraceTypeQuery
 	TraceTypeQuery_MAX	UMETA(Hidden)
 };
 
+
 /** Enum indicating which physics scene to use. */
 UENUM()
 enum EPhysicsSceneType
@@ -474,6 +493,7 @@ enum EPhysicsSceneType
 	PST_MAX,
 };
 
+
 /** Enum indicating how each type should respond */
 UENUM(BlueprintType)
 enum ECollisionResponse
@@ -484,6 +504,7 @@ enum ECollisionResponse
 	ECR_MAX,
 };
 
+
 UENUM()
 enum EFilterInterpolationType
 {
@@ -493,6 +514,7 @@ enum EFilterInterpolationType
 	BSIT_MAX
 };
 
+
 UENUM()
 enum EInputConsumeOptions
 {
@@ -505,6 +527,7 @@ enum EInputConsumeOptions
 	ICO_MAX
 };
 
+
 namespace EWorldType
 {
 	enum Type
@@ -518,6 +541,7 @@ namespace EWorldType
 	};
 }
 
+
 enum class EFlushLevelStreamingType
 {
 	None,			
@@ -525,6 +549,7 @@ enum class EFlushLevelStreamingType
 	Visibility,		// Flush visibility only, do not allow load requests, flushes async loading as well
 };
 
+
 USTRUCT()
 struct FResponseChannel
 {
@@ -547,6 +572,7 @@ struct FResponseChannel
 		, Response(InResponse) {}
 };
 
+
 /**
  *	Container for indicating a set of collision channels that this object will collide with.
  */
@@ -750,13 +776,16 @@ struct ENGINE_API FCollisionResponseContainer
 	static FCollisionResponseContainer CreateMinContainer(const FCollisionResponseContainer& A, const FCollisionResponseContainer& B);
 
 	static const struct FCollisionResponseContainer & GetDefaultResponseContainer() { return DefaultResponseContainer; }
+
 private:
+
 	/** static variable for default data to be used without reconstructing everytime **/
 	static FCollisionResponseContainer DefaultResponseContainer;
 
 	friend class UCollisionProfile;
 };
 
+
 /** Enum for controlling the falloff of strength of a radial impulse as a function of distance from Origin. */
 UENUM()
 enum ERadialImpulseFalloff
@@ -768,6 +797,7 @@ enum ERadialImpulseFalloff
 	RIF_MAX,
 };
 
+
 /** Presets of values used in considering when put this body to sleep. */
 UENUM()
 enum ESleepFamily
@@ -779,6 +809,7 @@ enum ESleepFamily
 	SF_MAX,
 };
 
+
 /** Enum used to indicate what type of timeline signature a function matches */
 UENUM()
 enum ETimelineSigType
@@ -791,6 +822,7 @@ enum ETimelineSigType
 	ETS_MAX,
 };
 
+
 /** Enum used to describe what type of collision is enabled on a body */
 UENUM()
 namespace ECollisionEnabled 
@@ -806,6 +838,7 @@ namespace ECollisionEnabled
 	}; 
 } 
 
+
 /** describes the physical state of a rigid body */
 USTRUCT()
 struct FRigidBodyState
@@ -833,10 +866,10 @@ struct FRigidBodyState
 		, LinVel(ForceInit)
 		, AngVel(ForceInit)
 		, Flags(0)
-	{
-	}
+	{ }
 };
 
+
 namespace ERigidBodyFlags
 {
 	enum Type
@@ -847,6 +880,7 @@ namespace ERigidBodyFlags
 	};
 }
 
+
 /** Rigid body error correction data */
 USTRUCT()
 struct FRigidBodyErrorCorrection
@@ -889,11 +923,12 @@ struct FRigidBodyErrorCorrection
 		, AngularInterpAlpha(0.1f)
 		, AngularRecipFixTime(1.0f)
 		, BodySpeedThresholdSq(0.2f)
-	{
-	}
+	{ }
 };
 
-/** Information about one contact between a pair of rigid bodies
+
+/**
+ * Information about one contact between a pair of rigid bodies.
  */
 USTRUCT()
 struct FRigidBodyContactInfo
@@ -920,19 +955,18 @@ struct FRigidBodyContactInfo
 	{
 		for (int32 ElementIndex = 0; ElementIndex < 2; ElementIndex++)
 		{
-			PhysMaterial[ElementIndex] = NULL;
+			PhysMaterial[ElementIndex] = nullptr;
 		}
 	}
 
-
 	FRigidBodyContactInfo(	const FVector& InContactPosition, 
 							const FVector& InContactNormal, 
 							float InPenetration, 
 							UPhysicalMaterial* InPhysMat0, 
 							UPhysicalMaterial* InPhysMat1 )
-	:	ContactPosition(InContactPosition)
-	,	ContactNormal(InContactNormal)
-	,	ContactPenetration(InPenetration)
+		: ContactPosition(InContactPosition)
+		, ContactNormal(InContactNormal)
+		, ContactPenetration(InPenetration)
 	{
 		PhysMaterial[0] = InPhysMat0;
 		PhysMaterial[1] = InPhysMat1;
@@ -942,7 +976,9 @@ struct FRigidBodyContactInfo
 	void SwapOrder();
 };
 
-/** Information about an overall collision, including contacts
+
+/**
+ * Information about an overall collision, including contacts.
  */
 USTRUCT()
 struct FCollisionImpactData
@@ -986,12 +1022,12 @@ struct FFractureEffect
 	class USoundBase* Sound;
 
 	FFractureEffect()
-		: ParticleSystem(NULL)
-		, Sound(NULL)
-	{
-	}
+		: ParticleSystem(nullptr)
+		, Sound(nullptr)
+	{ }
 };
 
+
 /**	Struct for handling positions relative to a base actor, which is potentially moving */
 USTRUCT()
 struct ENGINE_API FBasedPosition
@@ -1023,6 +1059,7 @@ struct ENGINE_API FBasedPosition
 	friend FArchive& operator<<( FArchive& Ar, FBasedPosition& T );
 };
 
+
 /** A line of subtitle text and the time at which it should be displayed. */
 USTRUCT()
 struct FSubtitleCue
@@ -1039,10 +1076,10 @@ struct FSubtitleCue
 
 	FSubtitleCue()
 		: Time(0)
-	{
-	}
+	{ }
 };
 
+
 /**	A subtitle localized to a specific language. */
 USTRUCT()
 struct FLocalizedSubtitle
@@ -1076,10 +1113,10 @@ struct FLocalizedSubtitle
 		: bMature(false)
 		, bManualWordWrap(false)
 		, bSingleLine(false)
-	{
-	}
+	{ }
 };
 
+
 /**	Per-light settings for Lightmass */
 USTRUCT()
 struct FLightmassLightSettings
@@ -1097,10 +1134,10 @@ struct FLightmassLightSettings
 	FLightmassLightSettings()
 		: IndirectLightingSaturation(1.0f)
 		, ShadowExponent(2.0f)
-	{
-	}
+	{ }
 };
 
+
 /**	Point/spot settings for Lightmass */
 USTRUCT()
 struct FLightmassPointLightSettings : public FLightmassLightSettings
@@ -1108,6 +1145,7 @@ struct FLightmassPointLightSettings : public FLightmassLightSettings
 	GENERATED_USTRUCT_BODY()
 };
 
+
 /**	Directional light settings for Lightmass */
 USTRUCT()
 struct FLightmassDirectionalLightSettings : public FLightmassLightSettings
@@ -1124,6 +1162,7 @@ struct FLightmassDirectionalLightSettings : public FLightmassLightSettings
 	}
 };
 
+
 /**	Per-object settings for Lightmass */
 USTRUCT()
 struct FLightmassPrimitiveSettings
@@ -1199,6 +1238,7 @@ struct FLightmassPrimitiveSettings
 	friend FArchive& operator<<(FArchive& Ar, FLightmassPrimitiveSettings& Settings);
 };
 
+
 /**	Debug options for Lightmass */
 USTRUCT()
 struct FLightmassDebugOptions
@@ -1293,6 +1333,7 @@ struct FLightmassDebugOptions
 	ENGINE_API FLightmassDebugOptions();
 };
 
+
 /**
  *	Debug options for Swarm
  */
@@ -1329,6 +1370,7 @@ struct FSwarmDebugOptions
 	void Touch();
 };
 
+
 UENUM()
 enum ELightMapPaddingType
 {
@@ -1337,6 +1379,7 @@ enum ELightMapPaddingType
 	LMPT_NoPadding
 };
 
+
 /**
  * Bit-field flags that affects storage (e.g. packing, streaming) and other info about a shadowmap.
  */
@@ -1349,6 +1392,7 @@ enum EShadowMapFlags
 	SMF_Streamed		= 0x00000001
 };
 
+
 /** reference to a specific material in a PrimitiveComponent */
 USTRUCT()
 struct FPrimitiveMaterialRef
@@ -1365,27 +1409,25 @@ struct FPrimitiveMaterialRef
 	int32 ElementIndex;
 
 	FPrimitiveMaterialRef()
-		: Primitive(NULL)
-		, Decal(NULL)
+		: Primitive(nullptr)
+		, Decal(nullptr)
 		, ElementIndex(0)
-	{
-	}
+	{ }
 
 	FPrimitiveMaterialRef(UPrimitiveComponent* InPrimitive, int32 InElementIndex)
 		: Primitive(InPrimitive)
-		, Decal(NULL)
+		, Decal(nullptr)
 		, ElementIndex(InElementIndex)
-	{
-	}
+	{ 	}
 
 	FPrimitiveMaterialRef(UDecalComponent* InDecal, int32 InElementIndex)
-		: Primitive(NULL)
+		: Primitive(nullptr)
 		, Decal(InDecal)
 		, ElementIndex(InElementIndex)
-	{
-	}
+	{ 	}
 };
 
+
 /**
  * Structure containing information about one hit of a trace, such as point of impact and surface normal at that point.
  */
@@ -1551,7 +1593,7 @@ struct ENGINE_API FHitResult
 				return &InHits[HitIdx];
 			}
 		}
-		return NULL;
+		return nullptr;
 	}
 
 	/** Static utility function that returns the number of blocking hits in array. */
@@ -1575,6 +1617,7 @@ struct ENGINE_API FHitResult
 	}
 };
 
+
 template<>
 struct TStructOpsTypeTraits<FHitResult> : public TStructOpsTypeTraitsBase
 {
@@ -1584,6 +1627,7 @@ struct TStructOpsTypeTraits<FHitResult> : public TStructOpsTypeTraitsBase
 	};
 };
 
+
 /** Structure containing information about one hit of an overlap test */
 USTRUCT()
 struct ENGINE_API FOverlapResult
@@ -1619,6 +1663,7 @@ struct ENGINE_API FOverlapResult
 	}
 };
 
+
 /** Structure containing information about minimum translation direction (MTD) */
 USTRUCT()
 struct ENGINE_API FMTDResult
@@ -1639,6 +1684,7 @@ struct ENGINE_API FMTDResult
 	}
 };
 
+
 /** Struct used for passing information from Matinee to an Actor for blending animations during a sequence. */
 USTRUCT()
 struct FAnimSlotInfo
@@ -1654,6 +1700,7 @@ struct FAnimSlotInfo
 	TArray<float> ChannelWeights;
 };
 
+
 /** Used to indicate each slot name and how many channels they have. */
 USTRUCT()
 struct FAnimSlotDesc
@@ -1671,11 +1718,11 @@ struct FAnimSlotDesc
 
 	FAnimSlotDesc()
 		: NumChannels(0)
-	{
-	}
+	{ }
 
 };
 
+
 /** Container for Animation Update Rate parameters.
  * They are shared for all components of an Actor, so they can be updated in sync. */
 USTRUCT()
@@ -1706,15 +1753,15 @@ private:
 	bool bSkipEvaluation;
 
 public:
-	/** Default constructor */
+
+	/** Default constructor. */
 	FAnimUpdateRateParameters()
 		: UpdateRate(1)
 		, EvaluationRate(1)
 		, bInterpolateSkippedFrames(false)
 		, bSkipUpdate(false)
 		, bSkipEvaluation(false)
-	{
-	}
+	{ }
 
 	/** Set parameters and verify inputs.
 	 * @param : Owner Actor owner calling this.
@@ -1755,6 +1802,7 @@ public:
 	}
 };
 
+
 /**
  * Point Of View type.
  */
@@ -1790,6 +1838,7 @@ struct FPOV
 	}
 };
 
+
 /**
  * The importance of a mesh feature when automatically generating mesh LODs.
  */
@@ -1807,6 +1856,7 @@ namespace EMeshFeatureImportance
 	};
 }
 
+
 /**
  * Settings used to reduce a mesh.
  */
@@ -1860,8 +1910,7 @@ struct FMeshReductionSettings
 		, ShadingImportance(EMeshFeatureImportance::Normal)
 		, bRecalculateNormals(false)
 		, BaseLODModel(0)
-	{
-	}
+	{ }
 
 	/** Equality operator. */
 	bool operator==(const FMeshReductionSettings& Other) const
@@ -1884,6 +1933,7 @@ struct FMeshReductionSettings
 	}
 };
 
+
 /**
  * Settings applied when building a mesh.
  */
@@ -1960,7 +2010,7 @@ struct FMeshBuildSettings
 		, BuildScale3D(1.0f, 1.0f, 1.0f)
 		, DistanceFieldResolutionScale(1.0f)
 		, bGenerateDistanceFieldAsIfTwoSided(false)
-	{}
+	{ }
 
 	/** Equality operator. */
 	bool operator==(const FMeshBuildSettings& Other) const
@@ -1986,6 +2036,7 @@ struct FMeshBuildSettings
 	}
 };
 
+
 USTRUCT()
 struct FMeshProxySettings
 {
@@ -2035,8 +2086,7 @@ struct FMeshProxySettings
 		, ClippingLevel(0.0)
 		, AxisIndex(0)
 		, bPlaneNegativeHalfspace(false)
-	{
-	}
+	{ }
 
 	/** Equality operator. */
 	bool operator==(const FMeshProxySettings& Other) const
@@ -2056,24 +2106,26 @@ struct FMeshProxySettings
 	}
 };
 
+
 USTRUCT()
 struct ENGINE_API FDamageEvent
 {
 	GENERATED_USTRUCT_BODY()
 
 public:
-	FDamageEvent() 
-	{}
+
+	/** Default constructor (no initialization). */
+	FDamageEvent() { }
 
 	FDamageEvent(FDamageEvent const& InDamageEvent)
 		: DamageTypeClass(InDamageEvent.DamageTypeClass)
-	{}
+	{ }
 	
 	explicit FDamageEvent(TSubclassOf<class UDamageType> InDamageTypeClass)
 		: DamageTypeClass(InDamageTypeClass)
-	{}
+	{ }
 
-	/** Optional DamageType for this event.  If NULL, UDamageType will be assumed. */
+	/** Optional DamageType for this event.  If nullptr, UDamageType will be assumed. */
 	UPROPERTY()
 	TSubclassOf<class UDamageType> DamageTypeClass;
 
@@ -2087,6 +2139,7 @@ public:
 	virtual void GetBestHitInfo(AActor const* HitActor, AActor const* HitInstigator, struct FHitResult& OutHitInfo, FVector& OutImpulseDir) const;
 };
 
+
 USTRUCT()
 struct ENGINE_API FPointDamageEvent : public FDamageEvent
 {
@@ -2202,6 +2255,7 @@ enum ENetRole
 	ROLE_MAX,
 };
 
+
 UENUM()
 enum ENetDormancy
 {
@@ -2218,6 +2272,7 @@ enum ENetDormancy
 	DORN_MAX,
 };
 
+
 UENUM()
 namespace EAutoReceiveInput
 {
@@ -2235,6 +2290,7 @@ namespace EAutoReceiveInput
 	};
 }
 
+
 UENUM()
 enum class EAutoPossessAI : uint8
 {
@@ -2244,6 +2300,7 @@ enum class EAutoPossessAI : uint8
 	PlacedInWorldOrSpawned, // Pawn is automatically possessed by an AI Controller whenever it is created.
 };
 
+
 UENUM(BlueprintType)
 namespace EEndPlayReason
 {
@@ -2258,6 +2315,7 @@ namespace EEndPlayReason
 
 }
 
+
 DECLARE_DYNAMIC_DELEGATE(FTimerDynamicDelegate);
 
 /** Replicated movement data of our RootComponent.
@@ -2388,7 +2446,8 @@ struct FRepMovement
 	}
 };
 
-/** Handles attachment replication to clients. Movement replication will not happen while AttachParent is non-NULL */
+
+/** Handles attachment replication to clients. Movement replication will not happen while AttachParent is non-nullptr */
 USTRUCT()
 struct FRepAttachment
 {
@@ -2413,13 +2472,13 @@ struct FRepAttachment
 	class USceneComponent* AttachComponent;
 
 	FRepAttachment()
-		: AttachParent(NULL)
+		: AttachParent(nullptr)
 		, LocationOffset(ForceInit)
 		, RelativeScale3D(ForceInit)
 		, RotationOffset(ForceInit)
 		, AttachSocket(NAME_None)
-		, AttachComponent(NULL)
-	{}
+		, AttachComponent(nullptr)
+	{ }
 };
 
 
@@ -2475,16 +2534,14 @@ struct FWalkableSlopeOverride
 	float WalkableSlopeAngle;
 
 	FWalkableSlopeOverride()
-	: WalkableSlopeBehavior(WalkableSlope_Default)
-	, WalkableSlopeAngle(0.f)
-	{
-	}
+		: WalkableSlopeBehavior(WalkableSlope_Default)
+		, WalkableSlopeAngle(0.f)
+	{ }
 
 	FWalkableSlopeOverride(EWalkableSlopeBehavior NewSlopeBehavior, float NewSlopeAngle)
-	: WalkableSlopeBehavior(NewSlopeBehavior)
-	, WalkableSlopeAngle(NewSlopeAngle)
-	{
-	}
+		: WalkableSlopeBehavior(NewSlopeBehavior)
+		, WalkableSlopeAngle(NewSlopeAngle)
+	{ }
 
 	// Given a walkable floor normal Z value, either relax or restrict the value if we override such behavior.
 	float ModifyWalkableFloorZ(float InWalkableFloorZ) const
@@ -2532,6 +2589,7 @@ struct TStructOpsTypeTraits<FRepMovement> : public TStructOpsTypeTraitsBase
 	};
 };
 
+
 /** Structure to hold and pass around transient flags used during replication. */
 struct FReplicationFlags
 {
@@ -2557,6 +2615,7 @@ struct FReplicationFlags
 	}
 };
 
+
 static_assert(sizeof(FReplicationFlags) == 4, "FReplicationFlags has invalid size.");
 
 /** Struct used to specify the property name of the component to constrain */
@@ -2570,6 +2629,7 @@ struct FConstrainComponentPropName
 	FName	ComponentName;
 };
 
+
 /** 
  *	Struct that allows for different ways to reference a component. 
  *	If just an Actor is specified, will return RootComponent of that Actor.
@@ -2594,6 +2654,7 @@ struct ENGINE_API FComponentReference
 	class USceneComponent* GetComponent(AActor* OwningActor) const;
 };
 
+
 /** Types of surfaces in the game. */
 UENUM(BlueprintType)
 enum EPhysicalSurface
@@ -2664,6 +2725,7 @@ enum EPhysicalSurface
 	SurfaceType_Max UMETA(Hidden)
 };
 
+
 /** Describes how often this component is allowed to move. */
 UENUM()
 namespace EComponentMobility
@@ -2697,6 +2759,7 @@ namespace EComponentMobility
 	};
 }
 
+
 UCLASS(abstract, config=Engine)
 class ENGINE_API UEngineTypes : public UObject
 {
@@ -2757,6 +2820,7 @@ struct FComponentSocketDescription
 	}
 };
 
+
 // ANGULAR DOF
 UENUM()
 enum EAngularConstraintMotion
@@ -2787,6 +2851,7 @@ struct FFilePath
 	FString FilePath;
 };
 
+
 /**
  * Structure for directory paths that are displayed in the UI.
  */
@@ -2802,6 +2867,7 @@ struct FDirectoryPath
 	FString Path;
 };
 
+
 /**
 * This is used for redirecting old name to new name
 * We use manually parsing array, but that makes harder to modify from property setting
@@ -2822,16 +2888,17 @@ struct ENGINE_API FRedirector
 	FRedirector()
 		: OldName(NAME_None)
 		, NewName(NAME_None)
-	{}
+	{ }
 
 	FRedirector(FName InOldName, FName InNewName)
 		: OldName(InOldName)
 		, NewName(InNewName)
-	{}
+	{ }
 };
 
+
 /** 
- * Structure for recording float values and displaying them as an Histogram through DrawDebugFloatHistory
+ * Structure for recording float values and displaying them as an Histogram through DrawDebugFloatHistory.
  */
 USTRUCT(BlueprintType)
 struct FDebugFloatHistory
@@ -2865,16 +2932,14 @@ public:
 		, MinValue(0.f)
 		, MaxValue(0.f)
 		, bAutoAdjustMinMax(true)
-	{
-	}
+	{ }
 
 	FDebugFloatHistory(float const & InMaxSamples, float const & InMinValue, float const & InMaxValue, bool const & InbAutoAdjustMinMax)
 		: MaxSamples(InMaxSamples)
 		, MinValue(InMinValue)
 		, MaxValue(InMaxValue)
 		, bAutoAdjustMinMax(InbAutoAdjustMinMax)
-	{
-	}
+	{ }
 
 	/**
 	 * Record a new Sample.
@@ -2936,6 +3001,7 @@ public:
 	}
 };
 
+
 /** info for glow when using depth field rendering */
 USTRUCT(BlueprintType)
 struct FDepthFieldGlowInfo
@@ -2968,33 +3034,32 @@ struct FDepthFieldGlowInfo
 		, GlowColor(ForceInit)
 		, GlowOuterRadius(ForceInit)
 		, GlowInnerRadius(ForceInit)
-	{
-	}
-
+	{ }
 
-		bool operator==(const FDepthFieldGlowInfo& Other) const
+	bool operator==(const FDepthFieldGlowInfo& Other) const
+	{
+		if (Other.bEnableGlow != bEnableGlow)
 		{
-			if (Other.bEnableGlow != bEnableGlow)
-			{
-				return false;
-			}
-			else if (!bEnableGlow)
-			{
-				// if the glow is disabled on both, the other values don't matter
-				return true;
-			}
-			else
-			{
-				return (Other.GlowColor == GlowColor && Other.GlowOuterRadius == GlowOuterRadius && Other.GlowInnerRadius == GlowInnerRadius);
-			}
+			return false;
 		}
-		bool operator!=(const FDepthFieldGlowInfo& Other) const
+		else if (!bEnableGlow)
 		{
-			return !(*this == Other);
+			// if the glow is disabled on both, the other values don't matter
+			return true;
 		}
-	
+		else
+		{
+			return (Other.GlowColor == GlowColor && Other.GlowOuterRadius == GlowOuterRadius && Other.GlowInnerRadius == GlowInnerRadius);
+		}
+	}
+
+	bool operator!=(const FDepthFieldGlowInfo& Other) const
+	{
+		return !(*this == Other);
+	}	
 };
 
+
 /** information used in font rendering */
 USTRUCT(BlueprintType)
 struct FFontRenderInfo
@@ -3018,6 +3083,7 @@ struct FFontRenderInfo
 	{}
 };
 
+
 /** Simple 2d triangle with UVs */
 USTRUCT()
 struct FCanvasUVTri
@@ -3070,8 +3136,8 @@ struct FCanvasUVTri
 		, V2_Pos(ForceInit)
 		, V2_UV(ForceInit)
 		, V2_Color(ForceInit)
-	{
-	}
+	{ }
 };
 
+
 template <> struct TIsZeroConstructType<FCanvasUVTri> { enum { Value = true }; };
diff --git a/Engine/Source/Runtime/Engine/Classes/Engine/ExponentialHeightFog.h b/Engine/Source/Runtime/Engine/Classes/Engine/ExponentialHeightFog.h
--- a/Engine/Source/Runtime/Engine/Classes/Engine/ExponentialHeightFog.h
+++ b/Engine/Source/Runtime/Engine/Classes/Engine/ExponentialHeightFog.h
@@ -1,11 +1,16 @@
 // Copyright 1998-2015 Epic Games, Inc. All Rights Reserved.
 
-
 #pragma once
+
 #include "ExponentialHeightFog.generated.h"
 
+
+/**
+ * Implements an Actor for exponential height fog.
+ */
 UCLASS(showcategories=(Movement, Rendering, "Utilities|Transformation"), ClassGroup=Fog, MinimalAPI)
-class AExponentialHeightFog : public AInfo
+class AExponentialHeightFog
+	: public AInfo
 {
 	GENERATED_UCLASS_BODY()
 
@@ -14,6 +19,7 @@ private_subobject:
 	DEPRECATED_FORGAME(4.6, "Component should not be accessed directly, please use GetComponent() function instead. Component will soon be private and your code will not compile.")
 	UPROPERTY(Category = ExponentialHeightFog, VisibleAnywhere, BlueprintReadOnly, meta = (AllowPrivateAccess = "true"))
 	class UExponentialHeightFogComponent* Component;
+
 public:
 
 	/** replicated copy of ExponentialHeightFogComponent's bEnabled property */
@@ -32,6 +38,3 @@ public:
 	/** Returns Component subobject **/
 	ENGINE_API class UExponentialHeightFogComponent* GetComponent() const;
 };
-
-
-
diff --git a/Engine/Source/Runtime/Engine/Classes/Engine/Font.h b/Engine/Source/Runtime/Engine/Classes/Engine/Font.h
--- a/Engine/Source/Runtime/Engine/Classes/Engine/Font.h
+++ b/Engine/Source/Runtime/Engine/Classes/Engine/Font.h
@@ -1,12 +1,16 @@
 ﻿// Copyright 1998-2015 Epic Games, Inc. All Rights Reserved.
 
 #pragma once
+
 #include "Engine/FontImportOptions.h"
 #include "FontProviderInterface.h"
 #include "SlateFontInfo.h"
 #include "Font.generated.h"
 
-/** What kind of caching is this font using? */
+
+/**
+ * Enumerates supported font caching types.
+ */
 UENUM()
 enum class EFontCacheType
 {
@@ -17,6 +21,7 @@ enum class EFontCacheType
 	Runtime,
 };
 
+
 /** this struct is serialized using native serialization so any changes to it require a package version bump */
 USTRUCT()
 struct FFontCharacter
@@ -49,8 +54,7 @@ struct FFontCharacter
 		, VSize(0)
 		, TextureIndex(0)
 		, VerticalOffset(0)
-	{
-	}
+	{ }
 
 	/**
 	 * Serialization.
@@ -72,6 +76,7 @@ struct FFontCharacter
 	}
 };
 
+
 template<>
 struct TStructOpsTypeTraits<FFontCharacter> : public TStructOpsTypeTraitsBase
 {
@@ -81,6 +86,7 @@ struct TStructOpsTypeTraits<FFontCharacter> : public TStructOpsTypeTraitsBase
 	};
 };
 
+
 /**
  * A font object, for use by Slate, UMG, and Canvas.
  *
@@ -165,8 +171,8 @@ class UFont : public UObject, public IFontProviderInterface
 	UPROPERTY()
 	FCompositeFont CompositeFont;
 
-
 public:
+
 	/** This is the character that RemapChar will return if the specified character doesn't exist in the font */
 	static const TCHAR NULLCHARACTER = 127;
 
@@ -219,7 +225,7 @@ public:
 	 *
 	 * @param	Text					the string to size
 	 *
-	 * @return	the width (in pixels) of the specified text, or 0 if Text was NULL.
+	 * @return	the width (in pixels) of the specified text, or 0 if Text was nullptr.
 	 */
 	ENGINE_API int32 GetStringSize( const TCHAR *Text ) const;
 
@@ -228,7 +234,7 @@ public:
 	 *
 	 * @param	Text					the string to size
 	 *
-	 * @return	the height (in pixels) of the specified text, or 0 if Text was NULL.
+	 * @return	the height (in pixels) of the specified text, or 0 if Text was nullptr.
 	 */
 	ENGINE_API int32 GetStringHeightSize( const TCHAR *Text ) const;
 
@@ -243,7 +249,6 @@ public:
 	 */
 	ENGINE_API void CacheCharacterCountAndMaxCharHeight();
 
-
 	/**
 	 *	Set the scaling factor
 	 *
diff --git a/Engine/Source/Runtime/Engine/Classes/Engine/FontImportOptions.h b/Engine/Source/Runtime/Engine/Classes/Engine/FontImportOptions.h
--- a/Engine/Source/Runtime/Engine/Classes/Engine/FontImportOptions.h
+++ b/Engine/Source/Runtime/Engine/Classes/Engine/FontImportOptions.h
@@ -1,9 +1,10 @@
 // Copyright 1998-2015 Epic Games, Inc. All Rights Reserved.
 
-
 #pragma once
+
 #include "FontImportOptions.generated.h"
 
+
 /** Font character set type for importing TrueType fonts */
 UENUM()
 enum EFontImportCharacterSet
@@ -14,6 +15,7 @@ enum EFontImportCharacterSet
 	FontICS_MAX,
 };
 
+
 /** Font import options */
 USTRUCT()
 struct FFontImportOptionsData
@@ -139,9 +141,8 @@ struct FFontImportOptionsData
 	UPROPERTY(EditAnywhere, Category=FontImportOptionsData, meta=(ClampMin = "0.0", ClampMax = "4.0"))
 	float DistanceFieldScanRadiusScale;
 
-
-
-		FFontImportOptionsData()
+	/** Default constructor. */
+	FFontImportOptionsData()
 		: FontName("Arial")
 		, Height(16.0f)
 		, bEnableAntialiasing(true)
@@ -167,13 +168,16 @@ struct FFontImportOptionsData
 		, bUseDistanceFieldAlpha(false)
 		, DistanceFieldScaleFactor(16)
 		, DistanceFieldScanRadiusScale(1.0f)
-		{
-		}
-	
+	{ }	
 };
 
+
+/**
+ * Holds options for importing fonts.
+ */
 UCLASS(hidecategories=Object, transient, MinimalAPI)
-class UFontImportOptions : public UObject
+class UFontImportOptions
+	: public UObject
 {
 	GENERATED_UCLASS_BODY()
 
@@ -182,4 +186,3 @@ class UFontImportOptions : public UObject
 	struct FFontImportOptionsData Data;
 
 };
-
diff --git a/Engine/Source/Runtime/Engine/Classes/Engine/GameEngine.h b/Engine/Source/Runtime/Engine/Classes/Engine/GameEngine.h
--- a/Engine/Source/Runtime/Engine/Classes/Engine/GameEngine.h
+++ b/Engine/Source/Runtime/Engine/Classes/Engine/GameEngine.h
@@ -1,20 +1,17 @@
 // Copyright 1998-2015 Epic Games, Inc. All Rights Reserved.
 
-//=============================================================================
-// GameEngine: The game subsystem.
-//=============================================================================
-
 #pragma once
 
 #include "Engine/Engine.h"
-
 #include "GameEngine.generated.h"
 
+
 /**
  * Engine that manages core systems that enable a game.
  */
 UCLASS(config=Engine, transient)
-class ENGINE_API UGameEngine : public UEngine
+class ENGINE_API UGameEngine
+	: public UEngine
 {
 	GENERATED_UCLASS_BODY()
 
@@ -30,6 +27,7 @@ class ENGINE_API UGameEngine : public UEngine
 	UGameInstance* GameInstance;
 
 public:
+
 	/** The game viewport window */
 	TWeakPtr<class SWindow> GameViewportWindow;
 	/** The primary scene viewport */
@@ -85,32 +83,38 @@ public:
 	 */
 	virtual void RedrawViewports( bool bShouldPresent = true );
 
-	// Begin UObject interface.
+public:
+
+	// UObject interface
+
 	virtual void FinishDestroy() override;
-	// End UObject interface.
 
-	// Begin UEngine interface.
+public:
+
+	// UEngine interface
+
 	virtual void Init(class IEngineLoop* InEngineLoop) override;
 	virtual void PreExit() override;
 	virtual void Tick( float DeltaSeconds, bool bIdleMode ) override;
 	virtual float GetMaxTickRate( float DeltaTime, bool bAllowFrameRateSmoothing = true ) const override;
 	virtual void ProcessToggleFreezeCommand( UWorld* InWorld ) override;
 	virtual void ProcessToggleFreezeStreamingCommand( UWorld* InWorld ) override;
-	// End UEngine interface.
 
-	// Begin FExec Interface
+public:
+
+	// FExec interface
+
 	virtual bool Exec( UWorld* InWorld, const TCHAR* Cmd, FOutputDevice& Ar=*GLog ) override;
-	// End FExec Interface
 
-	/** 
-	 * Exec command handlers
-	 */
+public:
 
+	// Exec command handlers
 	bool HandleCommand( const TCHAR* Cmd, FOutputDevice& Ar );
 	bool HandleExitCommand( const TCHAR* Cmd, FOutputDevice& Ar );
 	bool HandleMinimizeCommand( const TCHAR *Cmd, FOutputDevice &Ar );
 	bool HandleGetMaxTickRateCommand( const TCHAR* Cmd, FOutputDevice& Ar );
 	bool HandleCancelCommand( const TCHAR* Cmd, FOutputDevice& Ar, UWorld* InWorld );
+
 #if !UE_BUILD_SHIPPING
 	bool HandleApplyUserSettingsCommand( const TCHAR* Cmd, FOutputDevice& Ar );
 #endif // !UE_BUILD_SHIPPING
@@ -130,7 +134,6 @@ public:
 	/** Loads engine startup modules that the game needs to run.  This will also be called when the editor starts up. */
 	static void LoadRuntimeEngineStartupModules();
 
-
 	/**
 	 * This is a global, parameterless function used by the online subsystem modules.
 	 * It should never be used in gamecode - instead use the appropriate world context function 
@@ -139,10 +142,11 @@ public:
 	UWorld* GetGameWorld();
 
 protected:
+
 	float GetGameViewportDPIScale(UGameViewportClient* GameViewportClient) const;
 
 private:
+
 	virtual void HandleNetworkFailure_NotifyGameInstance(UWorld* World, UNetDriver* NetDriver, ENetworkFailure::Type FailureType) override;
 	virtual void HandleTravelFailure_NotifyGameInstance(UWorld* World, ETravelFailure::Type FailureType) override;
 };
-
diff --git a/Engine/Source/Runtime/Engine/Classes/Engine/GameViewportClient.h b/Engine/Source/Runtime/Engine/Classes/Engine/GameViewportClient.h
--- a/Engine/Source/Runtime/Engine/Classes/Engine/GameViewportClient.h
+++ b/Engine/Source/Runtime/Engine/Classes/Engine/GameViewportClient.h
@@ -1,6 +1,5 @@
 // Copyright 1998-2015 Epic Games, Inc. All Rights Reserved.
 
-
 #pragma once
 
 #include "ShowFlags.h"
@@ -659,8 +658,8 @@ private:
 
 	/** 
 	 * Applies requested changes to display configuration 
-	 * @param	Dimensions - Pointer to new dimensions of the display. NULL for no change.
-	 * @param	WindowMode - What window mode do we want to st the display to.
+	 * @param Dimensions Pointer to new dimensions of the display. nullptr for no change.
+	 * @param WindowMode What window mode do we want to st the display to.
 	 */
 	bool SetDisplayConfiguration( const FIntPoint* Dimensions, EWindowMode::Type WindowMode);
 
diff --git a/Engine/Source/Runtime/Engine/Classes/Engine/GeneratedMeshAreaLight.h b/Engine/Source/Runtime/Engine/Classes/Engine/GeneratedMeshAreaLight.h
--- a/Engine/Source/Runtime/Engine/Classes/Engine/GeneratedMeshAreaLight.h
+++ b/Engine/Source/Runtime/Engine/Classes/Engine/GeneratedMeshAreaLight.h
@@ -2,18 +2,15 @@
 
 #pragma once
 
-/**
- * GeneratedMeshAreaLight - A light type that is created after a lighting build with Lightmass and handles mesh area light influence on dynamic objects.
- */
 #include "Engine/SpotLight.h"
 #include "GeneratedMeshAreaLight.generated.h"
 
+/**
+ * Implements a light that is created after a lighting build with Lightmass and handles mesh area light influence on dynamic objects.
+ */
 UCLASS(notplaceable, MinimalAPI)
-class AGeneratedMeshAreaLight : public ASpotLight
+class AGeneratedMeshAreaLight
+	: public ASpotLight
 {
 	GENERATED_UCLASS_BODY()
-
 };
-
-
-
