In preparation for integrating a better solution, undid deprecation of AActor::GetDistanceTo, et al.  Also removed replacement functions (GetDistanceToActor, et al).

diff --git a/Engine/Source/Runtime/Engine/Classes/GameFramework/Actor.h b/Engine/Source/Runtime/Engine/Classes/GameFramework/Actor.h
--- a/Engine/Source/Runtime/Engine/Classes/GameFramework/Actor.h
+++ b/Engine/Source/Runtime/Engine/Classes/GameFramework/Actor.h
@@ -687,47 +687,25 @@ public:
 	UFUNCTION(BlueprintCallable, Category="Utilities|Orientation")
 	FVector GetActorScale3D() const;
 
-	/** Deprecated in favor of const/pure versions below. */
-	DEPRECATED(4.8, "AActor::GetDistanceTo will be removed. Use AActor::GetDistanceToActor instead.")
-	UFUNCTION(BlueprintCallable, meta = (DeprecatedFunction, DeprecationMessage = "Use the pure node GetDistanceToActor instead"), Category = "Utilities|Transformation")
-	float GetDistanceTo(AActor* OtherActor);
-	
-	DEPRECATED(4.8, "AActor::GetHorizontalDistanceTo will be removed. Use AActor::GetHorizontalDistanceToActor instead.")
-	UFUNCTION(BlueprintCallable, meta = (DeprecatedFunction, DeprecationMessage = "Use the pure node GetHorizontalDistanceToActor instead"), Category = "Utilities|Transformation")
-	float GetHorizontalDistanceTo(AActor* OtherActor);
-	
-	DEPRECATED(4.8, "AActor::GetVerticalDistanceTo will be removed. Use AActor::GetVerticalDistanceToActor instead.")
-	UFUNCTION(BlueprintCallable, meta = (DeprecatedFunction, DeprecationMessage = "Use the pure node GetVerticalDistanceToActor instead"), Category = "Utilities|Transformation")
-	float GetVerticalDistanceTo(AActor* OtherActor);
-	
-	DEPRECATED(4.8, "AActor::GetDotProductTo will be removed. Use AActor::GetDotProductToActor instead.")
-	UFUNCTION(BlueprintCallable, meta = (DeprecatedFunction, DeprecationMessage = "Use the pure node GetDotProductToActor instead"), Category = "Utilities|Transformation")
-	float GetDotProductTo(AActor* OtherActor);
-	
-	DEPRECATED(4.8, "AActor::GetHorizontalDotProductTo will be removed. Use AActor::GetHorizontalDotProductToActor instead.")
-	UFUNCTION(BlueprintCallable, meta = (DeprecatedFunction, DeprecationMessage = "Use the pure node GetHorizontalDotProductToActor instead"), Category = "Utilities|Transformation")
-	float GetHorizontalDotProductTo(AActor* OtherActor);
-
 	/** Returns the distance from this Actor to OtherActor. */
 	UFUNCTION(BlueprintCallable, Category = "Utilities|Transformation")
-	float GetDistanceToActor(AActor* OtherActor) const;
+	float GetDistanceTo(AActor* OtherActor);
 
 	/** Returns the distance from this Actor to OtherActor, ignoring Z. */
 	UFUNCTION(BlueprintCallable, Category = "Utilities|Transformation")
-	float GetHorizontalDistanceToActor(AActor* OtherActor) const;
+	float GetHorizontalDistanceTo(AActor* OtherActor);
 
 	/** Returns the distance from this Actor to OtherActor, ignoring XY. */
 	UFUNCTION(BlueprintCallable, Category = "Utilities|Transformation")
-	float GetVerticalDistanceToActor(AActor* OtherActor) const;
+	float GetVerticalDistanceTo(AActor* OtherActor);
 
 	/** Returns the dot product from this Actor to OtherActor. Returns -2.0 on failure. Returns 0.0 for coincidental actors. */
 	UFUNCTION(BlueprintCallable, Category = "Utilities|Transformation")
-	float GetDotProductToActor(AActor* OtherActor) const;
+	float GetDotProductTo(AActor* OtherActor);
 
 	/** Returns the dot product from this Actor to OtherActor, ignoring Z. Returns -2.0 on failure. Returns 0.0 for coincidental actors. */
 	UFUNCTION(BlueprintCallable, Category = "Utilities|Transformation")
-	float GetHorizontalDotProductToActor(AActor* OtherActor) const;
-
+	float GetHorizontalDotProductTo(AActor* OtherActor);
 
 	/**
 	 * Adds a delta to the location of this actor in world space.
@@ -2356,11 +2334,11 @@ FORCEINLINE FVector AActor::GetSimpleCollisionCylinderExtent() const
 	bool SetActorLocationAndRotation(FVector NewLocation, FRotator NewRotation, bool bSweep=false, FHitResult* OutSweepHitResult=nullptr) { return Super::SetActorLocationAndRotation(NewLocation, NewRotation, bSweep, OutSweepHitResult); } \
 	virtual bool TeleportTo( const FVector& DestLocation, const FRotator& DestRotation, bool bIsATest, bool bNoCheck ) override { return Super::TeleportTo(DestLocation, DestRotation, bIsATest, bNoCheck); } \
 	virtual FVector GetVelocity() const override { return Super::GetVelocity(); } \
-	float GetHorizontalDistanceToActor(AActor* OtherActor)  const { return Super::GetHorizontalDistanceToActor(OtherActor); } \
-	float GetVerticalDistanceToActor(AActor* OtherActor)  const { return Super::GetVerticalDistanceToActor(OtherActor); } \
-	float GetDotProductToActor(AActor* OtherActor) const { return Super::GetDotProductToActor(OtherActor); } \
-	float GetHorizontalDotProductToActor(AActor* OtherActor) const { return Super::GetHorizontalDotProductToActor(OtherActor); } \
-	float GetDistanceToActor(AActor* OtherActor) const { return Super::GetDistanceToActor(OtherActor); }
+	float GetHorizontalDistanceTo(AActor* OtherActor)  { return Super::GetHorizontalDistanceTo(OtherActor); } \
+	float GetVerticalDistanceTo(AActor* OtherActor)  { return Super::GetVerticalDistanceTo(OtherActor); } \
+	float GetDotProductTo(AActor* OtherActor) { return Super::GetDotProductTo(OtherActor); } \
+	float GetHorizontalDotProductTo(AActor* OtherActor) { return Super::GetHorizontalDotProductTo(OtherActor); } \
+	float GetDistanceTo(AActor* OtherActor) { return Super::GetDistanceTo(OtherActor); }
 
 
 
diff --git a/Engine/Source/Runtime/Engine/Private/Actor.cpp b/Engine/Source/Runtime/Engine/Private/Actor.cpp
--- a/Engine/Source/Runtime/Engine/Private/Actor.cpp
+++ b/Engine/Source/Runtime/Engine/Private/Actor.cpp
@@ -3478,25 +3478,21 @@ UMaterialInstanceDynamic* AActor::MakeMIDForMaterial(class UMaterialInterface* P
 	return NULL;
 }
 
-// deprecated
 float AActor::GetDistanceTo(AActor* OtherActor)
 {
 	return OtherActor ? (GetActorLocation() - OtherActor->GetActorLocation()).Size() : 0.f;
 }
 
-// deprecated
 float AActor::GetHorizontalDistanceTo(AActor* OtherActor)
 {
 	return OtherActor ? (GetActorLocation() - OtherActor->GetActorLocation()).Size2D() : 0.f;
 }
 
-// deprecated
 float AActor::GetVerticalDistanceTo(AActor* OtherActor)
 {
 	return OtherActor ? FMath::Abs((GetActorLocation().Z - OtherActor->GetActorLocation().Z)) : 0.f;
 }
 
-// deprecated
 float AActor::GetDotProductTo(AActor* OtherActor)
 {
 	if (OtherActor)
@@ -3509,7 +3505,6 @@ float AActor::GetDotProductTo(AActor* OtherActor)
 	return -2.0;
 }
 
-// deprecated
 float AActor::GetHorizontalDotProductTo(AActor* OtherActor)
 {
 	if (OtherActor)
@@ -3523,46 +3518,6 @@ float AActor::GetHorizontalDotProductTo(AActor* OtherActor)
 }
 
 
-
-float AActor::GetDistanceToActor(AActor* OtherActor) const
-{
-	return OtherActor ? (GetActorLocation() - OtherActor->GetActorLocation()).Size() : 0.f;
-}
-
-float AActor::GetHorizontalDistanceToActor(AActor* OtherActor) const
-{
-	return OtherActor ? (GetActorLocation() - OtherActor->GetActorLocation()).Size2D() : 0.f;
-}
-
-float AActor::GetVerticalDistanceToActor(AActor* OtherActor) const
-{
-	return OtherActor ? FMath::Abs((GetActorLocation().Z - OtherActor->GetActorLocation().Z)) : 0.f;
-}
-
-float AActor::GetDotProductToActor(AActor* OtherActor) const
-{
-	if (OtherActor)
-	{
-		FVector Dir = GetActorForwardVector();
-		FVector Offset = OtherActor->GetActorLocation() - GetActorLocation();
-		Offset = Offset.GetSafeNormal();
-		return FVector::DotProduct(Dir, Offset);
-	}
-	return -2.0;
-}
-
-float AActor::GetHorizontalDotProductToActor(AActor* OtherActor) const
-{
-	if (OtherActor)
-	{
-		FVector Dir = GetActorForwardVector();
-		FVector Offset = OtherActor->GetActorLocation() - GetActorLocation();
-		Offset = Offset.GetSafeNormal2D();
-		return FVector::DotProduct(Dir, Offset);
-	}
-	return -2.0;
-}
-
 #if WITH_EDITOR
 const int32 AActor::GetNumUncachedStaticLightingInteractions() const
 {
