Fix UE-7518 - Multicast RPCs do not work properly when client tries to async load arguments

diff --git a/Engine/Source/Runtime/Engine/Classes/Engine/ActorChannel.h b/Engine/Source/Runtime/Engine/Classes/Engine/ActorChannel.h
--- a/Engine/Source/Runtime/Engine/Classes/Engine/ActorChannel.h
+++ b/Engine/Source/Runtime/Engine/Classes/Engine/ActorChannel.h
@@ -75,6 +75,8 @@ class ENGINE_API UActorChannel
 
 	TArray< TWeakObjectPtr< UObject > >	CreateSubObjects;		// Any sub-object we created on this channel
 
+	TArray< FNetworkGUID >				MustBeMappedGuidsInLastBunch;		// Array of guids that will async load on client. This list is used for queued RPC's.
+
 	/**
 	 * Default constructor
 	 */
@@ -115,6 +117,9 @@ public:
 
 	void SetChannelActorForDestroy( struct FActorDestructionInfo *DestructInfo );
 
+	/** Append any "must be mapped" guids to front of bunch. These are guids that the client will wait on before processing this bunch. */
+	virtual void AppendMustBeMappedGuids( FOutBunch* Bunch ) override;
+
 	virtual void Serialize(FArchive& Ar) override;
 
 	static void AddReferencedObjects(UObject* InThis, FReferenceCollector& Collector);
diff --git a/Engine/Source/Runtime/Engine/Classes/Engine/Channel.h b/Engine/Source/Runtime/Engine/Classes/Engine/Channel.h
--- a/Engine/Source/Runtime/Engine/Classes/Engine/Channel.h
+++ b/Engine/Source/Runtime/Engine/Classes/Engine/Channel.h
@@ -104,6 +104,9 @@ public:
 	 */
 	void ReceivedRawBunch( FInBunch & Bunch, bool & bOutSkipAck );
 	
+	/** Append any "must be mapped" guids to front of bunch. These are guids that the client will wait on before processing this bunch. */
+	virtual void AppendMustBeMappedGuids( FOutBunch* Bunch );
+
 	/** Send a bunch if it's not overflowed, and queue it if it's reliable. */
 	virtual FPacketIdRange SendBunch(FOutBunch* Bunch, bool Merge);
 	
diff --git a/Engine/Source/Runtime/Engine/Private/DataChannel.cpp b/Engine/Source/Runtime/Engine/Private/DataChannel.cpp
--- a/Engine/Source/Runtime/Engine/Private/DataChannel.cpp
+++ b/Engine/Source/Runtime/Engine/Private/DataChannel.cpp
@@ -591,6 +591,52 @@ bool UChannel::ReceivedNextBunch( FInBunch & Bunch, bool & bOutSkipAck )
 	return false;
 }
 
+void UActorChannel::AppendMustBeMappedGuids( FOutBunch* Bunch )
+{
+	if ( MustBeMappedGuidsInLastBunch.Num() > 0 )
+	{
+		// Just add our list to the main list on package map so we can re-use the code in UChannel to add them all together
+		UPackageMapClient * PackageMapClient = CastChecked< UPackageMapClient >( Connection->PackageMap );
+
+		PackageMapClient->GetMustBeMappedGuidsInLastBunch().Append( MustBeMappedGuidsInLastBunch );
+
+		MustBeMappedGuidsInLastBunch.Empty();
+	}
+
+	// Actually add them to the bunch
+	Super::AppendMustBeMappedGuids( Bunch );
+}
+
+void UChannel::AppendMustBeMappedGuids( FOutBunch* Bunch )
+{
+	UPackageMapClient * PackageMapClient = CastChecked< UPackageMapClient >( Connection->PackageMap );
+
+	TArray< FNetworkGUID >& MustBeMappedGuidsInLastBunch = PackageMapClient->GetMustBeMappedGuidsInLastBunch();
+
+	if ( MustBeMappedGuidsInLastBunch.Num() > 0 )
+	{
+		// Rewrite the bunch with the unique guids in front
+		FOutBunch TempBunch( *Bunch );
+
+		Bunch->Reset();
+
+		// Write all the guids out
+		uint16 NumMustBeMappedGUIDs = MustBeMappedGuidsInLastBunch.Num();
+		*Bunch << NumMustBeMappedGUIDs;
+		for ( int32 i = 0; i < MustBeMappedGuidsInLastBunch.Num(); i++ )
+		{
+			*Bunch << MustBeMappedGuidsInLastBunch[i];
+		}
+
+		// Append the original bunch data at the end
+		Bunch->SerializeBits( TempBunch.GetData(), TempBunch.GetNumBits() );
+
+		Bunch->bHasMustBeMappedGUIDs = 1;
+
+		MustBeMappedGuidsInLastBunch.Empty();
+	}
+}
+
 FPacketIdRange UChannel::SendBunch( FOutBunch* Bunch, bool Merge )
 {
 	check(!Closing);
@@ -628,38 +674,18 @@ FPacketIdRange UChannel::SendBunch( FOutBunch* Bunch, bool Merge )
 		Merge = false;
 	}
 
-	// Append any "must be mapped" guids to front of bunch
-	// This is so the client will know to make sure they are loaded, and pause the stream until they are
-	TArray< FNetworkGUID > & MustBeMappedGuidsInLastBunch = PackageMapClient->GetMustBeMappedGuidsInLastBunch();
-
-	if ( MustBeMappedGuidsInLastBunch.Num() > 0 && Connection->Driver->IsServer() )
+	if ( Connection->Driver->IsServer() )
 	{
-		//UE_LOG( LogNet, Warning, TEXT( "SendBunch: Appending must be mapped guids. NumGuids: %i" ), MustBeMappedGuidsInLastBunch.Num() );
+		// Append any "must be mapped" guids to front of bunch from the packagemap
+		AppendMustBeMappedGuids( Bunch );
 
-		// Rewrite the bunch with the unique guids in front
-		FOutBunch TempBunch( *Bunch );
-
-		Bunch->Reset();
-
-		// Write all the guids out
-		uint16 NumMustBeMappedGUIDs = MustBeMappedGuidsInLastBunch.Num();
-		*Bunch << NumMustBeMappedGUIDs;
-		for ( int32 i = 0; i < MustBeMappedGuidsInLastBunch.Num(); i++ )
+		if ( Bunch->bHasMustBeMappedGUIDs )
 		{
-			*Bunch << MustBeMappedGuidsInLastBunch[i];
+			// We can't merge with this, since we need all the unique static guids in the front
+			Merge = false;
 		}
-
-		// Append the original bunch data at the end
-		Bunch->SerializeBits( TempBunch.GetData(), TempBunch.GetNumBits() );
-
-		Bunch->bHasMustBeMappedGUIDs = 1;
-
-		// We can't merge with this, since we need all the unique static guids in the front
-		Merge = false;
 	}
 
-	MustBeMappedGuidsInLastBunch.Empty();
-
 	//-----------------------------------------------------
 	// Contemplate merging.
 	//-----------------------------------------------------
@@ -1873,6 +1899,12 @@ bool UActorChannel::ReplicateActor()
 	checkSlow(Actor);
 	checkSlow(!Closing);
 
+	// The package map shouldn't have any carry over guids
+	if ( CastChecked< UPackageMapClient >( Connection->PackageMap )->GetMustBeMappedGuidsInLastBunch().Num() != 0 )
+	{
+		UE_LOG( LogNet, Warning, TEXT( "ReplicateActor: PackageMap->GetMustBeMappedGuidsInLastBunch().Num() != 0: %i" ), CastChecked< UPackageMapClient >( Connection->PackageMap )->GetMustBeMappedGuidsInLastBunch().Num() );
+	}
+
 	// Time how long it takes to replicate this particular actor
 	STAT( FScopeCycleCounterUObject FunctionScope(Actor) );
 
diff --git a/Engine/Source/Runtime/Engine/Private/DataReplication.cpp b/Engine/Source/Runtime/Engine/Private/DataReplication.cpp
--- a/Engine/Source/Runtime/Engine/Private/DataReplication.cpp
+++ b/Engine/Source/Runtime/Engine/Private/DataReplication.cpp
@@ -1066,6 +1066,18 @@ void FObjectReplicator::QueueRemoteFunctionBunch( UFunction* Func, FOutBunch &Bu
 	}
 
 	RemoteFunctions->SerializeBits( Bunch.GetData(), Bunch.GetNumBits() );
+
+	if ( Connection != NULL && Connection->PackageMap != NULL )
+	{
+		UPackageMapClient * PackageMapClient = CastChecked< UPackageMapClient >( Connection->PackageMap );
+
+		// We need to copy over any info that was obtained on the package map during serialization, and remember it until we actually call SendBunch
+		if ( PackageMapClient->GetMustBeMappedGuidsInLastBunch().Num() )
+		{
+			OwningChannel->MustBeMappedGuidsInLastBunch.Append( PackageMapClient->GetMustBeMappedGuidsInLastBunch() );
+			PackageMapClient->GetMustBeMappedGuidsInLastBunch().Empty();
+		}
+	}
 }
 
 bool FObjectReplicator::ReadyForDormancy(bool suppressLogs)
