Compile fix.

diff --git a/Engine/Source/Runtime/Engine/Classes/Niagara/NiagaraConstantSet.h b/Engine/Source/Runtime/Engine/Classes/Niagara/NiagaraConstantSet.h
--- a/Engine/Source/Runtime/Engine/Classes/Niagara/NiagaraConstantSet.h
+++ b/Engine/Source/Runtime/Engine/Classes/Niagara/NiagaraConstantSet.h
@@ -258,7 +258,7 @@ public:
 		}
 	}
 
-	void SetOrAdd(const FNiagaraVariableInfo& Constant, FVector4 Value)
+	void SetOrAdd(const FNiagaraVariableInfo& Constant, const FVector4& Value)
 	{
 		int32 Idx = VectorConstantsInfo.Find(Constant);
 		if (Idx == INDEX_NONE)
