#ue4 - "p.ShowInitialOverlaps" now also draws the triangles of a trianglemesh or heightfield which the query overlaps. Fixed DRAW_OVERLAPPING_TRIS define in CollisionConversions and repurposed it a bit.

diff --git a/Engine/Source/Runtime/Engine/Private/Collision/CollisionConversions.cpp b/Engine/Source/Runtime/Engine/Private/Collision/CollisionConversions.cpp
--- a/Engine/Source/Runtime/Engine/Private/Collision/CollisionConversions.cpp
+++ b/Engine/Source/Runtime/Engine/Private/Collision/CollisionConversions.cpp
@@ -10,6 +10,9 @@
 #include "Components/DestructibleComponent.h"
 #include "PhysicalMaterials/PhysicalMaterial.h"
 
+#define DRAW_OVERLAPPING_TRIS (!(UE_BUILD_SHIPPING || UE_BUILD_TEST))
+extern TAutoConsoleVariable<int32> CVarShowInitialOverlaps;
+
 // Sentinel for invalid query results.
 static const PxQueryHit InvalidQueryHit;
 
@@ -459,20 +462,20 @@ bool AddSweepResults(int32 NumHits, const PxSweepHit* Hits, float CheckLength, c
 	return bHadBlockingHit;
 }
 
-#define DRAW_OVERLAPPING_TRIS 0
-
 /* Function to find the best normal from the list of triangles that are overlapping our geom. */
 template<typename GeomType>
-FVector FindBestOverlappingNormal(const PxGeometry& Geom, const PxTransform& QueryTM, const GeomType& ShapeGeom, const PxTransform& PShapeWorldPose, PxU32* HitTris, int32 NumTrisHit)
+FVector FindBestOverlappingNormal(const PxGeometry& Geom, const PxTransform& QueryTM, const GeomType& ShapeGeom, const PxTransform& PShapeWorldPose, PxU32* HitTris, int32 NumTrisHit, bool bCanDrawOverlaps = false)
 {
 #if DRAW_OVERLAPPING_TRIS
-	TArray<FOverlapResult> Overlaps;
-	DrawGeomOverlaps(World, Geom, QueryTM, Overlaps);
-
-	TArray<FBatchedLine> Lines;
-	const FLinearColor LineColor = FLinearColor(1.f, 0.7f, 0.7f);
-	const FLinearColor NormalColor = FLinearColor(1.f, 1.f, 1.f);
 	const float Lifetime = 5.f;
+	if (bCanDrawOverlaps)
+	{
+		TArray<FOverlapResult> Overlaps;
+		DrawGeomOverlaps(GWorld, Geom, QueryTM, Overlaps, Lifetime);
+	}
+	const FLinearColor LineColor = FLinearColor::Green;
+	const FLinearColor NormalColor = FLinearColor::Red;
+	const FLinearColor PointColor = FLinearColor::Yellow;
 #endif // DRAW_OVERLAPPING_TRIS
 
 	// Track the best triangle plane distance
@@ -506,20 +509,24 @@ FVector FindBestOverlappingNormal(const PxGeometry& Geom, const PxTransform& Que
 		}
 
 #if DRAW_OVERLAPPING_TRIS
-		Lines.Add(FBatchedLine(A, B, LineColor, Lifetime, 0.1f, SDPG_Foreground));
-		Lines.Add(FBatchedLine(B, C, LineColor, Lifetime, 0.1f, SDPG_Foreground));
-		Lines.Add(FBatchedLine(C, A, LineColor, Lifetime, 0.1f, SDPG_Foreground));
-		Lines.Add(FBatchedLine(A, A + (50.f*TriNormal), NormalColor, Lifetime, 0.1f, SDPG_Foreground));
+		if (bCanDrawOverlaps && GWorld->PersistentLineBatcher)
+		{
+			static float LineThickness = 0.9f;
+			static float NormalThickness = 0.75f;
+			static float PointThickness = 5.0f;
+			GWorld->PersistentLineBatcher->DrawLine(A, B, LineColor, SDPG_Foreground, LineThickness, Lifetime);
+			GWorld->PersistentLineBatcher->DrawLine(B, C, LineColor, SDPG_Foreground, LineThickness, Lifetime);
+			GWorld->PersistentLineBatcher->DrawLine(C, A, LineColor, SDPG_Foreground, LineThickness, Lifetime);
+			const FVector Centroid((A + B + C) / 3.f);
+			GWorld->PersistentLineBatcher->DrawLine(Centroid, Centroid + (35.0f*TriNormal), NormalColor, SDPG_Foreground, NormalThickness, Lifetime);
+			GWorld->PersistentLineBatcher->DrawPoint(Centroid + (35.0f*TriNormal), NormalColor, PointThickness, SDPG_Foreground, Lifetime);
+			GWorld->PersistentLineBatcher->DrawPoint(A, PointColor, PointThickness, SDPG_Foreground, Lifetime);
+			GWorld->PersistentLineBatcher->DrawPoint(B, PointColor, PointThickness, SDPG_Foreground, Lifetime);
+			GWorld->PersistentLineBatcher->DrawPoint(C, PointColor, PointThickness, SDPG_Foreground, Lifetime);
+		}
 #endif // DRAW_OVERLAPPING_TRIS
 	}
 
-#if DRAW_OVERLAPPING_TRIS
-	if (World->PersistentLineBatcher)
-	{
-		World->PersistentLineBatcher->DrawLines(Lines);
-	}
-#endif // DRAW_OVERLAPPING_TRIS
-
 	return BestPlaneNormal;
 }
 
@@ -628,6 +635,39 @@ static bool ComputeInflatedMTD(const float MtdInflation, const PxLocationHit& PH
 }
 
 
+static bool FindOverlappedTriangleNormal(const PxGeometry& Geom, const PxTransform& QueryTM, const PxShape* PShape, const PxTransform& PShapeWorldPose, FVector& OutNormal, bool bCanDrawOverlaps = false)
+{
+	PxTriangleMeshGeometry PTriMeshGeom;
+	PxHeightFieldGeometry PHeightfieldGeom;
+
+	if (PShape->getTriangleMeshGeometry(PTriMeshGeom) || PShape->getHeightFieldGeometry(PHeightfieldGeom))
+	{
+		PxGeometryType::Enum GeometryType = PShape->getGeometryType();
+		const bool bIsTriMesh = (GeometryType == PxGeometryType::eTRIANGLEMESH);
+		PxU32 HitTris[64];
+		bool bOverflow = false;
+
+		const int32 NumTrisHit = bIsTriMesh ? PxMeshQuery::findOverlapTriangleMesh(Geom, QueryTM, PTriMeshGeom, PShapeWorldPose, HitTris, ARRAY_COUNT(HitTris), 0, bOverflow) :
+											  PxMeshQuery::findOverlapHeightField(Geom, QueryTM, PHeightfieldGeom, PShapeWorldPose, HitTris, ARRAY_COUNT(HitTris), 0, bOverflow);
+		if (NumTrisHit > 0)
+		{
+			if (bIsTriMesh)
+			{
+				OutNormal = FindBestOverlappingNormal(Geom, QueryTM, PTriMeshGeom, PShapeWorldPose, HitTris, NumTrisHit, bCanDrawOverlaps);
+			}
+			else
+			{
+				OutNormal = FindBestOverlappingNormal(Geom, QueryTM, PHeightfieldGeom, PShapeWorldPose, HitTris, NumTrisHit, bCanDrawOverlaps);
+			}
+
+			return true;
+		}
+	}
+
+	return false;
+}
+
+
 /** Util to convert an overlapped shape into a sweep hit result, returns whether it was a blocking hit. */
 static bool ConvertOverlappedShapeToImpactHit(const PxLocationHit& PHit, const FVector& StartLoc, const FVector& EndLoc, FHitResult& OutResult, const PxGeometry& Geom, const PxTransform& QueryTM, const PxFilterData& QueryFilter, bool bReturnPhysMat)
 {
@@ -675,6 +715,15 @@ static bool ConvertOverlappedShapeToImpactHit(const PxLocationHit& PHit, const F
 		}
 	}
 
+#if DRAW_OVERLAPPING_TRIS
+	if (CVarShowInitialOverlaps.GetValueOnAnyThread() != 0)
+	{
+		FVector DummyNormal(0.f);
+		const PxTransform PShapeWorldPose = PxShapeExt::getGlobalPose(*PShape, *PActor);
+		FindOverlappedTriangleNormal(Geom, QueryTM, PShape, PShapeWorldPose, DummyNormal, true);
+	}
+#endif
+
 	// Zero-distance hits are often valid hits and we can extract the hit normal from a larger shape with MTD.
 	// For invalid normals we can try other methods as well (get overlapping triangles).
 
@@ -693,30 +742,9 @@ static bool ConvertOverlappedShapeToImpactHit(const PxLocationHit& PHit, const F
 		}
 		else
 		{
-			PxTriangleMeshGeometry PTriMeshGeom;
-			PxHeightFieldGeometry PHeightfieldGeom;
-
-			if (PShape->getTriangleMeshGeometry(PTriMeshGeom) || PShape->getHeightFieldGeometry(PHeightfieldGeom))
+			if (FindOverlappedTriangleNormal(Geom, QueryTM, PShape, PShapeWorldPose, OutResult.ImpactNormal, false))
 			{
-				PxGeometryType::Enum GeometryType = PShape->getGeometryType();
-				const bool bIsTriMesh = (GeometryType == PxGeometryType::eTRIANGLEMESH);
-				PxU32 HitTris[64];
-				bool bOverflow = false;
-
-				const int32 NumTrisHit = bIsTriMesh ? PxMeshQuery::findOverlapTriangleMesh(Geom, QueryTM, PTriMeshGeom, PShapeWorldPose, HitTris, ARRAY_COUNT(HitTris), 0, bOverflow) :
-													  PxMeshQuery::findOverlapHeightField( Geom, QueryTM, PHeightfieldGeom, PShapeWorldPose, HitTris, ARRAY_COUNT(HitTris), 0, bOverflow);
-				if (NumTrisHit > 0)
-				{
-					if (bIsTriMesh)
-					{
-						OutResult.ImpactNormal = FindBestOverlappingNormal(Geom, QueryTM, PTriMeshGeom, PShapeWorldPose, HitTris, NumTrisHit);
-					}
-					else
-					{
-						OutResult.ImpactNormal = FindBestOverlappingNormal(Geom, QueryTM, PHeightfieldGeom, PShapeWorldPose, HitTris, NumTrisHit);
-					}
-
-				}
+				// Success
 			}
 			else
 			{
diff --git a/Engine/Source/Runtime/Engine/Private/PrimitiveComponent.cpp b/Engine/Source/Runtime/Engine/Private/PrimitiveComponent.cpp
--- a/Engine/Source/Runtime/Engine/Private/PrimitiveComponent.cpp
+++ b/Engine/Source/Runtime/Engine/Private/PrimitiveComponent.cpp
@@ -51,7 +51,7 @@ static TAutoConsoleVariable<float> CVarInitialOverlapTolerance(
 	ECVF_Default);
 
 #if !(UE_BUILD_SHIPPING || UE_BUILD_TEST)
-static TAutoConsoleVariable<int32> CVarShowInitialOverlaps(
+TAutoConsoleVariable<int32> CVarShowInitialOverlaps(
 	TEXT("p.ShowInitialOverlaps"),
 	0,
 	TEXT("Show initial overlaps when moving a component, including estimated 'exit' direction.\n")
