Added a sanity check for WorldSetting's existence when checking if navigation system should be created #UE4

diff --git a/Engine/Source/Runtime/Engine/Private/World.cpp b/Engine/Source/Runtime/Engine/Private/World.cpp
--- a/Engine/Source/Runtime/Engine/Private/World.cpp
+++ b/Engine/Source/Runtime/Engine/Private/World.cpp
@@ -764,7 +764,7 @@ void UWorld::InitWorld(const InitializationValues IVS)
 	if (IVS.bCreateNavigation)
 	{
 		AWorldSettings* WorldSettings = GetWorldSettings();
-		if (WorldSettings->bEnableNavigationSystem)
+		if (WorldSettings && WorldSettings->bEnableNavigationSystem)
 		{
 			UNavigationSystem::CreateNavigationSystem(this);
 		}
