Core - More stats for memory profiling

diff --git a/Engine/Source/Runtime/Engine/Private/Animation/Skeleton.cpp b/Engine/Source/Runtime/Engine/Private/Animation/Skeleton.cpp
--- a/Engine/Source/Runtime/Engine/Private/Animation/Skeleton.cpp
+++ b/Engine/Source/Runtime/Engine/Private/Animation/Skeleton.cpp
@@ -94,6 +94,8 @@ void USkeleton::PostDuplicate(bool bDuplicateForPIE)
 
 void USkeleton::Serialize( FArchive& Ar )
 {
+	DECLARE_SCOPE_CYCLE_COUNTER( TEXT("USkeleton::Serialize"), STAT_Skeleton_Serialize, STATGROUP_LoadTime );
+
 	Super::Serialize(Ar);
 
 	if( Ar.UE4Ver() >= VER_UE4_REFERENCE_SKELETON_REFACTOR )
diff --git a/Engine/Source/Runtime/Engine/Private/AudioDerivedData.cpp b/Engine/Source/Runtime/Engine/Private/AudioDerivedData.cpp
--- a/Engine/Source/Runtime/Engine/Private/AudioDerivedData.cpp
+++ b/Engine/Source/Runtime/Engine/Private/AudioDerivedData.cpp
@@ -938,6 +938,8 @@ void USoundWave::SerializeCookedPlatformData(FArchive& Ar)
 		return;
 	}
 
+	DECLARE_SCOPE_CYCLE_COUNTER( TEXT("USoundWave::SerializeCookedPlatformData"), STAT_SoundWave_SerializeCookedPlatformData, STATGROUP_LoadTime );
+
 #if WITH_EDITORONLY_DATA
 	if (Ar.IsCooking() && Ar.IsPersistent())
 	{
diff --git a/Engine/Source/Runtime/Engine/Private/SkeletalMesh.cpp b/Engine/Source/Runtime/Engine/Private/SkeletalMesh.cpp
--- a/Engine/Source/Runtime/Engine/Private/SkeletalMesh.cpp
+++ b/Engine/Source/Runtime/Engine/Private/SkeletalMesh.cpp
@@ -1001,6 +1001,7 @@ void FMultiSizeIndexContainer::CopyIndexBuffer(const TArray<uint32>& NewArray)
 
 void FMultiSizeIndexContainer::Serialize(FArchive& Ar, bool bNeedsCPUAccess)
 {
+	DECLARE_SCOPE_CYCLE_COUNTER( TEXT("FMultiSizeIndexContainer::Serialize"), STAT_MultiSizeIndexContainer_Serialize, STATGROUP_LoadTime );
 	if (Ar.UE4Ver() < VER_UE4_KEEP_SKEL_MESH_INDEX_DATA)
 	{
 		bool bOldNeedsCPUAccess = true;
@@ -1151,6 +1152,8 @@ FArchive& operator<<(FArchive& Ar,FSkelMeshSection& S)
 */
 void FStaticLODModel::Serialize( FArchive& Ar, UObject* Owner, int32 Idx )
 {
+	DECLARE_SCOPE_CYCLE_COUNTER( TEXT("FStaticLODModel::Serialize"), STAT_StaticLODModel_Serialize, STATGROUP_LoadTime );
+
 	const uint8 LodAdjacencyStripFlag = 1;
 	FStripDataFlags StripFlags( Ar, Ar.IsCooking() && !Ar.CookingTarget()->SupportsFeature(ETargetPlatformFeatures::Tessellation) ? LodAdjacencyStripFlag : 0 );
 
@@ -1687,6 +1690,8 @@ void FSkeletalMeshSourceData::Clear()
 /** Serialization. */
 void FSkeletalMeshSourceData::Serialize( FArchive& Ar, USkeletalMesh* SkeletalMesh )
 {
+	DECLARE_SCOPE_CYCLE_COUNTER( TEXT("FSkeletalMeshSourceData::Serialize"), STAT_SkeletalMeshSourceData_Serialize, STATGROUP_LoadTime );
+
 	if ( Ar.IsLoading() )
 	{
 		bool bHaveSourceData = false;
@@ -1864,6 +1869,8 @@ void FSkeletalMeshResource::ReleaseResources()
 
 void FSkeletalMeshResource::Serialize(FArchive& Ar, USkeletalMesh* Owner)
 {
+	DECLARE_SCOPE_CYCLE_COUNTER( TEXT("FSkeletalMeshResource::Serialize"), STAT_SkeletalMeshResource_Serialize, STATGROUP_LoadTime );
+
 	LODModels.Serialize(Ar,Owner);
 }
 
@@ -2278,6 +2285,8 @@ bool USkeletalMesh::IsReadyForFinishDestroy()
 
 void USkeletalMesh::Serialize( FArchive& Ar )
 {
+	DECLARE_SCOPE_CYCLE_COUNTER( TEXT("USkeletalMesh::Serialize"), STAT_SkeletalMesh_Serialize, STATGROUP_LoadTime );
+
 	Super::Serialize(Ar);
 
 	FStripDataFlags StripFlags( Ar );
diff --git a/Engine/Source/Runtime/Engine/Private/SoundWave.cpp b/Engine/Source/Runtime/Engine/Private/SoundWave.cpp
--- a/Engine/Source/Runtime/Engine/Private/SoundWave.cpp
+++ b/Engine/Source/Runtime/Engine/Private/SoundWave.cpp
@@ -16,6 +16,8 @@
 
 void FStreamedAudioChunk::Serialize(FArchive& Ar, UObject* Owner, int32 ChunkIndex)
 {
+	DECLARE_SCOPE_CYCLE_COUNTER( TEXT("FStreamedAudioChunk::Serialize"), STAT_StreamedAudioChunk_Serialize, STATGROUP_LoadTime );
+
 	bool bCooked = Ar.IsCooking();
 	Ar << bCooked;
 
@@ -135,6 +137,8 @@ void USoundWave::GetAssetRegistryTags(TArray<FAssetRegistryTag>& OutTags) const
 
 void USoundWave::Serialize( FArchive& Ar )
 {
+	DECLARE_SCOPE_CYCLE_COUNTER( TEXT("USoundWave::Serialize"), STAT_SoundWave_Serialize, STATGROUP_LoadTime );
+
 	Super::Serialize( Ar );
 
 	bool bCooked = Ar.IsCooking();
diff --git a/Engine/Source/Runtime/Engine/Private/StaticMesh.cpp b/Engine/Source/Runtime/Engine/Private/StaticMesh.cpp
--- a/Engine/Source/Runtime/Engine/Private/StaticMesh.cpp
+++ b/Engine/Source/Runtime/Engine/Private/StaticMesh.cpp
@@ -168,6 +168,8 @@ void FStaticMeshVertexBuffer::ConvertToFullPrecisionUVs()
 */
 void FStaticMeshVertexBuffer::Serialize( FArchive& Ar, bool bNeedsCPUAccess )
 {
+	DECLARE_SCOPE_CYCLE_COUNTER( TEXT("FStaticMeshVertexBuffer::Serialize"), STAT_StaticMeshVertexBuffer_Serialize, STATGROUP_LoadTime );
+
 	FStripDataFlags StripFlags(Ar, 0, VER_UE4_STATIC_SKELETAL_MESH_SERIALIZATION_FIX);
 
 	Ar << NumTexCoords << Stride << NumVertices;
@@ -273,6 +275,8 @@ FArchive& operator<<(FArchive& Ar, FStaticMeshSection& Section)
 
 void FStaticMeshLODResources::Serialize(FArchive& Ar, UObject* Owner, int32 Index)
 {
+	DECLARE_SCOPE_CYCLE_COUNTER( TEXT("FStaticMeshLODResources::Serialize"), STAT_StaticMeshLODResources_Serialize, STATGROUP_LoadTime );
+
 	// Note: this is all derived data, native versioning is not needed, but be sure to bump STATICMESH_DERIVEDDATA_VER when modifying!
 
 	// On cooked platforms we never need the resource data.
@@ -600,6 +604,8 @@ FStaticMeshRenderData::FStaticMeshRenderData()
 
 void FStaticMeshRenderData::Serialize(FArchive& Ar, UStaticMesh* Owner, bool bCooked)
 {
+	DECLARE_SCOPE_CYCLE_COUNTER( TEXT("FStaticMeshRenderData::Serialize"), STAT_StaticMeshRenderData_Serialize, STATGROUP_LoadTime );
+
 	// Note: this is all derived data, native versioning is not needed, but be sure to bump STATICMESH_DERIVEDDATA_VER when modifying!
 #if WITH_EDITOR
 	if (Ar.IsSaving())
@@ -1683,6 +1689,8 @@ FUObjectAnnotationSparseBool GStaticMeshesThatNeedMaterialFixup;
  */
 void UStaticMesh::Serialize(FArchive& Ar)
 {
+	DECLARE_SCOPE_CYCLE_COUNTER( TEXT("UStaticMesh::Serialize"), STAT_StaticMesh_Serialize, STATGROUP_LoadTime );
+
 	Super::Serialize(Ar);
 
 	FStripDataFlags StripFlags( Ar );
diff --git a/Engine/Source/Runtime/Engine/Private/TextureDerivedData.cpp b/Engine/Source/Runtime/Engine/Private/TextureDerivedData.cpp
--- a/Engine/Source/Runtime/Engine/Private/TextureDerivedData.cpp
+++ b/Engine/Source/Runtime/Engine/Private/TextureDerivedData.cpp
@@ -1185,6 +1185,8 @@ static void SerializePlatformData(
 	bool bCooked
 	)
 {
+	DECLARE_SCOPE_CYCLE_COUNTER( TEXT("SerializePlatformData"), STAT_Texture_SerializePlatformData, STATGROUP_LoadTime );
+
 	UEnum* PixelFormatEnum = UTexture::GetPixelFormatEnum();
 
 	Ar << PlatformData->SizeX;
@@ -1669,6 +1671,8 @@ void UTexture::SerializeCookedPlatformData(FArchive& Ar)
 		return;
 	}
 
+	DECLARE_SCOPE_CYCLE_COUNTER( TEXT("UTexture::SerializeCookedPlatformData"), STAT_Texture_SerializeCookedData, STATGROUP_LoadTime );
+
 	UEnum* PixelFormatEnum = UTexture::GetPixelFormatEnum();
 
 
diff --git a/Engine/Source/Runtime/Engine/Private/UnrealEngine.cpp b/Engine/Source/Runtime/Engine/Private/UnrealEngine.cpp
--- a/Engine/Source/Runtime/Engine/Private/UnrealEngine.cpp
+++ b/Engine/Source/Runtime/Engine/Private/UnrealEngine.cpp
@@ -1251,6 +1251,8 @@ void LoadSpecialMaterial(const FString& MaterialName, UMaterial*& Material, bool
  */
 void UEngine::InitializeObjectReferences()
 {
+	DECLARE_SCOPE_CYCLE_COUNTER(TEXT("UEngine::InitializeObjectReferences"), STAT_InitializeObjectReferences, STATGROUP_LoadTime);
+
 	// initialize the special engine/editor materials
 	if (AllowDebugViewmodes())
 	{
@@ -8598,6 +8600,8 @@ void LoadGametypeContent(FWorldContext &Context, const FURL& URL)
 
 bool UEngine::LoadMap( FWorldContext& WorldContext, FURL URL, class UPendingNetGame* Pending, FString& Error )
 {
+	DECLARE_SCOPE_CYCLE_COUNTER(TEXT("UEngine::LoadMap"), STAT_LoadMap, STATGROUP_LoadTime);
+
 	NETWORK_PROFILER(GNetworkProfiler.TrackSessionChange(true,URL));
 	MALLOC_PROFILER( FMallocProfiler::SnapshotMemoryLoadMapStart( URL.Map ) );
 	Error = TEXT("");
@@ -9016,12 +9020,16 @@ bool UEngine::LoadMap( FWorldContext& WorldContext, FURL URL, class UPendingNetG
 		}
 	}
 
-	// load any per-map packages
-	check(WorldContext.World()->PersistentLevel);
-	LoadPackagesFully(WorldContext.World(), FULLYLOAD_Map, WorldContext.World()->PersistentLevel->GetOutermost()->GetName());
+	{
+		DECLARE_SCOPE_CYCLE_COUNTER(TEXT("UEngine::LoadMap.LoadPackagesFully"), STAT_LoadMap_LoadPackagesFully, STATGROUP_LoadTime);
+
+		// load any per-map packages
+		check(WorldContext.World()->PersistentLevel);
+		LoadPackagesFully(WorldContext.World(), FULLYLOAD_Map, WorldContext.World()->PersistentLevel->GetOutermost()->GetName());
 
-	// Make sure "always loaded" sub-levels are fully loaded
-	WorldContext.World()->FlushLevelStreaming(EFlushLevelStreamingType::Visibility);
+		// Make sure "always loaded" sub-levels are fully loaded
+		WorldContext.World()->FlushLevelStreaming(EFlushLevelStreamingType::Visibility);
+	}
 	
 	UNavigationSystem::InitializeForWorld(WorldContext.World(), FNavigationSystem::GameMode);
 	
