Core - More stats for memory profiling

diff --git a/Engine/Source/Runtime/Core/Private/Serialization/AsyncIOSystemBase.cpp b/Engine/Source/Runtime/Core/Private/Serialization/AsyncIOSystemBase.cpp
--- a/Engine/Source/Runtime/Core/Private/Serialization/AsyncIOSystemBase.cpp
+++ b/Engine/Source/Runtime/Core/Private/Serialization/AsyncIOSystemBase.cpp
@@ -7,6 +7,8 @@
 #include "CorePrivatePCH.h"
 #include "Serialization/AsyncIOSystemBase.h"
 
+DECLARE_STATS_GROUP_VERBOSE(TEXT("AsyncIOSystem"),STATGROUP_AsyncIOSystem, STATCAT_Advanced);
+
 /*-----------------------------------------------------------------------------
 	FAsyncIOSystemBase implementation.
 -----------------------------------------------------------------------------*/
@@ -106,6 +108,8 @@ void FAsyncIOSystemBase::LogIORequest(const FString& Message, const FAsyncIORequ
 
 bool FAsyncIOSystemBase::InternalRead( IFileHandle* FileHandle, int64 Offset, int64 Size, void* Dest )
 {
+	DECLARE_SCOPE_CYCLE_COUNTER(TEXT("FAsyncIOSystemBase::InternalRead"), STAT_AsyncIOSystemBase_InternalRead, STATGROUP_AsyncIOSystem);
+
 	FScopeLock ScopeLock( ExclusiveReadCriticalSection );
 
 	bool bRetVal = false;
@@ -182,6 +186,8 @@ int64 FAsyncIOSystemBase::PlatformMinimumReadSize()
 
 void FAsyncIOSystemBase::FulfillCompressedRead( const FAsyncIORequest& IORequest, IFileHandle* FileHandle )
 {
+	DECLARE_SCOPE_CYCLE_COUNTER(TEXT("FAsyncIOSystemBase::FulfillCompressedRead"), STAT_AsyncIOSystemBase_FulfillCompressedRead, STATGROUP_AsyncIOSystem);
+
 	if (GbLogAsyncLoading == true)
 	{
 		LogIORequest(TEXT("FulfillCompressedRead"), IORequest);
diff --git a/Engine/Source/Runtime/Engine/Private/AudioStreaming.cpp b/Engine/Source/Runtime/Engine/Private/AudioStreaming.cpp
--- a/Engine/Source/Runtime/Engine/Private/AudioStreaming.cpp
+++ b/Engine/Source/Runtime/Engine/Private/AudioStreaming.cpp
@@ -34,6 +34,8 @@ FAsyncStreamDerivedChunkWorker::FAsyncStreamDerivedChunkWorker(
 /** Retrieves the derived chunk from the derived data cache. */
 void FAsyncStreamDerivedChunkWorker::DoWork()
 {
+	DECLARE_SCOPE_CYCLE_COUNTER(TEXT("FAsyncStreamDerivedChunkWorker::DoWork"), STAT_AsyncStreamDerivedChunkWorker_DoWork, STATGROUP_StreamingDetails);
+
 	UE_LOG(LogAudio, Verbose, TEXT("Start of ASync DDC Chunk read for key: %s"), *DerivedDataKey);
 
 	TArray<uint8> DerivedChunkData;
diff --git a/Engine/Source/Runtime/Engine/Private/ContentStreaming.cpp b/Engine/Source/Runtime/Engine/Private/ContentStreaming.cpp
--- a/Engine/Source/Runtime/Engine/Private/ContentStreaming.cpp
+++ b/Engine/Source/Runtime/Engine/Private/ContentStreaming.cpp
@@ -41,9 +41,6 @@ DEFINE_STAT(STAT_NumWantingTextures);
 
 DEFINE_STAT(STAT_StreamingTextures);
 
-DEFINE_STAT(STAT_RenderingThreadUpdateTime);
-DEFINE_STAT(STAT_RenderingThreadFinalizeTime);
-
 DEFINE_STAT(STAT_TotalStaticTextureHeuristicSize);
 DEFINE_STAT(STAT_TotalDynamicHeuristicSize);
 DEFINE_STAT(STAT_TotalLastRenderHeuristicSize);
@@ -1046,9 +1043,9 @@ public:
 	}
 
 	/** Statistics for the async work. */
-	struct FThreadStats
+	struct FAsyncStats
 	{
-		FThreadStats()
+		FAsyncStats()
 		{
 			Reset();
 		}
@@ -1095,7 +1092,7 @@ public:
 	}
 
 	/** Returns the thread statistics. */
-	const FThreadStats& GetStats() const
+	const FAsyncStats& GetStats() const
 	{
 		return ThreadStats;
 	}
@@ -1105,6 +1102,7 @@ private:
 	/** Performs the async work. */
 	void DoWork()
 	{
+		DECLARE_SCOPE_CYCLE_COUNTER(TEXT("FAsyncTextureStreaming::DoWork"), STAT_AsyncTextureStreaming_DoWork, STATGROUP_StreamingDetails);
 		PrioritizedTextures.Empty( StreamingManager.StreamingTextures.Num() );
 
 		// Calculate DynamicWantedMips and DynamicMinDistanceSq for all dynamic textures.
@@ -1220,7 +1218,7 @@ private:
 	/** Context (temporary info) used for the async work. */
 	FStreamingContext			ThreadContext;
 	/** Thread statistics. */
-	FThreadStats				ThreadStats;
+	FAsyncStats				ThreadStats;
 	/** Whether the async work should abort its processing. */
 	volatile bool				bAbort;
 };
@@ -3264,7 +3262,7 @@ void FStreamingManagerTexture::StreamTextures( bool bProcessEverything )
 	FStreamingContext Context( bProcessEverything, IndividualStreamingTexture, bCollectTextureStats );
 
 	const TArray<FTexturePriority>& PrioritizedTextures = AsyncWork->GetTask().GetPrioritizedTextures();
-	FAsyncTextureStreaming::FThreadStats ThreadStats = AsyncWork->GetTask().GetStats();
+	FAsyncTextureStreaming::FAsyncStats ThreadStats = AsyncWork->GetTask().GetStats();
 	Context.AddStats( AsyncWork->GetTask().GetContext() );
 
 #if STATS
diff --git a/Engine/Source/Runtime/Engine/Private/Level.cpp b/Engine/Source/Runtime/Engine/Private/Level.cpp
--- a/Engine/Source/Runtime/Engine/Private/Level.cpp
+++ b/Engine/Source/Runtime/Engine/Private/Level.cpp
@@ -284,6 +284,8 @@ struct FLegacyCoverIndexPair
 
 void ULevel::Serialize( FArchive& Ar )
 {
+	DECLARE_SCOPE_CYCLE_COUNTER(TEXT("ULevel::Serialize"), STAT_Level_Serialize, STATGROUP_LoadTime);
+
 	Super::Serialize( Ar );
 
 	Ar << Actors;
diff --git a/Engine/Source/Runtime/Engine/Private/ReflectionCaptureComponent.cpp b/Engine/Source/Runtime/Engine/Private/ReflectionCaptureComponent.cpp
--- a/Engine/Source/Runtime/Engine/Private/ReflectionCaptureComponent.cpp
+++ b/Engine/Source/Runtime/Engine/Private/ReflectionCaptureComponent.cpp
@@ -849,6 +849,8 @@ void UReflectionCaptureComponent::SerializeSourceData(FArchive& Ar)
 
 void UReflectionCaptureComponent::Serialize(FArchive& Ar)
 {
+	DECLARE_SCOPE_CYCLE_COUNTER(TEXT("UReflectionCaptureComponent::Serialize"), STAT_ReflectionCaptureComponent_Serialize, STATGROUP_LoadTime);
+
 	Super::Serialize(Ar);
 
 	bool bCooked = false;
diff --git a/Engine/Source/Runtime/Engine/Private/Texture2D.cpp b/Engine/Source/Runtime/Engine/Private/Texture2D.cpp
--- a/Engine/Source/Runtime/Engine/Private/Texture2D.cpp
+++ b/Engine/Source/Runtime/Engine/Private/Texture2D.cpp
@@ -1546,11 +1546,11 @@ void FTexture2DResource::BeginCancelUpdate()
  */
 void FTexture2DResource::UpdateMipCount()
 {
+	DECLARE_SCOPE_CYCLE_COUNTER(TEXT("FTexture2DResource::UpdateMipCount"), STAT_Texture2DResource_UpdateMipCount, STATGROUP_StreamingDetails);
+
 	FTexture2DScopedDebugInfo ScopedDebugInfo(Owner);
 	const TIndirectArray<FTexture2DMipMap>& OwnerMips = Owner->GetPlatformMips();
 
-	SCOPE_CYCLE_COUNTER(STAT_RenderingThreadUpdateTime);
-
 	static auto CVarVirtualTextureReducedMemoryEnabled = IConsoleManager::Get().FindTConsoleVariableDataInt(TEXT("r.VirtualTextureReducedMemory"));
 	check(CVarVirtualTextureReducedMemoryEnabled);
 	if ( (CVarVirtualTextureReducedMemoryEnabled->GetValueOnRenderThread() != 0) && CanCreateAsVirtualTexture(Owner, Texture2DRHI->GetFlags()) )
@@ -1717,7 +1717,8 @@ void FTexture2DResource::UpdateMipCount()
  */
 void FTexture2DResource::LoadMipData()
 {
-	SCOPE_CYCLE_COUNTER(STAT_RenderingThreadUpdateTime);
+	DECLARE_SCOPE_CYCLE_COUNTER(TEXT("FTexture2DResource::LoadMipData"), STAT_Texture2DResource_LoadMipData, STATGROUP_StreamingDetails);
+
 	check(Owner->bIsStreamable);
 	check(Owner->PendingMipChangeRequestStatus.GetValue() == TexState_InProgress_Loading);
 
@@ -1914,6 +1915,8 @@ void FTexture2DResource::LoadMipData()
 
 void FCreateTextureTask::DoWork()
 {
+	DECLARE_SCOPE_CYCLE_COUNTER(TEXT("FCreateTextureTask::DoWork"), STAT_CreateTextureTask_DoWork, STATGROUP_StreamingDetails);
+
 	{
 		FTexture2DRHIRef AsyncTexture = RHIAsyncCreateTexture2D(Args.SizeX,Args.SizeY,Args.Format,Args.NumMips,Args.Flags,Args.MipData,Args.NumNewMips);
 		check(IsValidRef(AsyncTexture));
@@ -1935,7 +1938,7 @@ void FCreateTextureTask::DoWork()
 
 void FTexture2DResource::UploadMipData()
 {
-	SCOPE_CYCLE_COUNTER(STAT_RenderingThreadFinalizeTime);
+	DECLARE_SCOPE_CYCLE_COUNTER(TEXT("FTexture2DResource::UploadMipData"), STAT_Texture2DResource_UploadMipData, STATGROUP_StreamingDetails);
 
 	check(Owner->bIsStreamable);
 	check(Owner->PendingMipChangeRequestStatus.GetValue()==TexState_InProgress_Upload);
@@ -2171,7 +2174,7 @@ void FTexture2DResource::FinalizeMipCount()
  */
 void FTexture2DResource::CancelUpdate()
 {
-	SCOPE_CYCLE_COUNTER(STAT_RenderingThreadUpdateTime);
+	DECLARE_SCOPE_CYCLE_COUNTER(TEXT("FTexture2DResource::CancelUpdate"), STAT_Texture2DResource_CancelUpdate, STATGROUP_StreamingDetails);
 
 	// TexState_InProgress_Finalization is valid as the request status gets decremented in the main thread. The actual
 	// call to FinalizeMipCount will happen after this one though.
diff --git a/Engine/Source/Runtime/Engine/Private/TextureCube.cpp b/Engine/Source/Runtime/Engine/Private/TextureCube.cpp
--- a/Engine/Source/Runtime/Engine/Private/TextureCube.cpp
+++ b/Engine/Source/Runtime/Engine/Private/TextureCube.cpp
@@ -17,6 +17,8 @@ UTextureCube::UTextureCube(const FObjectInitializer& ObjectInitializer)
 
 void UTextureCube::Serialize(FArchive& Ar)
 {
+	DECLARE_SCOPE_CYCLE_COUNTER(TEXT("UTextureCube::Serialize"), STAT_TextureCube_Serialize, STATGROUP_LoadTime);
+
 	Super::Serialize(Ar);
 
 	FStripDataFlags StripFlags(Ar);
diff --git a/Engine/Source/Runtime/Engine/Public/ContentStreaming.h b/Engine/Source/Runtime/Engine/Public/ContentStreaming.h
--- a/Engine/Source/Runtime/Engine/Public/ContentStreaming.h
+++ b/Engine/Source/Runtime/Engine/Public/ContentStreaming.h
@@ -32,8 +32,6 @@ DECLARE_MEMORY_STAT_POOL_EXTERN(TEXT("Textures Streamed In, Frame"),STAT_Request
 DECLARE_MEMORY_STAT_POOL_EXTERN(TEXT("Textures Streamed In, Total"),STAT_RequestSizeTotal,STATGROUP_StreamingDetails,FPlatformMemory::MCR_TexturePool, );
 DECLARE_MEMORY_STAT_POOL_EXTERN(TEXT("Lightmaps Streamed In, Total"),STAT_LightmapRequestSizeTotal,STATGROUP_StreamingDetails,FPlatformMemory::MCR_TexturePool, );
 DECLARE_CYCLE_STAT_EXTERN(TEXT("Game Thread Update Time"),STAT_GameThreadUpdateTime,STATGROUP_Streaming, );
-DECLARE_CYCLE_STAT_EXTERN(TEXT("Rendering Thread Update Time"),STAT_RenderingThreadUpdateTime,STATGROUP_StreamingDetails, );
-DECLARE_CYCLE_STAT_EXTERN(TEXT("Rendering Thread Finalize Time"),STAT_RenderingThreadFinalizeTime,STATGROUP_StreamingDetails, );
 DECLARE_FLOAT_ACCUMULATOR_STAT_EXTERN(TEXT("Streaming Latency, Average (sec)"),STAT_StreamingLatency,STATGROUP_StreamingDetails, );
 DECLARE_FLOAT_ACCUMULATOR_STAT_EXTERN(TEXT("Streaming Bandwidth, Average (MB/s)"),STAT_StreamingBandwidth,STATGROUP_StreamingDetails, );
 DECLARE_FLOAT_ACCUMULATOR_STAT_EXTERN(TEXT("Dynamic Streaming Total Time (sec)"),STAT_DynamicStreamingTotal,STATGROUP_StreamingDetails, );
