Fix for crash when GetRelevance_Internal is called from render thread  (occurs immediately after creating a new landscape)

It would be nice if we just made all these two-version functions templated based on which thread they're for.

diff --git a/Engine/Plugins/2D/Paper2D/Source/Paper2D/Private/PaperFlipbookSceneProxy.cpp b/Engine/Plugins/2D/Paper2D/Source/Paper2D/Private/PaperFlipbookSceneProxy.cpp
--- a/Engine/Plugins/2D/Paper2D/Source/Paper2D/Private/PaperFlipbookSceneProxy.cpp
+++ b/Engine/Plugins/2D/Paper2D/Source/Paper2D/Private/PaperFlipbookSceneProxy.cpp
@@ -12,5 +12,5 @@ FPaperFlipbookSceneProxy::FPaperFlipbookSceneProxy(const UPaperFlipbookComponent
 	//@TODO: PAPER2D: WireframeColor = RenderComp->GetWireframeColor();
 
 	Material = InComponent->GetMaterial(0);
-	MaterialRelevance = InComponent->GetMaterialRelevance(GetScene().GetFeatureLevel());
+	MaterialRelevance = InComponent->GetMaterialRelevance_GameThread(GetScene().GetFeatureLevel());
 }
diff --git a/Engine/Plugins/2D/Paper2D/Source/Paper2D/Private/PaperSpriteSceneProxy.cpp b/Engine/Plugins/2D/Paper2D/Source/Paper2D/Private/PaperSpriteSceneProxy.cpp
--- a/Engine/Plugins/2D/Paper2D/Source/Paper2D/Private/PaperSpriteSceneProxy.cpp
+++ b/Engine/Plugins/2D/Paper2D/Source/Paper2D/Private/PaperSpriteSceneProxy.cpp
@@ -15,7 +15,7 @@ FPaperSpriteSceneProxy::FPaperSpriteSceneProxy(const UPaperSpriteComponent* InCo
 	Material = InComponent->GetMaterial(0);
 	AlternateMaterial = InComponent->GetMaterial(1);
 	MaterialSplitIndex = INDEX_NONE;
-	MaterialRelevance = InComponent->GetMaterialRelevance(GetScene().GetFeatureLevel());
+	MaterialRelevance = InComponent->GetMaterialRelevance_GameThread(GetScene().GetFeatureLevel());
 
 	SourceSprite = InComponent->SourceSprite; //@TODO: This is totally not threadsafe, and won't keep up to date if the actor's sprite changes, etc....
 }
diff --git a/Engine/Plugins/Runtime/CableComponent/Source/CableComponent/Private/CableComponent.cpp b/Engine/Plugins/Runtime/CableComponent/Source/CableComponent/Private/CableComponent.cpp
--- a/Engine/Plugins/Runtime/CableComponent/Source/CableComponent/Private/CableComponent.cpp
+++ b/Engine/Plugins/Runtime/CableComponent/Source/CableComponent/Private/CableComponent.cpp
@@ -93,7 +93,7 @@ public:
 		: FPrimitiveSceneProxy(Component)
 		, Material(NULL)
 		, DynamicData(NULL)
-		, MaterialRelevance(Component->GetMaterialRelevance(GetScene().GetFeatureLevel()))
+		, MaterialRelevance(Component->GetMaterialRelevance_GameThread(GetScene().GetFeatureLevel()))
 		, NumSegments(Component->NumSegments)
 		, CableWidth(Component->CableWidth)
 		, NumSides(Component->NumSides)
diff --git a/Engine/Plugins/Runtime/CustomMeshComponent/Source/CustomMeshComponent/Private/CustomMeshComponent.cpp b/Engine/Plugins/Runtime/CustomMeshComponent/Source/CustomMeshComponent/Private/CustomMeshComponent.cpp
--- a/Engine/Plugins/Runtime/CustomMeshComponent/Source/CustomMeshComponent/Private/CustomMeshComponent.cpp
+++ b/Engine/Plugins/Runtime/CustomMeshComponent/Source/CustomMeshComponent/Private/CustomMeshComponent.cpp
@@ -81,7 +81,7 @@ public:
 
 	FCustomMeshSceneProxy(UCustomMeshComponent* Component)
 		: FPrimitiveSceneProxy(Component)
-		, MaterialRelevance(Component->GetMaterialRelevance(GetScene().GetFeatureLevel()))
+		, MaterialRelevance(Component->GetMaterialRelevance_GameThread(GetScene().GetFeatureLevel()))
 	{
 		const FColor VertexColor(255,255,255);
 
diff --git a/Engine/Source/Runtime/Engine/Classes/Components/MeshComponent.h b/Engine/Source/Runtime/Engine/Classes/Components/MeshComponent.h
--- a/Engine/Source/Runtime/Engine/Classes/Components/MeshComponent.h
+++ b/Engine/Source/Runtime/Engine/Classes/Components/MeshComponent.h
@@ -40,8 +40,8 @@ class ENGINE_API UMeshComponent : public UPrimitiveComponent
 	virtual void GetUsedMaterials(TArray<UMaterialInterface*>& OutMaterials) const override;
 	// End UPrimitiveComponent Interface
 
-	/** Accesses the scene relevance information for the materials applied to the mesh. */
-	FMaterialRelevance GetMaterialRelevance(ERHIFeatureLevel::Type InFeatureLevel) const;
+	/** Accesses the scene relevance information for the materials applied to the mesh. Valid from game thread only. */
+	FMaterialRelevance GetMaterialRelevance_GameThread(ERHIFeatureLevel::Type InFeatureLevel) const;
 
 	/**
 	 *	Tell the streaming system whether or not all mip levels of all textures used by this component should be loaded and remain loaded.
diff --git a/Engine/Source/Runtime/Engine/Classes/Materials/MaterialInterface.h b/Engine/Source/Runtime/Engine/Classes/Materials/MaterialInterface.h
--- a/Engine/Source/Runtime/Engine/Classes/Materials/MaterialInterface.h
+++ b/Engine/Source/Runtime/Engine/Classes/Materials/MaterialInterface.h
@@ -354,7 +354,8 @@ public:
 	ENGINE_API FMaterialRelevance GetRelevance_Concurrent(ERHIFeatureLevel::Type InFeatureLevel) const;
 private:
 	// might get called from game or render thread
-	ENGINE_API FMaterialRelevance GetRelevance_Internal(const UMaterial* Material, ERHIFeatureLevel::Type InFeatureLevel) const;
+	template<bool IsGameThread>
+	FMaterialRelevance GetRelevance_Internal(const UMaterial* Material, ERHIFeatureLevel::Type InFeatureLevel) const;
 public:
 
 	int32 GetWidth() const;
diff --git a/Engine/Source/Runtime/Engine/Private/Materials/MaterialInterface.cpp b/Engine/Source/Runtime/Engine/Private/Materials/MaterialInterface.cpp
--- a/Engine/Source/Runtime/Engine/Private/Materials/MaterialInterface.cpp
+++ b/Engine/Source/Runtime/Engine/Private/Materials/MaterialInterface.cpp
@@ -39,6 +39,7 @@ void UMaterialInterface::PostLoad()
 	PostLoadDefaultMaterials();
 }
 
+template<bool IsGameThread>
 FMaterialRelevance UMaterialInterface::GetRelevance_Internal(const UMaterial* Material, ERHIFeatureLevel::Type InFeatureLevel) const
 {
 	if(Material)
@@ -58,7 +59,7 @@ FMaterialRelevance UMaterialInterface::GetRelevance_Internal(const UMaterial* Ma
 		MaterialRelevance.bNormalTranslucency = bIsTranslucent && !Material->bEnableSeparateTranslucency;
 		MaterialRelevance.bDisableDepthTest = bIsTranslucent && Material->bDisableDepthTest;
 		MaterialRelevance.bSubsurfaceProfile = (Material->MaterialDomain == MD_Surface) && !bIsTranslucent && (ShadingModel == MSM_SubsurfaceProfile);
-		MaterialRelevance.bHasWorldPositionOffset = MaterialResource->MaterialModifiesMeshPosition_GameThread();
+		MaterialRelevance.bHasWorldPositionOffset = IsGameThread ? MaterialResource->MaterialModifiesMeshPosition_GameThread() : MaterialResource->MaterialModifiesMeshPosition_RenderThread();
 
 		return MaterialRelevance;
 	}
@@ -72,7 +73,7 @@ FMaterialRelevance UMaterialInterface::GetRelevance(ERHIFeatureLevel::Type InFea
 {
 	// Find the interface's concrete material.
 	const UMaterial* Material = GetMaterial();
-	return GetRelevance_Internal(Material, InFeatureLevel);
+	return GetRelevance_Internal<true>(Material, InFeatureLevel);
 }
 
 FMaterialRelevance UMaterialInterface::GetRelevance_Concurrent(ERHIFeatureLevel::Type InFeatureLevel) const
@@ -80,7 +81,7 @@ FMaterialRelevance UMaterialInterface::GetRelevance_Concurrent(ERHIFeatureLevel:
 	// Find the interface's concrete material.
 	TMicRecursionGuard RecursionGuard;
 	const UMaterial* Material = GetMaterial_Concurrent(RecursionGuard);
-	return GetRelevance_Internal(Material, InFeatureLevel);
+	return GetRelevance_Internal<false>(Material, InFeatureLevel);
 }
 
 int32 UMaterialInterface::GetWidth() const
diff --git a/Engine/Source/Runtime/Engine/Private/MeshComponent.cpp b/Engine/Source/Runtime/Engine/Private/MeshComponent.cpp
--- a/Engine/Source/Runtime/Engine/Private/MeshComponent.cpp
+++ b/Engine/Source/Runtime/Engine/Private/MeshComponent.cpp
@@ -60,7 +60,7 @@ void UMeshComponent::SetMaterial(int32 ElementIndex, UMaterialInterface* Materia
 	}
 }
 
-FMaterialRelevance UMeshComponent::GetMaterialRelevance(ERHIFeatureLevel::Type InFeatureLevel) const
+FMaterialRelevance UMeshComponent::GetMaterialRelevance_GameThread(ERHIFeatureLevel::Type InFeatureLevel) const
 {
 	// Combine the material relevance for all materials.
 	FMaterialRelevance Result;
@@ -71,7 +71,7 @@ FMaterialRelevance UMeshComponent::GetMaterialRelevance(ERHIFeatureLevel::Type I
 		{
 			MaterialInterface = UMaterial::GetDefaultMaterial(MD_Surface);
 		}
-		Result |= MaterialInterface->GetRelevance_Concurrent(InFeatureLevel);
+		Result |= MaterialInterface->GetRelevance(InFeatureLevel);
 	}
 	return Result;
 }
diff --git a/Engine/Source/Runtime/Engine/Private/SkeletalMesh.cpp b/Engine/Source/Runtime/Engine/Private/SkeletalMesh.cpp
--- a/Engine/Source/Runtime/Engine/Private/SkeletalMesh.cpp
+++ b/Engine/Source/Runtime/Engine/Private/SkeletalMesh.cpp
@@ -3946,7 +3946,7 @@ FSkeletalMeshSceneProxy::FSkeletalMeshSceneProxy(const USkinnedMeshComponent* Co
 		,	PhysicsAssetForDebug(Component->GetPhysicsAsset())
 		,	bForceWireframe(Component->bForceWireframe)
 		,	bCanHighlightSelectedSections(Component->bCanHighlightSelectedSections)
-		,	MaterialRelevance(Component->GetMaterialRelevance(GetScene().GetFeatureLevel()))
+		,	MaterialRelevance(Component->GetMaterialRelevance_GameThread(GetScene().GetFeatureLevel()))
 		,	bMaterialsNeedMorphUsage_GameThread(false)
 {
 	check(MeshObject);
diff --git a/Engine/Source/Runtime/Engine/Private/StaticMeshRender.cpp b/Engine/Source/Runtime/Engine/Private/StaticMeshRender.cpp
--- a/Engine/Source/Runtime/Engine/Private/StaticMeshRender.cpp
+++ b/Engine/Source/Runtime/Engine/Private/StaticMeshRender.cpp
@@ -58,7 +58,7 @@ FStaticMeshSceneProxy::FStaticMeshSceneProxy(UStaticMeshComponent* InComponent):
 	ForcedLodModel(InComponent->ForcedLodModel),
 	bCastShadow(InComponent->CastShadow),
 	CollisionTraceFlag(ECollisionTraceFlag::CTF_UseDefault),
-	MaterialRelevance(InComponent->GetMaterialRelevance(GetScene().GetFeatureLevel())),
+	MaterialRelevance(InComponent->GetMaterialRelevance_GameThread(GetScene().GetFeatureLevel())),
 	CollisionResponse(InComponent->GetCollisionResponseToChannels())
 {
 	check(RenderData);
