Fix root motion weighting issue when blending between animations in sync groups

Previous to this change, only root motion from the group leader was applied. Now root motion is blended in from all animations in the group, with their repsective weights.

diff --git a/Engine/Source/Runtime/Engine/Private/Animation/AnimInstance.cpp b/Engine/Source/Runtime/Engine/Private/Animation/AnimInstance.cpp
--- a/Engine/Source/Runtime/Engine/Private/Animation/AnimInstance.cpp
+++ b/Engine/Source/Runtime/Engine/Private/Animation/AnimInstance.cpp
@@ -409,6 +409,10 @@ void UAnimInstance::UpdateAnimation(float DeltaSeconds)
 			FAnimAssetTickContext TickContext(DeltaSeconds, RootMotionMode);
 			FAnimTickRecord& GroupLeader = SyncGroup.ActivePlayers[GroupLeaderIndex];
 			GroupLeader.SourceAsset->TickAssetPlayerInstance(GroupLeader, this, TickContext);
+			if (RootMotionMode == ERootMotionMode::RootMotionFromEverything && TickContext.RootMotionMovementParams.bHasRootMotion)
+			{
+				ExtractedRootMotion.AccumulateWithBlend(TickContext.RootMotionMovementParams.RootMotionTransform, GroupLeader.EffectiveBlendWeight);
+			}
 
 			// Update everything else to follow the leader
 			if (SyncGroup.ActivePlayers.Num() > 1)
@@ -421,14 +425,13 @@ void UAnimInstance::UpdateAnimation(float DeltaSeconds)
 					{
 						const FAnimTickRecord& AssetPlayer = SyncGroup.ActivePlayers[TickIndex];
 						AssetPlayer.SourceAsset->TickAssetPlayerInstance(AssetPlayer, this, TickContext);
+						if (RootMotionMode == ERootMotionMode::RootMotionFromEverything && TickContext.RootMotionMovementParams.bHasRootMotion)
+						{
+							ExtractedRootMotion.AccumulateWithBlend(TickContext.RootMotionMovementParams.RootMotionTransform, AssetPlayer.EffectiveBlendWeight);
+						}
 					}
 				}
 			}
-
-			if (RootMotionMode == ERootMotionMode::RootMotionFromEverything && TickContext.RootMotionMovementParams.bHasRootMotion)
-			{
-				ExtractedRootMotion.AccumulateWithBlend(TickContext.RootMotionMovementParams.RootMotionTransform, GroupLeader.EffectiveBlendWeight);
-			}
 		}
 	}
 
