Fix for editor only compilation

diff --git a/Engine/Source/Runtime/Engine/Private/Collision/CollisionConversions.cpp b/Engine/Source/Runtime/Engine/Private/Collision/CollisionConversions.cpp
--- a/Engine/Source/Runtime/Engine/Private/Collision/CollisionConversions.cpp
+++ b/Engine/Source/Runtime/Engine/Private/Collision/CollisionConversions.cpp
@@ -8,8 +8,10 @@
 #include "CollisionDebugDrawing.h"
 #include "CollisionConversions.h"
 #include "Components/DestructibleComponent.h"
+#include "Components/LineBatchComponent.h"
 #include "PhysicalMaterials/PhysicalMaterial.h"
 
+
 #define DRAW_OVERLAPPING_TRIS (!(UE_BUILD_SHIPPING || UE_BUILD_TEST))
 extern TAutoConsoleVariable<int32> CVarShowInitialOverlaps;
 
