Fix for UE-6545 'Changing Preview Rendering Level in StrategyGame turns Wall_EmptySlot material black'
 - Ensure all levels are added to the scene after a feature level switch.
 - Fix situation that caused 100% shadow when no importance volumes exist.
 - Ensures the light volume octree is correct for feature level when adding to the scene.

diff --git a/Engine/Source/Runtime/Engine/Private/PrecomputedLightVolume.cpp b/Engine/Source/Runtime/Engine/Private/PrecomputedLightVolume.cpp
--- a/Engine/Source/Runtime/Engine/Private/PrecomputedLightVolume.cpp
+++ b/Engine/Source/Runtime/Engine/Private/PrecomputedLightVolume.cpp
@@ -222,6 +222,7 @@ void FPrecomputedLightVolume::AddToScene(FSceneInterface* Scene)
 	if (bInitialized && Scene)
 	{
 		Scene->AddPrecomputedLightVolume(this);
+		OctreeForRendering = AllowHighQualityLightmaps(Scene->GetFeatureLevel()) ? &HighQualityLightmapOctree : &LowQualityLightmapOctree;
 	}
 }
 
diff --git a/Engine/Source/Runtime/Engine/Private/World.cpp b/Engine/Source/Runtime/Engine/Private/World.cpp
--- a/Engine/Source/Runtime/Engine/Private/World.cpp
+++ b/Engine/Source/Runtime/Engine/Private/World.cpp
@@ -5395,15 +5395,21 @@ void UWorld::ChangeFeatureLevel(ERHIFeatureLevel::Type InFeatureLevel)
             SlowTask.EnterProgressFrame(10.0f);
             if (Scene)
             {
-                PersistentLevel->ReleaseRenderingResources();
+				for (auto* Level : Levels)
+				{
+					Level->ReleaseRenderingResources();
+				}
+
                 Scene->Release();
                 GetRendererModule().RemoveScene(Scene);
-
                 GetRendererModule().AllocateScene(this, bRequiresHitProxies, FXSystem != nullptr, InFeatureLevel );
 
-                PersistentLevel->InitializeRenderingResources();
-                PersistentLevel->PrecomputedVisibilityHandler.UpdateScene(Scene);
-                PersistentLevel->PrecomputedVolumeDistanceField.UpdateScene(Scene);
+				for (auto* Level : Levels)
+				{
+					Level->InitializeRenderingResources();
+					Level->PrecomputedVisibilityHandler.UpdateScene(Scene);
+					Level->PrecomputedVolumeDistanceField.UpdateScene(Scene);
+				}
             }
 
             SlowTask.EnterProgressFrame(10.0f);
diff --git a/Engine/Source/Runtime/Renderer/Private/IndirectLightingCache.cpp b/Engine/Source/Runtime/Renderer/Private/IndirectLightingCache.cpp
--- a/Engine/Source/Runtime/Renderer/Private/IndirectLightingCache.cpp
+++ b/Engine/Source/Runtime/Renderer/Private/IndirectLightingCache.cpp
@@ -747,7 +747,7 @@ void FIndirectLightingCache::InterpolatePoint(
 	else
 	{
 		OutIncidentRadiance = AccumulatedIncidentRadiance;
-		OutDirectionalShadowing = AccumulatedDirectionalShadowing;
+		OutDirectionalShadowing = 1;
 		// Use an unoccluded vector if no valid samples were found for interpolation
 		OutSkyBentNormal = FVector(0, 0, 1);
 	}
