UE-7623 Giving a category to the OnDestroyed and OnEndPlay Actor delegates (were under the scary 'Event Dispatchers' category before)

diff --git a/Engine/Source/Runtime/Engine/Classes/GameFramework/Actor.h b/Engine/Source/Runtime/Engine/Classes/GameFramework/Actor.h
--- a/Engine/Source/Runtime/Engine/Classes/GameFramework/Actor.h
+++ b/Engine/Source/Runtime/Engine/Classes/GameFramework/Actor.h
@@ -1086,7 +1086,7 @@ public:
 	virtual void ReceiveDestroyed();
 
 	/** Event triggered when the actor is destroyed. */
-	UPROPERTY(BlueprintAssignable)
+	UPROPERTY(BlueprintAssignable, Category="Game")
 	FActorDestroyedSignature OnDestroyed;
 
 
@@ -1095,7 +1095,7 @@ public:
 	virtual void ReceiveEndPlay(EEndPlayReason::Type EndPlayReason);
 
 	/** Event triggered when the actor is being removed from a level. */
-	UPROPERTY(BlueprintAssignable)
+	UPROPERTY(BlueprintAssignable, Category="Game")
 	FActorEndPlaySignature OnEndPlay;
 	
 	// Begin UObject Interface
