PR #727: Expose several PhysicsHandleComponent properties to blueprint graphs (Contributed by [REDACTED_PERSON])
https://github.com/EpicGames/UnrealEngine/pull/727

diff --git a/Engine/Source/Runtime/Engine/Classes/PhysicsEngine/PhysicsHandleComponent.h b/Engine/Source/Runtime/Engine/Classes/PhysicsEngine/PhysicsHandleComponent.h
--- a/Engine/Source/Runtime/Engine/Classes/PhysicsEngine/PhysicsHandleComponent.h
+++ b/Engine/Source/Runtime/Engine/Classes/PhysicsEngine/PhysicsHandleComponent.h
@@ -7,7 +7,7 @@
 
 namespace physx
 {
-	class PxJoint;
+	class PxD6Joint;
 	class PxRigidDynamic;
 }
 
@@ -33,19 +33,19 @@ class UPhysicsHandleComponent : public UActorComponent
 	uint32 bRotationConstrained:1;
 
 	/** Linear damping of the handle spring. */
-	UPROPERTY(EditAnywhere, Category=PhysicsHandle)
+	UPROPERTY(EditAnywhere, BlueprintReadOnly, Category=PhysicsHandle)
 	float LinearDamping;
 
 	/** Linear stiffness of the handle spring */
-	UPROPERTY(EditAnywhere, Category=PhysicsHandle)
+	UPROPERTY(EditAnywhere, BlueprintReadOnly, Category=PhysicsHandle)
 	float LinearStiffness;
 
 	/** Angular stiffness of the handle spring */
-	UPROPERTY(EditAnywhere, Category=PhysicsHandle)
+	UPROPERTY(EditAnywhere, BlueprintReadOnly, Category=PhysicsHandle)
 	float AngularDamping;
 
 	/** Angular stiffness of the handle spring */
-	UPROPERTY(EditAnywhere, Category=PhysicsHandle)
+	UPROPERTY(EditAnywhere, BlueprintReadOnly, Category=PhysicsHandle)
 	float AngularStiffness;
 
 	/** Target transform */
@@ -54,13 +54,13 @@ class UPhysicsHandleComponent : public UActorComponent
 	FTransform CurrentTransform;
 
 	/** How quickly we interpolate the physics target transform */
-	UPROPERTY(EditAnywhere, Category=PhysicsHandle)
+	UPROPERTY(EditAnywhere, BlueprintReadOnly, Category=PhysicsHandle)
 	float InterpolationSpeed;
 
 
 protected:
 	/** Pointer to PhysX joint used by the handle*/
-	physx::PxJoint* HandleData;
+	physx::PxD6Joint* HandleData;
 	/** Pointer to kinematic actor jointed to grabbed object */
 	physx::PxRigidDynamic* KinActorData;
 
@@ -95,9 +95,32 @@ public:
 	UFUNCTION(BlueprintCallable, Category="Physics|Components|PhysicsHandle")
 	ENGINE_API void GetTargetLocationAndRotation(FVector& TargetLocation, FRotator& TargetRotation) const;
 
+	/** Set linear damping */
+	UFUNCTION(BlueprintCallable, Category = "Physics|Components|PhysicsHandle")
+	ENGINE_API void SetLinearDamping(float NewLinearDamping);
+
+	/** Set linear stiffness */
+	UFUNCTION(BlueprintCallable, Category = "Physics|Components|PhysicsHandle")
+	ENGINE_API void SetLinearStiffness(float NewLinearStiffness);
+
+	/** Set angular damping */
+	UFUNCTION(BlueprintCallable, Category = "Physics|Components|PhysicsHandle")
+	ENGINE_API void SetAngularDamping(float NewAngularDamping);
+
+	/** Set angular stiffness */
+	UFUNCTION(BlueprintCallable, Category = "Physics|Components|PhysicsHandle")
+	ENGINE_API void SetAngularStiffness(float NewAngularStiffness);
+
+	/** Set interpolation speed */
+	UFUNCTION(BlueprintCallable, Category = "Physics|Components|PhysicsHandle")
+	ENGINE_API void SetInterpolationSpeed(float NewInterpolationSpeed);
+
 protected:
 	/** Move the kinematic handle to the specified */
 	void UpdateHandleTransform(const FTransform& NewTransform);
+
+	/** Update the underlying constraint drive settings from the params in this component */
+	void UpdateDriveSettings();
 };
 
 
diff --git a/Engine/Source/Runtime/Engine/Private/PhysicsEngine/PhysicsHandleComponent.cpp b/Engine/Source/Runtime/Engine/Private/PhysicsEngine/PhysicsHandleComponent.cpp
--- a/Engine/Source/Runtime/Engine/Private/PhysicsEngine/PhysicsHandleComponent.cpp
+++ b/Engine/Source/Runtime/Engine/Private/PhysicsEngine/PhysicsHandleComponent.cpp
@@ -130,29 +130,15 @@ void UPhysicsHandleComponent::GrabComponent(UPrimitiveComponent* InComponent, FN
 			NewJoint->setMotion(PxD6Axis::eX, PxD6Motion::eFREE);
 			NewJoint->setMotion(PxD6Axis::eY, PxD6Motion::eFREE);
 			NewJoint->setMotion(PxD6Axis::eZ, PxD6Motion::eFREE);
-
-			NewJoint->setDrive(PxD6Drive::eX, PxD6JointDrive(LinearStiffness, LinearDamping, PX_MAX_F32, PxD6JointDriveFlag::eACCELERATION));
-			NewJoint->setDrive(PxD6Drive::eY, PxD6JointDrive(LinearStiffness, LinearDamping, PX_MAX_F32, PxD6JointDriveFlag::eACCELERATION));
-			NewJoint->setDrive(PxD6Drive::eZ, PxD6JointDrive(LinearStiffness, LinearDamping, PX_MAX_F32, PxD6JointDriveFlag::eACCELERATION));
 			NewJoint->setDrivePosition(PxTransform(PxVec3(0,0,0)));
 
-
 			NewJoint->setMotion(PxD6Axis::eTWIST, PxD6Motion::eFREE);
 			NewJoint->setMotion(PxD6Axis::eSWING1, PxD6Motion::eFREE);
 			NewJoint->setMotion(PxD6Axis::eSWING2, PxD6Motion::eFREE);
 			
 			bRotationConstrained = bConstrainRotation;
 			
-			if (bRotationConstrained)
-			{
-				NewJoint->setDrive(PxD6Drive::eSLERP, PxD6JointDrive(AngularStiffness, AngularDamping, PX_MAX_F32, PxD6JointDriveFlag::eACCELERATION));
-
-				//NewJoint->setDrive(PxD6Drive::eTWIST, PxD6JointDrive(AngularStiffness, AngularDamping, PX_MAX_F32, PxD6JointDriveFlag::eACCELERATION));
-				//NewJoint->setDrive(PxD6Drive::eSWING, PxD6JointDrive(AngularStiffness, AngularDamping, PX_MAX_F32, PxD6JointDriveFlag::eACCELERATION));
-				
-				//PosJointDesc.setGlobalAxis(NxVec3(0,0,1));
-			}
-
+			UpdateDriveSettings();
 		}
 	
 	}
@@ -163,6 +149,25 @@ void UPhysicsHandleComponent::GrabComponent(UPrimitiveComponent* InComponent, FN
 	GrabbedBoneName = InBoneName;
 }
 
+void UPhysicsHandleComponent::UpdateDriveSettings()
+{
+#if WITH_PHYSX
+	if(HandleData != nullptr)
+	{
+		HandleData->setDrive(PxD6Drive::eX, PxD6JointDrive(LinearStiffness, LinearDamping, PX_MAX_F32, PxD6JointDriveFlag::eACCELERATION));
+		HandleData->setDrive(PxD6Drive::eY, PxD6JointDrive(LinearStiffness, LinearDamping, PX_MAX_F32, PxD6JointDriveFlag::eACCELERATION));
+		HandleData->setDrive(PxD6Drive::eZ, PxD6JointDrive(LinearStiffness, LinearDamping, PX_MAX_F32, PxD6JointDriveFlag::eACCELERATION));
+
+		if (bRotationConstrained)
+		{
+			HandleData->setDrive(PxD6Drive::eSLERP, PxD6JointDrive(AngularStiffness, AngularDamping, PX_MAX_F32, PxD6JointDriveFlag::eACCELERATION));
+
+			//NewJoint->setDrive(PxD6Drive::eTWIST, PxD6JointDrive(AngularStiffness, AngularDamping, PX_MAX_F32, PxD6JointDriveFlag::eACCELERATION));
+			//NewJoint->setDrive(PxD6Drive::eSWING, PxD6JointDrive(AngularStiffness, AngularDamping, PX_MAX_F32, PxD6JointDriveFlag::eACCELERATION));
+		}
+	}
+#endif // WITH_PHYSX
+}
 
 void UPhysicsHandleComponent::ReleaseComponent()
 {
@@ -198,6 +203,8 @@ void UPhysicsHandleComponent::ReleaseComponent()
 }
 
 
+
+
 void UPhysicsHandleComponent::SetTargetLocation(FVector NewLocation)
 {
 	TargetTransform.SetTranslation(NewLocation);
@@ -286,3 +293,32 @@ void UPhysicsHandleComponent::GetTargetLocationAndRotation(FVector& OutLocation,
 	OutRotation = TargetTransform.Rotator();
 	OutLocation = TargetTransform.GetTranslation();
 }
+
+void UPhysicsHandleComponent::SetLinearDamping(float NewLinearDamping)
+{
+	LinearDamping = NewLinearDamping;
+	UpdateDriveSettings();
+}
+
+void UPhysicsHandleComponent::SetLinearStiffness(float NewLinearStiffness)
+{
+	LinearStiffness = NewLinearStiffness;
+	UpdateDriveSettings();
+}
+
+void UPhysicsHandleComponent::SetAngularDamping(float NewAngularDamping)
+{
+	AngularDamping = NewAngularDamping;
+	UpdateDriveSettings();
+}
+
+void UPhysicsHandleComponent::SetAngularStiffness(float NewAngularStiffness)
+{
+	AngularStiffness = NewAngularStiffness;
+	UpdateDriveSettings();
+}
+
+void UPhysicsHandleComponent::SetInterpolationSpeed(float NewInterpolationSpeed)
+{
+	InterpolationSpeed = NewInterpolationSpeed;
+}
\ No newline at end of file
