PR #734: Implemented GetMaterial and SetMaterial so individual material elements ... (Contributed by [REDACTED_PERSON])
https://github.com/EpicGames/UnrealEngine/pull/734
PR #744: Update MaterialBillboardComponent.cpp (Contributed by [REDACTED_PERSON])
https://github.com/EpicGames/UnrealEngine/pull/744

diff --git a/Engine/Source/Runtime/Engine/Classes/Components/MaterialBillboardComponent.h b/Engine/Source/Runtime/Engine/Classes/Components/MaterialBillboardComponent.h
--- a/Engine/Source/Runtime/Engine/Classes/Components/MaterialBillboardComponent.h
+++ b/Engine/Source/Runtime/Engine/Classes/Components/MaterialBillboardComponent.h
@@ -72,5 +72,7 @@ class ENGINE_API UMaterialBillboardComponent : public UPrimitiveComponent
 	// Begin UPrimitiveComponent Interface
 	virtual FPrimitiveSceneProxy* CreateSceneProxy() override;
 	virtual FBoxSphereBounds CalcBounds(const FTransform& LocalToWorld) const override;
+	virtual UMaterialInterface* GetMaterial(int32 Index) const override;
+	virtual void SetMaterial(int32 ElementIndex, class UMaterialInterface* Material) override;
 	// End UPrimitiveComponent Interface
 };
diff --git a/Engine/Source/Runtime/Engine/Private/MaterialBillboardComponent.cpp b/Engine/Source/Runtime/Engine/Private/MaterialBillboardComponent.cpp
--- a/Engine/Source/Runtime/Engine/Private/MaterialBillboardComponent.cpp
+++ b/Engine/Source/Runtime/Engine/Private/MaterialBillboardComponent.cpp
@@ -188,9 +188,9 @@ public:
 						VertexArray.Vertices[0].Position  = -WorldSizeX * LocalCameraRight + +WorldSizeY * LocalCameraUp;
 						VertexArray.Vertices[0].TexCoords = FVector2D(0,0);
 						VertexArray.Vertices[1].Position  = +WorldSizeX * LocalCameraRight + +WorldSizeY * LocalCameraUp;
-						VertexArray.Vertices[1].TexCoords = FVector2D(1,0);
+						VertexArray.Vertices[1].TexCoords = FVector2D(0,1);
 						VertexArray.Vertices[2].Position  = -WorldSizeX * LocalCameraRight + -WorldSizeY * LocalCameraUp;
-						VertexArray.Vertices[2].TexCoords = FVector2D(0,1);
+						VertexArray.Vertices[2].TexCoords = FVector2D(1,0);
 						VertexArray.Vertices[3].Position  = +WorldSizeX * LocalCameraRight + -WorldSizeY * LocalCameraUp;
 						VertexArray.Vertices[3].TexCoords = FVector2D(1,1);
 			
@@ -286,3 +286,23 @@ void UMaterialBillboardComponent::AddElement(
 
 	MarkRenderStateDirty();
 }
+
+UMaterialInterface* UMaterialBillboardComponent::GetMaterial(int32 Index) const
+{
+	UMaterialInterface* ResultMI = nullptr;
+	if (Elements.IsValidIndex(Index))
+	{
+		ResultMI = Elements[Index].Material;
+	}
+	return ResultMI;
+}
+
+void UMaterialBillboardComponent::SetMaterial(int32 ElementIndex, class UMaterialInterface* Material)
+{
+	if (Elements.IsValidIndex(ElementIndex))
+	{
+		Elements[ElementIndex].Material = Material;
+
+		MarkRenderStateDirty();
+	}
+}
\ No newline at end of file
