UE-6766, UE-7010 Fixing places where BP-callable functions could crash with a NULL world. This can happen in particular when using the GameInstance ReceiveInit BP event in a standalone game.

diff --git a/Engine/Source/Runtime/Engine/Private/GameplayStatics.cpp b/Engine/Source/Runtime/Engine/Private/GameplayStatics.cpp
--- a/Engine/Source/Runtime/Engine/Private/GameplayStatics.cpp
+++ b/Engine/Source/Runtime/Engine/Private/GameplayStatics.cpp
@@ -38,15 +38,18 @@ APlayerController* UGameplayStatics::GetPlayerController(UObject* WorldContextOb
 	if (WorldContextObject)
 	{
 		UWorld* World = GEngine->GetWorldFromContextObject(WorldContextObject);
-		uint32 Index = 0;
-		for (FConstPlayerControllerIterator Iterator = World->GetPlayerControllerIterator(); Iterator; ++Iterator)
+		if(World != nullptr)
 		{
-			APlayerController* PlayerController = *Iterator;
-			if (Index == PlayerIndex)
+			uint32 Index = 0;
+			for (FConstPlayerControllerIterator Iterator = World->GetPlayerControllerIterator(); Iterator; ++Iterator)
 			{
-				return PlayerController;
+				APlayerController* PlayerController = *Iterator;
+				if (Index == PlayerIndex)
+				{
+					return PlayerController;
+				}
+				Index++;
 			}
-			Index++;
 		}
 	}
 
@@ -82,7 +85,7 @@ APlayerController* UGameplayStatics::CreatePlayer(UObject* WorldContextObject, i
 	UWorld* World = GEngine->GetWorldFromContextObject(WorldContextObject);
 	FString Error;
 
-	ULocalPlayer* LocalPlayer = World->GetGameInstance()->CreateLocalPlayer(ControllerId, Error, bSpawnPawn);
+	ULocalPlayer* LocalPlayer = World ? World->GetGameInstance()->CreateLocalPlayer(ControllerId, Error, bSpawnPawn) : NULL;
 
 	if (Error.Len() > 0)
 	{
@@ -112,33 +115,43 @@ class UClass* UGameplayStatics::GetObjectClass(const UObject* Object)
 float UGameplayStatics::GetGlobalTimeDilation(UObject* WorldContextObject)
 {
 	UWorld* const World = GEngine->GetWorldFromContextObject( WorldContextObject );
-	return World->GetWorldSettings()->TimeDilation;
+	return World ? World->GetWorldSettings()->TimeDilation : 1.f;
 }
 
 void UGameplayStatics::SetGlobalTimeDilation(UObject* WorldContextObject, float TimeDilation)
 {
 	UWorld* const World = GEngine->GetWorldFromContextObject( WorldContextObject );
-	if (TimeDilation < 0.0001f || TimeDilation > 20.f)
+	if(World != nullptr)
 	{
-		UE_LOG(LogBlueprintUserMessages, Warning, TEXT("Time Dilation must be between 0.0001 and 20.  Clamping value to that range."));
-		TimeDilation = FMath::Clamp(TimeDilation, 0.0001f, 20.0f);
+		if (TimeDilation < 0.0001f || TimeDilation > 20.f)
+		{
+			UE_LOG(LogBlueprintUserMessages, Warning, TEXT("Time Dilation must be between 0.0001 and 20.  Clamping value to that range."));
+			TimeDilation = FMath::Clamp(TimeDilation, 0.0001f, 20.0f);
+		}
+		World->GetWorldSettings()->TimeDilation = TimeDilation;
 	}
-	World->GetWorldSettings()->TimeDilation = TimeDilation;
 }
 
 bool UGameplayStatics::SetGamePaused(UObject* WorldContextObject, bool bPaused)
 {
 	UWorld* const World = GEngine->GetWorldFromContextObject( WorldContextObject );
-	APlayerController* const PC = World->GetFirstPlayerController();
-	check(PC); // Gathering some information for TTP #303973
+	if(World != nullptr)
+	{
+		APlayerController* const PC = World->GetFirstPlayerController();
+		check(PC); // Gathering some information for TTP #303973
 
-	return PC->SetPause(bPaused);
+		return PC->SetPause(bPaused);
+	}
+	else
+	{
+		return false;
+	}
 }
 
 bool UGameplayStatics::IsGamePaused(UObject* WorldContextObject)
 {
 	UWorld* const World = GEngine->GetWorldFromContextObject( WorldContextObject );
-	return World->IsPaused();
+	return World ? World->IsPaused() : false;
 }
 
 /** @RETURN True if weapon trace from Origin hits component VictimComp.  OutHitResult will contain properties of the hit. */
@@ -342,22 +355,28 @@ AActor* UGameplayStatics::FinishSpawningActor(AActor* Actor, const FTransform& S
 void UGameplayStatics::LoadStreamLevel(UObject* WorldContextObject, FName LevelName,bool bMakeVisibleAfterLoad,bool bShouldBlockOnLoad,FLatentActionInfo LatentInfo)
 {
 	UWorld* World = GEngine->GetWorldFromContextObject(WorldContextObject);
-	FLatentActionManager& LatentManager = World->GetLatentActionManager();
-	if (LatentManager.FindExistingAction<FStreamLevelAction>(LatentInfo.CallbackTarget, LatentInfo.UUID) == NULL)
+	if(World != nullptr)
 	{
-		FStreamLevelAction* NewAction = new FStreamLevelAction(true, LevelName, bMakeVisibleAfterLoad, bShouldBlockOnLoad, LatentInfo, World);
-		LatentManager.AddNewAction(LatentInfo.CallbackTarget, LatentInfo.UUID, NewAction);
+		FLatentActionManager& LatentManager = World->GetLatentActionManager();
+		if (LatentManager.FindExistingAction<FStreamLevelAction>(LatentInfo.CallbackTarget, LatentInfo.UUID) == NULL)
+		{
+			FStreamLevelAction* NewAction = new FStreamLevelAction(true, LevelName, bMakeVisibleAfterLoad, bShouldBlockOnLoad, LatentInfo, World);
+			LatentManager.AddNewAction(LatentInfo.CallbackTarget, LatentInfo.UUID, NewAction);
+		}
 	}
 }
 
 void UGameplayStatics::UnloadStreamLevel(UObject* WorldContextObject, FName LevelName,FLatentActionInfo LatentInfo)
 {
 	UWorld* World = GEngine->GetWorldFromContextObject(WorldContextObject);
-	FLatentActionManager& LatentManager = World->GetLatentActionManager();
-	if (LatentManager.FindExistingAction<FStreamLevelAction>(LatentInfo.CallbackTarget, LatentInfo.UUID) == NULL)
+	if (World != nullptr)
 	{
-		FStreamLevelAction* NewAction = new FStreamLevelAction(false, LevelName, false, false, LatentInfo, World );
-		LatentManager.AddNewAction(LatentInfo.CallbackTarget, LatentInfo.UUID, NewAction );
+		FLatentActionManager& LatentManager = World->GetLatentActionManager();
+		if (LatentManager.FindExistingAction<FStreamLevelAction>(LatentInfo.CallbackTarget, LatentInfo.UUID) == NULL)
+		{
+			FStreamLevelAction* NewAction = new FStreamLevelAction(false, LevelName, false, false, LatentInfo, World );
+			LatentManager.AddNewAction(LatentInfo.CallbackTarget, LatentInfo.UUID, NewAction );
+		}
 	}
 }
 
@@ -366,16 +385,19 @@ ULevelStreaming* UGameplayStatics::GetStreamingLevel(UObject* WorldContextObject
 	if (InPackageName != NAME_None)
 	{
 		UWorld* World = GEngine->GetWorldFromContextObject(WorldContextObject);
-		const FString SearchPackageName = FStreamLevelAction::MakeSafeLevelName(InPackageName, World);
-
-		for (ULevelStreaming* LevelStreaming : World->StreamingLevels)
+		if (World != nullptr)
 		{
-			// We check only suffix of package name, to handle situations when packages were saved for play into a temporary folder
-			// Like Saved/Autosaves/PackageName
-			if (LevelStreaming && 
-				LevelStreaming->GetWorldAssetPackageName().EndsWith(SearchPackageName, ESearchCase::IgnoreCase))
+			const FString SearchPackageName = FStreamLevelAction::MakeSafeLevelName(InPackageName, World);
+
+			for (ULevelStreaming* LevelStreaming : World->StreamingLevels)
 			{
-				return LevelStreaming;
+				// We check only suffix of package name, to handle situations when packages were saved for play into a temporary folder
+				// Like Saved/Autosaves/PackageName
+				if (LevelStreaming && 
+					LevelStreaming->GetWorldAssetPackageName().EndsWith(SearchPackageName, ESearchCase::IgnoreCase))
+				{
+					return LevelStreaming;
+				}
 			}
 		}
 	}
@@ -466,7 +488,7 @@ void UGameplayStatics::GetAllActorsOfClass(UObject* WorldContextObject, TSubclas
 	UWorld* World = GEngine->GetWorldFromContextObject(WorldContextObject);
 
 	// We do nothing if not class provided, rather than giving ALL actors!
-	if(ActorClass != NULL)
+	if(ActorClass != NULL && World != nullptr)
 	{
 		for(TActorIterator<AActor> It(World, ActorClass); It; ++It)
 		{
@@ -485,7 +507,7 @@ void UGameplayStatics::GetAllActorsWithInterface(UObject* WorldContextObject, TS
 
 	UWorld* World = GEngine->GetWorldFromContextObject(WorldContextObject);
 	// We do nothing if not class provided, rather than giving ALL actors!
-	if(Interface != NULL)
+	if(Interface != NULL && World != nullptr)
 	{
 		for(FActorIterator It(World); It; ++It)
 		{
@@ -502,7 +524,10 @@ void UGameplayStatics::GetAllActorsWithInterface(UObject* WorldContextObject, TS
 void UGameplayStatics::PlayWorldCameraShake(UObject* WorldContextObject, TSubclassOf<class UCameraShake> Shake, FVector Epicenter, float InnerRadius, float OuterRadius, float Falloff, bool bOrientShakeTowardsEpicenter)
 {
 	UWorld* World = GEngine->GetWorldFromContextObject(WorldContextObject);
-	APlayerCameraManager::PlayWorldCameraShake(World, Shake, Epicenter, InnerRadius, OuterRadius, Falloff, bOrientShakeTowardsEpicenter);
+	if(World != nullptr)
+	{
+		APlayerCameraManager::PlayWorldCameraShake(World, Shake, Epicenter, InnerRadius, OuterRadius, Falloff, bOrientShakeTowardsEpicenter);
+	}
 }
 
 UParticleSystemComponent* CreateParticleSystem(UParticleSystem* EmitterTemplate, UWorld* World, AActor* Actor, bool bAutoDestroy)
@@ -525,13 +550,15 @@ UParticleSystemComponent* UGameplayStatics::SpawnEmitterAtLocation(UObject* Worl
 	if (EmitterTemplate)
 	{
 		UWorld* World = GEngine->GetWorldFromContextObject(WorldContextObject);
+		if(World != nullptr)
+		{
+			PSC = CreateParticleSystem(EmitterTemplate, World, NULL, bAutoDestroy);
 
-		PSC = CreateParticleSystem(EmitterTemplate, World, NULL, bAutoDestroy);
-
-		PSC->SetAbsolute(true, true, true);
-		PSC->SetWorldLocationAndRotation(SpawnLocation, SpawnRotation);
-		PSC->SetRelativeScale3D(FVector(1.f));
-		PSC->ActivateSystem(true);
+			PSC->SetAbsolute(true, true, true);
+			PSC->SetWorldLocationAndRotation(SpawnLocation, SpawnRotation);
+			PSC->SetRelativeScale3D(FVector(1.f));
+			PSC->ActivateSystem(true);
+		}
 	}
 	return PSC;
 }
@@ -604,6 +631,10 @@ void UGameplayStatics::PlaySoundAtLocation(UObject* WorldContextObject, class US
 	}
 
 	UWorld* ThisWorld = GEngine->GetWorldFromContextObject(WorldContextObject);
+	if(ThisWorld == nullptr)
+	{
+		return;
+	}
 
 	const bool bIsInGameWorld = ThisWorld->IsGameWorld();
 
@@ -660,7 +691,7 @@ void UGameplayStatics::PlayDialogueAtLocation(UObject* WorldContextObject, class
 	}
 
 	UWorld* ThisWorld = GEngine->GetWorldFromContextObject(WorldContextObject);
-	if (GEngine && GEngine->UseSound() && ThisWorld->bAllowAudioPlayback)
+	if (GEngine && GEngine->UseSound() && (ThisWorld != nullptr) && ThisWorld->bAllowAudioPlayback)
 	{
 		const bool bIsInGameWorld = ThisWorld->IsGameWorld();
 
@@ -889,8 +920,11 @@ UDecalComponent* UGameplayStatics::SpawnDecalAtLocation(UObject* WorldContextObj
 	if (DecalMaterial)
 	{
 		UWorld* World = GEngine->GetWorldFromContextObject(WorldContextObject);
-		DecalComp = CreateDecalComponent(DecalMaterial, DecalSize, World, NULL, LifeSpan);
-		DecalComp->SetWorldLocationAndRotation(Location, Rotation);
+		if(World != nullptr)
+		{
+			DecalComp = CreateDecalComponent(DecalMaterial, DecalSize, World, NULL, LifeSpan);
+			DecalComp->SetWorldLocationAndRotation(Location, Rotation);
+		}
 	}
 
 	return DecalComp;
@@ -1059,17 +1093,20 @@ USaveGame* UGameplayStatics::LoadGameFromSlot(const FString& SlotName, const int
 
 float UGameplayStatics::GetWorldDeltaSeconds(UObject* WorldContextObject)
 {
-	return GEngine->GetWorldFromContextObject(WorldContextObject)->GetDeltaSeconds();
+	UWorld* World = GEngine->GetWorldFromContextObject(WorldContextObject);
+	return World ? World->GetDeltaSeconds() : 0.f;
 }
 
 float UGameplayStatics::GetRealTimeSeconds(UObject* WorldContextObject)
 {
-	return GEngine->GetWorldFromContextObject(WorldContextObject)->GetRealTimeSeconds();
+	UWorld* World = GEngine->GetWorldFromContextObject(WorldContextObject);
+	return World ? World->GetRealTimeSeconds() : 0.f;
 }
 
 float UGameplayStatics::GetAudioTimeSeconds(UObject* WorldContextObject)
 {
-	return GEngine->GetWorldFromContextObject(WorldContextObject)->GetAudioTimeSeconds();
+	UWorld* World = GEngine->GetWorldFromContextObject(WorldContextObject);
+	return World ? World->GetAudioTimeSeconds() : 0.f;
 }
 
 void UGameplayStatics::GetAccurateRealTime(UObject* WorldContextObject, int32& Seconds, float& PartialSeconds)
@@ -1275,11 +1312,11 @@ bool UGameplayStatics::SuggestProjectileVelocity(UObject* WorldContextObject, FV
 FIntVector UGameplayStatics::GetWorldOriginLocation(UObject* WorldContextObject)
 {
 	UWorld* World = GEngine->GetWorldFromContextObject(WorldContextObject);
-	return World->OriginLocation;
+	return World ? World->OriginLocation : FIntVector::ZeroValue;
 }
 
 void UGameplayStatics::SetWorldOriginLocation(UObject* WorldContextObject, FIntVector NewLocation)
 {
 	UWorld* World = GEngine->GetWorldFromContextObject(WorldContextObject);
-	return World->RequestNewWorldOrigin(NewLocation);
+	return World ? World->RequestNewWorldOrigin(NewLocation) : false;
 }
\ No newline at end of file
diff --git a/Engine/Source/Runtime/Engine/Private/KismetMathLibrary.cpp b/Engine/Source/Runtime/Engine/Private/KismetMathLibrary.cpp
--- a/Engine/Source/Runtime/Engine/Private/KismetMathLibrary.cpp
+++ b/Engine/Source/Runtime/Engine/Private/KismetMathLibrary.cpp
@@ -1951,10 +1951,14 @@ void UKismetMathLibrary::MinimumAreaRectangle(class UObject* WorldContextObject,
 
 	if( bDebugDraw )
 	{
-		DrawDebugSphere(GEngine->GetWorldFromContextObject(WorldContextObject), OutRectCenter, 10.f, 12, FColor::Yellow, true);
-		DrawDebugCoordinateSystem(GEngine->GetWorldFromContextObject(WorldContextObject), OutRectCenter, SurfaceNormalMatrix.Rotator(), 100.f, true);
-		DrawDebugLine(GEngine->GetWorldFromContextObject(WorldContextObject), OutRectCenter - RectSideA * 0.5f + FVector(0,0,10.f), OutRectCenter + RectSideA * 0.5f + FVector(0,0,10.f), FColor::Green, true,-1, 0, 5.f);
-		DrawDebugLine(GEngine->GetWorldFromContextObject(WorldContextObject), OutRectCenter - RectSideB * 0.5f + FVector(0,0,10.f), OutRectCenter + RectSideB * 0.5f + FVector(0,0,10.f), FColor::Blue, true,-1, 0, 5.f);
+		UWorld* World = GEngine->GetWorldFromContextObject(WorldContextObject);
+		if(World != nullptr)
+		{
+			DrawDebugSphere(World, OutRectCenter, 10.f, 12, FColor::Yellow, true);
+			DrawDebugCoordinateSystem(World, OutRectCenter, SurfaceNormalMatrix.Rotator(), 100.f, true);
+			DrawDebugLine(World, OutRectCenter - RectSideA * 0.5f + FVector(0,0,10.f), OutRectCenter + RectSideA * 0.5f + FVector(0,0,10.f), FColor::Green, true,-1, 0, 5.f);
+			DrawDebugLine(World, OutRectCenter - RectSideB * 0.5f + FVector(0,0,10.f), OutRectCenter + RectSideB * 0.5f + FVector(0,0,10.f), FColor::Blue, true,-1, 0, 5.f);
+		}
 	}
 }
 
diff --git a/Engine/Source/Runtime/Engine/Private/KismetSystemLibrary.cpp b/Engine/Source/Runtime/Engine/Private/KismetSystemLibrary.cpp
--- a/Engine/Source/Runtime/Engine/Private/KismetSystemLibrary.cpp
+++ b/Engine/Source/Runtime/Engine/Private/KismetSystemLibrary.cpp
@@ -88,13 +88,13 @@ bool UKismetSystemLibrary::DoesImplementInterface(UObject* TestObject, TSubclass
 float UKismetSystemLibrary::GetGameTimeInSeconds(UObject* WorldContextObject)
 {
 	UWorld* World = GEngine->GetWorldFromContextObject( WorldContextObject );
-	return World->GetTimeSeconds();
+	return World ? World->GetTimeSeconds() : 0.f;
 }
 
 bool UKismetSystemLibrary::IsServer(UObject* WorldContextObject)
 {
 	UWorld* World = GEngine->GetWorldFromContextObject( WorldContextObject );
-	return (World->GetNetMode() != NM_Client);
+	return World ? (World->GetNetMode() != NM_Client) : false;
 }
 
 bool UKismetSystemLibrary::IsDedicatedServer(UObject* WorldContextObject)
@@ -291,9 +291,12 @@ void UKismetSystemLibrary::K2_SetTimerDelegate(FTimerDynamicDelegate Delegate, f
 	if (Delegate.IsBound())
 	{
 		const UWorld* const World = GEngine->GetWorldFromContextObject(Delegate.GetUObject());
-		auto& TimerManager = World->GetTimerManager();
-		auto Handle = TimerManager.K2_FindDynamicTimerHandle(Delegate);
-		TimerManager.SetTimer(Handle, Delegate, Time, bLooping);
+		if(World)
+		{
+			auto& TimerManager = World->GetTimerManager();
+			auto Handle = TimerManager.K2_FindDynamicTimerHandle(Delegate);
+			TimerManager.SetTimer(Handle, Delegate, Time, bLooping);
+		}
 	}
 	else
 	{
@@ -344,9 +347,12 @@ void UKismetSystemLibrary::K2_PauseTimerDelegate(FTimerDynamicDelegate Delegate)
 	if (Delegate.IsBound())
 	{
 		UWorld* World = GEngine->GetWorldFromContextObject(Delegate.GetUObject());
-		auto& TimerManager = World->GetTimerManager();
-		auto Handle = TimerManager.K2_FindDynamicTimerHandle(Delegate);
-		World->GetTimerManager().PauseTimer(Handle);
+		if(World)
+		{
+			auto& TimerManager = World->GetTimerManager();
+			auto Handle = TimerManager.K2_FindDynamicTimerHandle(Delegate);
+			World->GetTimerManager().PauseTimer(Handle);
+		}
 	}
 	else
 	{
@@ -369,9 +375,12 @@ void UKismetSystemLibrary::K2_UnPauseTimerDelegate(FTimerDynamicDelegate Delegat
 	if (Delegate.IsBound())
 	{
 		UWorld* World = GEngine->GetWorldFromContextObject(Delegate.GetUObject());
-		auto& TimerManager = World->GetTimerManager();
-		auto Handle = TimerManager.K2_FindDynamicTimerHandle(Delegate);
-		World->GetTimerManager().UnPauseTimer(Handle);
+		if(World)
+		{
+			auto& TimerManager = World->GetTimerManager();
+			auto Handle = TimerManager.K2_FindDynamicTimerHandle(Delegate);
+			World->GetTimerManager().UnPauseTimer(Handle);
+		}
 	}
 	else
 	{
@@ -394,9 +403,16 @@ bool UKismetSystemLibrary::K2_IsTimerActiveDelegate(FTimerDynamicDelegate Delega
 	if (Delegate.IsBound())
 	{
 		UWorld* World = GEngine->GetWorldFromContextObject(Delegate.GetUObject());
-		auto& TimerManager = World->GetTimerManager();
-		auto Handle = TimerManager.K2_FindDynamicTimerHandle(Delegate);
-		return World->GetTimerManager().IsTimerActive(Handle);
+		if(World)
+		{
+			auto& TimerManager = World->GetTimerManager();
+			auto Handle = TimerManager.K2_FindDynamicTimerHandle(Delegate);
+			return World->GetTimerManager().IsTimerActive(Handle);
+		}
+		else
+		{
+			return false;
+		}
 	}
 	else
 	{
@@ -420,9 +436,16 @@ bool UKismetSystemLibrary::K2_IsTimerPausedDelegate(FTimerDynamicDelegate Delega
 	if (Delegate.IsBound())
 	{
 		UWorld* World = GEngine->GetWorldFromContextObject(Delegate.GetUObject());
-		auto& TimerManager = World->GetTimerManager();
-		auto Handle = TimerManager.K2_FindDynamicTimerHandle(Delegate);
-		return World->GetTimerManager().IsTimerPaused(Handle);
+		if(World)
+		{
+			auto& TimerManager = World->GetTimerManager();
+			auto Handle = TimerManager.K2_FindDynamicTimerHandle(Delegate);
+			return World->GetTimerManager().IsTimerPaused(Handle);
+		}
+		else
+		{
+			return false;
+		}
 	}
 	else
 	{
@@ -446,9 +469,16 @@ bool UKismetSystemLibrary::K2_TimerExistsDelegate(FTimerDynamicDelegate Delegate
 	if (Delegate.IsBound())
 	{
 		UWorld* World = GEngine->GetWorldFromContextObject(Delegate.GetUObject());
-		auto& TimerManager = World->GetTimerManager();
-		auto Handle = TimerManager.K2_FindDynamicTimerHandle(Delegate);
-		return World->GetTimerManager().TimerExists(Handle);
+		if(World)
+		{
+			auto& TimerManager = World->GetTimerManager();
+			auto Handle = TimerManager.K2_FindDynamicTimerHandle(Delegate);
+			return World->GetTimerManager().TimerExists(Handle);
+		}
+		else
+		{
+			return false;
+		}
 	}
 	else
 	{
@@ -472,9 +502,16 @@ float UKismetSystemLibrary::K2_GetTimerElapsedTimeDelegate(FTimerDynamicDelegate
 	if (Delegate.IsBound())
 	{
 		UWorld* World = GEngine->GetWorldFromContextObject(Delegate.GetUObject());
-		auto& TimerManager = World->GetTimerManager();
-		auto Handle = TimerManager.K2_FindDynamicTimerHandle(Delegate);
-		return World->GetTimerManager().GetTimerElapsed(Handle);
+		if(World)
+		{
+			auto& TimerManager = World->GetTimerManager();
+			auto Handle = TimerManager.K2_FindDynamicTimerHandle(Delegate);
+			return World->GetTimerManager().GetTimerElapsed(Handle);
+		}
+		else
+		{
+			return false;
+		}
 	}
 	else
 	{
@@ -498,9 +535,16 @@ float UKismetSystemLibrary::K2_GetTimerRemainingTimeDelegate(FTimerDynamicDelega
 	if (Delegate.IsBound())
 	{
 		UWorld* World = GEngine->GetWorldFromContextObject(Delegate.GetUObject());
-		auto& TimerManager = World->GetTimerManager();
-		auto Handle = TimerManager.K2_FindDynamicTimerHandle(Delegate);
-		return World->GetTimerManager().GetTimerRemaining(Handle);
+		if(World)
+		{
+			auto& TimerManager = World->GetTimerManager();
+			auto Handle = TimerManager.K2_FindDynamicTimerHandle(Delegate);
+			return World->GetTimerManager().GetTimerRemaining(Handle);
+		}
+		else
+		{
+			return false;
+		}
 	}
 	else
 	{
@@ -926,7 +970,10 @@ bool UKismetSystemLibrary::SphereOverlapComponents_NEW(UObject* WorldContextObje
 
 
 	UWorld* World = GEngine->GetWorldFromContextObject( WorldContextObject );
-	World->OverlapMulti(Overlaps, SpherePos, FQuat::Identity, FCollisionShape::MakeSphere(SphereRadius), Params, ObjectParams);
+	if(World != nullptr)
+	{
+		World->OverlapMulti(Overlaps, SpherePos, FQuat::Identity, FCollisionShape::MakeSphere(SphereRadius), Params, ObjectParams);
+	}
 
 	for (int32 OverlapIdx=0; OverlapIdx<Overlaps.Num(); ++OverlapIdx)
 	{
@@ -978,7 +1025,10 @@ bool UKismetSystemLibrary::BoxOverlapComponents_NEW(UObject* WorldContextObject,
 	}
 
 	UWorld* World = GEngine->GetWorldFromContextObject( WorldContextObject );
-	World->OverlapMulti(Overlaps, BoxPos, FQuat::Identity, FCollisionShape::MakeBox(BoxExtent), Params, ObjectParams);
+	if (World != nullptr)
+	{
+		World->OverlapMulti(Overlaps, BoxPos, FQuat::Identity, FCollisionShape::MakeBox(BoxExtent), Params, ObjectParams);
+	}
 
 	for (int32 OverlapIdx=0; OverlapIdx<Overlaps.Num(); ++OverlapIdx)
 	{
@@ -1029,7 +1079,10 @@ bool UKismetSystemLibrary::CapsuleOverlapComponents_NEW(UObject* WorldContextObj
 	}
 
 	UWorld* World = GEngine->GetWorldFromContextObject( WorldContextObject );	
-	World->OverlapMulti(Overlaps, CapsulePos, FQuat::Identity, FCollisionShape::MakeCapsule(Radius, HalfHeight), Params, ObjectParams);
+	if (World != nullptr)
+	{
+		World->OverlapMulti(Overlaps, CapsulePos, FQuat::Identity, FCollisionShape::MakeCapsule(Radius, HalfHeight), Params, ObjectParams);
+	}
 
 	for (int32 OverlapIdx=0; OverlapIdx<Overlaps.Num(); ++OverlapIdx)
 	{
@@ -1303,9 +1356,9 @@ bool UKismetSystemLibrary::BoxTraceSingle(UObject* WorldContextObject, const FVe
 	}
 
 	UWorld* World = GEngine->GetWorldFromContextObject(WorldContextObject);
-	bool const bHit = World->SweepSingle(OutHit, Start, End, Orientation.Quaternion(), UEngineTypes::ConvertToCollisionChannel(TraceChannel), FCollisionShape::MakeBox(HalfSize), Params);
+	bool const bHit = World ? World->SweepSingle(OutHit, Start, End, Orientation.Quaternion(), UEngineTypes::ConvertToCollisionChannel(TraceChannel), FCollisionShape::MakeBox(HalfSize), Params) : false;
 
-	if (DrawDebugType != EDrawDebugTrace::None)
+	if (DrawDebugType != EDrawDebugTrace::None && (World != nullptr))
 	{
 		bool bPersistent = DrawDebugType == EDrawDebugTrace::Persistent;
 		float LifeTime = (DrawDebugType == EDrawDebugTrace::ForDuration) ? KISMET_TRACE_DEBUG_DRAW_DURATION : 0.f;
@@ -1358,9 +1411,9 @@ bool UKismetSystemLibrary::BoxTraceMulti(UObject* WorldContextObject, const FVec
 	}
 
 	UWorld* World = GEngine->GetWorldFromContextObject(WorldContextObject);
-	bool const bHit = World->SweepMulti(OutHits, Start, End, Orientation.Quaternion(), UEngineTypes::ConvertToCollisionChannel(TraceChannel), FCollisionShape::MakeBox(HalfSize), Params);
+	bool const bHit = World ? World->SweepMulti(OutHits, Start, End, Orientation.Quaternion(), UEngineTypes::ConvertToCollisionChannel(TraceChannel), FCollisionShape::MakeBox(HalfSize), Params) : false;
 
-	if (DrawDebugType != EDrawDebugTrace::None)
+	if (DrawDebugType != EDrawDebugTrace::None && (World != nullptr))
 	{
 		bool bPersistent = DrawDebugType == EDrawDebugTrace::Persistent;
 		float LifeTime = (DrawDebugType == EDrawDebugTrace::ForDuration) ? KISMET_TRACE_DEBUG_DRAW_DURATION : 0.f;
@@ -2104,9 +2157,9 @@ bool UKismetSystemLibrary::BoxTraceSingleForObjects(UObject* WorldContextObject,
 	}
 
 	UWorld* World = GEngine->GetWorldFromContextObject(WorldContextObject);
-	bool const bHit = World->SweepSingle(OutHit, Start, End, Orientation.Quaternion(), FCollisionShape::MakeBox(HalfSize), Params, ObjectParams);
+	bool const bHit = World ? World->SweepSingle(OutHit, Start, End, Orientation.Quaternion(), FCollisionShape::MakeBox(HalfSize), Params, ObjectParams) : false;
 
-	if (DrawDebugType != EDrawDebugTrace::None)
+	if (DrawDebugType != EDrawDebugTrace::None && (World != nullptr))
 	{
 		bool bPersistent = DrawDebugType == EDrawDebugTrace::Persistent;
 		float LifeTime = (DrawDebugType == EDrawDebugTrace::ForDuration) ? KISMET_TRACE_DEBUG_DRAW_DURATION : 0.f;
@@ -2187,9 +2240,9 @@ bool UKismetSystemLibrary::BoxTraceMultiForObjects(UObject* WorldContextObject,
 	}
 
 	UWorld* World = GEngine->GetWorldFromContextObject(WorldContextObject);
-	bool const bHit = World->SweepMulti(OutHits, Start, End, Orientation.Quaternion(), FCollisionShape::MakeBox(HalfSize), Params, ObjectParams);
+	bool const bHit = World ? World->SweepMulti(OutHits, Start, End, Orientation.Quaternion(), FCollisionShape::MakeBox(HalfSize), Params, ObjectParams) : false;
 
-	if (DrawDebugType != EDrawDebugTrace::None)
+	if (DrawDebugType != EDrawDebugTrace::None && (World != nullptr))
 	{
 		bool bPersistent = DrawDebugType == EDrawDebugTrace::Persistent;
 		float LifeTime = (DrawDebugType == EDrawDebugTrace::ForDuration) ? KISMET_TRACE_DEBUG_DRAW_DURATION : 0.f;
@@ -2419,103 +2472,157 @@ bool UKismetSystemLibrary::CapsuleTraceMultiByObject_DEPRECATED(UObject* WorldCo
 /** Draw a debug line */
 void UKismetSystemLibrary::DrawDebugLine(UObject* WorldContextObject, FVector const LineStart, FVector const LineEnd, FLinearColor Color, float LifeTime, float Thickness)
 {
-	::DrawDebugLine(GEngine->GetWorldFromContextObject(WorldContextObject), LineStart, LineEnd, Color, false, LifeTime, SDPG_World, Thickness);
+	UWorld* World = GEngine->GetWorldFromContextObject(WorldContextObject);
+	if(World != nullptr)
+	{
+		::DrawDebugLine(World, LineStart, LineEnd, Color, false, LifeTime, SDPG_World, Thickness);
+	}
 }
 
 /** Draw a debug point */
 void UKismetSystemLibrary::DrawDebugPoint(UObject* WorldContextObject, FVector const Position, float Size, FLinearColor PointColor, float LifeTime)
 {
-	::DrawDebugPoint(GEngine->GetWorldFromContextObject(WorldContextObject), Position, Size, PointColor, false, LifeTime, SDPG_World);
+	UWorld* World = GEngine->GetWorldFromContextObject(WorldContextObject);
+	if (World != nullptr)
+	{
+		::DrawDebugPoint(World, Position, Size, PointColor, false, LifeTime, SDPG_World);
+	}
 }
 
 /** Draw directional arrow, pointing from LineStart to LineEnd. */
 void UKismetSystemLibrary::DrawDebugArrow(UObject* WorldContextObject, FVector const LineStart, FVector const LineEnd, float ArrowSize, FLinearColor Color, float LifeTime)
 {
-	::DrawDebugDirectionalArrow(GEngine->GetWorldFromContextObject(WorldContextObject), LineStart, LineEnd, ArrowSize, Color, false, LifeTime, SDPG_World);
+	UWorld* World = GEngine->GetWorldFromContextObject(WorldContextObject);
+	if (World != nullptr)
+	{
+		::DrawDebugDirectionalArrow(World, LineStart, LineEnd, ArrowSize, Color, false, LifeTime, SDPG_World);
+	}
 }
 
 /** Draw a debug box */
 void UKismetSystemLibrary::DrawDebugBox(UObject* WorldContextObject, FVector const Center, FVector Extent, FLinearColor Color, const FRotator Rotation, float LifeTime)
 {
 	UWorld* World = GEngine->GetWorldFromContextObject(WorldContextObject);
-	if (Rotation == FRotator::ZeroRotator)
-	{
-		::DrawDebugBox(World, Center, Extent, Color.ToFColor(true), false, LifeTime, SDPG_World);
-	}
-	else
+	if(World != nullptr)
 	{
-		::DrawDebugBox(World, Center, Extent, Rotation.Quaternion(), Color.ToFColor(true), false, LifeTime, SDPG_World);
+		if (Rotation == FRotator::ZeroRotator)
+		{
+			::DrawDebugBox(World, Center, Extent, Color.ToFColor(true), false, LifeTime, SDPG_World);
+		}
+		else
+		{
+			::DrawDebugBox(World, Center, Extent, Rotation.Quaternion(), Color.ToFColor(true), false, LifeTime, SDPG_World);
+		}
 	}
 }
 
 /** Draw a debug coordinate system. */
 void UKismetSystemLibrary::DrawDebugCoordinateSystem(UObject* WorldContextObject, FVector const AxisLoc, FRotator const AxisRot, float Scale, float LifeTime)
 {
-	::DrawDebugCoordinateSystem(GEngine->GetWorldFromContextObject(WorldContextObject), AxisLoc, AxisRot, Scale, false, LifeTime, SDPG_World);
+	UWorld* World = GEngine->GetWorldFromContextObject(WorldContextObject);
+	if (World != nullptr)
+	{
+		::DrawDebugCoordinateSystem(World, AxisLoc, AxisRot, Scale, false, LifeTime, SDPG_World);
+	}
 }
 
 /** Draw a debug sphere */
 void UKismetSystemLibrary::DrawDebugSphere(UObject* WorldContextObject, FVector const Center, float Radius, int32 Segments, FLinearColor Color, float LifeTime)
 {
-	::DrawDebugSphere(GEngine->GetWorldFromContextObject(WorldContextObject), Center, Radius, Segments, Color.ToFColor(true), false, LifeTime, SDPG_World);
+	UWorld* World = GEngine->GetWorldFromContextObject(WorldContextObject);
+	if (World != nullptr)
+	{
+		::DrawDebugSphere(World, Center, Radius, Segments, Color.ToFColor(true), false, LifeTime, SDPG_World);
+	}
 }
 
 /** Draw a debug cylinder */
 void UKismetSystemLibrary::DrawDebugCylinder(UObject* WorldContextObject, FVector const Start, FVector const End, float Radius, int32 Segments, FLinearColor Color, float LifeTime)
 {
-	::DrawDebugCylinder(GEngine->GetWorldFromContextObject(WorldContextObject), Start, End, Radius, Segments, Color.ToFColor(true), false, LifeTime, SDPG_World);
+	UWorld* World = GEngine->GetWorldFromContextObject(WorldContextObject);
+	if (World != nullptr)
+	{
+		::DrawDebugCylinder(World, Start, End, Radius, Segments, Color.ToFColor(true), false, LifeTime, SDPG_World);
+	}
 }
 
 /** Draw a debug cone */
 void UKismetSystemLibrary::DrawDebugCone(UObject* WorldContextObject, FVector const Origin, FVector const Direction, float Length, float AngleWidth, float AngleHeight, int32 NumSides, FLinearColor Color)
 {
-	::DrawDebugCone(GEngine->GetWorldFromContextObject(WorldContextObject), Origin, Direction, Length, AngleWidth, AngleHeight, NumSides, Color.ToFColor(true), false, -1.f, SDPG_World);
+	UWorld* World = GEngine->GetWorldFromContextObject(WorldContextObject);
+	if (World != nullptr)
+	{
+		::DrawDebugCone(World, Origin, Direction, Length, AngleWidth, AngleHeight, NumSides, Color.ToFColor(true), false, -1.f, SDPG_World);
+	}
 }
 
 void UKismetSystemLibrary::DrawDebugConeInDegrees(UObject* WorldContextObject, FVector const Origin, FVector const Direction, float Length, float AngleWidth, float AngleHeight, int32 NumSides, FLinearColor Color, float LifeTime)
 {
-	::DrawDebugCone(GEngine->GetWorldFromContextObject(WorldContextObject), Origin, Direction, Length, FMath::DegreesToRadians(AngleWidth), FMath::DegreesToRadians(AngleHeight), NumSides, Color.ToFColor(true), false, LifeTime, SDPG_World);
+	UWorld* World = GEngine->GetWorldFromContextObject(WorldContextObject);
+	if (World != nullptr)
+	{
+		::DrawDebugCone(World, Origin, Direction, Length, FMath::DegreesToRadians(AngleWidth), FMath::DegreesToRadians(AngleHeight), NumSides, Color.ToFColor(true), false, LifeTime, SDPG_World);
+	}
 }
 
 /** Draw a debug capsule */
 void UKismetSystemLibrary::DrawDebugCapsule(UObject* WorldContextObject, FVector const Center, float HalfHeight, float Radius, const FRotator Rotation, FLinearColor Color, float LifeTime)
 {
-	::DrawDebugCapsule(GEngine->GetWorldFromContextObject(WorldContextObject), Center, HalfHeight, Radius, Rotation.Quaternion(), Color.ToFColor(true), false, LifeTime, SDPG_World);
+	UWorld* World = GEngine->GetWorldFromContextObject(WorldContextObject);
+	if (World != nullptr)
+	{
+		::DrawDebugCapsule(World, Center, HalfHeight, Radius, Rotation.Quaternion(), Color.ToFColor(true), false, LifeTime, SDPG_World);
+	}
 }
 
 /** Draw a debug string at a 3d world location. */
 void UKismetSystemLibrary::DrawDebugString(UObject* WorldContextObject, FVector const TextLocation, const FString& Text, class AActor* TestBaseActor, FLinearColor TextColor, float Duration)
 {
-	::DrawDebugString(GEngine->GetWorldFromContextObject( WorldContextObject ), TextLocation, Text, TestBaseActor, TextColor.ToFColor(true), Duration);
+	UWorld* World = GEngine->GetWorldFromContextObject(WorldContextObject);
+	if (World != nullptr)
+	{
+		::DrawDebugString(World, TextLocation, Text, TestBaseActor, TextColor.ToFColor(true), Duration);
+	}
 }
 
 /** Removes all debug strings. */
 void UKismetSystemLibrary::FlushDebugStrings( UObject* WorldContextObject )
 {
 	UWorld* World = GEngine->GetWorldFromContextObject( WorldContextObject );
-	::FlushDebugStrings( World );
+	if(World != nullptr)
+	{
+		::FlushDebugStrings( World );
+	}
 }
 
 /** Draws a debug plane. */
 void UKismetSystemLibrary::DrawDebugPlane(UObject* WorldContextObject, FPlane const& P, FVector const Loc, float Size, FLinearColor Color, float LifeTime)
 {
-	::DrawDebugSolidPlane(GEngine->GetWorldFromContextObject( WorldContextObject ), P, Loc, Size, Color.ToFColor(true), false, LifeTime, SDPG_World);
+	UWorld* World = GEngine->GetWorldFromContextObject(WorldContextObject);
+	if (World != nullptr)
+	{
+		::DrawDebugSolidPlane(World, P, Loc, Size, Color.ToFColor(true), false, LifeTime, SDPG_World);
+	}
 }
 
 /** Flush all persistent debug lines and shapes */
 void UKismetSystemLibrary::FlushPersistentDebugLines(UObject* WorldContextObject)
 {
 	UWorld* World = GEngine->GetWorldFromContextObject( WorldContextObject );
-	::FlushPersistentDebugLines( World );
+	if(World != nullptr)
+	{
+		::FlushPersistentDebugLines( World );
+	}
 }
 
 /** Draws a debug frustum. */
 void UKismetSystemLibrary::DrawDebugFrustum(UObject* WorldContextObject, const FTransform& FrustumTransform, FLinearColor FrustumColor, float Duration)
 {
-	if( FrustumTransform.IsRotationNormalized() )
+	UWorld* World = GEngine->GetWorldFromContextObject(WorldContextObject);
+	if( World != nullptr && FrustumTransform.IsRotationNormalized() )
 	{
 		FMatrix FrustumToWorld =  FrustumTransform.ToMatrixWithScale();
-		::DrawDebugFrustum(GEngine->GetWorldFromContextObject(WorldContextObject), FrustumToWorld, FrustumColor, false, Duration, SDPG_World);
+		::DrawDebugFrustum(World, FrustumToWorld, FrustumColor, false, Duration, SDPG_World);
 	}
 }
 
diff --git a/Engine/Source/Runtime/Engine/Private/UnrealEngine.cpp b/Engine/Source/Runtime/Engine/Private/UnrealEngine.cpp
--- a/Engine/Source/Runtime/Engine/Private/UnrealEngine.cpp
+++ b/Engine/Source/Runtime/Engine/Private/UnrealEngine.cpp
@@ -10463,7 +10463,7 @@ bool UEngine::ShouldAbsorbAuthorityOnlyEvent()
 			}
 		}
 
-		if (useIt)
+		if (useIt && (Context.World() != nullptr))
 		{
 			return (Context.World()->GetNetMode() ==  NM_Client);
 		}
@@ -10504,7 +10504,7 @@ bool UEngine::ShouldAbsorbCosmeticOnlyEvent()
 			}
 		}
 
-		if (useIt)
+		if (useIt && (Context.World() != nullptr))
 		{
 			return (Context.World()->GetNetMode() == NM_DedicatedServer);
 		}
