UE-7616 - Fix issue where queued RPC bunches can incorrectly cause the exported guids to go out on the wrong bunch

diff --git a/Engine/Source/Runtime/Engine/Classes/Engine/ActorChannel.h b/Engine/Source/Runtime/Engine/Classes/Engine/ActorChannel.h
--- a/Engine/Source/Runtime/Engine/Classes/Engine/ActorChannel.h
+++ b/Engine/Source/Runtime/Engine/Classes/Engine/ActorChannel.h
@@ -75,7 +75,9 @@ class ENGINE_API UActorChannel
 
 	TArray< TWeakObjectPtr< UObject > >	CreateSubObjects;		// Any sub-object we created on this channel
 
-	TArray< FNetworkGUID >				MustBeMappedGuidsInLastBunch;		// Array of guids that will async load on client. This list is used for queued RPC's.
+	TArray< FNetworkGUID >				QueuedMustBeMappedGuidsInLastBunch;		// Array of guids that will async load on client. This list is used for queued RPC's.
+
+	TArray< class FOutBunch * >			QueuedExportBunches;			// Bunches that need to be appended to the export list on the next SendBunch call. This list is used for queued RPC's.
 
 	/**
 	 * Default constructor
@@ -117,6 +119,9 @@ public:
 
 	void SetChannelActorForDestroy( struct FActorDestructionInfo *DestructInfo );
 
+	/** Append any export bunches */
+	virtual void AppendExportBunches( TArray< FOutBunch* >& OutExportBunches ) override;
+
 	/** Append any "must be mapped" guids to front of bunch. These are guids that the client will wait on before processing this bunch. */
 	virtual void AppendMustBeMappedGuids( FOutBunch* Bunch ) override;
 
diff --git a/Engine/Source/Runtime/Engine/Classes/Engine/Channel.h b/Engine/Source/Runtime/Engine/Classes/Engine/Channel.h
--- a/Engine/Source/Runtime/Engine/Classes/Engine/Channel.h
+++ b/Engine/Source/Runtime/Engine/Classes/Engine/Channel.h
@@ -104,6 +104,9 @@ public:
 	 */
 	void ReceivedRawBunch( FInBunch & Bunch, bool & bOutSkipAck );
 	
+	/** Append any export bunches */
+	virtual void AppendExportBunches( TArray< FOutBunch* >& OutExportBunches );
+
 	/** Append any "must be mapped" guids to front of bunch. These are guids that the client will wait on before processing this bunch. */
 	virtual void AppendMustBeMappedGuids( FOutBunch* Bunch );
 
diff --git a/Engine/Source/Runtime/Engine/Private/DataChannel.cpp b/Engine/Source/Runtime/Engine/Private/DataChannel.cpp
--- a/Engine/Source/Runtime/Engine/Private/DataChannel.cpp
+++ b/Engine/Source/Runtime/Engine/Private/DataChannel.cpp
@@ -591,20 +591,12 @@ bool UChannel::ReceivedNextBunch( FInBunch & Bunch, bool & bOutSkipAck )
 	return false;
 }
 
-void UActorChannel::AppendMustBeMappedGuids( FOutBunch* Bunch )
+void UChannel::AppendExportBunches( TArray<FOutBunch *>& OutExportBunches )
 {
-	if ( MustBeMappedGuidsInLastBunch.Num() > 0 )
-	{
-		// Just add our list to the main list on package map so we can re-use the code in UChannel to add them all together
-		UPackageMapClient * PackageMapClient = CastChecked< UPackageMapClient >( Connection->PackageMap );
-
-		PackageMapClient->GetMustBeMappedGuidsInLastBunch().Append( MustBeMappedGuidsInLastBunch );
-
-		MustBeMappedGuidsInLastBunch.Empty();
-	}
+	UPackageMapClient * PackageMapClient = CastChecked< UPackageMapClient >( Connection->PackageMap );
 
-	// Actually add them to the bunch
-	Super::AppendMustBeMappedGuids( Bunch );
+	// Let the package map add any outgoing bunches it needs to send
+	PackageMapClient->AppendExportBunches( OutExportBunches );
 }
 
 void UChannel::AppendMustBeMappedGuids( FOutBunch* Bunch )
@@ -637,6 +629,34 @@ void UChannel::AppendMustBeMappedGuids( FOutBunch* Bunch )
 	}
 }
 
+void UActorChannel::AppendExportBunches( TArray<FOutBunch *>& OutExportBunches )
+{
+	Super::AppendExportBunches( OutExportBunches );
+
+	if ( QueuedExportBunches.Num() )
+	{
+		OutExportBunches.Append( QueuedExportBunches );
+		QueuedExportBunches.Empty();
+	}
+}
+
+void UActorChannel::AppendMustBeMappedGuids( FOutBunch* Bunch )
+{
+	if ( QueuedMustBeMappedGuidsInLastBunch.Num() > 0 )
+	{
+		// Just add our list to the main list on package map so we can re-use the code in UChannel to add them all together
+		UPackageMapClient * PackageMapClient = CastChecked< UPackageMapClient >( Connection->PackageMap );
+
+		PackageMapClient->GetMustBeMappedGuidsInLastBunch().Append( QueuedMustBeMappedGuidsInLastBunch );
+
+		QueuedMustBeMappedGuidsInLastBunch.Empty();
+	}
+
+	// Actually add them to the bunch
+	// NOTE - We do this LAST since we want to capture the append that happened above
+	Super::AppendMustBeMappedGuids( Bunch );
+}
+
 FPacketIdRange UChannel::SendBunch( FOutBunch* Bunch, bool Merge )
 {
 	check(!Closing);
@@ -665,12 +685,13 @@ FPacketIdRange UChannel::SendBunch( FOutBunch* Bunch, bool Merge )
 
 	TArray<FOutBunch *> OutgoingBunches;
 
-	UPackageMapClient * PackageMapClient = CastChecked< UPackageMapClient >( Connection->PackageMap );
+	// Add any export bunches
+	AppendExportBunches( OutgoingBunches );
 
-	// Let the package map add any outgoing bunches it needs to send
-	if ( PackageMapClient->AppendExportBunches( OutgoingBunches ) )
+	if ( OutgoingBunches.Num() )
 	{
-		// Dont merge if we have multiple bunches to send.
+		// Don't merge if we are exporting guid's
+		// We can't be for sure if the last bunch has exported guids as well, so this just simplifies things
 		Merge = false;
 	}
 
@@ -1497,6 +1518,17 @@ bool UActorChannel::CleanUp( const bool bForDestroy )
 	// We don't care about any leftover pending guids at this point
 	PendingGuidResolves.Empty();
 
+	// Free export bunches list
+	for ( int32 i = 0; i < QueuedExportBunches.Num(); i++ )
+	{
+		delete QueuedExportBunches[i];
+	}
+
+	QueuedExportBunches.Empty();
+
+	// Free the must be mapped list
+	QueuedMustBeMappedGuidsInLastBunch.Empty();
+
 	// Free any queued bunches
 	for ( int32 i = 0; i < QueuedBunches.Num(); i++ )
 	{
diff --git a/Engine/Source/Runtime/Engine/Private/DataReplication.cpp b/Engine/Source/Runtime/Engine/Private/DataReplication.cpp
--- a/Engine/Source/Runtime/Engine/Private/DataReplication.cpp
+++ b/Engine/Source/Runtime/Engine/Private/DataReplication.cpp
@@ -1074,9 +1074,12 @@ void FObjectReplicator::QueueRemoteFunctionBunch( UFunction* Func, FOutBunch &Bu
 		// We need to copy over any info that was obtained on the package map during serialization, and remember it until we actually call SendBunch
 		if ( PackageMapClient->GetMustBeMappedGuidsInLastBunch().Num() )
 		{
-			OwningChannel->MustBeMappedGuidsInLastBunch.Append( PackageMapClient->GetMustBeMappedGuidsInLastBunch() );
+			OwningChannel->QueuedMustBeMappedGuidsInLastBunch.Append( PackageMapClient->GetMustBeMappedGuidsInLastBunch() );
 			PackageMapClient->GetMustBeMappedGuidsInLastBunch().Empty();
 		}
+
+		// Copy over any exported bunches
+		PackageMapClient->AppendExportBunches( OwningChannel->QueuedExportBunches );
 	}
 }
 
