Fixed handling serialization of navigation system SupportedAgents' NavigationDataClass refs #UE4

- addressed UE-7509
- removed UNavigationSystem.RequiredNavigationDataClassNames as no longer needed with changes to SupportedAgents
- also, fixes Project Settings for navigation system not using game-specific navigation system class, if one has been declared

diff --git a/Engine/Config/BaseEngine.ini b/Engine/Config/BaseEngine.ini
--- a/Engine/Config/BaseEngine.ini
+++ b/Engine/Config/BaseEngine.ini
@@ -1942,7 +1942,6 @@ EditorViewButtonsObject=EditorCylinder
 [/Script/Engine.NavigationSystem]
 bAutoCreateNavigationData=true
 bAddPlayersToGenerationSeeds=true
-+RequiredNavigationDataClassNames=/Script/Engine.RecastNavMesh
 
 [/Script/Engine.RecastNavMesh]
 ; runtime params
diff --git a/Engine/Source/Editor/ProjectSettingsViewer/Private/ProjectSettingsViewerModule.cpp b/Engine/Source/Editor/ProjectSettingsViewer/Private/ProjectSettingsViewerModule.cpp
--- a/Engine/Source/Editor/ProjectSettingsViewer/Private/ProjectSettingsViewerModule.cpp
+++ b/Engine/Source/Editor/ProjectSettingsViewer/Private/ProjectSettingsViewerModule.cpp
@@ -127,11 +127,14 @@ protected:
 			GetMutableDefault<UInputSettings>()
 		);
 
-		// navigation system
+		// navigation system's class can be game specific so we need to find appropriate CDO
+		UNavigationSystem* NavSysCDO = (*GEngine->NavigationSystemClass != nullptr)
+			? GetMutableDefault<UNavigationSystem>(GEngine->NavigationSystemClass)
+			: GetMutableDefault<UNavigationSystem>();
 		SettingsModule.RegisterSettings("Project", "Engine", "NavigationSystem",
 			LOCTEXT("NavigationSystemSettingsName", "Navigation System"),
 			LOCTEXT("NavigationSystemSettingsDescription", "Settings for the navigation system."),
-			GetMutableDefault<UNavigationSystem>()
+			NavSysCDO
 		);
 
 		// navigation mesh
diff --git a/Engine/Source/Runtime/Engine/Classes/AI/Navigation/NavigationSystem.h b/Engine/Source/Runtime/Engine/Classes/AI/Navigation/NavigationSystem.h
--- a/Engine/Source/Runtime/Engine/Classes/AI/Navigation/NavigationSystem.h
+++ b/Engine/Source/Runtime/Engine/Classes/AI/Navigation/NavigationSystem.h
@@ -202,10 +202,6 @@ public:
 private:
 	TWeakObjectPtr<UCrowdManager> CrowdManager;
 
-	// required navigation data 
-	UPROPERTY(config)
-	TArray<FStringClassReference> RequiredNavigationDataClassNames;
-
 	/** set to true when navigation processing was blocked due to missing nav bounds */
 	uint32 bNavDataRemovedDueToMissingNavBounds:1;
 
@@ -312,6 +308,11 @@ protected:
 	/** spawn new crowd manager */
 	virtual void CreateCrowdManager();
 
+	/** Used to properly set navigation class for indicated agent and propagate information to other places
+	 *	(like project settings) that may need this information 
+	 */
+	void SetSupportedAgentsNavigationClass(int32 AgentIndex, TSubclassOf<ANavigationData> NavigationDataClass);
+
 public:
 	//----------------------------------------------------------------------//
 	// Public querying interface                                                                
diff --git a/Engine/Source/Runtime/Engine/Classes/AI/Navigation/NavigationTypes.h b/Engine/Source/Runtime/Engine/Classes/AI/Navigation/NavigationTypes.h
--- a/Engine/Source/Runtime/Engine/Classes/AI/Navigation/NavigationTypes.h
+++ b/Engine/Source/Runtime/Engine/Classes/AI/Navigation/NavigationTypes.h
@@ -393,10 +393,10 @@ struct ENGINE_API FNavDataConfig : public FNavAgentProperties
 	UPROPERTY(EditAnywhere, BlueprintReadWrite, Category = Display)
 	FColor Color;
 
-	UPROPERTY(EditAnywhere, BlueprintReadOnly, Category = Querying, config)
+	UPROPERTY(EditAnywhere, BlueprintReadOnly, Category = Querying)
 	FVector DefaultQueryExtent;
 
-	UPROPERTY(EditAnywhere, BlueprintReadOnly, Category = Navigation, config)
+	UPROPERTY(Transient, EditAnywhere, BlueprintReadOnly, Category = Navigation)
 	TSubclassOf<ANavigationData> NavigationDataClass;
 
 	UPROPERTY(config)
diff --git a/Engine/Source/Runtime/Engine/Private/AI/Navigation/NavigationSystem.cpp b/Engine/Source/Runtime/Engine/Private/AI/Navigation/NavigationSystem.cpp
--- a/Engine/Source/Runtime/Engine/Private/AI/Navigation/NavigationSystem.cpp
+++ b/Engine/Source/Runtime/Engine/Private/AI/Navigation/NavigationSystem.cpp
@@ -119,8 +119,10 @@ FNavDataConfig::FNavDataConfig(float Radius, float Height)
 	, Color(140, 255, 0, 164)
 	, DefaultQueryExtent(DEFAULT_NAV_QUERY_EXTENT_HORIZONTAL, DEFAULT_NAV_QUERY_EXTENT_HORIZONTAL, DEFAULT_NAV_QUERY_EXTENT_VERTICAL)
 	, NavigationDataClass(ARecastNavMesh::StaticClass())
+	, NavigationDataClassName(NavigationDataClass)
 {
 }
+
 //----------------------------------------------------------------------//
 // FNavigationLockContext                                                                
 //----------------------------------------------------------------------//
@@ -304,31 +306,66 @@ void UNavigationSystem::UpdateAbstractNavData()
 		AbstractNavData->OnRegistered();
 	}
 }
+
+void UNavigationSystem::SetSupportedAgentsNavigationClass(int32 AgentIndex, TSubclassOf<ANavigationData> NavigationDataClass)
+{
+	check(SupportedAgents.IsValidIndex(AgentIndex));
+	SupportedAgents[AgentIndex].NavigationDataClass = NavigationDataClass;
+
+	if (NavigationDataClass != nullptr)
+	{
+		SupportedAgents[AgentIndex].NavigationDataClassName = FStringClassReference::GetOrCreateIDForClass(NavigationDataClass);
+	}
+	else
+	{
+		SupportedAgents[AgentIndex].NavigationDataClassName.Reset();
+	}
+
+#if WITH_EDITOR
+	if (GIsEditor)
+	{
+		if (HasAnyFlags(RF_ClassDefaultObject) == false)
+		{
+			// set it at CDO to properly show up in project settings
+			// @hack the reason for doing it this way is that engine doesn't handle default TSubclassOf properties
+			//	set to game-specific classes;
+			UNavigationSystem* NavigationSystemCDO = GetMutableDefault<UNavigationSystem>(GetClass());
+			NavigationSystemCDO->SetSupportedAgentsNavigationClass(AgentIndex, NavigationDataClass);
+		}
+	}
+#endif // WITH_EDITOR
+}
+
 void UNavigationSystem::PostInitProperties()
 {
 	Super::PostInitProperties();
 
 	if (HasAnyFlags(RF_ClassDefaultObject) == false)
 	{
-		// make sure there's at leas one supported navigation agent size
+		// make sure there's at least one supported navigation agent size
 		if (SupportedAgents.Num() == 0)
 		{
 			SupportedAgents.Add(FNavDataConfig(FNavigationSystem::FallbackAgentRadius, FNavigationSystem::FallbackAgentHeight));
 		}
-
-		// gather navigation creators
-		NavDataClasses.Empty(RequiredNavigationDataClassNames.Num());
-		for (int32 Index = 0; Index < RequiredNavigationDataClassNames.Num(); ++Index)
+		else
 		{
-			TSubclassOf<ANavigationData> NavDataClass = LoadClass<ANavigationData>(NULL, *RequiredNavigationDataClassNames[Index].ToString(), NULL, LOAD_None, NULL);
-			if (NavDataClass)
-			{
-				NavDataClasses.AddUnique(NavDataClass);
-			}
-			else
+			for (int32 AgentIndex = 0; AgentIndex < SupportedAgents.Num(); ++AgentIndex)
 			{
-				UE_LOG(LogNavigation, Warning, TEXT("Unable to find navigation data class \'%s\' while setting up require navigation types")
-					, *RequiredNavigationDataClassNames[Index].ToString());
+				FNavDataConfig& SupportedAgentConfig = SupportedAgents[AgentIndex];
+				// a piece of legacy maintanance 
+				if (SupportedAgentConfig.NavigationDataClass != nullptr && SupportedAgentConfig.NavigationDataClassName.IsValid() == false)
+				{
+					// fill NavigationDataClassName
+					SupportedAgentConfig.NavigationDataClassName = FStringClassReference(SupportedAgentConfig.NavigationDataClass);
+				}
+				else
+				{
+					TSubclassOf<ANavigationData> NavigationDataClass = SupportedAgentConfig.NavigationDataClassName.IsValid()
+						? LoadClass<ANavigationData>(NULL, *SupportedAgentConfig.NavigationDataClassName.ToString(), NULL, LOAD_None, NULL)
+						: nullptr;
+
+					SetSupportedAgentsNavigationClass(AgentIndex, NavigationDataClass);
+				}
 			}
 		}
 
@@ -351,15 +388,6 @@ void UNavigationSystem::PostInitProperties()
 		FCoreUObjectDelegates::PostLoadMap.AddUObject(this, &UNavigationSystem::OnPostLoadMap);
 		UNavigationSystem::NavigationDirtyEvent.AddUObject(this, &UNavigationSystem::OnNavigationDirtied);
 	}
-	
-	// update SupportedActors' navigation classes
-	for (FNavDataConfig& SupportedAgentConfig : SupportedAgents)
-	{
-		if (SupportedAgentConfig.NavigationDataClassName.IsValid())
-		{
-			SupportedAgentConfig.NavigationDataClass = LoadClass<ANavigationData>(NULL, *SupportedAgentConfig.NavigationDataClassName.ToString(), NULL, LOAD_None, NULL);
-		}
-	}
 }
 
 bool UNavigationSystem::ConditionalPopulateNavOctree()
@@ -426,17 +454,11 @@ void UNavigationSystem::PostEditChangeProperty(FPropertyChangedEvent& PropertyCh
 		if (PropName == NAME_NavigationDataClass)
 		{
 			int32 SupportedAgentIndex = PropertyChangedEvent.GetArrayIndex(NAME_SupportedAgents.ToString());
-			if (SupportedAgentIndex != INDEX_NONE)
+			if (SupportedAgents.IsValidIndex(SupportedAgentIndex))
 			{
 				// reflect the change to SupportedAgent's 
-				if (SupportedAgents[SupportedAgentIndex].NavigationDataClass != nullptr)
-				{
-					SupportedAgents[SupportedAgentIndex].NavigationDataClassName = FStringClassReference::GetOrCreateIDForClass(SupportedAgents[SupportedAgentIndex].NavigationDataClass);
-				}
-				else
-				{
-					SupportedAgents[SupportedAgentIndex].NavigationDataClassName.Reset();
-				}
+				SetSupportedAgentsNavigationClass(SupportedAgentIndex, SupportedAgents[SupportedAgentIndex].NavigationDataClass);
+				SaveConfig();
 			}
 		}
 	}
