Fixed compilation error on ps4.  Return value is void.

diff --git a/Engine/Source/Runtime/Engine/Private/GameplayStatics.cpp b/Engine/Source/Runtime/Engine/Private/GameplayStatics.cpp
--- a/Engine/Source/Runtime/Engine/Private/GameplayStatics.cpp
+++ b/Engine/Source/Runtime/Engine/Private/GameplayStatics.cpp
@@ -1318,5 +1318,8 @@ FIntVector UGameplayStatics::GetWorldOriginLocation(UObject* WorldContextObject)
 void UGameplayStatics::SetWorldOriginLocation(UObject* WorldContextObject, FIntVector NewLocation)
 {
 	UWorld* World = GEngine->GetWorldFromContextObject(WorldContextObject);
-	return World ? World->RequestNewWorldOrigin(NewLocation) : false;
+	if ( World )
+	{
+		World->RequestNewWorldOrigin(NewLocation);
+	}
 }
\ No newline at end of file
