Added a rename flag to allow blueprint renaming without renaming the generated classes.
#UE-7128 - Duplicating a map crashes the editor (Blueprint compile crash)
#Branch UE4
#Proj Engine, CoreUObject

#info map duplication was crashing because of a name collision during regeneration of BPGC and it's CDO for the LevelScriptBlueprint. The Collision was caused by the level renaming the blueprint and the blueprint renaming it's BPGC and CDO, this was complicated more by the fact that at this point the BPGC was still the class for the original LevelScriptBlueprint ( the one we duplicated from, so we can't rename or mark as TRASH ). Alternative fix would be to temporarily null out the LevelScriptBP Generated Class and restore it after renaming.

diff --git a/Engine/Source/Runtime/CoreUObject/Public/UObject/ObjectBase.h b/Engine/Source/Runtime/CoreUObject/Public/UObject/ObjectBase.h
--- a/Engine/Source/Runtime/CoreUObject/Public/UObject/ObjectBase.h
+++ b/Engine/Source/Runtime/CoreUObject/Public/UObject/ObjectBase.h
@@ -1273,6 +1273,7 @@ typedef uint32 ERenameFlags;
 #define REN_DontCreateRedirectors (0x0010) // Don't create an object redirector, even if the class is marked RF_Public
 #define REN_NonTransactional	(0x0020) // Don't call Modify() on the objects, so they won't be stored in the transaction buffer
 #define REN_ForceGlobalUnique	(0x0040) // Force unique names across all packages not just while the scope of the new outer
+#define REN_SkipGeneratedClasses (0x0080) // Prevent renaming of any child generated classes and CDO's in blueprints
 
 /*-----------------------------------------------------------------------------
 	Misc.
diff --git a/Engine/Source/Runtime/Engine/Private/Blueprint.cpp b/Engine/Source/Runtime/Engine/Private/Blueprint.cpp
--- a/Engine/Source/Runtime/Engine/Private/Blueprint.cpp
+++ b/Engine/Source/Runtime/Engine/Private/Blueprint.cpp
@@ -364,29 +364,33 @@ void UBlueprint::Serialize(FArchive& Ar)
 
 bool UBlueprint::RenameGeneratedClasses( const TCHAR* InName, UObject* NewOuter, ERenameFlags Flags )
 {
-	FName SkelClassName, GenClassName;
-	GetBlueprintClassNames(GenClassName, SkelClassName, FName(InName));
+	const bool bRenameGeneratedClasses = !(Flags & REN_SkipGeneratedClasses );
 
-	UPackage* NewTopLevelObjectOuter = NewOuter ? NewOuter->GetOutermost() : NULL;
-	if (GeneratedClass != NULL)
+	if(bRenameGeneratedClasses)
 	{
-		bool bMovedOK = GeneratedClass->Rename(*GenClassName.ToString(), NewTopLevelObjectOuter, Flags);
-		if (!bMovedOK)
+		FName SkelClassName, GenClassName;
+		GetBlueprintClassNames(GenClassName, SkelClassName, FName(InName));
+
+		UPackage* NewTopLevelObjectOuter = NewOuter ? NewOuter->GetOutermost() : NULL;
+		if (GeneratedClass != NULL)
 		{
-			return false;
+			bool bMovedOK = GeneratedClass->Rename(*GenClassName.ToString(), NewTopLevelObjectOuter, Flags);
+			if (!bMovedOK)
+			{
+				return false;
+			}
 		}
-	}
 
-	// Also move skeleton class, if different from generated class, to new package (again, to create redirector)
-	if (SkeletonGeneratedClass != NULL && SkeletonGeneratedClass != GeneratedClass)
-	{
-		bool bMovedOK = SkeletonGeneratedClass->Rename(*SkelClassName.ToString(), NewTopLevelObjectOuter, Flags);
-		if (!bMovedOK)
+		// Also move skeleton class, if different from generated class, to new package (again, to create redirector)
+		if (SkeletonGeneratedClass != NULL && SkeletonGeneratedClass != GeneratedClass)
 		{
-			return false;
+			bool bMovedOK = SkeletonGeneratedClass->Rename(*SkelClassName.ToString(), NewTopLevelObjectOuter, Flags);
+			if (!bMovedOK)
+			{
+				return false;
+			}
 		}
 	}
-
 	return true;
 }
 
diff --git a/Engine/Source/Runtime/Engine/Private/Level.cpp b/Engine/Source/Runtime/Engine/Private/Level.cpp
--- a/Engine/Source/Runtime/Engine/Private/Level.cpp
+++ b/Engine/Source/Runtime/Engine/Private/Level.cpp
@@ -557,7 +557,7 @@ void ULevel::PostLoad()
 		if (LevelScriptBlueprint && OuterWorld && LevelScriptBlueprint->GetFName() != OuterWorld->GetFName())
 		{
 			// Use LevelScriptBlueprint->GetOuter() instead of NULL to make sure the generated top level objects are moved appropriately
-			LevelScriptBlueprint->Rename(*OuterWorld->GetName(), LevelScriptBlueprint->GetOuter(), REN_DoNotDirty | REN_DontCreateRedirectors | REN_ForceNoResetLoaders | REN_NonTransactional);
+			LevelScriptBlueprint->Rename(*OuterWorld->GetName(), LevelScriptBlueprint->GetOuter(), REN_DoNotDirty | REN_DontCreateRedirectors | REN_ForceNoResetLoaders | REN_NonTransactional | REN_SkipGeneratedClasses);
 		}
 	}
 #endif
