Fix for crash found by user.
When altering modules of an emitter that was in the level, the instances of the module needed to be recreated.

diff --git a/Engine/Source/Runtime/Engine/Classes/Particles/ParticleSystemComponent.h b/Engine/Source/Runtime/Engine/Classes/Particles/ParticleSystemComponent.h
--- a/Engine/Source/Runtime/Engine/Classes/Particles/ParticleSystemComponent.h
+++ b/Engine/Source/Runtime/Engine/Classes/Particles/ParticleSystemComponent.h
@@ -855,7 +855,7 @@ public:
 
 
 	// @todo document
-	void UpdateInstances();
+	void UpdateInstances(bool bEmptyInstances = false);
 
 
 	// @todo document
diff --git a/Engine/Source/Runtime/Engine/Private/Particles/ParticleComponents.cpp b/Engine/Source/Runtime/Engine/Private/Particles/ParticleComponents.cpp
--- a/Engine/Source/Runtime/Engine/Private/Particles/ParticleComponents.cpp
+++ b/Engine/Source/Runtime/Engine/Private/Particles/ParticleComponents.cpp
@@ -2051,11 +2051,14 @@ void UParticleSystem::PostEditChangeProperty(FPropertyChangedEvent& PropertyChan
 {
 	UpdateTime_Delta = 1.0f / UpdateTime_FPS;
 
+	//If the property is NULL then we don't really know what's happened. 
+	//Could well be a module change, requiring all instances to be destroyed and recreated.
+	bool bEmptyInstances = PropertyChangedEvent.Property == NULL;
 	for (TObjectIterator<UParticleSystemComponent> It;It;++It)
 	{
 		if (It->Template == this)
 		{
-			It->UpdateInstances();
+			It->UpdateInstances(bEmptyInstances);
 		}
 	}
 
@@ -4731,12 +4734,12 @@ void UParticleSystemComponent::ResetToDefaults()
 	}
 }
 
-void UParticleSystemComponent::UpdateInstances()
+void UParticleSystemComponent::UpdateInstances(bool bEmptyInstances)
 {
 	if (GIsEditor && IsRegistered())
 	{
 		ForceAsyncWorkCompletion(STALL);
-		ResetParticles();
+		ResetParticles(bEmptyInstances);
 
 		InitializeSystem();
 		if (bAutoActivate)
