Fixed: AlwaysLoaded sub-levels should be loaded together with persistent level
UE-7664

diff --git a/Engine/Source/Runtime/Engine/Private/World.cpp b/Engine/Source/Runtime/Engine/Private/World.cpp
--- a/Engine/Source/Runtime/Engine/Private/World.cpp
+++ b/Engine/Source/Runtime/Engine/Private/World.cpp
@@ -2212,7 +2212,7 @@ void UWorld::UpdateLevelStreamingInner(ULevelStreaming* StreamingLevel)
 	// Figure out whether level should be loaded, visible and block on load if it should be loaded but currently isn't.
 	bool bShouldBeLoaded	= bHasVisibilityRequestPending || (!GEngine->bUseBackgroundLevelStreaming && !bShouldForceUnloadStreamingLevels && !StreamingLevel->bIsRequestingUnloadAndRemoval);
 	bool bShouldBeVisible	= bHasVisibilityRequestPending || bShouldForceVisibleStreamingLevels;
-	bool bShouldBlockOnLoad	= StreamingLevel->bShouldBlockOnLoad;
+	bool bShouldBlockOnLoad	= StreamingLevel->bShouldBlockOnLoad || StreamingLevel->ShouldBeAlwaysLoaded();
 
 	// Don't update if the code requested this level object to be unloaded and removed.
 	if(!bShouldForceUnloadStreamingLevels && !StreamingLevel->bIsRequestingUnloadAndRemoval)
@@ -2228,7 +2228,7 @@ void UWorld::UpdateLevelStreamingInner(ULevelStreaming* StreamingLevel)
 	// on purpose as well so the GC code has a chance to execute between consecutive loads of maps.
 	//
 	// NOTE: AllowLevelLoadRequests not an invariant as streaming might affect the result, do NOT pulled out of the loop.
-	bool bAllowLevelLoadRequests =	AllowLevelLoadRequests() || bShouldBlockOnLoad;
+	bool bAllowLevelLoadRequests =	bShouldBlockOnLoad || AllowLevelLoadRequests();
 
 	// Figure out whether there are any levels we haven't collected garbage yet.
 	bool bAreLevelsPendingPurge	=	FLevelStreamingGCHelper::GetNumLevelsPendingPurge() > 0;
