Added deferred addition and removal of PhysX bodies to scenes, to allow us to use faster scene insertion path.

diff --git a/Engine/Source/Runtime/Engine/Classes/PhysicsEngine/BodyInstance.h b/Engine/Source/Runtime/Engine/Classes/PhysicsEngine/BodyInstance.h
--- a/Engine/Source/Runtime/Engine/Classes/PhysicsEngine/BodyInstance.h
+++ b/Engine/Source/Runtime/Engine/Classes/PhysicsEngine/BodyInstance.h
@@ -109,6 +109,15 @@ private:
 	friend struct FBodyInstance;
 };
 
+enum class BodyInstanceSceneState
+{
+	NotAdded,
+	AwaitingAdd,
+	Added,
+	AwaitingRemove,
+	Removed
+};
+
 /** Container for a physics representation of an object */
 USTRUCT()
 struct ENGINE_API FBodyInstance
@@ -329,6 +338,13 @@ public:
 	TSharedPtr<TArray<ANSICHAR>> CharDebugName;
 #endif	//WITH_PHYSX
 
+	/** Per instance data used to initialize dynamic instances */
+	/** Initial physx velocity to apply to dynamic instances */
+	FVector InitialLinearVelocity;
+
+	/** When initializing dynamic instances their component or velocity can override the bStartAwake flag */
+	bool bWokenExternally;
+
 	/** PrimitiveComponent containing this body.   */
 	TWeakObjectPtr<class UPrimitiveComponent> OwnerComponent;
 
@@ -361,7 +377,7 @@ public:
 	 *	@param InRBScene The physics scene to place the body into
 	 *	@param InAggregate An aggregate to place the body into
 	 */
-	void InitBody(class UBodySetup* Setup, const FTransform& Transform, class UPrimitiveComponent* PrimComp, class FPhysScene* InRBScene, PhysXAggregateType InAggregate = NULL);
+	void InitBody(class UBodySetup* Setup, const FTransform& Transform, class UPrimitiveComponent* PrimComp, class FPhysScene* InRBScene, PhysXAggregateType InAggregate = NULL, bool bDefer = false);
 
 	/** Initialise multiple rigid bodies 
 	 *	@param Bodies List of body instances to initialise
@@ -371,7 +387,7 @@ public:
 	 *	@param InRBScene The physics scene to put the bodies into
 	 *	@param InAggregate Aggregate to place the bodies into, be aware for PhysX we can only have 128 elements in an aggregate!
 	 */
-	static void InitBodies(TArray<FBodyInstance*>& Bodies, TArray<FTransform>& Transforms, class UBodySetup* BodySetup, class UPrimitiveComponent* PrimitiveComp, class FPhysScene* InRBScene, PhysXAggregateType InAggregate = NULL);
+	static void InitBodies(TArray<FBodyInstance*>& Bodies, TArray<FTransform>& Transforms, class UBodySetup* BodySetup, class UPrimitiveComponent* PrimitiveComp, class FPhysScene* InRBScene, PhysXAggregateType InAggregate = NULL, bool Defer = false);
 
 	/** Initialises one or more rigid bodies
 	 *	@param Bodies First body instance to initialise
@@ -382,7 +398,7 @@ public:
 	 *	@param InRBScene Scene to place the bodies into
 	 *	@param InAggregate Aggregate to place the bodies into, be aware for PhysX we can only have 128 elements in an aggregate!
 	 */
-	static void InitBodies(FBodyInstance** Bodies, FTransform* Transforms, int32 NumBodies, class UBodySetup* BodySetup, class UPrimitiveComponent* PrimitiveComp, class FPhysScene* InRBScene, PhysXAggregateType InAggregate = NULL);
+	static void InitBodies(FBodyInstance** Bodies, FTransform* Transforms, int32 NumBodies, class UBodySetup* BodySetup, class UPrimitiveComponent* PrimitiveComp, class FPhysScene* InRBScene, PhysXAggregateType InAggregate = NULL, bool Defer = false);
 
 	void InitBody_Legacy(class UBodySetup* Setup, const FTransform& Transform, class UPrimitiveComponent* PrimComp, class FPhysScene* InRBScene, PhysXAggregateType InAggregate = NULL);
 
@@ -393,6 +409,12 @@ public:
 	 */
 	static bool ValidateTransform(const FTransform &Transform, FString& DebugName, UBodySetup* Setup);
 
+#if WITH_PHYSX
+	/** Initialise dynamic properties for this instance when using PhysX - this must be done after scene addition.
+	 */
+	void InitDynamicProperties();
+#endif // WITH_PHYSX
+
 #endif	//UE_WITH_PHYSICS
 
 	void TermBody();
@@ -718,6 +740,9 @@ public:
 public:
 	FPhysxUserData PhysxUserData;
 
+	// Current state of the physx body for tracking deferred addition and removal.
+	BodyInstanceSceneState CurrentSceneState;
+
 private:
 	/**
 	 *  Trace a box against just this bodyinstance
diff --git a/Engine/Source/Runtime/Engine/Private/InstancedStaticMesh.cpp b/Engine/Source/Runtime/Engine/Private/InstancedStaticMesh.cpp
--- a/Engine/Source/Runtime/Engine/Private/InstancedStaticMesh.cpp
+++ b/Engine/Source/Runtime/Engine/Private/InstancedStaticMesh.cpp
@@ -911,7 +911,7 @@ void UInstancedStaticMeshComponent::CreateAllInstanceBodies()
 
 	if(NumBodies > 0)
 	{
-		FBodyInstance::InitBodies(InstanceBodies, Transforms, GetBodySetup(), this, GetWorld()->GetPhysicsScene());
+		FBodyInstance::InitBodies(InstanceBodies, Transforms, GetBodySetup(), this, GetWorld()->GetPhysicsScene(), nullptr, true);
 	}
 }
 
diff --git a/Engine/Source/Runtime/Engine/Private/PhysicsEngine/BodyInstance.cpp b/Engine/Source/Runtime/Engine/Private/PhysicsEngine/BodyInstance.cpp
--- a/Engine/Source/Runtime/Engine/Private/PhysicsEngine/BodyInstance.cpp
+++ b/Engine/Source/Runtime/Engine/Private/PhysicsEngine/BodyInstance.cpp
@@ -240,7 +240,10 @@ FBodyInstance::FBodyInstance()
 , RigidActorSync(NULL)
 , RigidActorAsync(NULL)
 , BodyAggregate(NULL)
+, InitialLinearVelocity(0.0f)
+, bWokenExternally(false)
 , PhysxUserData(this)
+, CurrentSceneState(BodyInstanceSceneState::NotAdded)
 #endif	//WITH_PHYSX
 #if WITH_BOX2D
 , BodyInstancePtr(nullptr)
@@ -887,17 +890,17 @@ DEFINE_STAT(STAT_BulkSceneAdd);
 
 #if UE_WITH_PHYSICS
 
-void FBodyInstance::InitBody(class UBodySetup* Setup, const FTransform& Transform, class UPrimitiveComponent* PrimComp, class FPhysScene* InRBScene, PhysXAggregateType InAggregate /*= NULL*/ )
+void FBodyInstance::InitBody(class UBodySetup* Setup, const FTransform& Transform, class UPrimitiveComponent* PrimComp, class FPhysScene* InRBScene, PhysXAggregateType InAggregate /*= NULL*/, bool bDefer /*= false*/ )
 {
 	FBodyInstance* Instance = this;
 	FTransform TransformCopy = Transform;
-	InitBodies(&Instance, &TransformCopy, 1, Setup, PrimComp, InRBScene, InAggregate);
+	InitBodies(&Instance, &TransformCopy, 1, Setup, PrimComp, InRBScene, InAggregate, bDefer);
 }
 
-void FBodyInstance::InitBodies(TArray<FBodyInstance*>& Bodies, TArray<FTransform>& Transforms, class UBodySetup* BodySetup, class UPrimitiveComponent* PrimitiveComp, class FPhysScene* InRBScene, PhysXAggregateType InAggregate /*= NULL*/)
+void FBodyInstance::InitBodies(TArray<FBodyInstance*>& Bodies, TArray<FTransform>& Transforms, class UBodySetup* BodySetup, class UPrimitiveComponent* PrimitiveComp, class FPhysScene* InRBScene, PhysXAggregateType InAggregate /*= NULL*/, bool Defer /*= false*/)
 {
 	check(Bodies.Num() == Transforms.Num());
-	InitBodies(Bodies.GetData(), Transforms.GetData(), Bodies.Num(), BodySetup, PrimitiveComp, InRBScene, InAggregate);
+	InitBodies(Bodies.GetData(), Transforms.GetData(), Bodies.Num(), BodySetup, PrimitiveComp, InRBScene, InAggregate, Defer);
 }
 
 void FBodyInstance::InitBody_Legacy(UBodySetup* Setup, const FTransform& Transform, UPrimitiveComponent* PrimComp, FPhysScene* InRBScene, PxAggregate* InAggregate)
@@ -1401,7 +1404,7 @@ void FBodyInstance::InitBody_Legacy(UBodySetup* Setup, const FTransform& Transfo
 }
 
 
-void FBodyInstance::InitBodies(FBodyInstance** Bodies, FTransform* Transforms, int32 NumBodies, class UBodySetup* BodySetup, class UPrimitiveComponent* PrimitiveComp, class FPhysScene* InRBScene, PhysXAggregateType InAggregate /*= NULL*/ )
+void FBodyInstance::InitBodies(FBodyInstance** Bodies, FTransform* Transforms, int32 NumBodies, class UBodySetup* BodySetup, class UPrimitiveComponent* PrimitiveComp, class FPhysScene* InRBScene, PhysXAggregateType InAggregate /*= NULL*/, bool bDefer /*= false*/ )
 {
 	SCOPE_CYCLE_COUNTER(STAT_BulkInit);
 
@@ -1423,6 +1426,9 @@ void FBodyInstance::InitBodies(FBodyInstance** Bodies, FTransform* Transforms, i
 	TArray<PxActor*> PhysXDynamicActors;
 	PhysXSyncActors.Reserve(NumBodies);
 
+	// No reason to defer if there's an existing aggregate.
+	const bool bCanDefer = bDefer && !InAggregate;
+
 	if(PSceneAsync)
 	{
 		// Only allocate this if we can have async bodies
@@ -1489,11 +1495,10 @@ void FBodyInstance::InitBodies(FBodyInstance** Bodies, FTransform* Transforms, i
 		}
 	}
 
-	bool bStartAwakeOverride = false;
-
+	bool bComponentAwake = false;
 	if(SkelMeshComp)
 	{
-		bStartAwakeOverride = SkelMeshComp->BodyInstance.bStartAwake;
+		bComponentAwake = SkelMeshComp->BodyInstance.bStartAwake;
 	}
 
 	FVector InitialLinVel(EForceInit::ForceInitToZero);
@@ -1503,7 +1508,7 @@ void FBodyInstance::InitBodies(FBodyInstance** Bodies, FTransform* Transforms, i
 
 		if(InitialLinVel.Size() > KINDA_SMALL_NUMBER)
 		{
-			bStartAwakeOverride = true;
+			bComponentAwake = true;
 		}
 	}
 	PxVec3 PxInitialLinVel = U2PVector(InitialLinVel);
@@ -1535,10 +1540,10 @@ void FBodyInstance::InitBodies(FBodyInstance** Bodies, FTransform* Transforms, i
 
 				if(SkelMeshComp)
 				{
-					bStartAwakeOverride = Instance->bStartAwake && SkelMeshComp->BodyInstance.bStartAwake;
+					bComponentAwake = Instance->bStartAwake && SkelMeshComp->BodyInstance.bStartAwake;
 				}
 
-				if(Instance->bStartAwake || bStartAwakeOverride)
+				if(Instance->bStartAwake || bComponentAwake)
 				{
 					BodyDefinition.awake = true;
 				}
@@ -1823,21 +1828,53 @@ void FBodyInstance::InitBodies(FBodyInstance** Bodies, FTransform* Transforms, i
 			continue;
 		}
 
-		PhysXSyncActors.Add(PNewActorSync);
-		if(PNewActorAsync)
+		if(bCanDefer)
 		{
-			PhysXAsyncActors.Add(PNewActorAsync);
+			if(PNewDynamic)
+			{
+				InRBScene->DeferAddActor(Instance, PNewDynamic, Instance->UseAsyncScene() ? PST_Async : PST_Sync);
+			}
+			else
+			{
+				if(PNewActorSync)
+				{
+					InRBScene->DeferAddActor(Instance, PNewActorSync, PST_Sync);
+				}
+				if(PNewActorAsync)
+				{
+					InRBScene->DeferAddActor(Instance, PNewActorAsync, PST_Async);
+				}
+			}
 		}
-
-		if(PNewDynamic)
+		else
 		{
-			PhysXDynamicActors.Add(PNewDynamic);
-		}
+			if(PNewActorSync)
+			{
+				PhysXSyncActors.Add(PNewActorSync);
+			}
+
+			if(PNewActorAsync)
+			{
+				PhysXAsyncActors.Add(PNewActorAsync);
+			}
+
+			if(PNewDynamic)
+			{
+				PhysXDynamicActors.Add(PNewDynamic);
+			}
 
+			// Set the instance to added as we'll add it to the scene in a moment.
+			// This makes sure if it ends up in one of the deferred queues later it 
+			// functions correctly
+			Instance->CurrentSceneState = BodyInstanceSceneState::Added;
+		}
+		
 		Instance->RigidActorSync = PNewActorSync;
 		Instance->RigidActorAsync = PNewActorAsync;
 		Instance->SceneIndexSync = InRBScene->PhysXSceneIndex[PST_Sync];
 		Instance->SceneIndexAsync = PSceneAsync ? InRBScene->PhysXSceneIndex[PST_Async] : 0;
+		Instance->InitialLinearVelocity = InitialLinVel;
+		Instance->bWokenExternally = bComponentAwake;
 		Instance->UpdatePhysicsFilterData();
 	}
 
@@ -1865,6 +1902,7 @@ void FBodyInstance::InitBodies(FBodyInstance** Bodies, FTransform* Transforms, i
 		}
 	}
 	
+	if(!bCanDefer)
 	{
 		SCOPE_CYCLE_COUNTER(STAT_BulkSceneAdd);
 		// Use the aggregate if it exists, has enough slots and is in the correct scene (or no scene)
@@ -1892,64 +1930,21 @@ void FBodyInstance::InitBodies(FBodyInstance** Bodies, FTransform* Transforms, i
 				}
 			}
 		}
-	}
 
-	// Set up dynamic instance data
-	if(!bStatic)
-	{
-		SCOPE_CYCLE_COUNTER(STAT_InitBodyPostAdd);
-		for(int32 BodyIdx = 0 ; BodyIdx < NumBodies ; ++BodyIdx)
+		// Set up dynamic instance data
+		if(!bStatic)
 		{
-			FBodyInstance* Instance = Bodies[BodyIdx];
-			
-			if(PxRigidDynamic* RigidActor = Instance->GetPxRigidDynamic())
+			SCOPE_CYCLE_COUNTER(STAT_InitBodyPostAdd);
+			for(int32 BodyIdx = 0 ; BodyIdx < NumBodies ; ++BodyIdx)
 			{
-				Instance->UpdateMassProperties();
-				Instance->UpdateDampingProperties();
-				Instance->SetMaxAngularVelocity(Instance->MaxAngularVelocity, false);
-				Instance->SetMaxDepenetrationVelocity(Instance->bOverrideMaxDepenetrationVelocity ? Instance->MaxDepenetrationVelocity : UPhysicsSettings::Get()->MaxDepenetrationVelocity);
-
-				if(Instance->ShouldInstanceSimulatingPhysics())
-				{
-					RigidActor->setLinearVelocity(PxInitialLinVel);
-				}
-
-				float SleepEnergyThresh = RigidActor->getSleepThreshold();
-				if(Instance->SleepFamily == SF_Sensitive)
-				{
-					SleepEnergyThresh /= 20.f;
-				}
-				RigidActor->setSleepThreshold(SleepEnergyThresh);
-				// set solver iteration count 
-				int32 PositionIterCount = FMath::Clamp(Instance->PositionSolverIterationCount, 1, 255);
-				int32 VelocityIterCount = FMath::Clamp(Instance->VelocitySolverIterationCount, 1, 255);
-				RigidActor->setSolverIterationCounts(PositionIterCount, VelocityIterCount);
-
-				if(RigidActor->getScene())
-				{
-					Instance->CreateDOFLock();
-					if(IsRigidBodyNonKinematic(RigidActor))
-					{
-						if(SkelMeshComp)
-						{
-							bStartAwakeOverride = Instance->bStartAwake && SkelMeshComp->BodyInstance.bStartAwake;
-						}
-
-						if(Instance->bStartAwake || bStartAwakeOverride)
-						{
-							RigidActor->wakeUp();
-						}
-						else
-						{
-							RigidActor->putToSleep();
-						}
-					}
-				}
+				FBodyInstance* Instance = Bodies[BodyIdx];
+				Instance->InitDynamicProperties();
 			}
-			
 		}
 	}
-#endif
+
+	
+#endif // WITH_PHYSX
 }
 
 
@@ -4493,6 +4488,51 @@ bool FBodyInstance::ValidateTransform(const FTransform &Transform, FString& Debu
 	return true;
 }
 
+#if WITH_PHYSX
+void FBodyInstance::InitDynamicProperties()
+{
+	if(PxRigidDynamic* RigidActor = GetPxRigidDynamic())
+	{
+		UpdateMassProperties();
+		UpdateDampingProperties();
+		SetMaxAngularVelocity(MaxAngularVelocity, false);
+		SetMaxDepenetrationVelocity(bOverrideMaxDepenetrationVelocity ? MaxDepenetrationVelocity : UPhysicsSettings::Get()->MaxDepenetrationVelocity);
+
+		if(ShouldInstanceSimulatingPhysics())
+		{
+			RigidActor->setLinearVelocity(U2PVector(InitialLinearVelocity));
+		}
+
+		float SleepEnergyThresh = RigidActor->getSleepThreshold();
+		if(SleepFamily == SF_Sensitive)
+		{
+			SleepEnergyThresh /= 20.f;
+		}
+		RigidActor->setSleepThreshold(SleepEnergyThresh);
+		// set solver iteration count 
+		int32 PositionIterCount = FMath::Clamp(PositionSolverIterationCount, 1, 255);
+		int32 VelocityIterCount = FMath::Clamp(VelocitySolverIterationCount, 1, 255);
+		RigidActor->setSolverIterationCounts(PositionIterCount, VelocityIterCount);
+
+		if(RigidActor->getScene())
+		{
+			CreateDOFLock();
+			if(IsRigidBodyNonKinematic(RigidActor))
+			{
+				if(bStartAwake || bWokenExternally)
+				{
+					RigidActor->wakeUp();
+				}
+				else
+				{
+					RigidActor->putToSleep();
+				}
+			}
+		}
+	}
+}
+#endif
+
 ////////////////////////////////////////////////////////////////////////////
 // FBodyInstanceEditorHelpers
 
diff --git a/Engine/Source/Runtime/Engine/Private/PhysicsEngine/PhysScene.cpp b/Engine/Source/Runtime/Engine/Private/PhysicsEngine/PhysScene.cpp
--- a/Engine/Source/Runtime/Engine/Private/PhysicsEngine/PhysScene.cpp
+++ b/Engine/Source/Runtime/Engine/Private/PhysicsEngine/PhysScene.cpp
@@ -329,13 +329,25 @@ void FPhysScene::RemoveActiveBody(FBodyInstance* BodyInstance, uint32 SceneType)
 #endif
 void FPhysScene::TermBody(FBodyInstance* BodyInstance)
 {
-#if WITH_SUBSTEPPING
 	if (PxRigidBody * PRigidBody = BodyInstance->GetPxRigidBody())
 	{
+#if WITH_SUBSTEPPING
 		FPhysSubstepTask * PhysSubStepper = PhysSubSteppers[SceneType(BodyInstance)];
 		PhysSubStepper->RemoveBodyInstance(BodyInstance);
-	}
 #endif
+
+		// Remove body from any pending deferred addition / removal
+		for(int32 SceneIdx = 0 ; SceneIdx < PST_MAX ; ++SceneIdx)
+		{
+			int32 FoundIdx = INDEX_NONE;
+			if(DeferredAddActors[SceneIdx].Find(PRigidBody, FoundIdx))
+			{
+				DeferredAddActors->RemoveAt(FoundIdx);
+				DeferredAddInstances->RemoveAt(FoundIdx);
+			}
+		}
+	}
+
 #if WITH_PHYSX
 	RemoveActiveBody(BodyInstance, PST_Sync);
 	RemoveActiveBody(BodyInstance, PST_Async);
@@ -881,6 +893,10 @@ void FPhysScene::StartFrame()
 
 	//Update the collision disable table before ticking
 	FlushDeferredCollisionDisableTableQueue();
+#if WITH_PHYSX
+	// Flush list of deferred actors to add to the scene
+	FlushDeferredActors();
+#endif
 
 	// Run the sync scene
 	TickPhysScene(PST_Sync, PhysicsSubsceneCompletion[PST_Sync]);
@@ -1360,3 +1376,96 @@ void FPhysScene::TermPhysScene(uint32 SceneType)
 	}
 #endif
 }
+
+#if WITH_PHYSX
+void FPhysScene::FlushDeferredActors()
+{
+	check(IsInGameThread());
+
+	for(int32 SceneIdx = 0 ; SceneIdx < PST_MAX ; ++SceneIdx)
+	{
+		check(DeferredAddInstances[SceneIdx].Num() == DeferredAddActors[SceneIdx].Num());
+
+		if(DeferredAddInstances[SceneIdx].Num() > 0)
+		{
+			PxScene* Scene = GetPhysXScene(SceneIdx);
+			SCOPED_SCENE_WRITE_LOCK(Scene);
+
+			Scene->addActors(DeferredAddActors[SceneIdx].GetData(), DeferredAddActors[SceneIdx].Num());
+
+			for(FBodyInstance* Instance : DeferredAddInstances[SceneIdx])
+			{
+				Instance->CurrentSceneState = BodyInstanceSceneState::Added;
+
+				if(Instance->IsDynamic())
+				{
+					// Extra setup necessary for dynamic objects.
+					Instance->InitDynamicProperties();
+				}
+			}
+
+			DeferredAddInstances->Empty();
+			DeferredAddActors->Empty();
+		}
+
+		check(DeferredRemoveInstances[SceneIdx].Num() == DeferredRemoveActors[SceneIdx].Num());
+
+		if(DeferredRemoveInstances[SceneIdx].Num() > 0)
+		{
+			PxScene* Scene = GetPhysXScene(SceneIdx);
+			SCOPED_SCENE_WRITE_LOCK(Scene);
+
+			Scene->removeActors(DeferredRemoveActors[SceneIdx].GetData(), DeferredRemoveActors[SceneIdx].Num());
+
+			for(FBodyInstance* Instance : DeferredAddInstances[SceneIdx])
+			{
+				Instance->CurrentSceneState = BodyInstanceSceneState::Removed;
+			}
+
+			DeferredRemoveInstances->Empty();
+			DeferredRemoveActors->Empty();
+		}
+	}
+}
+
+void FPhysScene::DeferAddActor(FBodyInstance* OwningInstance, PxActor* Actor, EPhysicsSceneType SceneType)
+{
+	check(OwningInstance && Actor && SceneType < PST_MAX);
+	// Allowed to be unadded or awaiting add here (objects can be in more than one scene
+	if(OwningInstance->CurrentSceneState == BodyInstanceSceneState::NotAdded ||
+	   OwningInstance->CurrentSceneState == BodyInstanceSceneState::AwaitingAdd)
+	{
+		OwningInstance->CurrentSceneState = BodyInstanceSceneState::AwaitingAdd;
+		DeferredAddInstances[SceneType].Add(OwningInstance);
+		DeferredAddActors[SceneType].Add(Actor);
+	}
+	else if(OwningInstance->CurrentSceneState == BodyInstanceSceneState::AwaitingRemove)
+	{
+		// We were waiting to be removed, but we're canceling that
+
+		OwningInstance->CurrentSceneState = BodyInstanceSceneState::Added;
+		DeferredRemoveInstances[SceneType].RemoveSingle(OwningInstance);
+		DeferredRemoveActors[SceneType].RemoveSingle(Actor);
+	}
+}
+
+void FPhysScene::DeferRemoveActor(FBodyInstance* OwningInstance, PxActor* Actor, EPhysicsSceneType SceneType)
+{
+	check(OwningInstance && Actor && SceneType < PST_MAX);
+	if(OwningInstance->CurrentSceneState == BodyInstanceSceneState::Added ||
+	   OwningInstance->CurrentSceneState == BodyInstanceSceneState::AwaitingRemove)
+	{
+		OwningInstance->CurrentSceneState = BodyInstanceSceneState::AwaitingRemove;
+		DeferredRemoveInstances[SceneType].Add(OwningInstance);
+		DeferredRemoveActors[SceneType].Add(Actor);
+	}
+	else if(OwningInstance->CurrentSceneState == BodyInstanceSceneState::AwaitingAdd)
+	{
+		// We were waiting to add but now we're canceling it
+		OwningInstance->CurrentSceneState = BodyInstanceSceneState::Removed;
+		DeferredAddInstances[SceneType].RemoveSingle(OwningInstance);
+		DeferredAddActors[SceneType].RemoveSingle(Actor);
+	}
+}
+
+#endif
\ No newline at end of file
diff --git a/Engine/Source/Runtime/Engine/Public/PhysicsPublic.h b/Engine/Source/Runtime/Engine/Public/PhysicsPublic.h
--- a/Engine/Source/Runtime/Engine/Public/PhysicsPublic.h
+++ b/Engine/Source/Runtime/Engine/Public/PhysicsPublic.h
@@ -294,6 +294,37 @@ public:
 	 */
 	float							FrameTimeSmoothingFactor[PST_MAX];
 
+#if WITH_PHYSX
+	/** Body instances awaiting scene add */
+	TArray<FBodyInstance*> DeferredAddInstances[PST_MAX];
+	/** PhysX Actors awaiting scene add */
+	TArray<PxActor*> DeferredAddActors[PST_MAX];
+
+	/** Body instances awaiting scene remove */
+	TArray<FBodyInstance*> DeferredRemoveInstances[PST_MAX];
+	/** PhysX Actors awaiting scene remove */
+	TArray<PxActor*> DeferredRemoveActors[PST_MAX];
+
+	/** Flush the deferred actor and instance arrays, either adding or removing from the scene */
+	void FlushDeferredActors();
+
+	/** Defer the addition of an actor to a scene, this will actually be performed before the *next*
+	 *  Physics tick
+	 *	@param OwningInstance - The FBodyInstance that owns the actor
+	 *	@param Actor - The actual PhysX actor to add
+	 *	@param SceneType - The scene type to add the actor to
+	 */
+	void DeferAddActor(FBodyInstance* OwningInstance, PxActor* Actor, EPhysicsSceneType SceneType);
+
+	/** Defer the removal of an actor to a scene, this will actually be performed before the *next*
+	 *  Physics tick
+	 *	@param OwningInstance - The FBodyInstance that owns the actor
+	 *	@param Actor - The actual PhysX actor to add
+	 *	@param SceneType - The scene type to add the actor to
+	 */
+	void DeferRemoveActor(FBodyInstance* OwningInstance, PxActor* Actor, EPhysicsSceneType SceneType);
+#endif
+
 private:
 	/** DeltaSeconds from UWorld. */
 	float										DeltaSeconds;
