PR #676: Expose USkeletalMesh::FindSocket to the blueprint system. (Contributed by [REDACTED_PERSON])

https://github.com/EpicGames/UnrealEngine/pull/676

diff --git a/Engine/Source/Runtime/Engine/Classes/Engine/SkeletalMesh.h b/Engine/Source/Runtime/Engine/Classes/Engine/SkeletalMesh.h
--- a/Engine/Source/Runtime/Engine/Classes/Engine/SkeletalMesh.h
+++ b/Engine/Source/Runtime/Engine/Classes/Engine/SkeletalMesh.h
@@ -752,7 +752,8 @@ public:
 	 *	Find a socket object in this SkeletalMesh by name. 
 	 *	Entering NAME_None will return NULL. If there are multiple sockets with the same name, will return the first one.
 	 */
-	ENGINE_API class USkeletalMeshSocket const* FindSocket(FName InSocketName) const;
+	UFUNCTION(BlueprintCallable, Category="Animation")
+	ENGINE_API USkeletalMeshSocket* FindSocket(FName InSocketName) const;
 
 	// @todo document
 	ENGINE_API FMatrix GetRefPoseMatrix( int32 BoneIndex ) const;
diff --git a/Engine/Source/Runtime/Engine/Private/SkeletalMesh.cpp b/Engine/Source/Runtime/Engine/Private/SkeletalMesh.cpp
--- a/Engine/Source/Runtime/Engine/Private/SkeletalMesh.cpp
+++ b/Engine/Source/Runtime/Engine/Private/SkeletalMesh.cpp
@@ -2826,7 +2826,7 @@ UMorphTarget* USkeletalMesh::FindMorphTarget( FName MorphTargetName ) const
 	return NULL;
 }
 
-USkeletalMeshSocket const* USkeletalMesh::FindSocket(FName InSocketName) const
+USkeletalMeshSocket* USkeletalMesh::FindSocket(FName InSocketName) const
 {
 	if(InSocketName == NAME_None)
 	{
