Removed UNavigationSystem::NavDataClasses due to it no longer being filled but it being empty treated as an error #UE4

- addresses UE-7670, not being able to cook sample projects (or any project for that matter!)

diff --git a/Engine/Source/Runtime/Engine/Classes/AI/Navigation/NavigationSystem.h b/Engine/Source/Runtime/Engine/Classes/AI/Navigation/NavigationSystem.h
--- a/Engine/Source/Runtime/Engine/Classes/AI/Navigation/NavigationSystem.h
+++ b/Engine/Source/Runtime/Engine/Classes/AI/Navigation/NavigationSystem.h
@@ -732,9 +732,6 @@ protected:
 	/** self-registering exec command to handle nav sys console commands */
 	static FNavigationSystemExec ExecHandler;
 #endif // !UE_BUILD_SHIPPING
-
-	/** collection of delegates to create all available navigation data types */
-	static TArray<TSubclassOf<ANavigationData> > NavDataClasses;
 	
 	/** whether seamless navigation building is enabled */
 	static bool bNavigationAutoUpdateEnabled;
diff --git a/Engine/Source/Runtime/Engine/Private/AI/Navigation/NavigationSystem.cpp b/Engine/Source/Runtime/Engine/Private/AI/Navigation/NavigationSystem.cpp
--- a/Engine/Source/Runtime/Engine/Private/AI/Navigation/NavigationSystem.cpp
+++ b/Engine/Source/Runtime/Engine/Private/AI/Navigation/NavigationSystem.cpp
@@ -187,7 +187,6 @@ void FNavigationLockContext::UnlockUpdates()
 // UNavigationSystem                                                                
 //----------------------------------------------------------------------//
 bool UNavigationSystem::bNavigationAutoUpdateEnabled = true;
-TArray<TSubclassOf<ANavigationData> > UNavigationSystem::NavDataClasses;
 TArray<const UClass*> UNavigationSystem::NavAreaClasses;
 TArray<UClass*> UNavigationSystem::PendingNavAreaRegistration;
 TSubclassOf<UNavArea> UNavigationSystem::DefaultWalkableArea = NULL;
@@ -368,13 +367,6 @@ void UNavigationSystem::PostInitProperties()
 				}
 			}
 		}
-
-		if (NavDataClasses.Num() == 0)
-		{
-			// @note: if you don't want navigation system to be created at all you can disable it by 
-			// setting AWorldSettings.bEnableNavigationSystem to false
-			UE_LOG(LogNavigation, Error, TEXT("No navigation data types found while setting up required navigation types!"));
-		}
 	
 		bInitialBuildingLockActive = bInitialBuildingLocked;
 
@@ -2813,11 +2805,6 @@ void UNavigationSystem::Build()
 
 	SpawnMissingNavigationData();
 
-	if (NavDataClasses.Num() == 0)
-	{
-		return;
-	}
-
 	// make sure freshly created navigation instances are registered before we try to build them
 	ProcessRegistrationCandidates();
 
