UE4 - fix texture DDC threading issue that could crash on boot

diff --git a/Engine/Source/Runtime/Engine/Private/Texture.cpp b/Engine/Source/Runtime/Engine/Private/Texture.cpp
--- a/Engine/Source/Runtime/Engine/Private/Texture.cpp
+++ b/Engine/Source/Runtime/Engine/Private/Texture.cpp
@@ -405,6 +405,7 @@ UEnum* UTexture::GetPixelFormatEnum()
 	static UEnum* PixelFormatEnum = NULL;
 	if (PixelFormatEnum == NULL)
 	{
+		check(IsInGameThread());
 		UEnum::LookupEnumName(PixelFormatUnknownName, &PixelFormatEnum);
 		check(PixelFormatEnum);
 	}
diff --git a/Engine/Source/Runtime/Engine/Private/TextureDerivedData.cpp b/Engine/Source/Runtime/Engine/Private/TextureDerivedData.cpp
--- a/Engine/Source/Runtime/Engine/Private/TextureDerivedData.cpp
+++ b/Engine/Source/Runtime/Engine/Private/TextureDerivedData.cpp
@@ -866,6 +866,7 @@ public:
 		{
 			GetBuildInfo();
 		}
+		UTexture::GetPixelFormatEnum();
 	}
 
 	/** Does the work to cache derived data. Safe to call from any thread. */
@@ -1285,6 +1286,7 @@ FAsyncStreamDerivedMipWorker::FAsyncStreamDerivedMipWorker(
 	, bRequestFailed(false)
 	, ThreadSafeCounter(InThreadSafeCounter)
 {
+	UTexture::GetPixelFormatEnum();
 }
 
 /** Retrieves the derived mip from the derived data cache. */
