Fix for editor only Transform Curve data getting cooked which causes a crash on cooked platforms.

diff --git a/Engine/Source/Runtime/Engine/Private/Animation/AnimSequenceBase.cpp b/Engine/Source/Runtime/Engine/Private/Animation/AnimSequenceBase.cpp
--- a/Engine/Source/Runtime/Engine/Private/Animation/AnimSequenceBase.cpp
+++ b/Engine/Source/Runtime/Engine/Private/Animation/AnimSequenceBase.cpp
@@ -317,15 +317,19 @@ void FRawCurveTracks::Serialize(FArchive& Ar)
 			Curve.Serialize(Ar);
 		}
 	}
-#if WITH_EDITOR
-	if(Ar.UE4Ver() >= VER_UE4_ANIMATION_ADD_TRACKCURVES)
+#if WITH_EDITORONLY_DATA
+	if( !Ar.IsCooking() )
 	{
-		for(FTransformCurve& Curve : TransformCurves)
+		if( Ar.UE4Ver() >= VER_UE4_ANIMATION_ADD_TRACKCURVES )
 		{
-			Curve.Serialize(Ar);
+			for( FTransformCurve& Curve : TransformCurves )
+			{
+				Curve.Serialize( Ar );
+			}
+
 		}
 	}
-#endif // WITH_EDITOR
+#endif // WITH_EDITORONLY_DATA
 }
 
 void FRawCurveTracks::UpdateLastObservedNames(FSmartNameMapping* NameMapping, ESupportedCurveType SupportedCurveType /*= FloatType*/)
