Expanded the condition for "empty" navmesh to treat properly navmesh instances that had been initialized but don't have any tiles as empty as well #UE4

diff --git a/Engine/Source/Runtime/Engine/Private/AI/Navigation/RecastNavMeshGenerator.cpp b/Engine/Source/Runtime/Engine/Private/AI/Navigation/RecastNavMeshGenerator.cpp
--- a/Engine/Source/Runtime/Engine/Private/AI/Navigation/RecastNavMeshGenerator.cpp
+++ b/Engine/Source/Runtime/Engine/Private/AI/Navigation/RecastNavMeshGenerator.cpp
@@ -3092,32 +3092,34 @@ TArray<uint32> FRecastNavMeshGenerator::RemoveTileLayers(const int32 TileX, cons
 	TArray<uint32> ResultTileIndices;
 	dtNavMesh* DetourMesh = DestNavMesh->GetRecastNavMeshImpl()->GetRecastMesh();
 	
-	check(DetourMesh == nullptr || DetourMesh->isEmpty() == false);
-	const int32 NumLayers = DetourMesh != nullptr ? DetourMesh->getTileCountAt(TileX, TileY) : 0;
-	
-	if (NumLayers > 0)
+	if (DetourMesh != nullptr && DetourMesh->isEmpty() == false)
 	{
-		TArray<dtMeshTile*> Tiles;
-		Tiles.AddZeroed(NumLayers);
-		DetourMesh->getTilesAt(TileX, TileY, (const dtMeshTile**)Tiles.GetData(), NumLayers);
-	
-		for (int32 i = 0; i < NumLayers; i++)
+		const int32 NumLayers = DetourMesh != nullptr ? DetourMesh->getTileCountAt(TileX, TileY) : 0;
+
+		if (NumLayers > 0)
 		{
-			const int32 LayerIndex = Tiles[i]->header->layer;
-			const dtTileRef TileRef = DetourMesh->getTileRef(Tiles[i]);
+			TArray<dtMeshTile*> Tiles;
+			Tiles.AddZeroed(NumLayers);
+			DetourMesh->getTilesAt(TileX, TileY, (const dtMeshTile**)Tiles.GetData(), NumLayers);
 
-			NumActiveTiles--;
-			UE_LOG(LogNavigation, Log, TEXT("%s> Tile (%d,%d:%d), removing TileRef: 0x%X (active:%d)"),
-				*DestNavMesh->GetName(), TileX, TileY, LayerIndex, TileRef, NumActiveTiles);
+			for (int32 i = 0; i < NumLayers; i++)
+			{
+				const int32 LayerIndex = Tiles[i]->header->layer;
+				const dtTileRef TileRef = DetourMesh->getTileRef(Tiles[i]);
 
-			DetourMesh->removeTile(TileRef, nullptr, nullptr);
+				NumActiveTiles--;
+				UE_LOG(LogNavigation, Log, TEXT("%s> Tile (%d,%d:%d), removing TileRef: 0x%X (active:%d)"),
+					*DestNavMesh->GetName(), TileX, TileY, LayerIndex, TileRef, NumActiveTiles);
+
+				DetourMesh->removeTile(TileRef, nullptr, nullptr);
 
-			ResultTileIndices.AddUnique(DetourMesh->decodePolyIdTile(TileRef));
+				ResultTileIndices.AddUnique(DetourMesh->decodePolyIdTile(TileRef));
+			}
 		}
-	}
 
-	// Remove compressed tile cache layers
-	DestNavMesh->RemoveTileCacheLayers(TileX, TileY);
+		// Remove compressed tile cache layers
+		DestNavMesh->RemoveTileCacheLayers(TileX, TileY);
+	}
 
 	return ResultTileIndices;
 }
diff --git a/Engine/Source/Runtime/Navmesh/Public/Detour/DetourNavMesh.h b/Engine/Source/Runtime/Navmesh/Public/Detour/DetourNavMesh.h
--- a/Engine/Source/Runtime/Navmesh/Public/Detour/DetourNavMesh.h
+++ b/Engine/Source/Runtime/Navmesh/Public/Detour/DetourNavMesh.h
@@ -735,7 +735,10 @@ public:
 
 	inline bool isEmpty() const
 	{
-		return (m_tileWidth > 0 && m_tileHeight > 0) == false;
+		// has no tile grid set up
+		return (m_tileWidth > 0 && m_tileHeight > 0) == false
+			// or no actual tiles present (first free tile is the first one in tiles array)
+			|| (m_nextFree == m_tiles);
 	}
 	//@UE4 END
 	
