This updates the logic that determines if a curve should be deep copied when duplicating timelines.
#UE-4678 - LIVE: CRASH: Editor crashes when attempting to open a new level from a level that has an external curve in a timeline within the level blueprint
#Branch UE4
#Proj Engine

#info this changes seeks to only deep copy local curves when an object is duplicated outside of blueprint re-instancing, which incidentally is the only time I could find during testing when the owning package was transient. External curves are never deep copied.

diff --git a/Engine/Source/Runtime/Engine/Private/TimelineTemplate.cpp b/Engine/Source/Runtime/Engine/Private/TimelineTemplate.cpp
--- a/Engine/Source/Runtime/Engine/Private/TimelineTemplate.cpp
+++ b/Engine/Source/Runtime/Engine/Private/TimelineTemplate.cpp
@@ -231,16 +231,20 @@ void UTimelineTemplate::PostDuplicate(bool bDuplicateForPIE)
 #endif // WITH_EDITOR
 
 	const bool bTransientPackage = GetOutermost() == GetTransientPackage();
+	// Prevent curves being duplicated during blueprint reinstancing
+	const bool bDuplicateCurves = !( bTransientPackage || GIsDuplicatingClassForReinstancing );
 
 	for(TArray<struct FTTFloatTrack>::TIterator It = FloatTracks.CreateIterator();It;++It)
 	{
 		FTTFloatTrack& Track = *It;
 		if( Track.CurveFloat != NULL )
 		{
-			// Do not duplicate external curves unless duplicating to a transient package
-			if(!Track.CurveFloat->GetOuter()->IsA(UPackage::StaticClass()) || bTransientPackage )
+			if( bDuplicateCurves && !Track.bIsExternalCurve )
 			{
-				Track.CurveFloat = DuplicateObject<UCurveFloat>(Track.CurveFloat, NewCurveOuter, *MakeUniqueCurveName(Track.CurveFloat, NewCurveOuter));
+				if(!Track.CurveFloat->GetOuter()->IsA(UPackage::StaticClass()))
+				{
+					Track.CurveFloat = DuplicateObject<UCurveFloat>(Track.CurveFloat, NewCurveOuter, *MakeUniqueCurveName(Track.CurveFloat, NewCurveOuter));
+				}
 			}
 		}
 		else
@@ -254,10 +258,12 @@ void UTimelineTemplate::PostDuplicate(bool bDuplicateForPIE)
 		FTTEventTrack& Track = *It;
 		if( Track.CurveKeys != NULL )
 		{
-			// Do not duplicate external curves unless duplicating to a transient package
-			if(!Track.CurveKeys->GetOuter()->IsA(UPackage::StaticClass()) || bTransientPackage )
+			if( bDuplicateCurves && !Track.bIsExternalCurve )
 			{
-				Track.CurveKeys = DuplicateObject<UCurveFloat>(Track.CurveKeys, NewCurveOuter, *MakeUniqueCurveName(Track.CurveKeys, NewCurveOuter));
+				if(!Track.CurveKeys->GetOuter()->IsA(UPackage::StaticClass()))
+				{
+					Track.CurveKeys = DuplicateObject<UCurveFloat>(Track.CurveKeys, NewCurveOuter, *MakeUniqueCurveName(Track.CurveKeys, NewCurveOuter));
+				}
 			}
 		}
 		else
@@ -271,10 +277,12 @@ void UTimelineTemplate::PostDuplicate(bool bDuplicateForPIE)
 		FTTVectorTrack& Track = *It;
 		if( Track.CurveVector != NULL )
 		{
-			// Do not duplicate external curves unless duplicating to a transient package
-			if(!Track.CurveVector->GetOuter()->IsA(UPackage::StaticClass()) || bTransientPackage )
+			if( bDuplicateCurves && !Track.bIsExternalCurve )
 			{
-				Track.CurveVector = DuplicateObject<UCurveVector>(Track.CurveVector, NewCurveOuter, *MakeUniqueCurveName(Track.CurveVector, NewCurveOuter));
+				if(!Track.CurveVector->GetOuter()->IsA(UPackage::StaticClass()))
+				{
+					Track.CurveVector = DuplicateObject<UCurveVector>(Track.CurveVector, NewCurveOuter, *MakeUniqueCurveName(Track.CurveVector, NewCurveOuter));
+				}
 			}
 		}
 		else
@@ -288,10 +296,12 @@ void UTimelineTemplate::PostDuplicate(bool bDuplicateForPIE)
 		FTTLinearColorTrack& Track = *It;
 		if( Track.CurveLinearColor != NULL )
 		{
-			// Do not duplicate external curves unless duplicating to a transient package
-			if(!Track.CurveLinearColor->GetOuter()->IsA(UPackage::StaticClass()) && bTransientPackage )
+			if( bDuplicateCurves && !Track.bIsExternalCurve )
 			{
-				Track.CurveLinearColor = DuplicateObject<UCurveLinearColor>(Track.CurveLinearColor, NewCurveOuter, *MakeUniqueCurveName(Track.CurveLinearColor, NewCurveOuter));
+				if(!Track.CurveLinearColor->GetOuter()->IsA(UPackage::StaticClass()))
+				{
+					Track.CurveLinearColor = DuplicateObject<UCurveLinearColor>(Track.CurveLinearColor, NewCurveOuter, *MakeUniqueCurveName(Track.CurveLinearColor, NewCurveOuter));
+				}
 			}
 		}
 		else
