Added oriented box shape to Visual Logger.

diff --git a/Engine/Source/Developer/LogVisualizer/Private/VisualLoggerRenderingActor.cpp b/Engine/Source/Developer/LogVisualizer/Private/VisualLoggerRenderingActor.cpp
--- a/Engine/Source/Developer/LogVisualizer/Private/VisualLoggerRenderingActor.cpp
+++ b/Engine/Source/Developer/LogVisualizer/Private/VisualLoggerRenderingActor.cpp
@@ -116,6 +116,9 @@ void AVisualLoggerRenderingActor::AddDebugRendering()
 		const FVector BoxExtent(100, 100, 100);
 		const FBox Box(FVector(128), FVector(300));
 		Boxes.Add(FDebugRenderSceneProxy::FDebugBox(Box, FColor::Red));
+		FTransform Trans;
+		Trans.SetRotation(FQuat::MakeFromEuler(FVector(0.1, 0.2, 1.2)));
+		Boxes.Add(FDebugRenderSceneProxy::FDebugBox(Box, FColor::Red, Trans));
 	}
 	{
 		const FVector Orgin = FVector(400,0,128);
@@ -354,8 +357,8 @@ void AVisualLoggerRenderingActor::OnItemSelectionChanged(const FVisualLogDevice:
 			const FVector* BoxExtent = ElementToDraw->Points.GetData();
 			for (int32 Index = 0; Index + 1 < ElementToDraw->Points.Num(); Index += 2, BoxExtent += 2)
 			{
-				const FBox Box(*BoxExtent, *(BoxExtent + 1));
-				Boxes.Add(FDebugRenderSceneProxy::FDebugBox(Box, Color));
+				const FBox Box = FBox(*BoxExtent, *(BoxExtent + 1));
+				Boxes.Add(FDebugRenderSceneProxy::FDebugBox(Box, Color, FTransform(ElementToDraw->TransformationMatrix)));
 
 				if (bDrawLabel)
 				{
diff --git a/Engine/Source/Runtime/Engine/Private/DebugRenderSceneProxy.cpp b/Engine/Source/Runtime/Engine/Private/DebugRenderSceneProxy.cpp
--- a/Engine/Source/Runtime/Engine/Private/DebugRenderSceneProxy.cpp
+++ b/Engine/Source/Runtime/Engine/Private/DebugRenderSceneProxy.cpp
@@ -163,11 +163,11 @@ void FDebugRenderSceneProxy::GetDynamicMeshElements(const TArray<const FSceneVie
 			{
 				if (DrawType & SolidAndWireMeshes || DrawType & WireMesh)
 				{
-					DrawWireBox(PDI, Box.Box, Box.Color, SDPG_World, DrawType & SolidAndWireMeshes ? 2 : 0, 0, true);
+					DrawWireBox(PDI, Box.Transform.ToMatrixWithScale(), Box.Box, Box.Color, SDPG_World, DrawType & SolidAndWireMeshes ? 2 : 0, 0, true);
 				}
 				if (DrawType & SolidAndWireMeshes || DrawType & SolidMesh)
 				{
-					GetBoxMesh(FTransform(Box.Box.GetCenter()).ToMatrixNoScale(), Box.Box.GetExtent(), MaterialCache[Box.Color.WithAlpha(60)], SDPG_World, ViewIndex, Collector);
+					GetBoxMesh(FTransform(Box.Box.GetCenter()).ToMatrixNoScale() * Box.Transform.ToMatrixWithScale(), Box.Box.GetExtent(), MaterialCache[Box.Color.WithAlpha(60)], SDPG_World, ViewIndex, Collector);
 				}
 			}
 
diff --git a/Engine/Source/Runtime/Engine/Private/PrimitiveDrawingUtils.cpp b/Engine/Source/Runtime/Engine/Private/PrimitiveDrawingUtils.cpp
--- a/Engine/Source/Runtime/Engine/Private/PrimitiveDrawingUtils.cpp
+++ b/Engine/Source/Runtime/Engine/Private/PrimitiveDrawingUtils.cpp
@@ -844,7 +844,7 @@ void DrawWireBox(FPrimitiveDrawInterface* PDI, const FBox& Box, const FLinearCol
 }
 
 
-void DrawWireBox( FPrimitiveDrawInterface* PDI, const FMatrix& Matrix, const FBox& Box, const FLinearColor& Color, uint8 DepthPriority )
+void DrawWireBox(FPrimitiveDrawInterface* PDI, const FMatrix& Matrix, const FBox& Box, const FLinearColor& Color, uint8 DepthPriority, float Thickness, float DepthBias, bool bScreenSpace)
 {
 	FVector B[ 2 ];
 	B[ 0 ] = Box.Min;
@@ -860,19 +860,19 @@ void DrawWireBox( FPrimitiveDrawInterface* PDI, const FMatrix& Matrix, const FBo
 			P.Y = B[ j ].Y; Q.Y = B[ j ].Y;
 			P.Z = B[ 0 ].Z; Q.Z = B[ 1 ].Z;
 			P = Matrix.TransformPosition( P ); Q = Matrix.TransformPosition( Q );
-			PDI->DrawLine( P, Q, Color, DepthPriority );
+			PDI->DrawLine(P, Q, Color, DepthPriority, Thickness, DepthBias, bScreenSpace);
 
 			P.Y = B[ i ].Y; Q.Y = B[ i ].Y;
 			P.Z = B[ j ].Z; Q.Z = B[ j ].Z;
 			P.X = B[ 0 ].X; Q.X = B[ 1 ].X;
 			P = Matrix.TransformPosition( P ); Q = Matrix.TransformPosition( Q );
-			PDI->DrawLine( P, Q, Color, DepthPriority );
+			PDI->DrawLine(P, Q, Color, DepthPriority, Thickness, DepthBias, bScreenSpace);
 
 			P.Z = B[ i ].Z; Q.Z = B[ i ].Z;
 			P.X = B[ j ].X; Q.X = B[ j ].X;
 			P.Y = B[ 0 ].Y; Q.Y = B[ 1 ].Y;
 			P = Matrix.TransformPosition( P ); Q = Matrix.TransformPosition( Q );
-			PDI->DrawLine( P, Q, Color, DepthPriority );
+			PDI->DrawLine(P, Q, Color, DepthPriority, Thickness, DepthBias, bScreenSpace);
 		}
 	}
 }
diff --git a/Engine/Source/Runtime/Engine/Private/VisualLogger/VisualLoggerKismetLibrary.cpp b/Engine/Source/Runtime/Engine/Private/VisualLogger/VisualLoggerKismetLibrary.cpp
--- a/Engine/Source/Runtime/Engine/Private/VisualLogger/VisualLoggerKismetLibrary.cpp
+++ b/Engine/Source/Runtime/Engine/Private/VisualLogger/VisualLoggerKismetLibrary.cpp
@@ -30,6 +30,6 @@ void UVisualLoggerKismetLibrary::LogBox(UObject* WorldContextObject, FBox Box, F
 {
 #if ENABLE_VISUAL_LOG
 	const ELogVerbosity::Type DefaultVerbosity = ELogVerbosity::Log;
-	FVisualLogger::GeometryShapeLogf(WorldContextObject, FLogCategoryBase(*CategoryName.ToString(), DefaultVerbosity, DefaultVerbosity), DefaultVerbosity, INDEX_NONE, Box, ObjectColor, *Text);
+	FVisualLogger::GeometryShapeLogf(WorldContextObject, FLogCategoryBase(*CategoryName.ToString(), DefaultVerbosity, DefaultVerbosity), DefaultVerbosity, INDEX_NONE, Box, FMatrix::Identity, ObjectColor, *Text);
 #endif
 }
diff --git a/Engine/Source/Runtime/Engine/Private/VisualLogger/VisualLoggerTypes.cpp b/Engine/Source/Runtime/Engine/Private/VisualLogger/VisualLoggerTypes.cpp
--- a/Engine/Source/Runtime/Engine/Private/VisualLogger/VisualLoggerTypes.cpp
+++ b/Engine/Source/Runtime/Engine/Private/VisualLogger/VisualLoggerTypes.cpp
@@ -140,7 +140,7 @@ void FVisualLogEntry::AddElement(const FVector& Start, const FVector& End, const
 	ElementsToDraw.Add(Element);
 }
 
-void FVisualLogEntry::AddElement(const FBox& Box, const FName& CategoryName, ELogVerbosity::Type Verbosity, const FColor& Color, const FString& Description, uint16 Thickness)
+void FVisualLogEntry::AddElement(const FBox& Box, const FMatrix& Matrix, const FName& CategoryName, ELogVerbosity::Type Verbosity, const FColor& Color, const FString& Description, uint16 Thickness)
 {
 	FVisualLogShapeElement Element(Description, Color, Thickness, CategoryName);
 	Element.Points.Reserve(2);
@@ -148,6 +148,7 @@ void FVisualLogEntry::AddElement(const FBox& Box, const FName& CategoryName, ELo
 	Element.Points.Add(Box.Max);
 	Element.Type = EVisualLoggerShapeElement::Box;
 	Element.Verbosity = Verbosity;
+	Element.TransformationMatrix = Matrix;
 	ElementsToDraw.Add(Element);
 }
 
@@ -252,6 +253,13 @@ FArchive& operator<<(FArchive& Ar, FVisualLogShapeElement& Element)
 	FVisualLoggerHelpers::Serialize(Ar, Element.Category);
 	Ar << Element.Description;
 	Ar << Element.Verbosity;
+	const int32 VLogsVer = Ar.CustomVer(EVisualLoggerVersion::GUID);
+
+	if (VLogsVer >= EVisualLoggerVersion::TransformationForShapes)
+	{
+		Ar << Element.TransformationMatrix;
+	}
+
 	Ar << Element.Points;
 	Ar << Element.UniqueId;
 	Ar << Element.Type;
diff --git a/Engine/Source/Runtime/Engine/Public/DebugRenderSceneProxy.h b/Engine/Source/Runtime/Engine/Public/DebugRenderSceneProxy.h
--- a/Engine/Source/Runtime/Engine/Public/DebugRenderSceneProxy.h
+++ b/Engine/Source/Runtime/Engine/Public/DebugRenderSceneProxy.h
@@ -82,13 +82,19 @@ public:
 	/** Struct to hold info about boxes to render. */
 	struct FDebugBox
 	{
-		FDebugBox( const FBox& InBox, const FColor& InColor )
-			: Box( InBox ), Color( InColor )
+		FDebugBox(const FBox& InBox, const FColor& InColor)
+			: Box(InBox), Color(InColor)
+		{
+		}
+
+		FDebugBox(const FBox& InBox, const FColor& InColor, const FTransform& InTransform)
+			: Box(InBox), Color(InColor), Transform(InTransform)
 		{
 		}
 
 		FBox Box;
 		FColor Color;
+		FTransform Transform;
 	};
 
 	/** Struct to hold info about cylinders to render. */
diff --git a/Engine/Source/Runtime/Engine/Public/SceneManagement.h b/Engine/Source/Runtime/Engine/Public/SceneManagement.h
--- a/Engine/Source/Runtime/Engine/Public/SceneManagement.h
+++ b/Engine/Source/Runtime/Engine/Public/SceneManagement.h
@@ -1554,7 +1554,7 @@ extern ENGINE_API void DrawFlatArrow(class FPrimitiveDrawInterface* PDI,const FV
 
 // Line drawing utility functions.
 extern ENGINE_API void DrawWireBox(class FPrimitiveDrawInterface* PDI, const FBox& Box, const FLinearColor& Color, uint8 DepthPriority, float Thickness = 0, float DepthBias = 0.0f, bool bScreenSpace = false);
-extern ENGINE_API void DrawWireBox(class FPrimitiveDrawInterface* PDI, const FMatrix& Matrix, const FBox& Box, const FLinearColor& Color, uint8 DepthPriority );
+extern ENGINE_API void DrawWireBox(class FPrimitiveDrawInterface* PDI, const FMatrix& Matrix, const FBox& Box, const FLinearColor& Color, uint8 DepthPriority, float Thickness = 0, float DepthBias = 0.0f, bool bScreenSpace = false);
 extern ENGINE_API void DrawCircle(class FPrimitiveDrawInterface* PDI, const FVector& Base, const FVector& X, const FVector& Y, const FLinearColor& Color, float Radius, int32 NumSides, uint8 DepthPriority, float Thickness = 0, float DepthBias = 0.0f, bool bScreenSpace = false);
 extern ENGINE_API void DrawArc(FPrimitiveDrawInterface* PDI, const FVector Base, const FVector X, const FVector Y, const float MinAngle, const float MaxAngle, const float Radius, const int32 Sections, const FLinearColor& Color, uint8 DepthPriority);
 extern ENGINE_API void DrawWireSphere(class FPrimitiveDrawInterface* PDI, const FVector& Base, const FLinearColor& Color, float Radius, int32 NumSides, uint8 DepthPriority, float Thickness = 0, float DepthBias = 0.0f, bool bScreenSpace = false);
diff --git a/Engine/Source/Runtime/Engine/Public/VisualLogger/VisualLogger.h b/Engine/Source/Runtime/Engine/Public/VisualLogger/VisualLogger.h
--- a/Engine/Source/Runtime/Engine/Public/VisualLogger/VisualLogger.h
+++ b/Engine/Source/Runtime/Engine/Public/VisualLogger/VisualLogger.h
@@ -27,8 +27,11 @@
 #define UE_VLOG_LOCATION(LogOwner, CategoryName, Verbosity, Location, Radius, Color, Format, ...) FVisualLogger::GeometryShapeLogf(LogOwner, CategoryName, ELogVerbosity::Verbosity, INDEX_NONE, Location, Radius, Color, Format, ##__VA_ARGS__)
 #define UE_CVLOG_LOCATION(Condition, LogOwner, CategoryName, Verbosity, Location, Radius, Color, Format, ...)  if(Condition) {UE_VLOG_LOCATION(LogOwner, CategoryName, Verbosity, Location, Radius, Color, Format, ##__VA_ARGS__);} 
 // Box shape
-#define UE_VLOG_BOX(LogOwner, CategoryName, Verbosity, Box, Color, Format, ...)  FVisualLogger::GeometryShapeLogf(LogOwner, CategoryName, ELogVerbosity::Verbosity, INDEX_NONE, Box, Color, Format, ##__VA_ARGS__)
-#define UE_CVLOG_BOX(Condition, LogOwner, CategoryName, Verbosity, Box, Color, Format, ...)  if(Condition) {UE_VLOG_BOX(Actor, CategoryName, Verbosity, Box, Color, Format, ##__VA_ARGS__);} 
+#define UE_VLOG_BOX(LogOwner, CategoryName, Verbosity, Box, Color, Format, ...)  FVisualLogger::GeometryShapeLogf(LogOwner, CategoryName, ELogVerbosity::Verbosity, INDEX_NONE, Box, FMatrix::Identity, Color, Format, ##__VA_ARGS__)
+#define UE_CVLOG_BOX(Condition, LogOwner, CategoryName, Verbosity, Box, Color, Format, ...)  if(Condition) {UE_VLOG_BOX(Actor, CategoryName, Verbosity, Box, FMatrix::Identity, Color, Format, ##__VA_ARGS__);} 
+// Oriented box shape
+#define UE_VLOG_OBOX(LogOwner, CategoryName, Verbosity, Box, Matrix, Color, Format, ...)  FVisualLogger::GeometryShapeLogf(LogOwner, CategoryName, ELogVerbosity::Verbosity, INDEX_NONE, Box, Matrix, Color, Format, ##__VA_ARGS__)
+#define UE_CVLOG_OBOX(Condition, LogOwner, CategoryName, Verbosity, Box, Matrix, Color, Format, ...)  if(Condition) {UE_VLOG_BOX(Actor, CategoryName, Verbosity, Box, Matrix, Color, Format, ##__VA_ARGS__);} 
 // Cone shape
 #define UE_VLOG_CONE(LogOwner, CategoryName, Verbosity, Orgin, Direction, Length, Angle, Color, Format, ...) FVisualLogger::GeometryShapeLogf(LogOwner, CategoryName, ELogVerbosity::Verbosity, INDEX_NONE, Orgin, Direction, Length, Angle, Color, Format, ##__VA_ARGS__)
 #define UE_CVLOG_CONE(Condition, LogOwner, CategoryName, Verbosity, Orgin, Direction, Length, Angle, Color, Format, ...)  if(Condition) {UE_VLOG_CONE(Actor, CategoryName, Verbosity, Orgin, Direction, Length, Angle, Color, Format, ##__VA_ARGS__);} 
@@ -68,6 +71,8 @@
 #define UE_CVLOG_LOCATION(Condition, Actor, CategoryName, Verbosity, Location, Radius, Color, DescriptionFormat, ...)
 #define UE_VLOG_BOX(Actor, CategoryName, Verbosity, Box, Color, DescriptionFormat, ...) 
 #define UE_CVLOG_BOX(Condition, Actor, CategoryName, Verbosity, Box, Color, DescriptionFormat, ...) 
+#define UE_VLOG_OBOX(LogOwner, CategoryName, Verbosity, Box, Matrix, Color, Format, ...)  
+#define UE_CVLOG_OBOX(Condition, LogOwner, CategoryName, Verbosity, Box, Matrix, Color, Format, ...) 
 #define UE_VLOG_CONE(Object, CategoryName, Verbosity, Orgin, Direction, Length, Angle, Color, DescriptionFormat, ...)
 #define UE_CVLOG_CONE(Condition, Object, CategoryName, Verbosity, Orgin, Direction, Length, Angle, Color, DescriptionFormat, ...)
 #define UE_VLOG_CYLINDER(Object, CategoryName, Verbosity, Start, End, Radius, Color, DescriptionFormat, ...)
@@ -113,7 +118,7 @@ public:
 	// Location/Shpere log
 	VARARG_DECL(static void, static void, return, GeometryShapeLogf, VARARG_NONE, const TCHAR*, VARARG_EXTRA(const class UObject* LogOwner) VARARG_EXTRA(const struct FLogCategoryBase& Category) VARARG_EXTRA(ELogVerbosity::Type Verbosity) VARARG_EXTRA(int32 UniqueLogId) VARARG_EXTRA(const FVector& Location) VARARG_EXTRA(float Radius) VARARG_EXTRA(const FColor& Color), VARARG_EXTRA(LogOwner) VARARG_EXTRA(Category) VARARG_EXTRA(Verbosity) VARARG_EXTRA(UniqueLogId) VARARG_EXTRA(Location) VARARG_EXTRA(Radius) VARARG_EXTRA(Color));
 	// Box log
-	VARARG_DECL(static void, static void, return, GeometryShapeLogf, VARARG_NONE, const TCHAR*, VARARG_EXTRA(const class UObject* LogOwner) VARARG_EXTRA(const struct FLogCategoryBase& Category) VARARG_EXTRA(ELogVerbosity::Type Verbosity) VARARG_EXTRA(int32 UniqueLogId) VARARG_EXTRA(const FBox& Box) VARARG_EXTRA(const FColor& Color), VARARG_EXTRA(LogOwner) VARARG_EXTRA(Category) VARARG_EXTRA(Verbosity) VARARG_EXTRA(UniqueLogId) VARARG_EXTRA(Box) VARARG_EXTRA(Color));
+	VARARG_DECL(static void, static void, return, GeometryShapeLogf, VARARG_NONE, const TCHAR*, VARARG_EXTRA(const class UObject* LogOwner) VARARG_EXTRA(const struct FLogCategoryBase& Category) VARARG_EXTRA(ELogVerbosity::Type Verbosity) VARARG_EXTRA(int32 UniqueLogId) VARARG_EXTRA(const FBox& Box) VARARG_EXTRA(const FMatrix& Matrix) VARARG_EXTRA(const FColor& Color), VARARG_EXTRA(LogOwner) VARARG_EXTRA(Category) VARARG_EXTRA(Verbosity) VARARG_EXTRA(UniqueLogId) VARARG_EXTRA(Box) VARARG_EXTRA(Matrix) VARARG_EXTRA(Color));
 	// Cone log
 	VARARG_DECL(static void, static void, return, GeometryShapeLogf, VARARG_NONE, const TCHAR*, VARARG_EXTRA(const class UObject* LogOwner) VARARG_EXTRA(const struct FLogCategoryBase& Category) VARARG_EXTRA(ELogVerbosity::Type Verbosity) VARARG_EXTRA(int32 UniqueLogId) VARARG_EXTRA(const FVector& Orgin) VARARG_EXTRA(const FVector& Direction) VARARG_EXTRA(const float Length) VARARG_EXTRA(const float Angle) VARARG_EXTRA(const FColor& Color), VARARG_EXTRA(LogOwner) VARARG_EXTRA(Category) VARARG_EXTRA(Verbosity) VARARG_EXTRA(UniqueLogId) VARARG_EXTRA(Orgin) VARARG_EXTRA(Direction) VARARG_EXTRA(Length) VARARG_EXTRA(Angle) VARARG_EXTRA(Color));
 	// Cylinder log
diff --git a/Engine/Source/Runtime/Engine/Public/VisualLogger/VisualLogger.inl b/Engine/Source/Runtime/Engine/Public/VisualLogger/VisualLogger.inl
--- a/Engine/Source/Runtime/Engine/Public/VisualLogger/VisualLogger.inl
+++ b/Engine/Source/Runtime/Engine/Public/VisualLogger/VisualLogger.inl
@@ -228,7 +228,7 @@ VARARG_BODY(void, FVisualLogger::GeometryShapeLogf, const TCHAR*, VARARG_EXTRA(c
 }
 
 FORCEINLINE
-VARARG_BODY(void, FVisualLogger::GeometryShapeLogf, const TCHAR*, VARARG_EXTRA(const class UObject* Object) VARARG_EXTRA(const struct FLogCategoryBase& Category) VARARG_EXTRA(ELogVerbosity::Type Verbosity) VARARG_EXTRA(int32 UniqueLogId) VARARG_EXTRA(const FBox& Box) VARARG_EXTRA(const FColor& Color))
+VARARG_BODY(void, FVisualLogger::GeometryShapeLogf, const TCHAR*, VARARG_EXTRA(const class UObject* Object) VARARG_EXTRA(const struct FLogCategoryBase& Category) VARARG_EXTRA(ELogVerbosity::Type Verbosity) VARARG_EXTRA(int32 UniqueLogId) VARARG_EXTRA(const FBox& Box) VARARG_EXTRA(const FMatrix& Matrix) VARARG_EXTRA(const FColor& Color))
 {
 	SCOPE_CYCLE_COUNTER(STAT_VisualLog);
 	UWorld *World = NULL;
@@ -239,7 +239,7 @@ VARARG_BODY(void, FVisualLogger::GeometryShapeLogf, const TCHAR*, VARARG_EXTRA(c
 	}
 
 	COLLAPSED_LOGF(
-		CurrentEntry->AddElement(Box, Category.GetCategoryName(), Verbosity, Color, Buffer);
+		CurrentEntry->AddElement(Box, Matrix, Category.GetCategoryName(), Verbosity, Color, Buffer);
 	);
 }
 
diff --git a/Engine/Source/Runtime/Engine/Public/VisualLogger/VisualLoggerTypes.h b/Engine/Source/Runtime/Engine/Public/VisualLogger/VisualLoggerTypes.h
--- a/Engine/Source/Runtime/Engine/Public/VisualLogger/VisualLoggerTypes.h
+++ b/Engine/Source/Runtime/Engine/Public/VisualLogger/VisualLoggerTypes.h
@@ -29,6 +29,7 @@ namespace EVisualLoggerVersion
 		HistogramGraphsSerialization = 1,
 		AddedOwnerClassName = 2,
 		StatusCategoryWithChildren = 3,
+		TransformationForShapes = 4,
 		// -----<new versions can be added before this line>-------------------------------------------------
 		// - this needs to be the last line (see note below)
 		VersionPlusOne,
@@ -115,6 +116,7 @@ struct ENGINE_API FVisualLogShapeElement
 	FName Category;
 	TEnumAsByte<ELogVerbosity::Type> Verbosity;
 	TArray<FVector> Points;
+	FMatrix TransformationMatrix;
 	int32 UniqueId;
 	EVisualLoggerShapeElement Type;
 	uint8 Color;
@@ -180,7 +182,7 @@ struct ENGINE_API FVisualLogEntry
 	// segment
 	void AddElement(const FVector& Start, const FVector& End, const FName& CategoryName, ELogVerbosity::Type Verbosity, const FColor& Color = FColor::White, const FString& Description = TEXT(""), uint16 Thickness = 0);
 	// box
-	void AddElement(const FBox& Box, const FName& CategoryName, ELogVerbosity::Type Verbosity, const FColor& Color = FColor::White, const FString& Description = TEXT(""), uint16 Thickness = 0);
+	void AddElement(const FBox& Box, const FMatrix& Matrix, const FName& CategoryName, ELogVerbosity::Type Verbosity, const FColor& Color = FColor::White, const FString& Description = TEXT(""), uint16 Thickness = 0);
 	// Cone
 	void AddElement(const FVector& Orgin, const FVector& Direction, float Length, float AngleWidth, float AngleHeight, const FName& CategoryName, ELogVerbosity::Type Verbosity, const FColor& Color = FColor::White, const FString& Description = TEXT(""), uint16 Thickness = 0);
 	// Cylinder
@@ -323,6 +325,7 @@ void FVisualLogStatusCategory::AddChild(const FVisualLogStatusCategory& Child)
 FORCEINLINE
 FVisualLogShapeElement::FVisualLogShapeElement(EVisualLoggerShapeElement InType)
 : Verbosity(ELogVerbosity::All)
+, TransformationMatrix(FMatrix::Identity)
 , Type(InType)
 , Color(0xff)
 , Thicknes(0)
@@ -334,6 +337,7 @@ FORCEINLINE
 FVisualLogShapeElement::FVisualLogShapeElement(const FString& InDescription, const FColor& InColor, uint16 InThickness, const FName& InCategory)
 : Category(InCategory)
 , Verbosity(ELogVerbosity::All)
+, TransformationMatrix(FMatrix::Identity)
 , Type(EVisualLoggerShapeElement::Invalid)
 , Thicknes(InThickness)
 {
