Merging from John Barrett:

Added commandlet for enabling/disabling plugins, to aid automation tests; non-shipping only.

diff --git a/Engine/Source/Runtime/Engine/Classes/Commandlets/PluginCommandlet.h b/Engine/Source/Runtime/Engine/Classes/Commandlets/PluginCommandlet.h
new file mode 100644
--- /dev/null
+++ b/Engine/Source/Runtime/Engine/Classes/Commandlets/PluginCommandlet.h
@@ -0,0 +1,20 @@
+// Copyright 1998-2015 Epic Games, Inc. All Rights Reserved.
+
+#pragma once
+#include "Commandlets/Commandlet.h"
+#include "PluginCommandlet.generated.h"
+
+UCLASS()
+class UPluginCommandlet : public UCommandlet
+{
+	GENERATED_UCLASS_BODY()
+
+	/** Parsed commandline tokens */
+	TArray<FString> CmdLineTokens;
+
+	/** Parsed commandline switches */
+	TArray<FString> CmdLineSwitches;
+
+
+	virtual int32 Main(const FString& Params) override;
+};
diff --git a/Engine/Source/Runtime/Engine/Engine.Build.cs b/Engine/Source/Runtime/Engine/Engine.Build.cs
--- a/Engine/Source/Runtime/Engine/Engine.Build.cs
+++ b/Engine/Source/Runtime/Engine/Engine.Build.cs
@@ -68,6 +68,7 @@ public class Engine : ModuleRules
 				"VectorVM",
 				"Landscape",
                 "UMG",
+				"Projects"
 			}
         );
 
diff --git a/Engine/Source/Runtime/Engine/Private/Commandlets/PluginCommandlet.cpp b/Engine/Source/Runtime/Engine/Private/Commandlets/PluginCommandlet.cpp
new file mode 100644
--- /dev/null
+++ b/Engine/Source/Runtime/Engine/Private/Commandlets/PluginCommandlet.cpp
@@ -0,0 +1,112 @@
+// Copyright 1998-2015 Epic Games, Inc. All Rights Reserved.
+
+#include "EnginePrivate.h"
+#include "Commandlets/PluginCommandlet.h"
+
+#include "Projects.h"
+
+
+DEFINE_LOG_CATEGORY_STATIC(LogPluginCommandlet, Log, All);
+
+/**
+ * UPluginCommandlet
+ *
+ * Commandlet used for enabling/disabling plugins
+ *
+ * Usage:
+ *	Plugin Enable/Disable PluginName
+ *
+ * Example:
+ *	Plugin Enable NetcodeUnitTest
+ */
+
+UPluginCommandlet::UPluginCommandlet(const FObjectInitializer& ObjectInitializer)
+	: Super(ObjectInitializer)
+{
+	IsClient = false;
+	IsServer = false;
+	IsEditor = false;
+}
+
+int32 UPluginCommandlet::Main(const FString& Params)
+{
+	bool bSuccess = false;
+
+#if !UE_BUILD_SHIPPING
+	const TCHAR* ParamStr = *Params;
+	ParseCommandLine(ParamStr, CmdLineTokens, CmdLineSwitches);
+
+	if (CmdLineTokens.Num() > 0)
+	{
+		bool bEnable = false;
+
+		int32 CommandIdx = CmdLineTokens.IndexOfByPredicate(
+			[&bEnable](const FString& CurToken)
+			{
+				bool bReturnVal = false;
+
+				if (CurToken == TEXT("Enable") || CurToken == TEXT("Disable"))
+				{
+					bEnable = CurToken == TEXT("Enable");
+					bReturnVal = true;
+				}
+
+				return bReturnVal;
+			});
+
+		if (CommandIdx != INDEX_NONE)
+		{
+			// Trim all tokens up to the command, to keep things simple
+			CmdLineTokens.RemoveAt(0, CommandIdx+1);
+
+			FString PluginName = (CmdLineTokens.Num() > 0 ? CmdLineTokens.Pop() : TEXT(""));
+
+			if (PluginName.Len() > 0)
+			{
+				if (IProjectManager::Get().LoadProjectFile(FPaths::GetProjectFilePath()))
+				{
+					FText FailReason;
+
+					bSuccess = IProjectManager::Get().SetPluginEnabled(PluginName, bEnable, FailReason);
+
+					if (bSuccess)
+					{
+						UE_LOG(LogPluginCommandlet, Log, TEXT("Successfully %s plugin '%s'"),
+								(bEnable ? TEXT("enabled") : TEXT("disabled")), *PluginName);
+					}
+					else
+					{
+						UE_LOG(LogPluginCommandlet, Error, TEXT("Failed to %s plugin '%s' - error: %s"),
+								(bEnable ? TEXT("enable") : TEXT("disable")), *PluginName, *FailReason.ToString());
+					}
+				}
+				else
+				{
+					UE_LOG(LogPluginCommandlet, Error, TEXT("Failed to load project file."));
+					bSuccess = false;
+				}
+			}
+			else
+			{
+				UE_LOG(LogPluginCommandlet, Error, TEXT("Failed to parse plugin name."));
+				bSuccess = false;
+			}
+		}
+		else
+		{
+			UE_LOG(LogPluginCommandlet, Error, TEXT("No command specified or unknown command."));
+			bSuccess = false;
+		}
+	}
+	else
+	{
+		UE_LOG(LogPluginCommandlet, Error, TEXT("No command specified. Example: 'PluginCommandlet Enable NetcodeUnitTest'"));
+		bSuccess = false;
+	}
+#else
+	UE_LOG(LogPluginCommandlet, Error, TEXT("Plugin commandlet disabled in shipping mode."));
+	bSuccess = false;
+#endif
+
+	return bSuccess ? 0 : 1;
+}
