UEngine::MakeSureMapNameIsValid now checks for packages that may only be in memory, which is the case for a PIE client trying to load an unsaved map.
Fixes UE-6583.

diff --git a/Engine/Source/Runtime/Engine/Private/UnrealEngine.cpp b/Engine/Source/Runtime/Engine/Private/UnrealEngine.cpp
--- a/Engine/Source/Runtime/Engine/Private/UnrealEngine.cpp
+++ b/Engine/Source/Runtime/Engine/Private/UnrealEngine.cpp
@@ -8048,7 +8048,9 @@ bool UEngine::MakeSureMapNameIsValid(FString& InOutMapName)
 	bool bIsValid = !FPackageName::IsShortPackageName(TestMapName);
 	if (bIsValid)
 	{
-		bIsValid = FPackageName::DoesPackageExist(TestMapName);
+		// If the user starts a multiplayer PIE session with an unsaved map,
+		// DoesPackageExist won't find it, so we have to try to find the package in memory as well.
+		bIsValid = (FindObjectFast<UPackage>(nullptr, FName(*TestMapName)) != nullptr) || FPackageName::DoesPackageExist(TestMapName);
 	}
 	else
 	{
