xMake FComponentInstanceDataBase::MatchesComponent non-virtual and of subclasses saying that the component didn't match when the cache is being considered, have the Apply function decide not to apply the cache data based on the contents.

diff --git a/Engine/Source/Runtime/Engine/Private/Atmosphere/Atmosphere.cpp b/Engine/Source/Runtime/Engine/Private/Atmosphere/Atmosphere.cpp
--- a/Engine/Source/Runtime/Engine/Private/Atmosphere/Atmosphere.cpp
+++ b/Engine/Source/Runtime/Engine/Private/Atmosphere/Atmosphere.cpp
@@ -705,11 +705,6 @@ public:
 	virtual ~FAtmospherePrecomputeInstanceData()
 	{}
 
-	virtual bool MatchesComponent(const UActorComponent* Component) const override
-	{
-		return (PrecomputeParameter == CastChecked<UAtmosphericFogComponent>(Component)->GetPrecomputeParameters());
-	}
-
 	struct FAtmospherePrecomputeParameters PrecomputeParameter;
 
 	FByteBulkData TransmittanceData;
@@ -769,6 +764,11 @@ void UAtmosphericFogComponent::ApplyComponentInstanceData(FComponentInstanceData
 	check(ComponentInstanceData);
 	FAtmospherePrecomputeInstanceData* PrecomputedData = static_cast<FAtmospherePrecomputeInstanceData*>(const_cast<FComponentInstanceDataBase*>(ComponentInstanceData));
 
+	if (PrecomputedData->PrecomputeParameter != GetPrecomputeParameters())
+	{
+		return;
+	}
+
 	FComponentReregisterContext ReregisterContext(this);
 	ReleaseResource();
 
diff --git a/Engine/Source/Runtime/Engine/Private/InstancedStaticMesh.cpp b/Engine/Source/Runtime/Engine/Private/InstancedStaticMesh.cpp
--- a/Engine/Source/Runtime/Engine/Private/InstancedStaticMesh.cpp
+++ b/Engine/Source/Runtime/Engine/Private/InstancedStaticMesh.cpp
@@ -694,11 +694,6 @@ public:
 	{
 	}
 
-	virtual bool MatchesComponent(const UActorComponent* Component) const override
-	{
-		return (CastChecked<UStaticMeshComponent>(Component)->StaticMesh == StaticMesh && FComponentInstanceDataBase::MatchesComponent(Component));
-	}
-
 	/** Used to store lightmap data during RerunConstructionScripts */
 	struct FLightMapInstanceData
 	{
@@ -780,6 +775,11 @@ void UInstancedStaticMeshComponent::ApplyComponentInstanceData(FComponentInstanc
 
 	FInstancedStaticMeshComponentInstanceData* InstancedMeshData  = static_cast<FInstancedStaticMeshComponentInstanceData*>(ComponentInstanceData);
 
+	if (StaticMesh != InstancedMeshData->StaticMesh)
+	{
+		return;
+	}
+
 	// See if data matches current state
 	if (InstancedMeshData->bHasCachedStaticLighting)
 	{
diff --git a/Engine/Source/Runtime/Engine/Private/LightComponent.cpp b/Engine/Source/Runtime/Engine/Private/LightComponent.cpp
--- a/Engine/Source/Runtime/Engine/Private/LightComponent.cpp
+++ b/Engine/Source/Runtime/Engine/Private/LightComponent.cpp
@@ -878,13 +878,6 @@ public:
 	virtual ~FPrecomputedLightInstanceData()
 	{}
 
-	// Begin FComponentInstanceDataBase interface
-	virtual bool MatchesComponent(const UActorComponent* Component) const override
-	{
-		return Transform.Equals(CastChecked<ULightComponent>(Component)->ComponentToWorld);
-	}
-	// End FComponentInstanceDataBase interface
-
 	FTransform Transform;
 	FGuid LightGuid;
 	int32 ShadowMapChannel;
@@ -909,6 +902,11 @@ void ULightComponent::ApplyComponentInstanceData(FComponentInstanceDataBase* Com
 	check(ComponentInstanceData);
 	FPrecomputedLightInstanceData* LightMapData  = static_cast<FPrecomputedLightInstanceData*>(ComponentInstanceData);
 
+	if (!LightMapData->Transform.Equals(ComponentToWorld))
+	{
+		return;
+	}
+
 	LightGuid = LightMapData->LightGuid;
 	ShadowMapChannel = LightMapData->ShadowMapChannel;
 	PreviewShadowMapChannel = LightMapData->PreviewShadowMapChannel;
diff --git a/Engine/Source/Runtime/Engine/Private/StaticMeshComponent.cpp b/Engine/Source/Runtime/Engine/Private/StaticMeshComponent.cpp
--- a/Engine/Source/Runtime/Engine/Private/StaticMeshComponent.cpp
+++ b/Engine/Source/Runtime/Engine/Private/StaticMeshComponent.cpp
@@ -33,11 +33,6 @@ public:
 	{
 	}
 
-	virtual bool MatchesComponent(const UActorComponent* Component) const override
-	{
-		return (CastChecked<UStaticMeshComponent>(Component)->StaticMesh == StaticMesh && FComponentInstanceDataBase::MatchesComponent(Component));
-	}
-
 	/** Add vertex color data for a specified LOD before RerunConstructionScripts is called */
 	void AddVertexColorData(const struct FStaticMeshComponentLODInfo& LODInfo, uint32 LODIndex)
 	{
@@ -1580,6 +1575,11 @@ void UStaticMeshComponent::ApplyComponentInstanceData(FComponentInstanceDataBase
 	// That means all component state that is modified here must be mirrored on the scene proxy, which will be recreated to receive the changes later due to MarkRenderStateDirty.
 	FStaticMeshComponentInstanceData* StaticMeshInstanceData  = static_cast<FStaticMeshComponentInstanceData*>(ComponentInstanceData);
 
+	if (StaticMesh != StaticMeshInstanceData->StaticMesh)
+	{
+		return;
+	}
+
 	// See if data matches current state
 	if(	StaticMeshInstanceData->bHasCachedStaticLighting && StaticMeshInstanceData->CachedStaticLighting.Transform.Equals(ComponentToWorld) )
 	{
diff --git a/Engine/Source/Runtime/Engine/Public/ComponentInstanceDataCache.h b/Engine/Source/Runtime/Engine/Public/ComponentInstanceDataCache.h
--- a/Engine/Source/Runtime/Engine/Public/ComponentInstanceDataCache.h
+++ b/Engine/Source/Runtime/Engine/Public/ComponentInstanceDataCache.h
@@ -18,7 +18,7 @@ public:
 	virtual ~FComponentInstanceDataBase()
 	{}
 
-	virtual bool MatchesComponent(const UActorComponent* Component) const;
+	bool MatchesComponent(const UActorComponent* Component) const;
 
 protected:
 	/** The name of the source component */
diff --git a/Engine/Source/Runtime/UMG/Private/Components/WidgetComponent.cpp b/Engine/Source/Runtime/UMG/Private/Components/WidgetComponent.cpp
--- a/Engine/Source/Runtime/UMG/Private/Components/WidgetComponent.cpp
+++ b/Engine/Source/Runtime/UMG/Private/Components/WidgetComponent.cpp
@@ -577,11 +577,6 @@ public:
 		, RenderTarget( SourceComponent->GetRenderTarget() )
 	{}
 
-	virtual bool MatchesComponent( const UActorComponent* Component ) const override
-	{
-		return (CastChecked<UWidgetComponent>(Component)->GetWidgetClass() == WidgetClass && FComponentInstanceDataBase::MatchesComponent(Component));
-	}
-
 public:
 	TSubclassOf<UUserWidget> WidgetClass;
 	UTextureRenderTarget2D* RenderTarget;
@@ -606,6 +601,11 @@ void UWidgetComponent::ApplyComponentInstanceData(FComponentInstanceDataBase* Co
 	// That means all component state that is modified here must be mirrored on the scene proxy, which will be recreated to receive the changes later due to MarkRenderStateDirty.
 	FWidgetComponentInstanceData* WidgetInstanceData  = static_cast<FWidgetComponentInstanceData*>(ComponentInstanceData);
 
+	if (GetWidgetClass() != WidgetClass)
+	{
+		return;
+	}
+
 	RenderTarget = WidgetInstanceData->RenderTarget;
 	if( MaterialInstance )
 	{
