Add empty function to quad tree

diff --git a/Engine/Source/Runtime/Engine/Public/GenericQuadTree.h b/Engine/Source/Runtime/Engine/Public/GenericQuadTree.h
--- a/Engine/Source/Runtime/Engine/Public/GenericQuadTree.h
+++ b/Engine/Source/Runtime/Engine/Public/GenericQuadTree.h
@@ -23,6 +23,9 @@ public:
 	/** Does a deep copy of the tree by going through and re-creating the internal data. Cheaper than re-insertion as it should be linear instead of nlogn */
 	void Duplicate(TreeType* OutDuplicate) const;
 
+	/** Removes all elements of the tree */
+	void Empty();
+
 	~TQuadTree();
 
 private:
@@ -314,4 +317,21 @@ void TQuadTree<ElementType, NodeCapacity>::Duplicate(TreeType* OutDuplicate) con
 	OutDuplicate->TreeBox = TreeBox;
 	OutDuplicate->Position = Position;
 	OutDuplicate->bInternal = bInternal;
+}
+
+template <typename ElementType, int32 NodeCapacity>
+void TQuadTree<ElementType, NodeCapacity>::Empty()
+{
+	for (int32 TreeIdx = 0; TreeIdx < 4; ++TreeIdx)
+	{
+		if (TreeType* SubTree = SubTrees[TreeIdx])
+		{
+			delete SubTree;
+			SubTrees[TreeIdx] = nullptr;
+		}
+
+	}
+	
+	Nodes.Empty();
+	bInternal = false;
 }
\ No newline at end of file
