#github PR748, fix for Log Reverse Attenuation Function from korypostma

Taking suggest fix from pull req but keeping the clamp for Max(foo, 0) because as Distance/Falloff becomes greater than 0.99, we'll get negative attenuations.

Note of course, logf(0.0f) results in -1.#INF but this will test as a negative number (and has precedence in the ATTENUATION_Logarithmic case as well for when Distance == 0.0f).

diff --git a/Engine/Source/Runtime/Engine/Private/SoundAttenuation.cpp b/Engine/Source/Runtime/Engine/Private/SoundAttenuation.cpp
--- a/Engine/Source/Runtime/Engine/Private/SoundAttenuation.cpp
+++ b/Engine/Source/Runtime/Engine/Private/SoundAttenuation.cpp
@@ -102,7 +102,7 @@ float AttenuationEval(const ESoundDistanceModel DistanceModel, const float Dista
 
 			case ATTENUATION_LogReverse:
 
-				return FMath::Max( 0.5f * FMath::Loge( 1.0f / ( 1.0f - ( Distance / Falloff ) ) ), 0.0f );
+				return FMath::Max(1.0f + 0.5f * FMath::Loge(1.0f - (Distance / Falloff)), 0.0f);
 
 			case ATTENUATION_NaturalSound:
 			{
