UE4 - randomly adjust LOD for each instance based on foliage.RandomLODRange

diff --git a/Engine/Shaders/LocalVertexFactory.usf b/Engine/Shaders/LocalVertexFactory.usf
--- a/Engine/Shaders/LocalVertexFactory.usf
+++ b/Engine/Shaders/LocalVertexFactory.usf
@@ -18,7 +18,7 @@ float4 ShadowMapCoordinateScaleBias;
 #if USE_INSTANCING
 
 #if USE_DITHERED_LOD_TRANSITION
-    float4 InstancingViewZCompareZero;
+    float4 InstancingViewZCompareZero;  // w contains random lod scale 
     float4 InstancingViewZCompareOne;
     float4 InstancingViewZConstant;
     float4 InstancingWorldViewOriginZero;
@@ -459,8 +459,9 @@ FVertexFactoryIntermediates GetVertexFactoryIntermediates(FVertexFactoryInput In
 	Intermediates.PerInstanceParams.y = 1.0 - saturate((length(InstanceLocation + View.PreViewTranslation.xyz) - InstancingFadeOutParams.x) * InstancingFadeOutParams.y);
 	Intermediates.PerInstanceParams.z = InstancingFadeOutParams.z * SelectedValue + InstancingFadeOutParams.w * (1-SelectedValue);
 	#if USE_DITHERED_LOD_TRANSITION
-		float ViewZZero = length(InstanceLocation - InstancingWorldViewOriginZero.xyz);
-		float ViewZOne = length(InstanceLocation - InstancingWorldViewOriginOne.xyz);
+		float RandomLOD = InstancingViewZCompareZero.w * Intermediates.PerInstanceParams.x;
+		float ViewZZero = length(InstanceLocation - InstancingWorldViewOriginZero.xyz) + RandomLOD;
+		float ViewZOne = length(InstanceLocation - InstancingWorldViewOriginOne.xyz) + RandomLOD;
 		Intermediates.PerInstanceParams.w = 
 			dot(ViewZZero > InstancingViewZCompareZero.xyz, InstancingViewZConstant.xyz) * InstancingWorldViewOriginZero.w +
 			dot(ViewZOne > InstancingViewZCompareOne.xyz, InstancingViewZConstant.xyz) * InstancingWorldViewOriginOne.w;
diff --git a/Engine/Source/Runtime/Engine/Private/HierarchicalInstancedStaticMesh.cpp b/Engine/Source/Runtime/Engine/Private/HierarchicalInstancedStaticMesh.cpp
--- a/Engine/Source/Runtime/Engine/Private/HierarchicalInstancedStaticMesh.cpp
+++ b/Engine/Source/Runtime/Engine/Private/HierarchicalInstancedStaticMesh.cpp
@@ -61,6 +61,11 @@ TAutoConsoleVariable<float> CVarFoliageLODDistanceScale(
 	1.0f,
 	TEXT("Scale factor for the distance used in computing LOD for foliage."));
 
+TAutoConsoleVariable<float> CVarRandomLODRange(
+	TEXT("foliage.RandomLODRange"),
+	500.0f,
+	TEXT("Random distance added to each instance distance to compute LOD."));
+
 
 DECLARE_CYCLE_STAT(TEXT("Traversal Time"),STAT_FoliageTraversalTime,STATGROUP_Foliage);
 DECLARE_CYCLE_STAT(TEXT("Build Time"), STAT_FoliageBuildTime, STATGROUP_Foliage);
@@ -1162,6 +1167,7 @@ void FHierarchicalStaticMeshSceneProxy::GetDynamicMeshElements(const TArray<cons
 
 				float MinSize = CVarFoliageMinimumScreenSize.GetValueOnRenderThread();
 				float LODScale = CVarFoliageLODDistanceScale.GetValueOnRenderThread();
+				float LODRandom = CVarRandomLODRange.GetValueOnRenderThread();
 				float MaxDrawDistanceScale = GetCachedScalabilityCVars().ViewDistanceScale;
 				float SphereRadius = RenderData->Bounds.SphereRadius;
 
@@ -1190,10 +1196,10 @@ void FHierarchicalStaticMeshSceneProxy::GetDynamicMeshElements(const TArray<cons
 					for (int32 LODIndex = 1; LODIndex < InstanceParams.LODs; LODIndex++)
 					{
 						float Val = FMath::Min(FMath::Sqrt(Fac / RenderData->ScreenSize[LODIndex]), FinalCull);
-						InstanceParams.LODPlanesMin[LODIndex - 1] = FMath::Min(InstanceParams.LODPlanesMin[LODIndex - 1], Val);
+						InstanceParams.LODPlanesMin[LODIndex - 1] = FMath::Min(InstanceParams.LODPlanesMin[LODIndex - 1], Val - LODRandom);
 						InstanceParams.LODPlanesMax[LODIndex - 1] = FMath::Max(InstanceParams.LODPlanesMax[LODIndex - 1], Val);
 					}
-					InstanceParams.LODPlanesMin[InstanceParams.LODs - 1] = FMath::Min(InstanceParams.LODPlanesMin[InstanceParams.LODs - 1], FinalCull);
+					InstanceParams.LODPlanesMin[InstanceParams.LODs - 1] = FMath::Min(InstanceParams.LODPlanesMin[InstanceParams.LODs - 1], FinalCull - LODRandom);
 					InstanceParams.LODPlanesMax[InstanceParams.LODs - 1] = FMath::Max(InstanceParams.LODPlanesMax[InstanceParams.LODs - 1], FinalCull);
 				}
 
diff --git a/Engine/Source/Runtime/Engine/Private/InstancedStaticMesh.cpp b/Engine/Source/Runtime/Engine/Private/InstancedStaticMesh.cpp
--- a/Engine/Source/Runtime/Engine/Private/InstancedStaticMesh.cpp
+++ b/Engine/Source/Runtime/Engine/Private/InstancedStaticMesh.cpp
@@ -1518,8 +1518,8 @@ void FInstancedStaticMeshVertexFactoryShaderParameters::SetMesh( FRHICommandList
 	const FInstancingUserData* InstancingUserData = (const FInstancingUserData*)BatchElement.UserData;
 	if( InstancingWorldViewOriginOneParameter.IsBound() )
 	{
-		FVector4 InstancingViewZCompareZero(MIN_flt, MIN_flt, MAX_flt, MAX_flt);
-		FVector4 InstancingViewZCompareOne(MIN_flt, MIN_flt, MAX_flt, MAX_flt);
+		FVector4 InstancingViewZCompareZero(MIN_flt, MIN_flt, MAX_flt, 1.0f);
+		FVector4 InstancingViewZCompareOne(MIN_flt, MIN_flt, MAX_flt, 0.0f);
 		FVector4 InstancingViewZConstant(ForceInit);
 		FVector4 InstancingWorldViewOriginZero(ForceInit);
 		FVector4 InstancingWorldViewOriginOne(ForceInit);
@@ -1529,6 +1529,7 @@ void FInstancedStaticMeshVertexFactoryShaderParameters::SetMesh( FRHICommandList
 			float SphereRadius = InstancingUserData->MeshRenderData->Bounds.SphereRadius;
 			float MinSize = CVarFoliageMinimumScreenSize.GetValueOnRenderThread();
 			float LODScale = CVarFoliageLODDistanceScale.GetValueOnRenderThread();
+			float LODRandom = CVarRandomLODRange.GetValueOnRenderThread();
 			float MaxDrawDistanceScale = GetCachedScalabilityCVars().ViewDistanceScale;
 
 			if (BatchElement.InstancedLODIndex)
@@ -1601,6 +1602,8 @@ void FInstancedStaticMeshVertexFactoryShaderParameters::SetMesh( FRHICommandList
 			float Alpha = View.GetTemporalLODTransition();
 			InstancingWorldViewOriginZero.W = 1.0f - Alpha;
 			InstancingWorldViewOriginOne.W = Alpha;
+
+			InstancingViewZCompareZero.W = LODRandom;
 		}
 		SetShaderValue(RHICmdList, VS, InstancingViewZCompareZeroParameter, InstancingViewZCompareZero );
 		SetShaderValue(RHICmdList, VS, InstancingViewZCompareOneParameter, InstancingViewZCompareOne );
diff --git a/Engine/Source/Runtime/Engine/Private/InstancedStaticMesh.h b/Engine/Source/Runtime/Engine/Private/InstancedStaticMesh.h
--- a/Engine/Source/Runtime/Engine/Private/InstancedStaticMesh.h
+++ b/Engine/Source/Runtime/Engine/Private/InstancedStaticMesh.h
@@ -53,6 +53,7 @@
 
 extern TAutoConsoleVariable<float> CVarFoliageMinimumScreenSize;
 extern TAutoConsoleVariable<float> CVarFoliageLODDistanceScale;
+extern TAutoConsoleVariable<float> CVarRandomLODRange;
 
 
 // This must match the maximum a user could specify in the material (see 
