Altered InitBodies to no longer take raw data pointers

diff --git a/Engine/Source/Runtime/Engine/Classes/PhysicsEngine/BodyInstance.h b/Engine/Source/Runtime/Engine/Classes/PhysicsEngine/BodyInstance.h
--- a/Engine/Source/Runtime/Engine/Classes/PhysicsEngine/BodyInstance.h
+++ b/Engine/Source/Runtime/Engine/Classes/PhysicsEngine/BodyInstance.h
@@ -389,17 +389,6 @@ public:
 	 */
 	static void InitBodies(TArray<FBodyInstance*>& Bodies, TArray<FTransform>& Transforms, class UBodySetup* BodySetup, class UPrimitiveComponent* PrimitiveComp, class FPhysScene* InRBScene, PhysXAggregateType InAggregate = NULL, bool Defer = false);
 
-	/** Initialises one or more rigid bodies
-	 *	@param Bodies First body instance to initialise
-	 *	@param Transforms First transform to use
-	 *	@param NumBodies Number of bodies to iterate
-	 *	@param BodySetup Setup to use for bodies
-	 *	@param PrimitiveComp The owning component for the bodies
-	 *	@param InRBScene Scene to place the bodies into
-	 *	@param InAggregate Aggregate to place the bodies into, be aware for PhysX we can only have 128 elements in an aggregate!
-	 */
-	static void InitBodies(FBodyInstance** Bodies, FTransform* Transforms, int32 NumBodies, class UBodySetup* BodySetup, class UPrimitiveComponent* PrimitiveComp, class FPhysScene* InRBScene, PhysXAggregateType InAggregate = NULL, bool Defer = false);
-
 	void InitBody_Legacy(class UBodySetup* Setup, const FTransform& Transform, class UPrimitiveComponent* PrimComp, class FPhysScene* InRBScene, PhysXAggregateType InAggregate = NULL);
 
 	/** Validate a body transform, outputting debug info
diff --git a/Engine/Source/Runtime/Engine/Private/PhysicsEngine/BodyInstance.cpp b/Engine/Source/Runtime/Engine/Private/PhysicsEngine/BodyInstance.cpp
--- a/Engine/Source/Runtime/Engine/Private/PhysicsEngine/BodyInstance.cpp
+++ b/Engine/Source/Runtime/Engine/Private/PhysicsEngine/BodyInstance.cpp
@@ -892,15 +892,23 @@ DEFINE_STAT(STAT_BulkSceneAdd);
 
 void FBodyInstance::InitBody(class UBodySetup* Setup, const FTransform& Transform, class UPrimitiveComponent* PrimComp, class FPhysScene* InRBScene, PhysXAggregateType InAggregate /*= NULL*/, bool bDefer /*= false*/ )
 {
-	FBodyInstance* Instance = this;
-	FTransform TransformCopy = Transform;
-	InitBodies(&Instance, &TransformCopy, 1, Setup, PrimComp, InRBScene, InAggregate, bDefer);
-}
+	static TArray<FBodyInstance*> Bodies;
+	static TArray<FTransform> Transforms;
 
-void FBodyInstance::InitBodies(TArray<FBodyInstance*>& Bodies, TArray<FTransform>& Transforms, class UBodySetup* BodySetup, class UPrimitiveComponent* PrimitiveComp, class FPhysScene* InRBScene, PhysXAggregateType InAggregate /*= NULL*/, bool Defer /*= false*/)
-{
-	check(Bodies.Num() == Transforms.Num());
-	InitBodies(Bodies.GetData(), Transforms.GetData(), Bodies.Num(), BodySetup, PrimitiveComp, InRBScene, InAggregate, Defer);
+	if(Bodies.Num() == 1 && Transforms.Num() == 1)
+	{
+		Bodies[0] = this;
+		Transforms[0] = Transform;
+
+	}
+	else
+	{
+		Bodies.Add(this);
+		Transforms.Add(Transform);
+	}
+	check(Bodies.Num() == Transforms.Num() == 1);
+
+	InitBodies(Bodies, Transforms, Setup, PrimComp, InRBScene, InAggregate, bDefer);
 }
 
 void FBodyInstance::InitBody_Legacy(UBodySetup* Setup, const FTransform& Transform, UPrimitiveComponent* PrimComp, FPhysScene* InRBScene, PxAggregate* InAggregate)
@@ -1404,14 +1412,15 @@ void FBodyInstance::InitBody_Legacy(UBodySetup* Setup, const FTransform& Transfo
 }
 
 
-void FBodyInstance::InitBodies(FBodyInstance** Bodies, FTransform* Transforms, int32 NumBodies, class UBodySetup* BodySetup, class UPrimitiveComponent* PrimitiveComp, class FPhysScene* InRBScene, PhysXAggregateType InAggregate /*= NULL*/, bool bDefer /*= false*/ )
+void FBodyInstance::InitBodies(TArray<FBodyInstance*>& Bodies, TArray<FTransform>& Transforms, class UBodySetup* BodySetup, class UPrimitiveComponent* PrimitiveComp, class FPhysScene* InRBScene, PhysXAggregateType InAggregate /*= NULL*/, bool bDefer /*= false*/ )
 {
 	SCOPE_CYCLE_COUNTER(STAT_BulkInit);
 
 	check(BodySetup);
 	check(InRBScene);
-	check(NumBodies > 0);
+	check(Bodies.Num() > 0);
 
+	int32 NumBodies = Bodies.Num();
 	AActor* OwningActor = PrimitiveComp ? PrimitiveComp->GetOwner() : nullptr;
 
 	const bool bStatic = PrimitiveComp == nullptr || PrimitiveComp->Mobility != EComponentMobility::Movable;
