Fix a garbage collection crash with material instances.

diff --git a/Engine/Source/Runtime/Engine/Private/Materials/MaterialInstance.cpp b/Engine/Source/Runtime/Engine/Private/Materials/MaterialInstance.cpp
--- a/Engine/Source/Runtime/Engine/Private/Materials/MaterialInstance.cpp
+++ b/Engine/Source/Runtime/Engine/Private/Materials/MaterialInstance.cpp
@@ -1824,7 +1824,10 @@ void UMaterialInstance::AddReferencedObjects(UObject* InThis, FReferenceCollecto
 			for (int32 FeatureLevelIndex = 0; FeatureLevelIndex < ERHIFeatureLevel::Num; FeatureLevelIndex++)
 			{
 				FMaterialResource* CurrentResource = This->StaticPermutationMaterialResources[QualityLevelIndex][FeatureLevelIndex];
-				CurrentResource->AddReferencedObjects(Collector);
+				if (CurrentResource)
+				{
+					CurrentResource->AddReferencedObjects(Collector);
+				}
 			}
 		}
 	}
