InheritableComponents - WIP - various fixes

diff --git a/Engine/Source/Runtime/Engine/Classes/Engine/InheritableComponentHandler.h b/Engine/Source/Runtime/Engine/Classes/Engine/InheritableComponentHandler.h
--- a/Engine/Source/Runtime/Engine/Classes/Engine/InheritableComponentHandler.h
+++ b/Engine/Source/Runtime/Engine/Classes/Engine/InheritableComponentHandler.h
@@ -25,7 +25,7 @@ struct ENGINE_API FComponentKey
 	{}
 
 	FComponentKey(UBlueprintGeneratedClass* InOwnerClass, FName InVariableNam)
-		: OwnerClass(nullptr)
+		: OwnerClass(InOwnerClass)
 		, VariableName(InVariableNam)
 	{}
 
@@ -69,7 +69,6 @@ private:
 
 public:
 	UActorComponent* CreateOverridenComponentTemplate(FComponentKey Key);
-	void UpdateOverridenComponentTemplate(FComponentKey Key);
 	void UpdateOwnerClass(UBlueprintGeneratedClass* OwnerClass);
 	void RemoveInvalidAndUnnecessaryTemplates();
 	bool IsValid() const;
diff --git a/Engine/Source/Runtime/Engine/Private/ActorConstruction.cpp b/Engine/Source/Runtime/Engine/Private/ActorConstruction.cpp
--- a/Engine/Source/Runtime/Engine/Private/ActorConstruction.cpp
+++ b/Engine/Source/Runtime/Engine/Private/ActorConstruction.cpp
@@ -224,7 +224,7 @@ void AActor::DestroyConstructedComponents()
 				// Rename component to avoid naming conflicts in the case where we rerun the SCS and name the new components the same way.
 				FName const NewBaseName( *(FString::Printf(TEXT("TRASH_%s"), *Component->GetClass()->GetName())) );
 				FName const NewObjectName = MakeUniqueObjectName(this, GetClass(), NewBaseName);
-				Component->Rename(*NewObjectName.ToString(), this, REN_ForceNoResetLoaders);
+				Component->Rename(*NewObjectName.ToString(), this, REN_ForceNoResetLoaders|REN_DontCreateRedirectors);
 			}
 		}
 	}
@@ -551,7 +551,7 @@ UActorComponent* AActor::CreateComponentFromTemplate(UActorComponent* Template,
 	if(Template != NULL)
 	{
 		// Note we aren't copying the the RF_ArchetypeObject flag. Also note the result is non-transactional by default.
-		NewActorComp = (UActorComponent*)StaticDuplicateObject(Template, this, *InName, RF_AllFlags & ~(RF_ArchetypeObject|RF_Transactional|RF_WasLoaded) );
+		NewActorComp = (UActorComponent*)StaticDuplicateObject(Template, this, *InName, RF_AllFlags & ~(RF_ArchetypeObject|RF_Transactional|RF_WasLoaded|RF_Public) );
 		//NewActorComp = ConstructObject<UActorComponent>(Template->GetClass(), this, *InName, RF_NoFlags, Template);
 
 		NewActorComp->bCreatedByConstructionScript = true;
diff --git a/Engine/Source/Runtime/Engine/Private/Blueprint.cpp b/Engine/Source/Runtime/Engine/Private/Blueprint.cpp
--- a/Engine/Source/Runtime/Engine/Private/Blueprint.cpp
+++ b/Engine/Source/Runtime/Engine/Private/Blueprint.cpp
@@ -538,11 +538,6 @@ void UBlueprint::PostLoad()
 	}
 
 	FStructureEditorUtils::RemoveInvalidStructureMemberVariableFromBlueprint(this);
-
-	if (InheritableComponentHandler)
-	{
-		InheritableComponentHandler->RemoveInvalidAndUnnecessaryTemplates();
-	}
 }
 
 void UBlueprint::DebuggingWorldRegistrationHelper(UObject* ObjectProvidingWorld, UObject* ValueToRegister)
diff --git a/Engine/Source/Runtime/Engine/Private/InheritableComponentHandler.cpp b/Engine/Source/Runtime/Engine/Private/InheritableComponentHandler.cpp
--- a/Engine/Source/Runtime/Engine/Private/InheritableComponentHandler.cpp
+++ b/Engine/Source/Runtime/Engine/Private/InheritableComponentHandler.cpp
@@ -31,7 +31,7 @@ UActorComponent* UInheritableComponentHandler::CreateOverridenComponentTemplate(
 	}
 	ensure(Cast<UBlueprintGeneratedClass>(GetOuter()));
 	auto NewComponentTemplate = ConstructObject<UActorComponent>(
-		BestArchetype->GetClass(), GetOuter(), BestArchetype->GetFName(), RF_ArchetypeObject, BestArchetype);
+		BestArchetype->GetClass(), GetOuter(), BestArchetype->GetFName(), RF_ArchetypeObject | RF_Public, BestArchetype);
 
 	FComponentOverrideRecord NewRecord;
 	NewRecord.ComponentKey = Key;
@@ -41,27 +41,6 @@ UActorComponent* UInheritableComponentHandler::CreateOverridenComponentTemplate(
 	return NewComponentTemplate;
 }
 
-void UInheritableComponentHandler::UpdateOverridenComponentTemplate(FComponentKey Key)
-{
-	auto FoundRecord = FindRecord(Key);
-	if (!FoundRecord)
-	{
-		UE_LOG(LogBlueprint, Warning, TEXT("UpdateOverridenComponentTemplate '%s': cannot find overriden template for component '%s' from '%s'"),
-			*GetPathNameSafe(this), *Key.VariableName.ToString(), *GetPathNameSafe(Key.OwnerClass));
-		return;
-	}
-
-	check(FoundRecord->ComponentTemplate);
-
-	// TODO: save diff against current (previous) archetype. This part is tricky:
-	// the diff should be made before the archetype is updated
-	// OR the diff should be made from saved archetype
-
-	// TODO: copy data from best archetype (it should be new/updated one)
-
-	// TODO: restore custom data from diff
-}
-
 void UInheritableComponentHandler::UpdateOwnerClass(UBlueprintGeneratedClass* OwnerClass)
 {
 	for (auto& Record : Records)
@@ -278,4 +257,4 @@ UActorComponent* FComponentKey::GetOriginalTemplate() const
 	auto ParentSCS = OwnerClass ? OwnerClass->SimpleConstructionScript : nullptr;
 	auto SCSNode = ParentSCS ? ParentSCS->FindSCSNode(VariableName) : nullptr;
 	return SCSNode ? SCSNode->ComponentTemplate : nullptr;
-}
\ No newline at end of file
+}
diff --git a/Engine/Source/Runtime/Engine/Private/SimpleConstructionScript.cpp b/Engine/Source/Runtime/Engine/Private/SimpleConstructionScript.cpp
--- a/Engine/Source/Runtime/Engine/Private/SimpleConstructionScript.cpp
+++ b/Engine/Source/Runtime/Engine/Private/SimpleConstructionScript.cpp
@@ -174,6 +174,23 @@ void USimpleConstructionScript::PostLoad()
 			Node->CategoryName = TEXT("Default");
 		}
 	}
+
+	// Templates are used as archetypes, so they should be Public
+	auto OwnerClass = GetOwnerClass();
+	if (OwnerClass)
+	{
+		for (auto Node : Nodes)
+		{
+			auto ComponentTemplate = Node ? Node->ComponentTemplate : nullptr;
+			if (ComponentTemplate && !ComponentTemplate->HasAllFlags(RF_Public) && (ComponentTemplate->GetOuter() == OwnerClass))
+			{
+				ComponentTemplate->SetFlags(RF_Public);
+
+				// Package should be marked as dirty here, but it's not user friendly
+			}
+		}
+	}
+
 #endif // WITH_EDITOR
 	// Fix up native/inherited parent attachments, in case anything has changed
 	FixupRootNodeParentReferences();
@@ -720,7 +737,7 @@ USCS_Node* USimpleConstructionScript::CreateNode(UClass* NewComponentClass, FNam
 
 	ensure(NULL != Cast<UBlueprintGeneratedClass>(Blueprint->GeneratedClass));
 	UActorComponent* NewComponentTemplate = ConstructObject<UActorComponent>(NewComponentClass, Blueprint->GeneratedClass);
-	NewComponentTemplate->SetFlags(RF_ArchetypeObject|RF_Transactional);
+	NewComponentTemplate->SetFlags(RF_ArchetypeObject|RF_Transactional|RF_Public);
 
 	return CreateNode(NewComponentTemplate, NewComponentVariableName);
 }
