Don't update InstancedStaticMeshComponent instance render data whenever render state is dirtied
UE-5507

diff --git a/Engine/Source/Runtime/Engine/Classes/Components/InstancedStaticMeshComponent.h b/Engine/Source/Runtime/Engine/Classes/Components/InstancedStaticMeshComponent.h
--- a/Engine/Source/Runtime/Engine/Classes/Components/InstancedStaticMeshComponent.h
+++ b/Engine/Source/Runtime/Engine/Classes/Components/InstancedStaticMeshComponent.h
@@ -11,6 +11,7 @@ DECLARE_STATS_GROUP(TEXT("Foliage"), STATGROUP_Foliage, STATCAT_Advanced);
 class FStaticLightingTextureMapping_InstancedStaticMesh;
 class FInstancedLightMap2D;
 class FInstancedShadowMap2D;
+struct FPerInstanceRenderData;
 
 USTRUCT()
 struct FInstancedStaticMeshInstanceData
@@ -128,6 +129,11 @@ class ENGINE_API UInstancedStaticMeshComponent : public UStaticMeshComponent
 	virtual bool ShouldCreatePhysicsState() const override;
 
 public:
+	/** Render data will be initialized once we create scene proxy for this component
+	 *  Released on the rendering thread
+	 */
+	TSharedPtr<FPerInstanceRenderData, ESPMode::ThreadSafe> PerInstanceRenderData;
+		
 #if WITH_EDITOR
 	/** One bit per instance if the instance is selected. */
 	TBitArray<> SelectedInstances;
@@ -170,6 +176,7 @@ public:
 	//Begin UObject Interface
 	virtual void Serialize(FArchive& Ar) override;
 	virtual SIZE_T GetResourceSize(EResourceSizeMode::Type Mode) override;
+	void BeginDestroy() override;
 #if WITH_EDITOR
 	virtual void PostEditChangeChainProperty(FPropertyChangedChainEvent& PropertyChangedEvent) override;
 	virtual void PostEditUndo() override;
@@ -182,6 +189,12 @@ public:
 	/** Select/deselect an instance or group of instances */
 	void SelectInstance(bool bInSelected, int32 InInstanceIndex, int32 InInstanceCount = 1);
 
+	/** 
+	 * Transfers ownership of instance render data to a render thread 
+	 * instance render data will be released in scene proxy dtor or on render thread task
+	 */
+	void ReleasePerInstanceRenderData();
+
 private:
 	/** Creates body instances for all instances owned by this component */
 	void CreateAllInstanceBodies();
@@ -191,7 +204,7 @@ private:
 
 	/** Sets up new instance data to sensible defaults, creates physics counterparts if possible */
 	void SetupNewInstanceData(FInstancedStaticMeshInstanceData& InOutNewInstanceData, int32 InInstanceIndex, const FTransform& InInstanceTransform);
-
+	
 protected:
 	/** Request to navigation system to update only part of navmesh occupied by specified instance */
 	virtual void PartialNavigationUpdate(int32 InstanceIdx);
diff --git a/Engine/Source/Runtime/Engine/Private/HierarchicalInstancedStaticMesh.cpp b/Engine/Source/Runtime/Engine/Private/HierarchicalInstancedStaticMesh.cpp
--- a/Engine/Source/Runtime/Engine/Private/HierarchicalInstancedStaticMesh.cpp
+++ b/Engine/Source/Runtime/Engine/Private/HierarchicalInstancedStaticMesh.cpp
@@ -1409,6 +1409,7 @@ bool UHierarchicalInstancedStaticMeshComponent::RemoveInstance(int32 InstanceInd
 		}
 	}
 
+	ReleasePerInstanceRenderData();
 	MarkRenderStateDirty();
 	
 	return true;
@@ -1740,6 +1741,7 @@ void UHierarchicalInstancedStaticMeshComponent::ApplyBuildTreeAsync(ENamedThread
 			FlushAccumulatedNavigationUpdates();
 		}
 
+		ReleasePerInstanceRenderData();
 	    MarkRenderStateDirty();
 	}
 }
diff --git a/Engine/Source/Runtime/Engine/Private/InstancedStaticMesh.cpp b/Engine/Source/Runtime/Engine/Private/InstancedStaticMesh.cpp
--- a/Engine/Source/Runtime/Engine/Private/InstancedStaticMesh.cpp
+++ b/Engine/Source/Runtime/Engine/Private/InstancedStaticMesh.cpp
@@ -456,12 +456,14 @@ void FInstancedStaticMeshRenderData::InitStaticMeshVertexFactories(
 
 		if (bInstanced)
 		{
+			const FStaticMeshInstanceBuffer& InstanceBuffer = InstancedRenderData->PerInstanceRenderData->InstanceBuffer;
+			
 			// Shadow map bias (and random instance ID)
 			int32 CurInstanceBufferOffset = 0;
 			Data.InstancedShadowMapBiasComponent = FVertexStreamComponent(
-				&InstancedRenderData->InstanceBuffer,
+				&InstanceBuffer,
 				CurInstanceBufferOffset, 
-				InstancedRenderData->InstanceBuffer.GetStride(),
+				InstanceBuffer.GetStride(),
 				VET_Float4,
 				true
 				);
@@ -470,9 +472,9 @@ void FInstancedStaticMeshRenderData::InitStaticMeshVertexFactories(
 			for (int32 MatrixRow = 0; MatrixRow < 3; MatrixRow++)
 			{
 				Data.InstancedTransformComponent[MatrixRow] = FVertexStreamComponent(
-					&InstancedRenderData->InstanceBuffer,
+					&InstanceBuffer,
 					CurInstanceBufferOffset, 
-					InstancedRenderData->InstanceBuffer.GetStride(),
+					InstanceBuffer.GetStride(),
 					VET_Float4,
 					true
 					);
@@ -480,9 +482,9 @@ void FInstancedStaticMeshRenderData::InitStaticMeshVertexFactories(
 			}
 
 			Data.InstancedLightmapUVBiasComponent = FVertexStreamComponent(
-				&InstancedRenderData->InstanceBuffer,
+				&InstanceBuffer,
 				CurInstanceBufferOffset, 
-				InstancedRenderData->InstanceBuffer.GetStride(),
+				InstanceBuffer.GetStride(),
 				VET_Float4,
 				true
 				);
@@ -568,7 +570,7 @@ int32 FInstancedStaticMeshSceneProxy::GetNumMeshBatches() const
 	}
 	else
 	{
-		const uint32 NumInstances = InstancedRenderData.InstanceBuffer.GetNumInstances();
+		const uint32 NumInstances = InstancedRenderData.PerInstanceRenderData->InstanceBuffer.GetNumInstances();
 		const uint32 MaxInstancesPerBatch = FInstancedStaticMeshVertexFactory::NumBitsForVisibilityMask();
 		const uint32 NumBatches = FMath::DivideAndRoundUp(NumInstances, MaxInstancesPerBatch);
 		return NumBatches;
@@ -579,7 +581,7 @@ void FInstancedStaticMeshSceneProxy::SetupInstancedMeshBatch(int32 LODIndex, int
 {
 	const bool bInstanced = RHISupportsInstancing(GetFeatureLevelShaderPlatform(InstancedRenderData.FeatureLevel));
 	OutMeshBatch.VertexFactory = &InstancedRenderData.VertexFactories[LODIndex];
-	const uint32 NumInstances = InstancedRenderData.InstanceBuffer.GetNumInstances();
+	const uint32 NumInstances = InstancedRenderData.PerInstanceRenderData->InstanceBuffer.GetNumInstances();
 	FMeshBatchElement& BatchElement0 = OutMeshBatch.Elements[0];
 	BatchElement0.UserData = (void*)&UserData_AllInstances;
 	BatchElement0.InstancedLODIndex = LODIndex;
@@ -656,10 +658,11 @@ void FInstancedStaticMeshSceneProxy::GetDistancefieldAtlasData(FBox& LocalVolume
 
 	ObjectLocalToWorldTransforms.Reset();
 
-	const FInstancingUserData::FInstanceStream* InstanceStream = ((const FInstancingUserData::FInstanceStream*)InstancedRenderData.InstanceBuffer.GetRawData());
+	const FInstancingUserData::FInstanceStream* InstanceStream = ((const FInstancingUserData::FInstanceStream*)InstancedRenderData.PerInstanceRenderData->InstanceBuffer.GetRawData());
 	FVector4 FourthVector(0, 0, 0, 1);
 
-	for (uint32 InstanceIndex = 0; InstanceIndex < InstancedRenderData.InstanceBuffer.GetNumInstances(); InstanceIndex++)
+	const uint32 NumInstances = InstancedRenderData.PerInstanceRenderData->InstanceBuffer.GetNumInstances();
+	for (uint32 InstanceIndex = 0; InstanceIndex < NumInstances; InstanceIndex++)
 	{
 		const FMatrix TransposedInstanceToLocal(
 			(const FPlane&)InstanceStream[InstanceIndex].InstanceTransform[0], 
@@ -815,7 +818,8 @@ void UInstancedStaticMeshComponent::ApplyComponentInstanceData(FComponentInstanc
 	}
 
 	SelectedInstances = InstancedMeshData->SelectedInstances;
-
+	
+	ReleasePerInstanceRenderData();
 	MarkRenderStateDirty();
 #endif
 }
@@ -1144,6 +1148,25 @@ void UInstancedStaticMeshComponent::ApplyLightMapping(FStaticLightingTextureMapp
 }
 #endif
 
+void UInstancedStaticMeshComponent::ReleasePerInstanceRenderData()
+{
+	if (PerInstanceRenderData.IsValid())
+	{
+		if (SceneProxy == nullptr)
+		{
+			typedef TSharedPtr<FPerInstanceRenderData, ESPMode::ThreadSafe> FPerInstanceRenderDataPtr;
+			// Release PerInstanceRenderData on a render thread
+			ENQUEUE_UNIQUE_RENDER_COMMAND_ONEPARAMETER(
+				FReleasePerInstanceRenderData,
+				FPerInstanceRenderDataPtr,InPerInstanceRenderData,PerInstanceRenderData,
+			{
+				InPerInstanceRenderData.Reset();
+			});
+		}
+		// Will be released by scene proxy or render thread command
+		PerInstanceRenderData.Reset();
+	}
+}
 
 void UInstancedStaticMeshComponent::GetLightAndShadowMapMemoryUsage( int32& LightMapMemoryUsage, int32& ShadowMapMemoryUsage ) const
 {
@@ -1185,6 +1208,7 @@ int32 UInstancedStaticMeshComponent::AddInstance(const FTransform& InstanceTrans
 	}
 #endif
 
+	ReleasePerInstanceRenderData();
 	MarkRenderStateDirty();
 
 	PartialNavigationUpdate(InstanceIdx);
@@ -1229,6 +1253,7 @@ bool UInstancedStaticMeshComponent::RemoveInstance(int32 InstanceIndex)
 	}
 
 	// Indicate we need to update render state to reflect changes
+	ReleasePerInstanceRenderData();
 	MarkRenderStateDirty();
 
 	return true;
@@ -1283,6 +1308,7 @@ bool UInstancedStaticMeshComponent::UpdateInstanceTransform(int32 InstanceIndex,
 	// Request navigation update
 	PartialNavigationUpdate(InstanceIndex);
 
+	ReleasePerInstanceRenderData();
 	MarkRenderStateDirty();
 
 	return true;
@@ -1307,6 +1333,7 @@ void UInstancedStaticMeshComponent::ClearInstances()
 
 	// Indicate we need to update render state to reflect changes
 	bNeverNeedsRenderUpdate = false;
+	ReleasePerInstanceRenderData();
 	MarkRenderStateDirty();
 
 	UNavigationSystem::UpdateNavOctree(this);
@@ -1451,6 +1478,12 @@ SIZE_T UInstancedStaticMeshComponent::GetResourceSize( EResourceSizeMode::Type M
 	return ResSize;
 }
 
+void UInstancedStaticMeshComponent::BeginDestroy()
+{
+	Super::BeginDestroy();
+	ReleasePerInstanceRenderData();
+}
+
 #if WITH_EDITOR
 void UInstancedStaticMeshComponent::PostEditChangeChainProperty(FPropertyChangedChainEvent& PropertyChangedEvent)
 {
@@ -1465,7 +1498,8 @@ void UInstancedStaticMeshComponent::PostEditChangeChainProperty(FPropertyChanged
 			// added via the property editor, so we will want to interactively work with instances
 			bHasPerInstanceHitProxies = true;
 		}
-
+		
+		ReleasePerInstanceRenderData();
 		MarkRenderStateDirty();
 	}
 
@@ -1507,6 +1541,8 @@ void UInstancedStaticMeshComponent::SelectInstance(bool bInSelected, int32 InIns
 	{
 		SelectedInstances[InstanceIndex] = bInSelected;
 	}
+
+	ReleasePerInstanceRenderData();
 #endif
 }
 
@@ -1645,7 +1681,7 @@ void FInstancedStaticMeshVertexFactoryShaderParameters::SetMesh( FRHICommandList
 		if (CPUInstanceShadowMapBias.IsBound())
 		{
 			auto* InstancingData = (const FInstancingUserData*)BatchElement.UserData;
-			auto* InstanceStream = ((const FInstancingUserData::FInstanceStream*)InstancingData->RenderData->InstanceBuffer.GetRawData()) + BatchElement.UserIndex;
+			auto* InstanceStream = ((const FInstancingUserData::FInstanceStream*)InstancingData->RenderData->PerInstanceRenderData->InstanceBuffer.GetRawData()) + BatchElement.UserIndex;
 			SetShaderValue(RHICmdList, VS, CPUInstanceShadowMapBias, InstanceStream->InstanceShadowmapUVBias);
 			SetShaderValueArray(RHICmdList, VS, CPUInstanceTransform, InstanceStream->InstanceTransform, 3);
 			SetShaderValue(RHICmdList, VS, CPUInstanceLightmapUVBias, InstanceStream->InstanceLightmapUVBias);
diff --git a/Engine/Source/Runtime/Engine/Private/InstancedStaticMesh.h b/Engine/Source/Runtime/Engine/Private/InstancedStaticMesh.h
--- a/Engine/Source/Runtime/Engine/Private/InstancedStaticMesh.h
+++ b/Engine/Source/Runtime/Engine/Private/InstancedStaticMesh.h
@@ -289,6 +289,51 @@ private:
 	FShaderParameter CPUInstanceLightmapUVBias;
 };
 
+/*-----------------------------------------------------------------------------
+	FPerInstanceRenderData
+	Holds render data that can persist between scene proxy reconstruction
+-----------------------------------------------------------------------------*/
+struct FPerInstanceRenderData
+{
+	// Should be always constructed on main thread
+	FPerInstanceRenderData(UInstancedStaticMeshComponent* InComponent, ERHIFeatureLevel::Type InFeaureLevel)
+		: InstanceBuffer(InFeaureLevel)
+	{
+		// Create hit proxies for each instance if the component wants
+		if (GIsEditor && InComponent->bHasPerInstanceHitProxies)
+		{
+			HitProxies.Empty(InComponent->PerInstanceSMData.Num());
+			for (int32 InstanceIdx=0; InstanceIdx < InComponent->PerInstanceSMData.Num(); InstanceIdx++)
+			{
+				HitProxies.Add(new HInstancedStaticMeshInstance(InComponent, InstanceIdx));
+			}
+		}
+			
+		// initialize the instance buffer from the component's instances
+		InstanceBuffer.Init(InComponent, HitProxies);
+		BeginInitResource(&InstanceBuffer);
+	}
+
+	FPerInstanceRenderData(UInstancedStaticMeshComponent* InComponent, FStaticMeshInstanceData& Other, ERHIFeatureLevel::Type InFeaureLevel)
+		: InstanceBuffer(InFeaureLevel)
+	{
+		InstanceBuffer.InitFromPreallocatedData(InComponent, Other);
+		BeginInitResource(&InstanceBuffer);
+	}
+	
+	// Should be always destructed on render thread
+	~FPerInstanceRenderData()
+	{
+		InstanceBuffer.ReleaseResource();
+	}
+		
+	/** Instance buffer */
+	FStaticMeshInstanceBuffer			InstanceBuffer;
+	/** Hit proxies for the instances */
+	TArray<TRefCountPtr<HHitProxy>>		HitProxies;
+};
+
+
 /*-----------------------------------------------------------------------------
 	FInstancedStaticMeshRenderData
 -----------------------------------------------------------------------------*/
@@ -299,39 +344,33 @@ public:
 
 	FInstancedStaticMeshRenderData(UInstancedStaticMeshComponent* InComponent, ERHIFeatureLevel::Type InFeatureLevel)
 	  : Component(InComponent)
-	  , InstanceBuffer(InFeatureLevel)
+	  , PerInstanceRenderData(InComponent->PerInstanceRenderData)
 	  , LODModels(Component->StaticMesh->RenderData->LODResources)
 	  , FeatureLevel(InFeatureLevel)
 	{
 		// Allocate the vertex factories for each LOD
 		InitVertexFactories();
-
-		// Create hit proxies for each instance if the component wants
-		if( GIsEditor && InComponent->bHasPerInstanceHitProxies )
+		
+		if (!PerInstanceRenderData.IsValid())
 		{
-			HitProxies.Empty(Component->PerInstanceSMData.Num());
-			for( int32 InstanceIdx=0;InstanceIdx<Component->PerInstanceSMData.Num();InstanceIdx++ )
-			{
-				HitProxies.Add(new HInstancedStaticMeshInstance(InComponent, InstanceIdx));
-			}
+			// initialize the instance buffer from the component's instances
+			PerInstanceRenderData = InComponent->PerInstanceRenderData = MakeShareable(new FPerInstanceRenderData(InComponent, InFeatureLevel));
+			InComponent->PerInstanceRenderData = PerInstanceRenderData;
 		}
 
-		// initialize the instance buffer from the component's instances
-		InstanceBuffer.Init(Component, HitProxies);
-		NumInstances = InstanceBuffer.GetNumInstances();
+		NumInstances = PerInstanceRenderData->InstanceBuffer.GetNumInstances();
 		InitResources();
 	}
 
 	FInstancedStaticMeshRenderData(UInstancedStaticMeshComponent* InComponent, ERHIFeatureLevel::Type InFeatureLevel, FStaticMeshInstanceData& Other)
 		: Component(InComponent)
-		, InstanceBuffer(InFeatureLevel)
 		, LODModels(Component->StaticMesh->RenderData->LODResources)
 		, FeatureLevel(InFeatureLevel)
 	{
 		InitVertexFactories();
 		// initialize the instance buffer from the component's instances
-		InstanceBuffer.InitFromPreallocatedData(Component, Other);
-		NumInstances = InstanceBuffer.GetNumInstances();
+		PerInstanceRenderData = MakeShareable(new FPerInstanceRenderData(InComponent, Other, InFeatureLevel));
+		NumInstances = PerInstanceRenderData->InstanceBuffer.GetNumInstances();
 		InitResources();
 	}
 
@@ -341,8 +380,6 @@ public:
 
 	void InitResources()
 	{
-		BeginInitResource(&InstanceBuffer);
-
 		// Initialize the static mesh's vertex factory.
 		ENQUEUE_UNIQUE_RENDER_COMMAND_THREEPARAMETER(
 			CallInitStaticMeshVertexFactory,
@@ -379,7 +416,6 @@ public:
 			}
 		}
 
-		InstanceBuffer.ReleaseResource();
 		for( int32 LODIndex=0;LODIndex<VertexFactories.Num();LODIndex++ )
 		{
 			VertexFactories[LODIndex].ReleaseResource();
@@ -394,22 +430,19 @@ public:
 	/** Source component */
 	UInstancedStaticMeshComponent* Component;
 
-	/** Instance buffer */
-	FStaticMeshInstanceBuffer InstanceBuffer;
+	/** Per instance render data, could be shared with component */
+	TSharedPtr<FPerInstanceRenderData, ESPMode::ThreadSafe> PerInstanceRenderData;
 
 	/** Vertex factory */
 	TArray<FInstancedStaticMeshVertexFactory> VertexFactories;
 
 	/** LOD render data from the static mesh. */
 	TIndirectArray<FStaticMeshLODResources>& LODModels;
-
-	/** Hit proxies for the instances */
-	TArray<TRefCountPtr<HHitProxy> > HitProxies;
-
+	
 	/** Feature level used when creating instance data */
 	ERHIFeatureLevel::Type FeatureLevel;
 
-	/** Feature level used when creating instance data */
+	/** Number of instances */
 	int32 NumInstances;
 
 private:
@@ -503,10 +536,10 @@ public:
 	 */
 	virtual HHitProxy* CreateHitProxies(UPrimitiveComponent* Component,TArray<TRefCountPtr<HHitProxy> >& OutHitProxies) override
 	{
-		if( InstancedRenderData.HitProxies.Num() )
+		if( InstancedRenderData.PerInstanceRenderData->HitProxies.Num() )
 		{
 			// Add any per-instance hit proxies.
-			OutHitProxies += InstancedRenderData.HitProxies;
+			OutHitProxies += InstancedRenderData.PerInstanceRenderData->HitProxies;
 
 			// No default hit proxy.
 			return NULL;
@@ -566,7 +599,7 @@ private:
 			}
 		}
 
-		check(InstancedRenderData.InstanceBuffer.GetStride() == sizeof(FInstancingUserData::FInstanceStream));
+		check(InstancedRenderData.PerInstanceRenderData->InstanceBuffer.GetStride() == sizeof(FInstancingUserData::FInstanceStream));
 
 		const bool bInstanced = RHISupportsInstancing(GetFeatureLevelShaderPlatform(InstancedRenderData.FeatureLevel));
 
diff --git a/Engine/Source/Runtime/Foliage/Private/InstancedFoliage.cpp b/Engine/Source/Runtime/Foliage/Private/InstancedFoliage.cpp
--- a/Engine/Source/Runtime/Foliage/Private/InstancedFoliage.cpp
+++ b/Engine/Source/Runtime/Foliage/Private/InstancedFoliage.cpp
@@ -745,6 +745,7 @@ void FFoliageMeshInfo::SelectInstances(AInstancedFoliageActor* InIFA, bool bSele
 		if (bSelect)
 		{
 			// Apply selections to the component
+			Component->ReleasePerInstanceRenderData();
 			Component->MarkRenderStateDirty();
 
 			if (Component->SelectedInstances.Num() != Component->PerInstanceSMData.Num())
@@ -762,6 +763,7 @@ void FFoliageMeshInfo::SelectInstances(AInstancedFoliageActor* InIFA, bool bSele
 		{
 			if (Component->SelectedInstances.Num())
 			{
+				Component->ReleasePerInstanceRenderData();
 				Component->MarkRenderStateDirty();
 
 				for (int32 i : InInstances)
@@ -1223,6 +1225,7 @@ void AInstancedFoliageActor::SelectInstance(UInstancedStaticMeshComponent* InCom
 				if (MeshInfo.Component->SelectedInstances.Num() > 0)
 				{
 					MeshInfo.Component->SelectedInstances.Empty();
+					MeshInfo.Component->ReleasePerInstanceRenderData();
 					MeshInfo.Component->MarkRenderStateDirty();
 				}
 			}
@@ -1248,6 +1251,7 @@ void AInstancedFoliageActor::SelectInstance(UInstancedStaticMeshComponent* InCom
 				if (InInstanceIndex < InComponent->SelectedInstances.Num())
 				{
 					InComponent->SelectedInstances[InInstanceIndex] = false;
+					InComponent->ReleasePerInstanceRenderData();
 					InComponent->MarkRenderStateDirty();
 				}
 
@@ -1271,6 +1275,7 @@ void AInstancedFoliageActor::SelectInstance(UInstancedStaticMeshComponent* InCom
 						InComponent->SelectedInstances.Init(false, InComponent->PerInstanceSMData.Num());
 					}
 					InComponent->SelectedInstances[InInstanceIndex] = true;
+					InComponent->ReleasePerInstanceRenderData();
 					InComponent->MarkRenderStateDirty();
 
 					SelectedMesh = Type;
@@ -1312,6 +1317,7 @@ void AInstancedFoliageActor::ApplySelectionToComponents(bool bApply)
 
 				// Apply any selections in the component
 				MeshInfo.Component->SelectedInstances.Init(false, MeshInfo.Component->PerInstanceSMData.Num());
+				MeshInfo.Component->ReleasePerInstanceRenderData();
 				MeshInfo.Component->MarkRenderStateDirty();
 				for (int32 i : MeshInfo.SelectedIndices)
 				{
@@ -1325,6 +1331,7 @@ void AInstancedFoliageActor::ApplySelectionToComponents(bool bApply)
 			{
 				// remove any selections in the component
 				MeshInfo.Component->SelectedInstances.Empty();
+				MeshInfo.Component->ReleasePerInstanceRenderData();
 				MeshInfo.Component->MarkRenderStateDirty();
 			}
 		}
