Removed unused FOldGuidPair, FLegacyCoverIndexPair

diff --git a/Engine/Source/Runtime/Engine/Private/Level.cpp b/Engine/Source/Runtime/Engine/Private/Level.cpp
--- a/Engine/Source/Runtime/Engine/Private/Level.cpp
+++ b/Engine/Source/Runtime/Engine/Private/Level.cpp
@@ -256,32 +256,6 @@ void ULevel::AddReferencedObjects(UObject* InThis, FReferenceCollector& Collecto
 	Super::AddReferencedObjects( This, Collector );
 }
 
-// Compatibility classes
-struct FOldGuidPair
-{
-public:
-	FGuid	Guid;
-	uint32	RefId;
-
-	friend FArchive& operator<<( FArchive& Ar, FOldGuidPair& GP )
-	{
-		Ar << GP.Guid << GP.RefId;
-		return Ar;
-	}
-};
-
-struct FLegacyCoverIndexPair
-{
-	TLazyObjectPtr<class ACoverLink> CoverRef;
-	uint8	SlotIdx;
-
-	friend FArchive& operator<<( FArchive& Ar, struct FLegacyCoverIndexPair& IP )
-	{
-		Ar << IP.CoverRef << IP.SlotIdx;
-		return Ar;
-	}
-};
-
 void ULevel::Serialize( FArchive& Ar )
 {
 	DECLARE_SCOPE_CYCLE_COUNTER(TEXT("ULevel::Serialize"), STAT_Level_Serialize, STATGROUP_LoadTime);
