Fixed non STATS builds (fix for CL 2409266 build errors)

diff --git a/Engine/Source/Runtime/Engine/Private/UserInterface/Console.cpp b/Engine/Source/Runtime/Engine/Private/UserInterface/Console.cpp
--- a/Engine/Source/Runtime/Engine/Private/UserInterface/Console.cpp
+++ b/Engine/Source/Runtime/Engine/Private/UserInterface/Console.cpp
@@ -221,6 +221,7 @@ void UConsole::BuildRuntimeAutoCompleteList(bool bForce)
 		AutoCompleteList[NewIdx].Desc = FString(TEXT("open 127.0.0.1 (opens connection to localhost)"));
 	}
 
+#if STATS
 	// stat commands
 	{
 		const TSet<FName>& StatGroupNames = FStatGroupGameThreadNotifier::Get().StatGroupNames;
@@ -236,6 +237,7 @@ void UConsole::BuildRuntimeAutoCompleteList(bool bForce)
 			NewIdx++;
 		}
 	}
+#endif
 
 	// build the magic tree!
 	for (int32 ListIdx = 0; ListIdx < AutoCompleteList.Num(); ListIdx++)
