PR #753: Made UGameViewportClient::NotifyPlayerRemoved() virtual (Contributed by [REDACTED_PERSON])

diff --git a/Engine/Source/Runtime/Engine/Classes/Engine/GameViewportClient.h b/Engine/Source/Runtime/Engine/Classes/Engine/GameViewportClient.h
--- a/Engine/Source/Runtime/Engine/Classes/Engine/GameViewportClient.h
+++ b/Engine/Source/Runtime/Engine/Classes/Engine/GameViewportClient.h
@@ -413,7 +413,7 @@ public:
 	 * @param	PlayerIndex		the index [into the GamePlayers array] where the player was located
 	 * @param	RemovedPlayer	the player that was removed
 	 */
-	void NotifyPlayerRemoved( int32 PlayerIndex, class ULocalPlayer* RemovedPlayer );
+	virtual void NotifyPlayerRemoved( int32 PlayerIndex, class ULocalPlayer* RemovedPlayer );
 
 	/**
 	 * Notification of server travel error messages, generally network connection related (package verification, client server handshaking, etc) 
