Fix use of negative play rates in composite animations

diff --git a/Engine/Source/Runtime/Engine/Classes/Animation/AnimCompositeBase.h b/Engine/Source/Runtime/Engine/Classes/Animation/AnimCompositeBase.h
--- a/Engine/Source/Runtime/Engine/Classes/Animation/AnimCompositeBase.h
+++ b/Engine/Source/Runtime/Engine/Classes/Animation/AnimCompositeBase.h
@@ -98,7 +98,7 @@ struct FAnimSegment
 
 	float GetLength() const
 	{
-		return (float(LoopingCount) * (AnimEndTime - AnimStartTime)) / GetValidPlayRate();
+		return (float(LoopingCount) * (AnimEndTime - AnimStartTime)) / FMath::Abs(GetValidPlayRate());
 	}
 
 	bool IsInRange(float CurPos) const
diff --git a/Engine/Source/Runtime/Engine/Private/Animation/AnimCompositeBase.cpp b/Engine/Source/Runtime/Engine/Private/Animation/AnimCompositeBase.cpp
--- a/Engine/Source/Runtime/Engine/Private/Animation/AnimCompositeBase.cpp
+++ b/Engine/Source/Runtime/Engine/Private/Animation/AnimCompositeBase.cpp
@@ -28,13 +28,20 @@ UAnimSequenceBase * FAnimSegment::GetAnimationData(float PositionInTrack, float&
 			if (LoopingCount > 1)
 			{
 				// we need to consider looping count
-				float AnimPlayLength = (AnimEndTime - AnimStartTime) / ValidPlayRate;
+				float AnimPlayLength = (AnimEndTime - AnimStartTime) / FMath::Abs(ValidPlayRate);
 				Delta = FMath::Fmod(Delta, AnimPlayLength);
 			}
 
 			// no blending supported
 			Weight = 1.f;
-			PositionInAnim = AnimStartTime + Delta * ValidPlayRate;
+			if (ValidPlayRate > 0.f)
+			{
+				PositionInAnim = AnimStartTime + Delta * ValidPlayRate;
+			}
+			else
+			{
+				PositionInAnim = AnimEndTime + Delta * ValidPlayRate;
+			}
 			return AnimReference;
 		}
 	}
@@ -50,9 +57,11 @@ float FAnimSegment::ConvertTrackPosToAnimPos(const float& TrackPosition) const
 	const float AnimPositionUnWrapped = (TrackPosition - StartPos) * GetValidPlayRate();
 
 	// Figure out how many times animation is allowed to be looped.
-	const float LoopCount = FMath::Min(FMath::FloorToInt(AnimPositionUnWrapped / AnimLength), FMath::Max(LoopingCount-1, 0));
+	const float LoopCount = FMath::Min(FMath::FloorToInt(FMath::Abs(AnimPositionUnWrapped) / AnimLength), FMath::Max(LoopingCount-1, 0));
 	// Position within AnimSequence
-	const float AnimPosition = AnimStartTime + AnimPositionUnWrapped - float(LoopCount) * AnimLength;
+	float AnimPoint = AnimPositionUnWrapped > 0.f ? AnimStartTime : AnimEndTime;
+
+	const float AnimPosition = AnimPoint + (AnimPositionUnWrapped - float(LoopCount) * AnimLength);
 	
 	return AnimPosition;
 }
@@ -64,12 +73,12 @@ void FAnimSegment::GetAnimNotifiesFromTrackPositions(const float& PreviousTrackP
 		return;
 	}
 
-	const bool bPlayingBackwards = (PreviousTrackPosition > CurrentTrackPosition);
+	const bool bTrackPlayingBackwards = (PreviousTrackPosition > CurrentTrackPosition);
 	const float SegmentStartPos = StartPos;
 	const float SegmentEndPos = StartPos + GetLength();
 
 	// if track range overlaps segment
-	if( bPlayingBackwards 
+	if( bTrackPlayingBackwards 
 		? ((CurrentTrackPosition < SegmentEndPos) && (PreviousTrackPosition > SegmentStartPos)) 
 		: ((PreviousTrackPosition < SegmentEndPos) && (CurrentTrackPosition > SegmentStartPos)) 
 		)
@@ -79,25 +88,31 @@ void FAnimSegment::GetAnimNotifiesFromTrackPositions(const float& PreviousTrackP
 		if( AnimSequence )
 		{
 			const float ValidPlayRate = GetValidPlayRate();
+			const float AbsValidPlayRate = FMath::Abs(ValidPlayRate);
+
 			// Get starting position, closest overlap.
-			float AnimStartPosition = ConvertTrackPosToAnimPos( bPlayingBackwards ? FMath::Min(PreviousTrackPosition, SegmentEndPos) : FMath::Max(PreviousTrackPosition, SegmentStartPos) );
+			float AnimStartPosition = ConvertTrackPosToAnimPos( bTrackPlayingBackwards ? FMath::Min(PreviousTrackPosition, SegmentEndPos) : FMath::Max(PreviousTrackPosition, SegmentStartPos) );
 			check( (AnimStartPosition >= AnimStartTime) && (AnimStartPosition <= AnimEndTime) );
 			float TrackTimeToGo = CurrentTrackPosition - PreviousTrackPosition;
 
+			// The track can be playing backwards and the animation can be playing backwards, so we
+			// need to combine to work out what direction we are traveling through the animation
+			bool bAnimPlayingBackwards = bTrackPlayingBackwards ^ (ValidPlayRate < 0.f);
+			
 			// Abstract out end point since animation can be playing forward or backward.
-			const float AnimEndPoint = bPlayingBackwards ? AnimStartTime : AnimEndTime;
+			const float AnimEndPoint = bAnimPlayingBackwards ? AnimStartTime : AnimEndTime;
 
 			for(int32 IterationsLeft=FMath::Max(LoopingCount, 1); ((IterationsLeft > 0) && (TrackTimeToGo > 0.f)); --IterationsLeft)
 			{
 				// Track time left to reach end point of animation.
-				const float TrackTimeToAnimEndPoint = (AnimEndPoint - AnimStartPosition) / ValidPlayRate;
+				const float TrackTimeToAnimEndPoint = (AnimEndPoint - AnimStartPosition) / AbsValidPlayRate;
 
 				// If our time left is shorter than time to end point, no problem. End there.
 				if( FMath::Abs(TrackTimeToGo) < FMath::Abs(TrackTimeToAnimEndPoint) )
 				{
 					const float AnimEndPosition = ConvertTrackPosToAnimPos(CurrentTrackPosition);
 					// Make sure we have not wrapped around, positions should be contiguous.
-					check( bPlayingBackwards ? (AnimEndPosition <= AnimStartPosition) : (AnimStartPosition <= AnimEndPosition) );
+					check(bAnimPlayingBackwards ? (AnimEndPosition <= AnimStartPosition) : (AnimStartPosition <= AnimEndPosition));
 					AnimSequence->GetAnimNotifiesFromDeltaPositions(AnimStartPosition, AnimEndPosition, OutActiveNotifies);
 					break;
 				}
@@ -106,13 +121,13 @@ void FAnimSegment::GetAnimNotifiesFromTrackPositions(const float& PreviousTrackP
 				{
 					// Add that piece for extraction.
 					// Make sure we have not wrapped around, positions should be contiguous.
-					check( bPlayingBackwards ? (AnimEndPoint <= AnimStartPosition) : (AnimStartPosition <= AnimEndPoint) );
+					check(bAnimPlayingBackwards ? (AnimEndPoint <= AnimStartPosition) : (AnimStartPosition <= AnimEndPoint));
 					AnimSequence->GetAnimNotifiesFromDeltaPositions(AnimStartPosition, AnimEndPoint, OutActiveNotifies);
 
 					// decrease our TrackTimeToGo if we have to do another iteration.
 					// and put ourselves back at the beginning of the animation.
 					TrackTimeToGo -= TrackTimeToAnimEndPoint;
-					AnimStartPosition = bPlayingBackwards ? AnimEndTime : AnimStartTime;
+					AnimStartPosition = bAnimPlayingBackwards ? AnimEndTime : AnimStartTime;
 				}
 			}
 		}
@@ -135,25 +150,31 @@ void FAnimSegment::GetRootMotionExtractionStepsForTrackRange(TArray<FRootMotionE
 	UAnimSequence * AnimSequence = Cast<UAnimSequence>(AnimReference);
 	if( AnimSequence )
 	{
-		const bool bPlayingBackwards = (StartTrackPosition > EndTrackPosition);
+		const bool bTrackPlayingBackwards = (StartTrackPosition > EndTrackPosition);
 
 		const float SegmentStartPos = StartPos;
 		const float SegmentEndPos = StartPos + GetLength();
 
 		// if range overlaps segment
-		if( bPlayingBackwards 
+		if (bTrackPlayingBackwards
 			? ((EndTrackPosition < SegmentEndPos) && (StartTrackPosition > SegmentStartPos)) 
 			: ((StartTrackPosition < SegmentEndPos) && (EndTrackPosition > SegmentStartPos)) 
 			)
 		{
 			const float ValidPlayRate = GetValidPlayRate();
+			const float AbsValidPlayRate = FMath::Abs(ValidPlayRate);
+
 			// Get starting position, closest overlap.
-			float AnimStartPosition = ConvertTrackPosToAnimPos( bPlayingBackwards ? FMath::Min(StartTrackPosition, SegmentEndPos) : FMath::Max(StartTrackPosition, SegmentStartPos) );
+			float AnimStartPosition = ConvertTrackPosToAnimPos(bTrackPlayingBackwards ? FMath::Min(StartTrackPosition, SegmentEndPos) : FMath::Max(StartTrackPosition, SegmentStartPos));
 			check( (AnimStartPosition >= AnimStartTime) && (AnimStartPosition <= AnimEndTime) );
 			float TrackTimeToGo = EndTrackPosition - StartTrackPosition;
 
+			// The track can be playing backwards and the animation can be playing backwards, so we
+			// need to combine to work out what direction we are traveling through the animation
+			bool bAnimPlayingBackwards = bTrackPlayingBackwards ^ (ValidPlayRate < 0.f);
+
 			// Abstract out end point since animation can be playing forward or backward.
-			const float AnimEndPoint = bPlayingBackwards ? AnimStartTime : AnimEndTime;
+			const float AnimEndPoint = bAnimPlayingBackwards ? AnimStartTime : AnimEndTime;
 
 			for(int32 IterationsLeft=FMath::Max(LoopingCount, 1); ((IterationsLeft > 0) && (TrackTimeToGo > 0.f)); --IterationsLeft)
 			{
@@ -165,7 +186,7 @@ void FAnimSegment::GetRootMotionExtractionStepsForTrackRange(TArray<FRootMotionE
 				{
 					const float AnimEndPosition = ConvertTrackPosToAnimPos(EndTrackPosition);
 					// Make sure we have not wrapped around, positions should be contiguous.
-					check( bPlayingBackwards ? (AnimEndPosition <= AnimStartPosition) : (AnimStartPosition <= AnimEndPosition) );
+					check(bAnimPlayingBackwards ? (AnimEndPosition <= AnimStartPosition) : (AnimStartPosition <= AnimEndPosition));
 					RootMotionExtractionSteps.Add(FRootMotionExtractionStep(AnimSequence, AnimStartPosition, AnimEndPosition));
 					break;
 				}
@@ -174,13 +195,13 @@ void FAnimSegment::GetRootMotionExtractionStepsForTrackRange(TArray<FRootMotionE
 				{
 					// Add that piece for extraction.
 					// Make sure we have not wrapped around, positions should be contiguous.
-					check( bPlayingBackwards ? (AnimEndPoint <= AnimStartPosition) : (AnimStartPosition <= AnimEndPoint) );
+					check(bAnimPlayingBackwards ? (AnimEndPoint <= AnimStartPosition) : (AnimStartPosition <= AnimEndPoint));
 					RootMotionExtractionSteps.Add(FRootMotionExtractionStep(AnimSequence, AnimStartPosition, AnimEndPoint));
 
 					// decrease our TrackTimeToGo if we have to do another iteration.
 					// and put ourselves back at the beginning of the animation.
 					TrackTimeToGo -= TrackTimeToAnimEndPoint;
-					AnimStartPosition = bPlayingBackwards ? AnimEndTime : AnimStartTime;
+					AnimStartPosition = bAnimPlayingBackwards ? AnimEndTime : AnimStartTime;
 				}
 			}
 		}
