Fix for root motion during negative play rates

diff --git a/Engine/Source/Runtime/Engine/Private/Animation/AnimCompositeBase.cpp b/Engine/Source/Runtime/Engine/Private/Animation/AnimCompositeBase.cpp
--- a/Engine/Source/Runtime/Engine/Private/Animation/AnimCompositeBase.cpp
+++ b/Engine/Source/Runtime/Engine/Private/Animation/AnimCompositeBase.cpp
@@ -93,7 +93,7 @@ void FAnimSegment::GetAnimNotifiesFromTrackPositions(const float& PreviousTrackP
 			// Get starting position, closest overlap.
 			float AnimStartPosition = ConvertTrackPosToAnimPos( bTrackPlayingBackwards ? FMath::Min(PreviousTrackPosition, SegmentEndPos) : FMath::Max(PreviousTrackPosition, SegmentStartPos) );
 			check( (AnimStartPosition >= AnimStartTime) && (AnimStartPosition <= AnimEndTime) );
-			float TrackTimeToGo = CurrentTrackPosition - PreviousTrackPosition;
+			float TrackTimeToGo = FMath::Abs(CurrentTrackPosition - PreviousTrackPosition);
 
 			// The track can be playing backwards and the animation can be playing backwards, so we
 			// need to combine to work out what direction we are traveling through the animation
@@ -167,7 +167,7 @@ void FAnimSegment::GetRootMotionExtractionStepsForTrackRange(TArray<FRootMotionE
 			// Get starting position, closest overlap.
 			float AnimStartPosition = ConvertTrackPosToAnimPos(bTrackPlayingBackwards ? FMath::Min(StartTrackPosition, SegmentEndPos) : FMath::Max(StartTrackPosition, SegmentStartPos));
 			check( (AnimStartPosition >= AnimStartTime) && (AnimStartPosition <= AnimEndTime) );
-			float TrackTimeToGo = EndTrackPosition - StartTrackPosition;
+			float TrackTimeToGo = FMath::Abs(EndTrackPosition - StartTrackPosition);
 
 			// The track can be playing backwards and the animation can be playing backwards, so we
 			// need to combine to work out what direction we are traveling through the animation
