Instance SceneComponents will maintain attachment if their BP created parent is reconstructed

diff --git a/Engine/Source/Runtime/Engine/Classes/Components/SceneComponent.h b/Engine/Source/Runtime/Engine/Classes/Components/SceneComponent.h
--- a/Engine/Source/Runtime/Engine/Classes/Components/SceneComponent.h
+++ b/Engine/Source/Runtime/Engine/Classes/Components/SceneComponent.h
@@ -501,6 +501,9 @@ public:
 	virtual void UpdateComponentToWorld(bool bSkipPhysicsMove = false) override final;
 	virtual void OnComponentDestroyed() override;
 	virtual void ApplyWorldOffset(const FVector& InOffset, bool bWorldShift) override;
+	virtual class FComponentInstanceDataBase* GetComponentInstanceData() const override;
+	virtual FName GetComponentInstanceDataType() const override;
+	virtual void ApplyComponentInstanceData(class FComponentInstanceDataBase* ComponentInstanceData ) override;
 	// End ActorComponent interface
 
 	// Begin UObject Interface
@@ -785,6 +788,20 @@ protected:
 	virtual void UpdateNavigationData();
 };
 
+/** 
+ *  Component instance cached data base class for scene components. 
+ *  Stores a list of instance components attached to the 
+ */
+class ENGINE_API FSceneComponentInstanceData : public FComponentInstanceDataBase
+{
+public:
+	FSceneComponentInstanceData(const USceneComponent* SourceComponent);
+			
+	virtual ~FSceneComponentInstanceData()
+	{}
+
+	TArray<USceneComponent*> AttachedInstanceComponents;
+};
 
 
 /**
diff --git a/Engine/Source/Runtime/Engine/Private/Atmosphere/Atmosphere.cpp b/Engine/Source/Runtime/Engine/Private/Atmosphere/Atmosphere.cpp
--- a/Engine/Source/Runtime/Engine/Private/Atmosphere/Atmosphere.cpp
+++ b/Engine/Source/Runtime/Engine/Private/Atmosphere/Atmosphere.cpp
@@ -695,11 +695,11 @@ void UAtmosphericFogComponent::Serialize(FArchive& Ar)
 }
 
 /** Used to store lightmap data during RerunConstructionScripts */
-class FAtmospherePrecomputeInstanceData : public FComponentInstanceDataBase
+class FAtmospherePrecomputeInstanceData : public FSceneComponentInstanceData
 {
 public:
 	FAtmospherePrecomputeInstanceData(const UAtmosphericFogComponent* SourceComponent)
-		: FComponentInstanceDataBase(SourceComponent)
+		: FSceneComponentInstanceData(SourceComponent)
 	{}
 
 	virtual ~FAtmospherePrecomputeInstanceData()
@@ -761,6 +761,8 @@ FComponentInstanceDataBase* UAtmosphericFogComponent::GetComponentInstanceData()
 // Restore the precomputed data after re-running Blueprint construction script
 void UAtmosphericFogComponent::ApplyComponentInstanceData(FComponentInstanceDataBase* ComponentInstanceData)
 {
+	Super::ApplyComponentInstanceData(ComponentInstanceData);
+
 	check(ComponentInstanceData);
 	FAtmospherePrecomputeInstanceData* PrecomputedData = static_cast<FAtmospherePrecomputeInstanceData*>(const_cast<FComponentInstanceDataBase*>(ComponentInstanceData));
 
diff --git a/Engine/Source/Runtime/Engine/Private/ChildActorComponent.cpp b/Engine/Source/Runtime/Engine/Private/ChildActorComponent.cpp
--- a/Engine/Source/Runtime/Engine/Private/ChildActorComponent.cpp
+++ b/Engine/Source/Runtime/Engine/Private/ChildActorComponent.cpp
@@ -40,11 +40,11 @@ void UChildActorComponent::OnComponentDestroyed()
 	DestroyChildActor();
 }
 
-class FChildActorComponentInstanceData : public FComponentInstanceDataBase
+class FChildActorComponentInstanceData : public FSceneComponentInstanceData
 {
 public:
 	FChildActorComponentInstanceData(const UChildActorComponent* Component)
-		: FComponentInstanceDataBase(Component)
+		: FSceneComponentInstanceData(Component)
 		, ChildActorName(Component->ChildActorName)
 	{
 		if (Component->ChildActor)
@@ -108,6 +108,8 @@ FComponentInstanceDataBase* UChildActorComponent::GetComponentInstanceData() con
 
 void UChildActorComponent::ApplyComponentInstanceData(FComponentInstanceDataBase* ComponentInstanceData)
 {
+	Super::ApplyComponentInstanceData(ComponentInstanceData);
+
 	check(ComponentInstanceData);
 	FChildActorComponentInstanceData* ChildActorInstanceData  = static_cast<FChildActorComponentInstanceData*>(ComponentInstanceData);
 	ChildActorName = ChildActorInstanceData->ChildActorName;
diff --git a/Engine/Source/Runtime/Engine/Private/InstancedStaticMesh.cpp b/Engine/Source/Runtime/Engine/Private/InstancedStaticMesh.cpp
--- a/Engine/Source/Runtime/Engine/Private/InstancedStaticMesh.cpp
+++ b/Engine/Source/Runtime/Engine/Private/InstancedStaticMesh.cpp
@@ -687,11 +687,11 @@ UInstancedStaticMeshComponent::UInstancedStaticMeshComponent(const FObjectInitia
 
 #if WITH_EDITOR
 /** Helper class used to preserve lighting/selection state across blueprint reinstancing */
-class FInstancedStaticMeshComponentInstanceData : public FComponentInstanceDataBase
+class FInstancedStaticMeshComponentInstanceData : public FSceneComponentInstanceData
 {
 public:
 	FInstancedStaticMeshComponentInstanceData(const UInstancedStaticMeshComponent& InComponent)
-		: FComponentInstanceDataBase(&InComponent)
+		: FSceneComponentInstanceData(&InComponent)
 		, StaticMesh(InComponent.StaticMesh)
 		, bHasCachedStaticLighting(false)
 	{
@@ -773,6 +773,9 @@ FComponentInstanceDataBase* UInstancedStaticMeshComponent::GetComponentInstanceD
 
 void UInstancedStaticMeshComponent::ApplyComponentInstanceData(FComponentInstanceDataBase* ComponentInstanceData)
 {
+	// Skip UStaticMeshComponent implementation
+	USceneComponent::ApplyComponentInstanceData(ComponentInstanceData);
+
 #if WITH_EDITOR
 	check(ComponentInstanceData);
 
diff --git a/Engine/Source/Runtime/Engine/Private/LightComponent.cpp b/Engine/Source/Runtime/Engine/Private/LightComponent.cpp
--- a/Engine/Source/Runtime/Engine/Private/LightComponent.cpp
+++ b/Engine/Source/Runtime/Engine/Private/LightComponent.cpp
@@ -863,11 +863,11 @@ void ULightComponent::InvalidateLightingCacheInner(bool bRecreateLightGuids)
 }
 
 /** Used to store lightmap data during RerunConstructionScripts */
-class FPrecomputedLightInstanceData : public FComponentInstanceDataBase
+class FPrecomputedLightInstanceData : public FSceneComponentInstanceData
 {
 public:
 	FPrecomputedLightInstanceData(const ULightComponent* SourceComponent)
-		: FComponentInstanceDataBase(SourceComponent)
+		: FSceneComponentInstanceData(SourceComponent)
 		, Transform(SourceComponent->ComponentToWorld)
 		, LightGuid(SourceComponent->LightGuid)
 		, ShadowMapChannel(SourceComponent->ShadowMapChannel)
@@ -899,6 +899,8 @@ FComponentInstanceDataBase* ULightComponent::GetComponentInstanceData() const
 
 void ULightComponent::ApplyComponentInstanceData(FComponentInstanceDataBase* ComponentInstanceData)
 {
+	Super::ApplyComponentInstanceData(ComponentInstanceData);
+
 	check(ComponentInstanceData);
 	FPrecomputedLightInstanceData* LightMapData  = static_cast<FPrecomputedLightInstanceData*>(ComponentInstanceData);
 
diff --git a/Engine/Source/Runtime/Engine/Private/SceneComponent.cpp b/Engine/Source/Runtime/Engine/Private/SceneComponent.cpp
--- a/Engine/Source/Runtime/Engine/Private/SceneComponent.cpp
+++ b/Engine/Source/Runtime/Engine/Private/SceneComponent.cpp
@@ -1007,6 +1007,51 @@ bool USceneComponent::IsAttachedTo(class USceneComponent* TestComp) const
 	return false;
 }
 
+FSceneComponentInstanceData::FSceneComponentInstanceData(const USceneComponent* SourceComponent)
+	: FComponentInstanceDataBase(SourceComponent)
+{
+	for (USceneComponent* SceneComponent : SourceComponent->AttachChildren)
+	{
+		if (SceneComponent && !SceneComponent->bCreatedByConstructionScript)
+		{
+			AttachedInstanceComponents.Add(SceneComponent);
+		}
+	}
+}
+
+FComponentInstanceDataBase* USceneComponent::GetComponentInstanceData() const
+{
+	FComponentInstanceDataBase* InstanceData = nullptr;
+
+	for (USceneComponent* Child : AttachChildren)
+	{
+		if (!Child->bCreatedByConstructionScript)
+		{
+			InstanceData = new FSceneComponentInstanceData(this);
+			break;
+		}
+	}
+
+	return InstanceData;
+}
+
+FName USceneComponent::GetComponentInstanceDataType() const
+{
+	static const FName SceneComponentInstanceDataTypeName(TEXT("SceneComponentInstanceData"));
+	return SceneComponentInstanceDataTypeName;
+}
+
+void USceneComponent::ApplyComponentInstanceData(class FComponentInstanceDataBase* ComponentInstanceData )
+{
+	check(ComponentInstanceData);
+	FSceneComponentInstanceData* SceneComponentInstanceData  = static_cast<FSceneComponentInstanceData*>(ComponentInstanceData);
+
+	for (USceneComponent* ChildComponent : SceneComponentInstanceData->AttachedInstanceComponents)
+	{
+		ChildComponent->AttachTo(this);
+	}
+}
+
 void USceneComponent::UpdateChildTransforms()
 {
 	for(int32 i=0; i<AttachChildren.Num(); i++)
diff --git a/Engine/Source/Runtime/Engine/Private/SkyLightComponent.cpp b/Engine/Source/Runtime/Engine/Private/SkyLightComponent.cpp
--- a/Engine/Source/Runtime/Engine/Private/SkyLightComponent.cpp
+++ b/Engine/Source/Runtime/Engine/Private/SkyLightComponent.cpp
@@ -325,11 +325,11 @@ bool USkyLightComponent::IsReadyForFinishDestroy()
 }
 
 /** Used to store lightmap data during RerunConstructionScripts */
-class FPrecomputedSkyLightInstanceData : public FComponentInstanceDataBase
+class FPrecomputedSkyLightInstanceData : public FSceneComponentInstanceData
 {
 public:
 	FPrecomputedSkyLightInstanceData(const USkyLightComponent* SourceComponent)
-		: FComponentInstanceDataBase(SourceComponent)
+		: FSceneComponentInstanceData(SourceComponent)
 	{}
 
 	FGuid LightGuid;
@@ -359,6 +359,9 @@ FComponentInstanceDataBase* USkyLightComponent::GetComponentInstanceData() const
 
 void USkyLightComponent::ApplyComponentInstanceData(FComponentInstanceDataBase* ComponentInstanceData)
 {
+	//Skip ULightComponent implementationn
+	USceneComponent::ApplyComponentInstanceData(ComponentInstanceData);
+
 	check(ComponentInstanceData);
 	FPrecomputedSkyLightInstanceData* LightMapData  = static_cast<FPrecomputedSkyLightInstanceData*>(ComponentInstanceData);
 
diff --git a/Engine/Source/Runtime/Engine/Private/SplineComponent.cpp b/Engine/Source/Runtime/Engine/Private/SplineComponent.cpp
--- a/Engine/Source/Runtime/Engine/Private/SplineComponent.cpp
+++ b/Engine/Source/Runtime/Engine/Private/SplineComponent.cpp
@@ -563,11 +563,11 @@ void USplineComponent::RefreshSplineInputs()
 
 
 /** Used to store spline data during RerunConstructionScripts */
-class FSplineInstanceData : public FComponentInstanceDataBase
+class FSplineInstanceData : public FSceneComponentInstanceData
 {
 public:
 	explicit FSplineInstanceData(const USplineComponent* SourceComponent)
-		: FComponentInstanceDataBase(SourceComponent)
+		: FSceneComponentInstanceData(SourceComponent)
 	{
 	}
 
@@ -595,6 +595,8 @@ FComponentInstanceDataBase* USplineComponent::GetComponentInstanceData() const
 
 void USplineComponent::ApplyComponentInstanceData(FComponentInstanceDataBase* ComponentInstanceData)
 {
+	Super::ApplyComponentInstanceData(ComponentInstanceData);
+
 	if (ComponentInstanceData)
 	{
 		FSplineInstanceData* SplineInstanceData  = static_cast<FSplineInstanceData*>(ComponentInstanceData);
diff --git a/Engine/Source/Runtime/Engine/Private/StaticMeshComponent.cpp b/Engine/Source/Runtime/Engine/Private/StaticMeshComponent.cpp
--- a/Engine/Source/Runtime/Engine/Private/StaticMeshComponent.cpp
+++ b/Engine/Source/Runtime/Engine/Private/StaticMeshComponent.cpp
@@ -23,11 +23,11 @@
 
 #define LOCTEXT_NAMESPACE "StaticMeshComponent"
 
-class FStaticMeshComponentInstanceData : public FComponentInstanceDataBase
+class FStaticMeshComponentInstanceData : public FSceneComponentInstanceData
 {
 public:
 	FStaticMeshComponentInstanceData(const UStaticMeshComponent* SourceComponent)
-		: FComponentInstanceDataBase(SourceComponent)
+		: FSceneComponentInstanceData(SourceComponent)
 		, StaticMesh(SourceComponent->StaticMesh)
 		, bHasCachedStaticLighting(false)
 	{
@@ -1569,6 +1569,8 @@ FComponentInstanceDataBase* UStaticMeshComponent::GetComponentInstanceData() con
 
 void UStaticMeshComponent::ApplyComponentInstanceData(FComponentInstanceDataBase* ComponentInstanceData)
 {
+	Super::ApplyComponentInstanceData(ComponentInstanceData);
+
 	check(ComponentInstanceData);
 
 	// Note: ApplyComponentInstanceData is called while the component is registered so the rendering thread is already using this component
diff --git a/Engine/Source/Runtime/UMG/Private/Components/WidgetComponent.cpp b/Engine/Source/Runtime/UMG/Private/Components/WidgetComponent.cpp
--- a/Engine/Source/Runtime/UMG/Private/Components/WidgetComponent.cpp
+++ b/Engine/Source/Runtime/UMG/Private/Components/WidgetComponent.cpp
@@ -568,11 +568,11 @@ void UWidgetComponent::TickComponent(float DeltaTime, enum ELevelTick TickType,
 #endif // !UE_SERVER
 }
 
-class FWidgetComponentInstanceData : public FComponentInstanceDataBase
+class FWidgetComponentInstanceData : public FSceneComponentInstanceData
 {
 public:
 	FWidgetComponentInstanceData( const UWidgetComponent* SourceComponent )
-		: FComponentInstanceDataBase(SourceComponent)
+		: FSceneComponentInstanceData(SourceComponent)
 		, WidgetClass ( SourceComponent->GetWidgetClass() )
 		, RenderTarget( SourceComponent->GetRenderTarget() )
 	{}
@@ -595,6 +595,8 @@ FComponentInstanceDataBase* UWidgetComponent::GetComponentInstanceData() const
 
 void UWidgetComponent::ApplyComponentInstanceData(FComponentInstanceDataBase* ComponentInstanceData)
 {
+	Super::ApplyComponentInstanceData(ComponentInstanceData);
+
 	check(ComponentInstanceData);
 
 	// Note: ApplyComponentInstanceData is called while the component is registered so the rendering thread is already using this component
