Fix crash when there is a null entry in the RootComponent's AttachChildren array

diff --git a/Engine/Source/Runtime/Engine/Private/ActorEditor.cpp b/Engine/Source/Runtime/Engine/Private/ActorEditor.cpp
--- a/Engine/Source/Runtime/Engine/Private/ActorEditor.cpp
+++ b/Engine/Source/Runtime/Engine/Private/ActorEditor.cpp
@@ -220,7 +220,7 @@ AActor::FActorTransactionAnnotation::FActorTransactionAnnotation(const AActor* A
 		for (USceneComponent* AttachChild : RootComponent->AttachChildren)
 		{
 			AActor* ChildOwner = (AttachChild ? AttachChild->GetOwner() : NULL);
-			if (ChildOwner != Actor)
+			if (ChildOwner && ChildOwner != Actor)
 			{
 				// Save info about actor to reattach
 				FActorRootComponentReconstructionData::FAttachedActorInfo Info;
