bettern names and comments in the texture compression settings

diff --git a/Engine/Source/Runtime/Engine/Classes/Engine/Texture.h b/Engine/Source/Runtime/Engine/Classes/Engine/Texture.h
--- a/Engine/Source/Runtime/Engine/Classes/Engine/Texture.h
+++ b/Engine/Source/Runtime/Engine/Classes/Engine/Texture.h
@@ -12,18 +12,18 @@
 UENUM()
 enum TextureCompressionSettings
 {
-	TC_Default,
-	TC_Normalmap,
-	TC_Masks,
-	TC_Grayscale,
-	TC_Displacementmap,
-	TC_VectorDisplacementmap,
-	TC_HDR,
-	TC_EditorIcon UMETA(DisplayName="TC_UserInterface2D"),
-	TC_Alpha,
-	TC_DistanceFieldFont,
-	TC_HDR_Compressed,
-	TC_BC7,
+	TC_Default					UMETA(DisplayName="Default (DXT1/5, BC1/5 on DX11)"),
+	TC_Normalmap				UMETA(DisplayName="Normalmap (DXT5, BC5 on DX11)"),
+	TC_Masks					UMETA(DisplayName="Masks (no sRGB)"),
+	TC_Grayscale				UMETA(DisplayName="Grayscale (R8, RGB8 sRGB)"),
+	TC_Displacementmap			UMETA(DisplayName="Displacementmap (8/16bit)"),
+	TC_VectorDisplacementmap	UMETA(DisplayName="VectorDisplacementmap (RGBA8)"),
+	TC_HDR						UMETA(DisplayName="HDR (RGB, no sRGB)"),
+	TC_EditorIcon				UMETA(DisplayName="UserInterface2D (RGBA)"),
+	TC_Alpha					UMETA(DisplayName="Alpha (no sRGB, BC4 on DX11)"),
+	TC_DistanceFieldFont		UMETA(DisplayName="DistanceFieldFont (R8)"),
+	TC_HDR_Compressed			UMETA(DisplayName="HDRCompressed (RGB, BC6H, DX11)"),
+	TC_BC7						UMETA(DisplayName="BC7 (DX11, optional A)"),
 	TC_MAX,
 };
 
