fixed particle debug code

diff --git a/Engine/Source/Runtime/Engine/Private/Particles/ParticleGpuSimulation.cpp b/Engine/Source/Runtime/Engine/Private/Particles/ParticleGpuSimulation.cpp
--- a/Engine/Source/Runtime/Engine/Private/Particles/ParticleGpuSimulation.cpp
+++ b/Engine/Source/Runtime/Engine/Private/Particles/ParticleGpuSimulation.cpp
@@ -3836,24 +3836,24 @@ void DumpTileAllocations()
 		const TSet<FGPUSpriteParticleEmitterInstance*>& Emitters = It.Value();
 		int32 TotalAllocatedTiles = 0;
 
-		UE_LOG(LogParticles,Info,TEXT("---------- GPU Particle Tile Allocations : FXSystem=0x%016x ----------"), (PTRINT)FXSystem);
+		UE_LOG(LogParticles,Display,TEXT("---------- GPU Particle Tile Allocations : FXSystem=0x%016x ----------"), (PTRINT)FXSystem);
 		for (TSet<FGPUSpriteParticleEmitterInstance*>::TConstIterator It(Emitters); It; ++It)
 		{
 			FGPUSpriteParticleEmitterInstance* Emitter = *It;
 			int32 TileCount = Emitter->GetAllocatedTileCount();
-			UE_LOG(LogParticles,Info,
+			UE_LOG(LogParticles,Display,
 				TEXT("%s|%s|0x%016x %d tiles"),
 				*Emitter->Component->GetName(),*Emitter->Component->Template->GetName(),(PTRINT)Emitter, TileCount);
 			TotalAllocatedTiles += TileCount;
 		}
 
-		UE_LOG(LogParticles,Info,TEXT("---"));
-		UE_LOG(LogParticles,Info,TEXT("Total Allocated: %d"), TotalAllocatedTiles);
-		UE_LOG(LogParticles,Info,TEXT("Free (est.): %d"), GParticleSimulationTileCount - TotalAllocatedTiles);
+		UE_LOG(LogParticles,Display,TEXT("---"));
+		UE_LOG(LogParticles,Display,TEXT("Total Allocated: %d"), TotalAllocatedTiles);
+		UE_LOG(LogParticles,Display,TEXT("Free (est.): %d"), GParticleSimulationTileCount - TotalAllocatedTiles);
 		if (FXSystem)
 		{
-			UE_LOG(LogParticles,Info,TEXT("Free (actual): %d"), FXSystem->GetParticleSimulationResources()->GetFreeTileCount());
-			UE_LOG(LogParticles,Info,TEXT("Leaked: %d"), GParticleSimulationTileCount - TotalAllocatedTiles - FXSystem->GetParticleSimulationResources()->GetFreeTileCount());
+			UE_LOG(LogParticles,Display,TEXT("Free (actual): %d"), FXSystem->GetParticleSimulationResources()->GetFreeTileCount());
+			UE_LOG(LogParticles,Display,TEXT("Leaked: %d"), GParticleSimulationTileCount - TotalAllocatedTiles - FXSystem->GetParticleSimulationResources()->GetFreeTileCount());
 		}
 	}
 }
@@ -3861,7 +3861,7 @@ void DumpTileAllocations()
 FAutoConsoleCommand DumpTileAllocsCommand(
 	TEXT("FX.DumpTileAllocations"),
 	TEXT("Dump GPU particle tile allocations."),
-	FConsoleCommandDelegate::CreateRaw(DumpTileAllocations)
+	FConsoleCommandDelegate::CreateStatic(DumpTileAllocations)
 	);
 #endif // #if TRACK_TILE_ALLOCATIONS
 
