nicer UI

diff --git a/Engine/Source/Runtime/Engine/Classes/Components/PrimitiveComponent.h b/Engine/Source/Runtime/Engine/Classes/Components/PrimitiveComponent.h
--- a/Engine/Source/Runtime/Engine/Classes/Components/PrimitiveComponent.h
+++ b/Engine/Source/Runtime/Engine/Classes/Components/PrimitiveComponent.h
@@ -264,11 +264,11 @@ public:
 	uint32 bAffectDistanceFieldLighting:1;
 
 	/** Controls whether the primitive should cast shadows in the case of non precomputed shadowing.  This flag is only used if CastShadow is true. **/
-	UPROPERTY(EditAnywhere, BlueprintReadOnly, Category=Lighting, AdvancedDisplay, meta=(EditCondition="CastShadow"))
+	UPROPERTY(EditAnywhere, BlueprintReadOnly, Category=Lighting, AdvancedDisplay, meta=(EditCondition="CastShadow", DisplayName = "Dynamic Shadow"))
 	uint32 bCastDynamicShadow:1;
 
 	/** Whether the object should cast a static shadow from shadow casting lights.  This flag is only used if CastShadow is true. */
-	UPROPERTY(EditAnywhere, BlueprintReadOnly, Category=Lighting, AdvancedDisplay, meta=(EditCondition="CastShadow"))
+	UPROPERTY(EditAnywhere, BlueprintReadOnly, Category=Lighting, AdvancedDisplay, meta=(EditCondition="CastShadow", DisplayName = "Static Shadow"))
 	uint32 bCastStaticShadow:1;
 
 	/** 
@@ -276,7 +276,7 @@ public:
 	 * Volumetric translucent shadows are useful for primitives with smoothly changing opacity like particles representing a volume, 
 	 * But have artifacts when used on highly opaque surfaces.
 	 */
-	UPROPERTY(EditAnywhere, AdvancedDisplay, BlueprintReadOnly, Category=Lighting, meta=(EditCondition="CastShadow"))
+	UPROPERTY(EditAnywhere, AdvancedDisplay, BlueprintReadOnly, Category=Lighting, meta=(EditCondition="CastShadow", DisplayName = "Volumetric Translucent Shadow"))
 	uint32 bCastVolumetricTranslucentShadow:1;
 
 	/** 
@@ -289,7 +289,7 @@ public:
 	 * Whether this component should create a per-object shadow that gives higher effective shadow resolution. 
 	 * Useful for cinematic character shadowing. Assumed to be enabled if bSelfShadowOnly is enabled.
 	 */
-	UPROPERTY(EditAnywhere, AdvancedDisplay, BlueprintReadOnly, Category=Lighting, meta=(EditCondition="CastShadow"))
+	UPROPERTY(EditAnywhere, AdvancedDisplay, BlueprintReadOnly, Category=Lighting, meta=(EditCondition="CastShadow", DisplayName = "Inset Shadow"))
 	uint32 bCastInsetShadow:1;
 
 	/** 
@@ -297,11 +297,11 @@ public:
 	 *	Controls whether the primitive should cast shadows when hidden.
 	 *	This flag is only used if CastShadow is true.
 	 */
-	UPROPERTY(EditAnywhere, AdvancedDisplay, BlueprintReadOnly, Category=Lighting, meta=(EditCondition="CastShadow"))
+	UPROPERTY(EditAnywhere, AdvancedDisplay, BlueprintReadOnly, Category=Lighting, meta=(EditCondition="CastShadow", DisplayName = "Hidden Shadow"))
 	uint32 bCastHiddenShadow:1;
 
 	/** Whether this primitive should cast dynamic shadows as if it were a two sided material. */
-	UPROPERTY(EditAnywhere, AdvancedDisplay, BlueprintReadOnly, Category=Lighting, meta=(EditCondition="CastShadow"))
+	UPROPERTY(EditAnywhere, AdvancedDisplay, BlueprintReadOnly, Category=Lighting, meta=(EditCondition="CastShadow", DisplayName = "Shadow Two Sided"))
 	uint32 bCastShadowAsTwoSided:1;
 
 	/** 
