refactor shadow code

diff --git a/Engine/Source/Runtime/Engine/Private/DirectionalLightComponent.cpp b/Engine/Source/Runtime/Engine/Private/DirectionalLightComponent.cpp
--- a/Engine/Source/Runtime/Engine/Private/DirectionalLightComponent.cpp
+++ b/Engine/Source/Runtime/Engine/Private/DirectionalLightComponent.cpp
@@ -70,7 +70,7 @@ public:
 	 * Number of cascades to split the view frustum into for the whole scene dynamic shadow.  
 	 * More cascades result in better shadow resolution and allow WholeSceneDynamicShadowRadius to be further, but add rendering cost.
 	 */
-	int32 DynamicShadowCascades;
+	uint32 DynamicShadowCascades;
 
 	/** 
 	 * Exponent that is applied to the cascade transition distances as a fraction of WholeSceneDynamicShadowRadius.
@@ -100,7 +100,7 @@ public:
 		OcclusionMaskDarkness(Component->OcclusionMaskDarkness),
 		OcclusionDepthRange(Component->OcclusionDepthRange),
 		LightShaftOverrideDirection(Component->LightShaftOverrideDirection),
-		DynamicShadowCascades(Component->DynamicShadowCascades),
+		DynamicShadowCascades(Component->DynamicShadowCascades > 0 ? Component->DynamicShadowCascades : 0),
 		CascadeDistributionExponent(Component->CascadeDistributionExponent),
 		CascadeTransitionFraction(Component->CascadeTransitionFraction),
 		ShadowDistanceFadeoutFraction(Component->ShadowDistanceFadeoutFraction),
@@ -179,7 +179,7 @@ public:
 		if (HasStaticShadowing())
 		{
 			CascadeDistributionExponent = 4;
-			DynamicShadowCascades = CVarUnbuiltNumWholeSceneDynamicShadowCascades.GetValueOnAnyThread();
+			DynamicShadowCascades = FMath::Max(0, CVarUnbuiltNumWholeSceneDynamicShadowCascades.GetValueOnAnyThread());
 			WholeSceneDynamicShadowRadius = CVarUnbuiltWholeSceneDynamicShadowRadius.GetValueOnAnyThread();
 		}
 	}
@@ -195,30 +195,32 @@ public:
 	}
 
 	/** Returns the number of view dependent shadows this light will create. */
-	virtual int32 GetNumViewDependentWholeSceneShadows(const FSceneView& View) const
+	virtual uint32 GetNumViewDependentWholeSceneShadows(const FSceneView& View) const
 	{ 
-		int32 NumCascades = GetNumShadowMappedCascades(View.MaxShadowCascades);
+		uint32 NumCascades = GetNumShadowMappedCascades(View.MaxShadowCascades);
+
+		uint32 ExtraCascades = 0;
 
 		if (ShouldCreateRayTracedCascade(View.GetFeatureLevel()))
 		{
 			// Add a cascade for the ray traced shadows
-			NumCascades++;
+			ExtraCascades = 1;
 		}
 
-		return NumCascades;
+		return NumCascades + ExtraCascades;
 	}
 
 	/**
 	 * Sets up a projected shadow initializer that's dependent on the current view for shadows from the entire scene.
 	 * @return True if the whole-scene projected shadow should be used.
 	 */
-	virtual bool GetViewDependentWholeSceneProjectedShadowInitializer(const FSceneView& View, int32 SplitIndex, FWholeSceneProjectedShadowInitializer& OutInitializer) const
+	virtual bool GetViewDependentWholeSceneProjectedShadowInitializer(const FSceneView& View, uint32 InShadowSplitIndex, FWholeSceneProjectedShadowInitializer& OutInitializer) const
 	{
 		const FMatrix& WorldToLight = GetWorldToLight();
 
 		FShadowCascadeSettings CascadeSettings;
 
-		FSphere Bounds = FDirectionalLightSceneProxy::GetShadowSplitBounds(View, SplitIndex, &CascadeSettings);
+		FSphere Bounds = FDirectionalLightSceneProxy::GetShadowSplitBounds(View, InShadowSplitIndex, &CascadeSettings);
 		OutInitializer.CascadeSettings = CascadeSettings;
 
 		const float ShadowExtent = Bounds.W / FMath::Sqrt(3.0f);
@@ -234,9 +236,9 @@ public:
 		// Reduce casting distance on a directional light
 		// This is necessary to improve floating point precision in several places, especially when deriving frustum verts from InvReceiverMatrix
 		OutInitializer.MaxDistanceToCastInLightW = HALF_WORLD_MAX / 32.0f;
-		OutInitializer.SplitIndex = SplitIndex;
+		OutInitializer.InitShadowSplitIndex = InShadowSplitIndex;
 		// Last cascade is the ray traced shadow, if enabled
-		OutInitializer.bRayTracedDistanceFieldShadow = ShouldCreateRayTracedCascade(View.GetFeatureLevel()) && SplitIndex == FDirectionalLightSceneProxy::GetNumViewDependentWholeSceneShadows(View) - 1;
+		OutInitializer.bRayTracedDistanceFieldShadow = ShouldCreateRayTracedCascade(View.GetFeatureLevel()) && InShadowSplitIndex == FDirectionalLightSceneProxy::GetNumViewDependentWholeSceneShadows(View) - 1;
 		return true;
 	}
 
@@ -310,9 +312,9 @@ public:
 
 private:
 
-	int32 GetNumShadowMappedCascades(int32 MaxShadowCascades) const
+	uint32 GetNumShadowMappedCascades(uint32 MaxShadowCascades) const
 	{
-		int32 NumCascades = WholeSceneDynamicShadowRadius > 0.0f ? DynamicShadowCascades : 0;
+		uint32 NumCascades = WholeSceneDynamicShadowRadius > 0.0f ? DynamicShadowCascades : 0;
 		return FMath::Min<int32>(NumCascades, MaxShadowCascades);
 	}
 
@@ -341,7 +343,7 @@ private:
 	bool ShouldCreateRayTracedCascade(ERHIFeatureLevel::Type InFeatureLevel) const
 	{
 		//@todo - handle View.MaxShadowCascades properly
-		const int32 NumCascades = GetNumShadowMappedCascades(10);
+		const uint32 NumCascades = GetNumShadowMappedCascades(10);
 		const float RaytracedShadowDistance = GetDistanceFieldShadowDistance();
 		const bool bCreateWithCSM = NumCascades > 0 && RaytracedShadowDistance > GetCSMMaxDistance();
 		const bool bCreateWithoutCSM = NumCascades == 0 && RaytracedShadowDistance > 0;
@@ -570,18 +572,18 @@ private:
 		return CascadeSphere;
 	}
 
-	virtual FSphere GetShadowSplitBounds(const FSceneView& View, int32 SplitIndex, FShadowCascadeSettings* OutCascadeSettings) const
+	virtual FSphere GetShadowSplitBounds(const FSceneView& View, uint32 InShadowSplitIndex, FShadowCascadeSettings* OutCascadeSettings) const
 	{
 		// this checks for WholeSceneDynamicShadowRadius and DynamicShadowCascades
-		int32 NumCascades = GetNumViewDependentWholeSceneShadows(View);
+		uint32 NumCascades = GetNumViewDependentWholeSceneShadows(View);
 
 		const bool bHasRayTracedCascade = ShouldCreateRayTracedCascade(View.GetFeatureLevel());
-		const bool bIsRayTracedCascade = bHasRayTracedCascade && SplitIndex == NumCascades - 1;
+		const bool bIsRayTracedCascade = bHasRayTracedCascade && InShadowSplitIndex == NumCascades - 1;
 
 		// Determine start and end distances to the current cascade's split planes
 		// Presence of the ray traced cascade does not change depth ranges for the shadow-mapped cascades
-		float SplitNear = GetSplitDistance(View, SplitIndex, NumCascades, View.NearClippingDistance, bHasRayTracedCascade);
-		float SplitFar = GetSplitDistance(View, SplitIndex + 1, NumCascades, View.NearClippingDistance, bHasRayTracedCascade && !bIsRayTracedCascade);
+		float SplitNear = GetSplitDistance(View, InShadowSplitIndex, NumCascades, View.NearClippingDistance, bHasRayTracedCascade);
+		float SplitFar = GetSplitDistance(View, InShadowSplitIndex + 1, NumCascades, View.NearClippingDistance, bHasRayTracedCascade && !bIsRayTracedCascade);
 		float FadePlane = SplitFar;
 
 		float LocalCascadeTransitionFraction = CascadeTransitionFraction * GetShadowTransitionScale();
@@ -589,7 +591,7 @@ private:
 		float FadeExtension = (SplitFar - SplitNear) * LocalCascadeTransitionFraction;
 
 		// Add the fade region to the end of the subfrustum, if this is not the last cascade.
-		if (SplitIndex < (NumCascades - 1))
+		if (InShadowSplitIndex < NumCascades - 1)
 		{
 			SplitFar += FadeExtension;			
 		}
@@ -600,11 +602,11 @@ private:
 
 			OutCascadeSettings->SplitNearFadeRegion = 0.0f;
 
-			if(SplitIndex >= 1)
+			if(InShadowSplitIndex >= 1)
 			{
 				// todo: optimize
-				float BeforeSplitNear = GetSplitDistance(View, SplitIndex - 1, NumCascades, View.NearClippingDistance, bHasRayTracedCascade);
-				float BeforeSplitFar = GetSplitDistance(View, SplitIndex, NumCascades, View.NearClippingDistance, bHasRayTracedCascade);
+				float BeforeSplitNear = GetSplitDistance(View, InShadowSplitIndex - 1, NumCascades, View.NearClippingDistance, bHasRayTracedCascade);
+				float BeforeSplitFar = GetSplitDistance(View, InShadowSplitIndex, NumCascades, View.NearClippingDistance, bHasRayTracedCascade);
 
 				OutCascadeSettings->SplitNearFadeRegion = (BeforeSplitFar - BeforeSplitNear) * LocalCascadeTransitionFraction;
 			}
diff --git a/Engine/Source/Runtime/Engine/Private/PointLightComponent.cpp b/Engine/Source/Runtime/Engine/Private/PointLightComponent.cpp
--- a/Engine/Source/Runtime/Engine/Private/PointLightComponent.cpp
+++ b/Engine/Source/Runtime/Engine/Private/PointLightComponent.cpp
@@ -82,7 +82,7 @@ public:
 			OutInitializer.WAxis = FVector4(0,0,1,0);
 			OutInitializer.MinLightW = 0.1f;
 			OutInitializer.MaxDistanceToCastInLightW = Radius;
-			OutInitializer.SplitIndex = INDEX_NONE;
+			OutInitializer.InitShadowSplitIndex = INDEX_NONE;
 			OutInitializer.bRayTracedDistanceFieldShadow = UseRayTracedDistanceFieldShadows() && DoesPlatformSupportDistanceFieldShadowing(ViewFamily.GetShaderPlatform());
 		
 			return true;
diff --git a/Engine/Source/Runtime/Engine/Private/SpotLightComponent.cpp b/Engine/Source/Runtime/Engine/Private/SpotLightComponent.cpp
--- a/Engine/Source/Runtime/Engine/Private/SpotLightComponent.cpp
+++ b/Engine/Source/Runtime/Engine/Private/SpotLightComponent.cpp
@@ -135,7 +135,7 @@ public:
 		OutInitializer.WAxis = FVector4(0,0,1,0);
 		OutInitializer.MinLightW = 0.1f;
 		OutInitializer.MaxDistanceToCastInLightW = Radius;
-		OutInitializer.SplitIndex = INDEX_NONE;
+		OutInitializer.InitShadowSplitIndex = INDEX_NONE;
 		OutInitializer.bRayTracedDistanceFieldShadow = UseRayTracedDistanceFieldShadows() && DoesPlatformSupportDistanceFieldShadowing(ViewFamily.GetShaderPlatform());
 		return true;
 	}
diff --git a/Engine/Source/Runtime/Engine/Public/SceneManagement.h b/Engine/Source/Runtime/Engine/Public/SceneManagement.h
--- a/Engine/Source/Runtime/Engine/Public/SceneManagement.h
+++ b/Engine/Source/Runtime/Engine/Public/SceneManagement.h
@@ -566,7 +566,7 @@ public:
 class ENGINE_API FWholeSceneProjectedShadowInitializer : public FProjectedShadowInitializer
 {
 public:
-	int32 SplitIndex;
+	int32 InitShadowSplitIndex;
 	
 	FShadowCascadeSettings CascadeSettings;
 
@@ -577,9 +577,9 @@ public:
 	bool bRayTracedDistanceFieldShadow;
 
 	FWholeSceneProjectedShadowInitializer()
-	:	SplitIndex(INDEX_NONE)
-	,	bOnePassPointLightShadow(false)
-	,	bRayTracedDistanceFieldShadow(false)
+		: InitShadowSplitIndex(INDEX_NONE)
+		, bOnePassPointLightShadow(false)
+		, bRayTracedDistanceFieldShadow(false)
 	{}	
 };
 
@@ -688,7 +688,7 @@ public:
 	/** Whether this light should create per object shadows for dynamic objects. */
 	virtual bool ShouldCreatePerObjectShadowsForDynamicObjects() const;
 
-	virtual int32 GetNumViewDependentWholeSceneShadows(const FSceneView& View) const { return 0; }
+	virtual uint32 GetNumViewDependentWholeSceneShadows(const FSceneView& View) const { return 0; }
 
 	/**
 	 * Sets up a projected shadow initializer that's dependent on the current view for shadows from the entire scene.
@@ -696,7 +696,7 @@ public:
 	 */
 	virtual bool GetViewDependentWholeSceneProjectedShadowInitializer(
 		const class FSceneView& View, 
-		int32 SplitIndex,
+		uint32 InShadowSplitIndex,
 		class FWholeSceneProjectedShadowInitializer& OutInitializer) const
 	{
 		return false;
@@ -726,7 +726,7 @@ public:
 	}
 
 	// @param OutCascadeSettings can be 0
-	virtual FSphere GetShadowSplitBounds(const class FSceneView& View, int32 SplitIndex, FShadowCascadeSettings* OutCascadeSettings) const { return FSphere(FVector::ZeroVector, 0); }
+	virtual FSphere GetShadowSplitBounds(const class FSceneView& View, uint32 InShadowSplitIndex, FShadowCascadeSettings* OutCascadeSettings) const { return FSphere(FVector::ZeroVector, 0); }
 	virtual FSphere GetShadowSplitBoundsDepthRange(const FSceneView& View, float SplitNear, float SplitFar, FShadowCascadeSettings* OutCascadeSettings) const { return FSphere(FVector::ZeroVector, 0); }
 
 	virtual bool GetScissorRect(FIntRect& ScissorRect, const FSceneView& View) const
diff --git a/Engine/Source/Runtime/Renderer/Private/SceneOcclusion.cpp b/Engine/Source/Runtime/Renderer/Private/SceneOcclusion.cpp
--- a/Engine/Source/Runtime/Renderer/Private/SceneOcclusion.cpp
+++ b/Engine/Source/Runtime/Renderer/Private/SceneOcclusion.cpp
@@ -227,10 +227,10 @@ void FSceneViewState::TrimOcclusionHistory(FRHICommandListImmediate& RHICmdList,
 	}
 }
 
-bool FSceneViewState::IsShadowOccluded(FRHICommandListImmediate& RHICmdList, FPrimitiveComponentId PrimitiveId, const ULightComponent* Light, int32 SplitIndex, bool bTranslucentShadow) const
+bool FSceneViewState::IsShadowOccluded(FRHICommandListImmediate& RHICmdList, FPrimitiveComponentId PrimitiveId, const ULightComponent* Light, int32 InShadowSplitIndex, bool bTranslucentShadow) const
 {
 	// Find the shadow's occlusion query from the previous frame.
-	const FSceneViewState::FProjectedShadowKey Key(PrimitiveId, Light, SplitIndex, bTranslucentShadow);
+	const FSceneViewState::FProjectedShadowKey Key(PrimitiveId, Light, InShadowSplitIndex, bTranslucentShadow);
 
 #if BUFFERED_OCCLUSION_QUERIES
 	// Get the oldest occlusion query	
diff --git a/Engine/Source/Runtime/Renderer/Private/ShadowSetup.cpp b/Engine/Source/Runtime/Renderer/Private/ShadowSetup.cpp
--- a/Engine/Source/Runtime/Renderer/Private/ShadowSetup.cpp
+++ b/Engine/Source/Runtime/Renderer/Private/ShadowSetup.cpp
@@ -125,6 +125,25 @@ static TAutoConsoleVariable<int32> CVarUseConservativeShadowBounds(
 	TEXT("Whether to use safe and conservative shadow frustum creation that wastes some shadowmap space"),
 	ECVF_RenderThreadSafe);
 
+#if !(UE_BUILD_SHIPPING || UE_BUILD_TEST)
+// read and written on the render thread
+bool GDumpShadowSetup = false;
+void DumpShadowDumpSetup()
+{
+	ENQUEUE_UNIQUE_RENDER_COMMAND(
+		DumpShadowDumpSetup,
+	{
+		GDumpShadowSetup = true;
+	});
+}
+
+FAutoConsoleCommand CmdDumpShadowDumpSetup(
+	TEXT("r.DumpShadows"),
+	TEXT("Dump shadow setup (for developer only, only for non shiping build)"),
+	FConsoleCommandDelegate::CreateStatic(DumpShadowDumpSetup)
+	);
+#endif // !(UE_BUILD_SHIPPING || UE_BUILD_TEST)
+
 /**
  * Helper function to determine fade alpha value for shadows based on resolution. In the below ASCII art (1) is
  * the MinShadowResolution and (2) is the ShadowFadeResolution. Alpha will be 0 below the min resolution and 1
@@ -484,7 +503,7 @@ FProjectedShadowInfo::FProjectedShadowInfo(
 ,	ResolutionX(InResolutionX)
 ,	ResolutionY(InResolutionY)
 ,	MaxScreenPercent(1.0f)
-,	ShadowSplitIndex(Initializer.SplitIndex)
+,	ShadowSplitIndex(Initializer.InitShadowSplitIndex)
 ,	bAllocated(false)
 ,	bAllocatedInTranslucentLayout(false)
 ,	bRendered(false)
@@ -572,7 +591,7 @@ FProjectedShadowInfo::FProjectedShadowInfo(
 			PreShadowTranslation = -SnappedWorldPosition;
 		}
 
-		if (Initializer.SplitIndex >= 0 && bDirectionalLight)
+		if (Initializer.InitShadowSplitIndex >= 0 && bDirectionalLight)
 		{
 			checkSlow(InDependentView);
 			ShadowBounds = InLightSceneInfo->Proxy->GetShadowSplitBounds(*InDependentView, ShadowSplitIndex, 0);
@@ -709,7 +728,7 @@ void FProjectedShadowInfo::AddSubjectPrimitive(FPrimitiveSceneInfo* PrimitiveSce
 
 			if (PrimitiveSceneInfo->StaticMeshes.Num() > 0)
 			{
-				for (int32 ViewIndex = 0; ViewIndex < Views.Num(); ViewIndex++)
+				for (int32 ViewIndex = 0, ViewCount = Views.Num(); ViewIndex < ViewCount; ViewIndex++)
 				{
 					FViewInfo& CurrentView = *Views[ViewIndex];
 
@@ -1807,6 +1826,48 @@ void FSceneRenderer::InitProjectedShadowVisibility(FRHICommandListImmediate& RHI
 			}
 		}
 	}
+
+#if !(UE_BUILD_SHIPPING || UE_BUILD_TEST)
+	if(GDumpShadowSetup)
+	{
+		GDumpShadowSetup = false;
+
+		UE_LOG(LogRenderer, Display, TEXT("Dump Shadow Setup:"));
+
+		for(int32 ViewIndex = 0;ViewIndex < Views.Num();ViewIndex++)
+		{
+			FViewInfo& View = Views[ViewIndex];
+
+			UE_LOG(LogRenderer, Display, TEXT(" View  %d/%d"), ViewIndex, Views.Num());
+
+			uint32 LightIndex = 0;
+			for(TSparseArray<FLightSceneInfoCompact>::TConstIterator LightIt(Scene->Lights); LightIt; ++LightIt, ++LightIndex)
+			{
+				FVisibleLightInfo& VisibleLightInfo = VisibleLightInfos[LightIt.GetIndex()];
+				FVisibleLightViewInfo& VisibleLightViewInfo = View.VisibleLightInfos[LightIt.GetIndex()];
+
+				UE_LOG(LogRenderer, Display, TEXT("  Light %d/%d:"), LightIndex, Scene->Lights.Num());
+
+				for( int32 ShadowIndex = 0, ShadowCount = VisibleLightInfo.AllProjectedShadows.Num(); ShadowIndex < ShadowCount; ShadowIndex++ )
+				{
+					FProjectedShadowInfo& ProjectedShadowInfo = *VisibleLightInfo.AllProjectedShadows[ShadowIndex];
+
+					if(VisibleLightViewInfo.bInViewFrustum)
+					{
+						UE_LOG(LogRenderer, Display, TEXT("   Shadow %d/%d: %d"),  ShadowIndex, ShadowCount, ProjectedShadowInfo.ShadowId);
+						UE_LOG(LogRenderer, Display, TEXT("    WholeSceneDir=%d SplitIndex=%d near=%f far=%f"),
+							ProjectedShadowInfo.IsWholeSceneDirectionalShadow(),
+							ProjectedShadowInfo.ShadowSplitIndex,
+							ProjectedShadowInfo.CascadeSettings.SplitNear,
+							ProjectedShadowInfo.CascadeSettings.SplitFar);
+						UE_LOG(LogRenderer, Display, TEXT("    bDistField=%d"),
+							ProjectedShadowInfo.bRayTracedDistanceFieldShadow);
+					}
+				}
+			}
+		}
+	}
+#endif // !(UE_BUILD_SHIPPING || UE_BUILD_TEST)
 }
 
 void FSceneRenderer::GatherShadowDynamicMeshElements()
@@ -2063,12 +2124,12 @@ void FSceneRenderer::AddViewDependentWholeSceneShadowsForView(
 		// If rendering in stereo mode we render shadow depths only for the left eye, but project for both eyes!
 		if (View.StereoPass != eSSP_RIGHT_EYE)
 		{
-			const int32 NumSplits = LightSceneInfo.Proxy->GetNumViewDependentWholeSceneShadows(View);
-			for (int32 SplitIndex = 0; SplitIndex < NumSplits; SplitIndex++)
+			const uint32 NumSplits = LightSceneInfo.Proxy->GetNumViewDependentWholeSceneShadows(View);
+			for (uint32 LocalShadowSplitIndex = 0; LocalShadowSplitIndex < NumSplits; LocalShadowSplitIndex++)
 			{
 				FWholeSceneProjectedShadowInitializer ProjectedShadowInitializer;
 
-				if (LightSceneInfo.Proxy->GetViewDependentWholeSceneProjectedShadowInitializer(View, SplitIndex, ProjectedShadowInitializer))
+				if (LightSceneInfo.Proxy->GetViewDependentWholeSceneProjectedShadowInitializer(View, LocalShadowSplitIndex, ProjectedShadowInitializer))
 				{
 					const FIntPoint ShadowBufferResolution = GSceneRenderTargets.GetShadowDepthTextureResolution();
 					// Create the projected shadow info.
@@ -2111,7 +2172,7 @@ void FSceneRenderer::AddViewDependentWholeSceneShadowsForView(
 					if (LightSceneInfo.Proxy->GetViewDependentRsmWholeSceneProjectedShadowInitializer(View, Lpv.GetBoundingBox(), ProjectedShadowInitializer))
 					{
 						// moved out from the FProjectedShadowInfo constructor
-						ProjectedShadowInitializer.SplitIndex = 0;
+						ProjectedShadowInitializer.InitShadowSplitIndex = 0;
 
 						const FIntPoint ShadowBufferResolution = GSceneRenderTargets.GetReflectiveShadowMapTextureResolution();
 
diff --git a/Engine/Source/Runtime/Renderer/Private/TranslucentLighting.cpp b/Engine/Source/Runtime/Renderer/Private/TranslucentLighting.cpp
--- a/Engine/Source/Runtime/Renderer/Private/TranslucentLighting.cpp
+++ b/Engine/Source/Runtime/Renderer/Private/TranslucentLighting.cpp
@@ -696,7 +696,7 @@ public:
 		const FViewInfo& View, 
 		const FLightSceneInfo* LightSceneInfo, 
 		const FProjectedShadowInfo* ShadowMap, 
-		int32 VolumeCascadeIndexValue, 
+		uint32 VolumeCascadeIndexValue,
 		bool bDynamicallyShadowed) const
 	{
 		SetDeferredLightParameters(RHICmdList, ShaderRHI, Shader->GetUniformBufferParameter<FDeferredLightUniformStruct>(), LightSceneInfo, View);
@@ -754,11 +754,11 @@ public:
 	}
 	FTranslucentObjectShadowingPS() {}
 
-	void SetParameters(FRHICommandList& RHICmdList, const FViewInfo& View, const FLightSceneInfo* LightSceneInfo, const FProjectedShadowInfo* ShadowMap, int32 VolumeCascadeIndexValue)
+	void SetParameters(FRHICommandList& RHICmdList, const FViewInfo& View, const FLightSceneInfo* LightSceneInfo, const FProjectedShadowInfo* ShadowMap, uint32 VolumeCascadeIndex)
 	{
 		FGlobalShader::SetParameters(RHICmdList, GetPixelShader(), View);
 		TranslucencyProjectionParameters.Set(RHICmdList, this);
-		TranslucentInjectParameters.Set(RHICmdList, GetPixelShader(), this, View, LightSceneInfo, ShadowMap, VolumeCascadeIndexValue, true);
+		TranslucentInjectParameters.Set(RHICmdList, GetPixelShader(), this, View, LightSceneInfo, ShadowMap, VolumeCascadeIndex, true);
 	}
 
 	virtual bool Serialize(FArchive& Ar)
@@ -856,7 +856,7 @@ public:
 			FShadowCascadeSettings ShadowCascadeSettings;
 
 			// todo: optimize, not all computed data is needed
-			InnerBounds = LightSceneInfo->Proxy->GetShadowSplitBounds(View, InnerSplitIndex, &ShadowCascadeSettings);
+			InnerBounds = LightSceneInfo->Proxy->GetShadowSplitBounds(View, (uint32)InnerSplitIndex, &ShadowCascadeSettings);
 
 			// near cascade plane
 			{
@@ -864,7 +864,7 @@ public:
 				Planes[0] = FVector4((FVector)(ShadowCascadeSettings.NearFrustumPlane),  -ShadowCascadeSettings.NearFrustumPlane.W);
 			}
 
-			int32 CascadeCount = LightSceneInfo->Proxy->GetNumViewDependentWholeSceneShadows(View);
+			uint32 CascadeCount = LightSceneInfo->Proxy->GetNumViewDependentWholeSceneShadows(View);
 
 			// far cascade plane
 			if(InnerSplitIndex != CascadeCount - 1)
@@ -1233,7 +1233,7 @@ void FDeferredShadingSceneRenderer::ClearTranslucentVolumePerObjectShadowing(FRH
 }
 
 /** Calculates volume texture bounds for the given light in the given translucent lighting volume cascade. */
-FVolumeBounds CalculateLightVolumeBounds(const FSphere& LightBounds, const FViewInfo& View, int32 VolumeCascadeIndex, bool bDirectionalLight)
+FVolumeBounds CalculateLightVolumeBounds(const FSphere& LightBounds, const FViewInfo& View, uint32 VolumeCascadeIndex, bool bDirectionalLight)
 {
 	FVolumeBounds VolumeBounds;
 
@@ -1283,7 +1283,7 @@ void FDeferredShadingSceneRenderer::AccumulateTranslucentVolumeObjectShadowing(F
 		auto ShaderMap = GetGlobalShaderMap(FeatureLevel);
 
 		// Inject into each volume cascade
-		for (int32 VolumeCascadeIndex = 0; VolumeCascadeIndex < TVC_MAX; VolumeCascadeIndex++)
+		for (uint32 VolumeCascadeIndex = 0; VolumeCascadeIndex < TVC_MAX; VolumeCascadeIndex++)
 		{
 			//@todo - support multiple views
 			const FViewInfo& View = Views[0];
@@ -1494,7 +1494,7 @@ static void InjectTranslucentLightArray(FRHICommandListImmediate& RHICmdList, co
 
 	// Inject into each volume cascade
 	// Operate on one cascade at a time to reduce render target switches
-	for (int32 VolumeCascadeIndex = 0; VolumeCascadeIndex < TVC_MAX; VolumeCascadeIndex++)
+	for (uint32 VolumeCascadeIndex = 0; VolumeCascadeIndex < TVC_MAX; VolumeCascadeIndex++)
 	{
 		const IPooledRenderTarget* RT0 = GSceneRenderTargets.TranslucencyLightingVolumeAmbient[VolumeCascadeIndex];
 		const IPooledRenderTarget* RT1 = GSceneRenderTargets.TranslucencyLightingVolumeDirectional[VolumeCascadeIndex];
