nicer UI

diff --git a/Engine/Source/Runtime/Engine/Classes/Components/DirectionalLightComponent.h b/Engine/Source/Runtime/Engine/Classes/Components/DirectionalLightComponent.h
--- a/Engine/Source/Runtime/Engine/Classes/Components/DirectionalLightComponent.h
+++ b/Engine/Source/Runtime/Engine/Classes/Components/DirectionalLightComponent.h
@@ -92,7 +92,7 @@ class UDirectionalLightComponent : public ULightComponent
 	 * If DynamicShadowDistanceStationaryLight is large (currently > 8000), this will be forced off.
 	 * Disabling this can reduce shadowing cost significantly with many movable objects.
 	 */
-	UPROPERTY(EditAnywhere, AdvancedDisplay, BlueprintReadOnly, Category=CascadedShadowMaps)
+	UPROPERTY(EditAnywhere, AdvancedDisplay, BlueprintReadOnly, Category=CascadedShadowMaps, DisplayName = "Inset Shadows For Movable Objects")
 	uint32 bUseInsetShadowsForMovableObjects : 1;
 
 	/** 
diff --git a/Engine/Source/Runtime/Engine/Classes/Components/LightComponent.h b/Engine/Source/Runtime/Engine/Classes/Components/LightComponent.h
--- a/Engine/Source/Runtime/Engine/Classes/Components/LightComponent.h
+++ b/Engine/Source/Runtime/Engine/Classes/Components/LightComponent.h
@@ -151,7 +151,7 @@ class ENGINE_API ULightComponent : public ULightComponentBase
 	 * They have less aliasing artifacts than standard shadowmaps, but inherit all the limitations of distance field representations (only uniform scale, no deformation).
 	 * These shadows have a low per-object cost (and don't depend on triangle count) so they are effective for distant shadows from a dynamic sun.
 	 */
-	UPROPERTY(EditAnywhere, BlueprintReadOnly, Category=DistanceFieldShadows, meta=(DisplayName = "Use RayTraced DistanceField Shadows"))
+	UPROPERTY(EditAnywhere, BlueprintReadOnly, Category=DistanceFieldShadows, meta=(DisplayName = "RayTraced DistanceField Shadows"))
 	bool bUseRayTracedDistanceFieldShadows;
 
 public:
