Analytics for Engine and Launcher now set the UserID to be MachineID|AccountID. AccountID will be blank until the user actually logs in. #UETOOL-239
Note for analytics: your ETL will have to start splitting up these fields starting at the next promoted build (exact CL# unknown until one is promoted, but it will be anything >= this CL in main).

diff --git a/Engine/Source/Runtime/Engine/Private/EngineAnalytics.cpp b/Engine/Source/Runtime/Engine/Private/EngineAnalytics.cpp
--- a/Engine/Source/Runtime/Engine/Private/EngineAnalytics.cpp
+++ b/Engine/Source/Runtime/Engine/Private/EngineAnalytics.cpp
@@ -202,7 +202,7 @@ void FEngineAnalytics::Initialize()
 				DefaultEngineAnalyticsConfig);
 			if (Analytics.IsValid())
 			{
-				Analytics->SetUserID(FPlatformMisc::GetUniqueDeviceId());
+				Analytics->SetUserID(FString::Printf(TEXT("%s|%s"), *FPlatformMisc::GetMachineId().ToString(EGuidFormats::Digits).ToLower(), *FPlatformMisc::GetEpicAccountId()));
 				Analytics->StartSession();
 			}
 		}
