Opting out of usage events will still create an analytics provider and start a session, it just won't send usage events. #UE-7588. We still need to make the setting a global one, and not a project setting. That is tracked by a different ticket.

diff --git a/Engine/Source/Runtime/Engine/Private/EngineAnalytics.cpp b/Engine/Source/Runtime/Engine/Private/EngineAnalytics.cpp
--- a/Engine/Source/Runtime/Engine/Private/EngineAnalytics.cpp
+++ b/Engine/Source/Runtime/Engine/Private/EngineAnalytics.cpp
@@ -9,6 +9,7 @@
 
 bool FEngineAnalytics::bIsInitialized;
 TSharedPtr<IAnalyticsProvider> FEngineAnalytics::Analytics;
+bool FEngineAnalytics::bShouldSendUsageEvents;
 
 #define DEBUG_ANALYTICS 0
 
@@ -110,7 +111,8 @@ FAnalytics::FProviderConfigurationDelegate& GetEngineAnalyticsOverrideConfigDele
  */
 IAnalyticsProvider& FEngineAnalytics::GetProvider()
 {
-	checkf(bIsInitialized && Analytics.IsValid(), TEXT("FEngineAnalytics::GetProvider called outside of Initialize/Shutdown."));
+	checkf(bIsInitialized && IsAvailable(), TEXT("FEngineAnalytics::GetProvider called outside of Initialize/Shutdown."));
+
 	return *Analytics.Get();
 }
  
@@ -119,19 +121,19 @@ void FEngineAnalytics::Initialize()
 	checkf(!bIsInitialized, TEXT("FEngineAnalytics::Initialize called more than once."));
 
 	// Never use analytics when running a commandlet tool
-	bool bShouldInitAnalytics = !IsRunningCommandlet();
+	const bool bShouldInitAnalytics = !IsRunningCommandlet();
 
 	check(GEngine);
 #if WITH_EDITORONLY_DATA
 	if (GIsEditor)
 	{
-		bShouldInitAnalytics = bShouldInitAnalytics && GEngine->bEditorAnalyticsEnabled;
+		bShouldSendUsageEvents = GEngine->bEditorAnalyticsEnabled;
 	}
 	else
 #endif
 	{
 		// Outside of the editor, the only engine analytics usage is the hardware survey
-		bShouldInitAnalytics = bShouldInitAnalytics && GEngine->bHardwareSurveyEnabled;
+		bShouldSendUsageEvents = GEngine->bHardwareSurveyEnabled;
 	}
 
 	if (bShouldInitAnalytics)
diff --git a/Engine/Source/Runtime/Engine/Public/EngineAnalytics.h b/Engine/Source/Runtime/Engine/Public/EngineAnalytics.h
--- a/Engine/Source/Runtime/Engine/Public/EngineAnalytics.h
+++ b/Engine/Source/Runtime/Engine/Public/EngineAnalytics.h
@@ -20,7 +20,7 @@ public:
 	 */
 	static ENGINE_API IAnalyticsProvider& GetProvider();
 	/** Helper function to determine if the provider is valid. */
-	static ENGINE_API bool IsAvailable() { return Analytics.IsValid(); }
+	static ENGINE_API bool IsAvailable() { return bShouldSendUsageEvents && Analytics.IsValid(); }
 	/** Called to initialize the singleton. */
 	static void Initialize();
 	/** Called to shut down the singleton */
@@ -28,6 +28,8 @@ public:
 
 private:
 	static bool bIsInitialized;
+	/** This allows us to "opt-out" of engine analytics. We still create a session, but we don't send any usage events. */
+	static bool bShouldSendUsageEvents;
 	static TSharedPtr<IAnalyticsProvider> Analytics;
 };
 
