Fixed navmesh streaming with "DynamicModifiersOnly" option

diff --git a/Engine/Source/Runtime/Engine/Private/AI/Navigation/RecastNavMeshDataChunk.cpp b/Engine/Source/Runtime/Engine/Private/AI/Navigation/RecastNavMeshDataChunk.cpp
--- a/Engine/Source/Runtime/Engine/Private/AI/Navigation/RecastNavMeshDataChunk.cpp
+++ b/Engine/Source/Runtime/Engine/Private/AI/Navigation/RecastNavMeshDataChunk.cpp
@@ -209,8 +209,9 @@ TArray<uint32> URecastNavMeshDataChunk::DetachTiles(FPImplRecastNavMesh* NavMesh
 		if (TileData.TileRef != 0)
 		{
 			// Detach tile cache layer and take ownership over compressed data
+			const dtMeshTile* MeshTile = NavMesh->getTileByRef(TileData.TileRef);
+			if (MeshTile)
 			{
-				const dtMeshTile* MeshTile = NavMesh->getTileByRef(TileData.TileRef);
 				FNavMeshTileData TileCacheData = NavMeshImpl->GetTileCacheLayer(MeshTile->header->x, MeshTile->header->y, MeshTile->header->layer);
 				if (TileCacheData.IsValid())
 				{
diff --git a/Engine/Source/Runtime/Engine/Private/AI/Navigation/RecastNavMeshGenerator.cpp b/Engine/Source/Runtime/Engine/Private/AI/Navigation/RecastNavMeshGenerator.cpp
--- a/Engine/Source/Runtime/Engine/Private/AI/Navigation/RecastNavMeshGenerator.cpp
+++ b/Engine/Source/Runtime/Engine/Private/AI/Navigation/RecastNavMeshGenerator.cpp
@@ -2677,6 +2677,12 @@ static int32 CaclulateMaxTilesCount(const TNavStatArray<FBox>& NavigableAreas, f
 	return FMath::CeilToInt(GridCellsCount * AvgLayersPerGridCell);
 }
 
+// Whether navmesh is static, does not support rebuild from geometry
+static bool IsGameStaticNavMesh(ARecastNavMesh* InNavMesh)
+{
+	return (InNavMesh->GetWorld()->IsGameWorld() && InNavMesh->RuntimeGeneration != ERuntimeGenerationType::Dynamic);
+}
+
 //----------------------------------------------------------------------//
 // FRecastNavMeshGenerator
 //----------------------------------------------------------------------//
@@ -3043,7 +3049,7 @@ void FRecastNavMeshGenerator::OnNavigationBoundsChanged()
 	UpdateNavigationBounds();
 	
 	dtNavMesh* DetourMesh = DestNavMesh->GetRecastNavMeshImpl()->GetRecastMesh();
-	if (DestNavMesh->IsResizable() && DetourMesh)
+	if (!IsGameStaticNavMesh(DestNavMesh) && DestNavMesh->IsResizable() && DetourMesh)
 	{
 		// Check whether Navmesh size needs to be changed
 		int32 MaxRequestedTiles = CaclulateMaxTilesCount(InclusionBounds, Config.tileSize * Config.cs, AvgLayersPerTile);
@@ -3173,7 +3179,7 @@ TArray<uint32> FRecastNavMeshGenerator::AddGeneratedTiles(const FRecastTileGener
 				dtTileRef ResultTileRef = 0;
 
 				// let navmesh know it's tile generator who owns the data
-				dtStatus status = DetourMesh->addTile(TileLayers[i].GetData(), TileLayers[i].DataSize, DT_TILE_FREE_DATA, 0, &ResultTileRef);
+				dtStatus status = DetourMesh->addTile(TileLayers[i].GetData(), TileLayers[i].DataSize, DT_TILE_FREE_DATA, OldTileRef, &ResultTileRef);
 
 				if (dtStatusFailed(status))
 				{
@@ -3272,10 +3278,18 @@ void FRecastNavMeshGenerator::MarkDirtyTiles(const TArray<FNavigationDirtyArea>&
 	const FVector NavMeshOrigin = FVector::ZeroVector;
 	int32 NumTilesMarked = 0;
 
+	const bool bGameStaticNavMesh = IsGameStaticNavMesh(DestNavMesh);
+		
 	// find all tiles that need regeneration
 	TSet<FPendingTileElement> DirtyTiles;
 	for (const FNavigationDirtyArea& DirtyArea : DirtyAreas)
 	{
+		// Static navmeshes accept only area modifiers updates
+		if (bGameStaticNavMesh && (!DirtyArea.HasFlag(ENavigationDirtyFlag::DynamicModifier) || DirtyArea.HasFlag(ENavigationDirtyFlag::NavigationBounds)))
+		{
+			continue;
+		}
+		
 		bool bDoTileInclusionTest = false;
 		FBox AdjustedAreaBounds = DirtyArea.Bounds;
 		
@@ -3457,7 +3471,8 @@ TArray<uint32> FRecastNavMeshGenerator::ProcessTileTasks(const int32 NumTasksToS
 	
 	TArray<uint32> UpdatedTiles;
 	const bool bHasTasksAtStart = GetNumRemaningBuildTasks() > 0;
-	
+	const bool bGameStaticNavMesh = IsGameStaticNavMesh(DestNavMesh);
+			
 	int32 NumSubmittedTasks = 0;
 	// Submit pending tile elements
 	for (int32 ElementIdx = PendingDirtyTiles.Num()-1; ElementIdx >= 0 && NumSubmittedTasks < NumTasksToSubmit; ElementIdx--)
@@ -3480,7 +3495,7 @@ TArray<uint32> FRecastNavMeshGenerator::ProcessTileTasks(const int32 NumTasksToS
 				RunningDirtyTiles.Add(RunningElement);
 				NumSubmittedTasks++;
 			}
-			else
+			else if (!bGameStaticNavMesh)
 			{
 				// If there is nothing to generate remove all tiles from navmesh at specified grid coordinates
 				UpdatedTiles.Append(
diff --git a/Engine/Source/Runtime/Navmesh/Public/Detour/DetourNavMesh.h b/Engine/Source/Runtime/Navmesh/Public/Detour/DetourNavMesh.h
--- a/Engine/Source/Runtime/Navmesh/Public/Detour/DetourNavMesh.h
+++ b/Engine/Source/Runtime/Navmesh/Public/Detour/DetourNavMesh.h
@@ -736,9 +736,7 @@ public:
 	inline bool isEmpty() const
 	{
 		// has no tile grid set up
-		return (m_tileWidth > 0 && m_tileHeight > 0) == false
-			// or no actual tiles present (first free tile is the first one in tiles array)
-			|| (m_nextFree == m_tiles);
+		return (m_tileWidth > 0 && m_tileHeight > 0) == false;
 	}
 	//@UE4 END
 	
