UE-7701 - Fixed crash caused by missing shader.
bIsMasked was being used to cache off if the blend mode was BLEND_Masked and the material met some other conditions. This was then accessind in IsMasked(). Several places however were accessing GetBlendMode() and testing for BLEND_Masked directly. This lead to a mismatch between various code paths considering a material masked. Leading to a missing shader and the crash.
I've deprecated bIsMasked and placed the extra checks in GetBlendMode() so that all access of a material's blend mode are consistent. The extra checkes aren't expensive.

diff --git a/Engine/Source/Runtime/Engine/Classes/Materials/Material.h b/Engine/Source/Runtime/Engine/Classes/Materials/Material.h
--- a/Engine/Source/Runtime/Engine/Classes/Materials/Material.h
+++ b/Engine/Source/Runtime/Engine/Classes/Materials/Material.h
@@ -644,9 +644,9 @@ public:
 	UPROPERTY()
 	TArray<struct FMaterialParameterCollectionInfo> MaterialParameterCollectionInfos;
 
-	/** true if Material is masked and uses custom opacity */
-	UPROPERTY()
-	uint32 bIsMasked:1;
+ 	/** true if Material is masked and uses custom opacity */
+ 	UPROPERTY()
+ 	uint32 bIsMasked_DEPRECATED:1;
 
 	/** true if Material is the preview material used in the material editor. */
 	UPROPERTY(transient, duplicatetransient)
diff --git a/Engine/Source/Runtime/Engine/Private/Materials/Material.cpp b/Engine/Source/Runtime/Engine/Private/Materials/Material.cpp
--- a/Engine/Source/Runtime/Engine/Private/Materials/Material.cpp
+++ b/Engine/Source/Runtime/Engine/Private/Materials/Material.cpp
@@ -2278,9 +2278,6 @@ void UMaterial::PostEditChangeProperty(FPropertyChangedEvent& PropertyChangedEve
 		}
 	}
 
-	// Check if the material is masked and uses a custom opacity (that's not 1.0f).
-	bIsMasked = (((EBlendMode(BlendMode) == BLEND_Masked) && (OpacityMask.Expression || (OpacityMask.UseConstant && OpacityMask.Constant < 0.999f))) || bUseMaterialAttributes);
-
 	bool bRequiresCompilation = true;
 	if( PropertyThatChanged ) 
 	{
@@ -3444,7 +3441,22 @@ float UMaterial::GetOpacityMaskClipValue(bool bIsInGameThread) const
 
 EBlendMode UMaterial::GetBlendMode(bool bIsInGameThread) const
 {
-	return BlendMode;
+	if (EBlendMode(BlendMode) == BLEND_Masked)
+	{
+		// Check if the material is masked and uses a custom opacity (that's not 1.0f).
+		if ((((OpacityMask.Expression || (OpacityMask.UseConstant && OpacityMask.Constant < 0.999f))) || bUseMaterialAttributes))
+		{
+			return BLEND_Masked;
+		}
+		else
+		{
+			return BLEND_Opaque;
+		}
+	}
+	else
+	{
+		return BlendMode;
+	}
 }
 
 EMaterialShadingModel UMaterial::GetShadingModel(bool bIsInGameThread) const
@@ -3473,7 +3485,7 @@ bool UMaterial::IsTwoSided(bool bIsInGameThread) const
 
 bool UMaterial::IsMasked(bool bIsInGameThread) const
 {
-	return bIsMasked != 0;
+	return GetBlendMode() == BLEND_Masked;
 }
 
 USubsurfaceProfile* UMaterial::GetSubsurfaceProfile_Internal() const
