Inheritable Component - fixes, ranaming (basic case)

diff --git a/Engine/Source/Editor/Kismet/Private/SSCSEditor.cpp b/Engine/Source/Editor/Kismet/Private/SSCSEditor.cpp
--- a/Engine/Source/Editor/Kismet/Private/SSCSEditor.cpp
+++ b/Engine/Source/Editor/Kismet/Private/SSCSEditor.cpp
@@ -859,13 +859,6 @@ UActorComponent* FSCSEditorTreeNode::GetOverridenComponentTemplate(UBlueprint* B
 	UActorComponent* OverridenComponent = NULL;
 
 	FComponentKey Key(GetSCSNode());
-	if (!Key.IsValid())
-	{
-		UActorComponent* ComponentTemplate = GetComponentTemplate();
-		auto OriginalComponentOwner = ComponentTemplate ? ComponentTemplate->GetOwner() : nullptr;
-		auto ComponentOwnerBPGC = OriginalComponentOwner ? Cast<UBlueprintGeneratedClass>(OriginalComponentOwner->GetClass()) : nullptr;
-		Key = FComponentKey(ComponentOwnerBPGC, GetVariableName());
-	}
 
 	const bool BlueprintCanOverrideComponentFormKey = Key.IsValid() 
 		&& Blueprint 
diff --git a/Engine/Source/Editor/KismetCompiler/Private/KismetCompiler.cpp b/Engine/Source/Editor/KismetCompiler/Private/KismetCompiler.cpp
--- a/Engine/Source/Editor/KismetCompiler/Private/KismetCompiler.cpp
+++ b/Engine/Source/Editor/KismetCompiler/Private/KismetCompiler.cpp
@@ -3277,9 +3277,10 @@ void FKismetCompilerContext::Compile()
 
 	if (CompileOptions.CompileType == EKismetCompileType::Full)
 	{
-		if (Blueprint->InheritableComponentHandler)
+		auto InheritableComponentHandler = Blueprint->GetInheritableComponentHandler(false);
+		if (InheritableComponentHandler)
 		{
-			Blueprint->InheritableComponentHandler->RemoveInvalidAndUnnecessaryTemplates();
+			InheritableComponentHandler->RemoveInvalidAndUnnecessaryTemplates();
 		}
 	}
 
diff --git a/Engine/Source/Editor/UnrealEd/Private/Kismet2/BlueprintEditorUtils.cpp b/Engine/Source/Editor/UnrealEd/Private/Kismet2/BlueprintEditorUtils.cpp
--- a/Engine/Source/Editor/UnrealEd/Private/Kismet2/BlueprintEditorUtils.cpp
+++ b/Engine/Source/Editor/UnrealEd/Private/Kismet2/BlueprintEditorUtils.cpp
@@ -2421,9 +2421,13 @@ bool FBlueprintEditorUtils::IsDataOnlyBlueprint(const UBlueprint* Blueprint)
 		return false;
 	}
 
-	if (Blueprint->InheritableComponentHandler && !Blueprint->InheritableComponentHandler->IsEmpty())
+	static const FBoolConfigValueHelper EnableInheritableComponents(TEXT("Kismet"), TEXT("bEnableInheritableComponents"), GEngineIni);
+	if (EnableInheritableComponents)
 	{
-		return false;
+		if (Blueprint->InheritableComponentHandler && !Blueprint->InheritableComponentHandler->IsEmpty())
+		{
+			return false;
+		}
 	}
 
 	return true;
@@ -3514,7 +3518,24 @@ void FBlueprintEditorUtils::RenameComponentMemberVariable(UBlueprint* Blueprint,
 	if (!NewName.IsEqual(Node->VariableName, ENameCase::CaseSensitive))
 	{
 		Blueprint->Modify();
-		
+
+		// Rename Inheritable Component Templates
+		{
+			const FComponentKey Key(Node);
+			TArray<UBlueprint*> Dependents;
+			GetDependentBlueprints(Blueprint, Dependents);
+			for (auto DepBP : Dependents)
+			{
+				auto InheritableComponentHandler = DepBP ? DepBP->GetInheritableComponentHandler(false) : nullptr;
+				if (InheritableComponentHandler && InheritableComponentHandler->GetOverridenComponentTemplate(Key))
+				{
+					InheritableComponentHandler->Modify();
+					InheritableComponentHandler->RenameTemplate(Key, NewName);
+					InheritableComponentHandler->MarkPackageDirty();
+				}
+			}
+		}
+
 		// Update the name
 		const FName OldName = Node->VariableName;
 		Node->Modify();
diff --git a/Engine/Source/Runtime/Engine/Classes/Engine/InheritableComponentHandler.h b/Engine/Source/Runtime/Engine/Classes/Engine/InheritableComponentHandler.h
--- a/Engine/Source/Runtime/Engine/Classes/Engine/InheritableComponentHandler.h
+++ b/Engine/Source/Runtime/Engine/Classes/Engine/InheritableComponentHandler.h
@@ -24,14 +24,9 @@ struct ENGINE_API FComponentKey
 		: OwnerClass(nullptr)
 	{}
 
-	FComponentKey(UBlueprintGeneratedClass* InOwnerClass, FName InVariableNam)
-		: OwnerClass(InOwnerClass)
-		, VariableName(InVariableNam)
-	{}
-
 	FComponentKey(USCS_Node* ParentNode);
 
-	bool Match(FComponentKey OtherKey) const;
+	bool Match(const FComponentKey OtherKey) const;
 
 	bool IsValid() const
 	{
@@ -89,13 +84,15 @@ public:
 
 	void PreloadAllTempates();
 
+	bool RenameTemplate(FComponentKey OldKey, FName NewName);
+
 #endif
 
 public:
 	UActorComponent* GetOverridenComponentTemplate(FComponentKey Key) const;
 
 private:
-	const FComponentOverrideRecord* FindRecord(FComponentKey Key) const;
+	const FComponentOverrideRecord* FindRecord(const FComponentKey Key) const;
 	
 	UPROPERTY()
 	TArray<FComponentOverrideRecord> Records;
diff --git a/Engine/Source/Runtime/Engine/Private/InheritableComponentHandler.cpp b/Engine/Source/Runtime/Engine/Private/InheritableComponentHandler.cpp
--- a/Engine/Source/Runtime/Engine/Private/InheritableComponentHandler.cpp
+++ b/Engine/Source/Runtime/Engine/Private/InheritableComponentHandler.cpp
@@ -218,6 +218,19 @@ UActorComponent* UInheritableComponentHandler::FindBestArchetype(FComponentKey K
 	return ClosestArchetype;
 }
 
+bool UInheritableComponentHandler::RenameTemplate(FComponentKey OldKey, FName NewName)
+{
+	for (auto& Record : Records)
+	{
+		if (Record.ComponentKey.Match(OldKey))
+		{
+			Record.ComponentKey.VariableName = NewName;
+			return true;
+		}
+	}
+	return false;
+}
+
 #endif
 
 UActorComponent* UInheritableComponentHandler::GetOverridenComponentTemplate(FComponentKey Key) const
@@ -226,7 +239,7 @@ UActorComponent* UInheritableComponentHandler::GetOverridenComponentTemplate(FCo
 	return Record ? Record->ComponentTemplate : nullptr;
 }
 
-const FComponentOverrideRecord* UInheritableComponentHandler::FindRecord(FComponentKey Key) const
+const FComponentOverrideRecord* UInheritableComponentHandler::FindRecord(const FComponentKey Key) const
 {
 	for (auto& Record : Records)
 	{
@@ -242,12 +255,31 @@ const FComponentOverrideRecord* UInheritableComponentHandler::FindRecord(FCompon
 
 FComponentKey::FComponentKey(USCS_Node* ParentNode) : OwnerClass(nullptr)
 {
+	VariableName = ParentNode ? ParentNode->GetVariableName() : NAME_None;
+	if (VariableName == NAME_None)
+	{
+		return;
+	}
+
+	//Find original blueprint where the component is declared
 	auto ParentSCS = ParentNode ? ParentNode->GetSCS() : nullptr;
 	OwnerClass = ParentSCS ? Cast<UBlueprintGeneratedClass>(ParentSCS->GetOwnerClass()) : nullptr;
-	VariableName = (ParentNode && OwnerClass) ? ParentNode->GetVariableName() : NAME_None;
+	for (auto SuperOwnerClass = OwnerClass ? Cast<UBlueprintGeneratedClass>(OwnerClass->GetSuperClass()) : nullptr; 
+		SuperOwnerClass;
+		SuperOwnerClass = Cast<UBlueprintGeneratedClass>(SuperOwnerClass->GetSuperClass()))
+	{
+		if (SuperOwnerClass->SimpleConstructionScript && SuperOwnerClass->SimpleConstructionScript->FindSCSNode(VariableName))
+		{
+			OwnerClass = SuperOwnerClass;
+		}
+		else
+		{
+			break;
+		}
+	}
 }
 
-bool FComponentKey::Match(FComponentKey OtherKey) const
+bool FComponentKey::Match(const FComponentKey OtherKey) const
 {
 	return (OwnerClass == OtherKey.OwnerClass) && (VariableName == OtherKey.VariableName);
 }
diff --git a/Engine/Source/Runtime/Engine/Private/SCS_Node.cpp b/Engine/Source/Runtime/Engine/Private/SCS_Node.cpp
--- a/Engine/Source/Runtime/Engine/Private/SCS_Node.cpp
+++ b/Engine/Source/Runtime/Engine/Private/SCS_Node.cpp
@@ -32,13 +32,16 @@ UActorComponent* USCS_Node::ExecuteNodeOnActor(AActor* Actor, USceneComponent* P
 	check((ParentComponent != nullptr && !ParentComponent->IsPendingKill()) || (RootTransform != nullptr)); // must specify either a parent component or a world transform
 
 	UActorComponent* OverridenComponentTemplate = nullptr;
+	static const FBoolConfigValueHelper EnableInheritableComponents(TEXT("Kismet"), TEXT("bEnableInheritableComponents"), GEngineIni);
+	if (EnableInheritableComponents)
 	{
+		const FComponentKey ComponentKey(this);
 		auto ActualBPGC = Cast<UBlueprintGeneratedClass>(Actor->GetClass());
 		while (!OverridenComponentTemplate && ActualBPGC)
 		{
 			if (ActualBPGC->InheritableComponentHandler)
 			{
-				OverridenComponentTemplate = ActualBPGC->InheritableComponentHandler->GetOverridenComponentTemplate(this);
+				OverridenComponentTemplate = ActualBPGC->InheritableComponentHandler->GetOverridenComponentTemplate(ComponentKey);
 			}
 			ActualBPGC = Cast<UBlueprintGeneratedClass>(ActualBPGC->GetSuperClass());
 		}
