Integrate CL 2412131 from UE4-UT

On disconnects while we have a PendingNetGame, transition to the world that has the PendingNetGame object if there is no world yet on the PendingNetGame object (this happens when we've requested to join a server but haven't gotten the NMT_Welcome message yet). Fixes UE-7525.

#fixversion=4.7

diff --git a/Engine/Source/Runtime/Engine/Private/UnrealEngine.cpp b/Engine/Source/Runtime/Engine/Private/UnrealEngine.cpp
--- a/Engine/Source/Runtime/Engine/Private/UnrealEngine.cpp
+++ b/Engine/Source/Runtime/Engine/Private/UnrealEngine.cpp
@@ -7706,7 +7706,9 @@ static inline void CallHandleDisconnectForFailure(UWorld* InWorld, UNetDriver* N
 		ULocalPlayer* const LP = Context.OwningGameInstance->GetFirstGamePlayer();
 		if(ensure(LP))
 		{
-			LP->HandleDisconnect(InWorld, NetDriver);
+			// Use the world on the context instead since InWorld is null, it should be valid. This way we can do a travel
+			// in UEngine::HandleDisconnect if it gets called
+			LP->HandleDisconnect(Context.World(), NetDriver);
 		}
 	}
 	else
@@ -8000,12 +8002,6 @@ void UEngine::HandleDisconnect( UWorld *InWorld, UNetDriver *NetDriver )
 	// If there is a context for this world, setup client travel.
 	if (FWorldContext* WorldContext = GetWorldContextFromWorld(InWorld))
 	{
-		if (InWorld)
-		{
-			// If we have a world, then the failing NetDriver must be the world' net driver
-			check(InWorld->GetNetDriver() == NetDriver);
-		}
-
 		// Remove ?Listen parameter, if it exists
 		WorldContext->LastURL.RemoveOption( TEXT("Listen") );
 		WorldContext->LastURL.RemoveOption( TEXT("LAN") );
@@ -8025,6 +8021,23 @@ void UEngine::HandleDisconnect( UWorld *InWorld, UNetDriver *NetDriver )
 		{
 			NetDriver->Shutdown();
 			NetDriver->LowLevelDestroy();
+
+			// In this case, the world is null and something went wrong, so we should travel back to the default world so that we
+			// can get back to a good state.
+			for (FWorldContext& WorldContext : WorldList)
+			{
+				if (WorldContext.WorldType == EWorldType::Game)
+				{
+					FURL DefaultURL;
+					DefaultURL.LoadURLConfig(TEXT("DefaultPlayer"), GGameIni);
+					const UGameMapsSettings* GameMapsSettings = GetDefault<UGameMapsSettings>();
+					if (GameMapsSettings)
+					{
+						WorldContext.TravelURL = FURL(&DefaultURL, *(GameMapsSettings->GetGameDefaultMap() + GameMapsSettings->LocalMapOptions), TRAVEL_Partial).ToString();
+						WorldContext.TravelType = TRAVEL_Partial;
+					}
+				}
+			}
 		}
 	}
 }
