UE4 - pack instance buffers into smaller space, using half4 etc

diff --git a/Engine/Shaders/LocalVertexFactory.usf b/Engine/Shaders/LocalVertexFactory.usf
--- a/Engine/Shaders/LocalVertexFactory.usf
+++ b/Engine/Shaders/LocalVertexFactory.usf
@@ -33,9 +33,9 @@ float4 InstancingFadeOutParams;
 
 #if USE_INSTANCING_EMULATED
 	// Required for CPU emulation of Instancing
-	float4 CPUInstanceShadowMapBias;
-	float3x4 CPUInstanceTransform;
-	float4 CPUInstanceLightmapUVBias;
+	float4 CPUInstanceOrigin;  // per-instance random in w 
+	float3x4 CPUInstanceTransform; // hitproxy.r + 256 * selected in .w; hitproxy.g in .w; hitproxy.b in .w
+	float4 CPUInstanceLightmapAndShadowMapBias; 
 #endif	// USE_INSTANCING_EMULATED
 
 #else // USE_INSTANCING
@@ -111,11 +111,11 @@ struct FVertexFactoryInput
 #endif
 
 #if USE_INSTANCING && !USE_INSTANCING_EMULATED
-	float4 InstanceShadowMapBias : ATTRIBUTE8;
-	float4 InstanceTransform1 : ATTRIBUTE9;
-	float4 InstanceTransform2 : ATTRIBUTE10;
-	float4 InstanceTransform3 : ATTRIBUTE11;
-	float4 InstanceLightmapUVBias : ATTRIBUTE12; /* per-instance random in Z */
+	float4 InstanceOrigin : ATTRIBUTE8;  // per-instance random in w 
+	half4 InstanceTransform1 : ATTRIBUTE9;  // hitproxy.r + 256 * selected in .w
+	half4 InstanceTransform2 : ATTRIBUTE10; // hitproxy.g in .w
+	half4 InstanceTransform3 : ATTRIBUTE11; // hitproxy.b in .w
+	float4 InstanceLightmapAndShadowMapUVBias : ATTRIBUTE12; 
 #else
 
 #endif //USE_INSTANCING && !USE_INSTANCING_EMULATED
@@ -130,9 +130,10 @@ struct FPositionOnlyVertexFactoryInput
 	float4	Position	: ATTRIBUTE0;
 	
 #if USE_INSTANCING && !USE_INSTANCING_EMULATED
-	float4 InstanceTransform1 : ATTRIBUTE9;
-	float4 InstanceTransform2 : ATTRIBUTE10;
-	float4 InstanceTransform3 : ATTRIBUTE11;
+	float4 InstanceOrigin : ATTRIBUTE8;  // per-instance random in w 
+	half4 InstanceTransform1 : ATTRIBUTE9;  // hitproxy.r + 256 * selected in .w
+	half4 InstanceTransform2 : ATTRIBUTE10; // hitproxy.g in .w
+	half4 InstanceTransform3 : ATTRIBUTE11; // hitproxy.b in .w
 #endif	// USE_INSTANCING && !USE_INSTANCING_EMULATED
 };
 
@@ -153,29 +154,69 @@ struct FVertexFactoryIntermediates
 float4x4 GetInstanceTransform(FVertexFactoryInput Input)
 {
 	#if !USE_INSTANCING_EMULATED
-	return float4x4(Input.InstanceTransform1, Input.InstanceTransform2, Input.InstanceTransform3, float4(0, 0, 0, 1));
+	return float4x4(float4(Input.InstanceTransform1.xyz, Input.InstanceOrigin.x), float4(Input.InstanceTransform2.xyz, Input.InstanceOrigin.y), float4(Input.InstanceTransform3.xyz, Input.InstanceOrigin.z), float4(0, 0, 0, 1));
 	#else
-	return float4x4(CPUInstanceTransform[0], CPUInstanceTransform[1], CPUInstanceTransform[2], float4(0,0,0,1));
+	return float4x4(float4(CPUInstanceTransform[0].xyz, CPUInstanceOrigin.x), float4(CPUInstanceTransform[1].xyz, CPUInstanceOrigin.y), float4(CPUInstanceTransform[2].xyz, CPUInstanceOrigin.z), float4(0,0,0,1));
 	#endif	// !USE_INSTANCING_EMULATED
 }
 
 float4x4 GetInstanceTransform(FPositionOnlyVertexFactoryInput Input)
 {
 	#if !USE_INSTANCING_EMULATED
-	return float4x4(Input.InstanceTransform1, Input.InstanceTransform2, Input.InstanceTransform3, float4(0, 0, 0, 1));
+	return float4x4(float4(Input.InstanceTransform1.xyz, Input.InstanceOrigin.x), float4(Input.InstanceTransform2.xyz, Input.InstanceOrigin.y), float4(Input.InstanceTransform3.xyz, Input.InstanceOrigin.z), float4(0, 0, 0, 1));
 	#else
-	return float4x4(CPUInstanceTransform[0], CPUInstanceTransform[1], CPUInstanceTransform[2], float4(0,0,0,1));
+	return float4x4(float4(CPUInstanceTransform[0].xyz, CPUInstanceOrigin.x), float4(CPUInstanceTransform[1].xyz, CPUInstanceOrigin.y), float4(CPUInstanceTransform[2].xyz, CPUInstanceOrigin.z), float4(0,0,0,1));
 #endif	// !USE_INSTANCING_EMULATED
 }
 
-float4 GetInstanceShadowMapBias(FVertexFactoryInput Input)
+float2 GetInstanceShadowMapBias(FVertexFactoryInput Input)
 {
 	#if !USE_INSTANCING_EMULATED
-	return Input.InstanceShadowMapBias;
+	return Input.InstanceLightmapAndShadowMapUVBias.zw;
 	#else
-	return CPUInstanceShadowMapBias;
+	return CPUInstanceLightmapAndShadowMapBias.zw;
 	#endif	// !USE_INSTANCING_EMULATED
 }
+
+float2 GetInstanceLightMapBias(FVertexFactoryInput Input)
+{
+	#if !USE_INSTANCING_EMULATED
+	return Input.InstanceLightmapAndShadowMapUVBias.xy;
+	#else
+	return CPUInstanceLightmapAndShadowMapBias.xy;
+	#endif	// !USE_INSTANCING_EMULATED
+}
+
+float GetInstanceSelected(FVertexFactoryInput Input)
+{
+	#if !USE_INSTANCING_EMULATED
+		float SelectedValue = trunc(Input.InstanceTransform1.w * (1.0 / 256.0));
+	#else
+		float SelectedValue = trunc(CPUInstanceTransform[0].w * (1.0 / 256.0));
+	#endif	// !USE_INSTANCING_EMULATED
+	return SelectedValue;
+}
+
+float GetInstanceRandom(FVertexFactoryInput Input)
+{
+	#if !USE_INSTANCING_EMULATED
+		float RandomVal = Input.InstanceOrigin.w;
+	#else
+		float RandomVal = CPUInstanceOrigin.w;
+	#endif	// !USE_INSTANCING_EMULATED
+	return RandomVal;
+}
+
+float3 GetInstanceOrigin(FVertexFactoryInput Input)
+{
+	#if !USE_INSTANCING_EMULATED
+		float3 Origin = Input.InstanceOrigin.xyz;
+	#else
+		float3 Origin = CPUInstanceOrigin.xyz;
+	#endif	// !USE_INSTANCING_EMULATED
+	return Origin;
+}
+
 #endif	// USE_INSTANCING
 
 /** Converts from vertex factory specific interpolants to a FMaterialPixelParameters, which is used by material inputs. */
@@ -447,16 +488,11 @@ FVertexFactoryIntermediates GetVertexFactoryIntermediates(FVertexFactoryInput In
 	// x = per-instance random, y = per-instance fade out factor, z = zero or one depending of if it is shown at all, w is dither cutoff 
 
 	// PerInstanceParams.z stores a hide/show flag for this instance
-	float SelectedValue = trunc(GetInstanceShadowMapBias(Input).z * (1.0 / 256.0));
-	// InstancingFadeOutParams.z,w are RenderSelected and RenderDeselected respectively.
-	#if !USE_INSTANCING_EMULATED
-		Intermediates.PerInstanceParams.x = Input.InstanceLightmapUVBias.z;
-		float3 InstanceLocation = TransformLocalToWorld(float3(Input.InstanceTransform1.w, Input.InstanceTransform2.w, Input.InstanceTransform3.w)).xyz;
-	#else
-		Intermediates.PerInstanceParams.x = CPUInstanceLightmapUVBias.z;
-		float3 InstanceLocation = TransformLocalToWorld(float3(CPUInstanceTransform[0].w, CPUInstanceTransform[1].w, CPUInstanceTransform[2].w)).xyz;
-	#endif // !USE_INSTANCING_EMULATED
+	float SelectedValue = GetInstanceSelected(Input);
+	Intermediates.PerInstanceParams.x = GetInstanceRandom(Input);
+	float3 InstanceLocation = TransformLocalToWorld(GetInstanceOrigin(Input)).xyz;
 	Intermediates.PerInstanceParams.y = 1.0 - saturate((length(InstanceLocation + View.PreViewTranslation.xyz) - InstancingFadeOutParams.x) * InstancingFadeOutParams.y);
+	// InstancingFadeOutParams.z,w are RenderSelected and RenderDeselected respectively.
 	Intermediates.PerInstanceParams.z = InstancingFadeOutParams.z * SelectedValue + InstancingFadeOutParams.w * (1-SelectedValue);
 	#if USE_DITHERED_LOD_TRANSITION
 		float RandomLOD = InstancingViewZCompareZero.w * Intermediates.PerInstanceParams.x;
@@ -528,18 +564,13 @@ FVertexFactoryInterpolantsVSToPS VertexFactoryGetInterpolantsVSToPS(FVertexFacto
 	float2 ShadowMapCoordinate = 0;
 
 	#if USE_INSTANCING
-		#if !USE_INSTANCING_EMULATED
-		float2 InstanceLightmapUVBias = Input.InstanceLightmapUVBias.xy;
-		#else
-		float2 InstanceLightmapUVBias = CPUInstanceLightmapUVBias.xy;
-		#endif // !USE_INSTANCING_EMULATED
-		LightMapCoordinate = Input.LightMapCoordinate * LightMapCoordinateScaleBias.xy + InstanceLightmapUVBias;
+		LightMapCoordinate = Input.LightMapCoordinate * LightMapCoordinateScaleBias.xy + GetInstanceLightMapBias(Input);
 	#else
 		LightMapCoordinate = Input.LightMapCoordinate * LightMapCoordinateScaleBias.xy + LightMapCoordinateScaleBias.zw;
 	#endif
 	#if STATICLIGHTING_TEXTUREMASK
 		#if USE_INSTANCING
-			ShadowMapCoordinate = Input.LightMapCoordinate * ShadowMapCoordinateScaleBias.xy + GetInstanceShadowMapBias(Input).xy;
+			ShadowMapCoordinate = Input.LightMapCoordinate * ShadowMapCoordinateScaleBias.xy + GetInstanceShadowMapBias(Input);
 		#else
 			ShadowMapCoordinate = Input.LightMapCoordinate * ShadowMapCoordinateScaleBias.xy + ShadowMapCoordinateScaleBias.zw;
 		#endif
@@ -690,12 +721,15 @@ float4 VertexFactoryGetPreviousWorldPosition(FVertexFactoryInput Input, FVertexF
 #if USE_INSTANCING
 float4 VertexFactoryGetInstanceHitProxyId(FVertexFactoryInput Input, FVertexFactoryIntermediates Intermediates)
 {
-	// Unpack HitProxy color from InstanceShadowMapBias z and w
-	float4 InstanceShadowMapBias = GetInstanceShadowMapBias(Input);
-	float SelectedMask = trunc(InstanceShadowMapBias.z * (1.0 / 256.0));
-	float R = InstanceShadowMapBias.z - 256.0 * SelectedMask;
-	float G = trunc(InstanceShadowMapBias.w * (1.0/256.0));
-	float B = InstanceShadowMapBias.w - 256.0 * G;
+	#if !USE_INSTANCING_EMULATED
+		float R = Input.InstanceTransform1.w - 256.0 * GetInstanceSelected(Input);
+		float G = Input.InstanceTransform2.w;
+		float B = Input.InstanceTransform3.w;
+	#else
+		float R = CPUInstanceTransform[0].w - 256.0 * GetInstanceSelected(Input);
+		float G = CPUInstanceTransform[1].w;
+		float B = CPUInstanceTransform[2].w;
+	#endif
 	return float4(R/255.0, G/255.0, B/255.0, 0);
 }
 #endif	// USE_INSTANCING
diff --git a/Engine/Source/Runtime/Engine/Classes/Components/HierarchicalInstancedStaticMeshComponent.h b/Engine/Source/Runtime/Engine/Classes/Components/HierarchicalInstancedStaticMeshComponent.h
--- a/Engine/Source/Runtime/Engine/Classes/Components/HierarchicalInstancedStaticMeshComponent.h
+++ b/Engine/Source/Runtime/Engine/Classes/Components/HierarchicalInstancedStaticMeshComponent.h
@@ -57,6 +57,8 @@ class ENGINE_API UHierarchicalInstancedStaticMeshComponent : public UInstancedSt
 {
 	GENERATED_UCLASS_BODY()
 
+	~UHierarchicalInstancedStaticMeshComponent();
+
 	TSharedPtr<TArray<FClusterNode>, ESPMode::ThreadSafe> ClusterTreePtr;
 
 	FStaticMeshInstanceData WriteOncePrebuiltInstanceBuffer;
diff --git a/Engine/Source/Runtime/Engine/Private/HierarchicalInstancedStaticMesh.cpp b/Engine/Source/Runtime/Engine/Private/HierarchicalInstancedStaticMesh.cpp
--- a/Engine/Source/Runtime/Engine/Private/HierarchicalInstancedStaticMesh.cpp
+++ b/Engine/Source/Runtime/Engine/Private/HierarchicalInstancedStaticMesh.cpp
@@ -75,6 +75,7 @@ DECLARE_DWORD_COUNTER_STAT(TEXT("Mesh Batches"), STAT_FoliageMeshBatches, STATGR
 DECLARE_DWORD_COUNTER_STAT(TEXT("Triangles"), STAT_FoliageTriangles, STATGROUP_Foliage);
 DECLARE_DWORD_COUNTER_STAT(TEXT("Instances"), STAT_FoliageInstances, STATGROUP_Foliage);
 DECLARE_DWORD_COUNTER_STAT(TEXT("Traversals"),STAT_FoliageTraversals,STATGROUP_Foliage);
+DECLARE_MEMORY_STAT(TEXT("Instance Buffers"),STAT_FoliageInstanceBuffers,STATGROUP_Foliage);
 
 struct FClusterTree
 {
@@ -1340,6 +1341,17 @@ UHierarchicalInstancedStaticMeshComponent::UHierarchicalInstancedStaticMeshCompo
 {
 }
 
+UHierarchicalInstancedStaticMeshComponent::~UHierarchicalInstancedStaticMeshComponent()
+{
+	if (ProxySize)
+	{
+		DEC_DWORD_STAT_BY(STAT_FoliageInstanceBuffers, ProxySize);
+	}
+	ProxySize = 0;
+}
+
+
+
 void UHierarchicalInstancedStaticMeshComponent::Serialize(FArchive& Ar)
 {
 	Super::Serialize(Ar);
@@ -1483,6 +1495,11 @@ void UHierarchicalInstancedStaticMeshComponent::ClearInstances()
 	SortedInstances.Empty();
 	UnbuiltInstanceBounds.Init();
 
+	if (ProxySize)
+	{
+		DEC_DWORD_STAT_BY(STAT_FoliageInstanceBuffers, ProxySize);
+	}
+
 	Super::ClearInstances();
 }
 
@@ -1566,7 +1583,7 @@ void UHierarchicalInstancedStaticMeshComponent::AcceptPrebuiltTree(TArray<FClust
 {
 	// this is only for prebuild data, already in the correct order
 	check(!PerInstanceSMData.Num());
-	NumBuiltInstances = WriteOncePrebuiltInstanceBuffer.GetNumInstances();
+	NumBuiltInstances = WriteOncePrebuiltInstanceBuffer.Num();
 	check(NumBuiltInstances);
 	UnbuiltInstanceBounds.Init();
 	RemovedInstances.Empty();
@@ -1803,12 +1820,16 @@ FPrimitiveSceneProxy* UHierarchicalInstancedStaticMeshComponent::CreateSceneProx
 {
 	QUICK_SCOPE_CYCLE_COUNTER(STAT_HierarchicalInstancedStaticMeshComponent_CreateSceneProxy);
 
+	if (ProxySize)
+	{
+		DEC_DWORD_STAT_BY(STAT_FoliageInstanceBuffers, ProxySize);
+	}
 	ProxySize = 0;
 
 	// Verify that the mesh is valid before using it.
 	const bool bMeshIsValid = 
 		// make sure we have instances
-		(PerInstanceSMData.Num() > 0 || WriteOncePrebuiltInstanceBuffer.GetNumInstances() > 0)&&
+		(PerInstanceSMData.Num() > 0 || WriteOncePrebuiltInstanceBuffer.Num() > 0)&&
 		// make sure we have an actual staticmesh
 		StaticMesh &&
 		StaticMesh->HasValidRenderData() &&
@@ -1826,15 +1847,17 @@ FPrimitiveSceneProxy* UHierarchicalInstancedStaticMeshComponent::CreateSceneProx
 			InstancingRandomSeed = FMath::Rand();
 		}
 
-		if (WriteOncePrebuiltInstanceBuffer.GetNumInstances() > 0)
+		if (WriteOncePrebuiltInstanceBuffer.Num() > 0)
 		{
-			ProxySize = FStaticMeshInstanceData::GetResourceSize(WriteOncePrebuiltInstanceBuffer.GetNumInstances());
+			ProxySize = FStaticMeshInstanceData::GetResourceSize(WriteOncePrebuiltInstanceBuffer.Num());
+			INC_DWORD_STAT_BY(STAT_FoliageInstanceBuffers, ProxySize);
 
 			//@todo - this is sketchy, what you really want to do is flush the grass cache instead of doing a render thread update
 			bNeverNeedsRenderUpdate = true; // we can never update it because the data is gone
 			return ::new FHierarchicalStaticMeshSceneProxy(this, GetWorld()->FeatureLevel, WriteOncePrebuiltInstanceBuffer);
 		}
 		ProxySize = FStaticMeshInstanceData::GetResourceSize(PerInstanceSMData.Num());
+		INC_DWORD_STAT_BY(STAT_FoliageInstanceBuffers, ProxySize);
 		return ::new FHierarchicalStaticMeshSceneProxy(this, GetWorld()->FeatureLevel);
 	}
 	return NULL;
diff --git a/Engine/Source/Runtime/Engine/Private/InstancedStaticMesh.cpp b/Engine/Source/Runtime/Engine/Private/InstancedStaticMesh.cpp
--- a/Engine/Source/Runtime/Engine/Private/InstancedStaticMesh.cpp
+++ b/Engine/Source/Runtime/Engine/Private/InstancedStaticMesh.cpp
@@ -44,19 +44,6 @@ void FStaticMeshInstanceBuffer::CleanUp()
 }
 
 
-static void FillInstanceRenderData(const FInstancedStaticMeshInstanceData& Instance, FVector4* RenderData, float RandomInstanceID, float Z = 0, float W = 0)
-{
-	RenderData[0] = FVector4(Instance.ShadowmapUVBias.X, Instance.ShadowmapUVBias.Y, Z, W);
-
-	// Instance -> local matrix.  Every mesh instance has it's own transformation into
-	// the actor's coordinate space.
-	RenderData[1] = FVector4(Instance.Transform.M[0][0], Instance.Transform.M[1][0], Instance.Transform.M[2][0], Instance.Transform.M[3][0]);
-	RenderData[2] = FVector4(Instance.Transform.M[0][1], Instance.Transform.M[1][1], Instance.Transform.M[2][1], Instance.Transform.M[3][1]);
-	RenderData[3] = FVector4(Instance.Transform.M[0][2], Instance.Transform.M[1][2], Instance.Transform.M[2][2], Instance.Transform.M[3][2]);
-
-	RenderData[4] = FVector4(Instance.LightmapUVBias.X, Instance.LightmapUVBias.Y, RandomInstanceID, 0);
-}
-
 void FStaticMeshInstanceBuffer::SetupCPUAccess(UInstancedStaticMeshComponent* InComponent)
 {
 	static const auto CVar = IConsoleManager::Get().FindTConsoleVariableDataInt(TEXT("r.GenerateMeshDistanceFields"));
@@ -82,60 +69,37 @@ void FStaticMeshInstanceBuffer::Init(UInstancedStaticMeshComponent* InComponent,
 
 	InstanceData->AllocateInstances(NumInstances + NumRemoved);
 
-	// @todo: Make LD-customizable per component?
-	const float RandomInstanceIDBase = 0.0f;
-	const float RandomInstanceIDRange = 1.0f;
-
 	// Setup our random number generator such that random values are generated consistently for any
 	// given instance index between reattaches
 	FRandomStream RandomStream( InComponent->InstancingRandomSeed );
 
+	FColor HitProxyColor(ForceInit);
+	bool bSelected = false;
 	for (int32 InstanceIndex = 0; InstanceIndex < (int32)NumInstances; InstanceIndex++)
 	{
-		const FInstancedStaticMeshInstanceData& Instance = InComponent->PerInstanceSMData[InstanceIndex];
-
-		FVector4* InstanceRenderData = InstanceData->GetInstanceWriteAddress(bUseRemapTable ? InComponent->InstanceReorderTable[InstanceIndex] : InstanceIndex);
-
-		// X, Y	: Shadow map UV bias
-		// Z, W : Encoded HitProxy ID and selection mask.
-		float Z = 0.f;
-		float W = 0.f;
+#if WITH_EDITOR
 		if (InHitProxies.Num() == NumInstances)
 		{
-			FColor HitProxyColor = InHitProxies[InstanceIndex]->Id.GetColor();
-			Z = (float)HitProxyColor.R;
-			W = (float)HitProxyColor.G * 256.f + (float)HitProxyColor.B;
+			HitProxyColor = InHitProxies[InstanceIndex]->Id.GetColor();
 		}
-#if WITH_EDITOR
 		// Record if the instance is selected
-		if (InstanceIndex < InComponent->SelectedInstances.Num() && InComponent->SelectedInstances[InstanceIndex])
-		{
-			Z += 256.f;
-		}
+		bSelected = InstanceIndex < InComponent->SelectedInstances.Num() && InComponent->SelectedInstances[InstanceIndex];
 #endif
 
-		const float RandomInstanceID = RandomInstanceIDBase + RandomStream.GetFraction() * RandomInstanceIDRange;
+		FInstanceStream* InstanceRenderData = InstanceData->GetInstanceWriteAddress(bUseRemapTable ? InComponent->InstanceReorderTable[InstanceIndex] : InstanceIndex);
+		const FInstancedStaticMeshInstanceData& Instance = InComponent->PerInstanceSMData[InstanceIndex];
 						
-		FillInstanceRenderData(Instance, InstanceRenderData, RandomInstanceID, Z, W);
+		InstanceRenderData->SetInstance(Instance.Transform, RandomStream.GetFraction(), Instance.LightmapUVBias, Instance.ShadowmapUVBias, HitProxyColor, bSelected);
 	}
 
 	// Hide any removed instances
 	if (NumRemoved)
 	{
 		check(bUseRemapTable);
-
 		for (int32 InstanceIndex = 0; InstanceIndex < NumRemoved; InstanceIndex++)
 		{
-			FVector4* RenderData = InstanceData->GetData() + InComponent->RemovedInstances[InstanceIndex] * (GetStride() / sizeof(FVector4));
-
-			RenderData[0] = FVector4(0, 0, 0, 0);
-
-			// transform is 0 matrix
-			RenderData[1] = FVector4(0, 0, 0, 0);
-			RenderData[2] = FVector4(0, 0, 0, 0);
-			RenderData[3] = FVector4(0, 0, 0, 0);
-		
-			RenderData[4] = FVector4(0, 0, 0, 0);
+			FInstanceStream* InstanceRenderData = InstanceData->GetInstanceWriteAddress(InComponent->RemovedInstances[InstanceIndex]);
+			InstanceRenderData->NullifyInstance();
 		}
 	}
 }
@@ -146,6 +110,7 @@ void FStaticMeshInstanceBuffer::InitFromPreallocatedData(UInstancedStaticMeshCom
 	SetupCPUAccess(InComponent);
 }
 
+#if 0 // delete me
 
 /** Serializer. */
 FArchive& operator<<(FArchive& Ar,FStaticMeshInstanceBuffer& InstanceBuffer)
@@ -165,6 +130,8 @@ FArchive& operator<<(FArchive& Ar,FStaticMeshInstanceBuffer& InstanceBuffer)
 	return Ar;
 }
 
+#endif
+
 /**
  * Specialized assignment operator, only used when importing LOD's.  
  */
@@ -251,9 +218,10 @@ void FInstancedStaticMeshVertexFactory::InitRHI()
 		if (bInstanced)
 		{
 			// toss in the instanced location stream
-			PositionOnlyStreamElements.Add(AccessPositionStreamComponent(Data.InstancedTransformComponent[0],9));
-			PositionOnlyStreamElements.Add(AccessPositionStreamComponent(Data.InstancedTransformComponent[1],10));
-			PositionOnlyStreamElements.Add(AccessPositionStreamComponent(Data.InstancedTransformComponent[2],11));
+			PositionOnlyStreamElements.Add(AccessPositionStreamComponent(Data.InstanceOriginComponent,8));
+			PositionOnlyStreamElements.Add(AccessPositionStreamComponent(Data.InstanceTransformComponent[0],9));
+			PositionOnlyStreamElements.Add(AccessPositionStreamComponent(Data.InstanceTransformComponent[1],10));
+			PositionOnlyStreamElements.Add(AccessPositionStreamComponent(Data.InstanceTransformComponent[2],11));
 		}
 		InitPositionDeclaration(PositionOnlyStreamElements);
 	}
@@ -319,11 +287,11 @@ void FInstancedStaticMeshVertexFactory::InitRHI()
 	// toss in the instanced location stream
 	if (bInstanced)
 	{
-		Elements.Add(AccessStreamComponent(Data.InstancedShadowMapBiasComponent,8));
-		Elements.Add(AccessStreamComponent(Data.InstancedTransformComponent[0],9));
-		Elements.Add(AccessStreamComponent(Data.InstancedTransformComponent[1],10));
-		Elements.Add(AccessStreamComponent(Data.InstancedTransformComponent[2],11));
-		Elements.Add(AccessStreamComponent(Data.InstancedLightmapUVBiasComponent, 12));
+		Elements.Add(AccessStreamComponent(Data.InstanceOriginComponent,8));
+		Elements.Add(AccessStreamComponent(Data.InstanceTransformComponent[0],9));
+		Elements.Add(AccessStreamComponent(Data.InstanceTransformComponent[1],10));
+		Elements.Add(AccessStreamComponent(Data.InstanceTransformComponent[2],11));
+		Elements.Add(AccessStreamComponent(Data.InstanceLightmapAndShadowMapUVBiasComponent,12));
 	}
 
 	// we don't need per-vertex shadow or lightmap rendering
@@ -458,37 +426,42 @@ void FInstancedStaticMeshRenderData::InitStaticMeshVertexFactories(
 		{
 			const FStaticMeshInstanceBuffer& InstanceBuffer = InstancedRenderData->PerInstanceRenderData->InstanceBuffer;
 			
-			// Shadow map bias (and random instance ID)
-			int32 CurInstanceBufferOffset = 0;
-			Data.InstancedShadowMapBiasComponent = FVertexStreamComponent(
+			Data.InstanceOriginComponent = FVertexStreamComponent(
 				&InstanceBuffer,
-				CurInstanceBufferOffset, 
+				STRUCT_OFFSET(FInstanceStream,InstanceOrigin), 
 				InstanceBuffer.GetStride(),
 				VET_Float4,
 				true
 				);
-			CurInstanceBufferOffset += sizeof(float) * 4;
-
-			for (int32 MatrixRow = 0; MatrixRow < 3; MatrixRow++)
-			{
-				Data.InstancedTransformComponent[MatrixRow] = FVertexStreamComponent(
-					&InstanceBuffer,
-					CurInstanceBufferOffset, 
-					InstanceBuffer.GetStride(),
-					VET_Float4,
-					true
-					);
-				CurInstanceBufferOffset += sizeof(float) * 4;
-			}
+			Data.InstanceTransformComponent[0] = FVertexStreamComponent(
+				&InstanceBuffer,
+				STRUCT_OFFSET(FInstanceStream,InstanceTransform1), 
+				InstanceBuffer.GetStride(),
+				VET_Half4,
+				true
+				);
+			Data.InstanceTransformComponent[1] = FVertexStreamComponent(
+				&InstanceBuffer,
+				STRUCT_OFFSET(FInstanceStream,InstanceTransform2), 
+				InstanceBuffer.GetStride(),
+				VET_Half4,
+				true
+				);
+			Data.InstanceTransformComponent[2] = FVertexStreamComponent(
+				&InstanceBuffer,
+				STRUCT_OFFSET(FInstanceStream,InstanceTransform3), 
+				InstanceBuffer.GetStride(),
+				VET_Half4,
+				true
+				);
 
-			Data.InstancedLightmapUVBiasComponent = FVertexStreamComponent(
+			Data.InstanceLightmapAndShadowMapUVBiasComponent = FVertexStreamComponent(
 				&InstanceBuffer,
-				CurInstanceBufferOffset, 
+				STRUCT_OFFSET(FInstanceStream,InstanceLightmapAndShadowMapUVBias), 
 				InstanceBuffer.GetStride(),
-				VET_Float4,
+				VET_Short4N,
 				true
 				);
-			CurInstanceBufferOffset += sizeof(float) * 4;
 		}
 
 		// Assign to the vertex factory for this LOD.
@@ -658,16 +631,28 @@ void FInstancedStaticMeshSceneProxy::GetDistancefieldAtlasData(FBox& LocalVolume
 
 	ObjectLocalToWorldTransforms.Reset();
 
-	const FInstancingUserData::FInstanceStream* InstanceStream = ((const FInstancingUserData::FInstanceStream*)InstancedRenderData.PerInstanceRenderData->InstanceBuffer.GetRawData());
+	const FInstanceStream* InstanceStream = InstancedRenderData.PerInstanceRenderData->InstanceBuffer.GetData();
 	FVector4 FourthVector(0, 0, 0, 1);
 
 	const uint32 NumInstances = InstancedRenderData.PerInstanceRenderData->InstanceBuffer.GetNumInstances();
 	for (uint32 InstanceIndex = 0; InstanceIndex < NumInstances; InstanceIndex++)
 	{
 		const FMatrix TransposedInstanceToLocal(
-			(const FPlane&)InstanceStream[InstanceIndex].InstanceTransform[0], 
-			(const FPlane&)InstanceStream[InstanceIndex].InstanceTransform[1], 
-			(const FPlane&)InstanceStream[InstanceIndex].InstanceTransform[2], 
+			FPlane(
+				InstanceStream[InstanceIndex].InstanceTransform1[0].GetFloat(),
+				InstanceStream[InstanceIndex].InstanceTransform1[1].GetFloat(),
+				InstanceStream[InstanceIndex].InstanceTransform1[2].GetFloat(),
+				InstanceStream[InstanceIndex].InstanceOrigin.X),
+			FPlane(
+				InstanceStream[InstanceIndex].InstanceTransform2[0].GetFloat(),
+				InstanceStream[InstanceIndex].InstanceTransform2[1].GetFloat(),
+				InstanceStream[InstanceIndex].InstanceTransform2[2].GetFloat(),
+				InstanceStream[InstanceIndex].InstanceOrigin.Y),
+			FPlane(
+				InstanceStream[InstanceIndex].InstanceTransform3[0].GetFloat(),
+				InstanceStream[InstanceIndex].InstanceTransform3[1].GetFloat(),
+				InstanceStream[InstanceIndex].InstanceTransform3[2].GetFloat(),
+				InstanceStream[InstanceIndex].InstanceOrigin.Z),
 			(const FPlane&)FourthVector);
 
 		ObjectLocalToWorldTransforms.Add(TransposedInstanceToLocal.GetTransposed() * GetLocalToWorld());
@@ -821,7 +806,7 @@ void UInstancedStaticMeshComponent::ApplyComponentInstanceData(FComponentInstanc
 	}
 
 	SelectedInstances = InstancedMeshData->SelectedInstances;
-	
+
 	ReleasePerInstanceRenderData();
 	MarkRenderStateDirty();
 #endif
@@ -1501,7 +1486,7 @@ void UInstancedStaticMeshComponent::PostEditChangeChainProperty(FPropertyChanged
 			// added via the property editor, so we will want to interactively work with instances
 			bHasPerInstanceHitProxies = true;
 		}
-		
+
 		ReleasePerInstanceRenderData();
 		MarkRenderStateDirty();
 	}
@@ -1681,13 +1666,41 @@ void FInstancedStaticMeshVertexFactoryShaderParameters::SetMesh( FRHICommandList
 	const bool bInstanced = RHISupportsInstancing(ShaderPlatform);
 	if (!bInstanced)
 	{
-		if (CPUInstanceShadowMapBias.IsBound())
+		if (CPUInstanceOrigin.IsBound())
 		{
+			const float ShortScale = 1.0f / 32767.5f;
 			auto* InstancingData = (const FInstancingUserData*)BatchElement.UserData;
-			auto* InstanceStream = ((const FInstancingUserData::FInstanceStream*)InstancingData->RenderData->PerInstanceRenderData->InstanceBuffer.GetRawData()) + BatchElement.UserIndex;
-			SetShaderValue(RHICmdList, VS, CPUInstanceShadowMapBias, InstanceStream->InstanceShadowmapUVBias);
-			SetShaderValueArray(RHICmdList, VS, CPUInstanceTransform, InstanceStream->InstanceTransform, 3);
-			SetShaderValue(RHICmdList, VS, CPUInstanceLightmapUVBias, InstanceStream->InstanceLightmapUVBias);
+			const FInstanceStream* InstanceStream = InstancingData->RenderData->PerInstanceRenderData->InstanceBuffer.GetData() + BatchElement.UserIndex;
+
+			FVector4 InstanceLightmapAndShadowMapUVBias(
+				float(InstanceStream->InstanceLightmapAndShadowMapUVBias[0]) * ShortScale,
+				float(InstanceStream->InstanceLightmapAndShadowMapUVBias[1]) * ShortScale,
+				float(InstanceStream->InstanceLightmapAndShadowMapUVBias[2]) * ShortScale,
+				float(InstanceStream->InstanceLightmapAndShadowMapUVBias[3]) * ShortScale
+				);
+
+			FVector4 InstanceTransform[3] = 
+			{
+				FVector4(
+					InstanceStream->InstanceTransform1[0].GetFloat(),
+					InstanceStream->InstanceTransform1[1].GetFloat(),
+					InstanceStream->InstanceTransform1[2].GetFloat(),
+					InstanceStream->InstanceTransform1[3].GetFloat()),
+				FVector4(
+					InstanceStream->InstanceTransform2[0].GetFloat(),
+					InstanceStream->InstanceTransform2[1].GetFloat(),
+					InstanceStream->InstanceTransform2[2].GetFloat(),
+					InstanceStream->InstanceTransform2[3].GetFloat()),
+				FVector4(
+					InstanceStream->InstanceTransform3[0].GetFloat(),
+					InstanceStream->InstanceTransform3[1].GetFloat(),
+					InstanceStream->InstanceTransform3[2].GetFloat(),
+					InstanceStream->InstanceTransform3[3].GetFloat())
+			};
+
+			SetShaderValue(RHICmdList, VS, CPUInstanceOrigin, InstanceStream->InstanceOrigin);
+			SetShaderValueArray(RHICmdList, VS, CPUInstanceTransform, InstanceTransform, 3);
+			SetShaderValue(RHICmdList, VS, CPUInstanceLightmapAndShadowMapBias, InstanceLightmapAndShadowMapUVBias); 
 		}
 	}
 }
diff --git a/Engine/Source/Runtime/Engine/Private/InstancedStaticMesh.h b/Engine/Source/Runtime/Engine/Private/InstancedStaticMesh.h
--- a/Engine/Source/Runtime/Engine/Private/InstancedStaticMesh.h
+++ b/Engine/Source/Runtime/Engine/Private/InstancedStaticMesh.h
@@ -108,7 +108,7 @@ public:
 	}
 	FORCEINLINE uint32 GetNumInstances() const
 	{
-		return InstanceData ? InstanceData->GetNumInstances() : 0;
+		return InstanceData ? InstanceData->Num() : 0;
 	}
 
 	const void* GetRawData() const
@@ -116,6 +116,11 @@ public:
 		return InstanceData->GetDataPointer();
 	}
 
+	const FInstanceStream* GetData() const
+	{
+		return InstanceData->GetData();
+	}
+
 	// FRenderResource interface.
 	virtual void InitRHI() override;
 	virtual FString GetFriendlyName() const { return TEXT("Static-mesh instances"); }
@@ -143,13 +148,6 @@ private:
 
 struct FInstancingUserData
 {
-	struct FInstanceStream
-	{
-		FVector4 InstanceShadowmapUVBias;
-		FVector4 InstanceTransform[3];
-		FVector4 InstanceLightmapUVBias;
-	};
-
 	class FInstancedStaticMeshRenderData* RenderData;
 	class FStaticMeshRenderData* MeshRenderData;
 
@@ -169,14 +167,14 @@ struct FInstancedStaticMeshVertexFactory : public FLocalVertexFactory
 public:
 	struct DataType : public FLocalVertexFactory::DataType
 	{
-		/** The stream to read shadow map bias (and random instance ID) from. */
-		FVertexStreamComponent InstancedShadowMapBiasComponent;
+		/** The stream to read the mesh transform from. */
+		FVertexStreamComponent InstanceOriginComponent;
 
 		/** The stream to read the mesh transform from. */
-		FVertexStreamComponent InstancedTransformComponent[3];
+		FVertexStreamComponent InstanceTransformComponent[3];
 
 		/** The stream to read the Lightmap Bias and Random instance ID from. */
-		FVertexStreamComponent InstancedLightmapUVBiasComponent;
+		FVertexStreamComponent InstanceLightmapAndShadowMapUVBiasComponent;
 	};
 
 	/**
@@ -253,9 +251,9 @@ class FInstancedStaticMeshVertexFactoryShaderParameters : public FLocalVertexFac
 		InstancingViewZConstantParameter.Bind(ParameterMap, TEXT("InstancingViewZConstant"));
 		InstancingWorldViewOriginZeroParameter.Bind(ParameterMap, TEXT("InstancingWorldViewOriginZero"));
 		InstancingWorldViewOriginOneParameter.Bind(ParameterMap, TEXT("InstancingWorldViewOriginOne"));
-		CPUInstanceShadowMapBias.Bind(ParameterMap, TEXT("CPUInstanceShadowMapBias"));
+		CPUInstanceOrigin.Bind(ParameterMap, TEXT("CPUInstanceOrigin"));
 		CPUInstanceTransform.Bind(ParameterMap, TEXT("CPUInstanceTransform"));
-		CPUInstanceLightmapUVBias.Bind(ParameterMap, TEXT("CPUInstanceLightmapUVBias"));
+		CPUInstanceLightmapAndShadowMapBias.Bind(ParameterMap, TEXT("CPUInstanceLightmapAndShadowMapBias"));
 	}
 
 	virtual void SetMesh(FRHICommandList& RHICmdList, FShader* VertexShader,const class FVertexFactory* VertexFactory,const class FSceneView& View,const struct FMeshBatchElement& BatchElement,uint32 DataFlags) const override;
@@ -269,9 +267,9 @@ class FInstancedStaticMeshVertexFactoryShaderParameters : public FLocalVertexFac
 		Ar << InstancingViewZConstantParameter;
 		Ar << InstancingWorldViewOriginZeroParameter;
 		Ar << InstancingWorldViewOriginOneParameter;
-		Ar << CPUInstanceShadowMapBias;
+		Ar << CPUInstanceOrigin;
 		Ar << CPUInstanceTransform;
-		Ar << CPUInstanceLightmapUVBias;
+		Ar << CPUInstanceLightmapAndShadowMapBias;
 	}
 
 	virtual uint32 GetSize() const { return sizeof(*this); }
@@ -284,9 +282,9 @@ private:
 	FShaderParameter InstancingWorldViewOriginZeroParameter;
 	FShaderParameter InstancingWorldViewOriginOneParameter;
 
-	FShaderParameter CPUInstanceShadowMapBias;
+	FShaderParameter CPUInstanceOrigin;
 	FShaderParameter CPUInstanceTransform;
-	FShaderParameter CPUInstanceLightmapUVBias;
+	FShaderParameter CPUInstanceLightmapAndShadowMapBias;
 };
 
 /*-----------------------------------------------------------------------------
@@ -350,7 +348,7 @@ public:
 	{
 		// Allocate the vertex factories for each LOD
 		InitVertexFactories();
-		
+
 		if (!PerInstanceRenderData.IsValid())
 		{
 			// initialize the instance buffer from the component's instances
@@ -438,7 +436,7 @@ public:
 
 	/** LOD render data from the static mesh. */
 	TIndirectArray<FStaticMeshLODResources>& LODModels;
-	
+
 	/** Feature level used when creating instance data */
 	ERHIFeatureLevel::Type FeatureLevel;
 
@@ -599,8 +597,6 @@ private:
 			}
 		}
 
-		check(InstancedRenderData.PerInstanceRenderData->InstanceBuffer.GetStride() == sizeof(FInstancingUserData::FInstanceStream));
-
 		const bool bInstanced = RHISupportsInstancing(GetFeatureLevelShaderPlatform(InstancedRenderData.FeatureLevel));
 
 		// Copy the parameters for LOD - all instances
diff --git a/Engine/Source/Runtime/Engine/Public/StaticMeshResources.h b/Engine/Source/Runtime/Engine/Public/StaticMeshResources.h
--- a/Engine/Source/Runtime/Engine/Public/StaticMeshResources.h
+++ b/Engine/Source/Runtime/Engine/Public/StaticMeshResources.h
@@ -896,31 +896,146 @@ protected:
 	FStaticMeshInstanceData
 -----------------------------------------------------------------------------*/
 
+struct FInstanceStream
+{
+	FVector4 InstanceOrigin;  // per-instance random in w 
+	FFloat16 InstanceTransform1[4];  // hitproxy.r + 256 * selected in .w
+	FFloat16 InstanceTransform2[4]; // hitproxy.g in .w
+	FFloat16 InstanceTransform3[4]; // hitproxy.b in .w
+	int16 InstanceLightmapAndShadowMapUVBias[4]; 
+
+	FORCEINLINE void SetInstance(const FMatrix& Transform, float RandomInstanceID)
+	{
+		const FMatrix* RESTRICT rTransform = (const FMatrix* RESTRICT)&Transform;
+		FInstanceStream* RESTRICT Me = (FInstanceStream* RESTRICT)this;
+
+		Me->InstanceOrigin.X = rTransform->M[3][0];
+		Me->InstanceOrigin.Y = rTransform->M[3][1];
+		Me->InstanceOrigin.Z = rTransform->M[3][2];
+		Me->InstanceOrigin.W = RandomInstanceID;
+
+		Me->InstanceTransform1[0] = Transform.M[0][0];
+		Me->InstanceTransform1[1] = Transform.M[1][0];
+		Me->InstanceTransform1[2] = Transform.M[2][0];
+		Me->InstanceTransform1[3] = FFloat16();
+
+		Me->InstanceTransform2[0] = Transform.M[0][1];
+		Me->InstanceTransform2[1] = Transform.M[1][1];
+		Me->InstanceTransform2[2] = Transform.M[2][1];
+		Me->InstanceTransform2[3] = FFloat16();
+
+		Me->InstanceTransform3[0] = Transform.M[0][2];
+		Me->InstanceTransform3[1] = Transform.M[1][2];
+		Me->InstanceTransform3[2] = Transform.M[2][2];
+		Me->InstanceTransform3[3] = FFloat16();
+
+		Me->InstanceLightmapAndShadowMapUVBias[0] = 0;
+		Me->InstanceLightmapAndShadowMapUVBias[1] = 0;
+		Me->InstanceLightmapAndShadowMapUVBias[2] = 0;
+		Me->InstanceLightmapAndShadowMapUVBias[3] = 0;
+	}
+
+	FORCEINLINE void SetInstance(const FMatrix& Transform, float RandomInstanceID, const FVector2D& LightmapUVBias, const FVector2D& ShadowmapUVBias, FColor HitProxyColor, bool bSelected)
+	{
+		const FMatrix* RESTRICT rTransform = (const FMatrix* RESTRICT)&Transform;
+		FInstanceStream* RESTRICT Me = (FInstanceStream* RESTRICT)this;
+		const FVector2D* RESTRICT rLightmapUVBias = (const FVector2D* RESTRICT)&LightmapUVBias;
+		const FVector2D* RESTRICT rShadowmapUVBias = (const FVector2D* RESTRICT)&ShadowmapUVBias;
+
+		Me->InstanceOrigin.X = rTransform->M[3][0];
+		Me->InstanceOrigin.Y = rTransform->M[3][1];
+		Me->InstanceOrigin.Z = rTransform->M[3][2];
+		Me->InstanceOrigin.W = RandomInstanceID;
+
+		Me->InstanceTransform1[0] = Transform.M[0][0];
+		Me->InstanceTransform1[1] = Transform.M[1][0];
+		Me->InstanceTransform1[2] = Transform.M[2][0];
+		Me->InstanceTransform1[3] = ((float)HitProxyColor.R) + (bSelected ? 256.f : 0.0f);
+
+		Me->InstanceTransform2[0] = Transform.M[0][1];
+		Me->InstanceTransform2[1] = Transform.M[1][1];
+		Me->InstanceTransform2[2] = Transform.M[2][1];
+		Me->InstanceTransform2[3] = (float)HitProxyColor.G;
+
+		Me->InstanceTransform3[0] = Transform.M[0][2];
+		Me->InstanceTransform3[1] = Transform.M[1][2];
+		Me->InstanceTransform3[2] = Transform.M[2][2];
+		Me->InstanceTransform3[3] = (float)HitProxyColor.B;
+
+		Me->InstanceLightmapAndShadowMapUVBias[0] = FMath::Clamp<int32>(FMath::TruncToInt(rLightmapUVBias->X * 32767.5f),MIN_int16,MAX_int16);
+		Me->InstanceLightmapAndShadowMapUVBias[1] = FMath::Clamp<int32>(FMath::TruncToInt(rLightmapUVBias->Y * 32767.5f),MIN_int16,MAX_int16);
+		Me->InstanceLightmapAndShadowMapUVBias[2] = FMath::Clamp<int32>(FMath::TruncToInt(rShadowmapUVBias->X * 32767.5f),MIN_int16,MAX_int16);
+		Me->InstanceLightmapAndShadowMapUVBias[3] = FMath::Clamp<int32>(FMath::TruncToInt(rShadowmapUVBias->Y * 32767.5f),MIN_int16,MAX_int16);
+	}
+
+	FORCEINLINE void NullifyInstance()
+	{
+		FInstanceStream* RESTRICT Me = (FInstanceStream* RESTRICT)this;
+		Me->InstanceTransform1[0] = FFloat16();
+		Me->InstanceTransform1[1] = FFloat16();
+		Me->InstanceTransform1[2] = FFloat16();
+
+		Me->InstanceTransform2[0] = FFloat16();
+		Me->InstanceTransform2[1] = FFloat16();
+		Me->InstanceTransform2[2] = FFloat16();
+
+		Me->InstanceTransform3[0] = FFloat16();
+		Me->InstanceTransform3[1] = FFloat16();
+		Me->InstanceTransform3[2] = FFloat16();
+	}
+
+	friend FArchive& operator<<( FArchive& Ar, FInstanceStream& V )
+	{
+		return Ar 
+			<< V.InstanceOrigin.X 
+			<< V.InstanceOrigin.Y
+			<< V.InstanceOrigin.Z
+			<< V.InstanceOrigin.W
+
+			<< V.InstanceTransform1[0]
+			<< V.InstanceTransform1[1]
+			<< V.InstanceTransform1[2]
+			<< V.InstanceTransform1[3]
+
+			<< V.InstanceTransform2[0]
+			<< V.InstanceTransform2[1]
+			<< V.InstanceTransform2[2]
+			<< V.InstanceTransform2[3]
+
+			<< V.InstanceTransform3[0]
+			<< V.InstanceTransform3[1]
+			<< V.InstanceTransform3[2]
+			<< V.InstanceTransform3[3]
+
+			<< V.InstanceLightmapAndShadowMapUVBias[0]
+			<< V.InstanceLightmapAndShadowMapUVBias[1]
+			<< V.InstanceLightmapAndShadowMapUVBias[2]
+			<< V.InstanceLightmapAndShadowMapUVBias[3];
+	}
+};
+
+
 /** The implementation of the static mesh instance data storage type. */
 class FStaticMeshInstanceData :
 	public FStaticMeshVertexDataInterface,
-	public TResourceArray<FVector4,VERTEXBUFFER_ALIGNMENT>
+	public TResourceArray<FInstanceStream, VERTEXBUFFER_ALIGNMENT>
 {
 public:
 
-	enum 
-	{
-		VectorsPerInstance = 5
-	};
-	typedef TResourceArray<FVector4,VERTEXBUFFER_ALIGNMENT> ArrayType;
+	typedef TResourceArray<FInstanceStream, VERTEXBUFFER_ALIGNMENT> ArrayType;
 
 	/**
 	 * Constructor
 	 * @param InNeedsCPUAccess - true if resource array data should be CPU accessible
 	 */
 	FStaticMeshInstanceData(bool InNeedsCPUAccess=false)
-		:	TResourceArray<FVector4,VERTEXBUFFER_ALIGNMENT>(InNeedsCPUAccess)
+		:	TResourceArray<FInstanceStream, VERTEXBUFFER_ALIGNMENT>(InNeedsCPUAccess)
 	{
 	}
 
 	static uint32 StaticGetStride()
 	{
-		return sizeof(FVector4) * VectorsPerInstance;
+		return sizeof(FInstanceStream);
 	}
 
 	static SIZE_T GetResourceSize(uint32 NumInstances)
@@ -928,56 +1043,50 @@ public:
 		return SIZE_T(NumInstances) * SIZE_T(StaticGetStride());
 	}
 
-	int32 GetNumInstances()
-	{
-		check(Num() % VectorsPerInstance == 0);
-		return Num() / VectorsPerInstance;
-	}
-
 	/**
 	 * Resizes the vertex data buffer, discarding any data which no longer fits.
 	 * @param NumVertices - The number of vertices to allocate the buffer for.
 	 */
-	virtual void ResizeBuffer(uint32 NumInstances)
+	virtual void ResizeBuffer(uint32 NumInstances) override
 	{
 		checkf(0, TEXT("ArrayType::Add is not supported on all platforms"));
 	}
 
-	virtual uint32 GetStride() const
+	virtual uint32 GetStride() const override
 	{
 		return StaticGetStride();
 	}
-	virtual uint8* GetDataPointer()
+	virtual uint8* GetDataPointer() override
 	{
 		return (uint8*)&(*this)[0];
 	}
-	virtual FResourceArrayInterface* GetResourceArray()
+	virtual FResourceArrayInterface* GetResourceArray() override
 	{
 		return this;
 	}
-	virtual void Serialize(FArchive& Ar)
+	virtual void Serialize(FArchive& Ar) override
 	{
-		TResourceArray<FVector4,VERTEXBUFFER_ALIGNMENT>::BulkSerialize(Ar);
+		TResourceArray<FInstanceStream,VERTEXBUFFER_ALIGNMENT>::BulkSerialize(Ar);
 	}
 
-	void Set(const TArray<FVector4>& RawData)
+#if 0
+	void Set(const TArray<FInstanceStream>& RawData)
 	{
-		*((ArrayType*)this) = TArray<FVector4,TAlignedHeapAllocator<VERTEXBUFFER_ALIGNMENT> >(RawData);
+		*((ArrayType*)this) = TArray<FInstanceStream, TAlignedHeapAllocator<VERTEXBUFFER_ALIGNMENT> >(RawData);
 	}
+#endif
 
 	void AllocateInstances(int32 NumInstances)
 	{
 		// We cannot write directly to the data on all platforms,
 		// so we make a TArray of the right type, then assign it
-		check( GetStride() % sizeof(FVector4) == 0 );
-		Empty((NumInstances * GetStride()) / sizeof(FVector4));
-		AddUninitialized((NumInstances * GetStride()) / sizeof(FVector4));
+		Empty(NumInstances);
+		AddUninitialized(NumInstances);
 	}
-	FORCEINLINE FVector4* GetInstanceWriteAddress(int32 InstanceIndex)
+	FORCEINLINE FInstanceStream* GetInstanceWriteAddress(int32 InstanceIndex)
 	{
-		return GetData() + InstanceIndex * VectorsPerInstance;
+		return GetData() + InstanceIndex;
 	}
-
 };
 
 
diff --git a/Engine/Source/Runtime/IOS/MetalRHI/Private/MetalVertexDeclaration.cpp b/Engine/Source/Runtime/IOS/MetalRHI/Private/MetalVertexDeclaration.cpp
--- a/Engine/Source/Runtime/IOS/MetalRHI/Private/MetalVertexDeclaration.cpp
+++ b/Engine/Source/Runtime/IOS/MetalRHI/Private/MetalVertexDeclaration.cpp
@@ -23,6 +23,7 @@ static MTLVertexFormat TranslateElementTypeToMTLType(EVertexElementType Type)
 		case VET_Short2N:		return MTLVertexFormatShort2Normalized;
 		case VET_Half2:			return MTLVertexFormatHalf2;
 		case VET_Half4:			return MTLVertexFormatHalf4;
+		case VET_Short4N:		return MTLVertexFormatShort4Normalized;
 		default:				UE_LOG(LogMetal, Fatal, TEXT("Unknown vertex element type!")); return MTLVertexFormatFloat;
 	};
 
@@ -45,6 +46,7 @@ static uint32 TranslateElementTypeToSize(EVertexElementType Type)
 		case VET_Short2N:		return 4;
 		case VET_Half2:			return 4;
 		case VET_Half4:			return 8;
+		case VET_Short4N:		return 8;
 		default:				UE_LOG(LogMetal, Fatal, TEXT("Unknown vertex element type!")); return 0;
 	};
 }
diff --git a/Engine/Source/Runtime/Landscape/Private/Landscape.cpp b/Engine/Source/Runtime/Landscape/Private/Landscape.cpp
--- a/Engine/Source/Runtime/Landscape/Private/Landscape.cpp
+++ b/Engine/Source/Runtime/Landscape/Private/Landscape.cpp
@@ -2514,10 +2514,6 @@ struct FAsyncGrassBuilder : public FGrassBuilderBase
 			TotalInstances += NumKept;
 
 			{
-				// @todo: Make LD-customizable per component?
-				const float RandomInstanceIDBase = 0.0f;
-				const float RandomInstanceIDRange = 1.0f;
-
 				InstanceTime -= FPlatformTime::Seconds();
 				InstanceBuffer.AllocateInstances(NumKept);
 				int32 InstanceIndex = 0;
@@ -2559,18 +2555,8 @@ struct FAsyncGrassBuilder : public FGrassBuilderBase
 								OutXForm = FRotationMatrix(FRotator(0.0f, Rot, 0.0f)) * FTranslationMatrix(Instance.Pos) * XForm;
 							}
 							InstanceTransforms[OutInstanceIndex] = OutXForm;
-							FVector4* RenderData = InstanceBuffer.GetInstanceWriteAddress(OutInstanceIndex++);
-
-							const float RandomInstanceID = RandomInstanceIDBase + RandomStream.GetFraction() * RandomInstanceIDRange;
-							RenderData[0] = FVector4(-1.0f, -1.0f, 0.0f, 0.0f);
-
-							// Instance -> local matrix.  Every mesh instance has it's own transformation into
-							// the actor's coordinate space.
-							RenderData[1] = FVector4(OutXForm.M[0][0], OutXForm.M[1][0], OutXForm.M[2][0], OutXForm.M[3][0]);
-							RenderData[2] = FVector4(OutXForm.M[0][1], OutXForm.M[1][1], OutXForm.M[2][1], OutXForm.M[3][1]);
-							RenderData[3] = FVector4(OutXForm.M[0][2], OutXForm.M[1][2], OutXForm.M[2][2], OutXForm.M[3][2]);
-
-							RenderData[4] = FVector4(-1.0f, -1.0f, RandomInstanceID, 0.0f);
+							FInstanceStream* RenderData = InstanceBuffer.GetInstanceWriteAddress(OutInstanceIndex++);
+							RenderData->SetInstance(OutXForm, RandomStream.GetFraction());
 						}
 						InstanceIndex++;
 					}
@@ -2584,18 +2570,17 @@ struct FAsyncGrassBuilder : public FGrassBuilderBase
 				UHierarchicalInstancedStaticMeshComponent::BuildTreeAnyThread(InstanceTransforms, MeshBox, ClusterTree, SortedInstances, InstanceReorderTable, FMath::Clamp<int32>(CVarMinInstancesPerLeaf.GetValueOnAnyThread(), 16, 256 * 128));
 
 				// in-place sort the instances
-				FVector4 SwapBuffer[FStaticMeshInstanceData::VectorsPerInstance];
+				FInstanceStream SwapBuffer;
 				int32 FirstUnfixedIndex = 0;
-				int32 Stride = FStaticMeshInstanceData::VectorsPerInstance * sizeof(FVector4);
 				for (int32 FirstUnfixedIndex = 0; FirstUnfixedIndex < NumKept; FirstUnfixedIndex++)
 				{
 					int32 LoadFrom = SortedInstances[FirstUnfixedIndex];
 					if (LoadFrom != FirstUnfixedIndex)
 					{
 						check(LoadFrom > FirstUnfixedIndex);
-						FMemory::Memcpy(SwapBuffer, InstanceBuffer.GetInstanceWriteAddress(FirstUnfixedIndex), Stride);
-						FMemory::Memcpy(InstanceBuffer.GetInstanceWriteAddress(FirstUnfixedIndex), InstanceBuffer.GetInstanceWriteAddress(LoadFrom), Stride);
-						FMemory::Memcpy(InstanceBuffer.GetInstanceWriteAddress(LoadFrom), SwapBuffer, Stride);
+						FMemory::Memcpy(&SwapBuffer, InstanceBuffer.GetInstanceWriteAddress(FirstUnfixedIndex), sizeof(FInstanceStream));
+						FMemory::Memcpy(InstanceBuffer.GetInstanceWriteAddress(FirstUnfixedIndex), InstanceBuffer.GetInstanceWriteAddress(LoadFrom), sizeof(FInstanceStream));
+						FMemory::Memcpy(InstanceBuffer.GetInstanceWriteAddress(LoadFrom), &SwapBuffer, sizeof(FInstanceStream));
 
 						int32 SwapGoesTo = InstanceReorderTable[FirstUnfixedIndex];
 						check(SwapGoesTo > FirstUnfixedIndex);
diff --git a/Engine/Source/Runtime/OpenGLDrv/Private/OpenGLVertexDeclaration.cpp b/Engine/Source/Runtime/OpenGLDrv/Private/OpenGLVertexDeclaration.cpp
--- a/Engine/Source/Runtime/OpenGLDrv/Private/OpenGLVertexDeclaration.cpp
+++ b/Engine/Source/Runtime/OpenGLDrv/Private/OpenGLVertexDeclaration.cpp
@@ -79,6 +79,7 @@ struct FOpenGLVertexDeclarationKey
 						SetupGLElement(GLElement, GL_SHORT, 4, false, true);
 					}
 					break;
+				case VET_Short4N:		SetupGLElement(GLElement, GL_SHORT,			4,			true,	true); break;
 				default: UE_LOG(LogRHI, Fatal,TEXT("Unknown RHI vertex element type %u"),(uint8)InElements[ElementIndex].Type);
 			};
 			VertexElements.Add(GLElement);
diff --git a/Engine/Source/Runtime/RHI/Public/RHIDefinitions.h b/Engine/Source/Runtime/RHI/Public/RHIDefinitions.h
--- a/Engine/Source/Runtime/RHI/Public/RHIDefinitions.h
+++ b/Engine/Source/Runtime/RHI/Public/RHIDefinitions.h
@@ -213,6 +213,7 @@ enum EVertexElementType
 	VET_Short2N,		// 16 bit word normalized to (value/32767.0,value/32767.0,0,0,1)
 	VET_Half2,			// 16 bit float using 1 bit sign, 5 bit exponent, 10 bit mantissa 
 	VET_Half4,
+	VET_Short4N,		// 4 X 16 bit word, normalized 
 	VET_MAX
 };
 
diff --git a/Engine/Source/Runtime/ShaderCore/Private/VertexFactory.cpp b/Engine/Source/Runtime/ShaderCore/Private/VertexFactory.cpp
--- a/Engine/Source/Runtime/ShaderCore/Private/VertexFactory.cpp
+++ b/Engine/Source/Runtime/ShaderCore/Private/VertexFactory.cpp
@@ -284,37 +284,6 @@ FVertexElement FVertexFactory::AccessPositionStreamComponent(const FVertexStream
 	return FVertexElement(PositionStream.AddUnique(VertexStream),Component.Offset,Component.Type,AttributeIndex,VertexStream.Stride,Component.bUseInstanceIndex);
 }
 
-//TTP:51684
-/**
- *	GetVertexElementSize returns the size of data based on its type
- *	it's needed for FixFXCardDeclarator function
- */
-static uint8 GetVertexElementSize( uint8 Type )
-{
-	switch( Type )
-	{
-		case VET_Float1:
-			return 4;
-		case VET_Float2:
-			return 8;
-		case VET_Float3:
-			return 12;
-		case VET_Float4:
-			return 16;
-		case VET_PackedNormal:
-		case VET_UByte4:
-		case VET_UByte4N:
-		case VET_Color:
-		case VET_Short2:
-		case VET_Short2N:
-		case VET_Short4:
-		case VET_Half2:
-		case VET_Half4:
-		default:
-			return 0;
-	}
-}
-
 void FVertexFactory::InitDeclaration(
 	FVertexDeclarationElementList& Elements, 
 	const DataType& InData)
diff --git a/Engine/Source/Runtime/Windows/D3D11RHI/Private/D3D11VertexDeclaration.cpp b/Engine/Source/Runtime/Windows/D3D11RHI/Private/D3D11VertexDeclaration.cpp
--- a/Engine/Source/Runtime/Windows/D3D11RHI/Private/D3D11VertexDeclaration.cpp
+++ b/Engine/Source/Runtime/Windows/D3D11RHI/Private/D3D11VertexDeclaration.cpp
@@ -40,6 +40,7 @@ struct FD3D11VertexDeclarationKey
 			case VET_Short2N:		D3DElement.Format = DXGI_FORMAT_R16G16_SNORM; break;
 			case VET_Half2:			D3DElement.Format = DXGI_FORMAT_R16G16_FLOAT; break;
 			case VET_Half4:			D3DElement.Format = DXGI_FORMAT_R16G16B16A16_FLOAT; break;
+			case VET_Short4N:		D3DElement.Format = DXGI_FORMAT_R16G16B16A16_SNORM; break;
 			default: UE_LOG(LogD3D11RHI, Fatal,TEXT("Unknown RHI vertex element type %u"),(uint8)InElements[ElementIndex].Type);
 			};
 			D3DElement.SemanticName = "ATTRIBUTE";
