Material compiler CustomOutput
Landscape Grass material output node
LandscapeGrassType asset and factory

diff --git a/Engine/Source/Developer/AssetTools/Private/AssetTools.cpp b/Engine/Source/Developer/AssetTools/Private/AssetTools.cpp
--- a/Engine/Source/Developer/AssetTools/Private/AssetTools.cpp
+++ b/Engine/Source/Developer/AssetTools/Private/AssetTools.cpp
@@ -53,7 +53,8 @@ FAssetTools::FAssetTools()
 	RegisterAssetTypeActions( MakeShareable(new FAssetTypeActions_InstancedFoliageSettings) );
 	RegisterAssetTypeActions( MakeShareable(new FAssetTypeActions_InterpData) );
 	RegisterAssetTypeActions( MakeShareable(new FAssetTypeActions_LandscapeLayer) );
-	RegisterAssetTypeActions( MakeShareable(new FAssetTypeActions_Material) );
+	RegisterAssetTypeActions( MakeShareable(new FAssetTypeActions_LandscapeGrassType));
+	RegisterAssetTypeActions( MakeShareable(new FAssetTypeActions_Material));
 	RegisterAssetTypeActions( MakeShareable(new FAssetTypeActions_MaterialFunction) );
 	RegisterAssetTypeActions( MakeShareable(new FAssetTypeActions_MaterialInstanceConstant) );
 	RegisterAssetTypeActions( MakeShareable(new FAssetTypeActions_MaterialInterface) );
diff --git a/Engine/Source/Developer/AssetTools/Private/AssetToolsPrivatePCH.h b/Engine/Source/Developer/AssetTools/Private/AssetToolsPrivatePCH.h
--- a/Engine/Source/Developer/AssetTools/Private/AssetToolsPrivatePCH.h
+++ b/Engine/Source/Developer/AssetTools/Private/AssetToolsPrivatePCH.h
@@ -51,6 +51,7 @@ DECLARE_LOG_CATEGORY_EXTERN(LogAssetTools, Log, All);
 #include "AssetTypeActions/AssetTypeActions_InstancedFoliageSettings.h"
 #include "AssetTypeActions/AssetTypeActions_InterpData.h"
 #include "AssetTypeActions/AssetTypeActions_LandscapeLayer.h"
+#include "AssetTypeActions/AssetTypeActions_LandscapeGrassType.h"
 #include "AssetTypeActions/AssetTypeActions_Material.h"
 #include "AssetTypeActions/AssetTypeActions_MaterialFunction.h"
 #include "AssetTypeActions/AssetTypeActions_MaterialInstanceConstant.h"
diff --git a/Engine/Source/Developer/AssetTools/Private/AssetTypeActions/AssetTypeActions_LandscapeGrassType.h b/Engine/Source/Developer/AssetTools/Private/AssetTypeActions/AssetTypeActions_LandscapeGrassType.h
new file mode 100644
--- /dev/null
+++ b/Engine/Source/Developer/AssetTools/Private/AssetTypeActions/AssetTypeActions_LandscapeGrassType.h
@@ -0,0 +1,15 @@
+// Copyright 1998-2015 Epic Games, Inc. All Rights Reserved.
+
+#pragma once
+
+#include "LandscapeGrassType.h"
+
+class FAssetTypeActions_LandscapeGrassType : public FAssetTypeActions_Base
+{
+public:
+	// IAssetTypeActions Implementation
+	virtual FText GetName() const override { return NSLOCTEXT("AssetTypeActions", "AssetTypeActions_LandscapeGrassType", "Landscape Grass Type"); }
+	virtual FColor GetTypeColor() const override { return FColor(128,255,128); }
+	virtual UClass* GetSupportedClass() const override { return ULandscapeGrassType::StaticClass(); }
+	virtual uint32 GetCategories() override { return EAssetTypeCategories::Misc; }
+};
\ No newline at end of file
diff --git a/Engine/Source/Editor/FoliageEdit/Private/LandscapeGrassTypeFactory.cpp b/Engine/Source/Editor/FoliageEdit/Private/LandscapeGrassTypeFactory.cpp
new file mode 100644
--- /dev/null
+++ b/Engine/Source/Editor/FoliageEdit/Private/LandscapeGrassTypeFactory.cpp
@@ -0,0 +1,26 @@
+// Copyright 1998-2015 Epic Games, Inc. All Rights Reserved.
+
+#include "UnrealEd.h"
+#include "LandscapeGrassTypeFactory.h"
+#include "IAssetTypeActions.h"
+#include "LandscapeGrassType.h"
+
+ULandscapeGrassTypeFactory::ULandscapeGrassTypeFactory(const FObjectInitializer& ObjectInitializer)
+	: Super(ObjectInitializer)
+{
+	bCreateNew = true;
+	bEditAfterNew = true;
+	SupportedClass = ULandscapeGrassType::StaticClass();
+}
+
+UObject* ULandscapeGrassTypeFactory::FactoryCreateNew(UClass* Class, UObject* InParent, FName Name, EObjectFlags Flags, UObject* Context, FFeedbackContext* Warn)
+{
+	ULandscapeGrassType* NewGrassType = ConstructObject<ULandscapeGrassType>(Class, InParent, Name, Flags | RF_Transactional);
+
+	return NewGrassType;
+}
+
+uint32 ULandscapeGrassTypeFactory::GetMenuCategories() const
+{
+	return EAssetTypeCategories::Misc;
+}
\ No newline at end of file
diff --git a/Engine/Source/Editor/FoliageEdit/Private/LandscapeGrassTypeFactory.h b/Engine/Source/Editor/FoliageEdit/Private/LandscapeGrassTypeFactory.h
new file mode 100644
--- /dev/null
+++ b/Engine/Source/Editor/FoliageEdit/Private/LandscapeGrassTypeFactory.h
@@ -0,0 +1,21 @@
+// Copyright 1998-2015 Epic Games, Inc. All Rights Reserved.
+
+/**
+ * Factory for LandscapeGrassType assets
+ */
+
+#pragma once
+
+#include "LandscapeGrassTypeFactory.generated.h"
+
+
+UCLASS()
+class ULandscapeGrassTypeFactory : public UFactory
+{
+	GENERATED_UCLASS_BODY()
+
+	// UFactory interface
+	virtual UObject* FactoryCreateNew(UClass* Class, UObject* InParent, FName Name, EObjectFlags Flags, UObject* Context, FFeedbackContext* Warn) override;
+	virtual uint32 GetMenuCategories() const;
+	// End of UFactory interface
+};
diff --git a/Engine/Source/Editor/MaterialEditor/Private/MaterialEditor.cpp b/Engine/Source/Editor/MaterialEditor/Private/MaterialEditor.cpp
--- a/Engine/Source/Editor/MaterialEditor/Private/MaterialEditor.cpp
+++ b/Engine/Source/Editor/MaterialEditor/Private/MaterialEditor.cpp
@@ -1128,17 +1128,8 @@ void FMaterialEditor::RegenerateCodeView(bool bForce)
 		return;
 	}
 
-	TMap<FMaterialExpressionKey,int32> ExpressionCodeMap[MP_MAX][SF_NumFrequencies];
-	for (int32 PropertyIndex = 0; PropertyIndex < MP_MAX; PropertyIndex++)
-	{
-		for (int32 FrequencyIndex = 0; FrequencyIndex < SF_NumFrequencies; FrequencyIndex++)
-		{
-			ExpressionCodeMap[PropertyIndex][FrequencyIndex].Empty();
-		}
-	}
-
 	FString MarkupSource;
-	if (Material->GetMaterialResource(GMaxRHIFeatureLevel)->GetMaterialExpressionSource(MarkupSource, ExpressionCodeMap))
+	if (Material->GetMaterialResource(GMaxRHIFeatureLevel)->GetMaterialExpressionSource(MarkupSource))
 	{
 		// Remove line-feeds and leave just CRs so the character counts match the selection ranges.
 		MarkupSource.ReplaceInline(TEXT("\r"), TEXT(""));
diff --git a/Engine/Source/Runtime/Engine/Classes/Materials/Material.h b/Engine/Source/Runtime/Engine/Classes/Materials/Material.h
--- a/Engine/Source/Runtime/Engine/Classes/Materials/Material.h
+++ b/Engine/Source/Runtime/Engine/Classes/Materials/Material.h
@@ -944,6 +944,20 @@ public:
 		return Result;
 	}
 
+	/* Get all expressions of the requested type */
+	template<typename ExpressionType>
+	void GetAllExpressionsOfType(TArray<const ExpressionType*>& OutExpressions) const
+	{
+		for (int32 ExpressionIndex = 0; ExpressionIndex < Expressions.Num(); ExpressionIndex++)
+		{
+			ExpressionType* ExpressionPtr = Cast<ExpressionType>(Expressions[ExpressionIndex]);
+			if (ExpressionPtr)
+			{
+				OutExpressions.Add(ExpressionPtr);
+			}
+		}
+	}
+
 	/** Determines whether each quality level has different nodes by inspecting the material's expressions. */
 	void GetQualityLevelNodeUsage(TArray<bool, TInlineAllocator<EMaterialQualityLevel::Num> >& QualityLevelsUsed);
 
diff --git a/Engine/Source/Runtime/Engine/Classes/Materials/MaterialExpressionCustomOutput.h b/Engine/Source/Runtime/Engine/Classes/Materials/MaterialExpressionCustomOutput.h
new file mode 100644
--- /dev/null
+++ b/Engine/Source/Runtime/Engine/Classes/Materials/MaterialExpressionCustomOutput.h
@@ -0,0 +1,19 @@
+// Copyright 1998-2015 Epic Games, Inc. All Rights Reserved.
+
+
+#pragma once
+#include "Materials/MaterialExpression.h"
+#include "MaterialExpressionCustomOutput.generated.h"
+
+UCLASS(collapsecategories, hidecategories=Object, MinimalAPI)
+class UMaterialExpressionCustomOutput : public UMaterialExpression
+{
+	GENERATED_UCLASS_BODY()
+
+	// Override to enable multiple outputs
+	virtual int32 GetNumOutputs() const { return 1; };
+	virtual FString GetFunctionName() const PURE_VIRTUAL(UMaterialExpressionCustomOutput::GetFunctionName, return TEXT("GetCustomOutput"););
+};
+
+
+
diff --git a/Engine/Source/Runtime/Engine/Private/Materials/HLSLMaterialTranslator.h b/Engine/Source/Runtime/Engine/Private/Materials/HLSLMaterialTranslator.h
--- a/Engine/Source/Runtime/Engine/Private/Materials/HLSLMaterialTranslator.h
+++ b/Engine/Source/Runtime/Engine/Private/Materials/HLSLMaterialTranslator.h
@@ -17,6 +17,7 @@
 #include "Materials/MaterialExpressionTransform.h"
 #include "Materials/MaterialExpressionTransformPosition.h"
 #include "Materials/MaterialExpressionCustom.h"
+#include "Materials/MaterialExpressionCustomOutput.h"
 #include "Materials/MaterialFunction.h"
 #include "MaterialCompiler.h"
 #include "MaterialUniformExpressions.h"
@@ -109,6 +110,10 @@ protected:
 	EShaderFrequency ShaderFrequency;
 	/** The current material property being compiled.  This affects the behavior of all compiler functions except GetFixedParameterCode. */
 	EMaterialProperty MaterialProperty;
+	/** The code chunks corresponding to the currently compiled property or custom output. */
+	TArray<FShaderCodeChunk>* CurrentScopeChunks;
+	/* Stack that tracks compiler state specific to the function currently being compiled. */
+	TArray<FMaterialFunctionCompileState> FunctionStack;
 	/** Material being compiled.  Only transient compilation output like error information can be stored on the FMaterial. */
 	FMaterial* Material;
 	/** Compilation output which will be stored in the DDC. */
@@ -137,14 +142,11 @@ protected:
 	FString MaterialTemplate;
 
 	// Array of code chunks per material property
-	TArray<FShaderCodeChunk> CodeChunks[MP_MAX][SF_NumFrequencies];
+	TArray<FShaderCodeChunk> PropertyCodeChunks[MP_MAX][SF_NumFrequencies];
 
 	// Uniform expressions used across all material properties
 	TArray<FShaderCodeChunk> UniformExpressions;
 
-	// Stack that tracks compiler state specific to each function being compiled
-	TArray<FMaterialFunctionCompileState> FunctionStack;
-
 	/** Parameter collections referenced by this material.  The position in this array is used as an index on the shader parameter. */
 	TArray<UMaterialParameterCollection*> ParameterCollections;
 
@@ -154,6 +156,9 @@ protected:
 	/** Any custom expression function implementations */
 	TArray<FString> CustomExpressionImplementations;
 
+	/** Any custom output function implementations */
+	TArray<FString> CustomOutputImplementations;
+
 	/** Whether the translation succeeded. */
 	uint32 bSuccess : 1;
 	/** Whether the compute shader material inputs were compiled. */
@@ -256,9 +261,6 @@ public:
 			check(!MaterialCompilationOutput.bNeedsSceneTextures);
 			MaterialCompilationOutput.bUsesEyeAdaptation = false;
 
-			// Add a state item for the root level
-			FunctionStack.Add(FMaterialFunctionCompileState(NULL));
-
 			bCompileForComputeShader = Material->IsLightFunction();
 
 			// Generate code
@@ -341,8 +343,6 @@ public:
 				}
 			}
 
-			check(FunctionStack.Num() == 1);
-
 			if (Chunk[MP_WorldPositionOffset] == -1)
 			{
 				MaterialCompilationOutput.bModifiesMeshPosition = false;
@@ -350,7 +350,7 @@ public:
 			else
 			{
 				int32 WPOFreq = (int32)GetMaterialPropertyShaderFrequency(MP_WorldPositionOffset);
-				FShaderCodeChunk& WPOChunk = CodeChunks[MP_WorldPositionOffset][WPOFreq][Chunk[MP_WorldPositionOffset]];
+				FShaderCodeChunk& WPOChunk = PropertyCodeChunks[MP_WorldPositionOffset][WPOFreq][Chunk[MP_WorldPositionOffset]];
 
 				// Determine whether the world position offset is used. 
 				// If the output chunk has a uniform expression, it is constant, and GetNumberValue returns the default property value then WPO isn't used.
@@ -470,7 +470,7 @@ public:
 
 				GetFixedParameterCode(
 					Chunk[PropertyId], 
-					(EMaterialProperty)PropertyId, GetMaterialPropertyShaderFrequency((EMaterialProperty)PropertyId),
+					PropertyCodeChunks[PropertyId][GetMaterialPropertyShaderFrequency((EMaterialProperty)PropertyId)],
 					TranslatedCodeChunkDefinitions[PropertyId],
 					TranslatedCodeChunks[PropertyId]);
 			}
@@ -480,19 +480,57 @@ public:
 				switch(PropertyId)
 				{
 				case CompiledMP_EmissiveColorCS:
-					if (bCompileForComputeShader)
-					{
-						GetFixedParameterCode(Chunk[PropertyId], MP_EmissiveColor, SF_Compute, TranslatedCodeChunkDefinitions[PropertyId], TranslatedCodeChunks[PropertyId]);
-					}
+			    	if (bCompileForComputeShader)
+				    {
+						GetFixedParameterCode(Chunk[PropertyId], PropertyCodeChunks[MP_EmissiveColor][SF_Compute], TranslatedCodeChunkDefinitions[PropertyId], TranslatedCodeChunks[PropertyId]);
+				    }
 					break;
 				case CompiledMP_PrevWorldPositionOffset:
-					GetFixedParameterCode(Chunk[PropertyId], MP_WorldPositionOffset, SF_Vertex, TranslatedCodeChunkDefinitions[PropertyId], TranslatedCodeChunks[PropertyId]);
+					GetFixedParameterCode(Chunk[PropertyId], PropertyCodeChunks[MP_WorldPositionOffset][SF_Vertex], TranslatedCodeChunkDefinitions[PropertyId], TranslatedCodeChunks[PropertyId]);
 					break;
 				default: check(0);
 					break;
 				}
 			}
 
+			// Gather and output the implementation for any custom output expressions
+			TArray<UMaterialExpressionCustomOutput*> CustomOutputExpressions;
+			Material->GatherCustomOutputExpressions(CustomOutputExpressions);
+			TSet<UClass*> SeenCustomOutputExpressionsClases;
+
+			for(UMaterialExpressionCustomOutput* CustomOutput : CustomOutputExpressions)
+			{
+				if(SeenCustomOutputExpressionsClases.Contains(CustomOutput->GetClass()))
+				{
+					Errorf(TEXT("The material can contain only one %s node"), *CustomOutput->GetDescription());
+				}
+				else
+				{
+					SeenCustomOutputExpressionsClases.Add(CustomOutput->GetClass());
+
+					int32 NumOutputs = CustomOutput->GetNumOutputs();
+					ResourcesString += FString::Printf(TEXT("#define NUM_MATERIAL_OUTPUTS_%s %d\r\n"), *CustomOutput->GetFunctionName().ToUpper(), NumOutputs);
+					if (NumOutputs > 0)
+					{
+						for (int32 Index = 0; Index < NumOutputs; Index++)
+						{
+							FunctionStack.Empty();
+							FunctionStack.Add(FMaterialFunctionCompileState(NULL));
+							MaterialProperty = MP_MAX; // Indicates we're not compiling any material property.
+							ShaderFrequency = SF_Pixel;
+							TArray<FShaderCodeChunk> CustomExpressionChunks;
+							CurrentScopeChunks = &CustomExpressionChunks;
+							CustomOutput->Compile(this, Index, 0);
+						}
+					}
+				}
+			}
+
+			for (int32 ExpressionIndex = 0; ExpressionIndex < CustomOutputImplementations.Num(); ExpressionIndex++)
+			{
+				ResourcesString += CustomOutputImplementations[ExpressionIndex] + "\r\n\r\n";
+			}
+
 			LoadShaderSourceFileChecked(TEXT("MaterialTemplate"), MaterialTemplate);
 
 			// Find the string index of the '#line' statement in MaterialTemplate.usf
@@ -693,8 +731,8 @@ protected:
 			return Default;
 		}
 
-		checkf(Index >= 0 && Index < CodeChunks[MaterialProperty][ShaderFrequency].Num(), TEXT("Index %d/%d, Platform=%d"), Index, CodeChunks[MaterialProperty][ShaderFrequency].Num(), Platform);
-		const FShaderCodeChunk& CodeChunk = CodeChunks[MaterialProperty][ShaderFrequency][Index];
+		checkf(Index >= 0 && Index < CurrentScopeChunks->Num(), TEXT("Index %d/%d, Platform=%d"), Index, CurrentScopeChunks->Num(), Platform);
+		const FShaderCodeChunk& CodeChunk = (*CurrentScopeChunks)[Index];
 		if((CodeChunk.UniformExpression && CodeChunk.UniformExpression->IsConstant()) || CodeChunk.bInline)
 		{
 			// Constant uniform expressions and code chunks which are marked to be inlined are accessed via Definition
@@ -706,7 +744,7 @@ protected:
 			{
 				// If the code chunk has a uniform expression, create a new code chunk to access it
 				const int32 AccessedIndex = AccessUniformExpression(Index);
-				const FShaderCodeChunk& AccessedCodeChunk = CodeChunks[MaterialProperty][ShaderFrequency][AccessedIndex];
+				const FShaderCodeChunk& AccessedCodeChunk = (*CurrentScopeChunks)[AccessedIndex];
 				if(AccessedCodeChunk.bInline)
 				{
 					// Handle the accessed code chunk being inlined
@@ -724,12 +762,12 @@ protected:
 	}
 
 	/** Creates a string of all definitions needed for the given material input. */
-	FString GetDefinitions(EMaterialProperty InProperty, EShaderFrequency InFrequency) const
+	FString GetDefinitions(TArray<FShaderCodeChunk>& CodeChunks) const
 	{
 		FString Definitions;
-		for (int32 ChunkIndex = 0; ChunkIndex < CodeChunks[InProperty][InFrequency].Num(); ChunkIndex++)
+		for (int32 ChunkIndex = 0; ChunkIndex < CodeChunks.Num(); ChunkIndex++)
 		{
-			const FShaderCodeChunk& CodeChunk = CodeChunks[InProperty][InFrequency][ChunkIndex];
+			const FShaderCodeChunk& CodeChunk = CodeChunks[ChunkIndex];
 			// Uniform expressions (both constant and variable) and inline expressions don't have definitions.
 			if (!CodeChunk.UniformExpression && !CodeChunk.bInline)
 			{
@@ -740,25 +778,25 @@ protected:
 	}
 
 	// GetFixedParameterCode
-	virtual void GetFixedParameterCode(int32 Index, EMaterialProperty InProperty, EShaderFrequency InFrequency, FString& OutDefinitions, FString& OutValue)
+	void GetFixedParameterCode(int32 ResultIndex, TArray<FShaderCodeChunk>& CodeChunks, FString& OutDefinitions, FString& OutValue)
 	{
-		if(Index != INDEX_NONE)
+		if (ResultIndex != INDEX_NONE)
 		{
-			checkf(Index >= 0 && Index < CodeChunks[InProperty][InFrequency].Num(), TEXT("Index out of range %d/%d [%s]"), Index, CodeChunks[InProperty][InFrequency].Num(), *Material->GetFriendlyName());
-			check(!CodeChunks[InProperty][InFrequency][Index].UniformExpression || CodeChunks[InProperty][InFrequency][Index].UniformExpression->IsConstant());
-			if (CodeChunks[InProperty][InFrequency][Index].UniformExpression && CodeChunks[InProperty][InFrequency][Index].UniformExpression->IsConstant())
+			checkf(ResultIndex >= 0 && ResultIndex < CodeChunks.Num(), TEXT("Index out of range %d/%d [%s]"), ResultIndex, CodeChunks.Num(), *Material->GetFriendlyName());
+			check(!CodeChunks[ResultIndex].UniformExpression || CodeChunks[ResultIndex].UniformExpression->IsConstant());
+			if (CodeChunks[ResultIndex].UniformExpression && CodeChunks[ResultIndex].UniformExpression->IsConstant())
 			{
 				// Handle a constant uniform expression being the only code chunk hooked up to a material input
-				const FShaderCodeChunk& CodeChunk = CodeChunks[InProperty][InFrequency][Index];
-				OutValue = CodeChunk.Definition;
+				const FShaderCodeChunk& ResultChunk = CodeChunks[ResultIndex];
+				OutValue = ResultChunk.Definition;
 			}
 			else
 			{
-				const FShaderCodeChunk& CodeChunk = CodeChunks[InProperty][InFrequency][Index];
+				const FShaderCodeChunk& ResultChunk = CodeChunks[ResultIndex];
 				// Combine the definition lines and the return statement
-				check(CodeChunk.bInline || CodeChunk.SymbolName.Len() > 0);
-				OutDefinitions = GetDefinitions(InProperty, InFrequency);
-				OutValue = CodeChunk.bInline ? CodeChunk.Definition : CodeChunk.SymbolName;
+				check(ResultChunk.bInline || ResultChunk.SymbolName.Len() > 0);
+				OutDefinitions = GetDefinitions(CodeChunks);
+				OutValue = ResultChunk.bInline ? ResultChunk.Definition : ResultChunk.SymbolName;
 			}
 		}
 		else
@@ -847,21 +885,21 @@ protected:
 
 		if (bInlined)
 		{
-			const int32 CodeIndex = CodeChunks[MaterialProperty][ShaderFrequency].Num();
+			const int32 CodeIndex = CurrentScopeChunks->Num();
 			// Adding an inline code chunk, the definition will be the code to inline
-			new(CodeChunks[MaterialProperty][ShaderFrequency]) FShaderCodeChunk(FormattedCode,TEXT(""),Type,true);
+			new(*CurrentScopeChunks) FShaderCodeChunk(FormattedCode,TEXT(""),Type,true);
 			return CodeIndex;
 		}
 		// Can only create temporaries for float and material attribute types.
 		else if (Type & (MCT_Float))
 		{
-			const int32 CodeIndex = CodeChunks[MaterialProperty][ShaderFrequency].Num();
+			const int32 CodeIndex = CurrentScopeChunks->Num();
 			// Allocate a local variable name
 			const FString SymbolName = CreateSymbolName(TEXT("Local"));
 			// Construct the definition string which stores the result in a temporary and adds a newline for readability
 			const FString LocalVariableDefinition = FString("	") + HLSLTypeString(Type) + TEXT(" ") + SymbolName + TEXT(" = ") + FormattedCode + TEXT(";") + LINE_TERMINATOR;
 			// Adding a code chunk that creates a local variable
-			new(CodeChunks[MaterialProperty][ShaderFrequency]) FShaderCodeChunk(*LocalVariableDefinition,SymbolName,Type,false);
+			new(*CurrentScopeChunks) FShaderCodeChunk(*LocalVariableDefinition,SymbolName,Type,false);
 			return CodeIndex;
 		}
 		else
@@ -973,13 +1011,13 @@ protected:
 				// This allows multiple material properties to share uniform expressions because AccessUniformExpression uses AddUniqueItem when adding uniform expressions.
 				check(Type == UniformExpressions[ExpressionIndex].Type);
 				// Search for an existing code chunk with the same uniform expression in the array of code chunks for this material property.
-				for(int32 ChunkIndex = 0;ChunkIndex < CodeChunks[MaterialProperty][ShaderFrequency].Num();ChunkIndex++)
+				for(int32 ChunkIndex = 0;ChunkIndex < CurrentScopeChunks->Num();ChunkIndex++)
 				{
-					FMaterialUniformExpression* OtherExpression = CodeChunks[MaterialProperty][ShaderFrequency][ChunkIndex].UniformExpression;
+					FMaterialUniformExpression* OtherExpression = (*CurrentScopeChunks)[ChunkIndex].UniformExpression;
 					if(OtherExpression && OtherExpression->IsIdentical(UniformExpression))
 					{
 						delete UniformExpression;
-						// Reuse the entry in CodeChunks[MaterialProperty][ShaderFrequency]
+						// Reuse the entry in CurrentScopeChunks
 						return ChunkIndex;
 					}
 				}
@@ -1003,9 +1041,9 @@ protected:
 		};
 		FormattedCode[Result] = 0;
 
-		const int32 ReturnIndex = CodeChunks[MaterialProperty][ShaderFrequency].Num();
+		const int32 ReturnIndex = CurrentScopeChunks->Num();
 		// Create a new code chunk for the uniform expression
-		new(CodeChunks[MaterialProperty][ShaderFrequency]) FShaderCodeChunk(UniformExpression,FormattedCode,Type);
+		new(*CurrentScopeChunks) FShaderCodeChunk(UniformExpression,FormattedCode,Type);
 
 		if (!bFoundExistingExpression)
 		{
@@ -1020,8 +1058,8 @@ protected:
 	// AccessUniformExpression - Adds code to access the value of a uniform expression to the Code array and returns its index.
 	int32 AccessUniformExpression(int32 Index)
 	{
-		check(Index >= 0 && Index < CodeChunks[MaterialProperty][ShaderFrequency].Num());
-		const FShaderCodeChunk&	CodeChunk = CodeChunks[MaterialProperty][ShaderFrequency][Index];
+		check(Index >= 0 && Index < CurrentScopeChunks->Num());
+		const FShaderCodeChunk&	CodeChunk = (*CurrentScopeChunks)[Index];
 		check(CodeChunk.UniformExpression && !CodeChunk.UniformExpression->IsConstant());
 
 		FMaterialUniformExpressionTexture* TextureUniformExpression = CodeChunk.UniformExpression->GetTextureUniformExpression();
@@ -1041,9 +1079,9 @@ protected:
 				}
 				else
 				{
-					const int32 ScalarInputIndex = MaterialCompilationOutput.UniformExpressionSet.PerFrameUniformScalarExpressions.AddUnique(CodeChunk.UniformExpression);
-					FCString::Sprintf(FormattedCode, TEXT("UE_Material_PerFrameScalarExpression%u"), ScalarInputIndex);
-				}
+				const int32 ScalarInputIndex = MaterialCompilationOutput.UniformExpressionSet.PerFrameUniformScalarExpressions.AddUnique(CodeChunk.UniformExpression);
+				FCString::Sprintf(FormattedCode, TEXT("UE_Material_PerFrameScalarExpression%u"), ScalarInputIndex);
+			}
 			}
 			else
 			{
@@ -1074,9 +1112,9 @@ protected:
 				}
 				else
 				{
-					const int32 VectorInputIndex = MaterialCompilationOutput.UniformExpressionSet.PerFrameUniformVectorExpressions.AddUnique(CodeChunk.UniformExpression);
-					FCString::Sprintf(FormattedCode, TEXT("UE_Material_PerFrameVectorExpression%u%s"), VectorInputIndex, Mask);
-				}
+				const int32 VectorInputIndex = MaterialCompilationOutput.UniformExpressionSet.PerFrameUniformVectorExpressions.AddUnique(CodeChunk.UniformExpression);
+				FCString::Sprintf(FormattedCode, TEXT("UE_Material_PerFrameVectorExpression%u%s"), VectorInputIndex, Mask);
+			}
 			}
 			else
 			{
@@ -1108,14 +1146,14 @@ protected:
 			UE_LOG(LogMaterial, Fatal,TEXT("User input of unknown type: %s"),DescribeType(CodeChunk.Type));
 		}
 
-		return AddInlinedCodeChunk(CodeChunks[MaterialProperty][ShaderFrequency][Index].Type,FormattedCode);
+		return AddInlinedCodeChunk((*CurrentScopeChunks)[Index].Type,FormattedCode);
 	}
 
 	// CoerceParameter
 	FString CoerceParameter(int32 Index,EMaterialValueType DestType)
 	{
-		check(Index >= 0 && Index < CodeChunks[MaterialProperty][ShaderFrequency].Num());
-		const FShaderCodeChunk&	CodeChunk = CodeChunks[MaterialProperty][ShaderFrequency][Index];
+		check(Index >= 0 && Index < CurrentScopeChunks->Num());
+		const FShaderCodeChunk&	CodeChunk = (*CurrentScopeChunks)[Index];
 		if( CodeChunk.Type == DestType )
 		{
 			return GetParameterCode(Index);
@@ -1149,16 +1187,16 @@ protected:
 	// GetParameterType
 	EMaterialValueType GetParameterType(int32 Index) const
 	{
-		check(Index >= 0 && Index < CodeChunks[MaterialProperty][ShaderFrequency].Num());
-		return CodeChunks[MaterialProperty][ShaderFrequency][Index].Type;
+		check(Index >= 0 && Index < CurrentScopeChunks->Num());
+		return (*CurrentScopeChunks)[Index].Type;
 	}
 
 	// GetParameterUniformExpression
 	FMaterialUniformExpression* GetParameterUniformExpression(int32 Index) const
 	{
-		check(Index >= 0 && Index < CodeChunks[MaterialProperty][ShaderFrequency].Num());
+		check(Index >= 0 && Index < CurrentScopeChunks->Num());
 
-		const FShaderCodeChunk& Chunk = CodeChunks[MaterialProperty][ShaderFrequency][Index];
+		const FShaderCodeChunk& Chunk = (*CurrentScopeChunks)[Index];
 
 		return Chunk.UniformExpression;
 	}
@@ -1197,11 +1235,11 @@ protected:
 
 	EMaterialValueType GetArithmeticResultType(int32 A,int32 B)
 	{
-		check(A >= 0 && A < CodeChunks[MaterialProperty][ShaderFrequency].Num());
-		check(B >= 0 && B < CodeChunks[MaterialProperty][ShaderFrequency].Num());
+		check(A >= 0 && A < CurrentScopeChunks->Num());
+		check(B >= 0 && B < CurrentScopeChunks->Num());
 
-		EMaterialValueType	TypeA = CodeChunks[MaterialProperty][ShaderFrequency][A].Type,
-			TypeB = CodeChunks[MaterialProperty][ShaderFrequency][B].Type;
+		EMaterialValueType	TypeA = (*CurrentScopeChunks)[A].Type,
+			TypeB = (*CurrentScopeChunks)[B].Type;
 
 		return GetArithmeticResultType(TypeA,TypeB);
 	}
@@ -1215,6 +1253,9 @@ protected:
 	 */
 	virtual void SetMaterialProperty(EMaterialProperty InProperty, EShaderFrequency OverrideShaderFrequency = SF_NumFrequencies, bool bUsePreviousFrameTime = false)
 	{
+		FunctionStack.Empty();
+		FunctionStack.Add(FMaterialFunctionCompileState(NULL));
+
 		MaterialProperty = InProperty;
 
 		if(OverrideShaderFrequency != SF_NumFrequencies)
@@ -1227,6 +1268,8 @@ protected:
 		}
 
 		bUsesPreviousFrameTime = bUsePreviousFrameTime;
+
+		CurrentScopeChunks = &PropertyCodeChunks[MaterialProperty][ShaderFrequency];
 	}
 	virtual EShaderFrequency GetCurrentShaderFrequency() const
 	{
@@ -1281,7 +1324,7 @@ protected:
 	virtual int32 CallExpression(FMaterialExpressionKey ExpressionKey,FMaterialCompiler* Compiler) override
 	{
 		// Check if this expression has already been translated.
-		int32* ExistingCodeIndex = FunctionStack.Last().ExpressionCodeMap[MaterialProperty][ShaderFrequency].Find(ExpressionKey);
+		int32* ExistingCodeIndex = FunctionStack.Last().ExpressionCodeMap.Find(ExpressionKey);
 		if(ExistingCodeIndex)
 		{
 			return *ExistingCodeIndex;
@@ -1307,7 +1350,7 @@ protected:
 			check(FunctionDepth == FunctionStack.Num());
 
 			// Cache the translation.
-			FunctionStack.Last().ExpressionCodeMap[MaterialProperty][ShaderFrequency].Add(ExpressionKey,Result);
+			FunctionStack.Last().ExpressionCodeMap.Add(ExpressionKey,Result);
 
 			return Result;
 		}
@@ -2849,7 +2892,7 @@ protected:
 		}
 
 		EMaterialValueType ResultType = GetArithmeticResultType(X,Y);
-		EMaterialValueType AlphaType = ResultType == CodeChunks[MaterialProperty][ShaderFrequency][A].Type ? ResultType : MCT_Float1;
+		EMaterialValueType AlphaType = ResultType == (*CurrentScopeChunks)[A].Type ? ResultType : MCT_Float1;
 		return AddCodeChunk(ResultType,TEXT("lerp(%s,%s,%s)"),*CoerceParameter(X,ResultType),*CoerceParameter(Y,ResultType),*CoerceParameter(A,AlphaType));
 	}
 
@@ -3606,6 +3649,51 @@ protected:
 		return ResultIdx;
 	}
 
+	virtual int32 CustomOutput(class UMaterialExpressionCustomOutput* Custom, int32 OutputIndex, int32 OutputCode) override
+	{
+		if (MaterialProperty != MP_MAX)
+		{
+			return Errorf(TEXT("A Custom Output node should not be attached to the %s material property"), *GetNameOfMaterialProperty(MaterialProperty));
+		}
+
+		if (OutputCode == INDEX_NONE)
+		{
+			return INDEX_NONE;
+		}
+
+		EMaterialValueType OutputType = GetParameterType(OutputCode);
+		FString OutputTypeString;
+		switch (OutputType)
+		{
+			case MCT_Float:
+			case MCT_Float1:
+				OutputTypeString = TEXT("MaterialFloat");
+				break;
+			case MCT_Float2:
+				OutputTypeString += TEXT("MaterialFloat2");
+				break;
+			case MCT_Float3:
+				OutputTypeString += TEXT("MaterialFloat3");
+				break;
+			case MCT_Float4:
+				OutputTypeString += TEXT("MaterialFloat4");
+				break;
+			default:
+				return Errorf(TEXT("Bad type %s for %s"), DescribeType(GetParameterType(OutputCode)), *Custom->GetDescription());
+				break;
+		}
+
+		FString Definitions;
+		FString Body;
+		GetFixedParameterCode(OutputCode, *CurrentScopeChunks, Definitions, Body);
+
+		FString ImplementationCode = FString::Printf(TEXT("%s %s%d(FMaterial%sParameters Parameters)\r\n{\r\n%s return %s;\r\n}\r\n"), *OutputTypeString, *Custom->GetFunctionName(), OutputIndex, ShaderFrequency == SF_Vertex ? TEXT("Vertex") : TEXT("Pixel"), *Definitions, *Body);
+		CustomOutputImplementations.Add(ImplementationCode);
+
+		// return value is not used
+		return INDEX_NONE;
+	}
+
 	/**
 	 * Adds code to return a random value shared by all geometry for any given instanced static mesh
 	 *
diff --git a/Engine/Source/Runtime/Engine/Private/Materials/Material.cpp b/Engine/Source/Runtime/Engine/Private/Materials/Material.cpp
--- a/Engine/Source/Runtime/Engine/Private/Materials/Material.cpp
+++ b/Engine/Source/Runtime/Engine/Private/Materials/Material.cpp
@@ -8,6 +8,7 @@
 #include "Materials/MaterialFunction.h"
 #include "Materials/MaterialExpressionCollectionParameter.h"
 #include "Materials/MaterialExpressionComment.h"
+#include "Materials/MaterialExpressionCustomOutput.h"
 #include "Materials/MaterialExpressionDynamicParameter.h"
 #include "Materials/MaterialExpressionFontSampleParameter.h"
 #include "Materials/MaterialExpressionMaterialFunctionCall.h"
@@ -117,6 +118,18 @@ int32 FMaterialResource::CompilePropertyAndSetMaterialProperty(EMaterialProperty
 	return Compiler->ForceCast(Ret, GetMaterialPropertyType(Property));
 }
 
+void FMaterialResource::GatherCustomOutputExpressions(TArray<UMaterialExpressionCustomOutput*>& OutCustomOutputs) const
+{
+	for (UMaterialExpression* Expression : Material->Expressions)
+	{
+		UMaterialExpressionCustomOutput* CustomOutput = Cast<UMaterialExpressionCustomOutput>(Expression);
+		if (CustomOutput)
+		{
+			OutCustomOutputs.Add(CustomOutput);
+		}
+	}
+}
+
 void FMaterialResource::GetShaderMapId(EShaderPlatform Platform, FMaterialShaderMapId& OutId) const
 {
 	FMaterial::GetShaderMapId(Platform, OutId);
diff --git a/Engine/Source/Runtime/Engine/Private/Materials/MaterialExpressions.cpp b/Engine/Source/Runtime/Engine/Private/Materials/MaterialExpressions.cpp
--- a/Engine/Source/Runtime/Engine/Private/Materials/MaterialExpressions.cpp
+++ b/Engine/Source/Runtime/Engine/Private/Materials/MaterialExpressions.cpp
@@ -40,6 +40,7 @@
 #include "Materials/MaterialExpressionDivide.h"
 #include "Materials/MaterialExpressionDotProduct.h"
 #include "Materials/MaterialExpressionDynamicParameter.h"
+#include "Materials/MaterialExpressionCustomOutput.h"
 #include "Materials/MaterialExpressionEyeAdaptation.h"
 #include "Materials/MaterialExpressionFeatureLevelSwitch.h"
 #include "Materials/MaterialExpressionFloor.h"
@@ -8900,6 +8901,13 @@ bool UMaterialExpressionSpeedTree::CanEditChange(const UProperty* InProperty) co
 
 #endif
 
+///////////////////////////////////////////////////////////////////////////////
+// UMaterialExpressionCustomOutput
+///////////////////////////////////////////////////////////////////////////////
+UMaterialExpressionCustomOutput::UMaterialExpressionCustomOutput(const FObjectInitializer& ObjectInitializer)
+: Super(ObjectInitializer)
+{
+}
 
 ///////////////////////////////////////////////////////////////////////////////
 // UMaterialExpressionEyeAdaptation
diff --git a/Engine/Source/Runtime/Engine/Private/Materials/MaterialShared.cpp b/Engine/Source/Runtime/Engine/Private/Materials/MaterialShared.cpp
--- a/Engine/Source/Runtime/Engine/Private/Materials/MaterialShared.cpp
+++ b/Engine/Source/Runtime/Engine/Private/Materials/MaterialShared.cpp
@@ -1867,7 +1867,7 @@ void FMaterial::GetReferencedTexturesHash(FSHAHash& OutHash) const
  * @param OutHighlightMap - source code highlight list
  * @return - true on Success
  */
-bool FMaterial::GetMaterialExpressionSource( FString& OutSource, TMap<FMaterialExpressionKey,int32> OutExpressionCodeMap[][SF_NumFrequencies] )
+bool FMaterial::GetMaterialExpressionSource( FString& OutSource )
 {
 #if WITH_EDITORONLY_DATA
 	class FViewSourceMaterialTranslator : public FHLSLMaterialTranslator
@@ -1876,14 +1876,6 @@ bool FMaterial::GetMaterialExpressionSource( FString& OutSource, TMap<FMaterialE
 		FViewSourceMaterialTranslator(FMaterial* InMaterial,FMaterialCompilationOutput& InMaterialCompilationOutput,const FStaticParameterSet& StaticParameters,EShaderPlatform InPlatform,EMaterialQualityLevel::Type InQualityLevel,ERHIFeatureLevel::Type InFeatureLevel)
 		:	FHLSLMaterialTranslator(InMaterial,InMaterialCompilationOutput,StaticParameters,InPlatform,InQualityLevel,InFeatureLevel)
 		{}
-
-		void GetExpressionCodeMap(TMap<FMaterialExpressionKey,int32> OutExpressionCodeMap[][SF_NumFrequencies])
-		{
-			for (int32 PropertyIndex = 0; PropertyIndex < MP_MAX; PropertyIndex++)
-			{
-				OutExpressionCodeMap[PropertyIndex][GetMaterialPropertyShaderFrequency((EMaterialProperty)PropertyIndex)] = FunctionStack.Last().ExpressionCodeMap[PropertyIndex][GetMaterialPropertyShaderFrequency((EMaterialProperty)PropertyIndex)];
-			}
-		}
 	};
 
 	FMaterialCompilationOutput TempOutput;
@@ -1894,9 +1886,6 @@ bool FMaterial::GetMaterialExpressionSource( FString& OutSource, TMap<FMaterialE
 	{
 		// Generate the HLSL
 		OutSource = MaterialTranslator.GetMaterialShaderCode();
-
-		// Save the Expression Code map.
-		MaterialTranslator.GetExpressionCodeMap(OutExpressionCodeMap);
 	}
 	return bSuccess;
 #else
diff --git a/Engine/Source/Runtime/Engine/Public/MaterialCompiler.h b/Engine/Source/Runtime/Engine/Public/MaterialCompiler.h
--- a/Engine/Source/Runtime/Engine/Public/MaterialCompiler.h
+++ b/Engine/Source/Runtime/Engine/Public/MaterialCompiler.h
@@ -159,7 +159,8 @@ public:
 	virtual int32 VertexNormal() = 0;
 	virtual int32 PixelNormalWS() = 0;
 
-	virtual int32 CustomExpression( class UMaterialExpressionCustom* Custom, TArray<int32>& CompiledInputs ) = 0;
+	virtual int32 CustomExpression(class UMaterialExpressionCustom* Custom, TArray<int32>& CompiledInputs) = 0;
+	virtual int32 CustomOutput(class UMaterialExpressionCustomOutput* Custom, int32 OutputIndex, int32 OutputCode) = 0;
 
 	virtual int32 DDX(int32 X) = 0;
 	virtual int32 DDY(int32 X) = 0;
@@ -313,7 +314,8 @@ public:
 	virtual int32 VertexNormal() override { return Compiler->VertexNormal(); }
 	virtual int32 PixelNormalWS() override { return Compiler->PixelNormalWS(); }
 
-	virtual int32 CustomExpression( class UMaterialExpressionCustom* Custom, TArray<int32>& CompiledInputs ) override { return Compiler->CustomExpression(Custom,CompiledInputs); }
+	virtual int32 CustomExpression(class UMaterialExpressionCustom* Custom, TArray<int32>& CompiledInputs) override { return Compiler->CustomExpression(Custom,CompiledInputs); }
+	virtual int32 CustomOutput(class UMaterialExpressionCustomOutput* Custom, int32 OutputIndex, int32 OutputCode) override{ return Compiler->CustomOutput(Custom, OutputIndex, OutputCode); }
 	virtual int32 DDX(int32 X) override { return Compiler->DDX(X); }
 	virtual int32 DDY(int32 X) override { return Compiler->DDY(X); }
 
diff --git a/Engine/Source/Runtime/Engine/Public/MaterialShared.h b/Engine/Source/Runtime/Engine/Public/MaterialShared.h
--- a/Engine/Source/Runtime/Engine/Public/MaterialShared.h
+++ b/Engine/Source/Runtime/Engine/Public/MaterialShared.h
@@ -803,7 +803,7 @@ public:
 	TArray<FMaterialExpressionKey> ExpressionStack;
 
 	/** A map from material expression to the index into CodeChunks of the code for the material expression. */
-	TMap<FMaterialExpressionKey,int32> ExpressionCodeMap[MP_MAX][SF_NumFrequencies];
+	TMap<FMaterialExpressionKey,int32> ExpressionCodeMap;
 
 	explicit FMaterialFunctionCompileState(UMaterialExpressionMaterialFunctionCall* InFunctionCall) :
 		FunctionCall(InFunctionCall)
@@ -1068,7 +1068,7 @@ public:
 	* @param OutHighlightMap - source code highlight list
 	* @return - true on Success
 	*/
-	ENGINE_API bool GetMaterialExpressionSource( FString& OutSource, TMap<FMaterialExpressionKey,int32> OutExpressionCodeMap[][SF_NumFrequencies] );
+	ENGINE_API bool GetMaterialExpressionSource(FString& OutSource);
 
 	/** 
 	 * Adds an FMaterial to the global list.
@@ -1101,6 +1101,9 @@ protected:
 	 */
 	virtual int32 CompilePropertyAndSetMaterialProperty(EMaterialProperty Property, class FMaterialCompiler* Compiler, EShaderFrequency OverrideShaderFrequency = SF_NumFrequencies, bool bUsePreviousFrameTime = false) const = 0;
 
+	/* Gather any UMaterialExpressionCustomOutput expressions they can be compiled in turn */
+	virtual void GatherCustomOutputExpressions(TArray<class UMaterialExpressionCustomOutput*>& OutCustomOutputs) const {}
+
 	/** Returns the index to the Expression in the Expressions array, or -1 if not found. */
 	int32 FindExpression(const TArray<TRefCountPtr<FMaterialUniformExpressionTexture> >&Expressions, const FMaterialUniformExpressionTexture &Expression);
 
@@ -1549,6 +1552,8 @@ protected:
 
 	/** Entry point for compiling a specific material property.  This must call SetMaterialProperty. */
 	ENGINE_API virtual int32 CompilePropertyAndSetMaterialProperty(EMaterialProperty Property, class FMaterialCompiler* Compiler, EShaderFrequency OverrideShaderFrequency, bool bUsePreviousFrameTime) const override;
+	/* Gives the material a chance to compile any custom output nodes it has added */
+	ENGINE_API virtual void GatherCustomOutputExpressions(TArray<class UMaterialExpressionCustomOutput*>& OutCustomOutputs) const override;
 	ENGINE_API virtual bool HasVertexPositionOffsetConnected() const;
 	ENGINE_API virtual bool HasMaterialAttributesConnected() const;
 	/** Useful for debugging. */
diff --git a/Engine/Source/Runtime/Landscape/Classes/LandscapeGrassType.h b/Engine/Source/Runtime/Landscape/Classes/LandscapeGrassType.h
new file mode 100644
--- /dev/null
+++ b/Engine/Source/Runtime/Landscape/Classes/LandscapeGrassType.h
@@ -0,0 +1,47 @@
+// Copyright 1998-2015 Epic Games, Inc. All Rights Reserved.
+
+#pragma once
+
+#include "LandscapeGrassType.generated.h"
+
+UCLASS(MinimalAPI)
+class ULandscapeGrassType : public UObject
+{
+	GENERATED_UCLASS_BODY()
+
+	UPROPERTY(EditAnywhere, Category=Grass)
+	class UStaticMesh* GrassMesh;
+
+	/* Instances per 10 square meters. */
+	UPROPERTY(EditAnywhere, Category=Grass)
+	float GrassDensity;
+
+	UPROPERTY(EditAnywhere, Category=Grass)
+	float PlacementJitter;
+
+	/* The distance where instances will begin to fade out if using a PerInstanceFadeAmount material node. 0 disables. */
+	UPROPERTY(EditAnywhere, Category=Grass)
+	int32 StartCullDistance;
+
+	/**
+	 * The distance where instances will have completely faded out when using a PerInstanceFadeAmount material node. 0 disables. 
+	 * When the entire cluster is beyond this distance, the cluster is completely culled and not rendered at all.
+	 */
+	UPROPERTY(EditAnywhere, Category = Grass)
+	int32 EndCullDistance;
+
+	UPROPERTY(EditAnywhere, Category = Grass)
+	bool RandomRotation;
+
+	UPROPERTY(EditAnywhere, Category = Grass)
+	bool AlignToSurface;
+
+#if WITH_EDITOR
+	// Begin UObject Interface
+	virtual void PostEditChangeProperty(FPropertyChangedEvent& PropertyChangedEvent) override;
+	// End UObject Interface
+#endif
+};
+
+
+
diff --git a/Engine/Source/Runtime/Landscape/Classes/Materials/MaterialExpressionLandscapeGrassOutput.h b/Engine/Source/Runtime/Landscape/Classes/Materials/MaterialExpressionLandscapeGrassOutput.h
new file mode 100644
--- /dev/null
+++ b/Engine/Source/Runtime/Landscape/Classes/Materials/MaterialExpressionLandscapeGrassOutput.h
@@ -0,0 +1,60 @@
+// Copyright 1998-2015 Epic Games, Inc. All Rights Reserved.
+
+
+#pragma once
+#include "LandscapeGrassType.h"
+
+#include "Materials/MaterialExpressionCustomOutput.h"
+#include "MaterialExpressionLandscapeGrassOutput.generated.h"
+
+USTRUCT()
+struct FGrassInput
+{
+	GENERATED_USTRUCT_BODY()
+
+	UPROPERTY(EditAnywhere, Category = Grass)
+	FName Name;
+
+	UPROPERTY(EditAnywhere, Category = Grass)
+	ULandscapeGrassType* GrassType;
+
+	UPROPERTY(meta = (RequiredInput = "true"))
+	FExpressionInput Input;
+
+	FGrassInput()
+	: GrassType(nullptr)
+	{}
+	
+	FGrassInput(FName InName)
+	: Name(InName)
+	, GrassType(nullptr)
+	{}
+};
+
+UCLASS(collapsecategories, hidecategories=Object, MinimalAPI)
+class UMaterialExpressionLandscapeGrassOutput : public UMaterialExpressionCustomOutput
+{
+	GENERATED_UCLASS_BODY()
+	virtual int32 Compile(class FMaterialCompiler* Compiler, int32 OutputIndex, int32 MultiplexIndex) override;
+	virtual void GetCaption(TArray<FString>& OutCaptions) const override;
+	virtual const TArray<FExpressionInput*> GetInputs() override;
+	virtual FExpressionInput* GetInput(int32 InputIndex) override;
+	virtual FString GetInputName(int32 InputIndex) const override;
+
+	// Begin UObject Interface
+#if WITH_EDITOR
+	virtual void PostEditChangeProperty(FPropertyChangedEvent& PropertyChangedEvent) override;
+	// End UObject Interface
+
+	virtual uint32 GetInputType(int32 InputIndex) override { return MCT_Float; }
+#endif
+
+	virtual int32 GetNumOutputs() const override { return GrassTypes.Num(); }
+	virtual FString GetFunctionName() const { return TEXT("GetGrassWeight"); }
+
+	UPROPERTY(EditAnywhere, Category = UMaterialExpressionLandscapeGrassOutput)
+	TArray<FGrassInput> GrassTypes;
+};
+
+
+
diff --git a/Engine/Source/Runtime/Landscape/Private/LandscapeGrass.cpp b/Engine/Source/Runtime/Landscape/Private/LandscapeGrass.cpp
new file mode 100644
--- /dev/null
+++ b/Engine/Source/Runtime/Landscape/Private/LandscapeGrass.cpp
@@ -0,0 +1,149 @@
+// Copyright 1998-2015 Epic Games, Inc. All Rights Reserved.
+
+#include "Landscape.h"
+#include "MaterialCompiler.h"
+#include "Materials/MaterialExpressionLandscapeGrassOutput.h"
+#include "EdGraph/EdGraphNode.h"
+
+#define LOCTEXT_NAMESPACE "Landscape"
+
+//
+// UMaterialExpressionLandscapeGrassOutput
+//
+UMaterialExpressionLandscapeGrassOutput::UMaterialExpressionLandscapeGrassOutput(const FObjectInitializer& ObjectInitializer)
+: Super(ObjectInitializer)
+{
+	// Structure to hold one-time initialization
+	struct FConstructorStatics
+	{
+		FString STRING_Landscape;
+		FName NAME_Grass;
+		FConstructorStatics()
+			: STRING_Landscape(LOCTEXT("Landscape", "Landscape").ToString())
+			, NAME_Grass("Grass")
+		{
+		}
+	};
+	static FConstructorStatics ConstructorStatics;
+
+	MenuCategories.Add(ConstructorStatics.STRING_Landscape);
+
+	// No outputs
+	Outputs.Reset();
+
+	// Default input
+	new(GrassTypes)FGrassInput(ConstructorStatics.NAME_Grass);
+}
+
+int32 UMaterialExpressionLandscapeGrassOutput::Compile(class FMaterialCompiler* Compiler, int32 OutputIndex, int32 MultiplexIndex)
+{
+	if (GrassTypes.IsValidIndex(OutputIndex))
+	{
+		if (GrassTypes[OutputIndex].Input.Expression)
+		{
+			return Compiler->CustomOutput(this, OutputIndex, GrassTypes[OutputIndex].Input.Compile(Compiler, MultiplexIndex));
+		}
+		else
+		{
+			return CompilerError(Compiler, TEXT("Input missing"));
+		}
+	}
+
+	return INDEX_NONE;
+}
+
+void UMaterialExpressionLandscapeGrassOutput::GetCaption(TArray<FString>& OutCaptions) const
+{
+	OutCaptions.Add(TEXT("Grass"));
+}
+
+const TArray<FExpressionInput*> UMaterialExpressionLandscapeGrassOutput::GetInputs()
+{
+	TArray<FExpressionInput*> OutInputs;
+	for (auto& GrassType : GrassTypes)
+	{
+		OutInputs.Add(&GrassType.Input);
+	}
+	return OutInputs;
+}
+
+FExpressionInput* UMaterialExpressionLandscapeGrassOutput::GetInput(int32 InputIndex)
+{
+	return &GrassTypes[InputIndex].Input;
+}
+
+FString UMaterialExpressionLandscapeGrassOutput::GetInputName(int32 InputIndex) const
+{
+	return GrassTypes[InputIndex].Name.ToString();
+}
+
+#if WITH_EDITOR
+void UMaterialExpressionLandscapeGrassOutput::PostEditChangeProperty(FPropertyChangedEvent& PropertyChangedEvent)
+{
+	Super::PostEditChangeProperty(PropertyChangedEvent);
+
+	if (PropertyChangedEvent.MemberProperty)
+	{
+		const FName PropertyName = PropertyChangedEvent.MemberProperty->GetFName();
+		if (PropertyName == GET_MEMBER_NAME_CHECKED(UMaterialExpressionLandscapeGrassOutput, GrassTypes))
+		{
+			if (GraphNode)
+			{
+				GraphNode->ReconstructNode();
+			}
+		}
+	}
+}
+#endif
+
+//
+// ULandscapeGrassType
+//
+ULandscapeGrassType::ULandscapeGrassType(const FObjectInitializer& ObjectInitializer)
+: Super(ObjectInitializer)
+{
+	GrassDensity = 400;
+	StartCullDistance = 10000.0f;
+	EndCullDistance = 10000.0f;
+	PlacementJitter = 1.0f;
+	RandomRotation = true;
+	AlignToSurface = true;
+}
+
+#if WITH_EDITOR
+void ULandscapeGrassType::PostEditChangeProperty(FPropertyChangedEvent& PropertyChangedEvent)
+{
+	Super::PostEditChangeProperty(PropertyChangedEvent);
+
+	if (GIsEditor)
+	{
+		// Only care current world object
+		for (TActorIterator<ALandscapeProxy> It(GWorld); It; ++It)
+		{
+			ALandscapeProxy* Proxy = *It;
+			const UMaterialInterface* MaterialInterface = Proxy->LandscapeMaterial;
+			if (MaterialInterface)
+			{
+				TArray<const UMaterialExpressionLandscapeGrassOutput*> GrassExpressions;
+				MaterialInterface->GetMaterial()->GetAllExpressionsOfType<UMaterialExpressionLandscapeGrassOutput>(GrassExpressions);
+
+				// Should only be one grass type node
+				if (GrassExpressions.Num() > 0)
+				{
+					for (auto& Output : GrassExpressions[0]->GrassTypes)
+					{
+						if (Output.GrassType == this)
+						{
+							Proxy->FlushFoliageComponents();
+							break;
+						}
+					}
+				}
+			}
+		}
+	}
+}
+#endif
+
+
+#undef LOCTEXT_NAMESPACE
