NULL ptr check to fix crash

diff --git a/Engine/Source/Runtime/Engine/Private/VoiceChannel.cpp b/Engine/Source/Runtime/Engine/Private/VoiceChannel.cpp
--- a/Engine/Source/Runtime/Engine/Private/VoiceChannel.cpp
+++ b/Engine/Source/Runtime/Engine/Private/VoiceChannel.cpp
@@ -144,7 +144,7 @@ void UVoiceChannel::AddVoicePacket(TSharedPtr<FVoicePacket> VoicePacket)
 		VoicePackets.Add(VoicePacket);
 
 		UE_LOG(LogNet, VeryVerbose, TEXT("AddVoicePacket: %s [%s] to=%s from=%s"),
-			*Connection->PlayerId->ToDebugString(),
+			Connection->PlayerId.IsValid() ? *Connection->PlayerId->ToDebugString() : TEXT("NULL"),	// Currently, the server's PlayerId is NULL, so we need to check for this
 			*Connection->Driver->GetDescription(),
 			*Connection->LowLevelDescribe(),
 			*VoicePacket->GetSender()->ToDebugString());
