UE4 - fix reattach on grass components

diff --git a/Engine/Source/Runtime/Engine/Classes/Components/InstancedStaticMeshComponent.h b/Engine/Source/Runtime/Engine/Classes/Components/InstancedStaticMeshComponent.h
--- a/Engine/Source/Runtime/Engine/Classes/Components/InstancedStaticMeshComponent.h
+++ b/Engine/Source/Runtime/Engine/Classes/Components/InstancedStaticMeshComponent.h
@@ -133,6 +133,8 @@ public:
 	 *  Released on the rendering thread
 	 */
 	TSharedPtr<FPerInstanceRenderData, ESPMode::ThreadSafe> PerInstanceRenderData;
+	/** This was prebuilt for grass, never destroy it. */
+	bool bPerInstanceRenderDataWasPrebuilt;
 		
 #if WITH_EDITOR
 	/** One bit per instance if the instance is selected. */
diff --git a/Engine/Source/Runtime/Engine/Private/HierarchicalInstancedStaticMesh.cpp b/Engine/Source/Runtime/Engine/Private/HierarchicalInstancedStaticMesh.cpp
--- a/Engine/Source/Runtime/Engine/Private/HierarchicalInstancedStaticMesh.cpp
+++ b/Engine/Source/Runtime/Engine/Private/HierarchicalInstancedStaticMesh.cpp
@@ -1829,7 +1829,7 @@ FPrimitiveSceneProxy* UHierarchicalInstancedStaticMeshComponent::CreateSceneProx
 	// Verify that the mesh is valid before using it.
 	const bool bMeshIsValid = 
 		// make sure we have instances
-		(PerInstanceSMData.Num() > 0 || WriteOncePrebuiltInstanceBuffer.Num() > 0)&&
+		(PerInstanceSMData.Num() > 0 || WriteOncePrebuiltInstanceBuffer.Num() > 0 || bPerInstanceRenderDataWasPrebuilt) &&
 		// make sure we have an actual staticmesh
 		StaticMesh &&
 		StaticMesh->HasValidRenderData() &&
@@ -1847,13 +1847,11 @@ FPrimitiveSceneProxy* UHierarchicalInstancedStaticMeshComponent::CreateSceneProx
 			InstancingRandomSeed = FMath::Rand();
 		}
 
-		if (WriteOncePrebuiltInstanceBuffer.Num() > 0)
+		if (WriteOncePrebuiltInstanceBuffer.Num() > 0 || bPerInstanceRenderDataWasPrebuilt)
 		{
 			ProxySize = FStaticMeshInstanceData::GetResourceSize(WriteOncePrebuiltInstanceBuffer.Num());
 			INC_DWORD_STAT_BY(STAT_FoliageInstanceBuffers, ProxySize);
 
-			//@todo - this is sketchy, what you really want to do is flush the grass cache instead of doing a render thread update
-			bNeverNeedsRenderUpdate = true; // we can never update it because the data is gone
 			return ::new FHierarchicalStaticMeshSceneProxy(this, GetWorld()->FeatureLevel, WriteOncePrebuiltInstanceBuffer);
 		}
 		ProxySize = FStaticMeshInstanceData::GetResourceSize(PerInstanceSMData.Num());
diff --git a/Engine/Source/Runtime/Engine/Private/InstancedStaticMesh.cpp b/Engine/Source/Runtime/Engine/Private/InstancedStaticMesh.cpp
--- a/Engine/Source/Runtime/Engine/Private/InstancedStaticMesh.cpp
+++ b/Engine/Source/Runtime/Engine/Private/InstancedStaticMesh.cpp
@@ -665,6 +665,7 @@ void FInstancedStaticMeshSceneProxy::GetDistancefieldAtlasData(FBox& LocalVolume
 
 UInstancedStaticMeshComponent::UInstancedStaticMeshComponent(const FObjectInitializer& ObjectInitializer)
 	: Super(ObjectInitializer)
+	, bPerInstanceRenderDataWasPrebuilt(false)
 {
 	Mobility = EComponentMobility::Movable;
 	BodyInstance.bSimulatePhysics = false;
@@ -1138,7 +1139,7 @@ void UInstancedStaticMeshComponent::ApplyLightMapping(FStaticLightingTextureMapp
 
 void UInstancedStaticMeshComponent::ReleasePerInstanceRenderData()
 {
-	if (PerInstanceRenderData.IsValid())
+	if (PerInstanceRenderData.IsValid() && !bPerInstanceRenderDataWasPrebuilt)
 	{
 		if (SceneProxy == nullptr)
 		{
@@ -1320,7 +1321,7 @@ void UInstancedStaticMeshComponent::ClearInstances()
 	ClearAllInstanceBodies();
 
 	// Indicate we need to update render state to reflect changes
-	bNeverNeedsRenderUpdate = false;
+	bPerInstanceRenderDataWasPrebuilt = false;
 	ReleasePerInstanceRenderData();
 	MarkRenderStateDirty();
 
@@ -1468,6 +1469,7 @@ SIZE_T UInstancedStaticMeshComponent::GetResourceSize( EResourceSizeMode::Type M
 
 void UInstancedStaticMeshComponent::BeginDestroy()
 {
+	bPerInstanceRenderDataWasPrebuilt = false;
 	Super::BeginDestroy();
 	ReleasePerInstanceRenderData();
 }
diff --git a/Engine/Source/Runtime/Engine/Private/InstancedStaticMesh.h b/Engine/Source/Runtime/Engine/Private/InstancedStaticMesh.h
--- a/Engine/Source/Runtime/Engine/Private/InstancedStaticMesh.h
+++ b/Engine/Source/Runtime/Engine/Private/InstancedStaticMesh.h
@@ -352,8 +352,9 @@ public:
 		if (!PerInstanceRenderData.IsValid())
 		{
 			// initialize the instance buffer from the component's instances
-			PerInstanceRenderData = InComponent->PerInstanceRenderData = MakeShareable(new FPerInstanceRenderData(InComponent, InFeatureLevel));
-			InComponent->PerInstanceRenderData = PerInstanceRenderData;
+			InComponent->PerInstanceRenderData = MakeShareable(new FPerInstanceRenderData(InComponent, InFeatureLevel));
+			PerInstanceRenderData = InComponent->PerInstanceRenderData;
+			InComponent->bPerInstanceRenderDataWasPrebuilt = false;
 		}
 
 		NumInstances = PerInstanceRenderData->InstanceBuffer.GetNumInstances();
@@ -362,12 +363,18 @@ public:
 
 	FInstancedStaticMeshRenderData(UInstancedStaticMeshComponent* InComponent, ERHIFeatureLevel::Type InFeatureLevel, FStaticMeshInstanceData& Other)
 		: Component(InComponent)
+		, PerInstanceRenderData(InComponent->PerInstanceRenderData)
 		, LODModels(Component->StaticMesh->RenderData->LODResources)
 		, FeatureLevel(InFeatureLevel)
 	{
 		InitVertexFactories();
-		// initialize the instance buffer from the component's instances
-		PerInstanceRenderData = MakeShareable(new FPerInstanceRenderData(InComponent, Other, InFeatureLevel));
+		// initialize the instance buffer from the prebuilt instances
+		if (!PerInstanceRenderData.IsValid())
+		{
+			InComponent->PerInstanceRenderData = MakeShareable(new FPerInstanceRenderData(InComponent, Other, InFeatureLevel));
+			PerInstanceRenderData = InComponent->PerInstanceRenderData;
+			InComponent->bPerInstanceRenderDataWasPrebuilt = true;
+		}
 		NumInstances = PerInstanceRenderData->InstanceBuffer.GetNumInstances();
 		InitResources();
 	}
