Add an internal network protocol version, to catch incompatible changes in things like the bunch header format
This value is hashed into the result of FNetworkVersion::GetLocalNetworkVersion()
Update this whenever a change is made to the protocol format

diff --git a/Engine/Source/Runtime/Engine/Private/UnrealNetwork.cpp b/Engine/Source/Runtime/Engine/Private/UnrealNetwork.cpp
--- a/Engine/Source/Runtime/Engine/Private/UnrealNetwork.cpp
+++ b/Engine/Source/Runtime/Engine/Private/UnrealNetwork.cpp
@@ -6,6 +6,7 @@
 
 FNetworkVersion::FGetLocalNetworkVersionOverride FNetworkVersion::GetLocalNetworkVersionOverride;
 FNetworkVersion::FIsNetworkCompatibleOverride FNetworkVersion::IsNetworkCompatibleOverride;
+const uint32 FNetworkVersion::InternalProtocolVersion = 1;
 
 uint32 FNetworkVersion::GetLocalNetworkVersion()
 {
@@ -25,9 +26,12 @@ uint32 FNetworkVersion::GetLocalNetworkVersion()
 	const FString& ProjectVersion = Cast<UGeneralProjectSettings>(UGeneralProjectSettings::StaticClass()->GetDefaultObject())->ProjectVersion;
 
 	// Start with engine version as seed, and then hash with project name + project version
-	const uint32 LocalNetworkVersion = FCrc::StrCrc32( *ProjectVersion, FCrc::StrCrc32( *ProjectName, GEngineNetVersion ) );
+	const uint32 VersionHash = FCrc::StrCrc32( *ProjectVersion, FCrc::StrCrc32( *ProjectName, GEngineNetVersion ) );
 
-	UE_LOG( LogNet, Log, TEXT( "GetLocalNetworkVersion: GEngineNetVersion: %i, ProjectName: %s, ProjectVersion: %s, LocalNetworkVersion: %i" ), GEngineNetVersion, *ProjectName, *ProjectVersion, LocalNetworkVersion );
+	// Hash with internal protocol version
+	const uint32 LocalNetworkVersion = FCrc::MemCrc32( &InternalProtocolVersion, sizeof( InternalProtocolVersion ), VersionHash );
+
+	UE_LOG( LogNet, Log, TEXT( "GetLocalNetworkVersion: GEngineNetVersion: %i, ProjectName: %s, ProjectVersion: %s, InternalProtocolVersion: %i, LocalNetworkVersion: %i" ), GEngineNetVersion, *ProjectName, *ProjectVersion, InternalProtocolVersion, LocalNetworkVersion );
 
 	return LocalNetworkVersion;
 }
diff --git a/Engine/Source/Runtime/Engine/Public/Net/UnrealNetwork.h b/Engine/Source/Runtime/Engine/Public/Net/UnrealNetwork.h
--- a/Engine/Source/Runtime/Engine/Public/Net/UnrealNetwork.h
+++ b/Engine/Source/Runtime/Engine/Public/Net/UnrealNetwork.h
@@ -63,6 +63,15 @@ struct ENGINE_API FNetworkVersion
 	 * @return true if the two instances can communicate, false otherwise
 	 */
 	static bool IsNetworkCompatible( const uint32 LocalNetworkVersion, const uint32 RemoteNetworkVersion );
+
+	/**
+	 * Internal network protocol version.
+	 * By default, this value is incorporated into the result of GetLocalNetworkVersion().
+	 * This value should be incremented manually along with any changes to the
+	 * internal network protocol, such as the bunch header format or
+	 * the RPC header format.
+	 */
+	static const uint32 InternalProtocolVersion;
 };
 
 /*-----------------------------------------------------------------------------
