Pack FNetworkGUIDs when serializing them to save bandwidth

diff --git a/Engine/Source/Runtime/Core/Public/Misc/NetworkGuid.h b/Engine/Source/Runtime/Core/Public/Misc/NetworkGuid.h
--- a/Engine/Source/Runtime/Core/Public/Misc/NetworkGuid.h
+++ b/Engine/Source/Runtime/Core/Public/Misc/NetworkGuid.h
@@ -39,7 +39,8 @@ public:
 	
 	friend FArchive& operator<<( FArchive& Ar, FNetworkGUID& G )
 	{
-		return Ar << G.Value;
+		Ar.SerializeIntPacked(G.Value);
+		return Ar;
 	}
 
 public:
diff --git a/Engine/Source/Runtime/Engine/Private/UnrealNetwork.cpp b/Engine/Source/Runtime/Engine/Private/UnrealNetwork.cpp
--- a/Engine/Source/Runtime/Engine/Private/UnrealNetwork.cpp
+++ b/Engine/Source/Runtime/Engine/Private/UnrealNetwork.cpp
@@ -6,7 +6,7 @@
 
 FNetworkVersion::FGetLocalNetworkVersionOverride FNetworkVersion::GetLocalNetworkVersionOverride;
 FNetworkVersion::FIsNetworkCompatibleOverride FNetworkVersion::IsNetworkCompatibleOverride;
-const uint32 FNetworkVersion::InternalProtocolVersion = 1;
+const uint32 FNetworkVersion::InternalProtocolVersion = 2;
 
 uint32 FNetworkVersion::GetLocalNetworkVersion()
 {
