Comment on infinite loop posibility

diff --git a/Engine/Source/Runtime/Engine/Private/DataChannel.cpp b/Engine/Source/Runtime/Engine/Private/DataChannel.cpp
--- a/Engine/Source/Runtime/Engine/Private/DataChannel.cpp
+++ b/Engine/Source/Runtime/Engine/Private/DataChannel.cpp
@@ -1589,7 +1589,7 @@ void UActorChannel::SetChannelActor( AActor* InActor )
 			// UE_LOG(LogNetTraffic, Log, TEXT("%i SYNCHRONIZING by sending %i"), ChIndex, Connection->PendingOutRec[ChIndex]);
 
 			FOutBunch Bunch( this, 0 );
-			if( !Bunch.IsError() )
+			if( !Bunch.IsError() )	// FIXME: This will be an infinite loop if this happens!!!
 			{
 				Bunch.bReliable = true;
 				SendBunch( &Bunch, 0 );
