nicer UI

diff --git a/Engine/Source/Runtime/Engine/Classes/Components/DirectionalLightComponent.h b/Engine/Source/Runtime/Engine/Classes/Components/DirectionalLightComponent.h
--- a/Engine/Source/Runtime/Engine/Classes/Components/DirectionalLightComponent.h
+++ b/Engine/Source/Runtime/Engine/Classes/Components/DirectionalLightComponent.h
@@ -64,7 +64,7 @@ class UDirectionalLightComponent : public ULightComponent
 	 * Controls whether the cascades are distributed closer to the camera (larger exponent) or further from the camera (smaller exponent).
 	 * An exponent of 1 means that cascade transitions will happen at a distance proportional to their resolution.
 	 */
-	UPROPERTY(EditAnywhere, BlueprintReadOnly, Category=CascadedShadowMaps, meta=(UIMin = "1", UIMax = "4"))
+	UPROPERTY(EditAnywhere, BlueprintReadOnly, Category=CascadedShadowMaps, meta=(UIMin = "1", UIMax = "4", DisplayName = "Distribution Exponent"))
 	float CascadeDistributionExponent;
 
 	/** 
@@ -76,14 +76,14 @@ class UDirectionalLightComponent : public ULightComponent
 	 * Ideal values are the smallest possible which still hide the transition.
 	 * An increased fade region size causes an increase in shadow rendering cost.
 	 */
-	UPROPERTY(EditAnywhere, BlueprintReadOnly, Category=CascadedShadowMaps, meta=(UIMin = "0", UIMax = "0.3"))
+	UPROPERTY(EditAnywhere, BlueprintReadOnly, Category=CascadedShadowMaps, meta=(UIMin = "0", UIMax = "0.3", DisplayName = "Transition Fraction"))
 	float CascadeTransitionFraction;
 
 	/** 
 	 * Controls the size of the fade out region at the far extent of the dynamic shadow's influence.  
 	 * This is specified as a fraction of DynamicShadowDistance. 
 	 */
-	UPROPERTY(EditAnywhere, BlueprintReadOnly, Category=CascadedShadowMaps, meta=(UIMin = "0", UIMax = "1.0"))
+	UPROPERTY(EditAnywhere, BlueprintReadOnly, Category=CascadedShadowMaps, meta=(UIMin = "0", UIMax = "1.0", DisplayName = "Distance Fadeout Fraction"))
 	float ShadowDistanceFadeoutFraction;
 
 	/** 
