better UI

diff --git a/Engine/Source/Runtime/Engine/Classes/Components/PrimitiveComponent.h b/Engine/Source/Runtime/Engine/Classes/Components/PrimitiveComponent.h
--- a/Engine/Source/Runtime/Engine/Classes/Components/PrimitiveComponent.h
+++ b/Engine/Source/Runtime/Engine/Classes/Components/PrimitiveComponent.h
@@ -289,7 +289,7 @@ public:
 	 * Whether this component should create a per-object shadow that gives higher effective shadow resolution. 
 	 * Useful for cinematic character shadowing. Assumed to be enabled if bSelfShadowOnly is enabled.
 	 */
-	UPROPERTY(EditAnywhere, AdvancedDisplay, BlueprintReadOnly, Category=Lighting, meta=(EditCondition="CastShadow", DisplayName = "Inset Shadow"))
+	UPROPERTY(EditAnywhere, AdvancedDisplay, BlueprintReadOnly, Category=Lighting, meta=(EditCondition="CastShadow", DisplayName = "Dynamic Inset Shadow"))
 	uint32 bCastInsetShadow:1;
 
 	/** 
