#ue4 - Fix "p.ShowInitialOverlaps 1" causing navigation build to sometimes create a crazy number of batched debug lines/points.

- Improved LineBatchComponent to handle large number of batched elements by avoiding frequent memory operations.

diff --git a/Engine/Source/Runtime/Engine/Private/Collision/CollisionConversions.cpp b/Engine/Source/Runtime/Engine/Private/Collision/CollisionConversions.cpp
--- a/Engine/Source/Runtime/Engine/Private/Collision/CollisionConversions.cpp
+++ b/Engine/Source/Runtime/Engine/Private/Collision/CollisionConversions.cpp
@@ -470,6 +470,7 @@ FVector FindBestOverlappingNormal(const PxGeometry& Geom, const PxTransform& Que
 {
 #if DRAW_OVERLAPPING_TRIS
 	const float Lifetime = 5.f;
+	bCanDrawOverlaps &= GWorld->IsGameWorld() && GWorld->PersistentLineBatcher && (GWorld->PersistentLineBatcher->BatchedLines.Num() < 2048);
 	if (bCanDrawOverlaps)
 	{
 		TArray<FOverlapResult> Overlaps;
@@ -511,11 +512,11 @@ FVector FindBestOverlappingNormal(const PxGeometry& Geom, const PxTransform& Que
 		}
 
 #if DRAW_OVERLAPPING_TRIS
-		if (bCanDrawOverlaps && GWorld->PersistentLineBatcher)
+		if (bCanDrawOverlaps && (GWorld->PersistentLineBatcher->BatchedLines.Num() < 2048))
 		{
-			static float LineThickness = 0.9f;
-			static float NormalThickness = 0.75f;
-			static float PointThickness = 5.0f;
+			static const float LineThickness = 0.9f;
+			static const float NormalThickness = 0.75f;
+			static const float PointThickness = 5.0f;
 			GWorld->PersistentLineBatcher->DrawLine(A, B, LineColor, SDPG_Foreground, LineThickness, Lifetime);
 			GWorld->PersistentLineBatcher->DrawLine(B, C, LineColor, SDPG_Foreground, LineThickness, Lifetime);
 			GWorld->PersistentLineBatcher->DrawLine(C, A, LineColor, SDPG_Foreground, LineThickness, Lifetime);
@@ -718,7 +719,7 @@ static bool ConvertOverlappedShapeToImpactHit(const PxLocationHit& PHit, const F
 	}
 
 #if DRAW_OVERLAPPING_TRIS
-	if (CVarShowInitialOverlaps.GetValueOnAnyThread() != 0)
+	if (CVarShowInitialOverlaps.GetValueOnAnyThread() != 0 && GWorld->IsGameWorld())
 	{
 		FVector DummyNormal(0.f);
 		const PxTransform PShapeWorldPose = PxShapeExt::getGlobalPose(*PShape, *PActor);
diff --git a/Engine/Source/Runtime/Engine/Private/LineBatchComponent.cpp b/Engine/Source/Runtime/Engine/Private/LineBatchComponent.cpp
--- a/Engine/Source/Runtime/Engine/Private/LineBatchComponent.cpp
+++ b/Engine/Source/Runtime/Engine/Private/LineBatchComponent.cpp
@@ -261,7 +261,7 @@ void ULineBatchComponent::TickComponent(float DeltaTime, enum ELevelTick TickTyp
 			if(Line.RemainingLifeTime <= 0.0f)
 			{
 				// The line has expired, remove it.
-				BatchedLines.RemoveAt(LineIndex--);
+				BatchedLines.RemoveAtSwap(LineIndex--);
 				bDirty = true;
 			}
 		}
@@ -277,7 +277,7 @@ void ULineBatchComponent::TickComponent(float DeltaTime, enum ELevelTick TickTyp
 			if(Pt.RemainingLifeTime <= 0.0f)
 			{
 				// The point has expired, remove it.
-				BatchedPoints.RemoveAt(PtIndex--);
+				BatchedPoints.RemoveAtSwap(PtIndex--);
 				bDirty = true;
 			}
 		}
@@ -293,7 +293,7 @@ void ULineBatchComponent::TickComponent(float DeltaTime, enum ELevelTick TickTyp
 			if(Mesh.RemainingLifeTime <= 0.0f)
 			{
 				// The mesh has expired, remove it.
-				BatchedMeshes.RemoveAt(MeshIndex--);
+				BatchedMeshes.RemoveAtSwap(MeshIndex--);
 				bDirty = true;
 			}
 		}
