#ue4 - Remove triangles with very small area when cooking collision trimesh data (area is controlled by physics settings "TriangleMeshTriangleMinAreaThreshold"). Fixes various issues caused by having degenerate triangles in PhysX collision data.

**Note** Geometry must be rebuilt for this to take effect.

Jira: UE-7526

diff --git a/Engine/Source/Runtime/Engine/Classes/PhysicsEngine/PhysicsSettings.h b/Engine/Source/Runtime/Engine/Classes/PhysicsEngine/PhysicsSettings.h
--- a/Engine/Source/Runtime/Engine/Classes/PhysicsEngine/PhysicsSettings.h
+++ b/Engine/Source/Runtime/Engine/Classes/PhysicsEngine/PhysicsSettings.h
@@ -90,6 +90,10 @@ class ENGINE_API UPhysicsSettings
 	UPROPERTY(config, EditAnywhere, meta = (ClampMin = "1", UIMin = "1", ClampMax = "127", UIMax = "127"), Category = Constants)
 	int32 RagdollAggregateThreshold;
 
+	/** Triangles from triangle meshes (BSP) with an area less than or equal to this value will be removed from physics collision data. Set to less than 0 to disable. */
+	UPROPERTY(config, EditAnywhere, AdvancedDisplay, meta = (ClampMin = "-1.0", UIMin = "-1.0", ClampMax = "10.0", UIMax = "10.0"), Category = Constants)
+	float TriangleMeshTriangleMinAreaThreshold;
+
 	/** Enables the use of an async scene */
 	UPROPERTY(config, EditAnywhere, AdvancedDisplay, Category=Simulation)
 	bool bEnableAsyncScene;
diff --git a/Engine/Source/Runtime/Engine/Private/ModelComponent.cpp b/Engine/Source/Runtime/Engine/Private/ModelComponent.cpp
--- a/Engine/Source/Runtime/Engine/Private/ModelComponent.cpp
+++ b/Engine/Source/Runtime/Engine/Private/ModelComponent.cpp
@@ -10,6 +10,7 @@
 #include "LightMap.h"
 #include "ShadowMap.h"
 #include "Components/ModelComponent.h"
+#include "PhysicsEngine/PhysicsSettings.h"
 
 FModelElement::FModelElement(UModelComponent* InComponent,UMaterialInterface* InMaterial):
 	Component(InComponent),
@@ -528,6 +529,8 @@ void UModelComponent::InvalidateCollisionData()
 bool UModelComponent::GetPhysicsTriMeshData(struct FTriMeshCollisionData* CollisionData, bool InUseAllTriData)
 {
 	check(Model);
+	int32 nBadArea = 0;
+	const float AreaThreshold = UPhysicsSettings::Get()->TriangleMeshTriangleMinAreaThreshold;
 
 	const int32 NumVerts = Model->VertexBuffer.Vertices.Num();
 	CollisionData->Vertices.AddUninitialized(NumVerts);
@@ -549,11 +552,33 @@ bool UModelComponent::GetPhysicsTriMeshData(struct FTriMeshCollisionData* Collis
 			Triangle.v1 = IndexBuffer->Indices[Element.FirstIndex + (TriIdx*3) + 1];
 			Triangle.v2 = IndexBuffer->Indices[Element.FirstIndex + (TriIdx*3) + 2];
 
+			if (AreaThreshold >= 0.f)
+			{
+				const FVector V0 = Model->VertexBuffer.Vertices[Triangle.v0].Position;
+				const FVector V1 = Model->VertexBuffer.Vertices[Triangle.v1].Position;
+				const FVector V2 = Model->VertexBuffer.Vertices[Triangle.v2].Position;
+
+				const FVector V01 = (V1 - V0);
+				const FVector V02 = (V2 - V0);
+				const FVector Cross = FVector::CrossProduct(V01, V02);
+				const float Area = Cross.Size() * 0.5f;
+				if (Area <= AreaThreshold)
+				{
+					nBadArea++;
+					continue;
+				}
+			}
+
 			CollisionData->Indices.Add(Triangle);
 			CollisionData->MaterialIndices.Add(ElementIndex);
 		}
 	}
 
+	if (nBadArea > 0)
+	{
+		UE_LOG(LogPhysics, Log, TEXT("Cooking removed %d triangle%s with area <= %f (%s)"), nBadArea, (nBadArea > 1 ? TEXT("s") : TEXT("")), AreaThreshold, *GetPathName(GetOuter()));
+	}
+
 	CollisionData->bFlipNormals = true;
 	return true;
 }
diff --git a/Engine/Source/Runtime/Engine/Private/PhysicsEngine/PhysicsSettings.cpp b/Engine/Source/Runtime/Engine/Private/PhysicsEngine/PhysicsSettings.cpp
--- a/Engine/Source/Runtime/Engine/Private/PhysicsEngine/PhysicsSettings.cpp
+++ b/Engine/Source/Runtime/Engine/Private/PhysicsEngine/PhysicsSettings.cpp
@@ -11,6 +11,7 @@ UPhysicsSettings::UPhysicsSettings(const FObjectInitializer& ObjectInitializer)
 	, DefaultTerminalVelocity(4000.f)
 	, DefaultFluidFriction(0.3f)
 	, RagdollAggregateThreshold(4)
+	, TriangleMeshTriangleMinAreaThreshold(0.25f)
 	, bEnableAsyncScene(false)
 	, bEnable2DPhysics(false)
 	, bSimulateSkeletalMeshOnDedicatedServer(true)
