Inheritable Components and SCS_Node use GUID

diff --git a/Engine/Source/Editor/KismetCompiler/Private/KismetCompiler.cpp b/Engine/Source/Editor/KismetCompiler/Private/KismetCompiler.cpp
--- a/Engine/Source/Editor/KismetCompiler/Private/KismetCompiler.cpp
+++ b/Engine/Source/Editor/KismetCompiler/Private/KismetCompiler.cpp
@@ -3280,7 +3280,7 @@ void FKismetCompilerContext::Compile()
 		auto InheritableComponentHandler = Blueprint->GetInheritableComponentHandler(false);
 		if (InheritableComponentHandler)
 		{
-			InheritableComponentHandler->RemoveInvalidAndUnnecessaryTemplates();
+			InheritableComponentHandler->ValidateTemplates();
 		}
 	}
 
diff --git a/Engine/Source/Runtime/Engine/Classes/Engine/InheritableComponentHandler.h b/Engine/Source/Runtime/Engine/Classes/Engine/InheritableComponentHandler.h
--- a/Engine/Source/Runtime/Engine/Classes/Engine/InheritableComponentHandler.h
+++ b/Engine/Source/Runtime/Engine/Classes/Engine/InheritableComponentHandler.h
@@ -18,21 +18,23 @@ struct ENGINE_API FComponentKey
 	UPROPERTY()
 	FName VariableName;
 
-	// TODO: a GUID should be added
+	UPROPERTY()
+	FGuid VariableGuid;
 
 	FComponentKey()
 		: OwnerClass(nullptr)
 	{}
 
-	FComponentKey(USCS_Node* ParentNode);
+	FComponentKey(USCS_Node* SCSNode);
 
 	bool Match(const FComponentKey OtherKey) const;
 
 	bool IsValid() const
 	{
-		return OwnerClass && (VariableName != NAME_None);
+		return OwnerClass && (VariableName != NAME_None) && VariableGuid.IsValid();
 	}
 
+	USCS_Node* FindSCSNode() const;
 	UActorComponent* GetOriginalTemplate() const;
 };
 
@@ -65,7 +67,7 @@ private:
 public:
 	UActorComponent* CreateOverridenComponentTemplate(FComponentKey Key);
 	void UpdateOwnerClass(UBlueprintGeneratedClass* OwnerClass);
-	void RemoveInvalidAndUnnecessaryTemplates();
+	void ValidateTemplates();
 	bool IsValid() const;
 	UActorComponent* FindBestArchetype(FComponentKey Key) const;
 
diff --git a/Engine/Source/Runtime/Engine/Classes/Engine/SCS_Node.h b/Engine/Source/Runtime/Engine/Classes/Engine/SCS_Node.h
--- a/Engine/Source/Runtime/Engine/Classes/Engine/SCS_Node.h
+++ b/Engine/Source/Runtime/Engine/Classes/Engine/SCS_Node.h
@@ -49,6 +49,9 @@ class USCS_Node : public UObject
 	UPROPERTY(EditAnywhere, Category=BPVariableDescription)
 	TArray<struct FBPVariableMetaDataEntry> MetaDataArray;
 
+	UPROPERTY()
+	FGuid VariableGuid;
+
 	/** (DEPRECATED) */
 	UPROPERTY()
 	bool bIsFalseRoot_DEPRECATED;
@@ -66,6 +69,11 @@ class USCS_Node : public UObject
 	bool bVariableNameAutoGenerated_DEPRECATED;
 
 #if WITH_EDITOR
+
+	// Begin UObject Interface
+	virtual void PostLoad() override;
+	// End UObject Interface
+
 	/** 
 	 * Create a list of existing names used in this function. 
 	 *
diff --git a/Engine/Source/Runtime/Engine/Classes/Engine/SimpleConstructionScript.h b/Engine/Source/Runtime/Engine/Classes/Engine/SimpleConstructionScript.h
--- a/Engine/Source/Runtime/Engine/Classes/Engine/SimpleConstructionScript.h
+++ b/Engine/Source/Runtime/Engine/Classes/Engine/SimpleConstructionScript.h
@@ -61,6 +61,9 @@ class USimpleConstructionScript : public UObject
 	/** Find the SCS_Node node by name and return it if found */
 	ENGINE_API USCS_Node* FindSCSNode(FName InName);
 
+	/** Find the SCS_Node node by name and return it if found */
+	ENGINE_API USCS_Node* FindSCSNodeByGuid(FGuid Guid);
+
 protected:
 	/** Checks the root node set for scene components and ensures that it is valid (e.g. after a removal) */
 	void ValidateSceneRootNodes();
diff --git a/Engine/Source/Runtime/Engine/Private/InheritableComponentHandler.cpp b/Engine/Source/Runtime/Engine/Private/InheritableComponentHandler.cpp
--- a/Engine/Source/Runtime/Engine/Private/InheritableComponentHandler.cpp
+++ b/Engine/Source/Runtime/Engine/Private/InheritableComponentHandler.cpp
@@ -53,13 +53,28 @@ void UInheritableComponentHandler::UpdateOwnerClass(UBlueprintGeneratedClass* Ow
 	}
 }
 
-void UInheritableComponentHandler::RemoveInvalidAndUnnecessaryTemplates()
+void UInheritableComponentHandler::ValidateTemplates()
 {
 	for (int32 Index = 0; Index < Records.Num();)
 	{
 		bool bIsValidAndNecessary = false;
 		{
-			const FComponentOverrideRecord& Record = Records[Index];
+			FComponentOverrideRecord& Record = Records[Index];
+			
+			if (Record.ComponentKey.OwnerClass && Record.ComponentKey.VariableGuid.IsValid())
+			{
+				// Update variable, based on value
+				auto SCSNode = Record.ComponentKey.FindSCSNode();
+				const FName ActualName = SCSNode ? SCSNode->GetVariableName() : NAME_None;
+				if ((ActualName != NAME_None) && (ActualName != Record.ComponentKey.VariableName))
+				{
+					UE_LOG(LogBlueprint, Log, TEXT("ValidateTemplates '%s': variable old name '%s' new name '%s'"),
+						*GetPathNameSafe(this), *Record.ComponentKey.VariableName.ToString(), *ActualName.ToString());
+					Record.ComponentKey.VariableName = ActualName;
+					MarkPackageDirty();
+				}
+			}
+
 			if (IsRecordValid(Record))
 			{
 				if (IsRecordNecessary(Record))
@@ -68,13 +83,13 @@ void UInheritableComponentHandler::RemoveInvalidAndUnnecessaryTemplates()
 				}
 				else
 				{
-					UE_LOG(LogBlueprint, Log, TEXT("RemoveInvalidAndUnnecessaryTemplates '%s': overriden template is unnecessary - component '%s' from '%s'"),
+					UE_LOG(LogBlueprint, Log, TEXT("ValidateTemplates '%s': overriden template is unnecessary - component '%s' from '%s'"),
 						*GetPathNameSafe(this), *Record.ComponentKey.VariableName.ToString(), *GetPathNameSafe(Record.ComponentKey.OwnerClass));
 				}
 			}
 			else
 			{
-				UE_LOG(LogBlueprint, Warning, TEXT("RemoveInvalidAndUnnecessaryTemplates '%s': overriden template is invalid - component '%s' from '%s'"),
+				UE_LOG(LogBlueprint, Warning, TEXT("ValidateTemplates '%s': overriden template is invalid - component '%s' from '%s'"),
 					*GetPathNameSafe(this), *Record.ComponentKey.VariableName.ToString(), *GetPathNameSafe(Record.ComponentKey.OwnerClass));
 			}
 		}
@@ -253,40 +268,30 @@ const FComponentOverrideRecord* UInheritableComponentHandler::FindRecord(const F
 
 // FComponentOverrideRecord
 
-FComponentKey::FComponentKey(USCS_Node* ParentNode) : OwnerClass(nullptr)
+FComponentKey::FComponentKey(USCS_Node* SCSNode) : OwnerClass(nullptr)
 {
-	VariableName = ParentNode ? ParentNode->GetVariableName() : NAME_None;
-	if (VariableName == NAME_None)
+	if (SCSNode)
 	{
-		return;
-	}
-
-	//Find original blueprint where the component is declared
-	auto ParentSCS = ParentNode ? ParentNode->GetSCS() : nullptr;
-	OwnerClass = ParentSCS ? Cast<UBlueprintGeneratedClass>(ParentSCS->GetOwnerClass()) : nullptr;
-	for (auto SuperOwnerClass = OwnerClass ? Cast<UBlueprintGeneratedClass>(OwnerClass->GetSuperClass()) : nullptr; 
-		SuperOwnerClass;
-		SuperOwnerClass = Cast<UBlueprintGeneratedClass>(SuperOwnerClass->GetSuperClass()))
-	{
-		if (SuperOwnerClass->SimpleConstructionScript && SuperOwnerClass->SimpleConstructionScript->FindSCSNode(VariableName))
-		{
-			OwnerClass = SuperOwnerClass;
-		}
-		else
-		{
-			break;
-		}
+		auto ParentSCS = SCSNode->GetSCS();
+		OwnerClass = ParentSCS ? Cast<UBlueprintGeneratedClass>(ParentSCS->GetOwnerClass()) : nullptr;
+		VariableGuid = SCSNode->VariableGuid;
+		VariableName = SCSNode->GetVariableName();
 	}
 }
 
 bool FComponentKey::Match(const FComponentKey OtherKey) const
 {
-	return (OwnerClass == OtherKey.OwnerClass) && (VariableName == OtherKey.VariableName);
+	return (OwnerClass == OtherKey.OwnerClass) && (VariableGuid == OtherKey.VariableGuid);
 }
 
-UActorComponent* FComponentKey::GetOriginalTemplate() const
+USCS_Node* FComponentKey::FindSCSNode() const
 {
 	auto ParentSCS = OwnerClass ? OwnerClass->SimpleConstructionScript : nullptr;
-	auto SCSNode = ParentSCS ? ParentSCS->FindSCSNode(VariableName) : nullptr;
+	return ParentSCS ? ParentSCS->FindSCSNodeByGuid(VariableGuid) : nullptr;
+}
+
+UActorComponent* FComponentKey::GetOriginalTemplate() const
+{
+	auto SCSNode = FindSCSNode();
 	return SCSNode ? SCSNode->ComponentTemplate : nullptr;
 }
diff --git a/Engine/Source/Runtime/Engine/Private/SCS_Node.cpp b/Engine/Source/Runtime/Engine/Private/SCS_Node.cpp
--- a/Engine/Source/Runtime/Engine/Private/SCS_Node.cpp
+++ b/Engine/Source/Runtime/Engine/Private/SCS_Node.cpp
@@ -461,4 +461,14 @@ FName USCS_Node::GenerateNewComponentName( TArray<FName>& CurrentNames, FName De
 	return NewName;
 }
 
+void USCS_Node::PostLoad()
+{
+	Super::PostLoad();
+
+	if (!VariableGuid.IsValid())
+	{
+		VariableGuid = FGuid::NewGuid();
+	}
+}
+
 #endif
diff --git a/Engine/Source/Runtime/Engine/Private/SimpleConstructionScript.cpp b/Engine/Source/Runtime/Engine/Private/SimpleConstructionScript.cpp
--- a/Engine/Source/Runtime/Engine/Private/SimpleConstructionScript.cpp
+++ b/Engine/Source/Runtime/Engine/Private/SimpleConstructionScript.cpp
@@ -616,6 +616,22 @@ USCS_Node* USimpleConstructionScript::FindSCSNode(FName InName)
 	return ReturnSCSNode;
 }
 
+USCS_Node* USimpleConstructionScript::FindSCSNodeByGuid(FGuid Guid)
+{
+	TArray<USCS_Node*> AllNodes = GetAllNodes();
+	USCS_Node* ReturnSCSNode = nullptr;
+
+	for (USCS_Node* SCSNode : AllNodes)
+	{
+		if (SCSNode->VariableGuid == Guid)
+		{
+			ReturnSCSNode = SCSNode;
+			break;
+		}
+	}
+	return ReturnSCSNode;
+}
+
 void USimpleConstructionScript::ValidateSceneRootNodes()
 {
 #if WITH_EDITOR
@@ -762,6 +778,8 @@ USCS_Node* USimpleConstructionScript::CreateNode(UActorComponent* NewComponentTe
 		// Note: This should match up with UEdGraphSchema_K2::VR_DefaultCategory
 		NewNode->CategoryName = TEXT("Default");
 
+		NewNode->VariableGuid = FGuid::NewGuid();
+
 		return NewNode;
 	}
 
