Backward compatibility fix for GUID in SCS_Node

diff --git a/Engine/Source/Editor/UnrealEd/Private/Kismet2/BlueprintEditorUtils.cpp b/Engine/Source/Editor/UnrealEd/Private/Kismet2/BlueprintEditorUtils.cpp
--- a/Engine/Source/Editor/UnrealEd/Private/Kismet2/BlueprintEditorUtils.cpp
+++ b/Engine/Source/Editor/UnrealEd/Private/Kismet2/BlueprintEditorUtils.cpp
@@ -435,6 +435,19 @@ void FBlueprintEditorUtils::PreloadConstructionScript(UBlueprint* Blueprint)
 			RootNodes[NodeIndex]->PreloadChain();
 		}
 	}
+
+	auto SimpleConstructionScript = Blueprint ? Blueprint->SimpleConstructionScript : nullptr;
+	if (SimpleConstructionScript)
+	{
+		auto AllNodes = SimpleConstructionScript->GetAllNodes();
+		for (auto SCSNode : AllNodes)
+		{
+			if (SCSNode)
+			{
+				SCSNode->ValidateGuid();
+			}
+		}
+	}
 }
 
 void FBlueprintEditorUtils::PatchNewCDOIntoLinker(UObject* CDO, ULinkerLoad* Linker, int32 ExportIndex, TArray<UObject*>& ObjLoaded)
diff --git a/Engine/Source/Runtime/Engine/Classes/Engine/SCS_Node.h b/Engine/Source/Runtime/Engine/Classes/Engine/SCS_Node.h
--- a/Engine/Source/Runtime/Engine/Classes/Engine/SCS_Node.h
+++ b/Engine/Source/Runtime/Engine/Classes/Engine/SCS_Node.h
@@ -93,6 +93,9 @@ class USCS_Node : public UObject
 
 	/** The scene component constructed for component editing in the SCS editor */
 	class USceneComponent* EditorComponentInstance;
+
+	/** Make sure, the guid is proper - backward compatibility */
+	ENGINE_API void ValidateGuid();
 #endif
 
 	/**
diff --git a/Engine/Source/Runtime/Engine/Private/SCS_Node.cpp b/Engine/Source/Runtime/Engine/Private/SCS_Node.cpp
--- a/Engine/Source/Runtime/Engine/Private/SCS_Node.cpp
+++ b/Engine/Source/Runtime/Engine/Private/SCS_Node.cpp
@@ -465,9 +465,22 @@ void USCS_Node::PostLoad()
 {
 	Super::PostLoad();
 
+	ValidateGuid();
+}
+
+void USCS_Node::ValidateGuid()
+{
+	// Backward compatibility: node requires a guid. 
+	// The guid for the node should be always the same (event when it's not saved). The guid is created using persistent name.
 	if (!VariableGuid.IsValid())
 	{
-		VariableGuid = FGuid::NewGuid();
+		const FName PersistentVariableName = GetVariableName();
+		if (PersistentVariableName != NAME_None)
+		{
+			const FString NameVariableString = PersistentVariableName.ToString();
+			const uint32 PersistentCrc = FCrc::StrCrc32(*NameVariableString);
+			VariableGuid = FGuid(PersistentCrc, 0, 0, 0);
+		}
 	}
 }
 
