UE4 - add a command to log out foliage

diff --git a/Engine/Source/Runtime/Engine/Private/HierarchicalInstancedStaticMesh.cpp b/Engine/Source/Runtime/Engine/Private/HierarchicalInstancedStaticMesh.cpp
--- a/Engine/Source/Runtime/Engine/Private/HierarchicalInstancedStaticMesh.cpp
+++ b/Engine/Source/Runtime/Engine/Private/HierarchicalInstancedStaticMesh.cpp
@@ -760,6 +760,19 @@ static FAutoConsoleCommand ToggleUseVectorCullCmd(
 	FConsoleCommandWithArgsDelegate::CreateStatic(&ToggleUseVectorCull)
 	);
 
+static uint32 GFrameNumberRenderThread_CaptureFoliageRuns = MAX_uint32;
+
+static void LogFoliageFrame(const TArray<FString>& Args)
+{
+	GFrameNumberRenderThread_CaptureFoliageRuns = GFrameNumberRenderThread + 2;
+}
+
+static FAutoConsoleCommand LogFoliageFrameCmd(
+	TEXT("foliage.LogFoliageFrame"),
+	TEXT("Useful for debugging. Logs all foliage rendered in a frame."),
+	FConsoleCommandWithArgsDelegate::CreateStatic(&LogFoliageFrame)
+	);
+
 const VectorRegister		VECTOR_HALF_HALF_HALF_ZERO				= DECLARE_VECTOR_REGISTER(0.5f, 0.5f, 0.5f, 0.0f);
 
 template<bool TUseVector>
@@ -907,6 +920,9 @@ struct FFoliageElementParams
 	bool bUseHoveredMaterial;
 	bool bInstanced;
 	bool bBlendLODs;
+	ERHIFeatureLevel::Type FeatureLevel;
+	bool ShadowFrustum;
+	float FinalCullDistance;
 };
 
 void FHierarchicalStaticMeshSceneProxy::FillDynamicMeshElements(FMeshElementCollector& Collector, const FFoliageElementParams& ElementParams, const FFoliageRenderInstanceParams& Params) const
@@ -985,6 +1001,27 @@ void FHierarchicalStaticMeshSceneProxy::FillDynamicMeshElements(FMeshElementColl
 						bDidStats = true;
 						int64 Tris = int64(RemainingInstances) * int64(BatchElement0.NumPrimitives);
 						TotalTriangles += Tris;
+#if STATS
+						if (GFrameNumberRenderThread_CaptureFoliageRuns == GFrameNumberRenderThread)
+						{
+							if (ElementParams.FinalCullDistance > 9.9E8)
+							{
+								UE_LOG(LogStaticMesh, Display, TEXT("lod:%1d/%1d   sel:%1d   section:%1d/%1d   runs:%4d   inst:%8d   tris:%9lld   cast shadow:%1d   cull:-NONE!!-   shadow:%1d     %s %s"), 
+									LODIndex, InstancedRenderData.VertexFactories.Num(), SelectionGroupIndex, SectionIndex, LODModel.Sections.Num(), RunArray.Num() / 2, 
+									RemainingInstances, Tris, (int)MeshElement.CastShadow, ElementParams.ShadowFrustum,
+									*StaticMesh->GetPathName(),
+									*MeshElement.MaterialRenderProxy->GetMaterial(ElementParams.FeatureLevel)->GetFriendlyName());
+							}
+							else
+							{
+								UE_LOG(LogStaticMesh, Display, TEXT("lod:%1d/%1d   sel:%1d   section:%1d/%1d   runs:%4d   inst:%8d   tris:%9lld   cast shadow:%1d   cull:%8.0f   shadow:%1d     %s %s"), 
+									LODIndex, InstancedRenderData.VertexFactories.Num(), SelectionGroupIndex, SectionIndex, LODModel.Sections.Num(), RunArray.Num() / 2, 
+									RemainingInstances, Tris, (int)MeshElement.CastShadow, ElementParams.FinalCullDistance, ElementParams.ShadowFrustum,
+									*StaticMesh->GetPathName(),
+									*MeshElement.MaterialRenderProxy->GetMaterial(ElementParams.FeatureLevel)->GetFriendlyName());
+							}
+						}
+#endif
 					}
 					if (ElementParams.bInstanced)
 					{
@@ -1071,6 +1108,7 @@ void FHierarchicalStaticMeshSceneProxy::GetDynamicMeshElements(const TArray<cons
 			ElementParams.bIsWireframe = ViewFamily.EngineShowFlags.Wireframe;
 			ElementParams.bUseHoveredMaterial = IsHovered();
 			ElementParams.bInstanced = RHISupportsInstancing(GetFeatureLevelShaderPlatform(InstancedRenderData.FeatureLevel));
+			ElementParams.FeatureLevel = InstancedRenderData.FeatureLevel;
 			ElementParams.ViewIndex = ViewIndex;
 			ElementParams.View = View;
 
@@ -1087,6 +1125,7 @@ void FHierarchicalStaticMeshSceneProxy::GetDynamicMeshElements(const TArray<cons
 				bool bUseVectorCull = GUseVectorCull;
 
 				bool bDisableCull = !!CVarDisableCull.GetValueOnRenderThread();
+				ElementParams.ShadowFrustum = !!View->ViewMatrices.GetDynamicMeshElementsShadowCullFrustum;
 				if (View->ViewMatrices.GetDynamicMeshElementsShadowCullFrustum)
 				{
 					float GuardBand = CVarFoliageOrthoGuardband.GetValueOnRenderThread();
@@ -1177,6 +1216,7 @@ void FHierarchicalStaticMeshSceneProxy::GetDynamicMeshElements(const TArray<cons
 					InstanceParams.LODPlanesMax[LODIndex] = MIN_flt;
 					InstanceParams.LODPlanesMin[LODIndex] = MAX_flt;
 				}
+				ElementParams.FinalCullDistance = MIN_flt;
 				for (int32 SampleIndex = 0; SampleIndex < 2; SampleIndex++)
 				{
 					// we aren't going to accurately figure out LOD planes for both samples and try all 4 possibilities, rather we just consider the worst case LOD range; the shaders will render correctly
@@ -1193,6 +1233,7 @@ void FHierarchicalStaticMeshSceneProxy::GetDynamicMeshElements(const TArray<cons
 						FinalCull = UserData_AllInstances.EndCullDistance;
 					}
 					FinalCull *= MaxDrawDistanceScale;
+					ElementParams.FinalCullDistance = FMath::Max(ElementParams.FinalCullDistance, FinalCull);
 
 					for (int32 LODIndex = 1; LODIndex < InstanceParams.LODs; LODIndex++)
 					{
@@ -1745,7 +1786,7 @@ void UHierarchicalInstancedStaticMeshComponent::ApplyBuildTreeAsync(ENamedThread
 		Exchange(SortedInstances, Builder->Result->SortedInstances);
 		RemovedInstances.Empty();
 
-	    UE_LOG(LogStaticMesh, Display, TEXT("Built a foliage hierarchy with %d of %d elements in %.1fs."), NumBuiltInstances, PerInstanceSMData.Num(), float(FPlatformTime::Seconds() - StartTime));
+	    //UE_LOG(LogStaticMesh, Display, TEXT("Built a foliage hierarchy with %d of %d elements in %.1fs."), NumBuiltInstances, PerInstanceSMData.Num(), float(FPlatformTime::Seconds() - StartTime));
     
 	    if (NumBuiltInstances < PerInstanceSMData.Num())
 	    {
