fixed comment/docs
UEDOC-947 Fwd: Incorrect documentation on MipValueMode

diff --git a/Engine/Source/Runtime/Engine/Classes/Materials/MaterialExpressionTextureSample.h b/Engine/Source/Runtime/Engine/Classes/Materials/MaterialExpressionTextureSample.h
--- a/Engine/Source/Runtime/Engine/Classes/Materials/MaterialExpressionTextureSample.h
+++ b/Engine/Source/Runtime/Engine/Classes/Materials/MaterialExpressionTextureSample.h
@@ -10,9 +10,9 @@
 UENUM()
 enum ETextureMipValueMode
 {
-	TMVM_None UMETA(DisplayName="None"),
-	TMVM_MipLevel UMETA(DisplayName="MipLevel"),
-	TMVM_MipBias UMETA(DisplayName="MipBias"),
+	TMVM_None UMETA(DisplayName="None (use computed mip level)"),
+	TMVM_MipLevel UMETA(DisplayName="MipLevel (absolute, 0 is full resolution)"),
+	TMVM_MipBias UMETA(DisplayName="MipBias (relative to the computed mip level)"),
 	TMVM_MAX,
 };
 
@@ -31,11 +31,11 @@ class ENGINE_API UMaterialExpressionTextureSample : public UMaterialExpressionTe
 	UPROPERTY(meta = (RequiredInput = "false", ToolTip = "Defaults to 'Texture' if not specified"))
 	FExpressionInput TextureObject;
 
-	/** Meaning depends on MipValueMode */
+	/** Meaning depends on MipValueMode, a single unit is one mip level  */
 	UPROPERTY(meta = (RequiredInput = "false", ToolTip = "Defaults to 'ConstMipValue' if not specified"))
 	FExpressionInput MipValue;
 
-	/** Noise function, affects performance and look */
+	/** Defines how the MipValue property is applied to the texture lookup */
 	UPROPERTY(EditAnywhere, Category=MaterialExpressionTextureSample, meta=(DisplayName = "MipValueMode"))
 	TEnumAsByte<enum ETextureMipValueMode> MipValueMode;
 
