Fix runtime divide by zero crash for matinees with zero FOV angle.

diff --git a/Engine/Source/Runtime/Engine/Private/LocalPlayer.cpp b/Engine/Source/Runtime/Engine/Private/LocalPlayer.cpp
--- a/Engine/Source/Runtime/Engine/Private/LocalPlayer.cpp
+++ b/Engine/Source/Runtime/Engine/Private/LocalPlayer.cpp
@@ -1003,8 +1003,9 @@ bool ULocalPlayer::GetProjectionData(FViewport* Viewport, EStereoscopicPass Ster
 			}
 			else
 			{
+				// Avoid divide by zero in the projection matrix calculation by clamping FOV
 				ProjectionData.ProjectionMatrix = FReversedZPerspectiveMatrix( 
-					ViewInfo.FOV * (float)PI / 360.0f,
+					FMath::Max(0.001f, ViewInfo.FOV) * (float)PI / 360.0f,
 					ViewInfo.AspectRatio,
 					1.0f,
 					GNearClippingPlane );
@@ -1012,7 +1013,8 @@ bool ULocalPlayer::GetProjectionData(FViewport* Viewport, EStereoscopicPass Ster
 		}
 		else 
 		{
-			float MatrixFOV = ViewInfo.FOV * (float)PI / 360.0f;
+			// Avoid divide by zero in the projection matrix calculation by clamping FOV
+			float MatrixFOV = FMath::Max(0.001f, ViewInfo.FOV) * (float)PI / 360.0f;
 			float XAxisMultiplier;
 			float YAxisMultiplier;
 
