Improved fix for UE-7698 following advice from Nick.

diff --git a/Engine/Source/Runtime/Core/Public/UObject/ObjectVersion.h b/Engine/Source/Runtime/Core/Public/UObject/ObjectVersion.h
--- a/Engine/Source/Runtime/Core/Public/UObject/ObjectVersion.h
+++ b/Engine/Source/Runtime/Core/Public/UObject/ObjectVersion.h
@@ -454,6 +454,8 @@ enum EUnrealEngineObjectUE4Version
 	VER_UE4_LANDSCAPE_STATIC_SECTION_OFFSET,
 	// New options for navigation data runtime generation (static, modifiers only, dynamic)
 	VER_UE4_ADD_MODIFIERS_RUNTIME_GENERATION,
+	// Tidied up material's handling of masked blend mode.
+	VER_UE4_MATERIAL_MASKED_BLENDMODE_TIDY,
 
 	// -----<new versions can be added before this line>-------------------------------------------------
 	// - this needs to be the last line (see note below)
diff --git a/Engine/Source/Runtime/Engine/Classes/Materials/Material.h b/Engine/Source/Runtime/Engine/Classes/Materials/Material.h
--- a/Engine/Source/Runtime/Engine/Classes/Materials/Material.h
+++ b/Engine/Source/Runtime/Engine/Classes/Materials/Material.h
@@ -644,6 +644,10 @@ public:
 	UPROPERTY()
 	TArray<struct FMaterialParameterCollectionInfo> MaterialParameterCollectionInfos;
 
+	/** true if this Material can be assumed Opaque when set to masked. */
+	UPROPERTY()
+	uint32 bCanMaskedBeAssumedOpaque : 1;
+
  	/** true if Material is masked and uses custom opacity */
  	UPROPERTY()
  	uint32 bIsMasked_DEPRECATED:1;
diff --git a/Engine/Source/Runtime/Engine/Private/Materials/Material.cpp b/Engine/Source/Runtime/Engine/Private/Materials/Material.cpp
--- a/Engine/Source/Runtime/Engine/Private/Materials/Material.cpp
+++ b/Engine/Source/Runtime/Engine/Private/Materials/Material.cpp
@@ -1750,6 +1750,20 @@ void UMaterial::Serialize(FArchive& Ar)
 	DoMaterialAttributeReorder(&SubsurfaceColor,		Ar.UE4Ver());
 	DoMaterialAttributeReorder(&AmbientOcclusion,		Ar.UE4Ver());
 	DoMaterialAttributeReorder(&Refraction,				Ar.UE4Ver());
+
+	if (Ar.UE4Ver() < VER_UE4_MATERIAL_MASKED_BLENDMODE_TIDY)
+	{
+		//Set based on old value. Real check may not be possible here in cooked builds?
+		//Cached using acutal check in PostEditChangProperty().
+		if (BlendMode == BLEND_Masked && !bIsMasked_DEPRECATED)
+		{
+			bCanMaskedBeAssumedOpaque = true;
+		}
+		else
+		{
+			bCanMaskedBeAssumedOpaque = false;
+		}
+	}
 }
 
 void UMaterial::PostDuplicate(bool bDuplicateForPIE)
@@ -2291,6 +2305,9 @@ void UMaterial::PostEditChangeProperty(FPropertyChangedEvent& PropertyChangedEve
 		}
 	}
 
+	//If we can be sure this material would be the same opaque as it is masked then allow it to be assumed opaque.
+	bCanMaskedBeAssumedOpaque = !OpacityMask.Expression && !(OpacityMask.UseConstant && OpacityMask.Constant < 0.999f) && !bUseMaterialAttributes;
+
 	bool bRequiresCompilation = true;
 	if( PropertyThatChanged ) 
 	{
@@ -3456,14 +3473,13 @@ EBlendMode UMaterial::GetBlendMode(bool bIsInGameThread) const
 {
 	if (EBlendMode(BlendMode) == BLEND_Masked)
 	{
-		// Check if the material is masked and uses a custom opacity (that's not 1.0f).
-		if ((((OpacityMask.Expression || (OpacityMask.UseConstant && OpacityMask.Constant < 0.999f))) || bUseMaterialAttributes))
+		if (bCanMaskedBeAssumedOpaque)
 		{
-			return BLEND_Masked;
+			return BLEND_Opaque;
 		}
 		else
 		{
-			return BLEND_Opaque;
+			return BLEND_Masked;
 		}
 	}
 	else
