Fix for crash closing PIE in mac while firing

diff --git a/Engine/Source/Runtime/Engine/Private/Slate/SceneViewport.cpp b/Engine/Source/Runtime/Engine/Private/Slate/SceneViewport.cpp
--- a/Engine/Source/Runtime/Engine/Private/Slate/SceneViewport.cpp
+++ b/Engine/Source/Runtime/Engine/Private/Slate/SceneViewport.cpp
@@ -471,11 +471,14 @@ void FSceneViewport::OnMouseEnter( const FGeometry& MyGeometry, const FPointerEv
 
 void FSceneViewport::OnMouseLeave( const FPointerEvent& MouseEvent )
 {
-	ViewportClient->MouseLeave( this );
-	
-	if ( IsPlayInEditorViewport() )
+	if( ViewportClient )
 	{
-		CachedMousePos = FIntPoint(-1, -1);
+		ViewportClient->MouseLeave( this );
+	
+		if ( IsPlayInEditorViewport() )
+		{
+			CachedMousePos = FIntPoint(-1, -1);
+		}
 	}
 }
 
