Fix incorrect use of skipped tick time variable

diff --git a/Engine/Source/Runtime/Engine/Private/SkeletalMeshComponentPhysics.cpp b/Engine/Source/Runtime/Engine/Private/SkeletalMeshComponentPhysics.cpp
--- a/Engine/Source/Runtime/Engine/Private/SkeletalMeshComponentPhysics.cpp
+++ b/Engine/Source/Runtime/Engine/Private/SkeletalMeshComponentPhysics.cpp
@@ -3108,7 +3108,7 @@ void USkeletalMeshComponent::PreClothTick(float DeltaTime)
 	// if skeletal mesh has clothing assets, call TickClothing
 	if (SkeletalMesh && SkeletalMesh->ClothingAssets.Num() > 0)
 	{
-		TickClothing(DeltaTime + SkippedTickDeltaTime);
+		TickClothing(DeltaTime);
 	}
 #endif
 }
