#ue4 - Fix polygon normals becoming incorrect when non-uniform scale is applied.

- Removed redundant calculation of rotation matrix.
- Removed pointless normal renormalization when vertices are simply translated.

diff --git a/Engine/Source/Runtime/Engine/Private/Polygon.cpp b/Engine/Source/Runtime/Engine/Private/Polygon.cpp
--- a/Engine/Source/Runtime/Engine/Private/Polygon.cpp
+++ b/Engine/Source/Runtime/Engine/Private/Polygon.cpp
@@ -450,11 +450,6 @@ void FPoly::Transform
 	Base = (Base - PreSubtract) + PostAdd;
 	for( i=0; i<Vertices.Num(); i++ )
 		Vertices[i]  = (Vertices[i] - PreSubtract) + PostAdd;
-
-	// Transform normal.  Since the transformation coordinate system is
-	// orthogonal but not orthonormal, it has to be renormalized here.
-	Normal = Normal.GetSafeNormal();
-
 }
 
 
@@ -464,20 +459,22 @@ void FPoly::Rotate
 	const FRotator&		Rotation
 )
 {
+	const FRotationMatrix RotMatrix(Rotation);
+
 	// Rotate the vertices.
 	for (int32 Vertex = 0; Vertex < Vertices.Num(); Vertex++)
 	{
-		Vertices[Vertex] = PreSubtract + FRotationMatrix(Rotation).TransformVector(Vertices[Vertex] - PreSubtract);
+		Vertices[Vertex] = PreSubtract + RotMatrix.TransformVector(Vertices[Vertex] - PreSubtract);
 	}
 
-	Base = PreSubtract + FRotationMatrix(Rotation).TransformVector(Base - PreSubtract);
+	Base = PreSubtract + RotMatrix.TransformVector(Base - PreSubtract);
 
 	// Rotate the texture vectors.
-	TextureU = FRotationMatrix(Rotation).TransformVector(TextureU);
-	TextureV = FRotationMatrix(Rotation).TransformVector(TextureV);
+	TextureU = RotMatrix.TransformVector(TextureU);
+	TextureV = RotMatrix.TransformVector(TextureV);
 
 	// Rotate the normal.
-	Normal = FRotationMatrix(Rotation).TransformVector(Normal);
+	Normal = RotMatrix.TransformVector(Normal);
 	Normal = Normal.GetSafeNormal();
 }
 
@@ -488,6 +485,13 @@ void FPoly::Scale
 	const FVector&		Scale
 )
 {
+	if (Scale.X == 1.f &&
+		Scale.Y == 1.f &&
+		Scale.Z == 1.f)
+	{
+		return;
+	}
+
 	// Scale the vertices.
 	for (int32 Vertex = 0; Vertex < Vertices.Num(); Vertex++)
 	{
@@ -500,8 +504,8 @@ void FPoly::Scale
 	TextureU *= Scale;
 	TextureV *= Scale;
 
-	// Renormalize the normal.
-	Normal = Normal.GetSafeNormal();
+	// Recalculate the normal. Non-uniform scale or mirroring requires this.
+	CalcNormal(true);
 }
 
 
