Added override to debug plane drawing to allow seperate axis extents

diff --git a/Engine/Source/Runtime/Engine/Private/DrawDebugHelpers.cpp b/Engine/Source/Runtime/Engine/Private/DrawDebugHelpers.cpp
--- a/Engine/Source/Runtime/Engine/Private/DrawDebugHelpers.cpp
+++ b/Engine/Source/Runtime/Engine/Private/DrawDebugHelpers.cpp
@@ -257,6 +257,38 @@ void DrawDebugSolidPlane(const UWorld* InWorld, FPlane const& P, FVector const&
 	}
 }
 
+ENGINE_API void DrawDebugSolidPlane(const UWorld* InWorld, FPlane const& P, FVector const& Loc, FVector2D const& Extents, FColor const& Color, bool bPersistent/*=false*/, float LifeTime/*=-1*/, uint8 DepthPriority /*= 0*/)
+{
+	// no debug line drawing on dedicated server
+	if(GEngine->GetNetMode(InWorld) != NM_DedicatedServer)
+	{
+		FVector const ClosestPtOnPlane = Loc - P.PlaneDot(Loc) * P;
+
+		FVector U, V;
+		P.FindBestAxisVectors(U, V);
+		U *= Extents.Y;
+		V *= Extents.X;
+
+		TArray<FVector> Verts;
+		Verts.AddUninitialized(4);
+		Verts[0] = ClosestPtOnPlane + U + V;
+		Verts[1] = ClosestPtOnPlane - U + V;
+		Verts[2] = ClosestPtOnPlane + U - V;
+		Verts[3] = ClosestPtOnPlane - U - V;
+
+		TArray<int32> Indices;
+		Indices.AddUninitialized(6);
+		Indices[0] = 0; Indices[1] = 2; Indices[2] = 1;
+		Indices[3] = 1; Indices[4] = 2; Indices[5] = 3;
+
+		// plane quad
+		DrawDebugMesh(InWorld, Verts, Indices, Color, bPersistent, LifeTime, DepthPriority);
+
+		// arrow indicating normal
+		DrawDebugDirectionalArrow(InWorld, ClosestPtOnPlane, ClosestPtOnPlane + P * 16.f, 8.f, FColor::White, bPersistent, LifeTime, DepthPriority);
+	}
+}
+
 void DrawDebugCoordinateSystem(const UWorld* InWorld, FVector const& AxisLoc, FRotator const& AxisRot, float Scale, bool bPersistentLines, float LifeTime, uint8 DepthPriority)
 {
 	// no debug line drawing on dedicated server
diff --git a/Engine/Source/Runtime/Engine/Public/DrawDebugHelpers.h b/Engine/Source/Runtime/Engine/Public/DrawDebugHelpers.h
--- a/Engine/Source/Runtime/Engine/Public/DrawDebugHelpers.h
+++ b/Engine/Source/Runtime/Engine/Public/DrawDebugHelpers.h
@@ -44,6 +44,7 @@ ENGINE_API void DrawDebugSolidBox(const UWorld* InWorld, FVector const& Center,
 ENGINE_API void DrawDebugSolidBox(const UWorld* InWorld, FVector const& Center, FVector const& Extent, FQuat const& Rotation, FColor const& Color, bool bPersistent=false, float LifeTime=-1.f, uint8 DepthPriority = 0);
 ENGINE_API void DrawDebugMesh(const UWorld* InWorld, TArray<FVector> const& Verts, TArray<int32> const& Indices, FColor const& Color, bool bPersistent=false, float LifeTime=-1.f, uint8 DepthPriority = 0);
 ENGINE_API void DrawDebugSolidPlane(const UWorld* InWorld, FPlane const& P, FVector const& Loc, float Size, FColor const& Color, bool bPersistent=false, float LifeTime=-1, uint8 DepthPriority = 0);
+ENGINE_API void DrawDebugSolidPlane(const UWorld* InWorld, FPlane const& P, FVector const& Loc, FVector2D const& Extents, FColor const& Color, bool bPersistent=false, float LifeTime=-1, uint8 DepthPriority = 0);
 
 /* Draws a 2D Histogram of size 'DrawSize' based FDebugFloatHistory struct, using DrawTransform for the position in the world. */
 ENGINE_API void DrawDebugFloatHistory(UWorld const & WorldRef, FDebugFloatHistory const & FloatHistory, FTransform const & DrawTransform, FVector2D const & DrawSize, FColor const & DrawColor, bool const & bPersistent = false, float const & LifeTime = -1.f, uint8 const & DepthPriority = 0);
