#ue4 - Bump cooked physics data version to force geometry rebuild to address zero-area BSP triangles fix from CL 2413559.

diff --git a/Engine/Source/Runtime/Engine/Private/PhysicsEngine/PhysDerivedData.h b/Engine/Source/Runtime/Engine/Private/PhysicsEngine/PhysDerivedData.h
--- a/Engine/Source/Runtime/Engine/Private/PhysicsEngine/PhysDerivedData.h
+++ b/Engine/Source/Runtime/Engine/Private/PhysicsEngine/PhysDerivedData.h
@@ -43,7 +43,9 @@ public:
 
 	virtual FString GetPluginSpecificCacheKeySuffix() const override
 	{
-		enum { UE_PHYSX_DERIVEDDATA_VER = 1 };
+		//  1 - base version
+		//  2 - cook out small area trimesh triangles from BSP (see UPhysicsSettings::TriangleMeshTriangleMinAreaThreshold)
+		enum { UE_PHYSX_DERIVEDDATA_VER = 2 };
 
 		const uint16 PhysXVersion = ((PX_PHYSICS_VERSION_MAJOR  & 0xF) << 12) |
 				((PX_PHYSICS_VERSION_MINOR  & 0xF) << 8) |
