#ue4 - Expose GetActorEyesViewpoint to blueprints.

GitHub #765

diff --git a/Engine/Source/Runtime/Engine/Classes/GameFramework/Actor.h b/Engine/Source/Runtime/Engine/Classes/GameFramework/Actor.h
--- a/Engine/Source/Runtime/Engine/Classes/GameFramework/Actor.h
+++ b/Engine/Source/Runtime/Engine/Classes/GameFramework/Actor.h
@@ -2059,8 +2059,8 @@ public:
 	virtual void PrestreamTextures( float Seconds, bool bEnableStreaming, int32 CinematicTextureGroups = 0 );
 
 	/**
-	 * returns the point of view of the actor.
-	 * note that this doesn't mean the camera, but the 'eyes' of the actor.
+	 * Returns the point of view of the actor.
+	 * Note that this doesn't mean the camera, but the 'eyes' of the actor.
 	 * For example, for a Pawn, this would define the eye height location,
 	 * and view rotation (which is different from the pawn rotation which has a zeroed pitch component).
 	 * A camera first person view will typically use this view point. Most traces (weapon, AI) will be done from this view point.
@@ -2068,6 +2068,7 @@ public:
 	 * @param	OutLocation - location of view point
 	 * @param	OutRotation - view rotation of actor.
 	 */
+	UFUNCTION(BlueprintCallable, Category = Actor)
 	virtual void GetActorEyesViewPoint( FVector& OutLocation, FRotator& OutRotation ) const;
 
 	/**
