UE4 - fix bounds crash

diff --git a/Engine/Source/Runtime/Engine/Private/HierarchicalInstancedStaticMesh.cpp b/Engine/Source/Runtime/Engine/Private/HierarchicalInstancedStaticMesh.cpp
--- a/Engine/Source/Runtime/Engine/Private/HierarchicalInstancedStaticMesh.cpp
+++ b/Engine/Source/Runtime/Engine/Private/HierarchicalInstancedStaticMesh.cpp
@@ -1489,7 +1489,10 @@ bool UHierarchicalInstancedStaticMeshComponent::UpdateInstanceTransform(int32 In
 
 	bool Result = Super::UpdateInstanceTransform(InstanceIndex, NewInstanceTransform, bWorldSpace);
 
-	UnbuiltInstanceBounds += StaticMesh->GetBounds().GetBox().TransformBy(NewInstanceTransform);
+	if (StaticMesh)
+	{
+		UnbuiltInstanceBounds += StaticMesh->GetBounds().GetBox().TransformBy(NewInstanceTransform);
+	}
 
 	if (!IsAsyncBuilding())
 	{
