Remove old protocol version variables from UNetConnection

diff --git a/Engine/Source/Runtime/Engine/Classes/Engine/NetConnection.h b/Engine/Source/Runtime/Engine/Classes/Engine/NetConnection.h
--- a/Engine/Source/Runtime/Engine/Classes/Engine/NetConnection.h
+++ b/Engine/Source/Runtime/Engine/Classes/Engine/NetConnection.h
@@ -165,8 +165,6 @@ class UNetConnection : public UPlayer
 
 public:
 	// Constants.
-	enum{ MAX_PROTOCOL_VERSION = 1     };	// Maximum protocol version supported.
-	enum{ MIN_PROTOCOL_VERSION = 1     };	// Minimum protocol version supported.
 	enum{ MAX_CHANNELS         = 10240 };	// Maximum channels. TODO: This needs to differ per game somehow but cannot with shared executable
 
 	// Connection information.
@@ -181,7 +179,6 @@ public:
 	TSharedPtr<class FUniqueNetId> PlayerId;
 
 	// Negotiated parameters.
-	int32			ProtocolVersion;		// Protocol version we're communicating with (<=PROTOCOL_VERSION).
 	int32			PacketOverhead;			// Bytes overhead per packet sent.
 	FString			Challenge;				// Server-generated challenge.
 	FString			ClientResponse;			// Client-generated response.
diff --git a/Engine/Source/Runtime/Engine/Private/NetConnection.cpp b/Engine/Source/Runtime/Engine/Private/NetConnection.cpp
--- a/Engine/Source/Runtime/Engine/Private/NetConnection.cpp
+++ b/Engine/Source/Runtime/Engine/Private/NetConnection.cpp
@@ -28,7 +28,6 @@ UNetConnection::UNetConnection(const FObjectInitializer& ObjectInitializer)
 ,	MaxPacket			( 0 )
 ,	InternalAck			( false )
 ,	State				( USOCK_Invalid )
-,	ProtocolVersion		( MIN_PROTOCOL_VERSION )
 ,	PacketOverhead		( 0 )
 ,	ResponseId			( 0 )
 
@@ -371,8 +370,6 @@ bool UNetConnection::Exec( UWorld* InWorld, const TCHAR* Cmd, FOutputDevice& Ar
 void UNetConnection::AssertValid()
 {
 	// Make sure this connection is in a reasonable state.
-	check(ProtocolVersion>=MIN_PROTOCOL_VERSION);
-	check(ProtocolVersion<=MAX_PROTOCOL_VERSION);
 	check(State==USOCK_Closed || State==USOCK_Pending || State==USOCK_Open);
 
 }
