Added OperationSystemId detection for windows. Other platforms are TBD. #UETOOL-237.
* Analytics now sends this along with the MachineID for Launcher and EngineAnalytics. #UETOOL-236

diff --git a/Engine/Source/Runtime/Core/Private/GenericPlatform/GenericPlatformMisc.cpp b/Engine/Source/Runtime/Core/Private/GenericPlatform/GenericPlatformMisc.cpp
--- a/Engine/Source/Runtime/Core/Private/GenericPlatform/GenericPlatformMisc.cpp
+++ b/Engine/Source/Runtime/Core/Private/GenericPlatform/GenericPlatformMisc.cpp
@@ -865,3 +865,9 @@ void FGenericPlatformMisc::PlatformPreInit()
 {
 	FGenericCrashContext::Initialize();
 }
+
+FString FGenericPlatformMisc::GetOperatingSystemId()
+{
+	// not implemented by default.
+	return FString();
+}
diff --git a/Engine/Source/Runtime/Core/Private/Windows/WindowsPlatformMisc.cpp b/Engine/Source/Runtime/Core/Private/Windows/WindowsPlatformMisc.cpp
--- a/Engine/Source/Runtime/Core/Private/Windows/WindowsPlatformMisc.cpp
+++ b/Engine/Source/Runtime/Core/Private/Windows/WindowsPlatformMisc.cpp
@@ -2322,3 +2322,11 @@ bool FWindowsPlatformMisc::IsRunningOnBattery()
 
 	return false;
 }
+
+FString FWindowsPlatformMisc::GetOperatingSystemId()
+{
+	FString Result;
+	// more info on this key can be found here: http://stackoverflow.com/questions/99880/generating-a-unique-machine-id
+	QueryRegKey(HKEY_LOCAL_MACHINE, TEXT("Software\\Microsoft\\Cryptography"), TEXT("MachineGuid"), Result);
+	return Result;
+}
diff --git a/Engine/Source/Runtime/Core/Public/GenericPlatform/GenericPlatformMisc.h b/Engine/Source/Runtime/Core/Public/GenericPlatform/GenericPlatformMisc.h
--- a/Engine/Source/Runtime/Core/Public/GenericPlatform/GenericPlatformMisc.h
+++ b/Engine/Source/Runtime/Core/Public/GenericPlatform/GenericPlatformMisc.h
@@ -819,7 +819,9 @@ public:
 	static bool IsRunningOnBattery();
 
 	/**
-	 * Get (or create) the unique ID used to identify this computer
+	 * Get (or create) the unique ID used to identify this computer.
+	 * This is sort of a misnomer, as there will actually be one per user account on the operating system.
+	 * This is NOT based on a machine fingerprint.
 	 */
 	static FGuid GetMachineId();
 
@@ -835,6 +837,12 @@ public:
 	 */
 	static bool SetEpicAccountId( const FString& AccountId );
 
+	/**
+	 * Gets a globally unique ID the represents a particular operating system install.
+	 * @returns an opaque string representing the ID, or an empty string if the platform doesn't support one.
+	 */
+	static FString GetOperatingSystemId();
+
 	/** 
 	 * Get a string description of the mode the engine was running in.
 	 */
diff --git a/Engine/Source/Runtime/Core/Public/Windows/WindowsPlatformMisc.h b/Engine/Source/Runtime/Core/Public/Windows/WindowsPlatformMisc.h
--- a/Engine/Source/Runtime/Core/Public/Windows/WindowsPlatformMisc.h
+++ b/Engine/Source/Runtime/Core/Public/Windows/WindowsPlatformMisc.h
@@ -213,6 +213,12 @@ struct CORE_API FWindowsPlatformMisc
 	 * Returns whether the platform is running on battery power or not.
 	 */
 	static bool IsRunningOnBattery();
+
+	/**
+	 * Gets a globally unique ID the represents a particular operating system install.
+	 */
+	static FString GetOperatingSystemId();
+
 };
 
 
diff --git a/Engine/Source/Runtime/Engine/Private/EngineAnalytics.cpp b/Engine/Source/Runtime/Engine/Private/EngineAnalytics.cpp
--- a/Engine/Source/Runtime/Engine/Private/EngineAnalytics.cpp
+++ b/Engine/Source/Runtime/Engine/Private/EngineAnalytics.cpp
@@ -10,89 +10,6 @@ bool FEngineAnalytics::bIsInitialized;
 TSharedPtr<IAnalyticsProvider> FEngineAnalytics::Analytics;
 bool FEngineAnalytics::bShouldSendUsageEvents;
 
-#define DEBUG_ANALYTICS 0
-
-#if DEBUG_ANALYTICS
-
-DEFINE_LOG_CATEGORY_STATIC(LogDebugAnalytics, Log, All);
-
-namespace DebugAnalyticsProviderDefs
-{
-	const static int32 MaxAttributeCount = 40;
-}
-
-/**
- * Debug analytics provider for debugging analytics events
- * without the need to connect to a real provider.
- */
-class FDebugAnalyticsProvider : public IAnalyticsProvider
-{
-public:
-	virtual bool StartSession(const TArray<FAnalyticsEventAttribute>& Attributes)
-	{
-		UE_LOG(LogDebugAnalytics, Log, TEXT("FDebugAnalyticsProvider: StartSession..."));
-		OutputAttributesToLog(Attributes);
-		return true;
-	}
-
-	virtual void EndSession()
-	{
-		UE_LOG(LogDebugAnalytics, Log, TEXT("FDebugAnalyticsProvider: EndSession"));
-	}
-
-	virtual FString GetSessionID() const
-	{
-		return SessionID;
-	}
-
-	virtual bool SetSessionID(const FString& InSessionID)
-	{
-		UE_LOG(LogDebugAnalytics, Log, TEXT("FDebugAnalyticsProvider: SetSessionID SessionID=%s"), *InSessionID);
-		SessionID = InSessionID;
-		return true;
-	}
-
-	virtual void FlushEvents()
-	{
-	}
-
-	virtual void SetUserID(const FString& InUserID)
-	{
-		UE_LOG(LogDebugAnalytics, Log, TEXT("FDebugAnalyticsProvider: SetUserID UserID=%s"), *InUserID);
-		UserID = InUserID;
-	}
-
-	virtual FString GetUserID() const
-	{
-		return UserID;
-	}
-
-	virtual void RecordEvent(const FString& EventName, const TArray<FAnalyticsEventAttribute>& Attributes)
-	{
-		if (Attributes.Num() > DebugAnalyticsProviderDefs::MaxAttributeCount)
-		{
-			UE_LOG(LogDebugAnalytics, Warning, TEXT("FDebugAnalyticsProvider: Event %s has too many attributes (%d)"), *EventName, Attributes.Num());
-		}
-
-		UE_LOG(LogDebugAnalytics, Log, TEXT("FDebugAnalyticsProvider: RecordEvent Name=%s"), *EventName);
-		OutputAttributesToLog(Attributes);
-	}
-
-private:
-	void OutputAttributesToLog(const TArray<FAnalyticsEventAttribute>& Attributes)
-	{
-		for (auto& Attrib : Attributes)
-		{
-			UE_LOG(LogDebugAnalytics, Log, TEXT("AnalyticsAttrib Name=%s, Value=%s"), *Attrib.AttrName, *Attrib.AttrValue);
-		}
-	}
-
-	FString SessionID;
-	FString UserID;
-};
-
-#endif // DEBUG_ANALYTICS
-
 /**
  * Engine analytics config log to initialize the engine analytics provider.
  * External code should bind this delegate if engine analytics are desired,
@@ -137,13 +54,6 @@ void FEngineAnalytics::Initialize()
 
 	if (bShouldInitAnalytics)
 	{
-#if DEBUG_ANALYTICS
-		// Dummy analytics provider just logs events locally
-		Analytics = MakeShareable(new FDebugAnalyticsProvider());
-		Analytics->SetUserID(FPlatformMisc::GetUniqueDeviceId());
-		Analytics->StartSession();
-#else
-
 		{
 			// Setup some default engine analytics if there is nothing custom bound
 			FAnalytics::FProviderConfigurationDelegate DefaultEngineAnalyticsConfig;
@@ -203,7 +113,7 @@ void FEngineAnalytics::Initialize()
 				DefaultEngineAnalyticsConfig);
 			if (Analytics.IsValid())
 			{
-				Analytics->SetUserID(FString::Printf(TEXT("%s|%s"), *FPlatformMisc::GetMachineId().ToString(EGuidFormats::Digits).ToLower(), *FPlatformMisc::GetEpicAccountId()));
+				Analytics->SetUserID(FString::Printf(TEXT("%s|%s|%s"), *FPlatformMisc::GetMachineId().ToString(EGuidFormats::Digits).ToLower(), *FPlatformMisc::GetEpicAccountId(), *FPlatformMisc::GetOperatingSystemId()));
 
 				TArray<FAnalyticsEventAttribute> StartSessionAttributes;
 				GEngine->CreateStartupAnalyticsAttributes( StartSessionAttributes );
@@ -211,7 +121,6 @@ void FEngineAnalytics::Initialize()
 				Analytics->StartSession( StartSessionAttributes );
 			}
 		}
-#endif
 	}
 	bIsInitialized = true;
 }
