NavMesh rendering optimization - removed redundant gathering of nav mesh data for drawing.

diff --git a/Engine/Source/Runtime/Engine/Classes/AI/Navigation/NavMeshRenderingComponent.h b/Engine/Source/Runtime/Engine/Classes/AI/Navigation/NavMeshRenderingComponent.h
--- a/Engine/Source/Runtime/Engine/Classes/AI/Navigation/NavMeshRenderingComponent.h
+++ b/Engine/Source/Runtime/Engine/Classes/AI/Navigation/NavMeshRenderingComponent.h
@@ -8,8 +8,11 @@ class ENGINE_API UNavMeshRenderingComponent : public UPrimitiveComponent
 {
 	GENERATED_UCLASS_BODY()
 
+public:
 	// Begin UPrimitiveComponent Interface
 	virtual FPrimitiveSceneProxy* CreateSceneProxy() override;
+	virtual void OnRegister()  override;
+	virtual void OnUnregister()  override;
 	// End UPrimitiveComponent Interface
 
 	// Begin UActorComponent Interface
@@ -22,4 +25,11 @@ class ENGINE_API UNavMeshRenderingComponent : public UPrimitiveComponent
 	// End USceneComponent Interface
 
 	void GatherData(struct FNavMeshSceneProxyData*) const;
+
+protected:
+	void TimerFunction();
+
+protected:
+	int32 bCollectNavigationData : 1;
+	FTimerHandle TimerHandle;
 };
diff --git a/Engine/Source/Runtime/Engine/Private/NavMeshRenderingComponent.cpp b/Engine/Source/Runtime/Engine/Private/NavMeshRenderingComponent.cpp
--- a/Engine/Source/Runtime/Engine/Private/NavMeshRenderingComponent.cpp
+++ b/Engine/Source/Runtime/Engine/Private/NavMeshRenderingComponent.cpp
@@ -11,9 +11,12 @@
 #include "AI/Navigation/RecastNavMeshGenerator.h"
 #include "AI/Navigation/NavigationSystem.h"
 
-#include "NavMeshRenderingHelpers.h"
+#if WITH_EDITOR
+#include "Editor.h"
+#include "EditorViewportClient.h"
+#endif
 
-static bool GForceDisableNavmeshRendering = false; // FIXME: Temporary Hack! to disable navmesh rendering completely because of performance drops (SebaK for Cameron change)
+#include "NavMeshRenderingHelpers.h"
 
 UNavMeshRenderingComponent::UNavMeshRenderingComponent(const FObjectInitializer& ObjectInitializer)
 	: Super(ObjectInitializer)
@@ -22,15 +25,100 @@ UNavMeshRenderingComponent::UNavMeshRenderingComponent(const FObjectInitializer&
 	AlwaysLoadOnClient = false;
 	AlwaysLoadOnServer = false;
 	bSelectable = false;
-	GConfig->GetBool(TEXT("NavMeshRenderingComponent"), TEXT("bForceDisableNavmeshRendering"), GForceDisableNavmeshRendering, GEngineIni);
+	bCollectNavigationData = false;
+}
+
+void UNavMeshRenderingComponent::TimerFunction()
+{
+	bool bShowNavigation = false;
+#if WITH_EDITOR
+	if (GEditor)
+	{
+		for (FEditorViewportClient* CurrentViewpoer : GEditor->AllViewportClients)
+		{
+			if (CurrentViewpoer->EngineShowFlags.Navigation)
+			{
+				bShowNavigation = true;
+				break;
+			}
+		}
+	}
+	else
+#endif //WITH_EDITOR
+	{
+		bShowNavigation = GetWorld()->GetGameViewport()->EngineShowFlags.Navigation;
+	}
+
+	if (bShowNavigation != !!bCollectNavigationData)
+	{
+		bCollectNavigationData = bShowNavigation;
+		MarkRenderStateDirty();
+	}
+}
+
+void UNavMeshRenderingComponent::OnRegister()
+{
+	Super::OnRegister();
+
+#if WITH_RECAST && !UE_BUILD_SHIPPING && !UE_BUILD_TEST
+	// it's a kind of HACK but there is no event or other information that show flag was changed by user => we have to check it periodically
+#if WITH_EDITOR
+	if (GEditor)
+	{
+		GEditor->GetTimerManager()->SetTimer(TimerHandle, FTimerDelegate::CreateUObject(this, &UNavMeshRenderingComponent::TimerFunction), 1, true);
+	}
+	else
+#endif //WITH_EDITOR
+	{
+		GetWorld()->GetTimerManager().SetTimer(TimerHandle, FTimerDelegate::CreateUObject(this, &UNavMeshRenderingComponent::TimerFunction), 1, true);
+	}
+#endif //WITH_RECAST && !UE_BUILD_SHIPPING && !UE_BUILD_TEST
+}
+
+void UNavMeshRenderingComponent::OnUnregister()
+{
+#if WITH_RECAST && !UE_BUILD_SHIPPING && !UE_BUILD_TEST
+	// it's a kind of HACK but there is no event or other information that show flag was changed by user => we have to check it periodically
+#if WITH_EDITOR
+	if (GEditor)
+	{
+		GEditor->GetTimerManager()->ClearTimer(TimerHandle);
+	}
+	else
+#endif //WITH_EDITOR
+	{
+		GetWorld()->GetTimerManager().ClearTimer(TimerHandle);
+	}
+#endif //WITH_RECAST && !UE_BUILD_SHIPPING && !UE_BUILD_TEST
+	Super::OnUnregister();
 }
 
 FPrimitiveSceneProxy* UNavMeshRenderingComponent::CreateSceneProxy()
 {
-#if WITH_RECAST && WITH_EDITOR
+#if WITH_RECAST && !UE_BUILD_SHIPPING && !UE_BUILD_TEST
 	FPrimitiveSceneProxy* SceneProxy = NULL;
 
-	if (!GForceDisableNavmeshRendering && IsVisible())
+	bool bShowNavigation = false;
+#if WITH_EDITOR
+	if (GEditor)
+	{
+		for (FEditorViewportClient* CurrentViewpoer : GEditor->AllViewportClients)
+		{
+			if (CurrentViewpoer->EngineShowFlags.Navigation)
+			{
+				bShowNavigation = true;
+				break;
+			}
+		}
+	}
+	else
+#endif //WITH_EDITOR
+	{
+		bShowNavigation = GetWorld()->GetGameViewport()->EngineShowFlags.Navigation;
+	}
+	bCollectNavigationData = bShowNavigation;
+
+	if (bCollectNavigationData && IsVisible())
 	{
 		FNavMeshSceneProxyData ProxyData;
 		ProxyData.Reset();
@@ -42,14 +130,14 @@ FPrimitiveSceneProxy* UNavMeshRenderingComponent::CreateSceneProxy()
 	return SceneProxy;
 #else
 	return NULL;
-#endif
+#endif //WITH_RECAST && !UE_BUILD_SHIPPING && !UE_BUILD_TEST
 }
 
 void UNavMeshRenderingComponent::CreateRenderState_Concurrent()
 {
 	Super::CreateRenderState_Concurrent();
 
-#if WITH_RECAST && WITH_EDITOR
+#if WITH_RECAST && !UE_BUILD_SHIPPING && !UE_BUILD_TEST
 	if (SceneProxy != NULL)
 	{
 		static_cast<FRecastRenderingSceneProxy*>(SceneProxy)->RegisterDebugDrawDelgate();
@@ -59,7 +147,7 @@ void UNavMeshRenderingComponent::CreateRenderState_Concurrent()
 
 void UNavMeshRenderingComponent::DestroyRenderState_Concurrent()
 {
-#if WITH_RECAST && WITH_EDITOR
+#if WITH_RECAST && !UE_BUILD_SHIPPING && !UE_BUILD_TEST
 	if (SceneProxy != NULL)
 	{
 		static_cast<FRecastRenderingSceneProxy*>(SceneProxy)->UnregisterDebugDrawDelgate();
