Native support for state entry/exit callbacks

Added simple callback from state machines into their owning UAnimInstance so that native code can be easily notified of state transitions.

diff --git a/Engine/Source/Runtime/Engine/Classes/Animation/AnimInstance.h b/Engine/Source/Runtime/Engine/Classes/Animation/AnimInstance.h
--- a/Engine/Source/Runtime/Engine/Classes/Animation/AnimInstance.h
+++ b/Engine/Source/Runtime/Engine/Classes/Animation/AnimInstance.h
@@ -575,6 +575,12 @@ public:
 	// Note: the node graph will not be evaluated if this function returns true
 	virtual bool NativeEvaluateAnimation(FPoseContext& Output);
 
+	// Called back from the anim graph when a state is entered
+	virtual void NativeStateStart(const FName& MachineName, const FName& StateName);
+
+	// Called back from the anim graph when a state is exited
+	virtual void NativeStateEnd(const FName& MachineName, const FName& StateName);
+
 	// Debug output for this anim instance 
 	void DisplayDebug(UCanvas* Canvas, const FDebugDisplayInfo& DebugDisplay, float& YL, float& YPos);
 public:
diff --git a/Engine/Source/Runtime/Engine/Private/Animation/AnimInstance.cpp b/Engine/Source/Runtime/Engine/Private/Animation/AnimInstance.cpp
--- a/Engine/Source/Runtime/Engine/Private/Animation/AnimInstance.cpp
+++ b/Engine/Source/Runtime/Engine/Private/Animation/AnimInstance.cpp
@@ -531,6 +531,14 @@ bool UAnimInstance::NativeEvaluateAnimation(FPoseContext& Output)
 	return false;
 }
 
+void UAnimInstance::NativeStateStart(const FName& MachineName, const FName& StateName)
+{
+}
+
+void UAnimInstance::NativeStateEnd(const FName& MachineName, const FName& StateName) 
+{
+}
+
 void OutputCurveMap(TMap<FName, float>& CurveMap, UCanvas* Canvas, UFont* RenderFont, float Indent, float& YPos, FFontRenderInfo RenderInfo, float& YL)
 {
 	TArray<FName> Names;
diff --git a/Engine/Source/Runtime/Engine/Private/Animation/AnimNode_StateMachine.cpp b/Engine/Source/Runtime/Engine/Private/Animation/AnimNode_StateMachine.cpp
--- a/Engine/Source/Runtime/Engine/Private/Animation/AnimNode_StateMachine.cpp
+++ b/Engine/Source/Runtime/Engine/Private/Animation/AnimNode_StateMachine.cpp
@@ -317,19 +317,24 @@ void FAnimNode_StateMachine::Update(const FAnimationUpdateContext& Context)
 				Context.AnimInstance->AddAnimNotifyFromGeneratedClass(ActiveTransitionArray[ActiveTransitionArray.Num()-1].InterruptNotify);
 			}
 
-			const int32 PreviousState = CurrentState;
-			const int32 NextState = PotentialTransition.TargetState;
+			const int32 PreviousStateIndex = CurrentState;
+			const int32 NextStateIndex = PotentialTransition.TargetState;
 
 			// Fire off Notifies for state transition
-			Context.AnimInstance->AddAnimNotifyFromGeneratedClass(GetStateInfo(PreviousState).EndNotify);
-			Context.AnimInstance->AddAnimNotifyFromGeneratedClass(GetStateInfo(NextState).StartNotify);
+			const FBakedAnimationState& PreviousState = GetStateInfo(PreviousStateIndex);
+			Context.AnimInstance->NativeStateEnd(PRIVATE_MachineDescription->MachineName, PreviousState.StateName);
+			Context.AnimInstance->AddAnimNotifyFromGeneratedClass(PreviousState.EndNotify);
+
+			const FBakedAnimationState& NextState = GetStateInfo(PreviousStateIndex);
+			Context.AnimInstance->NativeStateStart(PRIVATE_MachineDescription->MachineName, NextState.StateName);
+			Context.AnimInstance->AddAnimNotifyFromGeneratedClass(NextState.StartNotify);
 			
 			// Get the current weight of the next state, which may be non-zero
-			const float ExistingWeightOfNextState = GetStateWeight(NextState);
+			const float ExistingWeightOfNextState = GetStateWeight(NextStateIndex);
 
 			// Push the transition onto the stack
 			const FAnimationTransitionBetweenStates& ReferenceTransition = GetTransitionInfo(PotentialTransition.TransitionRule->TransitionIndex);
-			FAnimationActiveTransitionEntry* NewTransition = new (ActiveTransitionArray) FAnimationActiveTransitionEntry(NextState, ExistingWeightOfNextState, PreviousState, ReferenceTransition);
+			FAnimationActiveTransitionEntry* NewTransition = new (ActiveTransitionArray) FAnimationActiveTransitionEntry(NextStateIndex, ExistingWeightOfNextState, PreviousStateIndex, ReferenceTransition);
 			NewTransition->InitializeCustomGraphLinks(Context, *(PotentialTransition.TransitionRule));
 
 #if WITH_EDITORONLY_DATA
@@ -338,7 +343,7 @@ void FAnimNode_StateMachine::Update(const FAnimationUpdateContext& Context)
 
 			Context.AnimInstance->AddAnimNotifyFromGeneratedClass(NewTransition->StartNotify);
 			
-			SetState(Context, NextState);
+			SetState(Context, NextStateIndex);
 
 			TransitionCountThisFrame++;
 		}
