Guid in SCS node, better fix for backward compatibility

diff --git a/Engine/Source/Runtime/Engine/Private/SCS_Node.cpp b/Engine/Source/Runtime/Engine/Private/SCS_Node.cpp
--- a/Engine/Source/Runtime/Engine/Private/SCS_Node.cpp
+++ b/Engine/Source/Runtime/Engine/Private/SCS_Node.cpp
@@ -470,17 +470,18 @@ void USCS_Node::PostLoad()
 
 void USCS_Node::ValidateGuid()
 {
-	// Backward compatibility: node requires a guid. 
-	// The guid for the node should be always the same (event when it's not saved). The guid is created using persistent name.
-	if (!VariableGuid.IsValid())
+	// Backward compatibility:
+	// The guid for the node should be always the same (event when it's not saved). 
+	// The guid is created in an deterministic way using persistent name.
+	if (!VariableGuid.IsValid() && (VariableName != NAME_None))
 	{
-		const FName PersistentVariableName = GetVariableName();
-		if (PersistentVariableName != NAME_None)
-		{
-			const FString NameVariableString = PersistentVariableName.ToString();
-			const uint32 PersistentCrc = FCrc::StrCrc32(*NameVariableString);
-			VariableGuid = FGuid(PersistentCrc, 0, 0, 0);
-		}
+		const FString HashString = VariableName.ToString();
+		ensure(HashString.Len());
+
+		const uint32 BufferLength = HashString.Len() * sizeof(HashString[0]);
+		uint32 HashBuffer[5];
+		FSHA1::HashBuffer(*HashString, BufferLength, reinterpret_cast<uint8*>(HashBuffer));
+		VariableGuid = FGuid(HashBuffer[1], HashBuffer[2], HashBuffer[3], HashBuffer[4]);
 	}
 }
 
