We now replicate the component in FHitResult

diff --git a/Engine/Source/Runtime/Engine/Classes/Engine/EngineTypes.h b/Engine/Source/Runtime/Engine/Classes/Engine/EngineTypes.h
--- a/Engine/Source/Runtime/Engine/Classes/Engine/EngineTypes.h
+++ b/Engine/Source/Runtime/Engine/Classes/Engine/EngineTypes.h
@@ -1522,7 +1522,7 @@ struct ENGINE_API FHitResult
 	TWeakObjectPtr<class AActor> Actor;
 
 	/** PrimitiveComponent hit by the trace. */
-	UPROPERTY(NotReplicated)
+	UPROPERTY()
 	TWeakObjectPtr<class UPrimitiveComponent> Component;
 
 	/** Name of bone we hit (for skeletal meshes). */
diff --git a/Engine/Source/Runtime/Engine/Private/Collision/Collision.cpp b/Engine/Source/Runtime/Engine/Private/Collision/Collision.cpp
--- a/Engine/Source/Runtime/Engine/Private/Collision/Collision.cpp
+++ b/Engine/Source/Runtime/Engine/Private/Collision/Collision.cpp
@@ -89,7 +89,7 @@ bool FHitResult::NetSerialize(FArchive& Ar, class UPackageMap* Map, bool& bOutSu
 	Ar << Item;
 	Ar << PhysMaterial;
 	Ar << Actor;
-	// Skipping component on purpose
+	Ar << Component;
 	Ar << BoneName;
 	Ar << FaceIndex;
 
