Cleanup naming, remove unused variables

diff --git a/Engine/Source/Runtime/CoreUObject/Public/UObject/CoreNet.h b/Engine/Source/Runtime/CoreUObject/Public/UObject/CoreNet.h
--- a/Engine/Source/Runtime/CoreUObject/Public/UObject/CoreNet.h
+++ b/Engine/Source/Runtime/CoreUObject/Public/UObject/CoreNet.h
@@ -296,44 +296,37 @@ struct FNetDeltaSerializeInfo
 {
 	FNetDeltaSerializeInfo()
 	{
-		OutBunch = NULL;
-		NewState = NULL;
-		OldState = NULL;
-		Map = NULL;
-		Data = NULL;
+		Writer		= NULL;
+		Reader		= NULL;
 
-		Struct = NULL;
-		InArchive = NULL;
+		NewState	= NULL;
+		OldState	= NULL;
+		Map			= NULL;
+		Data		= NULL;
 
-		NumPotentialBits = 0;
-		NumBytesTouched = 0;
+		Struct		= NULL;
 
 		NetSerializeCB = NULL;
 	}
 
-	// Used for ggeneric TArray replication
-	FBitWriter * OutBunch;
-	
-	TSharedPtr<INetDeltaBaseState> *NewState;		// SharedPtr to new base state created by NetDeltaSerialize.
-	INetDeltaBaseState *  OldState;				// Pointer to the previous base state.
-	UPackageMap *	Map;
-	void* Data;
+	// Used when writing
+	FBitWriter*						Writer;
 
-	// Only used for fast TArray replication
-	UStruct *Struct;
+	// Used for when reading
+	FBitReader*						Reader;
 
-	// Used for TArray replication reading
-	FArchive *InArchive;
+	TSharedPtr<INetDeltaBaseState>*	NewState;		// SharedPtr to new base state created by NetDeltaSerialize.
+	INetDeltaBaseState*				OldState;				// Pointer to the previous base state.
+	UPackageMap*					Map;
+	void*							Data;
 
-	// Number of bits touched during full state serialization
-	int32 NumPotentialBits;
-	// Number of bytes touched (will be larger usually due to forced byte alignment)
-	int32 NumBytesTouched;
+	// Only used for fast TArray replication
+	UStruct*						Struct;
 
-	INetSerializeCB * NetSerializeCB;
+	INetSerializeCB*				NetSerializeCB;
 
 	// Debugging variables
-	FString DebugName;
+	FString							DebugName;
 };
 
 /**
diff --git a/Engine/Source/Runtime/Engine/Classes/Engine/NetSerialization.h b/Engine/Source/Runtime/Engine/Classes/Engine/NetSerialization.h
--- a/Engine/Source/Runtime/Engine/Classes/Engine/NetSerialization.h
+++ b/Engine/Source/Runtime/Engine/Classes/Engine/NetSerialization.h
@@ -374,13 +374,13 @@ bool FFastArraySerializer::FastArrayDeltaSerialize( TArray<Type> &Items, FNetDel
 {
 	class UScriptStruct* InnerStruct = Type::StaticStruct();
 
-	if (Parms.OutBunch)
+	if (Parms.Writer)
 	{
 		//-----------------------------
 		// Saving
 		//-----------------------------	
 		check(Parms.Struct);
-		FBitWriter& OutBunch = *Parms.OutBunch;
+		FBitWriter& Writer = *Parms.Writer;
 
 		// Create a new map from the current state of the array		
 		FNetFastTArrayBaseState * NewState = new FNetFastTArrayBaseState();
@@ -485,10 +485,10 @@ bool FFastArraySerializer::FastArrayDeltaSerialize( TArray<Type> &Items, FNetDel
 		//----------------------
 
 		uint32 ElemNum = DeletedElements.Num();
-		OutBunch << ElemNum;
+		Writer << ElemNum;
 
 		ElemNum = ChangedElements.Num();
-		OutBunch << ElemNum;
+		Writer << ElemNum;
 
 		UE_LOG(LogNetSerialization, Log, TEXT("   Writing Bunch. NumChange: %d. NumDel: %d"), ChangedElements.Num(), DeletedElements.Num() );
 
@@ -496,7 +496,7 @@ bool FFastArraySerializer::FastArrayDeltaSerialize( TArray<Type> &Items, FNetDel
 		for (auto It = DeletedElements.CreateIterator(); It; ++It)
 		{
 			int32 ID = *It;
-			OutBunch << ID;
+			Writer << ID;
 			UE_LOG(LogNetSerialization, Log, TEXT("   Deleted ElementID: %d"), ID);
 		}
 
@@ -507,13 +507,13 @@ bool FFastArraySerializer::FastArrayDeltaSerialize( TArray<Type> &Items, FNetDel
 
 			// Dont pack this, want property to be byte aligned
 			uint32 ID = It->ID;
-			OutBunch << ID;
+			Writer << ID;
 
 			UE_LOG(LogNetSerialization, Log, TEXT("   Changed ElementID: %d"), ID);
 
 			bool bHasUnmapped = false;
 
-			Parms.NetSerializeCB->NetSerializeStruct(InnerStruct, OutBunch, Parms.Map, ThisElement, bHasUnmapped);
+			Parms.NetSerializeCB->NetSerializeStruct(InnerStruct, Writer, Parms.Map, ThisElement, bHasUnmapped);
 
 			if (bHasUnmapped)
 			{
@@ -530,8 +530,8 @@ bool FFastArraySerializer::FastArrayDeltaSerialize( TArray<Type> &Items, FNetDel
 		//-----------------------------
 		// Loading
 		//-----------------------------	
-		check(Parms.InArchive);
-		FArchive &Ar = *Parms.InArchive;
+		check(Parms.Reader);
+		FBitReader& Reader = *Parms.Reader;
 
 		//---------------
 		// Build ItemMap if necessary. This maps ReplicationID to our local index into the Items array.
@@ -555,22 +555,22 @@ bool FFastArraySerializer::FastArrayDeltaSerialize( TArray<Type> &Items, FNetDel
 		// Read header
 		//---------------
 		uint32 NumDeletes;
-		Ar << NumDeletes;
+		Reader << NumDeletes;
 	
 		if ( NumDeletes > MAX_NUM_DELETED )
 		{
 			UE_LOG( LogNetSerialization, Warning, TEXT( "NumDeletes > MAX_NUM_DELETED: %d." ), NumDeletes );
-			Ar.SetError();
+			Reader.SetError();
 			return false;;
 		}
 
 		uint32 NumChanged;
-		Ar << NumChanged;
+		Reader << NumChanged;
 
 		if (NumChanged > MAX_NUM_CHANGED)
 		{
 			UE_LOG(LogNetSerialization, Warning, TEXT("NumChanged > MAX_NUM_CHANGED: %d."), NumChanged);
-			Ar.SetError();
+			Reader.SetError();
 			return false;;
 		}
 
@@ -586,7 +586,7 @@ bool FFastArraySerializer::FastArrayDeltaSerialize( TArray<Type> &Items, FNetDel
 			for (uint32 i = 0; i < NumDeletes; ++i)
 			{
 				int32 ElementID;
-				Ar << ElementID;
+				Reader << ElementID;
 
 				int32* ElementIndexPtr = ArraySerializer.ItemMap.Find(ElementID);
 				if (ElementIndexPtr)
@@ -613,7 +613,7 @@ bool FFastArraySerializer::FastArrayDeltaSerialize( TArray<Type> &Items, FNetDel
 		for(uint32 i=0; i < NumChanged; ++i)
 		{
 			int32 ElementID;
-			Ar << ElementID;
+			Reader << ElementID;
 
 			int32* ElementIndexPtr = ArraySerializer.ItemMap.Find(ElementID);
 			int32	ElementIndex = 0;
@@ -639,13 +639,13 @@ bool FFastArraySerializer::FastArrayDeltaSerialize( TArray<Type> &Items, FNetDel
 				ElementIndex = *ElementIndexPtr;
 				ThisElement = &Items[ElementIndex];
 			}
-			
+
 			bool bHasUnmapped = false;
-			Parms.NetSerializeCB->NetSerializeStruct( InnerStruct, Ar, Parms.Map, ThisElement, bHasUnmapped );
+			Parms.NetSerializeCB->NetSerializeStruct( InnerStruct, Reader, Parms.Map, ThisElement, bHasUnmapped );
 
-			if ( Ar.IsError() )
+			if ( Reader.IsError() )
 			{
-				UE_LOG( LogNetSerialization, Warning, TEXT( "Parms.NetSerializeCB->NetSerializeStruct: Ar.IsError() == true" ) );
+				UE_LOG( LogNetSerialization, Warning, TEXT( "Parms.NetSerializeCB->NetSerializeStruct: Reader.IsError() == true" ) );
 				return false;
 			}
 
diff --git a/Engine/Source/Runtime/Engine/Private/DataReplication.cpp b/Engine/Source/Runtime/Engine/Private/DataReplication.cpp
--- a/Engine/Source/Runtime/Engine/Private/DataReplication.cpp
+++ b/Engine/Source/Runtime/Engine/Private/DataReplication.cpp
@@ -79,7 +79,7 @@ bool FObjectReplicator::SerializeCustomDeltaProperty( UNetConnection * Connectio
 
 	FNetSerializeCB NetSerializeCB( Connection->Driver );
 
-	Parms.OutBunch			= &OutBunch;
+	Parms.Writer			= &OutBunch;
 	Parms.Map				= Connection->PackageMap;
 	Parms.OldState			= OldState.Get();
 	Parms.NewState			= &NewFullState;
@@ -542,7 +542,7 @@ bool FObjectReplicator::ReceivedBunch( FInBunch &Bunch, const FReplicationFlags&
 				Parms.DebugName			= StructProperty->GetName();
 				Parms.Struct			= InnerStruct;
 				Parms.Map				= PackageMap;
-				Parms.InArchive			= &Bunch;
+				Parms.Reader			= &Bunch;
 				Parms.NetSerializeCB	= &NetSerializeCB;
 
 				// Call the custom delta serialize function to handle it
