#UE4 Nav mesh rendering component will no longer crash on dedicated servers. In fact it will no longer even load in dedicated servers. Also, now in cases where we would have crashed due to a lack of viewport, the component is simply not rendered instead (might happen in clients with nullrhi)

diff --git a/Engine/Source/Runtime/Engine/Classes/AI/Navigation/NavMeshRenderingComponent.h b/Engine/Source/Runtime/Engine/Classes/AI/Navigation/NavMeshRenderingComponent.h
--- a/Engine/Source/Runtime/Engine/Classes/AI/Navigation/NavMeshRenderingComponent.h
+++ b/Engine/Source/Runtime/Engine/Classes/AI/Navigation/NavMeshRenderingComponent.h
@@ -9,6 +9,10 @@ class ENGINE_API UNavMeshRenderingComponent : public UPrimitiveComponent
 	GENERATED_UCLASS_BODY()
 
 public:
+	// Begin UObject Interface
+	virtual void NeedsLoadForServer() override;
+	// End UObject Interface
+	
 	// Begin UPrimitiveComponent Interface
 	virtual FPrimitiveSceneProxy* CreateSceneProxy() override;
 	virtual void OnRegister()  override;
diff --git a/Engine/Source/Runtime/Engine/Private/NavMeshRenderingComponent.cpp b/Engine/Source/Runtime/Engine/Private/NavMeshRenderingComponent.cpp
--- a/Engine/Source/Runtime/Engine/Private/NavMeshRenderingComponent.cpp
+++ b/Engine/Source/Runtime/Engine/Private/NavMeshRenderingComponent.cpp
@@ -28,15 +28,21 @@ UNavMeshRenderingComponent::UNavMeshRenderingComponent(const FObjectInitializer&
 	bCollectNavigationData = false;
 }
 
+bool UNavMeshRenderingComponent::NeedsLoadForServer() const
+{
+	// This rendering component is not needed in dedicated server builds
+	return false;
+}
+
 void UNavMeshRenderingComponent::TimerFunction()
 {
 	bool bShowNavigation = false;
 #if WITH_EDITOR
 	if (GEditor)
 	{
-		for (FEditorViewportClient* CurrentViewpoer : GEditor->AllViewportClients)
+		for (FEditorViewportClient* CurrentViewport : GEditor->AllViewportClients)
 		{
-			if (CurrentViewpoer->EngineShowFlags.Navigation)
+			if (CurrentViewport && CurrentViewport->EngineShowFlags.Navigation)
 			{
 				bShowNavigation = true;
 				break;
@@ -46,7 +52,9 @@ void UNavMeshRenderingComponent::TimerFunction()
 	else
 #endif //WITH_EDITOR
 	{
-		bShowNavigation = GetWorld()->GetGameViewport()->EngineShowFlags.Navigation;
+		UWorld* World = GetWorld();
+		UGameViewportClient* Viewport = World ? World->GetGameViewport() : nullptr;
+		bShowNavigation = Viewport ? Viewport->EngineShowFlags.Navigation : false;
 	}
 
 	if (bShowNavigation != !!bCollectNavigationData)
