#ANIM: Look at control interpolation/clamp support.

Merge from CL 2414755

diff --git a/Engine/Source/Runtime/Engine/Classes/Animation/BoneControllers/AnimNode_LookAt.h b/Engine/Source/Runtime/Engine/Classes/Animation/BoneControllers/AnimNode_LookAt.h
--- a/Engine/Source/Runtime/Engine/Classes/Animation/BoneControllers/AnimNode_LookAt.h
+++ b/Engine/Source/Runtime/Engine/Classes/Animation/BoneControllers/AnimNode_LookAt.h
@@ -19,6 +19,24 @@ namespace EAxisOption
 	};
 }
 
+
+UENUM()
+/** Various ways to interpolate TAlphaBlend. */
+namespace EInterpolationBlend
+{
+	enum Type
+	{
+		Linear,
+		Cubic,
+		Sinusoidal,
+		EaseInOutExponent2,
+		EaseInOutExponent3,
+		EaseInOutExponent4,
+		EaseInOutExponent5,
+		MAX
+	};
+}
+
 /**
  *	Simple controller that make a bone to look at the point or another bone
  */
@@ -43,14 +61,38 @@ struct ENGINE_API FAnimNode_LookAt : public FAnimNode_SkeletalControlBase
 	UPROPERTY(EditAnywhere, BlueprintReadWrite, Category=SkeletalControl) 
 	TEnumAsByte<EAxisOption::Type>	LookAtAxis;
 
+	/** Look at Clamp value in degree - if you're look at axis is Z, only X, Y degree of clamp will be used*/
+	UPROPERTY(EditAnywhere, BlueprintReadWrite, Category=SkeletalControl, meta=(PinHiddenByDefault))
+	float LookAtClamp;
+
+	UPROPERTY(EditAnywhere, BlueprintReadWrite, Category=SkeletalControl, meta=(PinHiddenByDefault))
+	TEnumAsByte<EInterpolationBlend::Type>	InterpolationType;
+
+	UPROPERTY(EditAnywhere, BlueprintReadWrite, Category=SkeletalControl, meta=(PinHiddenByDefault))
+	float	InterpolationTime;
+
+	UPROPERTY(EditAnywhere, BlueprintReadWrite, Category=SkeletalControl, meta=(PinHiddenByDefault))
+	float	InterpolationTriggerThreashold;
+
+	UPROPERTY(EditAnywhere, Category=SkeletalControl)
+	bool	bEnableDebug;
+
 	/** Debug transient data */
 	FVector CurrentLookAtLocation;
 
+	/** Current Target Location */
+	FVector CurrentTargetLocation;
+	FVector PreviousTargetLocation;
+
+	/** Current Alpha */
+	float AccumulatedInterpoolationTime;
+
 	// in the future, it would be nice to have more options, -i.e. lag, interpolation speed
 	FAnimNode_LookAt();
 
 	// FAnimNode_Base interface
 	virtual void GatherDebugData(FNodeDebugData& DebugData) override;
+	virtual void Update(const FAnimationUpdateContext& Context) override;
 	// End of FAnimNode_Base interface
 
 	// FAnimNode_SkeletalControlBase interface
@@ -64,4 +106,25 @@ private:
 	// End of FAnimNode_SkeletalControlBase interface
 
 	FVector GetAlignVector(FTransform& Transform, EAxisOption::Type AxisOption);
+
+	EAlphaBlendType GetInterpolationType()
+	{
+		switch (InterpolationType)
+		{
+		case EInterpolationBlend::Cubic:
+			return ABT_Cubic;
+		case EInterpolationBlend::Sinusoidal:
+			return ABT_Cubic;
+		case EInterpolationBlend::EaseInOutExponent2:
+			return ABT_EaseInOutExponent2;
+		case EInterpolationBlend::EaseInOutExponent3:
+			return ABT_EaseInOutExponent3;
+		case EInterpolationBlend::EaseInOutExponent4:
+			return ABT_EaseInOutExponent4;
+		case EInterpolationBlend::EaseInOutExponent5:
+			return ABT_EaseInOutExponent5;
+		}
+
+		return ABT_Linear;
+	}
 };
diff --git a/Engine/Source/Runtime/Engine/Private/Animation/AnimNode_LookAt.cpp b/Engine/Source/Runtime/Engine/Private/Animation/AnimNode_LookAt.cpp
--- a/Engine/Source/Runtime/Engine/Private/Animation/AnimNode_LookAt.cpp
+++ b/Engine/Source/Runtime/Engine/Private/Animation/AnimNode_LookAt.cpp
@@ -10,7 +10,11 @@
 FAnimNode_LookAt::FAnimNode_LookAt()
 	: LookAtLocation(FVector::ZeroVector)
 	, LookAtAxis(EAxis::X)
+	, LookAtClamp(0.f)
+	, InterpolationTime(0.f)
+	, InterpolationTriggerThreashold(0.f)
 	, CurrentLookAtLocation(FVector::ZeroVector)
+	, AccumulatedInterpoolationTime(0.f)
 {
 }
 
@@ -34,6 +38,14 @@ void FAnimNode_LookAt::GatherDebugData(FNodeDebugData& DebugData)
 	ComponentPose.GatherDebugData(DebugData);
 }
 
+void DrawDebugData(UWorld* World, const FTransform& TransformVector, const FVector& StartLoc, const FVector& TargetLoc, FColor Color)
+{
+	FVector Start = TransformVector.TransformPosition(StartLoc);
+	FVector End = TransformVector.TransformPosition(TargetLoc);
+
+	DrawDebugLine(World, Start, End, Color);
+}
+
 void FAnimNode_LookAt::EvaluateBoneTransforms(USkeletalMeshComponent* SkelComp, const FBoneContainer& RequiredBones, FA2CSPose& MeshBases, TArray<FBoneTransform>& OutBoneTransforms)
 {
 	check(OutBoneTransforms.Num() == 0);
@@ -52,7 +64,49 @@ void FAnimNode_LookAt::EvaluateBoneTransforms(USkeletalMeshComponent* SkelComp,
 		TargetLocationInComponentSpace = SkelComp->ComponentToWorld.InverseTransformPosition(LookAtLocation);
 	}
 	
-	CurrentLookAtLocation = TargetLocationInComponentSpace;
+	FVector OldCurrentTargetLocation = CurrentTargetLocation;
+	FVector NewCurrentTargetLocation = TargetLocationInComponentSpace;
+
+	if ((NewCurrentTargetLocation - OldCurrentTargetLocation).SizeSquared() > InterpolationTriggerThreashold*InterpolationTriggerThreashold)
+	{
+		if (AccumulatedInterpoolationTime >= InterpolationTime)
+		{
+			// reset current Alpha, we're starting to move
+			AccumulatedInterpoolationTime = 0.f;
+		}
+
+		PreviousTargetLocation = OldCurrentTargetLocation;
+		CurrentTargetLocation = NewCurrentTargetLocation;
+	}
+	else if (InterpolationTriggerThreashold == 0.f)
+	{
+		CurrentTargetLocation = NewCurrentTargetLocation;
+	}
+
+	if (InterpolationTime > 0.f)
+	{
+		float CurrentAlpha = AccumulatedInterpoolationTime/InterpolationTime;
+
+		if (CurrentAlpha < 1.f)
+		{
+			float BlendAlpha = AlphaToBlendType(CurrentAlpha, GetInterpolationType());
+
+			CurrentLookAtLocation = FMath::Lerp(PreviousTargetLocation, CurrentTargetLocation, BlendAlpha);
+		}
+	}
+	else
+	{
+		CurrentLookAtLocation = CurrentTargetLocation;
+	}
+
+	if (bEnableDebug)
+	{
+		UWorld* World = SkelComp->GetWorld();
+
+		DrawDebugData(World, SkelComp->GetComponentToWorld(), ComponentBoneTransform.GetLocation(), PreviousTargetLocation, FColor(0, 255, 0));
+		DrawDebugData(World, SkelComp->GetComponentToWorld(), ComponentBoneTransform.GetLocation(), CurrentTargetLocation, FColor(255, 0, 0));
+		DrawDebugData(World, SkelComp->GetComponentToWorld(), ComponentBoneTransform.GetLocation(), CurrentLookAtLocation, FColor(0, 0, 255));
+	}
 
 	// lookat vector
 	FVector LookAtVector = GetAlignVector(ComponentBoneTransform, LookAtAxis);
@@ -65,9 +119,27 @@ void FAnimNode_LookAt::EvaluateBoneTransforms(USkeletalMeshComponent* SkelComp,
 		ToTarget *= -1.f;
 	}
 	
+	if ( LookAtClamp > ZERO_ANIMWEIGHT_THRESH )
+	{
+		float LookAtClampInRadians = FMath::DegreesToRadians(LookAtClamp);
+		float DiffAngle = FMath::Acos(FVector::DotProduct(LookAtVector, ToTarget));
+		if (LookAtClampInRadians > 0.f && DiffAngle > LookAtClampInRadians)
+		{
+			FVector OldToTarget = ToTarget;
+			FVector DeltaTarget = ToTarget-LookAtVector;
+
+			float Ratio = LookAtClampInRadians/DiffAngle;
+			DeltaTarget *= Ratio;
+
+			ToTarget = LookAtVector + DeltaTarget;
+			ToTarget.Normalize();
+//			UE_LOG(LogAnimation, Warning, TEXT("Recalculation required - old target %f, new target %f"), 
+//				FMath::RadiansToDegrees(FMath::Acos(FVector::DotProduct(LookAtVector, OldToTarget))), FMath::RadiansToDegrees(FMath::Acos(FVector::DotProduct(LookAtVector, ToTarget))));
+		}
+	}
+
 	// get delta rotation
 	FQuat DeltaRot = FQuat::FindBetween(LookAtVector, ToTarget);
-
 	// transform current rotation to delta rotation
 	FQuat CurrentRot = ComponentBoneTransform.GetRotation();
 	FQuat NewRotation = DeltaRot * CurrentRot;
@@ -108,4 +180,11 @@ FVector FAnimNode_LookAt::GetAlignVector(FTransform& Transform, EAxisOption::Typ
 	}
 
 	return FVector(1.f, 0.f, 0.f);
-}
\ No newline at end of file
+}
+
+void FAnimNode_LookAt::Update(const FAnimationUpdateContext& Context)
+{
+	FAnimNode_SkeletalControlBase::Update(Context);
+
+	AccumulatedInterpoolationTime = FMath::Clamp(AccumulatedInterpoolationTime+Context.GetDeltaTime(), 0.f, InterpolationTime);;
+}
