#UE4 Fix minor build error

diff --git a/Engine/Source/Runtime/Engine/Classes/AI/Navigation/NavMeshRenderingComponent.h b/Engine/Source/Runtime/Engine/Classes/AI/Navigation/NavMeshRenderingComponent.h
--- a/Engine/Source/Runtime/Engine/Classes/AI/Navigation/NavMeshRenderingComponent.h
+++ b/Engine/Source/Runtime/Engine/Classes/AI/Navigation/NavMeshRenderingComponent.h
@@ -10,7 +10,7 @@ class ENGINE_API UNavMeshRenderingComponent : public UPrimitiveComponent
 
 public:
 	// Begin UObject Interface
-	virtual void NeedsLoadForServer() override;
+	virtual bool NeedsLoadForServer() override;
 	// End UObject Interface
 	
 	// Begin UPrimitiveComponent Interface
