Only load highres screenshot helper materials in editor builds

diff --git a/Engine/Source/Runtime/Engine/Private/HighResScreenshot.cpp b/Engine/Source/Runtime/Engine/Private/HighResScreenshot.cpp
--- a/Engine/Source/Runtime/Engine/Private/HighResScreenshot.cpp
+++ b/Engine/Source/Runtime/Engine/Private/HighResScreenshot.cpp
@@ -32,6 +32,7 @@ void FHighResScreenshotConfig::Init()
 		ImageCompressorHDR = ImageWrapperModule->CreateImageWrapper(EImageFormat::EXR);
 	}
 
+#if WITH_EDITOR
 	HighResScreenshotMaterial = LoadObject<UMaterial>(NULL, TEXT("/Engine/EngineMaterials/HighResScreenshot.HighResScreenshot"));
 	HighResScreenshotMaskMaterial = LoadObject<UMaterial>(NULL, TEXT("/Engine/EngineMaterials/HighResScreenshotMask.HighResScreenshotMask"));
 	HighResScreenshotCaptureRegionMaterial = LoadObject<UMaterial>(NULL, TEXT("/Engine/EngineMaterials/HighResScreenshotCaptureRegion.HighResScreenshotCaptureRegion"));
@@ -48,6 +49,7 @@ void FHighResScreenshotConfig::Init()
 	{
 		HighResScreenshotCaptureRegionMaterial->AddToRoot();
 	}
+#endif
 }
 
 void FHighResScreenshotConfig::ChangeViewport(TWeakPtr<FSceneViewport> InViewport)
