shape components can now export themselves to navigation as dynamic obstacles, added new area class for obstacles UE-7605

diff --git a/Engine/Source/Runtime/Engine/Classes/AI/Navigation/NavAreas/NavArea_Obstacle.h b/Engine/Source/Runtime/Engine/Classes/AI/Navigation/NavAreas/NavArea_Obstacle.h
new file mode 100644
--- /dev/null
+++ b/Engine/Source/Runtime/Engine/Classes/AI/Navigation/NavAreas/NavArea_Obstacle.h
@@ -0,0 +1,12 @@
+// Copyright 1998-2015 Epic Games, Inc. All Rights Reserved.
+
+#pragma once
+#include "AI/Navigation/NavAreas/NavArea.h"
+#include "NavArea_Obstacle.generated.h"
+
+/** In general represents a high cost area, that shouldn't be traversed by anyone unless no other path exist.*/
+UCLASS(Config = Engine)
+class ENGINE_API UNavArea_Obstacle : public UNavArea
+{
+	GENERATED_UCLASS_BODY()
+};
diff --git a/Engine/Source/Runtime/Engine/Classes/Components/ShapeComponent.h b/Engine/Source/Runtime/Engine/Classes/Components/ShapeComponent.h
--- a/Engine/Source/Runtime/Engine/Classes/Components/ShapeComponent.h
+++ b/Engine/Source/Runtime/Engine/Classes/Components/ShapeComponent.h
@@ -31,10 +31,20 @@ class ENGINE_API UShapeComponent : public UPrimitiveComponent
 	UPROPERTY()
 	uint32 bShouldCollideWhenPlacing:1;
 
+	/** If set, shape will be exported for navigation as dynamic modifier instead of using regular collision data */
+	UPROPERTY(EditAnywhere, Category = Navigation)
+	uint32 bDynamicObstacle : 1;
+
+	/** Navigation area type (empty = default obstacle) */
+	UPROPERTY(EditAnywhere, Category = Navigation)
+	TSubclassOf<class UNavArea> AreaClass;
+
 	// Begin UPrimitiveComponent interface.
 	virtual FPrimitiveSceneProxy* CreateSceneProxy() override;
 	virtual void GetUsedMaterials( TArray<UMaterialInterface*>& OutMaterials ) const override; 
 	virtual class UBodySetup* GetBodySetup() override;
+	virtual bool DoCustomNavigableGeometryExport(struct FNavigableGeometryExport* GeomExport) const override;
+	virtual void GetNavigationData(FNavigationRelevantData& Data) const override;
 	// End UPrimitiveComponent interface.
 
 	// Begin USceneComponent interface
diff --git a/Engine/Source/Runtime/Engine/Private/AI/Navigation/NavAreas/NavArea_Obstacle.cpp b/Engine/Source/Runtime/Engine/Private/AI/Navigation/NavAreas/NavArea_Obstacle.cpp
new file mode 100644
--- /dev/null
+++ b/Engine/Source/Runtime/Engine/Private/AI/Navigation/NavAreas/NavArea_Obstacle.cpp
@@ -0,0 +1,10 @@
+// Copyright 1998-2015 Epic Games, Inc. All Rights Reserved.
+
+#include "EnginePrivate.h"
+#include "AI/Navigation/NavAreas/NavArea_Obstacle.h"
+
+UNavArea_Obstacle::UNavArea_Obstacle(const FObjectInitializer& ObjectInitializer) : Super(ObjectInitializer)
+{
+	DrawColor = FColor(127, 51, 0);	// brownish
+	DefaultCost = 1000000.0f;
+}
diff --git a/Engine/Source/Runtime/Engine/Private/AI/Navigation/NavigationModifier.cpp b/Engine/Source/Runtime/Engine/Private/AI/Navigation/NavigationModifier.cpp
--- a/Engine/Source/Runtime/Engine/Private/AI/Navigation/NavigationModifier.cpp
+++ b/Engine/Source/Runtime/Engine/Private/AI/Navigation/NavigationModifier.cpp
@@ -11,6 +11,7 @@
 #include "AI/Navigation/NavLinkDefinition.h"
 #include "AI/Navigation/NavLinkTrivial.h"
 #include "AI/Navigation/NavAreas/NavAreaMeta.h"
+#include "Components/PrimitiveComponent.h"
 #include "Components/BrushComponent.h"
 
 // if square distance between two points is less than this the those points
@@ -578,6 +579,35 @@ FCompositeNavModifier FCompositeNavModifier::GetInstantiatedMetaModifier(const F
 	return Result;
 }
 
+void FCompositeNavModifier::CreateAreaModifiers(const UPrimitiveComponent* PrimComp, const TSubclassOf<UNavArea> AreaClass)
+{
+	UBodySetup* BodySetup = PrimComp ? ((UPrimitiveComponent*)PrimComp)->GetBodySetup() : nullptr;
+	if (BodySetup == nullptr)
+	{
+		return;
+	}
+
+	for (int32 Idx = 0; Idx < BodySetup->AggGeom.BoxElems.Num(); Idx++)
+	{
+		const FKBoxElem& BoxElem = BodySetup->AggGeom.BoxElems[Idx];
+		const FBox BoxSize = BoxElem.CalcAABB(FTransform::Identity, 1.0f);
+
+		FAreaNavModifier AreaMod(BoxSize, PrimComp->ComponentToWorld, AreaClass);
+		Add(AreaMod);
+	}
+
+	for (int32 Idx = 0; Idx < BodySetup->AggGeom.SphylElems.Num(); Idx++)
+	{
+		const FKSphylElem& SphylElem = BodySetup->AggGeom.SphylElems[Idx];
+		const FTransform AreaOffset(FVector(0, 0, -SphylElem.Length));
+
+		FAreaNavModifier AreaMod(SphylElem.Radius, SphylElem.Length, AreaOffset * PrimComp->ComponentToWorld, AreaClass);
+		Add(AreaMod);
+	}
+
+	// convex elements support?
+}
+
 uint32 FCompositeNavModifier::GetAllocatedSize() const
 {
 	uint32 MemUsed = Areas.GetAllocatedSize() + SimpleLinks.GetAllocatedSize() + CustomLinks.GetAllocatedSize();
diff --git a/Engine/Source/Runtime/Engine/Private/AI/Navigation/NavigationSystem.cpp b/Engine/Source/Runtime/Engine/Private/AI/Navigation/NavigationSystem.cpp
--- a/Engine/Source/Runtime/Engine/Private/AI/Navigation/NavigationSystem.cpp
+++ b/Engine/Source/Runtime/Engine/Private/AI/Navigation/NavigationSystem.cpp
@@ -2164,11 +2164,6 @@ void UNavigationSystem::AddElementToNavOctree(const FNavigationDirtyElement& Dir
 	const FBox BBox = GeneratedData.Bounds.GetBox();
 	const bool bValidBBox = BBox.IsValid && !BBox.GetSize().IsNearlyZero();
 
-	if (BBox.GetExtent().X > 400000)
-	{
-		volatile int32 i = 0;
-	}
-
 	if (bValidBBox && !GeneratedData.IsEmpty())
 	{
 		const int32 DirtyFlag = DirtyElement.FlagsOverride ? DirtyElement.FlagsOverride : GeneratedData.Data.GetDirtyFlag();
diff --git a/Engine/Source/Runtime/Engine/Private/Character.cpp b/Engine/Source/Runtime/Engine/Private/Character.cpp
--- a/Engine/Source/Runtime/Engine/Private/Character.cpp
+++ b/Engine/Source/Runtime/Engine/Private/Character.cpp
@@ -55,6 +55,7 @@ ACharacter::ACharacter(const FObjectInitializer& ObjectInitializer)
 	CapsuleComponent->bShouldUpdatePhysicsVolume = true;
 	CapsuleComponent->bCheckAsyncSceneOnMove = false;	
 	CapsuleComponent->bCanEverAffectNavigation = false;
+	CapsuleComponent->bDynamicObstacle = true;
 	RootComponent = CapsuleComponent;
 
 	JumpKeyHoldTime = 0.0f;
diff --git a/Engine/Source/Runtime/Engine/Private/DefaultPawn.cpp b/Engine/Source/Runtime/Engine/Private/DefaultPawn.cpp
--- a/Engine/Source/Runtime/Engine/Private/DefaultPawn.cpp
+++ b/Engine/Source/Runtime/Engine/Private/DefaultPawn.cpp
@@ -33,6 +33,7 @@ ADefaultPawn::ADefaultPawn(const FObjectInitializer& ObjectInitializer)
 	CollisionComponent->CanCharacterStepUpOn = ECB_No;
 	CollisionComponent->bShouldUpdatePhysicsVolume = true;
 	CollisionComponent->bCanEverAffectNavigation = false;
+	CollisionComponent->bDynamicObstacle = true;
 
 	RootComponent = CollisionComponent;
 
diff --git a/Engine/Source/Runtime/Engine/Private/ShapeComponent.cpp b/Engine/Source/Runtime/Engine/Private/ShapeComponent.cpp
--- a/Engine/Source/Runtime/Engine/Private/ShapeComponent.cpp
+++ b/Engine/Source/Runtime/Engine/Private/ShapeComponent.cpp
@@ -2,6 +2,7 @@
 
 
 #include "EnginePrivate.h"
+#include "AI/Navigation/NavAreas/NavArea_Obstacle.h"
 
 UShapeComponent::UShapeComponent(const FObjectInitializer& ObjectInitializer)
 	: Super(ObjectInitializer)
@@ -18,7 +19,11 @@ UShapeComponent::UShapeComponent(const FObjectInitializer& ObjectInitializer)
 	bCastDynamicShadow = false;
 	ShapeColor = FColor(223, 149, 157, 255);
 	bShouldCollideWhenPlacing = false;
+	
+	bHasCustomNavigableGeometry = EHasCustomNavigableGeometry::Yes;
 	bCanEverAffectNavigation = true;
+	bDynamicObstacle = false;
+	AreaClass = UNavArea_Obstacle::StaticClass();
 }
 
 FPrimitiveSceneProxy* UShapeComponent::CreateSceneProxy()
@@ -59,4 +64,17 @@ void UShapeComponent::PostEditChangeProperty(FPropertyChangedEvent& PropertyChan
 
 	Super::PostEditChangeProperty(PropertyChangedEvent);
 }
-#endif // WITH_EDITOR
\ No newline at end of file
+#endif // WITH_EDITOR
+
+bool UShapeComponent::DoCustomNavigableGeometryExport(struct FNavigableGeometryExport* GeomExport) const
+{
+	return !bDynamicObstacle;
+}
+
+void UShapeComponent::GetNavigationData(FNavigationRelevantData& Data) const
+{
+	if (bDynamicObstacle)
+	{
+		Data.Modifiers.CreateAreaModifiers(this, AreaClass);
+	}
+}
diff --git a/Engine/Source/Runtime/Engine/Public/AI/NavigationModifier.h b/Engine/Source/Runtime/Engine/Public/AI/NavigationModifier.h
--- a/Engine/Source/Runtime/Engine/Public/AI/NavigationModifier.h
+++ b/Engine/Source/Runtime/Engine/Public/AI/NavigationModifier.h
@@ -6,6 +6,7 @@
 
 class UNavArea;
 class UNavLinkDefinition;
+class UPrimitiveComponent;
 class UBrushComponent;
 
 struct ENGINE_API FNavigationModifier
@@ -238,6 +239,8 @@ struct ENGINE_API FCompositeNavModifier : public FNavigationModifier
 		bAdjustHeight |= Modifiers.HasAgentHeightAdjust();
 	}
 
+	void CreateAreaModifiers(const UPrimitiveComponent* PrimComp, const TSubclassOf<UNavArea> AreaClass);
+
 	FORCEINLINE const TArray<FAreaNavModifier>& GetAreas() const { return Areas; }
 	FORCEINLINE const TArray<FSimpleLinkNavModifier>& GetSimpleLinks() const { return SimpleLinks; }
 	FORCEINLINE const TArray<FCustomLinkNavModifier>& GetCustomLinks() const { return CustomLinks; }
