Fix UE-6190 : UVs skewed in Mobile Preview mode for some skeletal meshes
Fix for CPU skin code creating mismatching render verts. CPU skin code no longer selects vert type based on per chunk  bone influences.

diff --git a/Engine/Source/Runtime/Engine/Private/SkeletalRenderCPUSkin.cpp b/Engine/Source/Runtime/Engine/Private/SkeletalRenderCPUSkin.cpp
--- a/Engine/Source/Runtime/Engine/Private/SkeletalRenderCPUSkin.cpp
+++ b/Engine/Source/Runtime/Engine/Private/SkeletalRenderCPUSkin.cpp
@@ -614,11 +614,9 @@ const VectorRegister		VECTOR4_UNPACK_MINUS_1	= DECLARE_VECTOR_REGISTER(-1.f, -1.
 
 const VectorRegister		VECTOR_0001				= DECLARE_VECTOR_REGISTER(0.f, 0.f, 0.f, 1.f);
 
-template<int32 MaxBoneInfluences, typename BaseVertexType, typename VertexType>
+template<bool bExtraBoneInfluences, int32 MaxChunkBoneInfluences, typename BaseVertexType, typename VertexType>
 static void SkinVertexChunk( FFinalSkinVertex*& DestVertex, TArray<FVertexAnimEvalInfo>& AnimEvalInfos, const FSkelMeshChunk& Chunk, const FStaticLODModel &LOD, int32 VertexBufferBaseIndex, uint32 NumValidMorphs, int32 &CurBaseVertIdx, int32 LODIndex, int32 RigidInfluenceIndex, const FMatrix* RESTRICT ReferenceToLocal )
 {
-	const bool bExtraBoneInfluences = (MaxBoneInfluences > 4);
-
 	// VertexCopy for morph. Need to allocate right struct
 	// To avoid re-allocation, create 2 statics, and assign right struct
 	VertexType  VertexCopy;
@@ -730,7 +728,7 @@ static void SkinVertexChunk( FFinalSkinVertex*& DestVertex, TArray<FVertexAnimEv
 			SrcNormals[2] = Unpack4( &MorphedVertex->TangentZ.Vector.Packed );
 			VectorRegister Weights = VectorMultiply( VectorLoadByte4(BoneWeights), VECTOR_INV_255 );
 			VectorRegister ExtraWeights;
-			if (bExtraBoneInfluences)
+			if (MaxChunkBoneInfluences > 4)
 			{
 				ExtraWeights = VectorMultiply( VectorLoadByte4(&BoneWeights[MAX_INFLUENCES_PER_STREAM]), VECTOR_INV_255 );
 			}
@@ -743,7 +741,7 @@ static void SkinVertexChunk( FFinalSkinVertex*& DestVertex, TArray<FVertexAnimEv
 			VectorRegister M20	= VectorMultiply( VectorLoadAligned( &BoneMatrix0.M[2][0] ), Weight0 );
 			VectorRegister M30	= VectorMultiply( VectorLoadAligned( &BoneMatrix0.M[3][0] ), Weight0 );
 
-			if ( MaxBoneInfluences > 1 )
+			if ( MaxChunkBoneInfluences > 1 )
 			{
 				const FMatrix BoneMatrix1 = ReferenceToLocal[BoneMap[BoneIndices[INFLUENCE_1]]];
 				VectorRegister Weight1 = VectorReplicate( Weights, INFLUENCE_1 );
@@ -752,7 +750,7 @@ static void SkinVertexChunk( FFinalSkinVertex*& DestVertex, TArray<FVertexAnimEv
 				M20	= VectorMultiplyAdd( VectorLoadAligned( &BoneMatrix1.M[2][0] ), Weight1, M20 );
 				M30	= VectorMultiplyAdd( VectorLoadAligned( &BoneMatrix1.M[3][0] ), Weight1, M30 );
 
-				if ( MaxBoneInfluences > 2 )
+				if ( MaxChunkBoneInfluences > 2 )
 				{
 					const FMatrix BoneMatrix2 = ReferenceToLocal[BoneMap[BoneIndices[INFLUENCE_2]]];
 					VectorRegister Weight2 = VectorReplicate( Weights, INFLUENCE_2 );
@@ -761,7 +759,7 @@ static void SkinVertexChunk( FFinalSkinVertex*& DestVertex, TArray<FVertexAnimEv
 					M20	= VectorMultiplyAdd( VectorLoadAligned( &BoneMatrix2.M[2][0] ), Weight2, M20 );
 					M30	= VectorMultiplyAdd( VectorLoadAligned( &BoneMatrix2.M[3][0] ), Weight2, M30 );
 
-					if ( MaxBoneInfluences > 3 )
+					if ( MaxChunkBoneInfluences > 3 )
 					{
 						const FMatrix BoneMatrix3 = ReferenceToLocal[BoneMap[BoneIndices[INFLUENCE_3]]];
 						VectorRegister Weight3 = VectorReplicate( Weights, INFLUENCE_3 );
@@ -771,7 +769,7 @@ static void SkinVertexChunk( FFinalSkinVertex*& DestVertex, TArray<FVertexAnimEv
 						M30	= VectorMultiplyAdd( VectorLoadAligned( &BoneMatrix3.M[3][0] ), Weight3, M30 );
 					}
 
-					if (MaxBoneInfluences > 4)
+					if (MaxChunkBoneInfluences > 4)
 					{
 						const FMatrix BoneMatrix4 = ReferenceToLocal[BoneMap[BoneIndices[INFLUENCE_4]]];
 						VectorRegister Weight4 = VectorReplicate( ExtraWeights, INFLUENCE_4 - INFLUENCE_4 );
@@ -780,7 +778,7 @@ static void SkinVertexChunk( FFinalSkinVertex*& DestVertex, TArray<FVertexAnimEv
 						M20	= VectorMultiplyAdd( VectorLoadAligned( &BoneMatrix4.M[2][0] ), Weight4, M20 );
 						M30	= VectorMultiplyAdd( VectorLoadAligned( &BoneMatrix4.M[3][0] ), Weight4, M30 );
 
-						if (MaxBoneInfluences > 5)
+						if (MaxChunkBoneInfluences > 5)
 						{
 							const FMatrix BoneMatrix5 = ReferenceToLocal[BoneMap[BoneIndices[INFLUENCE_5]]];
 							VectorRegister Weight5 = VectorReplicate( ExtraWeights, INFLUENCE_5 - INFLUENCE_4 );
@@ -789,7 +787,7 @@ static void SkinVertexChunk( FFinalSkinVertex*& DestVertex, TArray<FVertexAnimEv
 							M20	= VectorMultiplyAdd( VectorLoadAligned( &BoneMatrix5.M[2][0] ), Weight5, M20 );
 							M30	= VectorMultiplyAdd( VectorLoadAligned( &BoneMatrix5.M[3][0] ), Weight5, M30 );
 
-							if (MaxBoneInfluences > 6)
+							if (MaxChunkBoneInfluences > 6)
 							{
 								const FMatrix BoneMatrix6 = ReferenceToLocal[BoneMap[BoneIndices[INFLUENCE_6]]];
 								VectorRegister Weight6 = VectorReplicate( ExtraWeights, INFLUENCE_6 - INFLUENCE_4 );
@@ -798,7 +796,7 @@ static void SkinVertexChunk( FFinalSkinVertex*& DestVertex, TArray<FVertexAnimEv
 								M20	= VectorMultiplyAdd( VectorLoadAligned( &BoneMatrix6.M[2][0] ), Weight6, M20 );
 								M30	= VectorMultiplyAdd( VectorLoadAligned( &BoneMatrix6.M[3][0] ), Weight6, M30 );
 
-								if (MaxBoneInfluences > 7)
+								if (MaxChunkBoneInfluences > 7)
 								{
 									const FMatrix BoneMatrix7 = ReferenceToLocal[BoneMap[BoneIndices[INFLUENCE_7]]];
 									VectorRegister Weight7 = VectorReplicate( ExtraWeights, INFLUENCE_7 - INFLUENCE_4 );
@@ -850,6 +848,23 @@ static void SkinVertexChunk( FFinalSkinVertex*& DestVertex, TArray<FVertexAnimEv
 	}
 }
 
+template< bool bExtraBoneInfluences, typename BaseVertexType, typename VertexType>
+static void SkinVertexChunk( FFinalSkinVertex*& DestVertex, TArray<FVertexAnimEvalInfo>& AnimEvalInfos, const FSkelMeshChunk& Chunk, const FStaticLODModel &LOD, int32 VertexBufferBaseIndex, uint32 NumValidMorphs, int32 &CurBaseVertIdx, int32 LODIndex, int32 RigidInfluenceIndex, const FMatrix* RESTRICT ReferenceToLocal)
+{
+	switch (Chunk.MaxBoneInfluences)
+	{
+		case 1: SkinVertexChunk<bExtraBoneInfluences, 1, BaseVertexType, VertexType>(DestVertex, AnimEvalInfos, Chunk, LOD, VertexBufferBaseIndex, NumValidMorphs, CurBaseVertIdx, LODIndex, RigidInfluenceIndex, ReferenceToLocal); break;
+		case 2: SkinVertexChunk<bExtraBoneInfluences, 2, BaseVertexType, VertexType>(DestVertex, AnimEvalInfos, Chunk, LOD, VertexBufferBaseIndex, NumValidMorphs, CurBaseVertIdx, LODIndex, RigidInfluenceIndex, ReferenceToLocal); break;
+		case 3: SkinVertexChunk<bExtraBoneInfluences, 3, BaseVertexType, VertexType>(DestVertex, AnimEvalInfos, Chunk, LOD, VertexBufferBaseIndex, NumValidMorphs, CurBaseVertIdx, LODIndex, RigidInfluenceIndex, ReferenceToLocal); break;
+		case 4: SkinVertexChunk<bExtraBoneInfluences, 4, BaseVertexType, VertexType>(DestVertex, AnimEvalInfos, Chunk, LOD, VertexBufferBaseIndex, NumValidMorphs, CurBaseVertIdx, LODIndex, RigidInfluenceIndex, ReferenceToLocal); break;
+		case 5: SkinVertexChunk<bExtraBoneInfluences, 5, BaseVertexType, VertexType>(DestVertex, AnimEvalInfos, Chunk, LOD, VertexBufferBaseIndex, NumValidMorphs, CurBaseVertIdx, LODIndex, RigidInfluenceIndex, ReferenceToLocal); break;
+		case 6: SkinVertexChunk<bExtraBoneInfluences, 6, BaseVertexType, VertexType>(DestVertex, AnimEvalInfos, Chunk, LOD, VertexBufferBaseIndex, NumValidMorphs, CurBaseVertIdx, LODIndex, RigidInfluenceIndex, ReferenceToLocal); break;
+		case 7: SkinVertexChunk<bExtraBoneInfluences, 7, BaseVertexType, VertexType>(DestVertex, AnimEvalInfos, Chunk, LOD, VertexBufferBaseIndex, NumValidMorphs, CurBaseVertIdx, LODIndex, RigidInfluenceIndex, ReferenceToLocal); break;
+		case 8: SkinVertexChunk<bExtraBoneInfluences, 8, BaseVertexType, VertexType>(DestVertex, AnimEvalInfos, Chunk, LOD, VertexBufferBaseIndex, NumValidMorphs, CurBaseVertIdx, LODIndex, RigidInfluenceIndex, ReferenceToLocal); break;
+		default: check(0);
+	}
+}
+
 template<typename BaseVertexType, typename VertexType>
 static void SkinVertices( FFinalSkinVertex* DestVertex, FMatrix* ReferenceToLocal, int32 LODIndex, FStaticLODModel& LOD, TArray<FActiveVertexAnim>& ActiveVertexAnims )
 {
@@ -879,17 +894,13 @@ static void SkinVertices( FFinalSkinVertex* DestVertex, FMatrix* ReferenceToLoca
 		FSkelMeshSection& Section = LOD.Sections[SectionIndex];
 		FSkelMeshChunk& Chunk = LOD.Chunks[Section.ChunkIndex];
 
-		switch (Chunk.MaxBoneInfluences)
+		if (LOD.DoChunksNeedExtraBoneInfluences())
 		{
-		case 1: SkinVertexChunk<1, BaseVertexType, VertexType>(DestVertex, AnimEvalInfos, Chunk, LOD, VertexBufferBaseIndex, NumValidMorphs, CurBaseVertIdx, LODIndex, RigidInfluenceIndex, ReferenceToLocal); break;
-		case 2: SkinVertexChunk<2, BaseVertexType, VertexType>(DestVertex, AnimEvalInfos, Chunk, LOD, VertexBufferBaseIndex, NumValidMorphs, CurBaseVertIdx, LODIndex, RigidInfluenceIndex, ReferenceToLocal); break;
-		case 3: SkinVertexChunk<3, BaseVertexType, VertexType>(DestVertex, AnimEvalInfos, Chunk, LOD, VertexBufferBaseIndex, NumValidMorphs, CurBaseVertIdx, LODIndex, RigidInfluenceIndex, ReferenceToLocal); break;
-		case 4: SkinVertexChunk<4, BaseVertexType, VertexType>(DestVertex, AnimEvalInfos, Chunk, LOD, VertexBufferBaseIndex, NumValidMorphs, CurBaseVertIdx, LODIndex, RigidInfluenceIndex, ReferenceToLocal); break;
-		case 5: SkinVertexChunk<5, BaseVertexType, VertexType>(DestVertex, AnimEvalInfos, Chunk, LOD, VertexBufferBaseIndex, NumValidMorphs, CurBaseVertIdx, LODIndex, RigidInfluenceIndex, ReferenceToLocal); break;
-		case 6: SkinVertexChunk<6, BaseVertexType, VertexType>(DestVertex, AnimEvalInfos, Chunk, LOD, VertexBufferBaseIndex, NumValidMorphs, CurBaseVertIdx, LODIndex, RigidInfluenceIndex, ReferenceToLocal); break;
-		case 7: SkinVertexChunk<7, BaseVertexType, VertexType>(DestVertex, AnimEvalInfos, Chunk, LOD, VertexBufferBaseIndex, NumValidMorphs, CurBaseVertIdx, LODIndex, RigidInfluenceIndex, ReferenceToLocal); break;
-		case 8: SkinVertexChunk<8, BaseVertexType, VertexType>(DestVertex, AnimEvalInfos, Chunk, LOD, VertexBufferBaseIndex, NumValidMorphs, CurBaseVertIdx, LODIndex, RigidInfluenceIndex, ReferenceToLocal); break;
-		default: check(0);
+			SkinVertexChunk<true, BaseVertexType, VertexType>(DestVertex, AnimEvalInfos, Chunk, LOD, VertexBufferBaseIndex, NumValidMorphs, CurBaseVertIdx, LODIndex, RigidInfluenceIndex, ReferenceToLocal);
+		}
+		else
+		{
+			SkinVertexChunk<false, BaseVertexType, VertexType>(DestVertex, AnimEvalInfos, Chunk, LOD, VertexBufferBaseIndex, NumValidMorphs, CurBaseVertIdx, LODIndex, RigidInfluenceIndex, ReferenceToLocal);
 		}
 	}
 
