Fix non unity compile error accessing MessageLog and UObjectToken

diff --git a/Engine/Source/Runtime/Engine/Private/LODActor.cpp b/Engine/Source/Runtime/Engine/Private/LODActor.cpp
--- a/Engine/Source/Runtime/Engine/Private/LODActor.cpp
+++ b/Engine/Source/Runtime/Engine/Private/LODActor.cpp
@@ -7,6 +7,8 @@
 #include "EnginePrivate.h"
 #include "Engine/LODActor.h"
 #include "MapErrors.h"
+#include "MessageLog.h"
+#include "UObjectToken.h"
 
 #define LOCTEXT_NAMESPACE "LODActor"
 
