Fix issue where the update rate of skinned meshes on the same actor become out of sync.

diff --git a/Engine/Source/Runtime/Engine/Classes/Components/SkinnedMeshComponent.h b/Engine/Source/Runtime/Engine/Classes/Components/SkinnedMeshComponent.h
--- a/Engine/Source/Runtime/Engine/Classes/Components/SkinnedMeshComponent.h
+++ b/Engine/Source/Runtime/Engine/Classes/Components/SkinnedMeshComponent.h
@@ -909,16 +909,16 @@ private:
 	// Animation update rate control.
 public:
 	/** Animation Update Rate optimization parameters. */
-	UPROPERTY(Transient)
-	struct FAnimUpdateRateParameters AnimUpdateRateParams;
+	struct FAnimUpdateRateParameters* AnimUpdateRateParams;
 
 	/** Aimation Update Rate Tick. */
-	void AnimUpdateRateTick();
+	void AnimUpdateRateTick(uint8 UpdateRateShift);
 
 	/** Updates AnimUpdateRateParams, used by SkinnedMeshComponents.
+	* 
 	* @param bRecentlyRendered : true if at least one SkinnedMeshComponent on this Actor has been rendered in the last second.
 	* @param MaxDistanceFactor : Largest SkinnedMeshComponent of this Actor drawn on screen. */
-	void AnimUpdateRateSetParams(const bool & bRecentlyRendered, const float& MaxDistanceFactor, const bool & bPlayingRootMotion);
+	void AnimUpdateRateSetParams(uint8 UpdateRateShift, const bool & bRecentlyRendered, const float& MaxDistanceFactor, const bool & bPlayingRootMotion);
 
 	virtual bool IsPlayingRootMotion(){ return false; }
 };
diff --git a/Engine/Source/Runtime/Engine/Classes/Engine/EngineTypes.h b/Engine/Source/Runtime/Engine/Classes/Engine/EngineTypes.h
--- a/Engine/Source/Runtime/Engine/Classes/Engine/EngineTypes.h
+++ b/Engine/Source/Runtime/Engine/Classes/Engine/EngineTypes.h
@@ -1764,12 +1764,12 @@ public:
 	{ }
 
 	/** Set parameters and verify inputs.
-	 * @param : Owner Actor owner calling this.
+	 * @param : UpdateShiftRate. Shift our update frames so that updates across all skinned components are staggered
 	 * @param : NewUpdateRate. How often animation will be updated/ticked. 1 = every frame, 2 = every 2 frames, etc.
 	 * @param : NewEvaluationRate. How often animation will be evaluated. 1 = every frame, 2 = every 2 frames, etc.
 	 * @param : bNewInterpSkippedFrames. When skipping a frame, should it be interpolated or frozen?
 	 */
-	void Set(class AActor& Owner, const int32& NewUpdateRate, const int32& NewEvaluationRate, const bool & bNewInterpSkippedFrames);
+	void Set(uint8 UpdateRateShift, const int32& NewUpdateRate, const int32& NewEvaluationRate, const bool & bNewInterpSkippedFrames);
 
 	/* Getter for UpdateRate */
 	int32 GetUpdateRate() const
diff --git a/Engine/Source/Runtime/Engine/Classes/GameFramework/Actor.h b/Engine/Source/Runtime/Engine/Classes/GameFramework/Actor.h
--- a/Engine/Source/Runtime/Engine/Classes/GameFramework/Actor.h
+++ b/Engine/Source/Runtime/Engine/Classes/GameFramework/Actor.h
@@ -338,19 +338,6 @@ public:
 	/** Array of Actors whose Owner is this actor */
 	UPROPERTY(transient)
 	TArray<AActor*> Children;
-	
-	// Animation update rate control.
-protected:
-	/** Unique Tag assigned to spread updates of SkinnedMeshes over time. */
-	UPROPERTY(Transient)
-	uint32 AnimUpdateRateShiftTag;
-public:
-	/** Frame counter to call AnimUpdateRateTick() just once per frame. */
-	UPROPERTY(Transient)
-	uint32 AnimUpdateRateFrameCount;
-
-	/** Get a unique ID to share with all SkinnedMeshComponents in this actor. */
-	uint32 GetAnimUpdateRateShiftTag();
 
 protected:
 	/**
diff --git a/Engine/Source/Runtime/Engine/Private/Actor.cpp b/Engine/Source/Runtime/Engine/Private/Actor.cpp
--- a/Engine/Source/Runtime/Engine/Private/Actor.cpp
+++ b/Engine/Source/Runtime/Engine/Private/Actor.cpp
@@ -83,22 +83,6 @@ AActor::AActor(const FObjectInitializer& ObjectInitializer)
 	bPendingKillPending = false;
 	bFindCameraComponentWhenViewTarget = true;
 	bAllowReceiveTickEventOnDedicatedServer = true;
-	AnimUpdateRateShiftTag = 0;
-	AnimUpdateRateFrameCount = 0;
-}
-
-// Global counter to spread SkinnedMeshComponent tick updates.
-static uint32 SkinnedMeshUpdateRateGroupCount = 0;
-
-uint32 AActor::GetAnimUpdateRateShiftTag()
-{
-	// If hasn't been initialized yet, pick a unique ID, to spread population over frames.
-	if (AnimUpdateRateShiftTag == 0)
-	{
-		AnimUpdateRateShiftTag = ++SkinnedMeshUpdateRateGroupCount;
-	}
-
-	return AnimUpdateRateShiftTag;
 }
 
 void FActorTickFunction::ExecuteTick(float DeltaTime, enum ELevelTick TickType, ENamedThreads::Type CurrentThread, const FGraphEventRef& MyCompletionGraphEvent)
diff --git a/Engine/Source/Runtime/Engine/Private/SkeletalMesh.cpp b/Engine/Source/Runtime/Engine/Private/SkeletalMesh.cpp
--- a/Engine/Source/Runtime/Engine/Private/SkeletalMesh.cpp
+++ b/Engine/Source/Runtime/Engine/Private/SkeletalMesh.cpp
@@ -4541,6 +4541,7 @@ void FSkeletalMeshSceneProxy::UpdateMorphMaterialUsage_GameThread(bool bNeedsMor
 
 USkinnedMeshComponent::USkinnedMeshComponent(const FObjectInitializer& ObjectInitializer)
 	: Super(ObjectInitializer)
+	, AnimUpdateRateParams(NULL)
 {
 	bAutoActivate = true;
 	PrimaryComponentTick.bCanEverTick = true;
diff --git a/Engine/Source/Runtime/Engine/Private/SkeletalMeshComponent.cpp b/Engine/Source/Runtime/Engine/Private/SkeletalMeshComponent.cpp
--- a/Engine/Source/Runtime/Engine/Private/SkeletalMeshComponent.cpp
+++ b/Engine/Source/Runtime/Engine/Private/SkeletalMeshComponent.cpp
@@ -942,16 +942,16 @@ void USkeletalMeshComponent::RefreshBoneTransforms(FActorComponentTickFunction*
 	AnimEvaluationContext.AnimInstance = AnimScriptInstance;
 
 	//Handle update rate optimization setup
-	const bool bDoUpdateRateOptimization = bEnableUpdateRateOptimizations && (AnimUpdateRateParams.GetEvaluationRate() > 1);
+	const bool bDoUpdateRateOptimization = bEnableUpdateRateOptimizations && (AnimUpdateRateParams->GetEvaluationRate() > 1);
 	//Dont mark cache as invalid if we aren't performing optimization anyway
 	const bool bInvalidCachedBones = bDoUpdateRateOptimization &&
 									( (LocalAtoms.Num() != SkeletalMesh->RefSkeleton.GetNum())
 									  || (LocalAtoms.Num() != CachedLocalAtoms.Num())
 									  || (GetNumSpaceBases() != CachedSpaceBases.Num()) );
 
-	AnimEvaluationContext.bDoEvaluation = !bDoUpdateRateOptimization || bInvalidCachedBones || !AnimUpdateRateParams.ShouldSkipEvaluation();
+	AnimEvaluationContext.bDoEvaluation = !bDoUpdateRateOptimization || bInvalidCachedBones || !AnimUpdateRateParams->ShouldSkipEvaluation();
 	
-	AnimEvaluationContext.bDoInterpolation = bDoUpdateRateOptimization && !bInvalidCachedBones && AnimUpdateRateParams.ShouldInterpolateSkippedFrames();
+	AnimEvaluationContext.bDoInterpolation = bDoUpdateRateOptimization && !bInvalidCachedBones && AnimUpdateRateParams->ShouldInterpolateSkippedFrames();
 	AnimEvaluationContext.bDuplicateToCacheBones = bInvalidCachedBones || (bDoUpdateRateOptimization && AnimEvaluationContext.bDoEvaluation && !AnimEvaluationContext.bDoInterpolation);
 
 	if (!bDoUpdateRateOptimization)
@@ -1026,7 +1026,7 @@ void USkeletalMeshComponent::PostAnimEvaluation(FAnimationEvaluationContext& Eva
 	{
 		SCOPE_CYCLE_COUNTER(STAT_InterpolateSkippedFrames);
 
-		const float Alpha = 0.25f + (1.f / float(FMath::Max(AnimUpdateRateParams.GetEvaluationRate(), 2) * 2));
+		const float Alpha = 0.25f + (1.f / float(FMath::Max(AnimUpdateRateParams->GetEvaluationRate(), 2) * 2));
 		FAnimationRuntime::LerpBoneTransforms(LocalAtoms, CachedLocalAtoms, Alpha, RequiredBones);
 		if (bDoubleBufferedBlendSpaces)
 		{
diff --git a/Engine/Source/Runtime/Engine/Private/SkinnedMeshComponent.cpp b/Engine/Source/Runtime/Engine/Private/SkinnedMeshComponent.cpp
--- a/Engine/Source/Runtime/Engine/Private/SkinnedMeshComponent.cpp
+++ b/Engine/Source/Runtime/Engine/Private/SkinnedMeshComponent.cpp
@@ -26,6 +26,111 @@ FAutoConsoleVariableRef CVarSkeletalMeshLODBias(
 	ECVF_Scalability
 	);
 
+
+namespace FAnimUpdateRateManager
+{
+	// Global counter to spread SkinnedMeshComponent tick updates.
+	static uint8 UpdateRateGroupCount = 0;
+
+	struct FAnimUpdateRateParametersTracker
+	{
+		FAnimUpdateRateParameters UpdateRateParameters;
+
+		/** Frame counter to call AnimUpdateRateTick() just once per frame. */
+		uint32 AnimUpdateRateFrameCount;
+
+		/** Counter to stagger update and evaluation across skinned mesh components */
+		uint8 AnimUpdateRateShiftTag;
+
+		/** List of all USkinnedMeshComponents that use this set of parameters */
+		TArray<USkinnedMeshComponent*> RegisteredComponents;
+
+		FAnimUpdateRateParametersTracker() : AnimUpdateRateFrameCount(0), AnimUpdateRateShiftTag(0) {}
+
+		uint8 GetAnimUpdateRateShiftTag()
+		{
+			// If hasn't been initialized yet, pick a unique ID, to spread population over frames.
+			if (AnimUpdateRateShiftTag == 0)
+			{
+				AnimUpdateRateShiftTag = ++UpdateRateGroupCount;
+			}
+
+			return AnimUpdateRateShiftTag;
+		}
+	};
+
+
+	TMap<UObject*, FAnimUpdateRateParametersTracker*> ActorToUpdateRateParams;
+
+	UObject* GetMapIndexForComponent(USkinnedMeshComponent* SkinnedComponent)
+	{
+		UObject* TrackerIndex = SkinnedComponent->GetOwner();
+		if (TrackerIndex == nullptr)
+		{
+			TrackerIndex = SkinnedComponent;
+		}
+		return TrackerIndex;
+	}
+
+	FAnimUpdateRateParameters* GetUpdateRateParameters(USkinnedMeshComponent* SkinnedComponent)
+	{
+		UObject* TrackerIndex = GetMapIndexForComponent(SkinnedComponent);
+
+		FAnimUpdateRateParametersTracker** ExistingTracker = ActorToUpdateRateParams.Find(TrackerIndex);
+		if (!ExistingTracker)
+		{
+			ExistingTracker = &ActorToUpdateRateParams.Add(TrackerIndex);
+			(*ExistingTracker) = new FAnimUpdateRateParametersTracker();
+		}
+		
+		check(ExistingTracker);
+		check(*ExistingTracker);
+
+		checkSlow(!(*ExistingTracker)->RegisteredComponents.Contains(SkinnedComponent)); // We have already been registered? Something has gone very wrong!
+
+		(*ExistingTracker)->RegisteredComponents.Add(SkinnedComponent);
+		FPlatformMisc::LowLevelOutputDebugStringf(TEXT("MAPUPDATE: Count %i\n"), ActorToUpdateRateParams.Num());
+		return &(*ExistingTracker)->UpdateRateParameters;
+	}
+
+	void CleanupUpdateRateParametersRef(USkinnedMeshComponent* SkinnedComponent)
+	{
+		UObject* TrackerIndex = GetMapIndexForComponent(SkinnedComponent);
+
+		FAnimUpdateRateParametersTracker* Tracker = ActorToUpdateRateParams.FindChecked(TrackerIndex);
+		Tracker->RegisteredComponents.Remove(SkinnedComponent);
+		if (Tracker->RegisteredComponents.Num() == 0)
+		{
+			ActorToUpdateRateParams.Remove(TrackerIndex);
+			delete Tracker;
+		}
+		FPlatformMisc::LowLevelOutputDebugStringf(TEXT("MAPUPDATE: Count %i\n"), ActorToUpdateRateParams.Num());
+	}
+
+	void TickUpdateRateParameters(USkinnedMeshComponent* SkinnedComponent)
+	{
+		// Convert current frame counter from 64 to 32 bits.
+		const uint32 CurrentFrame32 = uint32(GFrameCounter % MAX_uint32);
+
+		UObject* TrackerIndex = GetMapIndexForComponent(SkinnedComponent);
+		FAnimUpdateRateParametersTracker* Tracker = ActorToUpdateRateParams.FindChecked(TrackerIndex);
+
+		if (CurrentFrame32 != Tracker->AnimUpdateRateFrameCount)
+		{
+			Tracker->AnimUpdateRateFrameCount = CurrentFrame32;
+			SkinnedComponent->AnimUpdateRateTick(Tracker->GetAnimUpdateRateShiftTag());
+		}
+	}
+
+	const TArray<USkinnedMeshComponent*> GetSharedMeshComponents(USkinnedMeshComponent* SkinnedComponent)
+	{
+		UObject* TrackerIndex = GetMapIndexForComponent(SkinnedComponent);
+		FAnimUpdateRateParametersTracker* Tracker = ActorToUpdateRateParams.FindChecked(TrackerIndex);
+		return Tracker->RegisteredComponents;
+	}
+
+}
+
 void USkinnedMeshComponent::OnRegister()
 {
 	// The reason this happens before register
@@ -33,6 +138,9 @@ void USkinnedMeshComponent::OnRegister()
 	// won't result in any issues of accessing SpaceBases
 	// This isn't really ideal solution because these transform won't have 
 	// any valid data yet. 
+
+	AnimUpdateRateParams = FAnimUpdateRateManager::GetUpdateRateParameters(this);
+
 	AllocateTransformData();
 	UpdateMasterBoneMap();	
 
@@ -45,6 +153,8 @@ void USkinnedMeshComponent::OnUnregister()
 {
 	DeallocateTransformData();
 	Super::OnUnregister();
+	FAnimUpdateRateManager::CleanupUpdateRateParametersRef(this);
+	AnimUpdateRateParams = NULL;
 }
 
 void USkinnedMeshComponent::CreateRenderState_Concurrent()
@@ -220,24 +330,16 @@ void USkinnedMeshComponent::TickUpdateRate()
 	SCOPE_CYCLE_COUNTER(STAT_TickUpdateRate);
 	if( bEnableUpdateRateOptimizations )
 	{
-		AActor* Owner = GetOwner();
-		if (Owner)
+		if (GetOwner())
 		{
-			// Convert current frame counter from 64 to 32 bits.
-			const uint32 CurrentFrame32 = uint32(GFrameCounter % MAX_uint32);
-
 			// Tick Owner once per frame. All attached SkinnedMeshComponents will share the same settings.
-			if (Owner->AnimUpdateRateFrameCount != CurrentFrame32)
-			{
-				Owner->AnimUpdateRateFrameCount = CurrentFrame32;
-				AnimUpdateRateTick();
-			}
+			FAnimUpdateRateManager::TickUpdateRateParameters(this);
 
 			// debug -- @todo: hook this up to a console command.
 			if (bDisplayDebugUpdateRateOptimizations)
 			{
 				FColor DrawColor;
-				switch (AnimUpdateRateParams.GetUpdateRate())
+				switch (AnimUpdateRateParams->GetUpdateRate())
 				{
 				case 1: DrawColor = FColor::Red; break;
 				case 2: DrawColor = FColor::Green; break;
@@ -265,7 +367,7 @@ void USkinnedMeshComponent::TickComponent(float DeltaTime, enum ELevelTick TickT
 	// Tick Pose first
 	if( ShouldTickPose() )
 	{
-		if (bEnableUpdateRateOptimizations && AnimUpdateRateParams.ShouldSkipUpdate())
+		if (bEnableUpdateRateOptimizations && AnimUpdateRateParams->ShouldSkipUpdate())
 		{
 			SkippedTickDeltaTime += DeltaTime;
 			if( !bRecentlyRendered )
@@ -1747,71 +1849,46 @@ TArray<FActiveVertexAnim> USkinnedMeshComponent::UpdateActiveVertexAnims(const U
 	return OutVertexAnims;
 }
 
-void USkinnedMeshComponent::AnimUpdateRateTick()
+void USkinnedMeshComponent::AnimUpdateRateTick(uint8 UpdateRateShift)
 {
-	AActor* Owner = GetOwner();
-
-	if (!Owner)
-	{
-		return;
-	}
+	const float RecentlyRenderedTime = (GetWorld()->TimeSeconds - 1.f);
 
 	// Go through components and figure out if they've been recently rendered, and the biggest MaxDistanceFactor
 	bool bRecentlyRendered = false;
 	bool bPlayingRootMotion = false;
 	float MaxDistanceFactor = 0.f;
 
-	// Gather Actor's components
-	TInlineComponentArray<USceneComponent *> ComponentStack;
-	ComponentStack.Reserve(NumInlinedActorComponents);
+	const TArray<USkinnedMeshComponent*> SkinnedComponents = FAnimUpdateRateManager::GetSharedMeshComponents(this);
 
-	USceneComponent* CurrentComponent = Owner->GetRootComponent();
-	while (CurrentComponent)
+	for (USkinnedMeshComponent* Component : SkinnedComponents)
 	{
-		// push children on the stack so they get tested later
-		ComponentStack.Append(CurrentComponent->AttachChildren);
-
-		if (CurrentComponent->IsRegistered())
-		{
-			USkinnedMeshComponent* SkinMeshComp = Cast<USkinnedMeshComponent>(CurrentComponent);
-			if (SkinMeshComp)
-			{
-				bRecentlyRendered = bRecentlyRendered || (SkinMeshComp->LastRenderTime > (GetWorld()->TimeSeconds - 1.f));
-				MaxDistanceFactor = FMath::Max(MaxDistanceFactor, SkinMeshComp->MaxDistanceFactor);
-				bPlayingRootMotion = bPlayingRootMotion || SkinMeshComp->IsPlayingRootMotion();
-			}
-		}
-
-		CurrentComponent = (ComponentStack.Num() > 0) ? ComponentStack.Pop(false) : NULL;
+		bRecentlyRendered = bRecentlyRendered || (Component->LastRenderTime > RecentlyRenderedTime);
+		MaxDistanceFactor = FMath::Max(MaxDistanceFactor, Component->MaxDistanceFactor);
+		bPlayingRootMotion = bPlayingRootMotion || Component->IsPlayingRootMotion();
 	}
 
 	// Figure out which update rate should be used.
-	AnimUpdateRateSetParams(bRecentlyRendered, MaxDistanceFactor, bPlayingRootMotion);
+	AnimUpdateRateSetParams(UpdateRateShift, bRecentlyRendered, MaxDistanceFactor, bPlayingRootMotion);
 }
 
-void USkinnedMeshComponent::AnimUpdateRateSetParams(const bool & bRecentlyRendered, const float& MaxDistanceFactor, const bool & bPlayingRootMotion)
+void USkinnedMeshComponent::AnimUpdateRateSetParams(uint8 UpdateRateShift, const bool & bRecentlyRendered, const float& MaxDistanceFactor, const bool & bPlayingRootMotion)
 {
 	// default rules for setting update rates
 
 	// Human controlled characters should be ticked always fully to minimize latency w/ game play events triggered by animation.
 	AActor* Owner = GetOwner();
-	if (!Owner)
-	{
-		return;
-	}
-
-	AController * Controller = Owner->GetInstigatorController();
+	AController * Controller = Owner ? Owner->GetInstigatorController() : NULL;
 	const bool bHumanControlled = Controller && (Cast<APlayerController>(Controller) != NULL);
 
 	// Not rendered, including dedicated servers. we can skip the Evaluation part.
 	if (!bRecentlyRendered)
 	{
-		AnimUpdateRateParams.Set(*Owner, (bHumanControlled ? 1 : 4), 4, false);
+		AnimUpdateRateParams->Set(UpdateRateShift, (bHumanControlled ? 1 : 4), 4, false);
 	}
 	// Visible controlled characters or playing root motion. Need evaluation and ticking done every frame.
 	else  if (bHumanControlled || bPlayingRootMotion)
 	{
-		AnimUpdateRateParams.Set(*Owner, 1, 1, false);
+		AnimUpdateRateParams->Set(UpdateRateShift, 1, 1, false);
 	}
 	else
 	{
@@ -1831,7 +1908,7 @@ void USkinnedMeshComponent::AnimUpdateRateSetParams(const bool & bRecentlyRender
 			DesiredEvaluationRate = 3;
 		}
 
-		AnimUpdateRateParams.Set(*Owner, DesiredEvaluationRate, DesiredEvaluationRate, true);
+		AnimUpdateRateParams->Set(UpdateRateShift, DesiredEvaluationRate, DesiredEvaluationRate, true);
 	}
 }
 
@@ -1866,7 +1943,7 @@ void USkinnedMeshComponent::SetSpaceBaseDoubleBuffering(bool bInDoubleBufferedBl
 	}
 }
 
-void FAnimUpdateRateParameters::Set(class AActor& Owner, const int32& NewUpdateRate, const int32& NewEvaluationRate, const bool & bNewInterpSkippedFrames)
+void FAnimUpdateRateParameters::Set(uint8 UpdateRateShift, const int32& NewUpdateRate, const int32& NewEvaluationRate, const bool & bNewInterpSkippedFrames)
 {
 	UpdateRate = FMath::Max(NewUpdateRate, 1);
 	// Make sure EvaluationRate is a multiple of UpdateRate.
@@ -1874,7 +1951,7 @@ void FAnimUpdateRateParameters::Set(class AActor& Owner, const int32& NewUpdateR
 	bInterpolateSkippedFrames = bNewInterpSkippedFrames;
 
 	// Make sure we don't overflow. we don't need very large numbers.
-	const uint32 Counter = (GFrameCounter % MAX_uint8) + (Owner.GetAnimUpdateRateShiftTag() % MAX_uint8);
+	const uint32 Counter = (GFrameCounter + UpdateRateShift)% MAX_uint8;
 
 	bSkipUpdate = ((Counter % UpdateRate) > 0);
 	bSkipEvaluation = ((Counter % EvaluationRate) > 0);
