Merging Loading Screen adjustments from MattO

//depot/UE4-UT/...

to //depot/UE4/...

diff --git a/Engine/Source/Runtime/Engine/Private/UnrealEngine.cpp b/Engine/Source/Runtime/Engine/Private/UnrealEngine.cpp
--- a/Engine/Source/Runtime/Engine/Private/UnrealEngine.cpp
+++ b/Engine/Source/Runtime/Engine/Private/UnrealEngine.cpp
@@ -8629,6 +8629,21 @@ bool UEngine::LoadMap( FWorldContext& WorldContext, FURL URL, class UPendingNetG
 
 	// send a callback message
 	FCoreUObjectDelegates::PreLoadMap.Broadcast();
+	// make sure there is a matching PostLoadMap() no matter how we exit
+	struct FPostLoadMapCaller
+	{
+		bool bCalled;
+		FPostLoadMapCaller()
+			: bCalled(false)
+		{}
+		~FPostLoadMapCaller()
+		{
+			if (!bCalled)
+			{
+				FCoreUObjectDelegates::PostLoadMap.Broadcast();
+			}
+		}
+	} PostLoadMapCaller;
 
 	// Cancel any pending texture streaming requests.  This avoids a significant delay on consoles 
 	// when loading a map and there are a lot of outstanding texture streaming requests from the previous map.
@@ -9091,6 +9106,7 @@ bool UEngine::LoadMap( FWorldContext& WorldContext, FURL URL, class UPendingNetG
 	}
 
 	// send a callback message
+	PostLoadMapCaller.bCalled = true;
 	FCoreUObjectDelegates::PostLoadMap.Broadcast();
 	
 	WorldContext.World()->bWorldWasLoadedThisTick = true;
