added far shadow cascade feature
 per scene primitive checkbox, default is on for landscape
 directional light property FarShadowCascadeCount and FarShadowDistance
 todo: ray traced distance field shadows don't fade with it yet, non landscape objects makes with far should be excluded from ry traced distance field

diff --git a/Engine/Source/Runtime/Engine/Classes/Components/DirectionalLightComponent.h b/Engine/Source/Runtime/Engine/Classes/Components/DirectionalLightComponent.h
--- a/Engine/Source/Runtime/Engine/Classes/Components/DirectionalLightComponent.h
+++ b/Engine/Source/Runtime/Engine/Classes/Components/DirectionalLightComponent.h
@@ -95,6 +95,16 @@ class UDirectionalLightComponent : public ULightComponent
 	UPROPERTY(EditAnywhere, AdvancedDisplay, BlueprintReadOnly, Category=CascadedShadowMaps, DisplayName = "Inset Shadows For Movable Objects")
 	uint32 bUseInsetShadowsForMovableObjects : 1;
 
+	/** 0: no DistantShadowCascades, otherwise the count of cascades between WholeSceneDynamicShadowRadius and DistantShadowDistance that are covered by distant shadow cascades. */
+	UPROPERTY(EditAnywhere, BlueprintReadOnly, Category=CascadedShadowMaps, meta=(UIMin = "0", UIMax = "4"), DisplayName = "Far Shadow Cascade Count")
+	int32 FarShadowCascadeCount;
+
+	/** 
+	 * Distance at which the far shadow cascade should end.  Far shadows will cover the range between 'Dynamic Shadow Distance' and this distance. 
+	 */
+	UPROPERTY(EditAnywhere, BlueprintReadOnly, Category=CascadedShadowMaps, meta=(UIMin = "0", UIMax = "800000"), DisplayName = "Far Shadow Distance")
+	float FarShadowDistance;
+
 	/** 
 	 * Distance at which the ray traced shadow cascade should end.  Distance field shadows will cover the range between 'Dynamic Shadow Distance' this distance. 
 	 */
diff --git a/Engine/Source/Runtime/Engine/Classes/Components/PrimitiveComponent.h b/Engine/Source/Runtime/Engine/Classes/Components/PrimitiveComponent.h
--- a/Engine/Source/Runtime/Engine/Classes/Components/PrimitiveComponent.h
+++ b/Engine/Source/Runtime/Engine/Classes/Components/PrimitiveComponent.h
@@ -285,6 +285,12 @@ public:
 	UPROPERTY(EditAnywhere, AdvancedDisplay, BlueprintReadOnly, Category=Lighting, meta=(EditCondition="CastShadow"))
 	uint32 bSelfShadowOnly:1;
 
+	/** 
+	 * When enabled, the component will be rendering into the far shadow cascades (only for directional lights).
+	 */
+	UPROPERTY(EditAnywhere, AdvancedDisplay, BlueprintReadOnly, Category=Lighting, meta=(EditCondition="CastShadow", DisplayName = "Far Shadow"))
+	uint32 bCastFarShadow:1;
+
 	/** 
 	 * Whether this component should create a per-object shadow that gives higher effective shadow resolution. 
 	 * Useful for cinematic character shadowing. Assumed to be enabled if bSelfShadowOnly is enabled.
diff --git a/Engine/Source/Runtime/Engine/Private/DirectionalLightComponent.cpp b/Engine/Source/Runtime/Engine/Private/DirectionalLightComponent.cpp
--- a/Engine/Source/Runtime/Engine/Private/DirectionalLightComponent.cpp
+++ b/Engine/Source/Runtime/Engine/Private/DirectionalLightComponent.cpp
@@ -118,6 +118,15 @@ public:
 		{
 			WholeSceneDynamicShadowRadius = Component->DynamicShadowDistanceStationaryLight;
 		}
+
+		const float FarCascadeSize = Component->FarShadowDistance - WholeSceneDynamicShadowRadius;
+
+		// 100 is just some number to avoid cascades of 0 size, the user still can still create many tiny cascades (very inefficient)
+		if(Component->FarShadowCascadeCount && FarCascadeSize > 100.0f)
+		{
+			FarShadowDistance = Component->FarShadowDistance;
+			FarShadowCascadeCount = Component->FarShadowCascadeCount;
+		}
 	}
 
 	void UpdateLightShaftOverrideDirection_GameThread(const UDirectionalLightComponent* Component)
@@ -194,34 +203,32 @@ public:
 			&& bUseInsetShadowsForMovableObjects;
 	}
 
-	/** Returns the number of view dependent shadows this light will create. */
+	/** Returns the number of view dependent shadows this light will create, not counting distance field shadow cascades. */
 	virtual uint32 GetNumViewDependentWholeSceneShadows(const FSceneView& View) const
 	{ 
-		uint32 NumCascades = GetNumShadowMappedCascades(View.MaxShadowCascades);
-
-		uint32 ExtraCascades = 0;
-
-		if (ShouldCreateRayTracedCascade(View.GetFeatureLevel()))
-		{
-			// Add a cascade for the ray traced shadows
-			ExtraCascades = 1;
-		}
+		uint32 TotalCascades = GetNumShadowMappedCascades(View.MaxShadowCascades) + FarShadowCascadeCount;
 
-		return NumCascades + ExtraCascades;
+		return TotalCascades;
 	}
 
 	/**
 	 * Sets up a projected shadow initializer that's dependent on the current view for shadows from the entire scene.
+	 * @param InCascadeIndex ShadowSplitIndex or INDEX_NONE for the the DistanceFieldCascade
 	 * @return True if the whole-scene projected shadow should be used.
 	 */
-	virtual bool GetViewDependentWholeSceneProjectedShadowInitializer(const FSceneView& View, uint32 InShadowSplitIndex, FWholeSceneProjectedShadowInitializer& OutInitializer) const
+	virtual bool GetViewDependentWholeSceneProjectedShadowInitializer(const FSceneView& View, int32 InCascadeIndex, FWholeSceneProjectedShadowInitializer& OutInitializer) const
 	{
 		const FMatrix& WorldToLight = GetWorldToLight();
 
+		const bool bRaytracedCascade = (InCascadeIndex == INDEX_NONE);
+
 		FShadowCascadeSettings CascadeSettings;
 
-		FSphere Bounds = FDirectionalLightSceneProxy::GetShadowSplitBounds(View, InShadowSplitIndex, &CascadeSettings);
-		OutInitializer.CascadeSettings = CascadeSettings;
+		FSphere Bounds = FDirectionalLightSceneProxy::GetShadowSplitBounds(View, InCascadeIndex, &OutInitializer.CascadeSettings);
+
+		uint32 NumNearCascades = GetNumShadowMappedCascades(View.MaxShadowCascades);
+
+		const uint32 ShadowSplitIndex = bRaytracedCascade ? NumNearCascades : InCascadeIndex;
 
 		const float ShadowExtent = Bounds.W / FMath::Sqrt(3.0f);
 		const FBoxSphereBounds SubjectBounds(Bounds.Center, FVector(ShadowExtent, ShadowExtent, ShadowExtent), Bounds.W);
@@ -236,9 +243,10 @@ public:
 		// Reduce casting distance on a directional light
 		// This is necessary to improve floating point precision in several places, especially when deriving frustum verts from InvReceiverMatrix
 		OutInitializer.MaxDistanceToCastInLightW = HALF_WORLD_MAX / 32.0f;
-		OutInitializer.InitShadowSplitIndex = InShadowSplitIndex;
 		// Last cascade is the ray traced shadow, if enabled
-		OutInitializer.bRayTracedDistanceFieldShadow = ShouldCreateRayTracedCascade(View.GetFeatureLevel()) && InShadowSplitIndex == FDirectionalLightSceneProxy::GetNumViewDependentWholeSceneShadows(View) - 1;
+		OutInitializer.InitShadowSplitIndex = ShadowSplitIndex;
+		OutInitializer.bRayTracedDistanceFieldShadow = bRaytracedCascade;
+		OutInitializer.CascadeSettings.bFarShadowCascade = !bRaytracedCascade && ShadowSplitIndex >= NumNearCascades;
 		return true;
 	}
 
@@ -287,7 +295,14 @@ public:
 
 	virtual FVector2D GetDirectionalLightDistanceFadeParameters(ERHIFeatureLevel::Type InFeatureLevel) const
 	{
-		const float FarDistance = ComputeShadowDistance(InFeatureLevel);
+		float FarDistance = GetCSMMaxDistance();
+		{
+			if (ShouldCreateRayTracedCascade(InFeatureLevel))
+			{
+				FarDistance = GetDistanceFieldShadowDistance();
+			}
+			FarDistance = FMath::Max(FarDistance, FarShadowDistance);
+		}
         
 		// The far distance for the dynamic to static fade is the range of the directional light.
 		// The near distance is placed at a depth of 90% of the light's range.
@@ -310,6 +325,16 @@ public:
 		return true;
 	}
 
+	virtual bool ShouldCreateRayTracedCascade(ERHIFeatureLevel::Type InFeatureLevel) const
+	{
+		//@todo - handle View.MaxShadowCascades properly
+		const uint32 NumCascades = GetNumShadowMappedCascades(10);
+		const float RaytracedShadowDistance = GetDistanceFieldShadowDistance();
+		const bool bCreateWithCSM = NumCascades > 0 && RaytracedShadowDistance > GetCSMMaxDistance();
+		const bool bCreateWithoutCSM = NumCascades == 0 && RaytracedShadowDistance > 0;
+		return DoesPlatformSupportDistanceFieldShadowing(GShaderPlatformForFeatureLevel[InFeatureLevel]) && (bCreateWithCSM || bCreateWithoutCSM);
+	}
+
 private:
 
 	uint32 GetNumShadowMappedCascades(uint32 MaxShadowCascades) const
@@ -340,28 +365,6 @@ private:
 		return DistanceFieldShadowDistance;
 	}
 
-	bool ShouldCreateRayTracedCascade(ERHIFeatureLevel::Type InFeatureLevel) const
-	{
-		//@todo - handle View.MaxShadowCascades properly
-		const uint32 NumCascades = GetNumShadowMappedCascades(10);
-		const float RaytracedShadowDistance = GetDistanceFieldShadowDistance();
-		const bool bCreateWithCSM = NumCascades > 0 && RaytracedShadowDistance > GetCSMMaxDistance();
-		const bool bCreateWithoutCSM = NumCascades == 0 && RaytracedShadowDistance > 0;
-		return DoesPlatformSupportDistanceFieldShadowing(GShaderPlatformForFeatureLevel[InFeatureLevel]) && (bCreateWithCSM || bCreateWithoutCSM);
-	}
-
-	float ComputeShadowDistance(ERHIFeatureLevel::Type InFeatureLevel) const
-	{
-		float Distance = GetCSMMaxDistance();
-
-		if (ShouldCreateRayTracedCascade(InFeatureLevel))
-		{
-			Distance = GetDistanceFieldShadowDistance();
-		}
-
-		return Distance;
-	}
-
 	float GetShadowTransitionScale() const
 	{
 		static const auto CVar = IConsoleManager::Get().FindTConsoleVariableDataFloat(TEXT("r.Shadow.CSM.TransitionScale"));
@@ -481,22 +484,35 @@ private:
 
 		return Ret / TotalScale;
 	}
-
-	// Given a particular split index and view near plane, returns the view space Z distance to the PSSM split plane.
-	inline float GetSplitDistance(const FSceneView& View, int32 SplitIndex, int32 NumCascades, float ShadowNear, bool bIgnoreLastCascade) const
+	
+	// @param SplitIndex 0: near, 1:end of 1st cascade, ...
+	inline float GetSplitDistance(const FSceneView& View, uint32 SplitIndex, bool bDistanceFieldShadows) const
 	{
-		check(SplitIndex >= 0);
+		// near cascade means non far and non distance field cascade
+		uint32 NumNearCascades = GetNumShadowMappedCascades(View.MaxShadowCascades);
 
-		float FarDistance = ComputeShadowDistance(View.GetFeatureLevel());
+		float CascadeDistanceWithoutFar = GetCSMMaxDistance();
+		float ShadowNear = View.NearClippingDistance;
 
-		if (bIgnoreLastCascade)
+		// non near cascades are split differently for distance field shadow
+		if(SplitIndex > NumNearCascades)
 		{
-			// Compute shadowmap cascade distances as if the ray traced cascade didn't exist
-			NumCascades--;
-			FarDistance = GetCSMMaxDistance();
+			if(bDistanceFieldShadows)
+			{
+				// there is only one distance field shadow cascade
+				check(SplitIndex == NumNearCascades + 1);
+				return DistanceFieldShadowDistance;
+			}
+			else
+			{
+				// the far cascades start at the after the near cascades
+				return CascadeDistanceWithoutFar + ComputeAccumulatedScale(CascadeDistributionExponent, SplitIndex - NumNearCascades, FarShadowCascadeCount) * (FarShadowDistance - CascadeDistanceWithoutFar);
+			}
+		}
+		else
+		{
+			return ShadowNear + ComputeAccumulatedScale(CascadeDistributionExponent, SplitIndex, NumNearCascades) * (CascadeDistanceWithoutFar - ShadowNear);
 		}
-
-		return ShadowNear + ComputeAccumulatedScale(CascadeDistributionExponent, SplitIndex, NumCascades) * (FarDistance - ShadowNear);
 	}
 
 	virtual FSphere GetShadowSplitBoundsDepthRange(const FSceneView& View, float SplitNear, float SplitFar, FShadowCascadeSettings* OutCascadeSettings) const
@@ -572,18 +588,25 @@ private:
 		return CascadeSphere;
 	}
 
-	virtual FSphere GetShadowSplitBounds(const FSceneView& View, uint32 InShadowSplitIndex, FShadowCascadeSettings* OutCascadeSettings) const
+	// @param InShadowSplitIndex cascade index or InShadowSplitIndex == INDEX_NONE for the distance field cascade
+	virtual FSphere GetShadowSplitBounds(const FSceneView& View, int32 InCascadeIndex, FShadowCascadeSettings* OutCascadeSettings) const
 	{
-		// this checks for WholeSceneDynamicShadowRadius and DynamicShadowCascades
-		uint32 NumCascades = GetNumViewDependentWholeSceneShadows(View);
+		uint32 NumNearCascades = GetNumShadowMappedCascades(View.MaxShadowCascades);
 
 		const bool bHasRayTracedCascade = ShouldCreateRayTracedCascade(View.GetFeatureLevel());
-		const bool bIsRayTracedCascade = bHasRayTracedCascade && InShadowSplitIndex == NumCascades - 1;
+
+		// this checks for WholeSceneDynamicShadowRadius and DynamicShadowCascades
+		uint32 NumNearAndFarCascades = GetNumViewDependentWholeSceneShadows(View);
+
+		uint32 NumTotalCascades = FMath::Max(NumNearAndFarCascades, NumNearCascades + (bHasRayTracedCascade ? 1 : 0));
+
+		const bool bIsRayTracedCascade = InCascadeIndex == INDEX_NONE;
+		const uint32 ShadowSplitIndex = bIsRayTracedCascade ? NumNearCascades : InCascadeIndex;		// todo -1?
 
 		// Determine start and end distances to the current cascade's split planes
 		// Presence of the ray traced cascade does not change depth ranges for the shadow-mapped cascades
-		float SplitNear = GetSplitDistance(View, InShadowSplitIndex, NumCascades, View.NearClippingDistance, bHasRayTracedCascade);
-		float SplitFar = GetSplitDistance(View, InShadowSplitIndex + 1, NumCascades, View.NearClippingDistance, bHasRayTracedCascade && !bIsRayTracedCascade);
+		float SplitNear = GetSplitDistance(View, ShadowSplitIndex, bIsRayTracedCascade);
+		float SplitFar = GetSplitDistance(View, ShadowSplitIndex + 1, bIsRayTracedCascade);
 		float FadePlane = SplitFar;
 
 		float LocalCascadeTransitionFraction = CascadeTransitionFraction * GetShadowTransitionScale();
@@ -591,7 +614,7 @@ private:
 		float FadeExtension = (SplitFar - SplitNear) * LocalCascadeTransitionFraction;
 
 		// Add the fade region to the end of the subfrustum, if this is not the last cascade.
-		if (InShadowSplitIndex < NumCascades - 1)
+		if ((int32)ShadowSplitIndex < (int32)NumTotalCascades - 1)
 		{
 			SplitFar += FadeExtension;			
 		}
@@ -602,11 +625,12 @@ private:
 
 			OutCascadeSettings->SplitNearFadeRegion = 0.0f;
 
-			if(InShadowSplitIndex >= 1)
+			if(ShadowSplitIndex >= 1)
 			{
-				// todo: optimize
-				float BeforeSplitNear = GetSplitDistance(View, InShadowSplitIndex - 1, NumCascades, View.NearClippingDistance, bHasRayTracedCascade);
-				float BeforeSplitFar = GetSplitDistance(View, InShadowSplitIndex, NumCascades, View.NearClippingDistance, bHasRayTracedCascade);
+				// only used to fade the translucency lighting volume
+
+				float BeforeSplitNear = GetSplitDistance(View, ShadowSplitIndex - 1, bIsRayTracedCascade);
+				float BeforeSplitFar = GetSplitDistance(View, ShadowSplitIndex, bIsRayTracedCascade);
 
 				OutCascadeSettings->SplitNearFadeRegion = (BeforeSplitFar - BeforeSplitNear) * LocalCascadeTransitionFraction;
 			}
@@ -651,6 +675,7 @@ UDirectionalLightComponent::UDirectionalLightComponent(const FObjectInitializer&
 	DynamicShadowDistanceStationaryLight = 0.f;
 
 	DistanceFieldShadowDistance = 30000.0f;
+	FarShadowDistance = 300000.0f;
 	LightSourceAngle = 1;
 
 	DynamicShadowCascades = 3;
@@ -692,22 +717,35 @@ bool UDirectionalLightComponent::CanEditChange(const UProperty* InProperty) cons
 	if (InProperty)
 	{
 		FString PropertyName = InProperty->GetName();
+		
+		bool bShadowCascases = CastShadows
+			&& CastDynamicShadows 
+			&& ((DynamicShadowDistanceMovableLight > 0 && Mobility == EComponentMobility::Movable)
+			|| (DynamicShadowDistanceStationaryLight > 0 && Mobility == EComponentMobility::Stationary));
 
 		if (PropertyName == GET_MEMBER_NAME_STRING_CHECKED(UDirectionalLightComponent, bUseInsetShadowsForMovableObjects))
 		{
 			return CastShadows && CastDynamicShadows && DynamicShadowDistanceStationaryLight > 0 && Mobility == EComponentMobility::Stationary;
 		}
 
+		if (PropertyName == GET_MEMBER_NAME_STRING_CHECKED(UDirectionalLightComponent, DynamicShadowDistanceMovableLight))
+		{
+			return CastShadows && CastDynamicShadows;
+		}
+
 		if (PropertyName == GET_MEMBER_NAME_STRING_CHECKED(UDirectionalLightComponent, DynamicShadowCascades)
 			|| PropertyName == GET_MEMBER_NAME_STRING_CHECKED(UDirectionalLightComponent, CascadeDistributionExponent)
 			|| PropertyName == GET_MEMBER_NAME_STRING_CHECKED(UDirectionalLightComponent, CascadeTransitionFraction)
 			|| PropertyName == GET_MEMBER_NAME_STRING_CHECKED(UDirectionalLightComponent, ShadowDistanceFadeoutFraction)
-			|| PropertyName == GET_MEMBER_NAME_STRING_CHECKED(UDirectionalLightComponent, bUseInsetShadowsForMovableObjects))
+			|| PropertyName == GET_MEMBER_NAME_STRING_CHECKED(UDirectionalLightComponent, bUseInsetShadowsForMovableObjects)
+			|| PropertyName == GET_MEMBER_NAME_STRING_CHECKED(UDirectionalLightComponent, FarShadowCascadeCount))
+		{
+			return bShadowCascases;
+		}
+
+		if (PropertyName == GET_MEMBER_NAME_STRING_CHECKED(UDirectionalLightComponent, FarShadowDistance))
 		{
-			return CastShadows 
-				&& CastDynamicShadows 
-				&& ((DynamicShadowDistanceMovableLight > 0 && Mobility == EComponentMobility::Movable)
-					|| (DynamicShadowDistanceStationaryLight > 0 && Mobility == EComponentMobility::Stationary));
+			return bShadowCascases && FarShadowCascadeCount > 0;
 		}
 
 		if (PropertyName == GET_MEMBER_NAME_STRING_CHECKED(UDirectionalLightComponent, DistanceFieldShadowDistance)
diff --git a/Engine/Source/Runtime/Engine/Private/LightComponent.cpp b/Engine/Source/Runtime/Engine/Private/LightComponent.cpp
--- a/Engine/Source/Runtime/Engine/Private/LightComponent.cpp
+++ b/Engine/Source/Runtime/Engine/Private/LightComponent.cpp
@@ -231,6 +231,8 @@ FLightSceneProxy::FLightSceneProxy(const ULightComponent* InLightComponent)
 	, ComponentName(InLightComponent->GetOwner() ? InLightComponent->GetOwner()->GetFName() : InLightComponent->GetFName())
 	, LevelName(InLightComponent->GetOutermost()->GetFName())
 	, StatId(InLightComponent->GetStatID(true))
+	, FarShadowDistance(0)
+	, FarShadowCascadeCount(0)
 {
 	// Brightness in Lumens
 	float LightBrightness = InLightComponent->ComputeLightBrightness();
@@ -539,6 +541,8 @@ void ULightComponent::PostEditChangeProperty(FPropertyChangedEvent& PropertyChan
 		PropertyName != GET_MEMBER_NAME_STRING_CHECKED(UDirectionalLightComponent, DynamicShadowDistanceMovableLight) &&
 		PropertyName != GET_MEMBER_NAME_STRING_CHECKED(UDirectionalLightComponent, DynamicShadowDistanceStationaryLight) &&
 		PropertyName != GET_MEMBER_NAME_STRING_CHECKED(UDirectionalLightComponent, DynamicShadowCascades) &&
+		PropertyName != GET_MEMBER_NAME_STRING_CHECKED(UDirectionalLightComponent, FarShadowDistance) &&
+		PropertyName != GET_MEMBER_NAME_STRING_CHECKED(UDirectionalLightComponent, FarShadowCascadeCount) &&
 		PropertyName != GET_MEMBER_NAME_STRING_CHECKED(UDirectionalLightComponent, CascadeDistributionExponent) &&
 		PropertyName != GET_MEMBER_NAME_STRING_CHECKED(UDirectionalLightComponent, CascadeTransitionFraction) &&
 		PropertyName != GET_MEMBER_NAME_STRING_CHECKED(UDirectionalLightComponent, ShadowDistanceFadeoutFraction) &&
diff --git a/Engine/Source/Runtime/Engine/Private/PrimitiveSceneProxy.cpp b/Engine/Source/Runtime/Engine/Private/PrimitiveSceneProxy.cpp
--- a/Engine/Source/Runtime/Engine/Private/PrimitiveSceneProxy.cpp
+++ b/Engine/Source/Runtime/Engine/Private/PrimitiveSceneProxy.cpp
@@ -45,6 +45,7 @@ FPrimitiveSceneProxy::FPrimitiveSceneProxy(const UPrimitiveComponent* InComponen
 ,	bCastShadowAsTwoSided(InComponent->bCastShadowAsTwoSided)
 ,	bSelfShadowOnly(InComponent->bSelfShadowOnly)
 ,	bCastInsetShadow(InComponent->bSelfShadowOnly ? true : InComponent->bCastInsetShadow)	// Assumed to be enabled if bSelfShadowOnly is enabled.
+,	bCastFarShadow(InComponent->bCastFarShadow)
 ,	bLightAttachmentsAsGroup(InComponent->bLightAttachmentsAsGroup)
 ,	bStaticElementsAlwaysUseProxyPrimitiveUniformBuffer(false)
 ,	bAlwaysHasVelocity(false)
diff --git a/Engine/Source/Runtime/Engine/Public/PrimitiveSceneProxy.h b/Engine/Source/Runtime/Engine/Public/PrimitiveSceneProxy.h
--- a/Engine/Source/Runtime/Engine/Public/PrimitiveSceneProxy.h
+++ b/Engine/Source/Runtime/Engine/Public/PrimitiveSceneProxy.h
@@ -346,6 +346,7 @@ public:
 	inline bool CastsShadowAsTwoSided() const { return bCastShadowAsTwoSided; }
 	inline bool CastsSelfShadowOnly() const { return bSelfShadowOnly; }
 	inline bool CastsInsetShadow() const { return bCastInsetShadow; }
+	inline bool CastsFarShadow() const { return bCastFarShadow; }
 	inline bool LightAttachmentsAsGroup() const { return bLightAttachmentsAsGroup; }
 	inline bool StaticElementsAlwaysUseProxyPrimitiveUniformBuffer() const { return bStaticElementsAlwaysUseProxyPrimitiveUniformBuffer; }
 	inline bool ShouldUseAsOccluder() const { return bUseAsOccluder; }
@@ -531,6 +532,9 @@ protected:
 	/** Whether this component should create a per-object shadow that gives higher effective shadow resolution. true if bSelfShadowOnly is true. */
 	uint32 bCastInsetShadow : 1;
 
+	/* When enabled, the component will be rendering into the distant shadow cascades (only for directional lights). */
+	uint32 bCastFarShadow : 1;
+
 	/** 
 	 * Whether to light this component and any attachments as a group.  This only has effect on the root component of an attachment tree.
 	 * When enabled, attached component shadowing settings like bCastInsetShadow, bCastVolumetricTranslucentShadow, etc, will be ignored.
diff --git a/Engine/Source/Runtime/Engine/Public/SceneManagement.h b/Engine/Source/Runtime/Engine/Public/SceneManagement.h
--- a/Engine/Source/Runtime/Engine/Public/SceneManagement.h
+++ b/Engine/Source/Runtime/Engine/Public/SceneManagement.h
@@ -521,6 +521,9 @@ public:
 	FPlane NearFrustumPlane;
 	FPlane FarFrustumPlane;
 
+	/** When enabled, the cascade only renders objects marked with bCastFarShadows enabled (e.g. Landscape). */
+	bool bFarShadowCascade;
+
 	FShadowCascadeSettings()
 		: SplitNear(0.0f)
 		, SplitFar(WORLD_MAX)
@@ -528,6 +531,7 @@ public:
 		, SplitFarFadeRegion(0.0f)
 		, FadePlaneOffset(SplitFar)
 		, FadePlaneLength(SplitFar - FadePlaneOffset)
+		, bFarShadowCascade(false)
 	{
 	}
 };
@@ -688,15 +692,17 @@ public:
 	/** Whether this light should create per object shadows for dynamic objects. */
 	virtual bool ShouldCreatePerObjectShadowsForDynamicObjects() const;
 
+	/** Returns the number of view dependent shadows this light will create, not counting distance field shadow cascades. */
 	virtual uint32 GetNumViewDependentWholeSceneShadows(const FSceneView& View) const { return 0; }
 
 	/**
 	 * Sets up a projected shadow initializer that's dependent on the current view for shadows from the entire scene.
+	 * @param InCascadeIndex cascade index or INDEX_NONE for the distance field cascade
 	 * @return True if the whole-scene projected shadow should be used.
 	 */
 	virtual bool GetViewDependentWholeSceneProjectedShadowInitializer(
 		const class FSceneView& View, 
-		uint32 InShadowSplitIndex,
+		int32 InCascadeIndex,
 		class FWholeSceneProjectedShadowInitializer& OutInitializer) const
 	{
 		return false;
@@ -725,8 +731,9 @@ public:
 		return false;
 	}
 
+	// @param InCascadeIndex cascade index or INDEX_NONE for the distance field cascade
 	// @param OutCascadeSettings can be 0
-	virtual FSphere GetShadowSplitBounds(const class FSceneView& View, uint32 InShadowSplitIndex, FShadowCascadeSettings* OutCascadeSettings) const { return FSphere(FVector::ZeroVector, 0); }
+	virtual FSphere GetShadowSplitBounds(const class FSceneView& View, int32 InCascadeIndex, FShadowCascadeSettings* OutCascadeSettings) const { return FSphere(FVector::ZeroVector, 0); }
 	virtual FSphere GetShadowSplitBoundsDepthRange(const FSceneView& View, float SplitNear, float SplitFar, FShadowCascadeSettings* OutCascadeSettings) const { return FSphere(FVector::ZeroVector, 0); }
 
 	virtual bool GetScissorRect(FIntRect& ScissorRect, const FSceneView& View) const
@@ -739,6 +746,8 @@ public:
 	{
 	}
 
+	virtual bool ShouldCreateRayTracedCascade(ERHIFeatureLevel::Type Type) const { return false; }
+
 	// Accessors.
 	float GetUserShadowBias() const { return ShadowBias; }
 
@@ -900,6 +909,12 @@ protected:
 	/** Used for dynamic stats */
 	TStatId StatId;
 
+	/** Only for whole scene directional lights, if FarShadowCascadeCount > 0 and FarShadowDistance >= WholeSceneDynamicShadowRadius, where far shadow cascade should end. */
+	float FarShadowDistance;
+
+	/** Only for whole scene directional lights, 0: no FarShadowCascades, otherwise the count of cascades between WholeSceneDynamicShadowRadius and FarShadowDistance that are covered by distant shadow cascades. */
+	uint32 FarShadowCascadeCount;
+
 	/**
 	 * Updates the light proxy's cached transforms.
 	 * @param InLightToWorld - The new light-to-world transform.
diff --git a/Engine/Source/Runtime/Landscape/Classes/LandscapeProxy.h b/Engine/Source/Runtime/Landscape/Classes/LandscapeProxy.h
--- a/Engine/Source/Runtime/Landscape/Classes/LandscapeProxy.h
+++ b/Engine/Source/Runtime/Landscape/Classes/LandscapeProxy.h
@@ -391,6 +391,10 @@ public:
 	UPROPERTY(EditAnywhere, AdvancedDisplay, Category=Lighting, meta=(DisplayName = "Shadow Two Sided"))
 	uint32 bCastShadowAsTwoSided:1;
 
+	/** Whether this primitive should cast shadows in the far shadow cascades. */
+	UPROPERTY(EditAnywhere, AdvancedDisplay, Category=Lighting, meta=(DisplayName = "Far Shadow"))
+	uint32 bCastFarShadow:1;
+
 	UPROPERTY()
 	uint32 bIsProxy:1;
 
diff --git a/Engine/Source/Runtime/Landscape/Private/Landscape.cpp b/Engine/Source/Runtime/Landscape/Private/Landscape.cpp
--- a/Engine/Source/Runtime/Landscape/Private/Landscape.cpp
+++ b/Engine/Source/Runtime/Landscape/Private/Landscape.cpp
@@ -83,6 +83,8 @@ ULandscapeComponent::ULandscapeComponent(const FObjectInitializer& ObjectInitial
 	SetCollisionProfileName(UCollisionProfile::BlockAll_ProfileName);
 	bGenerateOverlapEvents = false;
 	CastShadow = true;
+	// by default we want to see the Landscape shadows even in the far shadow cascades
+	bCastFarShadow = true;
 	bUseAsOccluder = true;
 	bAllowCullDistanceVolume = false;
 	CollisionMipLevel = 0;
@@ -516,6 +518,8 @@ ALandscapeProxy::ALandscapeProxy(const FObjectInitializer& ObjectInitializer)
 	bHidden = false;
 	bReplicateMovement = false;
 	bCanBeDamaged = false;
+	// by default we want to see the Landscape shadows even in the far shadow cascades
+	bCastFarShadow = true;
 
 	USceneComponent* SceneComponent = ObjectInitializer.CreateDefaultSubobject<USceneComponent>(this, TEXT("RootComponent0"));
 	RootComponent = SceneComponent;
diff --git a/Engine/Source/Runtime/Landscape/Private/LandscapeEdit.cpp b/Engine/Source/Runtime/Landscape/Private/LandscapeEdit.cpp
--- a/Engine/Source/Runtime/Landscape/Private/LandscapeEdit.cpp
+++ b/Engine/Source/Runtime/Landscape/Private/LandscapeEdit.cpp
@@ -3702,7 +3702,11 @@ void ALandscapeProxy::PostEditChangeProperty(FPropertyChangedEvent& PropertyChan
 	{
 		RecreateCollisionComponents();
 	}
-	else if (PropertyName == FName(TEXT("bCastStaticShadow")) || PropertyName == FName(TEXT("bCastShadowAsTwoSided")))
+	else 
+		if(PropertyName == FName(TEXT("bCastStaticShadow"))
+		|| PropertyName == FName(TEXT("bCastShadowAsTwoSided"))
+		|| PropertyName == FName(TEXT("bCastFarShadow"))
+		)
 	{
 		// Replicate shared properties to all components.
 		for (int32 ComponentIndex = 0; ComponentIndex < LandscapeComponents.Num(); ComponentIndex++)
@@ -3712,6 +3716,7 @@ void ALandscapeProxy::PostEditChangeProperty(FPropertyChangedEvent& PropertyChan
 			{
 				Comp->bCastStaticShadow = bCastStaticShadow;
 				Comp->bCastShadowAsTwoSided = bCastShadowAsTwoSided;
+				Comp->bCastFarShadow = bCastFarShadow;
 			}
 		}
 	}
diff --git a/Engine/Source/Runtime/Renderer/Private/DistanceFieldShadowing.cpp b/Engine/Source/Runtime/Renderer/Private/DistanceFieldShadowing.cpp
--- a/Engine/Source/Runtime/Renderer/Private/DistanceFieldShadowing.cpp
+++ b/Engine/Source/Runtime/Renderer/Private/DistanceFieldShadowing.cpp
@@ -822,7 +822,9 @@ void FProjectedShadowInfo::RenderRayTracedDistanceFieldProjection(FRHICommandLis
 					
 				if (bDirectionalLight)
 				{
-					check(CascadeSettings.FadePlaneLength == 0);
+					// todo: when using far shadows we need to fix the shader to get a soft transition
+//					check(CascadeSettings.FadePlaneLength == 0);
+
 					// first cascade rendered or old method doesn't require fading (CO_Min is needed to combine multiple shadow passes)
 					// The ray traced cascade should always be first
 					RHICmdList.SetBlendState(TStaticBlendState<CW_RG, BO_Min, BF_One, BF_One>::GetRHI());
diff --git a/Engine/Source/Runtime/Renderer/Private/ShadowSetup.cpp b/Engine/Source/Runtime/Renderer/Private/ShadowSetup.cpp
--- a/Engine/Source/Runtime/Renderer/Private/ShadowSetup.cpp
+++ b/Engine/Source/Runtime/Renderer/Private/ShadowSetup.cpp
@@ -591,10 +591,10 @@ FProjectedShadowInfo::FProjectedShadowInfo(
 			PreShadowTranslation = -SnappedWorldPosition;
 		}
 
-		if (Initializer.InitShadowSplitIndex >= 0 && bDirectionalLight)
+		if (ShadowSplitIndex >= 0 && bDirectionalLight)
 		{
 			checkSlow(InDependentView);
-			ShadowBounds = InLightSceneInfo->Proxy->GetShadowSplitBounds(*InDependentView, ShadowSplitIndex, 0);
+			ShadowBounds = InLightSceneInfo->Proxy->GetShadowSplitBounds(*InDependentView, bRayTracedDistanceFieldShadow ? INDEX_NONE : ShadowSplitIndex, 0);
 		}
 		else
 		{
@@ -732,6 +732,21 @@ void FProjectedShadowInfo::AddSubjectPrimitive(FPrimitiveSceneInfo* PrimitiveSce
 				{
 					FViewInfo& CurrentView = *Views[ViewIndex];
 
+					// 	in a far shadow cascade we only render objects marked with bCastFarShadows enabled (e.g. Landscape).
+					if(CascadeSettings.bFarShadowCascade)
+					{
+						if(!Proxy->CastsFarShadow())
+						{
+							// cull objects that are not supposed to be in the far shadow cascades
+							continue;
+						}
+						if(bRayTracedDistanceFieldShadow && Proxy->AffectsDistanceFieldLighting())
+						{
+							// cull objects that already have been handles by the ray traced distance field ones
+							continue;
+						}
+					}
+
 					const float DistanceSquared = ( Bounds.Origin - CurrentView.ShadowViewMatrices.ViewOrigin ).SizeSquared();
 					const bool bDrawShadowDepth = FMath::Square( Bounds.SphereRadius ) > FMath::Square( GMinScreenRadiusForShadowCaster ) * DistanceSquared;
 					if( !bDrawShadowDepth )
@@ -1807,14 +1822,14 @@ void FSceneRenderer::InitProjectedShadowVisibility(FRHICommandListImmediate& RHI
 								float Far = ProjectedShadowInfo.CascadeSettings.SplitFar;
 
 								// Camera Subfrustum
-								DrawFrustumWireframe(&ShadowFrustumPDI, (ViewMatrix * FPerspectiveMatrix(ActualFOV, AspectRatio, 1.0f, Near, Mid)).InverseFast(), Color, 0);
-								DrawFrustumWireframe(&ShadowFrustumPDI, (ViewMatrix * FPerspectiveMatrix(ActualFOV, AspectRatio, 1.0f, Mid, Far)).InverseFast(), FColor::White, 0);
+								DrawFrustumWireframe(&ShadowFrustumPDI, (ViewMatrix * FPerspectiveMatrix(ActualFOV, AspectRatio, 1.0f, Near, Mid)).Inverse(), Color, 0);
+								DrawFrustumWireframe(&ShadowFrustumPDI, (ViewMatrix * FPerspectiveMatrix(ActualFOV, AspectRatio, 1.0f, Mid, Far)).Inverse(), FColor::White, 0);
 
 								// Subfrustum Sphere Bounds
 								DrawWireSphere(&ShadowFrustumPDI, FTransform(ProjectedShadowInfo.ShadowBounds.Center), Color, ProjectedShadowInfo.ShadowBounds.W, 40, 0);
 
 								// Shadow Map Projection Bounds
-								DrawFrustumWireframe(&ShadowFrustumPDI, ProjectedShadowInfo.SubjectAndReceiverMatrix.InverseFast() * FTranslationMatrix(-ProjectedShadowInfo.PreShadowTranslation), Color, 0);
+								DrawFrustumWireframe(&ShadowFrustumPDI, ProjectedShadowInfo.SubjectAndReceiverMatrix.Inverse() * FTranslationMatrix(-ProjectedShadowInfo.PreShadowTranslation), Color, 0);
 							}
 							else
 							{
@@ -1854,14 +1869,24 @@ void FSceneRenderer::InitProjectedShadowVisibility(FRHICommandListImmediate& RHI
 
 					if(VisibleLightViewInfo.bInViewFrustum)
 					{
-						UE_LOG(LogRenderer, Display, TEXT("   Shadow %d/%d: %d"),  ShadowIndex, ShadowCount, ProjectedShadowInfo.ShadowId);
+						UE_LOG(LogRenderer, Display, TEXT("   Shadow %d/%d: ShadowId=%d"),  ShadowIndex, ShadowCount, ProjectedShadowInfo.ShadowId);
 						UE_LOG(LogRenderer, Display, TEXT("    WholeSceneDir=%d SplitIndex=%d near=%f far=%f"),
 							ProjectedShadowInfo.IsWholeSceneDirectionalShadow(),
 							ProjectedShadowInfo.ShadowSplitIndex,
 							ProjectedShadowInfo.CascadeSettings.SplitNear,
 							ProjectedShadowInfo.CascadeSettings.SplitFar);
-						UE_LOG(LogRenderer, Display, TEXT("    bDistField=%d"),
-							ProjectedShadowInfo.bRayTracedDistanceFieldShadow);
+						UE_LOG(LogRenderer, Display, TEXT("    bDistField=%d bFarShadows=%d Bounds=%f,%f,%f,%f"),
+							ProjectedShadowInfo.bRayTracedDistanceFieldShadow,
+							ProjectedShadowInfo.CascadeSettings.bFarShadowCascade,
+							ProjectedShadowInfo.ShadowBounds.Center.X,
+							ProjectedShadowInfo.ShadowBounds.Center.Y,
+							ProjectedShadowInfo.ShadowBounds.Center.Z,
+							ProjectedShadowInfo.ShadowBounds.W);
+						UE_LOG(LogRenderer, Display, TEXT("    SplitFadeRegion=%f .. %f FadePlaneOffset=%f FadePlaneLength=%f"),
+							ProjectedShadowInfo.CascadeSettings.SplitNearFadeRegion,
+							ProjectedShadowInfo.CascadeSettings.SplitFarFadeRegion,
+							ProjectedShadowInfo.CascadeSettings.FadePlaneOffset,
+							ProjectedShadowInfo.CascadeSettings.FadePlaneLength);			
 					}
 				}
 			}
@@ -2125,12 +2150,26 @@ void FSceneRenderer::AddViewDependentWholeSceneShadowsForView(
 		// If rendering in stereo mode we render shadow depths only for the left eye, but project for both eyes!
 		if (View.StereoPass != eSSP_RIGHT_EYE)
 		{
-			const uint32 NumSplits = LightSceneInfo.Proxy->GetNumViewDependentWholeSceneShadows(View);
-			for (uint32 LocalShadowSplitIndex = 0; LocalShadowSplitIndex < NumSplits; LocalShadowSplitIndex++)
+			const bool bExtraDistanceFieldCascade = LightSceneInfo.Proxy->ShouldCreateRayTracedCascade(View.GetFeatureLevel());
+
+			const int32 ProjectionCount = LightSceneInfo.Proxy->GetNumViewDependentWholeSceneShadows(View) + (bExtraDistanceFieldCascade?1:0);
+
+			checkSlow(INDEX_NONE == -1);
+
+			// todo: this code can be simplified by computing all the distances in one place - avoiding some redundant work and complexity
+			for (int32 Index = 0; Index < ProjectionCount; Index++)
 			{
 				FWholeSceneProjectedShadowInitializer ProjectedShadowInitializer;
 
-				if (LightSceneInfo.Proxy->GetViewDependentWholeSceneProjectedShadowInitializer(View, LocalShadowSplitIndex, ProjectedShadowInitializer))
+				int32 LocalIndex = Index;
+
+				// Indexing like this puts the raytraced shadow cascade last (might not be needed)
+				if(bExtraDistanceFieldCascade && LocalIndex + 1 == ProjectionCount)
+				{
+					LocalIndex = INDEX_NONE;
+				}
+
+				if (LightSceneInfo.Proxy->GetViewDependentWholeSceneProjectedShadowInitializer(View, LocalIndex, ProjectedShadowInitializer))
 				{
 					const FIntPoint ShadowBufferResolution = GSceneRenderTargets.GetShadowDepthTextureResolution();
 					// Create the projected shadow info.
