Fixed a comment on VisualizeSenses show flag to make more sense #UE4

diff --git a/Engine/Source/Runtime/Engine/Public/ShowFlagsValues.inl b/Engine/Source/Runtime/Engine/Public/ShowFlagsValues.inl
--- a/Engine/Source/Runtime/Engine/Public/ShowFlagsValues.inl
+++ b/Engine/Source/Runtime/Engine/Public/ShowFlagsValues.inl
@@ -280,7 +280,7 @@ SHOWFLAG_ALWAYS_ACCESSIBLE(DistanceFieldGI, SFG_LightingFeatures, LOCTEXT("Dista
 SHOWFLAG_FIXED_IN_SHIPPING(VisualizeSSR, 0, SFG_Visualize, LOCTEXT("VisualizeSSR", "Screen Space Reflections"))
 /** Force the use of the GBuffer. */
 SHOWFLAG_ALWAYS_ACCESSIBLE(ForceGBuffer, SFG_Hidden, LOCTEXT("ForceGBuffer", "Force usage of GBuffer"))
-/** Visualize the senses of the Senses component */
+/** Visualize the senses configuration of AIs' PawnSensingComponent */
 SHOWFLAG_ALWAYS_ACCESSIBLE(VisualizeSenses, SFG_Advanced, LOCTEXT("VisualizeSenses", "Senses"))
 
 #undef SHOWFLAG_ALWAYS_ACCESSIBLE
