Fixes to the prototype tonemapper (better controls, better controls for interpolation of post volumes).

diff --git a/Engine/Shaders/TonemapCommon.usf b/Engine/Shaders/TonemapCommon.usf
--- a/Engine/Shaders/TonemapCommon.usf
+++ b/Engine/Shaders/TonemapCommon.usf
@@ -95,7 +95,7 @@ half4 PhotoMatrixR_Lift;
 half4 PhotoMatrixG_LiftScale;
 half4 PhotoMatrixB_Gamma;
 half4 PhotoTint_Over;
-half4 PhotoOverTint;
+half4 PhotoOverTint_OneMinusMid;
 
 half3 FilmPostProcessPhoto(half3 LinearColor) 
 {
@@ -113,9 +113,9 @@ half3 FilmPostProcessPhoto(half3 LinearColor)
 	// This actually moves over to PhotoOverTint, 
 	// which can be a highly colored value, which later turns to white via highlight compression.
 	half OverLuma = dot(PhotoColor, half3(0.2126, 0.7152, 0.0722)); // BT.709 luma seems to work best.
-	half Over = (PhotoTint_Over.a*OverLuma) / ((1.0-0.18)+OverLuma);
-	PhotoColor = lerp(PhotoColor, PhotoOverTint.rgb*OverLuma, Over);
-	// Highlight compression without change of 18% mid-tone.
-	PhotoColor *= 1.0/((1.0-0.18)+PhotoColor);    
+	half Over = (PhotoTint_Over.a*OverLuma) / (PhotoOverTint_OneMinusMid.a+OverLuma);
+	PhotoColor = lerp(PhotoColor, PhotoOverTint_OneMinusMid.rgb*OverLuma, Over);
+	// Highlight compression without change of mid-tone.
+	PhotoColor *= 1.0/(PhotoOverTint_OneMinusMid.a+PhotoColor);    
 	return PhotoColor;
 }
diff --git a/Engine/Source/Runtime/Engine/Classes/Engine/Scene.h b/Engine/Source/Runtime/Engine/Classes/Engine/Scene.h
--- a/Engine/Source/Runtime/Engine/Classes/Engine/Scene.h
+++ b/Engine/Source/Runtime/Engine/Classes/Engine/Scene.h
@@ -73,6 +73,12 @@ struct FPostProcessSettings
 	UPROPERTY(BlueprintReadWrite, Category=Overrides, meta=(PinHiddenByDefault))
 	uint32 bOverride_FilmShadowTintAmount:1;
 
+	UPROPERTY(BlueprintReadWrite, Category=Overrides, meta=(PinHiddenByDefault))
+	uint32 bOverride_PhotoExposure:1;
+
+	UPROPERTY(BlueprintReadWrite, Category=Overrides, meta=(PinHiddenByDefault))
+	uint32 bOverride_PhotoMid:1;
+
 	UPROPERTY(BlueprintReadWrite, Category=Overrides, meta=(PinHiddenByDefault))
 	uint32 bOverride_PhotoWhite:1;
 
@@ -88,9 +94,6 @@ struct FPostProcessSettings
 	UPROPERTY(BlueprintReadWrite, Category=Overrides, meta=(PinHiddenByDefault))
 	uint32 bOverride_PhotoChannelMixerBlue:1;
 
-	UPROPERTY(BlueprintReadWrite, Category=Overrides, meta=(PinHiddenByDefault))
-	uint32 bOverride_PhotoBlackIn:1;
-
 	UPROPERTY(BlueprintReadWrite, Category=Overrides, meta=(PinHiddenByDefault))
 	uint32 bOverride_PhotoBlackOut:1;
 
@@ -400,7 +403,11 @@ struct FPostProcessSettings
 	UPROPERTY(interp, BlueprintReadWrite, Category=Film, AdvancedDisplay, meta=(UIMin = "1.0", UIMax = "4.0", editcondition = "bOverride_FilmDynamicRange", DisplayName = "Dynamic Range"))
 	float FilmDynamicRange;
 
-	UPROPERTY(interp, BlueprintReadWrite, Category=Photo, AdvancedDisplay, meta=(editcondition = "bOverride_PhotoWhite", DisplayName = "White", HideAlphaChannel))
+	UPROPERTY(interp, BlueprintReadWrite, Category=Photo, AdvancedDisplay, meta=(UIMin = "0.0", UIMax = "1.0", editcondition = "bOverride_PhotoMid", DisplayName = "Mid-point (linear)"))
+	float PhotoMid;
+	UPROPERTY(interp, BlueprintReadWrite, Category=Photo, AdvancedDisplay, meta=(UIMin = "-16.0", UIMax = "16.0", editcondition = "bOverride_PhotoExposure", DisplayName = "Exposure (+/- stops)"))
+	float PhotoExposure;
+	UPROPERTY(interp, BlueprintReadWrite, Category=Photo, AdvancedDisplay, meta=(editcondition = "bOverride_PhotoWhite", DisplayName = "White (pre-tonemap)", HideAlphaChannel))
 	FLinearColor PhotoWhite;
 	UPROPERTY(interp, BlueprintReadWrite, Category=Photo, AdvancedDisplay, meta=(UIMin = "0.0", UIMax = "2.0", editcondition = "bOverride_PhotoSaturation", DisplayName = "Saturation"))
 	float PhotoSaturation;
@@ -410,13 +417,11 @@ struct FPostProcessSettings
 	FLinearColor PhotoChannelMixerGreen;
 	UPROPERTY(interp, BlueprintReadWrite, Category=Photo, AdvancedDisplay, meta=(editcondition = "bOverride_PhotoChannelMixerBlue", DisplayName = " Channel Mixer Blue", HideAlphaChannel))
 	FLinearColor PhotoChannelMixerBlue;
-	UPROPERTY(interp, BlueprintReadWrite, Category=Photo, AdvancedDisplay, meta=(UIMin = "0.0", UIMax = "1.0", editcondition = "bOverride_PhotoBlackIn", DisplayName = "Black In"))
-	float PhotoBlackIn;
-	UPROPERTY(interp, BlueprintReadWrite, Category=Photo, AdvancedDisplay, meta=(UIMin = "0.0", UIMax = "1.0", editcondition = "bOverride_PhotoBlackOut", DisplayName = "Black Out"))
+	UPROPERTY(interp, BlueprintReadWrite, Category=Photo, AdvancedDisplay, meta=(UIMin = "-6.0", UIMax = "2.0", editcondition = "bOverride_PhotoBlackOut", DisplayName = "Lift (+/- stops from mid^3)"))
 	float PhotoBlackOut;
-	UPROPERTY(interp, BlueprintReadWrite, Category=Photo, AdvancedDisplay, meta=(UIMin = "0.5", UIMax = "4.0", editcondition = "bOverride_PhotoContrast", DisplayName = "Contrast"))
+	UPROPERTY(interp, BlueprintReadWrite, Category=Photo, AdvancedDisplay, meta=(UIMin = "1.0", UIMax = "4.0", editcondition = "bOverride_PhotoContrast", DisplayName = "Contrast (gamma)"))
 	float PhotoContrast;
-	UPROPERTY(interp, BlueprintReadWrite, Category=Photo, AdvancedDisplay, meta=(editcondition = "bOverride_PhotoTint", DisplayName = "Tint", HideAlphaChannel))
+	UPROPERTY(interp, BlueprintReadWrite, Category=Photo, AdvancedDisplay, meta=(editcondition = "bOverride_PhotoTint", DisplayName = "Tint (post-tonemap)", HideAlphaChannel))
 	FLinearColor PhotoTint;
 	UPROPERTY(interp, BlueprintReadWrite, Category=Photo, AdvancedDisplay, meta=(UIMin = "0.0", UIMax = "1.0", editcondition = "bOverride_PhotoOver", DisplayName = "Over Exposure Mix"))
 	float PhotoOver;
@@ -854,13 +859,14 @@ struct FPostProcessSettings
 		FilmShadowTintBlend = 0.5;
 		FilmShadowTintAmount = 0.0;
 
+		PhotoExposure = 0.0f;
+		PhotoMid = 0.18f;
 		PhotoWhite = FLinearColor(1.0f,1.0f,1.0f);
 		PhotoSaturation = 1.0f;
 		PhotoChannelMixerRed = FLinearColor(1.0f,0.0f,0.0f);
 		PhotoChannelMixerGreen = FLinearColor(0.0f,1.0f,0.0f);
 		PhotoChannelMixerBlue = FLinearColor(0.0f,0.0f,1.0f);
-		PhotoBlackIn = 1.0f/512.0f;
-		PhotoBlackOut = 1.0f/1024.0f;
+		PhotoBlackOut = -1.0f;
 		PhotoContrast = 1.2f;
 		PhotoTint = FLinearColor(1.0f,1.0f,1.0f);
 		PhotoOver = 1.0f/32.0f;
diff --git a/Engine/Source/Runtime/Engine/Private/SceneView.cpp b/Engine/Source/Runtime/Engine/Private/SceneView.cpp
--- a/Engine/Source/Runtime/Engine/Private/SceneView.cpp
+++ b/Engine/Source/Runtime/Engine/Private/SceneView.cpp
@@ -715,12 +715,13 @@ void FSceneView::OverridePostProcessSettings(const FPostProcessSettings& Src, fl
 	LERP_PP(FilmShadowTintBlend);
 	LERP_PP(FilmShadowTintAmount);
 
+	LERP_PP(PhotoExposure);
+	LERP_PP(PhotoMid);
 	LERP_PP(PhotoWhite);
 	LERP_PP(PhotoSaturation);
 	LERP_PP(PhotoChannelMixerRed);
 	LERP_PP(PhotoChannelMixerGreen);
 	LERP_PP(PhotoChannelMixerBlue);
-	LERP_PP(PhotoBlackIn);
 	LERP_PP(PhotoBlackOut);
 	LERP_PP(PhotoContrast);
 	LERP_PP(PhotoTint);
diff --git a/Engine/Source/Runtime/Renderer/Private/PostProcess/PostProcessTonemap.cpp b/Engine/Source/Runtime/Renderer/Private/PostProcess/PostProcessTonemap.cpp
--- a/Engine/Source/Runtime/Renderer/Private/PostProcess/PostProcessTonemap.cpp
+++ b/Engine/Source/Runtime/Renderer/Private/PostProcess/PostProcessTonemap.cpp
@@ -946,6 +946,8 @@ void FilmPostSetConstants(FVector4* RESTRICT const Constants, const uint32 Confi
 // Improved film post: setting of constants.
 void FilmPostSetConstantsPhoto(FVector4* RESTRICT const Constants, const FPostProcessSettings* RESTRICT const FinalPostProcessSettings)
 {
+	float PhotoExp = FMath::Pow(2.0f,FMath::Clamp(FinalPostProcessSettings->PhotoExposure, -64.0f, 64.0f));
+	float PhotoMid = 0.18f * FMath::Pow(2.0f,FMath::Clamp(FinalPostProcessSettings->PhotoMid, -64.0f, 64.0f));
 	FVector PhotoWhite(FinalPostProcessSettings->PhotoWhite);
 	float PhotoSat = FMath::Clamp(FinalPostProcessSettings->PhotoSaturation, 0.0f, 2.0f) - 1.0f;
 	FVector PhotoMixR(FinalPostProcessSettings->PhotoChannelMixerRed);
@@ -953,16 +955,11 @@ void FilmPostSetConstantsPhoto(FVector4* RESTRICT const Constants, const FPostPr
 	FVector PhotoMixB(FinalPostProcessSettings->PhotoChannelMixerBlue);
 	FVector PhotoOverTint(FinalPostProcessSettings->PhotoOverTint);
 	float PhotoGamma = FMath::Clamp(FinalPostProcessSettings->PhotoContrast, 0.0f, 4.0f);
-	float PhotoBlackOut = FMath::Clamp(FinalPostProcessSettings->PhotoBlackOut, 0.0f, 1.0f) * (1.0f/64.0f);
-	float PhotoBlackIn = FMath::Clamp(FinalPostProcessSettings->PhotoBlackIn, 0.0f, 1.0f) * (1.0f/64.0f);
+	float PhotoBlackIn = (PhotoMid*PhotoMid*PhotoMid);
+	float PhotoBlackOut = (PhotoMid*PhotoMid*PhotoMid)*FMath::Pow(2.0f,FMath::Clamp(FinalPostProcessSettings->PhotoBlackOut, -64.0f, 64.0f));
 	FVector PhotoTint(FinalPostProcessSettings->PhotoTint);
 	float PhotoOver = FMath::Clamp(FinalPostProcessSettings->PhotoOver, 0.0f, 1.0f);
 
-	// Place holder for eventual exposure adjustment.
-	float PhotoExp = 1.0;
-
-	// Constant, don't change...
-	float PhotoMid = 0.18f;
 	// Conversion from linear rgb to luma (using HDTV coef).
 	FVector LumaWeights = FVector(0.2126f, 0.7152f, 0.0722f);
 
@@ -1023,7 +1020,7 @@ void FilmPostSetConstantsPhoto(FVector4* RESTRICT const Constants, const FPostPr
 	Constants[1] = FVector4(OutMatrixG, OutLiftScale);
 	Constants[2] = FVector4(OutMatrixB, PhotoGamma); 
 	Constants[3] = FVector4(OutTint, PhotoOver*(1.0f/3.0f)); 
-	Constants[4] = FVector4(PhotoOverTint, 0.0f);
+	Constants[4] = FVector4(PhotoOverTint, 1.0f-PhotoMid);
 }
 
 
@@ -1105,7 +1102,7 @@ public:
 	FShaderParameter PhotoMatrixG_LiftScale;
 	FShaderParameter PhotoMatrixB_Gamma;
 	FShaderParameter PhotoTint_Over;
-	FShaderParameter PhotoOverTint;
+	FShaderParameter PhotoOverTint_OneMinusMid;
 
 	//@HACK
 	FShaderParameter OverlayColor;
@@ -1139,7 +1136,7 @@ public:
 		PhotoMatrixG_LiftScale.Bind(Initializer.ParameterMap, TEXT("PhotoMatrixG_LiftScale"));
 		PhotoMatrixB_Gamma.Bind(Initializer.ParameterMap, TEXT("PhotoMatrixB_Gamma"));
 		PhotoTint_Over.Bind(Initializer.ParameterMap, TEXT("PhotoTint_Over"));
-		PhotoOverTint.Bind(Initializer.ParameterMap, TEXT("PhotoOverTint"));
+		PhotoOverTint_OneMinusMid.Bind(Initializer.ParameterMap, TEXT("PhotoOverTint_OneMinusMid"));
 		
 		OverlayColor.Bind(Initializer.ParameterMap, TEXT("OverlayColor"));
 	}
@@ -1152,7 +1149,7 @@ public:
 			<< TexScale << VignetteColorIntensity << GrainScaleBiasJitter
 			<< ColorGradingLUT << ColorGradingLUTSampler
 			<< ColorMatrixR_ColorCurveCd1 << ColorMatrixG_ColorCurveCd3Cm3 << ColorMatrixB_ColorCurveCm2 << ColorCurve_Cm0Cd0_Cd2_Ch0Cm1_Ch3 << ColorCurve_Ch1_Ch2 << ColorShadow_Luma << ColorShadow_Tint1 << ColorShadow_Tint2
-			<< PhotoMatrixR_Lift << PhotoMatrixG_LiftScale << PhotoMatrixB_Gamma << PhotoTint_Over << PhotoOverTint
+			<< PhotoMatrixR_Lift << PhotoMatrixG_LiftScale << PhotoMatrixB_Gamma << PhotoTint_Over << PhotoOverTint_OneMinusMid
 			<< OverlayColor;
 		
 		return bShaderHasOutdatedParameters;
@@ -1280,7 +1277,7 @@ public:
 			SetShaderValue(Context.RHICmdList, ShaderRHI, PhotoMatrixG_LiftScale, Constants[1]);
 			SetShaderValue(Context.RHICmdList, ShaderRHI, PhotoMatrixB_Gamma, Constants[2]); 
 			SetShaderValue(Context.RHICmdList, ShaderRHI, PhotoTint_Over, Constants[3]); 
-			SetShaderValue(Context.RHICmdList, ShaderRHI, PhotoOverTint, Constants[4]);
+			SetShaderValue(Context.RHICmdList, ShaderRHI, PhotoOverTint_OneMinusMid, Constants[4]);
 		}
 	}
 	
