UE4 - Material property for opting-out writing into velocity pass; remove primitive relevance flag

diff --git a/Engine/Source/Runtime/Engine/Classes/Materials/Material.h b/Engine/Source/Runtime/Engine/Classes/Materials/Material.h
--- a/Engine/Source/Runtime/Engine/Classes/Materials/Material.h
+++ b/Engine/Source/Runtime/Engine/Classes/Materials/Material.h
@@ -614,6 +614,10 @@ public:
 	UPROPERTY(EditAnywhere, Category=Material, AdvancedDisplay)
 	uint32 Wireframe:1;
 
+	/** Skips outputting velocity during the base pass. */
+	UPROPERTY(EditAnywhere, Category=Material, AdvancedDisplay)
+	uint32 bSkipVelocityOnBasePass:1;
+
 	UPROPERTY()
 	int32 EditorX;
 
diff --git a/Engine/Source/Runtime/Engine/Classes/Materials/MaterialInterface.h b/Engine/Source/Runtime/Engine/Classes/Materials/MaterialInterface.h
--- a/Engine/Source/Runtime/Engine/Classes/Materials/MaterialInterface.h
+++ b/Engine/Source/Runtime/Engine/Classes/Materials/MaterialInterface.h
@@ -40,25 +40,25 @@ struct ENGINE_API FMaterialRelevance
 	uint32 bOpaque : 1;
 
 	UPROPERTY()
-		uint32 bMasked : 1;
+	uint32 bMasked : 1;
 
 	UPROPERTY()
-		uint32 bDistortion : 1;
+	uint32 bDistortion : 1;
 
 	UPROPERTY()
-		uint32 bSeparateTranslucency : 1;
+	uint32 bSeparateTranslucency : 1;
 
 	UPROPERTY()
-		uint32 bNormalTranslucency : 1;
+	uint32 bNormalTranslucency : 1;
 
 	UPROPERTY()
-		uint32 bDisableDepthTest : 1;
+	uint32 bDisableDepthTest : 1;
 
 	UPROPERTY()
-		uint32 bSubsurfaceProfile : 1;
+	uint32 bSubsurfaceProfile : 1;
 
 	UPROPERTY()
-		uint32 bHasWorldPositionOffset : 1;
+	uint32 bSkipVelocityRendering : 1;
 
 	/** Default constructor. */
 	FMaterialRelevance()
@@ -69,7 +69,7 @@ struct ENGINE_API FMaterialRelevance
 		, bNormalTranslucency(false)
 		, bDisableDepthTest(false)
 		, bSubsurfaceProfile(false)
-		, bHasWorldPositionOffset(false)
+		, bSkipVelocityRendering(false)
 	{}
 
 	/** Bitwise OR operator.  Sets any relevance bits which are present in either FMaterialRelevance. */
@@ -82,7 +82,7 @@ struct ENGINE_API FMaterialRelevance
 		bNormalTranslucency |= B.bNormalTranslucency;
 		bDisableDepthTest |= B.bDisableDepthTest;
 		bSubsurfaceProfile |= B.bSubsurfaceProfile;
-		bHasWorldPositionOffset |= B.bHasWorldPositionOffset;
+		bSkipVelocityRendering |= B.bSkipVelocityRendering;
 		return *this;
 	}
 
diff --git a/Engine/Source/Runtime/Engine/Private/Materials/MaterialInterface.cpp b/Engine/Source/Runtime/Engine/Private/Materials/MaterialInterface.cpp
--- a/Engine/Source/Runtime/Engine/Private/Materials/MaterialInterface.cpp
+++ b/Engine/Source/Runtime/Engine/Private/Materials/MaterialInterface.cpp
@@ -18,7 +18,6 @@ void FMaterialRelevance::SetPrimitiveViewRelevance(FPrimitiveViewRelevance& OutV
 	OutViewRelevance.bSeparateTranslucencyRelevance = bSeparateTranslucency;
 	OutViewRelevance.bNormalTranslucencyRelevance = bNormalTranslucency;
 	OutViewRelevance.bSubsurfaceProfileRelevance = bSubsurfaceProfile;
-	OutViewRelevance.bHasWorldPositionOffset = bHasWorldPositionOffset;
 }
 
 //////////////////////////////////////////////////////////////////////////
@@ -58,7 +57,7 @@ FMaterialRelevance UMaterialInterface::GetRelevance_Internal(const UMaterial* Ma
 		MaterialRelevance.bNormalTranslucency = bIsTranslucent && !Material->bEnableSeparateTranslucency;
 		MaterialRelevance.bDisableDepthTest = bIsTranslucent && Material->bDisableDepthTest;
 		MaterialRelevance.bSubsurfaceProfile = (Material->MaterialDomain == MD_Surface) && !bIsTranslucent && (ShadingModel == MSM_SubsurfaceProfile);
-		MaterialRelevance.bHasWorldPositionOffset = MaterialResource->MaterialMayModifyMeshPosition();
+		MaterialRelevance.bSkipVelocityRendering = Material->bSkipVelocityOnBasePass;
 
 		return MaterialRelevance;
 	}
diff --git a/Engine/Source/Runtime/Engine/Private/Materials/MaterialShared.cpp b/Engine/Source/Runtime/Engine/Private/Materials/MaterialShared.cpp
--- a/Engine/Source/Runtime/Engine/Private/Materials/MaterialShared.cpp
+++ b/Engine/Source/Runtime/Engine/Private/Materials/MaterialShared.cpp
@@ -781,6 +781,11 @@ bool FMaterialResource::IsFullyRough() const
 	return Material->bFullyRough;
 }
 
+bool FMaterialResource::HasSkipVelocityOnBasePass() const
+{
+	return Material->bSkipVelocityOnBasePass;
+}
+
 bool FMaterialResource::IsNonmetal() const
 {
 	return !Material->Metallic.IsConnected() && !Material->Specular.IsConnected();
diff --git a/Engine/Source/Runtime/Engine/Public/MaterialShared.h b/Engine/Source/Runtime/Engine/Public/MaterialShared.h
--- a/Engine/Source/Runtime/Engine/Public/MaterialShared.h
+++ b/Engine/Source/Runtime/Engine/Public/MaterialShared.h
@@ -922,6 +922,7 @@ public:
 	virtual bool IsCrackFreeDisplacementEnabled() const { return false; }
 	virtual bool IsAdaptiveTessellationEnabled() const { return false; }
 	virtual bool IsFullyRough() const { return false; }
+	virtual bool HasSkipVelocityOnBasePass() const { return false; }
 	virtual bool IsNonmetal() const { return false; }
 	virtual bool UseLmDirectionality() const { return true; }
 	virtual bool IsMasked() const = 0;
@@ -1496,6 +1497,7 @@ public:
 	ENGINE_API virtual bool IsCrackFreeDisplacementEnabled() const;
 	ENGINE_API virtual bool IsAdaptiveTessellationEnabled() const;
 	ENGINE_API virtual bool IsFullyRough() const;
+	ENGINE_API virtual bool HasSkipVelocityOnBasePass() const;
 	ENGINE_API virtual bool IsNonmetal() const;
 	ENGINE_API virtual bool UseLmDirectionality() const;
 	ENGINE_API virtual enum EBlendMode GetBlendMode() const;
diff --git a/Engine/Source/Runtime/Renderer/Private/PrimitiveSceneInfo.cpp b/Engine/Source/Runtime/Renderer/Private/PrimitiveSceneInfo.cpp
--- a/Engine/Source/Runtime/Renderer/Private/PrimitiveSceneInfo.cpp
+++ b/Engine/Source/Runtime/Renderer/Private/PrimitiveSceneInfo.cpp
@@ -454,7 +454,7 @@ bool FPrimitiveSceneInfo::ShouldRenderVelocity(const FViewInfo& View, bool bChec
 
 	const FPrimitiveViewRelevance& PrimitiveViewRelevance = View.PrimitiveViewRelevanceMap[PrimitiveId];
 
-	if (!Proxy->IsMovable() && !PrimitiveViewRelevance.bHasWorldPositionOffset)
+	if (!Proxy->IsMovable())
 	{
 		return false;
 	}
diff --git a/Engine/Source/Runtime/Renderer/Private/VelocityRendering.cpp b/Engine/Source/Runtime/Renderer/Private/VelocityRendering.cpp
--- a/Engine/Source/Runtime/Renderer/Private/VelocityRendering.cpp
+++ b/Engine/Source/Runtime/Renderer/Private/VelocityRendering.cpp
@@ -349,15 +349,15 @@ bool FVelocityDrawingPolicy::HasVelocityOnBasePass(const FViewInfo& View,const F
 	FScene* Scene = PrimitiveSceneInfo->Scene;
 	if (Scene->MotionBlurInfoData.GetPrimitiveMotionBlurInfo(PrimitiveSceneInfo, OutTransform))
 	{
-		const FMatrix& LocalToWorld = PrimitiveSceneProxy->GetLocalToWorld();
+		bOutHasTransform = true;
 /*
+		const FMatrix& LocalToWorld = PrimitiveSceneProxy->GetLocalToWorld();
 		// Hasn't moved (treat as background by not rendering any special velocities)?
 		if (LocalToWorld.Equals(OutTransform, 0.0001f))
 		{
 			return false;
 		}
 */
-		bOutHasTransform = true;
 		return true;
 	}
 
@@ -368,7 +368,8 @@ bool FVelocityDrawingPolicy::HasVelocityOnBasePass(const FViewInfo& View,const F
 	}
 
 	//@todo-rco: Optimize finding WPO!
-	return View.PrimitiveViewRelevanceMap[PrimitiveSceneInfo->GetIndex()].bHasWorldPositionOffset;
+	auto* Material = Mesh.MaterialRenderProxy->GetMaterial(Scene->GetFeatureLevel());
+	return Material->MaterialModifiesMeshPosition_RenderThread() && !Material->HasSkipVelocityOnBasePass();
 }
 
 FBoundShaderStateInput FVelocityDrawingPolicy::GetBoundShaderStateInput(ERHIFeatureLevel::Type InFeatureLevel)
