RTDF shadows don't force backfaces to be shadowed, improves materials with subsurfaces scattering
Added a distance based ray start offset which reduces incorrect self shadowing artifacts

diff --git a/Engine/Shaders/DistanceFieldShadowing.usf b/Engine/Shaders/DistanceFieldShadowing.usf
--- a/Engine/Shaders/DistanceFieldShadowing.usf
+++ b/Engine/Shaders/DistanceFieldShadowing.usf
@@ -300,6 +300,7 @@ RWTexture2D<float2> RWShadowFactors;
 float3 LightDirection;
 float4 LightPositionAndInvRadius;
 float LightSourceRadius;
+float RayStartOffsetDepthScale;
 float2 TanLightAngleAndNormalThreshold;
 int4 ScissorRectMinAndSize;
 
@@ -511,15 +512,10 @@ void DistanceFieldShadowingCS(
 	if (bTileShouldComputeShadowing)
 	{
 		FGBufferData GBufferData = GetGBufferData(ScreenUV);
-		float3 WorldNormal = GBufferData.WorldNormal;
 
-#if !POINT_LIGHT
-		BRANCH
-		if (dot(WorldNormal, LightDirection) > TanLightAngleAndNormalThreshold.y)
-#endif
 		{
 			// World space offset along the start of the ray to avoid incorrect self-shadowing
-			float RayStartOffset = 2;
+			float RayStartOffset = 2 + RayStartOffsetDepthScale * SceneDepth;
 			// Keeps result from going all the way sharp
 			float MinSphereRadius = .4f;
 			// Maintain reasonable culling bounds
diff --git a/Engine/Source/Runtime/Engine/Classes/Components/LightComponent.h b/Engine/Source/Runtime/Engine/Classes/Components/LightComponent.h
--- a/Engine/Source/Runtime/Engine/Classes/Components/LightComponent.h
+++ b/Engine/Source/Runtime/Engine/Classes/Components/LightComponent.h
@@ -154,6 +154,9 @@ class ENGINE_API ULightComponent : public ULightComponentBase
 	UPROPERTY(EditAnywhere, BlueprintReadOnly, Category=DistanceFieldShadows, meta=(DisplayName = "RayTraced DistanceField Shadows"))
 	bool bUseRayTracedDistanceFieldShadows;
 
+	UPROPERTY(EditAnywhere, BlueprintReadOnly, Category=DistanceFieldShadows, meta=(UIMin = "0", UIMax = ".1"), AdvancedDisplay)
+	float RayStartOffsetDepthScale;
+
 public:
 	/** Set intensity of the light */
 	UFUNCTION(BlueprintCallable, Category="Rendering|Components|Light")
diff --git a/Engine/Source/Runtime/Engine/Private/LightComponent.cpp b/Engine/Source/Runtime/Engine/Private/LightComponent.cpp
--- a/Engine/Source/Runtime/Engine/Private/LightComponent.cpp
+++ b/Engine/Source/Runtime/Engine/Private/LightComponent.cpp
@@ -227,6 +227,7 @@ FLightSceneProxy::FLightSceneProxy(const ULightComponent* InLightComponent)
 	, bAffectDynamicIndirectLighting(InLightComponent->bAffectDynamicIndirectLighting)
 	, bHasReflectiveShadowMap(InLightComponent->bAffectDynamicIndirectLighting && InLightComponent->GetLightType() == LightType_Directional)
 	, bUseRayTracedDistanceFieldShadows(InLightComponent->bUseRayTracedDistanceFieldShadows)
+	, RayStartOffsetDepthScale(InLightComponent->RayStartOffsetDepthScale)
 	, LightType(InLightComponent->GetLightType())	
 	, ComponentName(InLightComponent->GetOwner() ? InLightComponent->GetOwner()->GetFName() : InLightComponent->GetFName())
 	, LevelName(InLightComponent->GetOutermost()->GetFName())
@@ -329,6 +330,8 @@ ULightComponent::ULightComponent(const FObjectInitializer& ObjectInitializer)
 	BloomScale = .2f;
 	BloomThreshold = 0;
 	BloomTint = FColor::White;
+
+	RayStartOffsetDepthScale = .01f;
 }
 
 
@@ -471,8 +474,9 @@ bool ULightComponent::CanEditChange(const UProperty* InProperty) const
 				return false;
 			}
 		}
-
-		if (PropertyName == GET_MEMBER_NAME_STRING_CHECKED(ULightComponent, bUseRayTracedDistanceFieldShadows))
+		
+		if (PropertyName == GET_MEMBER_NAME_STRING_CHECKED(ULightComponent, bUseRayTracedDistanceFieldShadows)
+			|| PropertyName == GET_MEMBER_NAME_STRING_CHECKED(ULightComponent, RayStartOffsetDepthScale))
 		{
 			static const auto CVar = IConsoleManager::Get().FindTConsoleVariableDataInt(TEXT("r.GenerateMeshDistanceFields"));
 			bool bCanEdit = CastShadows && CastDynamicShadows && Mobility != EComponentMobility::Static && CVar->GetValueOnGameThread() != 0;
@@ -533,6 +537,7 @@ void ULightComponent::PostEditChangeProperty(FPropertyChangedEvent& PropertyChan
 		PropertyName != GET_MEMBER_NAME_STRING_CHECKED(ULightComponent, BloomThreshold) &&
 		PropertyName != GET_MEMBER_NAME_STRING_CHECKED(ULightComponent, BloomTint) &&
 		PropertyName != GET_MEMBER_NAME_STRING_CHECKED(ULightComponent, bUseRayTracedDistanceFieldShadows) &&
+		PropertyName != GET_MEMBER_NAME_STRING_CHECKED(ULightComponent, RayStartOffsetDepthScale) &&
 		PropertyName != GET_MEMBER_NAME_STRING_CHECKED(ULightComponent, bVisible) &&
 		// Point light properties that shouldn't unbuild lighting
 		PropertyName != GET_MEMBER_NAME_STRING_CHECKED(UPointLightComponent, SourceRadius) &&
diff --git a/Engine/Source/Runtime/Engine/Public/SceneManagement.h b/Engine/Source/Runtime/Engine/Public/SceneManagement.h
--- a/Engine/Source/Runtime/Engine/Public/SceneManagement.h
+++ b/Engine/Source/Runtime/Engine/Public/SceneManagement.h
@@ -781,6 +781,7 @@ public:
 	inline bool CastsTranslucentShadows() const { return bCastTranslucentShadows; }
 	inline bool AffectsTranslucentLighting() const { return bAffectTranslucentLighting; }
 	inline bool UseRayTracedDistanceFieldShadows() const { return bUseRayTracedDistanceFieldShadows; }
+	inline float GetRayStartOffsetDepthScale() const { return RayStartOffsetDepthScale; }
 	inline uint8 GetLightType() const { return LightType; }
 	inline FName GetComponentName() const { return ComponentName; }
 	inline FName GetLevelName() const { return LevelName; }
@@ -898,6 +899,8 @@ protected:
 	/** Whether to use ray traced distance field area shadows. */
 	const uint32 bUseRayTracedDistanceFieldShadows : 1;
 
+	float RayStartOffsetDepthScale;
+
 	/** The light type (ELightComponentType) */
 	const uint8 LightType;
 
diff --git a/Engine/Source/Runtime/Renderer/Private/DistanceFieldShadowing.cpp b/Engine/Source/Runtime/Renderer/Private/DistanceFieldShadowing.cpp
--- a/Engine/Source/Runtime/Renderer/Private/DistanceFieldShadowing.cpp
+++ b/Engine/Source/Runtime/Renderer/Private/DistanceFieldShadowing.cpp
@@ -388,6 +388,7 @@ public:
 		NumGroups.Bind(Initializer.ParameterMap, TEXT("NumGroups"));
 		LightDirection.Bind(Initializer.ParameterMap, TEXT("LightDirection"));
 		LightSourceRadius.Bind(Initializer.ParameterMap, TEXT("LightSourceRadius"));
+		RayStartOffsetDepthScale.Bind(Initializer.ParameterMap, TEXT("RayStartOffsetDepthScale"));
 		LightPositionAndInvRadius.Bind(Initializer.ParameterMap, TEXT("LightPositionAndInvRadius"));
 		TanLightAngleAndNormalThreshold.Bind(Initializer.ParameterMap, TEXT("TanLightAngleAndNormalThreshold"));
 		ScissorRectMinAndSize.Bind(Initializer.ParameterMap, TEXT("ScissorRectMinAndSize"));
@@ -434,6 +435,8 @@ public:
 		// Default light source radius of 0 gives poor results
 		SetShaderValue(RHICmdList, ShaderRHI, LightSourceRadius, LightSourceRadiusValue == 0 ? 20 : FMath::Clamp(LightSourceRadiusValue, .001f, 1.0f / (4 * LightPositionAndInvRadiusValue.W)));
 
+		SetShaderValue(RHICmdList, ShaderRHI, RayStartOffsetDepthScale, ProjectedShadowInfo->LightSceneInfo->Proxy->GetRayStartOffsetDepthScale());
+
 		const float LightSourceAngle = FMath::Clamp(ProjectedShadowInfo->LightSceneInfo->Proxy->GetLightSourceAngle(), 0.001f, 5.0f) * PI / 180.0f;
 		const FVector2D TanLightAngleAndNormalThresholdValue(FMath::Tan(LightSourceAngle), FMath::Cos(PI / 2 + LightSourceAngle));
 		SetShaderValue(RHICmdList, ShaderRHI, TanLightAngleAndNormalThreshold, TanLightAngleAndNormalThresholdValue);
@@ -467,6 +470,7 @@ public:
 		Ar << LightDirection;
 		Ar << LightPositionAndInvRadius;
 		Ar << LightSourceRadius;
+		Ar << RayStartOffsetDepthScale;
 		Ar << TanLightAngleAndNormalThreshold;
 		Ar << ScissorRectMinAndSize;
 		Ar << ObjectParameters;
@@ -485,6 +489,7 @@ private:
 	FShaderParameter LightDirection;
 	FShaderParameter LightPositionAndInvRadius;
 	FShaderParameter LightSourceRadius;
+	FShaderParameter RayStartOffsetDepthScale;
 	FShaderParameter TanLightAngleAndNormalThreshold;
 	FShaderParameter ScissorRectMinAndSize;
 	FDistanceFieldCulledObjectBufferParameters ObjectParameters;
