Added serialisation to generic quad tree, fixed possible null access in duplicate.

diff --git a/Engine/Source/Runtime/Engine/Public/GenericQuadTree.h b/Engine/Source/Runtime/Engine/Public/GenericQuadTree.h
--- a/Engine/Source/Runtime/Engine/Public/GenericQuadTree.h
+++ b/Engine/Source/Runtime/Engine/Public/GenericQuadTree.h
@@ -21,11 +21,15 @@ public:
 	void Remove(const ElementType& Instance, const FBox2D& Box);
 
 	/** Does a deep copy of the tree by going through and re-creating the internal data. Cheaper than re-insertion as it should be linear instead of nlogn */
-	void Duplicate(TreeType* OutDuplicate) const;
+	void Duplicate(TreeType& OutDuplicate) const;
 
 	/** Removes all elements of the tree */
 	void Empty();
 
+	void Serialize(FArchive& Ar);
+
+	TreeType& operator=(const TreeType& Other);
+
 	~TQuadTree();
 
 private:
@@ -53,10 +57,17 @@ private:
 		FBox2D Box;
 		ElementType Element;
 
+		FNode() {};
+
 		FNode(const ElementType& InElement, const FBox2D& InBox)
 		: Box(InBox)
 		, Element(InElement)
 		{}
+
+		friend FArchive& operator<<(FArchive& Ar, typename TQuadTree<ElementType, NodeCapacity>::FNode& Node)
+		{
+			return Ar << Node.Box << Node.Element;
+		}
 	};
 
 private:
@@ -75,6 +86,39 @@ private:
 	bool bInternal;
 };
 
+template <typename ElementType, int32 NodeCapacity /*= 4*/>
+typename TQuadTree<ElementType, NodeCapacity>::TreeType& TQuadTree<ElementType, NodeCapacity>::operator=(const TreeType& Other)
+{
+	Other.Duplicate(*this);
+	return *this;
+}
+
+template <typename ElementType, int32 NodeCapacity /*= 4*/>
+void TQuadTree<ElementType, NodeCapacity>::Serialize(FArchive& Ar)
+{
+	Ar << Nodes;
+	
+	bool SubTreeFlags[4] = {SubTrees[0] != nullptr, SubTrees[1] != nullptr, SubTrees[2] != nullptr, SubTrees[3] != nullptr};
+	Ar << SubTreeFlags[0] << SubTreeFlags[1] << SubTreeFlags[2] << SubTreeFlags[3];
+
+	for(int32 Idx = 0 ; Idx < 4 ; ++Idx)
+	{
+		if(SubTreeFlags[Idx])
+		{
+			if(Ar.ArIsLoading)
+			{
+				SubTrees[Idx] = new TreeType(FBox2D());
+			}
+
+			SubTrees[Idx]->Serialize(Ar);
+		}
+	}
+
+	Ar << TreeBox;
+	Ar << Position;
+	Ar << bInternal;
+}
+
 template <typename ElementType, int32 NodeCapacity>
 TQuadTree<ElementType, NodeCapacity>::TQuadTree(const FBox2D& Box)
 : TreeBox(Box)
@@ -301,22 +345,22 @@ void TQuadTree<ElementType, NodeCapacity>::GetLeaves(const FBox2D& Box, TArray<c
 }
 
 template <typename ElementType, int32 NodeCapacity>
-void TQuadTree<ElementType, NodeCapacity>::Duplicate(TreeType* OutDuplicate) const
+void TQuadTree<ElementType, NodeCapacity>::Duplicate(TreeType& OutDuplicate) const
 {
 	for (int32 TreeIdx = 0; TreeIdx < 4; ++TreeIdx)
 	{
 		if (TreeType* SubTree = SubTrees[TreeIdx])
 		{
-			OutDuplicate->SubTrees[TreeIdx] = new TreeType(FBox2D(0,0));
-			SubTree->Duplicate(OutDuplicate->SubTrees[TreeIdx]);	//duplicate sub trees
+			OutDuplicate.SubTrees[TreeIdx] = new TreeType(FBox2D(0,0));
+			SubTree->Duplicate(*OutDuplicate.SubTrees[TreeIdx]);	//duplicate sub trees
 		}
 		
 	}
 
-	OutDuplicate->Nodes = Nodes;
-	OutDuplicate->TreeBox = TreeBox;
-	OutDuplicate->Position = Position;
-	OutDuplicate->bInternal = bInternal;
+	OutDuplicate.Nodes = Nodes;
+	OutDuplicate.TreeBox = TreeBox;
+	OutDuplicate.Position = Position;
+	OutDuplicate.bInternal = bInternal;
 }
 
 template <typename ElementType, int32 NodeCapacity>
diff --git a/Engine/Source/Runtime/Foliage/Private/ProceduralFoliageBroadphase.cpp b/Engine/Source/Runtime/Foliage/Private/ProceduralFoliageBroadphase.cpp
--- a/Engine/Source/Runtime/Foliage/Private/ProceduralFoliageBroadphase.cpp
+++ b/Engine/Source/Runtime/Foliage/Private/ProceduralFoliageBroadphase.cpp
@@ -12,7 +12,7 @@ FProceduralFoliageBroadphase::FProceduralFoliageBroadphase(float TileSize)
 FProceduralFoliageBroadphase::FProceduralFoliageBroadphase(const FProceduralFoliageBroadphase& OtherBroadphase)
 : QuadTree(FBox2D(FVector2D(0, 0), FVector2D(0, 0)))
 {
-	OtherBroadphase.QuadTree.Duplicate(&QuadTree);
+	OtherBroadphase.QuadTree.Duplicate(QuadTree);
 }
 
 void FProceduralFoliageBroadphase::Empty()
