Added support for DistanceFieldReplacementMesh under BuildSettings

diff --git a/Engine/Source/Editor/StaticMeshEditor/Private/StaticMeshEditorTools.cpp b/Engine/Source/Editor/StaticMeshEditor/Private/StaticMeshEditorTools.cpp
--- a/Engine/Source/Editor/StaticMeshEditor/Private/StaticMeshEditorTools.cpp
+++ b/Engine/Source/Editor/StaticMeshEditor/Private/StaticMeshEditorTools.cpp
@@ -314,6 +314,16 @@ void FMeshBuildSettingsLayout::GenerateHeaderRowContent( FDetailWidgetRow& NodeR
 
 TAutoConsoleVariable<int32> GEnableMikkTSpaceCVar(TEXT("r.MikkTSPaceOptional"),0,TEXT("Set to be non-zero to display the option of using MikkTSpace to generate tangents for static meshes. MikkTSpace is the default."),ECVF_Default);
 
+FString FMeshBuildSettingsLayout::GetCurrentDistanceFieldReplacementMeshPath() const
+{
+	return BuildSettings.DistanceFieldReplacementMesh ? BuildSettings.DistanceFieldReplacementMesh->GetPathName() : FString("");
+}
+
+void FMeshBuildSettingsLayout::OnDistanceFieldReplacementMeshSelected(const FAssetData& AssetData)
+{
+	BuildSettings.DistanceFieldReplacementMesh = Cast<UStaticMesh>(AssetData.GetAsset());
+}
+
 void FMeshBuildSettingsLayout::GenerateChildContent( IDetailChildrenBuilder& ChildrenBuilder )
 {
 	{
@@ -532,6 +542,26 @@ void FMeshBuildSettingsLayout::GenerateChildContent( IDetailChildrenBuilder& Chi
 		];
 	}
 
+	{
+		TSharedRef<SWidget> PropWidget = SNew(SObjectPropertyEntryBox)
+			.AllowedClass(UStaticMesh::StaticClass())
+			.AllowClear(true)
+			.ObjectPath(this, &FMeshBuildSettingsLayout::GetCurrentDistanceFieldReplacementMeshPath)
+			.OnObjectChanged(this, &FMeshBuildSettingsLayout::OnDistanceFieldReplacementMeshSelected);
+
+		ChildrenBuilder.AddChildContent( LOCTEXT("DistanceFieldReplacementMesh", "Distance Field Replacement Mesh") )
+		.NameContent()
+		[
+			SNew(STextBlock)
+			.Font( IDetailLayoutBuilder::GetDetailFont() )
+			.Text(LOCTEXT("DistanceFieldReplacementMesh", "Distance Field Replacement Mesh"))
+		]
+		.ValueContent()
+		[
+			PropWidget
+		];
+	}
+
 	{
 		ChildrenBuilder.AddChildContent( LOCTEXT("ApplyChanges", "Apply Changes") )
 		.ValueContent()
diff --git a/Engine/Source/Editor/StaticMeshEditor/Private/StaticMeshEditorTools.h b/Engine/Source/Editor/StaticMeshEditor/Private/StaticMeshEditorTools.h
--- a/Engine/Source/Editor/StaticMeshEditor/Private/StaticMeshEditorTools.h
+++ b/Engine/Source/Editor/StaticMeshEditor/Private/StaticMeshEditorTools.h
@@ -171,6 +171,9 @@ private:
 
 	void OnDistanceFieldResolutionScaleChanged(float NewValue);
 	void OnDistanceFieldResolutionScaleCommitted(float NewValue, ETextCommit::Type TextCommitType);
+	FString GetCurrentDistanceFieldReplacementMeshPath() const;
+	void OnDistanceFieldReplacementMeshSelected(const FAssetData& AssetData);
+
 private:
 	TWeakPtr<FLevelOfDetailSettingsLayout> ParentLODSettings;
 	FMeshBuildSettings BuildSettings;
diff --git a/Engine/Source/Runtime/Engine/Classes/Engine/EngineTypes.h b/Engine/Source/Runtime/Engine/Classes/Engine/EngineTypes.h
--- a/Engine/Source/Runtime/Engine/Classes/Engine/EngineTypes.h
+++ b/Engine/Source/Runtime/Engine/Classes/Engine/EngineTypes.h
@@ -1995,6 +1995,9 @@ struct FMeshBuildSettings
 	UPROPERTY(EditAnywhere, Category=BuildSettings)
 	bool bGenerateDistanceFieldAsIfTwoSided;
 
+	UPROPERTY(EditAnywhere, Category=BuildSettings)
+	class UStaticMesh* DistanceFieldReplacementMesh;
+
 	/** Default settings. */
 	FMeshBuildSettings()
 		: bUseMikkTSpace(true)
@@ -2010,6 +2013,7 @@ struct FMeshBuildSettings
 		, BuildScale3D(1.0f, 1.0f, 1.0f)
 		, DistanceFieldResolutionScale(1.0f)
 		, bGenerateDistanceFieldAsIfTwoSided(false)
+		, DistanceFieldReplacementMesh(NULL)
 	{ }
 
 	/** Equality operator. */
@@ -2026,7 +2030,8 @@ struct FMeshBuildSettings
 			&& DstLightmapIndex == Other.DstLightmapIndex
 			&& BuildScale3D == Other.BuildScale3D
 			&& DistanceFieldResolutionScale == Other.DistanceFieldResolutionScale
-			&& bGenerateDistanceFieldAsIfTwoSided == Other.bGenerateDistanceFieldAsIfTwoSided;
+			&& bGenerateDistanceFieldAsIfTwoSided == Other.bGenerateDistanceFieldAsIfTwoSided
+			&& DistanceFieldReplacementMesh == Other.DistanceFieldReplacementMesh;
 	}
 
 	/** Inequality. */
diff --git a/Engine/Source/Runtime/Engine/Private/DistanceFieldAtlas.cpp b/Engine/Source/Runtime/Engine/Private/DistanceFieldAtlas.cpp
--- a/Engine/Source/Runtime/Engine/Private/DistanceFieldAtlas.cpp
+++ b/Engine/Source/Runtime/Engine/Private/DistanceFieldAtlas.cpp
@@ -193,7 +193,7 @@ FDistanceFieldAsyncQueue* GDistanceFieldAsyncQueue = NULL;
 
 #if WITH_EDITORONLY_DATA
 
-void FDistanceFieldVolumeData::CacheDerivedData(const FString& InDDCKey, UStaticMesh* Mesh)
+void FDistanceFieldVolumeData::CacheDerivedData(const FString& InDDCKey, UStaticMesh* Mesh, UStaticMesh* GenerateSource, float DistanceFieldResolutionScale, bool bGenerateDistanceFieldAsIfTwoSided)
 {
 	TArray<uint8> DerivedData;
 
@@ -207,6 +207,9 @@ void FDistanceFieldVolumeData::CacheDerivedData(const FString& InDDCKey, UStatic
 		FAsyncDistanceFieldTask* NewTask = new FAsyncDistanceFieldTask;
 		NewTask->DDCKey = InDDCKey;
 		NewTask->StaticMesh = Mesh;
+		NewTask->GenerateSource = GenerateSource;
+		NewTask->DistanceFieldResolutionScale = DistanceFieldResolutionScale;
+		NewTask->bGenerateDistanceFieldAsIfTwoSided = bGenerateDistanceFieldAsIfTwoSided;
 		NewTask->GeneratedVolumeData = new FDistanceFieldVolumeData();
 
 		for (int32 MaterialIndex = 0; MaterialIndex < Mesh->Materials.Num(); MaterialIndex++)
@@ -370,15 +373,15 @@ void FDistanceFieldAsyncQueue::AddTask(FAsyncDistanceFieldTask* Task)
 void FDistanceFieldAsyncQueue::Build(FAsyncDistanceFieldTask* Task, FQueuedThreadPool& ThreadPool)
 {
 #if WITH_EDITOR
-	const FStaticMeshLODResources& LODModel = Task->StaticMesh->RenderData->LODResources[0];
+	const FStaticMeshLODResources& LODModel = Task->GenerateSource->RenderData->LODResources[0];
 
 	MeshUtilities->GenerateSignedDistanceFieldVolumeData(
 		LODModel,
 		ThreadPool,
 		Task->MaterialBlendModes,
-		Task->StaticMesh->RenderData->Bounds,
-		Task->StaticMesh->SourceModels[0].BuildSettings.DistanceFieldResolutionScale,
-		Task->StaticMesh->SourceModels[0].BuildSettings.bGenerateDistanceFieldAsIfTwoSided,
+		Task->GenerateSource->RenderData->Bounds,
+		Task->DistanceFieldResolutionScale,
+		Task->bGenerateDistanceFieldAsIfTwoSided,
 		*Task->GeneratedVolumeData);
 
 	CompletedTasks.Push(Task);
@@ -391,6 +394,7 @@ void FDistanceFieldAsyncQueue::AddReferencedObjects(FReferenceCollector& Collect
 	{
 		// Make sure none of the UObjects referenced by the async tasks are GC'ed during the task
 		Collector.AddReferencedObject(ReferencedTasks[TaskIndex]->StaticMesh);
+		Collector.AddReferencedObject(ReferencedTasks[TaskIndex]->GenerateSource);
 	}
 }
 
diff --git a/Engine/Source/Runtime/Engine/Private/StaticMesh.cpp b/Engine/Source/Runtime/Engine/Private/StaticMesh.cpp
--- a/Engine/Source/Runtime/Engine/Private/StaticMesh.cpp
+++ b/Engine/Source/Runtime/Engine/Private/StaticMesh.cpp
@@ -1025,6 +1025,13 @@ FArchive& operator<<(FArchive& Ar, FMeshBuildSettings& BuildSettings)
 	
 	Ar << BuildSettings.DistanceFieldResolutionScale;
 	Ar << BuildSettings.bGenerateDistanceFieldAsIfTwoSided;
+
+	if (BuildSettings.DistanceFieldReplacementMesh)
+	{
+		FString ReplacementMeshName = BuildSettings.DistanceFieldReplacementMesh->GetPathName();
+		Ar << ReplacementMeshName;
+	}
+
 	return Ar;
 }
 
@@ -1180,7 +1187,9 @@ void FStaticMeshRenderData::Cache(UStaticMesh* Owner, const FStaticMeshLODSettin
 			LODResources[0].DistanceFieldData = new FDistanceFieldVolumeData();
 		}
 
-		LODResources[0].DistanceFieldData->CacheDerivedData(DistanceFieldKey, Owner);
+		const FMeshBuildSettings& BuildSettings = Owner->SourceModels[0].BuildSettings;
+		UStaticMesh* MeshToGenerateFrom = BuildSettings.DistanceFieldReplacementMesh ? BuildSettings.DistanceFieldReplacementMesh : Owner;
+		LODResources[0].DistanceFieldData->CacheDerivedData(DistanceFieldKey, Owner, MeshToGenerateFrom, BuildSettings.DistanceFieldResolutionScale, BuildSettings.bGenerateDistanceFieldAsIfTwoSided);
 	}
 }
 #endif // #if WITH_EDITOR
diff --git a/Engine/Source/Runtime/Engine/Public/DistanceFieldAtlas.h b/Engine/Source/Runtime/Engine/Public/DistanceFieldAtlas.h
--- a/Engine/Source/Runtime/Engine/Public/DistanceFieldAtlas.h
+++ b/Engine/Source/Runtime/Engine/Public/DistanceFieldAtlas.h
@@ -148,7 +148,7 @@ public:
 
 #if WITH_EDITORONLY_DATA
 
-	void CacheDerivedData(const FString& InDDCKey, UStaticMesh* Mesh);
+	void CacheDerivedData(const FString& InDDCKey, UStaticMesh* Mesh, UStaticMesh* GenerateSource, float DistanceFieldResolutionScale, bool bGenerateDistanceFieldAsIfTwoSided);
 
 #endif
 
@@ -166,6 +166,9 @@ class FAsyncDistanceFieldTask
 public:
 	TArray<EBlendMode> MaterialBlendModes;
 	UStaticMesh* StaticMesh;
+	UStaticMesh* GenerateSource;
+	float DistanceFieldResolutionScale;
+	bool bGenerateDistanceFieldAsIfTwoSided;
 	FString DDCKey;
 	FDistanceFieldVolumeData* GeneratedVolumeData;
 };
