Fixed issue in which the spline wasn't updated when setting/unsetting the closed loop property via a Blueprint.

diff --git a/Engine/Source/Runtime/Engine/Classes/Components/SplineComponent.h b/Engine/Source/Runtime/Engine/Classes/Components/SplineComponent.h
--- a/Engine/Source/Runtime/Engine/Classes/Components/SplineComponent.h
+++ b/Engine/Source/Runtime/Engine/Classes/Components/SplineComponent.h
@@ -64,11 +64,11 @@ private:
 public:
 #if WITH_EDITORONLY_DATA
 	/** Color of an unselected spline component segment in the editor */
-	UPROPERTY(EditAnywhere, BlueprintReadWrite, Category = Editor)
+	UPROPERTY(EditAnywhere, Category = Editor)
 	FLinearColor EditorUnselectedSplineSegmentColor;
 
 	/** Color of a selected spline component segment in the editor */
-	UPROPERTY(EditAnywhere, BlueprintReadWrite, Category = Editor)
+	UPROPERTY(EditAnywhere, Category = Editor)
 	FLinearColor EditorSelectedSplineSegmentColor;
 #endif
 
@@ -81,6 +81,14 @@ public:
 	/** Update the spline tangents and SplineReparamTable */
 	void UpdateSpline();
 
+	/** Specify unselected spline component segment color in the editor */
+	UFUNCTION(BlueprintCallable, Category = Editor)
+	void SetUnselectedSplineSegmentColor(const FLinearColor& SegmentColor);
+
+	/** Specify selected spline component segment color in the editor */
+	UFUNCTION(BlueprintCallable, Category = Editor)
+	void SetSelectedSplineSegmentColor(const FLinearColor& SegmentColor);
+
 	/** Specify whether the spline is a closed loop or not */
 	UFUNCTION(BlueprintCallable, Category = Spline)
 	void SetClosedLoop(bool bInClosedLoop);
@@ -192,6 +200,9 @@ private:
 	/** Returns the parametric value t which would result in a spline segment of the given length between S(0)...S(t) */
 	float GetSegmentParamFromLength(const int32 Index, const float Length, const float SegmentLength) const;
 
+	/** Updates the spline points to reflect a closed loop or not */
+	void UpdateLoopEndpoint(bool bInClosedLoop);
+
 	/** Adds an additional point to the end of the spline, whose OutValue and tangents coincides with the first point */
 	void AddLoopEndpoint();
 
diff --git a/Engine/Source/Runtime/Engine/Private/SplineComponent.cpp b/Engine/Source/Runtime/Engine/Private/SplineComponent.cpp
--- a/Engine/Source/Runtime/Engine/Private/SplineComponent.cpp
+++ b/Engine/Source/Runtime/Engine/Private/SplineComponent.cpp
@@ -209,6 +209,13 @@ float USplineComponent::GetSegmentParamFromLength(const int32 Index, const float
 
 
 void USplineComponent::SetClosedLoop(bool bInClosedLoop)
+{
+	UpdateLoopEndpoint(bInClosedLoop);
+	UpdateSpline();
+}
+
+
+void USplineComponent::UpdateLoopEndpoint(bool bInClosedLoop)
 {
 	if (bClosedLoop != bInClosedLoop)
 	{
@@ -257,6 +264,22 @@ void USplineComponent::AddLoopEndpoint()
 }
 
 
+void USplineComponent::SetUnselectedSplineSegmentColor(const FLinearColor& Color)
+{
+#if WITH_EDITOR
+	EditorUnselectedSplineSegmentColor = Color;
+#endif
+}
+
+
+void USplineComponent::SetSelectedSplineSegmentColor(const FLinearColor& Color)
+{
+#if WITH_EDITOR
+	EditorSelectedSplineSegmentColor = Color;
+#endif
+}
+
+
 void USplineComponent::ClearSplinePoints()
 {
 	SplineInfo.Reset();
@@ -268,14 +291,14 @@ void USplineComponent::AddSplineWorldPoint(const FVector& Position)
 {
 	// If it's a closed loop, remove the endpoint before adding a new point
 	const bool bWasLoop = IsClosedLoop();
-	SetClosedLoop(false);
+	UpdateLoopEndpoint(false);
 
 	float InputKey = static_cast<float>(SplineInfo.Points.Num());
 	const int32 PointIndex = SplineInfo.AddPoint(InputKey, ComponentToWorld.InverseTransformPosition(Position));
 	SplineInfo.Points[PointIndex].InterpMode = CIM_CurveAuto;
 
 	// Then re-close the spline if required
-	SetClosedLoop(bWasLoop);
+	UpdateLoopEndpoint(bWasLoop);
 
 	UpdateSpline();
 }
@@ -285,14 +308,14 @@ void USplineComponent::AddSplineLocalPoint(const FVector& Position)
 {
 	// If it's a closed loop, remove the endpoint before adding a new point
 	const bool bWasLoop = IsClosedLoop();
-	SetClosedLoop(false);
+	UpdateLoopEndpoint(false);
 
 	float InputKey = static_cast<float>(SplineInfo.Points.Num());
 	const int32 PointIndex = SplineInfo.AddPoint(InputKey, Position);
 	SplineInfo.Points[PointIndex].InterpMode = CIM_CurveAuto;
 
 	// Then re-close the spline if required
-	SetClosedLoop(bWasLoop);
+	UpdateLoopEndpoint(bWasLoop);
 
 	UpdateSpline();
 }
@@ -302,7 +325,7 @@ void USplineComponent::SetSplineWorldPoints(const TArray<FVector>& Points)
 {
 	// If it's a closed loop, mark it as not closed before setting a new array of points
 	const bool bWasLoop = IsClosedLoop();
-	SetClosedLoop(false);
+	UpdateLoopEndpoint(false);
 
 	SplineInfo.Points.Reset(bWasLoop ? Points.Num() + 1 : Points.Num());
 	float InputKey = 0.0f;
@@ -313,7 +336,7 @@ void USplineComponent::SetSplineWorldPoints(const TArray<FVector>& Points)
 		InputKey += 1.0f;
 	}
 
-	SetClosedLoop(bWasLoop);
+	UpdateLoopEndpoint(bWasLoop);
 
 	UpdateSpline();
 }
@@ -322,7 +345,7 @@ void USplineComponent::SetSplineWorldPoints(const TArray<FVector>& Points)
 void USplineComponent::SetSplineLocalPoints(const TArray<FVector>& Points)
 {
 	const bool bWasLoop = IsClosedLoop();
-	SetClosedLoop(false);
+	UpdateLoopEndpoint(false);
 
 	SplineInfo.Points.Reset(bWasLoop ? Points.Num() + 1 : Points.Num());
 	float InputKey = 0.0f;
@@ -333,7 +356,7 @@ void USplineComponent::SetSplineLocalPoints(const TArray<FVector>& Points)
 		InputKey += 1.0f;
 	}
 
-	SetClosedLoop(bWasLoop);
+	UpdateLoopEndpoint(bWasLoop);
 
 	UpdateSpline();
 }
