Recast serialization workaround to load packages merged from 4.7 branch

diff --git a/Engine/Source/Runtime/Engine/Private/AI/Navigation/PImplRecastNavMesh.cpp b/Engine/Source/Runtime/Engine/Private/AI/Navigation/PImplRecastNavMesh.cpp
--- a/Engine/Source/Runtime/Engine/Private/AI/Navigation/PImplRecastNavMesh.cpp
+++ b/Engine/Source/Runtime/Engine/Private/AI/Navigation/PImplRecastNavMesh.cpp
@@ -365,7 +365,8 @@ void FPImplRecastNavMesh::Serialize( FArchive& Ar )
 					dtFree(TileHeader);
 
 					//
-					if (Ar.UE4Ver() >= VER_UE4_ADD_MODIFIERS_RUNTIME_GENERATION)
+					if (Ar.UE4Ver() >= VER_UE4_ADD_MODIFIERS_RUNTIME_GENERATION && 
+						Ar.UE4Ver() != VER_UE4_MERGED_ADD_MODIFIERS_RUNTIME_GENERATION_TO_4_7) // Merged package from 4.7 branch
 					{
 						unsigned char* ComressedTileData = NULL;
 						int32 CompressedTileDataSize = 0;
@@ -405,7 +406,8 @@ void FPImplRecastNavMesh::Serialize( FArchive& Ar )
 					DetourNavMesh->addTile(TileData, TileDataSize, DT_TILE_FREE_DATA, TileRef, NULL);
 
 					// Serialize compressed tile cache layer
-					if (Ar.UE4Ver() >= VER_UE4_ADD_MODIFIERS_RUNTIME_GENERATION)
+					if (Ar.UE4Ver() >= VER_UE4_ADD_MODIFIERS_RUNTIME_GENERATION &&
+						Ar.UE4Ver() != VER_UE4_MERGED_ADD_MODIFIERS_RUNTIME_GENERATION_TO_4_7) // Merged package from 4.7 branch
 					{
 						uint8* ComressedTileData = nullptr;
 						int32 CompressedTileDataSize = 0;
diff --git a/Engine/Source/Runtime/Engine/Private/AI/Navigation/RecastNavMeshDataChunk.cpp b/Engine/Source/Runtime/Engine/Private/AI/Navigation/RecastNavMeshDataChunk.cpp
--- a/Engine/Source/Runtime/Engine/Private/AI/Navigation/RecastNavMeshDataChunk.cpp
+++ b/Engine/Source/Runtime/Engine/Private/AI/Navigation/RecastNavMeshDataChunk.cpp
@@ -117,7 +117,8 @@ void URecastNavMeshDataChunk::SerializeRecastData(FArchive& Ar)
 				// Load compressed tile cache layer
 				int32 TileCacheDataSize = 0;
 				uint8* TileCacheRawData = nullptr;
-				if (Ar.UE4Ver() >= VER_UE4_ADD_MODIFIERS_RUNTIME_GENERATION)
+				if (Ar.UE4Ver() >= VER_UE4_ADD_MODIFIERS_RUNTIME_GENERATION && 
+					Ar.UE4Ver() != VER_UE4_MERGED_ADD_MODIFIERS_RUNTIME_GENERATION_TO_4_7) // Merged package from 4.7 branch
 				{
 					FPImplRecastNavMesh::SerializeCompressedTileCacheData(Ar, TileCacheRawData, TileCacheDataSize);
 				}
