Fix UE-7946 - editor icon scale changes after a feature level switch
- Set relative scale using accessor function, ensures UpdateComponentToWorld occurs.

diff --git a/Engine/Source/Runtime/Engine/Private/LightComponent.cpp b/Engine/Source/Runtime/Engine/Private/LightComponent.cpp
--- a/Engine/Source/Runtime/Engine/Private/LightComponent.cpp
+++ b/Engine/Source/Runtime/Engine/Private/LightComponent.cpp
@@ -123,7 +123,7 @@ void ULightComponentBase::UpdateLightSpriteTexture()
 		SpriteComponent->SetSprite(GetEditorSprite());
 
 		float SpriteScale = GetEditorSpriteScale();
-		SpriteComponent->RelativeScale3D = FVector(SpriteScale);
+		SpriteComponent->SetRelativeScale3D(FVector(SpriteScale));
 	}
 }
 
@@ -581,7 +581,7 @@ void ULightComponent::UpdateLightSpriteTexture()
 			UTexture2D* SpriteTexture = NULL;
 			SpriteTexture = LoadObject<UTexture2D>(NULL, TEXT("/Engine/EditorResources/LightIcons/S_LightError.S_LightError"));
 			SpriteComponent->SetSprite(SpriteTexture);
-			SpriteComponent->RelativeScale3D = FVector(0.5f, 0.5f, 0.5f);
+			SpriteComponent->SetRelativeScale3D(FVector(0.5f));
 		}
 		else
 		{
