Native transition rules & state entry/exit event delegates

Added support for binding native delegates to transition rules.
Changed callback for state entry/exit to also use delegates.
Bindings are registered before InitializeAnimation() is called (usually in your class constructor), then bound to states & transitions after the state machine has been initialized.
Added indicator in details panel when a transition is using a native rule.
Modified AnimBlueprint compiler so that 'always transition' warnings are not generated when a native delegate is installed.

diff --git a/Engine/Source/Editor/AnimGraph/Private/AnimStateTransitionNode.cpp b/Engine/Source/Editor/AnimGraph/Private/AnimStateTransitionNode.cpp
--- a/Engine/Source/Editor/AnimGraph/Private/AnimStateTransitionNode.cpp
+++ b/Engine/Source/Editor/AnimGraph/Private/AnimStateTransitionNode.cpp
@@ -13,6 +13,7 @@
 #include "AnimGraphNode_TransitionResult.h"
 #include "EdGraphUtilities.h"
 #include "Kismet2/Kismet2NameValidators.h"
+#include "Animation/AnimInstance.h"
 
 //////////////////////////////////////////////////////////////////////////
 // IAnimStateTransitionNodeSharedDataHelper
@@ -468,7 +469,29 @@ void UAnimStateTransitionNode::ValidateNodeDuringCompilation(class FCompilerResu
 		UEdGraphPin* BoolResultPin = ResultNode->Pins[0];
 		if ((BoolResultPin->LinkedTo.Num() == 0) && (BoolResultPin->DefaultValue.ToBool() == false))
 		{
-			MessageLog.Warning(TEXT("@@ will never be taken, please connect something to @@"), this, BoolResultPin);
+			// check for native transition rule before warning
+			bool bHasNative = false;
+			UBlueprint* Blueprint = FBlueprintEditorUtils::FindBlueprintForNodeChecked(this);
+			if(Blueprint && Blueprint->ParentClass)
+			{
+				UAnimInstance* AnimInstance = CastChecked<UAnimInstance>(Blueprint->ParentClass->GetDefaultObject());
+				if(AnimInstance)
+				{
+					UEdGraph* ParentGraph = GetGraph();
+					UAnimStateNodeBase* PrevState = GetPreviousState();
+					UAnimStateNodeBase* NextState = GetNextState();
+					if(PrevState != nullptr && NextState != nullptr && ParentGraph != nullptr)
+					{
+						FName FunctionName;
+						bHasNative = AnimInstance->HasNativeTransitionBinding(ParentGraph->GetFName(), FName(*PrevState->GetStateName()), FName(*NextState->GetStateName()), FunctionName);
+					}
+				}
+			}
+
+			if(!bHasNative)
+			{
+				MessageLog.Warning(TEXT("@@ will never be taken, please connect something to @@"), this, BoolResultPin);
+			}
 		}
 	}
 	else
diff --git a/Engine/Source/Editor/DetailCustomizations/Private/AnimTransitionNodeDetails.cpp b/Engine/Source/Editor/DetailCustomizations/Private/AnimTransitionNodeDetails.cpp
--- a/Engine/Source/Editor/DetailCustomizations/Private/AnimTransitionNodeDetails.cpp
+++ b/Engine/Source/Editor/DetailCustomizations/Private/AnimTransitionNodeDetails.cpp
@@ -13,6 +13,7 @@
 #include "SKismetLinearExpression.h"
 #include "STextEntryPopup.h"
 #include "SExpandableArea.h"
+#include "BlueprintEditorUtils.h"
 
 #define LOCTEXT_NAMESPACE "FAnimStateNodeDetails"
 
@@ -101,6 +102,33 @@ void FAnimTransitionNodeDetails::CustomizeDetails( IDetailLayoutBuilder& DetailB
 					CanExecPin = ResultNode->FindPin(TEXT("bCanEnterTransition"));
 				}
 			}
+
+			// indicate if a native transition rule applies to this
+			UBlueprint* Blueprint = FBlueprintEditorUtils::FindBlueprintForNodeChecked(TransitionNode.Get());
+			if(Blueprint && Blueprint->ParentClass)
+			{
+				UAnimInstance* AnimInstance = CastChecked<UAnimInstance>(Blueprint->ParentClass->GetDefaultObject());
+				if(AnimInstance)
+				{
+					UEdGraph* ParentGraph = TransitionNode->GetGraph();
+					UAnimStateNodeBase* PrevState = TransitionNode->GetPreviousState();
+					UAnimStateNodeBase* NextState = TransitionNode->GetNextState();
+					if(PrevState != nullptr && NextState != nullptr && ParentGraph != nullptr)
+					{
+						FName FunctionName;
+						if(AnimInstance->HasNativeTransitionBinding(ParentGraph->GetFName(), FName(*PrevState->GetStateName()), FName(*NextState->GetStateName()), FunctionName))
+						{
+							TransitionCategory.AddCustomRow( LOCTEXT("NativeBindingPresent", "Transition has native binding") )
+							[
+								SNew(STextBlock)
+								.Text(FText::Format(LOCTEXT("NativeBindingPresent", "Transition has native binding to {0}()"), FText::FromName(FunctionName)))
+								.Font( IDetailLayoutBuilder::GetDetailFontBold() )
+							];
+						}
+					}
+				}
+			}
+
 			TransitionCategory.AddCustomRow( CanExecPin ? CanExecPin->PinFriendlyName : FText::GetEmpty() )
 			[
 				SNew(SKismetLinearExpression, CanExecPin)
diff --git a/Engine/Source/Runtime/Engine/Classes/Animation/AnimBlueprintGeneratedClass.h b/Engine/Source/Runtime/Engine/Classes/Animation/AnimBlueprintGeneratedClass.h
--- a/Engine/Source/Runtime/Engine/Classes/Animation/AnimBlueprintGeneratedClass.h
+++ b/Engine/Source/Runtime/Engine/Classes/Animation/AnimBlueprintGeneratedClass.h
@@ -158,6 +158,19 @@ namespace EPropertySearchMode
 }
 #endif
 
+template<typename NodeType> 
+NodeType* GetNodeFromPropertyIndex(UAnimInstance* AnimInstance, const UAnimBlueprintGeneratedClass* AnimBlueprintClass, int32 PropertyIndex)
+{
+	if( PropertyIndex != INDEX_NONE )
+	{
+		UStructProperty* NodeProperty = AnimBlueprintClass->AnimNodeProperties[AnimBlueprintClass->AnimNodeProperties.Num() - 1 - PropertyIndex]; //@TODO: Crazysauce
+		check(NodeProperty->Struct == NodeType::StaticStruct());
+		return NodeProperty->ContainerPtrToValuePtr<NodeType>(AnimInstance);
+	}
+
+	return NULL;
+}
+
 UCLASS()
 class ENGINE_API UAnimBlueprintGeneratedClass : public UBlueprintGeneratedClass
 {
diff --git a/Engine/Source/Runtime/Engine/Classes/Animation/AnimInstance.h b/Engine/Source/Runtime/Engine/Classes/Animation/AnimInstance.h
--- a/Engine/Source/Runtime/Engine/Classes/Animation/AnimInstance.h
+++ b/Engine/Source/Runtime/Engine/Classes/Animation/AnimInstance.h
@@ -6,6 +6,7 @@
 #include "AnimSequence.h"
 #include "Components/SkeletalMeshComponent.h"
 #include "Components/SkinnedMeshComponent.h"
+#include "AnimStateMachineTypes.h"
 #include "AnimInstance.generated.h"
 
 struct FAnimMontageInstance;
@@ -46,6 +47,12 @@ DECLARE_DYNAMIC_MULTICAST_DELEGATE_TwoParams(FOnMontageEndedMCDelegate, UAnimMon
 */
 DECLARE_DYNAMIC_MULTICAST_DELEGATE_TwoParams(FOnMontageBlendingOutStartedMCDelegate, UAnimMontage*, Montage, bool, bInterrupted);
 
+/** Delegate that native code can hook to to provide additional transition logic */
+DECLARE_DELEGATE_RetVal(bool, FCanTakeTransition);
+
+/** Delegate that native code can hook into to handle state entry/exit */
+DECLARE_DELEGATE_ThreeParams(FOnGraphStateChanged, const struct FAnimNode_StateMachine& /*Machine*/, int32 /*PrevStateIndex*/, int32 /*NextStateIndex*/);
+
 /** Enum for controlling which reference frame a controller is applied in. */
 UENUM()
 enum EBoneControlSpace
@@ -248,6 +255,51 @@ struct FSlotEvaluationPose
 	}
 };
 
+/** Binding allowing native transition rule evaluation */
+struct FNativeTransitionBinding
+{
+	/** State machine to bind to */
+	FName MachineName;
+
+	/** Previous state the transition comes from */
+	FName PreviousStateName;
+
+	/** Next state the transition goes to */
+	FName NextStateName;
+
+	/** Delegate to use when checking transition */
+	FCanTakeTransition NativeTransitionDelegate;
+
+	FNativeTransitionBinding(const FName& InMachineName, const FName& InPreviousStateName, const FName& InNextStateName, const FCanTakeTransition& InNativeTransitionDelegate)
+		: MachineName(InMachineName)
+		, PreviousStateName(InPreviousStateName)
+		, NextStateName(InNextStateName)
+		, NativeTransitionDelegate(InNativeTransitionDelegate)
+	{
+	}
+};
+
+/** Binding allowing native notification of state changes */
+struct FNativeStateBinding
+{
+	/** State machine to bind to */
+	FName MachineName;
+
+	/** State to bind to */
+	FName StateName;
+
+	/** Delegate to use when checking transition */
+	FOnGraphStateChanged NativeStateDelegate;
+
+	FNativeStateBinding(const FName& InMachineName, const FName& InStateName, const FOnGraphStateChanged& InNativeStateDelegate)
+		: MachineName(InMachineName)
+		, StateName(InStateName)
+		, NativeStateDelegate(InNativeStateDelegate)
+	{
+	}
+};
+
+
 UCLASS(transient, Blueprintable, hideCategories=AnimInstance, BlueprintType)
 class ENGINE_API UAnimInstance : public UObject
 {
@@ -579,11 +631,24 @@ public:
 	// Note: the node graph will not be evaluated if this function returns true
 	virtual bool NativeEvaluateAnimation(FPoseContext& Output);
 
-	// Called back from the anim graph when a state is entered
-	virtual void NativeStateStart(const FName& MachineName, const FName& StateName);
+	// Sets up a native transition delegate between states with PrevStateName and NextStateName, in the state machine with name MachineName.
+	// Note that a transition already has to exist for this to succeed
+	void AddNativeTransitionBinding(const FName& MachineName, const FName& PrevStateName, const FName& NextStateName, const FCanTakeTransition& NativeTransitionDelegate);
+
+	// Check for whether a native rule is bound to the specified transition
+	bool HasNativeTransitionBinding(const FName& MachineName, const FName& PrevStateName, const FName& NextStateName, FName& OutBindingName);
+
+	// Sets up a native state entry delegate from state with StateName, in the state machine with name MachineName.
+	void AddNativeStateEntryBinding(const FName& MachineName, const FName& StateName, const FOnGraphStateChanged& NativeEnteredDelegate);
+	
+	// Check for whether a native entry delegate is bound to the specified state
+	bool HasNativeStateEntryBinding(const FName& MachineName, const FName& StateName, FName& OutBindingName);
+
+	// Sets up a native state exit delegate from state with StateName, in the state machine with name MachineName.
+	void AddNativeStateExitBinding(const FName& MachineName, const FName& StateName, const FOnGraphStateChanged& NativeExitedDelegate);
 
-	// Called back from the anim graph when a state is exited
-	virtual void NativeStateEnd(const FName& MachineName, const FName& StateName);
+	// Check for whether a native exit delegate is bound to the specified state
+	bool HasNativeStateExitBinding(const FName& MachineName, const FName& StateName, FName& OutBindingName);
 
 	// Debug output for this anim instance 
 	void DisplayDebug(UCanvas* Canvas, const FDebugDisplayInfo& DebugDisplay, float& YL, float& YPos);
@@ -684,5 +749,19 @@ private:
 
 	// Root motion extracted from animation since the last time ConsumeExtractedRootMotion was called
 	FRootMotionMovementParams ExtractedRootMotion;
+
+	// Native bindings
+private:
+	/** Bind any native delegates that we have set up */
+	void BindNativeDelegates();
+
+	/** Native transition rules */
+	TArray<FNativeTransitionBinding> NativeTransitionBindings;
+
+	/** Native state entry bindings */
+	TArray<FNativeStateBinding> NativeStateEntryBindings;
+
+	/** Native state exit bindings */
+	TArray<FNativeStateBinding> NativeStateExitBindings;
 };
 
diff --git a/Engine/Source/Runtime/Engine/Classes/Animation/AnimNode_StateMachine.h b/Engine/Source/Runtime/Engine/Classes/Animation/AnimNode_StateMachine.h
--- a/Engine/Source/Runtime/Engine/Classes/Animation/AnimNode_StateMachine.h
+++ b/Engine/Source/Runtime/Engine/Classes/Animation/AnimNode_StateMachine.h
@@ -149,6 +149,12 @@ protected:
 	// Used during transitions to make sure we don't double tick a state if it appears multiple times
 	TArray<int32> StatesUpdated;
 
+	// Delegates that native code can hook into to handle state entry
+	TArray<FOnGraphStateChanged> OnGraphStatesEntered;
+
+	// Delegates that native code can hook into to handle state exits
+	TArray<FOnGraphStateChanged> OnGraphStatesExited;
+
 private:
 	// true if it is the first update.
 	bool bFirstUpdate;
@@ -206,4 +212,7 @@ protected:
 	// transition type evaluation functions
 	void EvaluateTransitionStandardBlend(FPoseContext& Output, FAnimationActiveTransitionEntry& Transition, bool bIntermediatePoseIsValid);
 	void EvaluateTransitionCustomBlend(FPoseContext& Output, FAnimationActiveTransitionEntry& Transition, bool bIntermediatePoseIsValid);
+
+public:
+	friend class UAnimInstance;
 };
diff --git a/Engine/Source/Runtime/Engine/Classes/Animation/AnimNode_TransitionResult.h b/Engine/Source/Runtime/Engine/Classes/Animation/AnimNode_TransitionResult.h
--- a/Engine/Source/Runtime/Engine/Classes/Animation/AnimNode_TransitionResult.h
+++ b/Engine/Source/Runtime/Engine/Classes/Animation/AnimNode_TransitionResult.h
@@ -13,6 +13,9 @@ struct ENGINE_API FAnimNode_TransitionResult : public FAnimNode_Base
 	UPROPERTY(EditAnywhere, BlueprintReadWrite, Category=Result, meta=(AlwaysAsPin))
 	bool bCanEnterTransition;
 
+	/** Native delegate to use when checking transition */
+	FCanTakeTransition NativeTransitionDelegate;
+
 public:	
 	FAnimNode_TransitionResult();
 
diff --git a/Engine/Source/Runtime/Engine/Private/Animation/AnimInstance.cpp b/Engine/Source/Runtime/Engine/Private/Animation/AnimInstance.cpp
--- a/Engine/Source/Runtime/Engine/Private/Animation/AnimInstance.cpp
+++ b/Engine/Source/Runtime/Engine/Private/Animation/AnimInstance.cpp
@@ -19,6 +19,7 @@
 #include "Animation/AnimNotifies/AnimNotify.h"
 #include "Animation/AnimMontage.h"
 #include "Animation/AnimInstance.h"
+#include "Animation/AnimNode_TransitionResult.h"
 
 /** Anim stats */
 
@@ -296,6 +297,9 @@ void UAnimInstance::InitializeAnimation()
 		FAnimationInitializeContext InitContext(this);
 		RootNode->Initialize(InitContext);
 	}
+
+	// we can bind rules & events now the graph has been initialized
+	BindNativeDelegates();
 }
 
 #if WITH_EDITORONLY_DATA
@@ -531,12 +535,156 @@ bool UAnimInstance::NativeEvaluateAnimation(FPoseContext& Output)
 	return false;
 }
 
-void UAnimInstance::NativeStateStart(const FName& MachineName, const FName& StateName)
+void UAnimInstance::AddNativeTransitionBinding(const FName& MachineName, const FName& PrevStateName, const FName& NextStateName, const FCanTakeTransition& NativeTransitionDelegate)
+{
+	NativeTransitionBindings.Add(FNativeTransitionBinding(MachineName, PrevStateName, NextStateName, NativeTransitionDelegate));
+}
+
+bool UAnimInstance::HasNativeTransitionBinding(const FName& MachineName, const FName& PrevStateName, const FName& NextStateName, FName& OutBindingName)
+{
+	for(const auto& Binding : NativeTransitionBindings)
+	{
+		if(Binding.MachineName == MachineName && Binding.PreviousStateName == PrevStateName && Binding.NextStateName == NextStateName)
+		{
+			IDelegateInstance* DelegateInstance = Binding.NativeTransitionDelegate.GetDelegateInstance();
+			if(DelegateInstance)
+			{
+				OutBindingName = DelegateInstance->GetFunctionName();
+			}
+			return true;
+		}
+	}
+
+	return false;
+}
+
+void UAnimInstance::AddNativeStateEntryBinding(const FName& MachineName, const FName& StateName, const FOnGraphStateChanged& NativeEnteredDelegate)
+{
+	NativeStateEntryBindings.Add(FNativeStateBinding(MachineName, StateName, NativeEnteredDelegate));
+}
+	
+bool UAnimInstance::HasNativeStateEntryBinding(const FName& MachineName, const FName& StateName, FName& OutBindingName)
+{
+	for(const auto& Binding : NativeStateEntryBindings)
+	{
+		if(Binding.MachineName == MachineName && Binding.StateName == StateName)
+		{
+			IDelegateInstance* DelegateInstance = Binding.NativeStateDelegate.GetDelegateInstance();
+			if(DelegateInstance)
+			{
+				OutBindingName = DelegateInstance->GetFunctionName();
+			}
+			return true;
+		}
+	}
+
+	return false;
+}
+
+void UAnimInstance::AddNativeStateExitBinding(const FName& MachineName, const FName& StateName, const FOnGraphStateChanged& NativeExitedDelegate)
 {
+	NativeStateExitBindings.Add(FNativeStateBinding(MachineName, StateName, NativeExitedDelegate));
+}
+
+bool UAnimInstance::HasNativeStateExitBinding(const FName& MachineName, const FName& StateName, FName& OutBindingName)
+{
+	for(const auto& Binding : NativeStateExitBindings)
+	{
+		if(Binding.MachineName == MachineName && Binding.StateName == StateName)
+		{
+			IDelegateInstance* DelegateInstance = Binding.NativeStateDelegate.GetDelegateInstance();
+			if(DelegateInstance)
+			{
+				OutBindingName = DelegateInstance->GetFunctionName();
+			}
+			return true;
+		}
+	}
+
+	return false;
 }
 
-void UAnimInstance::NativeStateEnd(const FName& MachineName, const FName& StateName) 
+void UAnimInstance::BindNativeDelegates()
 {
+	auto ForEachStateLambda = [&](UAnimBlueprintGeneratedClass* AnimBlueprintGeneratedClass, const FName& MachineName, const FName& StateName, TFunctionRef<void (FAnimNode_StateMachine*, const FBakedAnimationState&, int32)> Predicate)
+	{
+		for(UStructProperty* Property : AnimBlueprintGeneratedClass->AnimNodeProperties)
+		{
+			if(Property && Property->Struct == FAnimNode_StateMachine::StaticStruct())
+			{
+				FAnimNode_StateMachine* StateMachine = Property->ContainerPtrToValuePtr<FAnimNode_StateMachine>(this);
+				if(StateMachine)
+				{
+					const FBakedAnimationStateMachine* MachineDescription = StateMachine->GetMachineDescription();
+					if(MachineDescription && MachineName == MachineDescription->MachineName)
+					{
+						// check each state transition for a match
+						int32 StateIndex = 0;
+						for(const FBakedAnimationState& State : MachineDescription->States)
+						{
+							if(State.StateName == StateName)
+							{
+								Predicate(StateMachine, State, StateIndex);
+							}
+							StateIndex++;
+						}
+					}
+				}
+			}
+		}
+	};
+
+	UAnimBlueprintGeneratedClass* AnimBlueprintGeneratedClass = Cast<UAnimBlueprintGeneratedClass>(GetClass());
+	if(AnimBlueprintGeneratedClass)
+	{
+		// transition delegates
+		for(const auto& Binding : NativeTransitionBindings)
+		{
+			ForEachStateLambda(AnimBlueprintGeneratedClass, Binding.MachineName, Binding.PreviousStateName, 
+				[&](FAnimNode_StateMachine* StateMachine, const FBakedAnimationState& State, int32 StateIndex)
+				{
+					for(const FBakedStateExitTransition& TransitionExit : State.Transitions)
+					{
+						if(TransitionExit.CanTakeDelegateIndex != INDEX_NONE)
+						{
+							const FAnimationTransitionBetweenStates& Transition = StateMachine->GetTransitionInfo(TransitionExit.TransitionIndex);
+							if(StateMachine->GetStateInfo(Transition.NextState).StateName == Binding.NextStateName)
+							{
+								FAnimNode_TransitionResult* ResultNode = GetNodeFromPropertyIndex<FAnimNode_TransitionResult>(this, AnimBlueprintGeneratedClass, TransitionExit.CanTakeDelegateIndex);
+								if(ResultNode)
+								{
+									ResultNode->NativeTransitionDelegate = Binding.NativeTransitionDelegate;
+								}
+							}
+						}
+					}
+				});
+		}
+
+		// state entry delegates
+		for(const auto& Binding : NativeStateEntryBindings)
+		{
+			ForEachStateLambda(AnimBlueprintGeneratedClass, Binding.MachineName, Binding.StateName, 
+				[&](FAnimNode_StateMachine* StateMachine, const FBakedAnimationState& State, int32 StateIndex)
+				{
+					// allocate enough space for all our states we need so far
+					StateMachine->OnGraphStatesEntered.SetNum(FMath::Max(StateIndex + 1, StateMachine->OnGraphStatesEntered.Num()));
+					StateMachine->OnGraphStatesEntered[StateIndex] = Binding.NativeStateDelegate;
+				});
+		}
+
+		// state exit delegates
+		for(const auto& Binding : NativeStateExitBindings)
+		{
+			ForEachStateLambda(AnimBlueprintGeneratedClass, Binding.MachineName, Binding.StateName, 
+				[&](FAnimNode_StateMachine* StateMachine, const FBakedAnimationState& State, int32 StateIndex)
+				{
+					// allocate enough space for all our states we need so far
+					StateMachine->OnGraphStatesExited.SetNum(FMath::Max(StateIndex + 1, StateMachine->OnGraphStatesExited.Num()));
+					StateMachine->OnGraphStatesExited[StateIndex] = Binding.NativeStateDelegate;
+				});
+		}
+	}
 }
 
 void OutputCurveMap(TMap<FName, float>& CurveMap, UCanvas* Canvas, UFont* RenderFont, float Indent, float& YPos, FFontRenderInfo RenderInfo, float& YL)
diff --git a/Engine/Source/Runtime/Engine/Private/Animation/AnimNode_StateMachine.cpp b/Engine/Source/Runtime/Engine/Private/Animation/AnimNode_StateMachine.cpp
--- a/Engine/Source/Runtime/Engine/Private/Animation/AnimNode_StateMachine.cpp
+++ b/Engine/Source/Runtime/Engine/Private/Animation/AnimNode_StateMachine.cpp
@@ -7,21 +7,6 @@
 #include "AnimationRuntime.h"
 #include "AnimTree.h"
 
-//////////////////////////////////////////////////////////////////////////
-
-template<typename NodeType> 
-NodeType* GetNodeFromPropertyIndex(UAnimInstance* AnimInstance, const UAnimBlueprintGeneratedClass* AnimBlueprintClass, int32 PropertyIndex)
-{
-	if( PropertyIndex != INDEX_NONE )
-	{
-		UStructProperty* NodeProperty = AnimBlueprintClass->AnimNodeProperties[AnimBlueprintClass->AnimNodeProperties.Num() - 1 - PropertyIndex]; //@TODO: Crazysauce
-		check(NodeProperty->Struct == NodeType::StaticStruct());
-		return NodeProperty->ContainerPtrToValuePtr<NodeType>(AnimInstance);
-	}
-
-	return NULL;
-}
-
 //////////////////////////////////////////////////////////////////////////
 // FAnimationActiveTransitionEntry
 
@@ -317,24 +302,19 @@ void FAnimNode_StateMachine::Update(const FAnimationUpdateContext& Context)
 				Context.AnimInstance->AddAnimNotifyFromGeneratedClass(ActiveTransitionArray[ActiveTransitionArray.Num()-1].InterruptNotify);
 			}
 
-			const int32 PreviousStateIndex = CurrentState;
-			const int32 NextStateIndex = PotentialTransition.TargetState;
+			const int32 PreviousState = CurrentState;
+			const int32 NextState = PotentialTransition.TargetState;
 
 			// Fire off Notifies for state transition
-			const FBakedAnimationState& PreviousState = GetStateInfo(PreviousStateIndex);
-			Context.AnimInstance->NativeStateEnd(PRIVATE_MachineDescription->MachineName, PreviousState.StateName);
-			Context.AnimInstance->AddAnimNotifyFromGeneratedClass(PreviousState.EndNotify);
-
-			const FBakedAnimationState& NextState = GetStateInfo(PreviousStateIndex);
-			Context.AnimInstance->NativeStateStart(PRIVATE_MachineDescription->MachineName, NextState.StateName);
-			Context.AnimInstance->AddAnimNotifyFromGeneratedClass(NextState.StartNotify);
+			Context.AnimInstance->AddAnimNotifyFromGeneratedClass(GetStateInfo(PreviousState).EndNotify);
+			Context.AnimInstance->AddAnimNotifyFromGeneratedClass(GetStateInfo(NextState).StartNotify);
 			
 			// Get the current weight of the next state, which may be non-zero
-			const float ExistingWeightOfNextState = GetStateWeight(NextStateIndex);
+			const float ExistingWeightOfNextState = GetStateWeight(NextState);
 
 			// Push the transition onto the stack
 			const FAnimationTransitionBetweenStates& ReferenceTransition = GetTransitionInfo(PotentialTransition.TransitionRule->TransitionIndex);
-			FAnimationActiveTransitionEntry* NewTransition = new (ActiveTransitionArray) FAnimationActiveTransitionEntry(NextStateIndex, ExistingWeightOfNextState, PreviousStateIndex, ReferenceTransition);
+			FAnimationActiveTransitionEntry* NewTransition = new (ActiveTransitionArray) FAnimationActiveTransitionEntry(NextState, ExistingWeightOfNextState, PreviousState, ReferenceTransition);
 			NewTransition->InitializeCustomGraphLinks(Context, *(PotentialTransition.TransitionRule));
 
 #if WITH_EDITORONLY_DATA
@@ -343,7 +323,7 @@ void FAnimNode_StateMachine::Update(const FAnimationUpdateContext& Context)
 
 			Context.AnimInstance->AddAnimNotifyFromGeneratedClass(NewTransition->StartNotify);
 			
-			SetState(Context, NextStateIndex);
+			SetState(Context, NextState);
 
 			TransitionCountThisFrame++;
 		}
@@ -446,8 +426,18 @@ bool FAnimNode_StateMachine::FindValidTransition(const FAnimationUpdateContext&
 
 		FAnimNode_TransitionResult* ResultNode = GetNodeFromPropertyIndex<FAnimNode_TransitionResult>(Context.AnimInstance, AnimBlueprintClass, TransitionRule.CanTakeDelegateIndex);
 
+		// reset transition flag (this would be always set if we had a bound graph, but if we only have a delegate it will not)
+		ResultNode->bCanEnterTransition = false;
+
 		// Execute it and see if we can take this rule
 		ResultNode->EvaluateGraphExposedInputs.Execute(Context);
+
+		// attempt to evaluate native rule
+		if(ResultNode->NativeTransitionDelegate.IsBound())
+		{
+			ResultNode->bCanEnterTransition |= ResultNode->NativeTransitionDelegate.Execute();
+		}
+
 		if (ResultNode->bCanEnterTransition == TransitionRule.bDesiredTransitionReturnValue)
 		{
 			int32 NextState = GetTransitionInfo(TransitionRule.TransitionIndex).NextState;
@@ -747,6 +737,12 @@ void FAnimNode_StateMachine::SetState(const FAnimationBaseContext& Context, int3
 {
 	if (NewStateIndex != CurrentState)
 	{
+		const int32 PrevStateIndex = CurrentState;
+		if(CurrentState != INDEX_NONE && CurrentState < OnGraphStatesEntered.Num())
+		{
+			OnGraphStatesEntered[CurrentState].ExecuteIfBound(*this, CurrentState, NewStateIndex);
+		}
+
 		// Determine if the new state is active or not
 		const bool bAlreadyActive = GetStateWeight(NewStateIndex) > 0.0f;
 
@@ -762,6 +758,11 @@ void FAnimNode_StateMachine::SetState(const FAnimationBaseContext& Context, int3
 			FAnimationCacheBonesContext CacheBoneContext(Context.AnimInstance);
 			StatePoseLinks[NewStateIndex].CacheBones(CacheBoneContext);
 		}
+
+		if(CurrentState != INDEX_NONE && CurrentState < OnGraphStatesExited.Num())
+		{
+			OnGraphStatesExited[CurrentState].ExecuteIfBound(*this, PrevStateIndex, CurrentState);
+		}
 	}
 }
 
