Constified USceneComponent::GetRelativeTransform.

diff --git a/Engine/Source/Runtime/Engine/Classes/Components/SceneComponent.h b/Engine/Source/Runtime/Engine/Classes/Components/SceneComponent.h
--- a/Engine/Source/Runtime/Engine/Classes/Components/SceneComponent.h
+++ b/Engine/Source/Runtime/Engine/Classes/Components/SceneComponent.h
@@ -244,7 +244,7 @@ public:
 
 	/** Returns the transform of this component relative to its parent */
 	UFUNCTION(BlueprintCallable, Category="Utilities|Transformation")
-	FTransform GetRelativeTransform();
+	FTransform GetRelativeTransform() const;
 
 	/** Set the transform of this component relative to its parent */
 	UFUNCTION(BlueprintCallable, Category="Utilities|Transformation")
diff --git a/Engine/Source/Runtime/Engine/Private/SceneComponent.cpp b/Engine/Source/Runtime/Engine/Private/SceneComponent.cpp
--- a/Engine/Source/Runtime/Engine/Private/SceneComponent.cpp
+++ b/Engine/Source/Runtime/Engine/Private/SceneComponent.cpp
@@ -447,7 +447,7 @@ void USceneComponent::SetRelativeTransform(const FTransform& NewTransform, bool
 	SetRelativeScale3D(NewTransform.GetScale3D());
 }
 
-FTransform USceneComponent::GetRelativeTransform()
+FTransform USceneComponent::GetRelativeTransform() const
 {
 	const FTransform RelativeTransform( RelativeRotation, RelativeLocation, RelativeScale3D );
 	return RelativeTransform;
