Added a function to VisualLog supporting adding of arbitrary log element #UE4

diff --git a/Engine/Source/Runtime/Engine/Private/VisualLogger/VisualLoggerTypes.cpp b/Engine/Source/Runtime/Engine/Private/VisualLogger/VisualLoggerTypes.cpp
--- a/Engine/Source/Runtime/Engine/Private/VisualLogger/VisualLoggerTypes.cpp
+++ b/Engine/Source/Runtime/Engine/Private/VisualLogger/VisualLoggerTypes.cpp
@@ -111,6 +111,11 @@ void FVisualLogEntry::AddText(const FString& TextLine, const FName& CategoryName
 	LogLines.Add(FVisualLogLine(CategoryName, Verbosity, TextLine));
 }
 
+void FVisualLogEntry::AddElement(FVisualLogShapeElement Element)
+{
+	ElementsToDraw.Emplace(Element);
+}
+
 void FVisualLogEntry::AddElement(const TArray<FVector>& Points, const FName& CategoryName, ELogVerbosity::Type Verbosity, const FColor& Color, const FString& Description, uint16 Thickness)
 {
 	FVisualLogShapeElement Element(Description, Color, Thickness, CategoryName);
diff --git a/Engine/Source/Runtime/Engine/Public/VisualLogger/VisualLoggerTypes.h b/Engine/Source/Runtime/Engine/Public/VisualLogger/VisualLoggerTypes.h
--- a/Engine/Source/Runtime/Engine/Public/VisualLogger/VisualLoggerTypes.h
+++ b/Engine/Source/Runtime/Engine/Public/VisualLogger/VisualLoggerTypes.h
@@ -189,6 +189,8 @@ struct ENGINE_API FVisualLogEntry
 	void AddElement(const FVector& Start, const FVector& End, float Radius, const FName& CategoryName, ELogVerbosity::Type Verbosity, const FColor& Color = FColor::White, const FString& Description = TEXT(""), uint16 Thickness = 0);
 	// capsule
 	void AddElement(const FVector& Center, float HalfHeight, float Radius, const FQuat & Rotation, const FName& CategoryName, ELogVerbosity::Type Verbosity, const FColor& Color = FColor::White, const FString& Description = TEXT(""));
+	// custom element
+	void AddElement(FVisualLogShapeElement Element);
 	// histogram sample
 	void AddHistogramData(const FVector2D& DataSample, const FName& CategoryName, ELogVerbosity::Type Verbosity, const FName& GraphName, const FName& DataName);
 	// Custom data block
