Remove spam deubg output

diff --git a/Engine/Source/Runtime/Engine/Private/SkinnedMeshComponent.cpp b/Engine/Source/Runtime/Engine/Private/SkinnedMeshComponent.cpp
--- a/Engine/Source/Runtime/Engine/Private/SkinnedMeshComponent.cpp
+++ b/Engine/Source/Runtime/Engine/Private/SkinnedMeshComponent.cpp
@@ -89,7 +89,6 @@ namespace FAnimUpdateRateManager
 		checkSlow(!(*ExistingTracker)->RegisteredComponents.Contains(SkinnedComponent)); // We have already been registered? Something has gone very wrong!
 
 		(*ExistingTracker)->RegisteredComponents.Add(SkinnedComponent);
-		FPlatformMisc::LowLevelOutputDebugStringf(TEXT("MAPUPDATE: Count %i\n"), ActorToUpdateRateParams.Num());
 		return &(*ExistingTracker)->UpdateRateParameters;
 	}
 
@@ -104,7 +103,6 @@ namespace FAnimUpdateRateManager
 			ActorToUpdateRateParams.Remove(TrackerIndex);
 			delete Tracker;
 		}
-		FPlatformMisc::LowLevelOutputDebugStringf(TEXT("MAPUPDATE: Count %i\n"), ActorToUpdateRateParams.Num());
 	}
 
 	void TickUpdateRateParameters(USkinnedMeshComponent* SkinnedComponent)
