Fixed UE-8195 "All Visual Logger automation tests fail with CurrentEntry->TimeStamp == -1; The value is not true."

diff --git a/Engine/Source/Runtime/Engine/Private/VisualLogger/VisualLoggerAutomationTests.cpp b/Engine/Source/Runtime/Engine/Private/VisualLogger/VisualLoggerAutomationTests.cpp
--- a/Engine/Source/Runtime/Engine/Private/VisualLogger/VisualLoggerAutomationTests.cpp
+++ b/Engine/Source/Runtime/Engine/Private/VisualLogger/VisualLoggerAutomationTests.cpp
@@ -114,11 +114,18 @@ bool FVisualLogTest::RunTest(const FString& Parameters)
 			CHECK_SUCCESS(CurrentEntry->LogLines[0].Category == LogVisual.GetCategoryName());
 			CHECK_SUCCESS(CurrentEntry->LogLines[0].Line == TextToLog);
 
-			FVisualLogger::Get().GetEntryToWrite(GWorld, CurrentTimestamp + 0.1); //generate new entry and serialize old one
+			const float NewTimestamp = CurrentTimestamp + 0.1;
+			FVisualLogEntry* NewEntry = FVisualLogger::Get().GetEntryToWrite(GWorld, NewTimestamp); //generate new entry and serialize old one
 			CurrentEntry = &Context.Device.LastEntry;
 			CHECK_SUCCESS(CurrentEntry != NULL);
-			CHECK_SUCCESS(CurrentEntry->TimeStamp == -1);
-			CHECK_SUCCESS(CurrentEntry->LogLines.Num() == 0);
+			CHECK_SUCCESS(CurrentEntry->TimeStamp == CurrentTimestamp);
+			CHECK_SUCCESS(CurrentEntry->LogLines.Num() == 1);
+			CHECK_SUCCESS(CurrentEntry->LogLines[0].Category == LogVisual.GetCategoryName());
+			CHECK_SUCCESS(CurrentEntry->LogLines[0].Line == TextToLog);
+
+			CHECK_SUCCESS(NewEntry != NULL);
+			CHECK_SUCCESS(NewEntry->TimeStamp - NewTimestamp <= SMALL_NUMBER);
+			CHECK_SUCCESS(NewEntry->LogLines.Num() == 0);
 		}
 	}
 
@@ -160,11 +167,20 @@ bool FVisualLogSegmentsTest::RunTest(const FString& Parameters)
 			CHECK_SUCCESS(CurrentEntry->ElementsToDraw[0].Points[0] == StartPoint);
 			CHECK_SUCCESS(CurrentEntry->ElementsToDraw[0].Points[1] == EndPoint);
 
-			FVisualLogger::Get().GetEntryToWrite(GWorld, CurrentTimestamp + 0.1); //generate new entry and serialize old one
+			const float NewTimestamp = CurrentTimestamp + 0.1;
+			FVisualLogEntry* NewEntry = FVisualLogger::Get().GetEntryToWrite(GWorld, NewTimestamp); //generate new entry and serialize old one
 			CurrentEntry = &Context.Device.LastEntry;
 			CHECK_SUCCESS(CurrentEntry != NULL);
-			CHECK_SUCCESS(CurrentEntry->TimeStamp == -1);
-			CHECK_SUCCESS(CurrentEntry->ElementsToDraw.Num() == 0);
+			CHECK_SUCCESS(CurrentEntry->TimeStamp == CurrentTimestamp);
+			CHECK_SUCCESS(CurrentEntry->ElementsToDraw.Num() == 1);
+			CHECK_SUCCESS(CurrentEntry->ElementsToDraw[0].GetType() == EVisualLoggerShapeElement::Segment);
+			CHECK_SUCCESS(CurrentEntry->ElementsToDraw[0].Points.Num() == 2);
+			CHECK_SUCCESS(CurrentEntry->ElementsToDraw[0].Points[0] == StartPoint);
+			CHECK_SUCCESS(CurrentEntry->ElementsToDraw[0].Points[1] == EndPoint);
+
+			CHECK_SUCCESS(NewEntry != NULL);
+			CHECK_SUCCESS(NewEntry->TimeStamp - NewTimestamp <= SMALL_NUMBER);
+			CHECK_SUCCESS(NewEntry->ElementsToDraw.Num() == 0);
 		}
 	}
 	return true;
@@ -226,11 +242,26 @@ bool FVisualLogEventsTest::RunTest(const FString& Parameters)
 		CHECK_SUCCESS(CurrentEntry->Events[1].UserFriendlyDesc == TEXT("Second simple event for vlog tests"));
 		CHECK_SUCCESS(CurrentEntry->Events[2].UserFriendlyDesc == TEXT("Third simple event for vlog tests"));
 
-		FVisualLogger::Get().GetEntryToWrite(GWorld, CurrentTimestamp+0.1); //generate new entry and serialize old one
+		const float NewTimestamp = CurrentTimestamp + 0.1;
+		FVisualLogEntry* NewEntry = FVisualLogger::Get().GetEntryToWrite(GWorld, NewTimestamp); //generate new entry and serialize old one
 		CurrentEntry = &Context.Device.LastEntry;
 		CHECK_SUCCESS(CurrentEntry != NULL);
-		CHECK_SUCCESS(CurrentEntry->TimeStamp == -1);
-		CHECK_SUCCESS(CurrentEntry->Events.Num() == 0);
+		CHECK_SUCCESS(CurrentEntry->TimeStamp == CurrentTimestamp);
+		CHECK_SUCCESS(CurrentEntry->Events.Num() == 3);
+		CHECK_SUCCESS(CurrentEntry->Events[0].Counter == 3);
+		CHECK_SUCCESS(CurrentEntry->Events[0].Name == TEXT("EventTest"));
+		CHECK_SUCCESS(CurrentEntry->Events[1].Counter == 2);
+		CHECK_SUCCESS(CurrentEntry->Events[1].Name == TEXT("EventTest2"));
+		CHECK_SUCCESS(CurrentEntry->Events[2].Counter == 1);
+		CHECK_SUCCESS(CurrentEntry->Events[2].Name == TEXT("EventTest3"));
+
+		CHECK_SUCCESS(CurrentEntry->Events[0].UserFriendlyDesc == TEXT("Simple event for vlog tests"));
+		CHECK_SUCCESS(CurrentEntry->Events[1].UserFriendlyDesc == TEXT("Second simple event for vlog tests"));
+		CHECK_SUCCESS(CurrentEntry->Events[2].UserFriendlyDesc == TEXT("Third simple event for vlog tests"));
+
+		CHECK_SUCCESS(NewEntry != NULL);
+		CHECK_SUCCESS(NewEntry->TimeStamp - NewTimestamp <= SMALL_NUMBER);
+		CHECK_SUCCESS(NewEntry->Events.Num() == 0);
 	}
 
 	const FName EventTag1 = TEXT("ATLAS_C_0");
