Fix foliage static lighting
I changed the multiplier/divisor used by the lightmap offset to 32767.*0* to match what OpenGL/D3D use when decoding a short to a float

diff --git a/Engine/Source/Runtime/Engine/Private/InstancedStaticMesh.cpp b/Engine/Source/Runtime/Engine/Private/InstancedStaticMesh.cpp
--- a/Engine/Source/Runtime/Engine/Private/InstancedStaticMesh.cpp
+++ b/Engine/Source/Runtime/Engine/Private/InstancedStaticMesh.cpp
@@ -1132,6 +1132,10 @@ void UInstancedStaticMeshComponent::ApplyLightMapping(FStaticLightingTextureMapp
 		IrrelevantLights = PossiblyIrrelevantLights.Array();
 
 		bHasCachedStaticLighting = true;
+
+		ReleasePerInstanceRenderData();
+		MarkRenderStateDirty();
+
 		MarkPackageDirty();
 	}
 }
@@ -1670,7 +1674,7 @@ void FInstancedStaticMeshVertexFactoryShaderParameters::SetMesh( FRHICommandList
 	{
 		if (CPUInstanceOrigin.IsBound())
 		{
-			const float ShortScale = 1.0f / 32767.5f;
+			const float ShortScale = 1.0f / 32767.0f;
 			auto* InstancingData = (const FInstancingUserData*)BatchElement.UserData;
 			const FInstanceStream* InstanceStream = InstancingData->RenderData->PerInstanceRenderData->InstanceBuffer.GetData() + BatchElement.UserIndex;
 
diff --git a/Engine/Source/Runtime/Engine/Private/LightMap.cpp b/Engine/Source/Runtime/Engine/Private/LightMap.cpp
--- a/Engine/Source/Runtime/Engine/Private/LightMap.cpp
+++ b/Engine/Source/Runtime/Engine/Private/LightMap.cpp
@@ -253,6 +253,9 @@ struct FLightMapAllocation
 			// TODO: We currently only support one LOD of static lighting in foliage
 			// Need to create per-LOD instance data to fix that
 			Component->PerInstanceSMData[InstanceIndex].LightmapUVBias = LightMap->GetCoordinateBias();
+
+			Component->ReleasePerInstanceRenderData();
+			Component->MarkRenderStateDirty();
 		}
 	}
 
diff --git a/Engine/Source/Runtime/Engine/Private/ShadowMap.cpp b/Engine/Source/Runtime/Engine/Private/ShadowMap.cpp
--- a/Engine/Source/Runtime/Engine/Private/ShadowMap.cpp
+++ b/Engine/Source/Runtime/Engine/Private/ShadowMap.cpp
@@ -96,6 +96,9 @@ struct FShadowMapAllocation
 			// TODO: We currently only support one LOD of static lighting in foliage
 			// Need to create per-LOD instance data to fix that
 			Component->PerInstanceSMData[InstanceIndex].ShadowmapUVBias = ShadowMap->GetCoordinateBias();
+
+			Component->ReleasePerInstanceRenderData();
+			Component->MarkRenderStateDirty();
 		}
 	}
 };
diff --git a/Engine/Source/Runtime/Engine/Public/StaticMeshResources.h b/Engine/Source/Runtime/Engine/Public/StaticMeshResources.h
--- a/Engine/Source/Runtime/Engine/Public/StaticMeshResources.h
+++ b/Engine/Source/Runtime/Engine/Public/StaticMeshResources.h
@@ -962,10 +962,10 @@ struct FInstanceStream
 		Me->InstanceTransform3[2] = Transform.M[2][2];
 		Me->InstanceTransform3[3] = (float)HitProxyColor.B;
 
-		Me->InstanceLightmapAndShadowMapUVBias[0] = FMath::Clamp<int32>(FMath::TruncToInt(rLightmapUVBias->X * 32767.5f),MIN_int16,MAX_int16);
-		Me->InstanceLightmapAndShadowMapUVBias[1] = FMath::Clamp<int32>(FMath::TruncToInt(rLightmapUVBias->Y * 32767.5f),MIN_int16,MAX_int16);
-		Me->InstanceLightmapAndShadowMapUVBias[2] = FMath::Clamp<int32>(FMath::TruncToInt(rShadowmapUVBias->X * 32767.5f),MIN_int16,MAX_int16);
-		Me->InstanceLightmapAndShadowMapUVBias[3] = FMath::Clamp<int32>(FMath::TruncToInt(rShadowmapUVBias->Y * 32767.5f),MIN_int16,MAX_int16);
+		Me->InstanceLightmapAndShadowMapUVBias[0] = FMath::Clamp<int32>(FMath::TruncToInt(rLightmapUVBias->X  * 32767.0f), MIN_int16, MAX_int16);
+		Me->InstanceLightmapAndShadowMapUVBias[1] = FMath::Clamp<int32>(FMath::TruncToInt(rLightmapUVBias->Y  * 32767.0f), MIN_int16, MAX_int16);
+		Me->InstanceLightmapAndShadowMapUVBias[2] = FMath::Clamp<int32>(FMath::TruncToInt(rShadowmapUVBias->X * 32767.0f), MIN_int16, MAX_int16);
+		Me->InstanceLightmapAndShadowMapUVBias[3] = FMath::Clamp<int32>(FMath::TruncToInt(rShadowmapUVBias->Y * 32767.0f), MIN_int16, MAX_int16);
 	}
 
 	FORCEINLINE void NullifyInstance()
