Merging SystemSettingsEditor enabled

//depot/UE4-UT/...

to //depot/UE4/...

diff --git a/Engine/Source/Runtime/Engine/Private/SystemSettings.cpp b/Engine/Source/Runtime/Engine/Private/SystemSettings.cpp
--- a/Engine/Source/Runtime/Engine/Private/SystemSettings.cpp
+++ b/Engine/Source/Runtime/Engine/Private/SystemSettings.cpp
@@ -334,6 +334,8 @@ void FSystemSettings::Initialize( bool bSetupForEditor )
 	FSystemSettingsData DefaultSettings;
 	DefaultSettings.LoadFromIni(GetSectionName(false), GEngineIni, false);
 
+	bIsEditor = bSetupForEditor;
+
 	(FSystemSettingsData&)(*this) = DefaultSettings;
 	LoadFromIni();
 
