UE4 - fix prestreaming of grass and texture with parented cameras

diff --git a/Engine/Source/Runtime/Engine/Classes/Matinee/MatineeActor.h b/Engine/Source/Runtime/Engine/Classes/Matinee/MatineeActor.h
--- a/Engine/Source/Runtime/Engine/Classes/Matinee/MatineeActor.h
+++ b/Engine/Source/Runtime/Engine/Classes/Matinee/MatineeActor.h
@@ -342,6 +342,9 @@ public:
 	/** Scan the matinee data for camera cuts and sets up the CameraCut array. */
 	void SetupCameraCuts();
 
+	/** Used when setting up the camera cuts to make sure the parent is updated. */
+	void UpdateInterpForParentMovementTracks( float Time, UInterpGroupInst* ViewGroupInst );
+
 	/** Disable the radio filter effect if "Disable Radio Filter" is checked. */
 	void DisableRadioFilterIfNeeded();
 
diff --git a/Engine/Source/Runtime/Engine/Private/Interpolation.cpp b/Engine/Source/Runtime/Engine/Private/Interpolation.cpp
--- a/Engine/Source/Runtime/Engine/Private/Interpolation.cpp
+++ b/Engine/Source/Runtime/Engine/Private/Interpolation.cpp
@@ -890,6 +890,37 @@ void AMatineeActor::TermInterp()
 	EnableCinematicMode(false);
 }
 
+void AMatineeActor::UpdateInterpForParentMovementTracks( float Time, UInterpGroupInst* ViewGroupInst )
+{
+	AActor* Actor = ViewGroupInst->GetGroupActor();
+	AActor* Parent = Actor->GetAttachParentActor();
+
+	UInterpGroupInst* ParentInst = FindGroupInst( Parent );
+	if(ParentInst)
+	{
+		TArray<UInterpTrack*> FoundTracks;
+
+		UInterpTrackInst* ParentTrackInst = nullptr;
+		for (UInterpTrackInst* Inst : ParentInst->TrackInst )
+		{
+			if( Inst->GetGroupActor() == Parent )
+			{
+				ParentTrackInst = Inst;
+				break;
+			}
+		}
+
+		if (ParentTrackInst)
+		{
+			TArray<UInterpTrack*> FoundTracks;
+			ParentInst->Group->FindTracksByClass( UInterpTrackMove::StaticClass(), FoundTracks );
+			//Just use the first one, multiple move tracks wouldnt work well anyway
+			UInterpTrackMove* MoveTrack = CastChecked<UInterpTrackMove>(FoundTracks[0]);
+			MoveTrack->ConditionalUpdateTrack(Time, ParentTrackInst, true);
+		}
+	}
+}
+
 void AMatineeActor::SetupCameraCuts()
 {
 	if( MatineeData )
@@ -900,6 +931,8 @@ void AMatineeActor::SetupCameraCuts()
 		{
 			CameraCuts.Reserve( DirTrack->CutTrack.Num() );
 
+			float OldInterpPosition = InterpPosition;
+
 			// Find the starting camera location for each cut.
 			for( int32 KeyFrameIndex=0; KeyFrameIndex < DirTrack->CutTrack.Num(); KeyFrameIndex++)
 			{
@@ -919,17 +952,14 @@ void AMatineeActor::SetupCameraCuts()
 							FRotator CameraRotation;
 
 							UInterpTrackInst* TrackInst = ViewGroupInst->TrackInst[ TrackIndex ];
-							bool bSucceeded = MoveTrack->GetLocationAtTime( TrackInst, Cut.Time, CameraCut.Location, CameraRotation );
-
-							// The first keyframe could be (0,0,0), try again slightly past it in time.
-							if ( !bSucceeded || CameraCut.Location.IsNearlyZero() == true )
-							{
-								bSucceeded = MoveTrack->GetLocationAtTime( TrackInst, Cut.Time+0.01f, CameraCut.Location, CameraRotation );
-							}
+							UpdateInterpForParentMovementTracks(Cut.Time + .01f, ViewGroupInst);
+							bool bSucceeded = MoveTrack->GetLocationAtTime( TrackInst, Cut.Time+0.01f, CameraCut.Location, CameraRotation );
+							UpdateInterpForParentMovementTracks(OldInterpPosition, ViewGroupInst);
 
 							// Only add locations that aren't (0,0,0)
 							if ( bSucceeded && CameraCut.Location.IsNearlyZero() == false )
 							{
+								//UE_LOG(LogTemp, Display, TEXT("Cut %d   %s    %f    %f %f %f"), CameraCuts.Num(), *Cut.TargetCamGroup.ToString(), Cut.Time, CameraCut.Location.X, CameraCut.Location.Y, CameraCut.Location.Z);
 								CameraCut.TimeStamp = Cut.Time;
 								CameraCuts.Add( CameraCut );
 								break;
