Fix for UE-5011.

ParticleSystemComponents spawned with the world as an outer (as opposed to an actor) were returning early from the collision module.
Removed the check as this setup should be fine.

diff --git a/Engine/Source/Runtime/Engine/Private/Particles/ParticleModules_Collision.cpp b/Engine/Source/Runtime/Engine/Private/Particles/ParticleModules_Collision.cpp
--- a/Engine/Source/Runtime/Engine/Private/Particles/ParticleModules_Collision.cpp
+++ b/Engine/Source/Runtime/Engine/Private/Particles/ParticleModules_Collision.cpp
@@ -161,11 +161,8 @@ void UParticleModuleCollision::Update(FParticleEmitterInstance* Owner, int32 Off
 		return;
 	}
 
+	//Gets the owning actor of the component. Can be NULL if the component is spawned with the World as an Outer, e.g. in UGameplayStatics::SpawnEmitterAtLocation().
 	AActor* Actor = Owner->Component->GetOwner();
-	if (!Actor && Owner->GetWorld()->IsGameWorld() )
-	{
-		return;
-	}
 
 	UParticleLODLevel* LODLevel	= Owner->SpriteTemplate->GetCurrentLODLevel(Owner);
 	check(LODLevel);
