Removed an unused variable.

diff --git a/Engine/Source/Runtime/Engine/Private/FloatingPawnMovement.cpp b/Engine/Source/Runtime/Engine/Private/FloatingPawnMovement.cpp
--- a/Engine/Source/Runtime/Engine/Private/FloatingPawnMovement.cpp
+++ b/Engine/Source/Runtime/Engine/Private/FloatingPawnMovement.cpp
@@ -48,7 +48,6 @@ void UFloatingPawnMovement::TickComponent(float DeltaTime, enum ELevelTick TickT
 		{
 			const FVector OldLocation = UpdatedComponent->GetComponentLocation();
 			const FRotator Rotation = UpdatedComponent->GetComponentRotation();
-			FVector TraceStart = OldLocation;
 
 			FHitResult Hit(1.f);
 			SafeMoveUpdatedComponent(Delta, Rotation, true, Hit);
