added another bloom stage for larger more efficient bloom

diff --git a/Engine/Source/Runtime/Engine/Classes/Engine/Scene.h b/Engine/Source/Runtime/Engine/Classes/Engine/Scene.h
--- a/Engine/Source/Runtime/Engine/Classes/Engine/Scene.h
+++ b/Engine/Source/Runtime/Engine/Classes/Engine/Scene.h
@@ -160,6 +160,12 @@ struct FPostProcessSettings
 	UPROPERTY(BlueprintReadWrite, Category=Overrides, meta=(PinHiddenByDefault))
 	uint32 bOverride_Bloom5Size:1;
 
+	UPROPERTY(BlueprintReadWrite, Category=Overrides, meta=(PinHiddenByDefault))
+	uint32 bOverride_Bloom6Tint:1;
+
+	UPROPERTY(BlueprintReadWrite, Category=Overrides, meta=(PinHiddenByDefault))
+	uint32 bOverride_Bloom6Size:1;
+
 	UPROPERTY(BlueprintReadWrite, Category=Overrides, meta=(PinHiddenByDefault))
 	uint32 bOverride_BloomSizeScale:1;
 
@@ -495,6 +501,13 @@ struct FPostProcessSettings
 	 */
 	UPROPERTY(interp, BlueprintReadWrite, Category=Bloom, AdvancedDisplay, meta=(ClampMin = "0.0", UIMax = "64.0", editcondition = "bOverride_Bloom5Size", DisplayName = "#5 Size"))
 	float Bloom5Size;
+	/**
+	 * Diameter size for Bloom6 in percent of the screen width
+	 * (is done in 1/64 resolution, larger values cost more performance, best for wide contributions)
+	 * >=0: can be clamped because of shader limitations
+	 */
+	UPROPERTY(interp, BlueprintReadWrite, Category=Bloom, AdvancedDisplay, meta=(ClampMin = "0.0", UIMax = "128.0", editcondition = "bOverride_Bloom6Size", DisplayName = "#6 Size"))
+	float Bloom6Size;
 
 	/** Bloom1 tint color */
 	UPROPERTY(interp, BlueprintReadWrite, Category=Bloom, AdvancedDisplay, meta=(editcondition = "bOverride_Bloom1Tint", DisplayName = "#1 Tint", HideAlphaChannel))
@@ -511,6 +524,9 @@ struct FPostProcessSettings
 	/** Bloom5 tint color */
 	UPROPERTY(interp, BlueprintReadWrite, Category=Bloom, AdvancedDisplay, meta=(editcondition = "bOverride_Bloom5Tint", DisplayName = "#5 Tint", HideAlphaChannel))
 	FLinearColor Bloom5Tint;
+	/** Bloom6 tint color */
+	UPROPERTY(interp, BlueprintReadWrite, Category=Bloom, AdvancedDisplay, meta=(editcondition = "bOverride_Bloom6Tint", DisplayName = "#6 Tint", HideAlphaChannel))
+	FLinearColor Bloom6Tint;
 
 	/** BloomDirtMask intensity */
 	UPROPERTY(interp, BlueprintReadWrite, Category=Bloom, meta=(ClampMin = "0.0", UIMax = "8.0", editcondition = "bOverride_BloomDirtMaskIntensity", DisplayName = "Dirt Mask Intensity"))
@@ -898,7 +914,9 @@ struct FPostProcessSettings
 		Bloom4Tint = FLinearColor(0.5f, 0.5f, 0.5f);
 		Bloom4Size = 32.0f;
 		Bloom5Tint = FLinearColor(0.5f, 0.5f, 0.5f);
-		Bloom5Size = 100.0f;
+		Bloom5Size = 64.0f;
+		Bloom6Tint = FLinearColor(0.5f, 0.5f, 0.5f);
+		Bloom6Size = 64.0f;
 		BloomDirtMaskIntensity = 1.0f;
 		BloomDirtMaskTint = FLinearColor(0.5f, 0.5f, 0.5f);
 		AmbientCubemapIntensity = 1.0f;
diff --git a/Engine/Source/Runtime/Engine/Private/SceneView.cpp b/Engine/Source/Runtime/Engine/Private/SceneView.cpp
--- a/Engine/Source/Runtime/Engine/Private/SceneView.cpp
+++ b/Engine/Source/Runtime/Engine/Private/SceneView.cpp
@@ -744,6 +744,8 @@ void FSceneView::OverridePostProcessSettings(const FPostProcessSettings& Src, fl
 	LERP_PP(Bloom4Size);
 	LERP_PP(Bloom5Tint);
 	LERP_PP(Bloom5Size);
+	LERP_PP(Bloom6Tint);
+	LERP_PP(Bloom6Size);
 	LERP_PP(BloomDirtMaskIntensity);
 	LERP_PP(BloomDirtMaskTint);
 	LERP_PP(AmbientCubemapIntensity);
diff --git a/Engine/Source/Runtime/Renderer/Private/PostProcess/PostProcessing.cpp b/Engine/Source/Runtime/Renderer/Private/PostProcess/PostProcessing.cpp
--- a/Engine/Source/Runtime/Renderer/Private/PostProcess/PostProcessing.cpp
+++ b/Engine/Source/Runtime/Renderer/Private/PostProcess/PostProcessing.cpp
@@ -435,8 +435,8 @@ static FRenderingCompositeOutputRef AddBloom(FPostprocessContext Context, FRende
 		3,// Q1
 		3,// Q2
 		4,// Q3
-		4,// Q4
-		5,// Q5
+		5,// Q4
+		6,// Q5
 	};
 
 	int32 BloomQuality;
@@ -447,12 +447,12 @@ static FRenderingCompositeOutputRef AddBloom(FPostprocessContext Context, FRende
 	}
 
 	// Perform down sample. Used by both bloom and lens flares.
-	static const int32 DownSampleStages = 5;
+	static const int32 DownSampleStages = 6;
 	FRenderingCompositeOutputRef PostProcessDownsamples[DownSampleStages] = {PostProcessDownsample0};
 	for (int i = 1; i < DownSampleStages; i++)
 	{
 		static const TCHAR* PassLabels[] =
-			{NULL, TEXT("BloomDownsample1"), TEXT("BloomDownsample2"), TEXT("BloomDownsample3"), TEXT("BloomDownsample4")};
+			{NULL, TEXT("BloomDownsample1"), TEXT("BloomDownsample2"), TEXT("BloomDownsample3"), TEXT("BloomDownsample4"), TEXT("BloomDownsample5")};
 		static_assert(ARRAY_COUNT(PassLabels) == DownSampleStages, "PassLabel count must be equal to DownSampleStages.");
 		FRenderingCompositePass* Pass = Context.Graph.RegisterPass(new(FMemStack::Get()) FRCPassPostProcessDownsample(PF_Unknown, 1, PassLabels[i]));
 		Pass->SetInput(ePId_Input0, PostProcessDownsamples[i - 1]);
@@ -479,6 +479,7 @@ static FRenderingCompositeOutputRef AddBloom(FPostprocessContext Context, FRende
 
 		FBloomStage BloomStages[] =
 		{
+			{ Settings.Bloom6Size, &Settings.Bloom6Tint},
 			{ Settings.Bloom5Size, &Settings.Bloom5Tint},
 			{ Settings.Bloom4Size, &Settings.Bloom4Tint},
 			{ Settings.Bloom3Size, &Settings.Bloom3Tint},
