[INTEGRATE] Change 2413891 by [REDACTED_PERSON][REDACTED_HANDLE] on 2015/01/21 11:44:13

	Add new NativeParentClass tag to BP assets to indicate first native class in parent hierarchy

diff --git a/Engine/Source/Runtime/Engine/Private/Blueprint.cpp b/Engine/Source/Runtime/Engine/Private/Blueprint.cpp
--- a/Engine/Source/Runtime/Engine/Private/Blueprint.cpp
+++ b/Engine/Source/Runtime/Engine/Private/Blueprint.cpp
@@ -655,13 +655,23 @@ void UBlueprint::GetAssetRegistryTags(TArray<FAssetRegistryTag>& OutTags) const
 	Super::GetAssetRegistryTags(OutTags);
 
 	FString ParentClassPackageName;
+	FString NativeParentClassName;
 	if ( ParentClass )
 	{
 		ParentClassPackageName = ParentClass->GetOutermost()->GetName();
+
+		// Walk up until we find a native class (ie 'while they are BP classes')
+		UClass* NativeParentClass = ParentClass;
+		while (Cast<UBlueprintGeneratedClass>(NativeParentClass) != nullptr) // can't use IsA on UClass
+		{
+			NativeParentClass = NativeParentClass->GetSuperClass();
+		}
+		NativeParentClassName = FString::Printf(TEXT("%s'%s'"), *UClass::StaticClass()->GetName(), *NativeParentClass->GetPathName());
 	}
 	else
 	{
 		ParentClassPackageName = TEXT("None");
+		NativeParentClassName = TEXT("None");
 	}
 
 	//NumReplicatedProperties
@@ -674,6 +684,7 @@ void UBlueprint::GetAssetRegistryTags(TArray<FAssetRegistryTag>& OutTags) const
 
 	OutTags.Add(FAssetRegistryTag("NumReplicatedProperties", FString::FromInt(NumReplicatedProperties), FAssetRegistryTag::TT_Numerical));
 	OutTags.Add(FAssetRegistryTag("ParentClassPackage", ParentClassPackageName, FAssetRegistryTag::TT_Hidden));
+	OutTags.Add(FAssetRegistryTag("NativeParentClass", NativeParentClassName, FAssetRegistryTag::TT_Alphabetical));
 	OutTags.Add(FAssetRegistryTag(GET_MEMBER_NAME_CHECKED(UBlueprint, BlueprintDescription), BlueprintDescription, FAssetRegistryTag::TT_Hidden));
 
 	uint32 ClassFlagsTagged = 0;
