[INTEGRATE] Change 2413933 by [REDACTED_PERSON][REDACTED_HANDLE] on 2015/01/21 12:12:34

	Added null check in UAnimSequence::GetAssetRegistryTags

	UE-7941 - Critical: Editor crashes after initial asset registry scan in 4.7 branch

diff --git a/Engine/Source/Runtime/Engine/Private/Animation/AnimSequence.cpp b/Engine/Source/Runtime/Engine/Private/Animation/AnimSequence.cpp
--- a/Engine/Source/Runtime/Engine/Private/Animation/AnimSequence.cpp
+++ b/Engine/Source/Runtime/Engine/Private/Animation/AnimSequence.cpp
@@ -125,7 +125,10 @@ SIZE_T UAnimSequence::GetResourceSize(EResourceSizeMode::Type Mode)
 void UAnimSequence::GetAssetRegistryTags(TArray<FAssetRegistryTag>& OutTags) const
 {
 #if WITH_EDITORONLY_DATA
-	OutTags.Add( FAssetRegistryTag(SourceFileTagName(), AssetImportData->SourceFilePath, FAssetRegistryTag::TT_Hidden) );
+	if (AssetImportData)
+	{
+		OutTags.Add( FAssetRegistryTag(SourceFileTagName(), AssetImportData->SourceFilePath, FAssetRegistryTag::TT_Hidden) );
+	}
 #endif
 
 	Super::GetAssetRegistryTags(OutTags);
