[INTEGRATE] Change 2415420 by [REDACTED_PERSON][REDACTED_HANDLE] on 2015/01/22 11:53:25

	Restructure Component Instance Data

diff --git a/Engine/Source/Runtime/Engine/Classes/Atmosphere/AtmosphericFogComponent.h b/Engine/Source/Runtime/Engine/Classes/Atmosphere/AtmosphericFogComponent.h
--- a/Engine/Source/Runtime/Engine/Classes/Atmosphere/AtmosphericFogComponent.h
+++ b/Engine/Source/Runtime/Engine/Classes/Atmosphere/AtmosphericFogComponent.h
@@ -242,7 +242,7 @@ public:
 	mutable FByteBulkData IrradianceData;
 	mutable FByteBulkData InscatterData;
 	
-	// Begin UObject interface.
+	// Begin UObject interface. 
 	virtual void PostLoad() override;
 	virtual void BeginDestroy() override;
 
@@ -260,9 +260,10 @@ public:
 	// Begin UActorComponent interface.
 	virtual FComponentInstanceDataBase* GetComponentInstanceData() const override;
 	virtual FName GetComponentInstanceDataType() const override;
-	virtual void ApplyComponentInstanceData(FComponentInstanceDataBase* ComponentInstanceData) override;
 	// End UActorComponent interface.
 
+	void ApplyComponentInstanceData(class FAtmospherePrecomputeInstanceData* ComponentInstanceData);
+
 	const FAtmospherePrecomputeParameters& GetPrecomputeParameters() const { return PrecomputeParams;  }
 
 private:
diff --git a/Engine/Source/Runtime/Engine/Classes/Components/ActorComponent.h b/Engine/Source/Runtime/Engine/Classes/Components/ActorComponent.h
--- a/Engine/Source/Runtime/Engine/Classes/Components/ActorComponent.h
+++ b/Engine/Source/Runtime/Engine/Classes/Components/ActorComponent.h
@@ -449,9 +449,6 @@ public:
 	/** The type of the component instance data that this component is interested in */
 	virtual FName GetComponentInstanceDataType() const { return NAME_None; }
 
-	/** Called after we create new components during RerunConstructionScripts, to optionally apply any data backed up during GetComponentInstanceData */
-	virtual void ApplyComponentInstanceData(class FComponentInstanceDataBase* ComponentInstanceData ) {}
-
 	// Begin UObject interface.
 	virtual void BeginDestroy() override;
 	virtual bool NeedsLoadForClient() const override;
diff --git a/Engine/Source/Runtime/Engine/Classes/Components/ChildActorComponent.h b/Engine/Source/Runtime/Engine/Classes/Components/ChildActorComponent.h
--- a/Engine/Source/Runtime/Engine/Classes/Components/ChildActorComponent.h
+++ b/Engine/Source/Runtime/Engine/Classes/Components/ChildActorComponent.h
@@ -33,10 +33,11 @@ class UChildActorComponent : public USceneComponent
 	virtual void OnRegister() override;
 	virtual FComponentInstanceDataBase* GetComponentInstanceData() const override;
 	virtual FName GetComponentInstanceDataType() const override;
-	virtual void ApplyComponentInstanceData(FComponentInstanceDataBase* ComponentInstanceData) override;
-
 	// End ActorComponent interface.
 
+	/** Apply the component instance data to the child actor component */
+	void ApplyComponentInstanceData(class FChildActorComponentInstanceData* ComponentInstanceData);
+
 	/** Create the child actor */
 	void CreateChildActor();
 
diff --git a/Engine/Source/Runtime/Engine/Classes/Components/InstancedStaticMeshComponent.h b/Engine/Source/Runtime/Engine/Classes/Components/InstancedStaticMeshComponent.h
--- a/Engine/Source/Runtime/Engine/Classes/Components/InstancedStaticMeshComponent.h
+++ b/Engine/Source/Runtime/Engine/Classes/Components/InstancedStaticMeshComponent.h
@@ -152,7 +152,6 @@ public:
 	// Begin UActorComponent interface 
 	virtual FComponentInstanceDataBase* GetComponentInstanceData() const override;
 	virtual FName GetComponentInstanceDataType() const override;
-	virtual void ApplyComponentInstanceData(FComponentInstanceDataBase* ComponentInstanceData) override;
 	// End UActorComponent interface 
 
 	// Begin UPrimitiveComponent Interface
@@ -185,6 +184,9 @@ public:
 #endif
 	//End UObject Interface
 
+	/** Applies the cached component instance data to a newly blueprint constructed component */
+	void ApplyComponentInstanceData(class FInstancedStaticMeshComponentInstanceData* ComponentInstanceData);
+
 	/** Check to see if an instance is selected */
 	bool IsInstanceSelected(int32 InInstanceIndex) const;
 
diff --git a/Engine/Source/Runtime/Engine/Classes/Components/LightComponent.h b/Engine/Source/Runtime/Engine/Classes/Components/LightComponent.h
--- a/Engine/Source/Runtime/Engine/Classes/Components/LightComponent.h
+++ b/Engine/Source/Runtime/Engine/Classes/Components/LightComponent.h
@@ -279,7 +279,7 @@ public:
 
 	virtual FComponentInstanceDataBase* GetComponentInstanceData() const override;
 	virtual FName GetComponentInstanceDataType() const override;
-	virtual void ApplyComponentInstanceData(FComponentInstanceDataBase* ComponentInstanceData) override;
+	 void ApplyComponentInstanceData(class FPrecomputedLightInstanceData* ComponentInstanceData);
 
 	/** @return number of material elements in this primitive */
 	virtual int32 GetNumMaterials() const;
diff --git a/Engine/Source/Runtime/Engine/Classes/Components/SceneComponent.h b/Engine/Source/Runtime/Engine/Classes/Components/SceneComponent.h
--- a/Engine/Source/Runtime/Engine/Classes/Components/SceneComponent.h
+++ b/Engine/Source/Runtime/Engine/Classes/Components/SceneComponent.h
@@ -503,7 +503,6 @@ public:
 	virtual void ApplyWorldOffset(const FVector& InOffset, bool bWorldShift) override;
 	virtual class FComponentInstanceDataBase* GetComponentInstanceData() const override;
 	virtual FName GetComponentInstanceDataType() const override;
-	virtual void ApplyComponentInstanceData(class FComponentInstanceDataBase* ComponentInstanceData ) override;
 	// End ActorComponent interface
 
 	// Begin UObject Interface
@@ -800,6 +799,8 @@ public:
 	virtual ~FSceneComponentInstanceData()
 	{}
 
+	virtual void ApplyToComponent(UActorComponent* Component) override;
+
 	TArray<USceneComponent*> AttachedInstanceComponents;
 };
 
diff --git a/Engine/Source/Runtime/Engine/Classes/Components/SkyLightComponent.h b/Engine/Source/Runtime/Engine/Classes/Components/SkyLightComponent.h
--- a/Engine/Source/Runtime/Engine/Classes/Components/SkyLightComponent.h
+++ b/Engine/Source/Runtime/Engine/Classes/Components/SkyLightComponent.h
@@ -150,7 +150,7 @@ class ENGINE_API USkyLightComponent : public ULightComponentBase
 
 	virtual FComponentInstanceDataBase* GetComponentInstanceData() const override;
 	virtual FName GetComponentInstanceDataType() const override;
-	virtual void ApplyComponentInstanceData(FComponentInstanceDataBase* ComponentInstanceData) override;
+	void ApplyComponentInstanceData(class FPrecomputedSkyLightInstanceData* ComponentInstanceData);
 
 	/** Called each tick to recapture and queued sky captures. */
 	static void UpdateSkyCaptureContents(UWorld* WorldToUpdate);
diff --git a/Engine/Source/Runtime/Engine/Classes/Components/SplineComponent.h b/Engine/Source/Runtime/Engine/Classes/Components/SplineComponent.h
--- a/Engine/Source/Runtime/Engine/Classes/Components/SplineComponent.h
+++ b/Engine/Source/Runtime/Engine/Classes/Components/SplineComponent.h
@@ -75,9 +75,10 @@ public:
 	// Begin UActorComponent interface.
 	virtual FComponentInstanceDataBase* GetComponentInstanceData() const override;
 	virtual FName GetComponentInstanceDataType() const override;
-	virtual void ApplyComponentInstanceData(FComponentInstanceDataBase* ComponentInstanceData) override;
 	// End UActorComponent interface.
 
+	void ApplyComponentInstanceData(class FSplineInstanceData* ComponentInstanceData);
+
 	/** Update the spline tangents and SplineReparamTable */
 	void UpdateSpline();
 
diff --git a/Engine/Source/Runtime/Engine/Classes/Components/StaticMeshComponent.h b/Engine/Source/Runtime/Engine/Classes/Components/StaticMeshComponent.h
--- a/Engine/Source/Runtime/Engine/Classes/Components/StaticMeshComponent.h
+++ b/Engine/Source/Runtime/Engine/Classes/Components/StaticMeshComponent.h
@@ -258,11 +258,8 @@ public:
 #endif
 	virtual FComponentInstanceDataBase* GetComponentInstanceData() const override;
 	virtual FName GetComponentInstanceDataType() const override;
-	virtual void ApplyComponentInstanceData(FComponentInstanceDataBase* ComponentInstanceData) override;
 	// End UActorComponent interface.
 
-
-
 	// Begin UPrimitiveComponent interface.
 	virtual int32 GetNumMaterials() const override;
 #if WITH_EDITOR
@@ -450,6 +447,8 @@ public:
 
 	/** Get this components index in its parents blueprint created components array (used for matching instance data) */
 	int32 GetBlueprintCreatedComponentIndex() const;
+
+	void ApplyComponentInstanceData(class FStaticMeshComponentInstanceData* ComponentInstanceData);
 };
 
 
diff --git a/Engine/Source/Runtime/Engine/Private/Atmosphere/Atmosphere.cpp b/Engine/Source/Runtime/Engine/Private/Atmosphere/Atmosphere.cpp
--- a/Engine/Source/Runtime/Engine/Private/Atmosphere/Atmosphere.cpp
+++ b/Engine/Source/Runtime/Engine/Private/Atmosphere/Atmosphere.cpp
@@ -705,6 +705,12 @@ public:
 	virtual ~FAtmospherePrecomputeInstanceData()
 	{}
 
+	virtual void ApplyToComponent(UActorComponent* Component) override
+	{
+		FSceneComponentInstanceData::ApplyToComponent(Component);
+		CastChecked<UAtmosphericFogComponent>(Component)->ApplyComponentInstanceData(this);
+	}
+
 	struct FAtmospherePrecomputeParameters PrecomputeParameter;
 
 	FByteBulkData TransmittanceData;
@@ -721,12 +727,13 @@ FName UAtmosphericFogComponent::GetComponentInstanceDataType() const
 // Backup the precomputed data before re-running Blueprint construction script
 FComponentInstanceDataBase* UAtmosphericFogComponent::GetComponentInstanceData() const
 {
-	FAtmospherePrecomputeInstanceData* PrecomputedData = nullptr;
+	FComponentInstanceDataBase* InstanceData = nullptr;
 
 	if (TransmittanceData.GetElementCount() && IrradianceData.GetElementCount() && InscatterData.GetElementCount() && PrecomputeCounter.GetValue() == EValid)
 	{
 		// Allocate new struct for holding light map data
-		 PrecomputedData = new FAtmospherePrecomputeInstanceData(this);
+		 FAtmospherePrecomputeInstanceData* PrecomputedData = new FAtmospherePrecomputeInstanceData(this);
+		 InstanceData = PrecomputedData;
 
 		// Fill in info
 		PrecomputedData->PrecomputeParameter = PrecomputeParams;
@@ -754,17 +761,18 @@ FComponentInstanceDataBase* UAtmosphericFogComponent::GetComponentInstanceData()
 			PrecomputedData->InscatterData.Unlock();
 		}
 	}
+	else
+	{
+		InstanceData = Super::GetComponentInstanceData();
+	}
 
-	return PrecomputedData;
+	return InstanceData;
 }
 
 // Restore the precomputed data after re-running Blueprint construction script
-void UAtmosphericFogComponent::ApplyComponentInstanceData(FComponentInstanceDataBase* ComponentInstanceData)
+void UAtmosphericFogComponent::ApplyComponentInstanceData(FAtmospherePrecomputeInstanceData* PrecomputedData)
 {
-	Super::ApplyComponentInstanceData(ComponentInstanceData);
-
-	check(ComponentInstanceData);
-	FAtmospherePrecomputeInstanceData* PrecomputedData = static_cast<FAtmospherePrecomputeInstanceData*>(const_cast<FComponentInstanceDataBase*>(ComponentInstanceData));
+	check(PrecomputedData);
 
 	if (PrecomputedData->PrecomputeParameter != GetPrecomputeParameters())
 	{
diff --git a/Engine/Source/Runtime/Engine/Private/ChildActorComponent.cpp b/Engine/Source/Runtime/Engine/Private/ChildActorComponent.cpp
--- a/Engine/Source/Runtime/Engine/Private/ChildActorComponent.cpp
+++ b/Engine/Source/Runtime/Engine/Private/ChildActorComponent.cpp
@@ -71,6 +71,12 @@ public:
 		}
 	}
 
+	virtual void ApplyToComponent(UActorComponent* Component) override
+	{
+		FSceneComponentInstanceData::ApplyToComponent(Component);
+		CastChecked<UChildActorComponent>(Component)->ApplyComponentInstanceData(this);
+	}
+
 	FName ChildActorName;
 
 	struct FAttachedActorInfo
@@ -106,12 +112,10 @@ FComponentInstanceDataBase* UChildActorComponent::GetComponentInstanceData() con
 	return InstanceData;
 }
 
-void UChildActorComponent::ApplyComponentInstanceData(FComponentInstanceDataBase* ComponentInstanceData)
+void UChildActorComponent::ApplyComponentInstanceData(FChildActorComponentInstanceData* ChildActorInstanceData)
 {
-	Super::ApplyComponentInstanceData(ComponentInstanceData);
+	check(ChildActorInstanceData);
 
-	check(ComponentInstanceData);
-	FChildActorComponentInstanceData* ChildActorInstanceData  = static_cast<FChildActorComponentInstanceData*>(ComponentInstanceData);
 	ChildActorName = ChildActorInstanceData->ChildActorName;
 	if (ChildActor)
 	{
diff --git a/Engine/Source/Runtime/Engine/Private/ComponentInstanceDataCache.cpp b/Engine/Source/Runtime/Engine/Private/ComponentInstanceDataCache.cpp
--- a/Engine/Source/Runtime/Engine/Private/ComponentInstanceDataCache.cpp
+++ b/Engine/Source/Runtime/Engine/Private/ComponentInstanceDataCache.cpp
@@ -121,7 +121,7 @@ void FComponentInstanceDataCache::ApplyToActor(AActor* Actor) const
 					{
 						if (ComponentInstanceData && ComponentInstanceData->MatchesComponent(Component))
 						{
-							Component->ApplyComponentInstanceData(ComponentInstanceData);
+							ComponentInstanceData->ApplyToComponent(Component);
 							break;
 						}
 					}
diff --git a/Engine/Source/Runtime/Engine/Private/InstancedStaticMesh.cpp b/Engine/Source/Runtime/Engine/Private/InstancedStaticMesh.cpp
--- a/Engine/Source/Runtime/Engine/Private/InstancedStaticMesh.cpp
+++ b/Engine/Source/Runtime/Engine/Private/InstancedStaticMesh.cpp
@@ -683,6 +683,12 @@ public:
 	{
 	}
 
+	virtual void ApplyToComponent(UActorComponent* Component) override
+	{
+		FSceneComponentInstanceData::ApplyToComponent(Component);
+		CastChecked<UInstancedStaticMeshComponent>(Component)->ApplyComponentInstanceData(this);
+	}
+
 	/** Used to store lightmap data during RerunConstructionScripts */
 	struct FLightMapInstanceData
 	{
@@ -719,36 +725,37 @@ FName UInstancedStaticMeshComponent::GetComponentInstanceDataType() const
 FComponentInstanceDataBase* UInstancedStaticMeshComponent::GetComponentInstanceData() const
 {
 #if WITH_EDITOR
-	FInstancedStaticMeshComponentInstanceData* InstanceData = nullptr;
+	FComponentInstanceDataBase* InstanceData = nullptr;
+	FInstancedStaticMeshComponentInstanceData* StaticMeshInstanceData = nullptr;
 
 	// Don't back up static lighting if there isn't any
 	if (bHasCachedStaticLighting || SelectedInstances.Num() > 0)
 	{
-		InstanceData = new FInstancedStaticMeshComponentInstanceData(*this);
+		InstanceData = StaticMeshInstanceData = new FInstancedStaticMeshComponentInstanceData(*this);	
 	}
 
 	// Don't back up static lighting if there isn't any
 	if (bHasCachedStaticLighting)
 	{
 		// Fill in info (copied from UStaticMeshComponent::GetComponentInstanceData)
-		InstanceData->bHasCachedStaticLighting = true;
-		InstanceData->CachedStaticLighting.Transform = ComponentToWorld;
-		InstanceData->CachedStaticLighting.IrrelevantLights = IrrelevantLights;
-		InstanceData->CachedStaticLighting.LODDataLightMap.Empty(LODData.Num());
+		StaticMeshInstanceData->bHasCachedStaticLighting = true;
+		StaticMeshInstanceData->CachedStaticLighting.Transform = ComponentToWorld;
+		StaticMeshInstanceData->CachedStaticLighting.IrrelevantLights = IrrelevantLights;
+		StaticMeshInstanceData->CachedStaticLighting.LODDataLightMap.Empty(LODData.Num());
 		for (const FStaticMeshComponentLODInfo& LODDataEntry : LODData)
 		{
-			InstanceData->CachedStaticLighting.LODDataLightMap.Add(LODDataEntry.LightMap);
-			InstanceData->CachedStaticLighting.LODDataShadowMap.Add(LODDataEntry.ShadowMap);
+			StaticMeshInstanceData->CachedStaticLighting.LODDataLightMap.Add(LODDataEntry.LightMap);
+			StaticMeshInstanceData->CachedStaticLighting.LODDataShadowMap.Add(LODDataEntry.ShadowMap);
 		}
 
 		// Back up per-instance lightmap/shadowmap info
-		InstanceData->PerInstanceSMData = PerInstanceSMData;
+		StaticMeshInstanceData->PerInstanceSMData = PerInstanceSMData;
 	}
 
 	// Back up instance selection
 	if (SelectedInstances.Num() > 0)
 	{
-		InstanceData->SelectedInstances = SelectedInstances;
+		StaticMeshInstanceData->SelectedInstances = SelectedInstances;
 	}
 
 	return InstanceData;
@@ -757,15 +764,10 @@ FComponentInstanceDataBase* UInstancedStaticMeshComponent::GetComponentInstanceD
 #endif
 }
 
-void UInstancedStaticMeshComponent::ApplyComponentInstanceData(FComponentInstanceDataBase* ComponentInstanceData)
+void UInstancedStaticMeshComponent::ApplyComponentInstanceData(FInstancedStaticMeshComponentInstanceData* InstancedMeshData)
 {
-	// Skip UStaticMeshComponent implementation
-	USceneComponent::ApplyComponentInstanceData(ComponentInstanceData);
-
 #if WITH_EDITOR
-	check(ComponentInstanceData);
-
-	FInstancedStaticMeshComponentInstanceData* InstancedMeshData  = static_cast<FInstancedStaticMeshComponentInstanceData*>(ComponentInstanceData);
+	check(InstancedMeshData);
 
 	if (StaticMesh != InstancedMeshData->StaticMesh)
 	{
diff --git a/Engine/Source/Runtime/Engine/Private/LightComponent.cpp b/Engine/Source/Runtime/Engine/Private/LightComponent.cpp
--- a/Engine/Source/Runtime/Engine/Private/LightComponent.cpp
+++ b/Engine/Source/Runtime/Engine/Private/LightComponent.cpp
@@ -883,9 +883,12 @@ public:
 		, PreviewShadowMapChannel(SourceComponent->PreviewShadowMapChannel)
 		, bPrecomputedLightingIsValid(SourceComponent->bPrecomputedLightingIsValid)
 	{}
-			
-	virtual ~FPrecomputedLightInstanceData()
-	{}
+
+	virtual void ApplyToComponent(UActorComponent* Component) override
+	{
+		FSceneComponentInstanceData::ApplyToComponent(Component);
+		CastChecked<ULightComponent>(Component)->ApplyComponentInstanceData(this);
+	}
 
 	FTransform Transform;
 	FGuid LightGuid;
@@ -906,12 +909,9 @@ FComponentInstanceDataBase* ULightComponent::GetComponentInstanceData() const
 	return new FPrecomputedLightInstanceData(this);
 }
 
-void ULightComponent::ApplyComponentInstanceData(FComponentInstanceDataBase* ComponentInstanceData)
+void ULightComponent::ApplyComponentInstanceData(FPrecomputedLightInstanceData* LightMapData)
 {
-	Super::ApplyComponentInstanceData(ComponentInstanceData);
-
-	check(ComponentInstanceData);
-	FPrecomputedLightInstanceData* LightMapData  = static_cast<FPrecomputedLightInstanceData*>(ComponentInstanceData);
+	check(LightMapData);
 
 	if (!LightMapData->Transform.Equals(ComponentToWorld))
 	{
diff --git a/Engine/Source/Runtime/Engine/Private/SceneComponent.cpp b/Engine/Source/Runtime/Engine/Private/SceneComponent.cpp
--- a/Engine/Source/Runtime/Engine/Private/SceneComponent.cpp
+++ b/Engine/Source/Runtime/Engine/Private/SceneComponent.cpp
@@ -1010,15 +1010,26 @@ bool USceneComponent::IsAttachedTo(class USceneComponent* TestComp) const
 FSceneComponentInstanceData::FSceneComponentInstanceData(const USceneComponent* SourceComponent)
 	: FComponentInstanceDataBase(SourceComponent)
 {
-	for (USceneComponent* SceneComponent : SourceComponent->AttachChildren)
+	for (int32 i = SourceComponent->AttachChildren.Num()-1; i >= 0; --i)
 	{
+		USceneComponent* SceneComponent = SourceComponent->AttachChildren[i];
 		if (SceneComponent && !SceneComponent->bCreatedByConstructionScript)
 		{
 			AttachedInstanceComponents.Add(SceneComponent);
+			SceneComponent->DetachFromParent(true);
 		}
 	}
 }
 
+void FSceneComponentInstanceData::ApplyToComponent(UActorComponent* Component)
+{
+	USceneComponent* SceneComponent = CastChecked<USceneComponent>(Component);
+	for (USceneComponent* ChildComponent : AttachedInstanceComponents)
+	{
+		ChildComponent->AttachTo(SceneComponent);
+	}
+}
+
 FComponentInstanceDataBase* USceneComponent::GetComponentInstanceData() const
 {
 	FComponentInstanceDataBase* InstanceData = nullptr;
@@ -1041,17 +1052,6 @@ FName USceneComponent::GetComponentInstanceDataType() const
 	return SceneComponentInstanceDataTypeName;
 }
 
-void USceneComponent::ApplyComponentInstanceData(class FComponentInstanceDataBase* ComponentInstanceData )
-{
-	check(ComponentInstanceData);
-	FSceneComponentInstanceData* SceneComponentInstanceData  = static_cast<FSceneComponentInstanceData*>(ComponentInstanceData);
-
-	for (USceneComponent* ChildComponent : SceneComponentInstanceData->AttachedInstanceComponents)
-	{
-		ChildComponent->AttachTo(this);
-	}
-}
-
 void USceneComponent::UpdateChildTransforms()
 {
 	for(int32 i=0; i<AttachChildren.Num(); i++)
diff --git a/Engine/Source/Runtime/Engine/Private/SkyLightComponent.cpp b/Engine/Source/Runtime/Engine/Private/SkyLightComponent.cpp
--- a/Engine/Source/Runtime/Engine/Private/SkyLightComponent.cpp
+++ b/Engine/Source/Runtime/Engine/Private/SkyLightComponent.cpp
@@ -332,6 +332,12 @@ public:
 		: FSceneComponentInstanceData(SourceComponent)
 	{}
 
+	virtual void ApplyToComponent(UActorComponent* Component) override
+	{
+		FSceneComponentInstanceData::ApplyToComponent(Component);
+		CastChecked<USkyLightComponent>(Component)->ApplyComponentInstanceData(this);
+	}
+
 	FGuid LightGuid;
 	bool bPrecomputedLightingIsValid;
 	// This has to be refcounted to keep it alive during the handoff without doing a deep copy
@@ -357,13 +363,9 @@ FComponentInstanceDataBase* USkyLightComponent::GetComponentInstanceData() const
 	return InstanceData;
 }
 
-void USkyLightComponent::ApplyComponentInstanceData(FComponentInstanceDataBase* ComponentInstanceData)
+void USkyLightComponent::ApplyComponentInstanceData(FPrecomputedSkyLightInstanceData* LightMapData)
 {
-	//Skip ULightComponent implementationn
-	USceneComponent::ApplyComponentInstanceData(ComponentInstanceData);
-
-	check(ComponentInstanceData);
-	FPrecomputedSkyLightInstanceData* LightMapData  = static_cast<FPrecomputedSkyLightInstanceData*>(ComponentInstanceData);
+	check(LightMapData);
 
 	LightGuid = LightMapData->LightGuid;
 	bPrecomputedLightingIsValid = LightMapData->bPrecomputedLightingIsValid;
diff --git a/Engine/Source/Runtime/Engine/Private/SplineComponent.cpp b/Engine/Source/Runtime/Engine/Private/SplineComponent.cpp
--- a/Engine/Source/Runtime/Engine/Private/SplineComponent.cpp
+++ b/Engine/Source/Runtime/Engine/Private/SplineComponent.cpp
@@ -594,6 +594,12 @@ public:
 	{
 	}
 
+	virtual void ApplyToComponent(UActorComponent* Component) override
+	{
+		FSceneComponentInstanceData::ApplyToComponent(Component);
+		CastChecked<USplineComponent>(Component)->ApplyComponentInstanceData(this);
+	}
+
 	FInterpCurveVector SplineInfo;
 	bool bClosedLoop;
 };
@@ -606,23 +612,26 @@ FName USplineComponent::GetComponentInstanceDataType() const
 
 FComponentInstanceDataBase* USplineComponent::GetComponentInstanceData() const
 {
-	FSplineInstanceData* SplineInstanceData = nullptr;
+	FComponentInstanceDataBase* InstanceData = nullptr;
 	if (bAllowSplineEditingPerInstance)
 	{
-		SplineInstanceData = new FSplineInstanceData(this);
+		FSplineInstanceData* SplineInstanceData = new FSplineInstanceData(this);
 		SplineInstanceData->SplineInfo = SplineInfo;
 		SplineInstanceData->bClosedLoop = bClosedLoop;
+
+		InstanceData = SplineInstanceData;
 	}
-	return SplineInstanceData;
+	else
+	{
+		InstanceData = Super::GetComponentInstanceData();
+	}
+	return InstanceData;
 }
 
-void USplineComponent::ApplyComponentInstanceData(FComponentInstanceDataBase* ComponentInstanceData)
+void USplineComponent::ApplyComponentInstanceData(FSplineInstanceData* SplineInstanceData)
 {
-	Super::ApplyComponentInstanceData(ComponentInstanceData);
-
-	if (ComponentInstanceData)
+	if (SplineInstanceData)
 	{
-		FSplineInstanceData* SplineInstanceData  = static_cast<FSplineInstanceData*>(ComponentInstanceData);
 		if (bAllowSplineEditingPerInstance)
 		{
 			SplineInfo = SplineInstanceData->SplineInfo;
diff --git a/Engine/Source/Runtime/Engine/Private/StaticMeshComponent.cpp b/Engine/Source/Runtime/Engine/Private/StaticMeshComponent.cpp
--- a/Engine/Source/Runtime/Engine/Private/StaticMeshComponent.cpp
+++ b/Engine/Source/Runtime/Engine/Private/StaticMeshComponent.cpp
@@ -33,6 +33,12 @@ public:
 	{
 	}
 
+	virtual void ApplyToComponent(UActorComponent* Component) override
+	{
+		FSceneComponentInstanceData::ApplyToComponent(Component);
+		CastChecked<UStaticMeshComponent>(Component)->ApplyComponentInstanceData(this);
+	}
+
 	/** Add vertex color data for a specified LOD before RerunConstructionScripts is called */
 	void AddVertexColorData(const struct FStaticMeshComponentLODInfo& LODInfo, uint32 LODIndex)
 	{
@@ -1528,22 +1534,22 @@ FName UStaticMeshComponent::GetComponentInstanceDataType() const
 
 FComponentInstanceDataBase* UStaticMeshComponent::GetComponentInstanceData() const
 {
-	FStaticMeshComponentInstanceData* InstanceData = nullptr;
+	FStaticMeshComponentInstanceData* StaticMeshInstanceData = nullptr;
 
 	// Don't back up static lighting if there isn't any
 	if(bHasCachedStaticLighting)
 	{
-		InstanceData = new FStaticMeshComponentInstanceData(this);
+		StaticMeshInstanceData = new FStaticMeshComponentInstanceData(this);
 
 		// Fill in info
-		InstanceData->bHasCachedStaticLighting = true;
-		InstanceData->CachedStaticLighting.Transform = ComponentToWorld;
-		InstanceData->CachedStaticLighting.IrrelevantLights = IrrelevantLights;
-		InstanceData->CachedStaticLighting.LODDataLightMap.Empty(LODData.Num());
+		StaticMeshInstanceData->bHasCachedStaticLighting = true;
+		StaticMeshInstanceData->CachedStaticLighting.Transform = ComponentToWorld;
+		StaticMeshInstanceData->CachedStaticLighting.IrrelevantLights = IrrelevantLights;
+		StaticMeshInstanceData->CachedStaticLighting.LODDataLightMap.Empty(LODData.Num());
 		for (const FStaticMeshComponentLODInfo& LODDataEntry : LODData)
 		{
-			InstanceData->CachedStaticLighting.LODDataLightMap.Add(LODDataEntry.LightMap);
-			InstanceData->CachedStaticLighting.LODDataShadowMap.Add(LODDataEntry.ShadowMap);
+			StaticMeshInstanceData->CachedStaticLighting.LODDataLightMap.Add(LODDataEntry.LightMap);
+			StaticMeshInstanceData->CachedStaticLighting.LODDataShadowMap.Add(LODDataEntry.ShadowMap);
 		}
 	}
 
@@ -1555,27 +1561,24 @@ FComponentInstanceDataBase* UStaticMeshComponent::GetComponentInstanceData() con
 
 		if ( LODInfo.OverrideVertexColors && LODInfo.OverrideVertexColors->GetNumVertices() > 0 && LODInfo.PaintedVertices.Num() > 0 )
 		{
-			if (!InstanceData)
+			if (!StaticMeshInstanceData)
 			{
-				InstanceData = new FStaticMeshComponentInstanceData(this);
+				StaticMeshInstanceData = new FStaticMeshComponentInstanceData(this);
 			}
 
-			InstanceData->AddVertexColorData(LODInfo, LODIndex);
+			StaticMeshInstanceData->AddVertexColorData(LODInfo, LODIndex);
 		}
 	}
 
-	return InstanceData;
+	return (StaticMeshInstanceData ? StaticMeshInstanceData : Super::GetComponentInstanceData());
 }
 
-void UStaticMeshComponent::ApplyComponentInstanceData(FComponentInstanceDataBase* ComponentInstanceData)
+void UStaticMeshComponent::ApplyComponentInstanceData(FStaticMeshComponentInstanceData* StaticMeshInstanceData)
 {
-	Super::ApplyComponentInstanceData(ComponentInstanceData);
-
-	check(ComponentInstanceData);
+	check(StaticMeshInstanceData);
 
 	// Note: ApplyComponentInstanceData is called while the component is registered so the rendering thread is already using this component
 	// That means all component state that is modified here must be mirrored on the scene proxy, which will be recreated to receive the changes later due to MarkRenderStateDirty.
-	FStaticMeshComponentInstanceData* StaticMeshInstanceData  = static_cast<FStaticMeshComponentInstanceData*>(ComponentInstanceData);
 
 	if (StaticMesh != StaticMeshInstanceData->StaticMesh)
 	{
diff --git a/Engine/Source/Runtime/Engine/Public/ComponentInstanceDataCache.h b/Engine/Source/Runtime/Engine/Public/ComponentInstanceDataCache.h
--- a/Engine/Source/Runtime/Engine/Public/ComponentInstanceDataCache.h
+++ b/Engine/Source/Runtime/Engine/Public/ComponentInstanceDataCache.h
@@ -20,6 +20,8 @@ public:
 
 	bool MatchesComponent(const UActorComponent* Component) const;
 
+	virtual void ApplyToComponent(UActorComponent* Component) = 0;
+
 protected:
 	/** The name of the source component */
 	FName SourceComponentName;
diff --git a/Engine/Source/Runtime/UMG/Private/Components/WidgetComponent.cpp b/Engine/Source/Runtime/UMG/Private/Components/WidgetComponent.cpp
--- a/Engine/Source/Runtime/UMG/Private/Components/WidgetComponent.cpp
+++ b/Engine/Source/Runtime/UMG/Private/Components/WidgetComponent.cpp
@@ -577,6 +577,12 @@ public:
 		, RenderTarget( SourceComponent->GetRenderTarget() )
 	{}
 
+	virtual void ApplyToComponent(UActorComponent* Component) override
+	{
+		FSceneComponentInstanceData::ApplyToComponent(Component);
+		CastChecked<UWidgetComponent>(Component)->ApplyComponentInstanceData(this);
+	}
+
 public:
 	TSubclassOf<UUserWidget> WidgetClass;
 	UTextureRenderTarget2D* RenderTarget;
@@ -593,15 +599,12 @@ FComponentInstanceDataBase* UWidgetComponent::GetComponentInstanceData() const
 	return new FWidgetComponentInstanceData( this );
 }
 
-void UWidgetComponent::ApplyComponentInstanceData(FComponentInstanceDataBase* ComponentInstanceData)
+void UWidgetComponent::ApplyComponentInstanceData(FWidgetComponentInstanceData* WidgetInstanceData)
 {
-	Super::ApplyComponentInstanceData(ComponentInstanceData);
-
-	check(ComponentInstanceData);
+	check(WidgetInstanceData);
 
 	// Note: ApplyComponentInstanceData is called while the component is registered so the rendering thread is already using this component
 	// That means all component state that is modified here must be mirrored on the scene proxy, which will be recreated to receive the changes later due to MarkRenderStateDirty.
-	FWidgetComponentInstanceData* WidgetInstanceData  = static_cast<FWidgetComponentInstanceData*>(ComponentInstanceData);
 
 	if (GetWidgetClass() != WidgetClass)
 	{
diff --git a/Engine/Source/Runtime/UMG/Public/Components/WidgetComponent.h b/Engine/Source/Runtime/UMG/Public/Components/WidgetComponent.h
--- a/Engine/Source/Runtime/UMG/Public/Components/WidgetComponent.h
+++ b/Engine/Source/Runtime/UMG/Public/Components/WidgetComponent.h
@@ -35,7 +35,8 @@ public:
 
 	virtual FComponentInstanceDataBase* GetComponentInstanceData() const override;
 	virtual FName GetComponentInstanceDataType() const override;
-	virtual void ApplyComponentInstanceData(FComponentInstanceDataBase* ComponentInstanceData) override;
+
+	void ApplyComponentInstanceData(class FWidgetComponentInstanceData* ComponentInstanceData);
 
 	// Begin UObject
 	virtual void PostLoad() override;
