[INTEGRATE] Change 2415780 by [REDACTED_PERSON][REDACTED_HANDLE] on 2015/01/22 15:51:55

	Don't detach children when collecting instance data as that can cause a detach during undo/redo

diff --git a/Engine/Source/Runtime/Engine/Private/SceneComponent.cpp b/Engine/Source/Runtime/Engine/Private/SceneComponent.cpp
--- a/Engine/Source/Runtime/Engine/Private/SceneComponent.cpp
+++ b/Engine/Source/Runtime/Engine/Private/SceneComponent.cpp
@@ -1016,7 +1016,6 @@ FSceneComponentInstanceData::FSceneComponentInstanceData(const USceneComponent*
 		if (SceneComponent && !SceneComponent->bCreatedByConstructionScript)
 		{
 			AttachedInstanceComponents.Add(SceneComponent);
-			SceneComponent->DetachFromParent(true);
 		}
 	}
 }
