[INTEGRATE] Change 2415914 by [REDACTED_PERSON][REDACTED_HANDLE] on 2015/01/22 17:18:20

	Don't use an instance component as the root when re-running simple construction script

diff --git a/Engine/Source/Runtime/Engine/Private/SimpleConstructionScript.cpp b/Engine/Source/Runtime/Engine/Private/SimpleConstructionScript.cpp
--- a/Engine/Source/Runtime/Engine/Private/SimpleConstructionScript.cpp
+++ b/Engine/Source/Runtime/Engine/Private/SimpleConstructionScript.cpp
@@ -340,6 +340,13 @@ void USimpleConstructionScript::ExecuteScriptOnActor(AActor* Actor, const FTrans
 				// Get all native scene components
 				TInlineComponentArray<USceneComponent*> Components;
 				Actor->GetComponents(Components);
+				for (int32 Index = Components.Num()-1; Index >= 0; --Index)
+				{
+					if (Components[Index]->bInstanceComponent)
+					{
+						Components.RemoveAt(Index);
+					}
+				}
 
 				// Get the native root component; if it's not set, the first native scene component will be used as root. This matches what's done in the SCS editor.
 				USceneComponent* RootComponent = Actor->GetRootComponent();
