[INTEGRATE] Change 2416010 by [REDACTED_PERSON][REDACTED_HANDLE] on 2015/01/22 17:53:58

	Fix crash during recompile after deleting a component.

diff --git a/Engine/Source/Runtime/Engine/Private/SceneComponent.cpp b/Engine/Source/Runtime/Engine/Private/SceneComponent.cpp
--- a/Engine/Source/Runtime/Engine/Private/SceneComponent.cpp
+++ b/Engine/Source/Runtime/Engine/Private/SceneComponent.cpp
@@ -1035,7 +1035,7 @@ FComponentInstanceDataBase* USceneComponent::GetComponentInstanceData() const
 
 	for (USceneComponent* Child : AttachChildren)
 	{
-		if (!Child->bCreatedByConstructionScript)
+		if (Child && !Child->bCreatedByConstructionScript)
 		{
 			InstanceData = new FSceneComponentInstanceData(this);
 			break;
