[INTEGRATE] Change 2416116 by [REDACTED_PERSON][REDACTED_HANDLE] on 2015/01/22 18:52:27

	Mark InstanceComponents as Instanced so that they are properly duplicated

diff --git a/Engine/Source/Runtime/Engine/Classes/GameFramework/Actor.h b/Engine/Source/Runtime/Engine/Classes/GameFramework/Actor.h
--- a/Engine/Source/Runtime/Engine/Classes/GameFramework/Actor.h
+++ b/Engine/Source/Runtime/Engine/Classes/GameFramework/Actor.h
@@ -2229,7 +2229,7 @@ public:
 
 private:
 	/** Array of ActorComponents that have been added by the user on a per-instance basis. */
-	UPROPERTY()
+	UPROPERTY(Instanced)
 	TArray<UActorComponent*> InstanceComponents;
 
 public:
