[INTEGRATE] Change 2416423 by [REDACTED_PERSON][REDACTED_HANDLE] on 2015/01/22 23:43:30

	Changed category for Actor tags.  Made it an advanced setting

diff --git a/Engine/Source/Runtime/Engine/Classes/GameFramework/Actor.h b/Engine/Source/Runtime/Engine/Classes/GameFramework/Actor.h
--- a/Engine/Source/Runtime/Engine/Classes/GameFramework/Actor.h
+++ b/Engine/Source/Runtime/Engine/Classes/GameFramework/Actor.h
@@ -466,7 +466,7 @@ public:
 	uint32 bEnableAutoLODGeneration:1;
 
 	/** Array of tags that can be used for grouping and categorizing. */
-	UPROPERTY(EditAnywhere, BlueprintReadWrite, Category=Tags)
+	UPROPERTY(EditAnywhere, BlueprintReadWrite, AdvancedDisplay, Category=Actor)
 	TArray<FName> Tags;
 
 	/** Bitflag to represent which views this actor is hidden in, via per-view layer visibility. */
