Naive fix for SplineMeshComponent after integrating changes from 4.7 branch.

diff --git a/Engine/Source/Runtime/Engine/Classes/Components/SplineMeshComponent.h b/Engine/Source/Runtime/Engine/Classes/Components/SplineMeshComponent.h
--- a/Engine/Source/Runtime/Engine/Classes/Components/SplineMeshComponent.h
+++ b/Engine/Source/Runtime/Engine/Classes/Components/SplineMeshComponent.h
@@ -134,9 +134,10 @@ class ENGINE_API USplineMeshComponent : public UStaticMeshComponent, public IInt
 	// Begin UActorComponent interface.
 	virtual FComponentInstanceDataBase* GetComponentInstanceData() const override;
 	virtual FName GetComponentInstanceDataType() const override;
-	virtual void ApplyComponentInstanceData(FComponentInstanceDataBase* ComponentInstanceData) override;
 	// End UActorComponent interface.
 
+	void ApplyComponentInstanceData(class FSplineMeshInstanceData* ComponentInstanceData);
+
 	//Begin USceneComponent Interface
 	virtual FPrimitiveSceneProxy* CreateSceneProxy() override;
 	virtual FBoxSphereBounds CalcBounds(const FTransform& LocalToWorld) const override;
diff --git a/Engine/Source/Runtime/Engine/Private/SplineMeshComponent.cpp b/Engine/Source/Runtime/Engine/Private/SplineMeshComponent.cpp
--- a/Engine/Source/Runtime/Engine/Private/SplineMeshComponent.cpp
+++ b/Engine/Source/Runtime/Engine/Private/SplineMeshComponent.cpp
@@ -790,14 +790,20 @@ void USplineMeshComponent::RecreateCollision()
 #endif
 
 /** Used to store spline mesh data during RerunConstructionScripts */
-class FSplineMeshInstanceData : public FComponentInstanceDataBase
+class FSplineMeshInstanceData : public FSceneComponentInstanceData
 {
 public:
 	explicit FSplineMeshInstanceData(const USplineMeshComponent* SourceComponent)
-		: FComponentInstanceDataBase(SourceComponent)
+		: FSceneComponentInstanceData(SourceComponent)
 	{
 	}
 
+	virtual void ApplyToComponent(UActorComponent* Component) override
+	{
+		FSceneComponentInstanceData::ApplyToComponent(Component);
+		CastChecked<USplineMeshComponent>(Component)->ApplyComponentInstanceData(this);
+	}
+
 	FVector StartPos;
 	FVector EndPos;
 	FVector StartTangent;
@@ -812,23 +818,27 @@ FName USplineMeshComponent::GetComponentInstanceDataType() const
 
 FComponentInstanceDataBase* USplineMeshComponent::GetComponentInstanceData() const
 {
-	FSplineMeshInstanceData* SplineMeshInstanceData = nullptr;
+	FComponentInstanceDataBase* InstanceData = nullptr;
 	if (bAllowSplineEditingPerInstance)
 	{
-		SplineMeshInstanceData = new FSplineMeshInstanceData(this);
+		FSplineMeshInstanceData *SplineMeshInstanceData = new FSplineMeshInstanceData(this);
 		SplineMeshInstanceData->StartPos = SplineParams.StartPos;
 		SplineMeshInstanceData->EndPos = SplineParams.EndPos;
 		SplineMeshInstanceData->StartTangent = SplineParams.StartTangent;
 		SplineMeshInstanceData->EndTangent = SplineParams.EndTangent;
+		InstanceData = SplineMeshInstanceData;
+	}
+	else
+	{
+		InstanceData = Super::GetComponentInstanceData();
 	}
-	return SplineMeshInstanceData;
+	return InstanceData;
 }
 
-void USplineMeshComponent::ApplyComponentInstanceData(FComponentInstanceDataBase* ComponentInstanceData)
+void USplineMeshComponent::ApplyComponentInstanceData(FSplineMeshInstanceData* SplineMeshInstanceData)
 {
-	if (ComponentInstanceData)
+	if (SplineMeshInstanceData)
 	{
-		FSplineMeshInstanceData* SplineMeshInstanceData = static_cast<FSplineMeshInstanceData*>(ComponentInstanceData);
 		if (bAllowSplineEditingPerInstance)
 		{
 			SplineParams.StartPos = SplineMeshInstanceData->StartPos;
