[INTEGRATE] Change 2417838 by [REDACTED_PERSON][REDACTED_HANDLE] on 2015/01/24 12:50:00

	Rename instance components that have the same name as newly added Blueprint created components to avoid them being found and the blueprint component not being created.

diff --git a/Engine/Source/Runtime/Engine/Private/ActorConstruction.cpp b/Engine/Source/Runtime/Engine/Private/ActorConstruction.cpp
--- a/Engine/Source/Runtime/Engine/Private/ActorConstruction.cpp
+++ b/Engine/Source/Runtime/Engine/Private/ActorConstruction.cpp
@@ -550,6 +550,38 @@ UActorComponent* AActor::CreateComponentFromTemplate(UActorComponent* Template,
 	UActorComponent* NewActorComp = NULL;
 	if(Template != NULL)
 	{
+		// If there is a Component with this name already (almost certainly because it is an Instance component), we need to rename it out of the way
+		if (!InName.IsEmpty())
+		{
+			UObject* ConflictingObject = FindObjectFast<UObject>(this, *InName);
+			if (ConflictingObject)
+			{
+				ensure(ConflictingObject->IsA<UActorComponent>() && CastChecked<UActorComponent>(ConflictingObject)->bInstanceComponent);
+				
+				// Try and pick a good name
+				FString ConflictingObjectName = ConflictingObject->GetName();
+				int32 CharIndex = ConflictingObjectName.Len()-1;
+				while (FChar::IsDigit(ConflictingObjectName[CharIndex]))
+				{
+					--CharIndex;
+				}
+				int32 Counter = 0;
+				if (CharIndex < ConflictingObjectName.Len()-1)
+				{
+					Counter = FCString::Atoi(*ConflictingObjectName.RightChop(CharIndex+1));
+					ConflictingObjectName = ConflictingObjectName.Left(CharIndex+1);
+				}
+				FString NewObjectName;
+				do
+				{
+					NewObjectName = ConflictingObjectName + FString::FromInt(++Counter);
+					
+				} while (FindObjectFast<UObject>(this, *NewObjectName) != nullptr);
+
+				ConflictingObject->Rename(*NewObjectName, this);
+			}
+		}
+
 		// Note we aren't copying the the RF_ArchetypeObject flag. Also note the result is non-transactional by default.
 		NewActorComp = (UActorComponent*)StaticDuplicateObject(Template, this, *InName, RF_AllFlags & ~(RF_ArchetypeObject|RF_Transactional|RF_WasLoaded|RF_Public) );
 		//NewActorComp = ConstructObject<UActorComponent>(Template->GetClass(), this, *InName, RF_NoFlags, Template);
