[INTEGRATE] Change 2417878 by [REDACTED_PERSON][REDACTED_HANDLE] on 2015/01/24 14:34:25

	Deprecate bCreatedByConstructionScript and bInstanceComponent and replace with an enumerated CreationMethod

diff --git a/Engine/Source/Editor/DetailCustomizations/Private/ActorMaterialCategory.cpp b/Engine/Source/Editor/DetailCustomizations/Private/ActorMaterialCategory.cpp
--- a/Engine/Source/Editor/DetailCustomizations/Private/ActorMaterialCategory.cpp
+++ b/Engine/Source/Editor/DetailCustomizations/Private/ActorMaterialCategory.cpp
@@ -70,7 +70,7 @@ public:
 						UPrimitiveComponent* PrimitiveComp = Cast<UPrimitiveComponent>( CurComponent );
 						UDecalComponent* DecalComponent = (PrimitiveComp ? NULL : Cast<UDecalComponent>( CurComponent ));
 
-						if( !CurComponent->bCreatedByConstructionScript )
+						if( CurComponent->CreationMethod != EComponentCreationMethod::ConstructionScript)
 						{
 							if( PrimitiveComp )
 							{
diff --git a/Engine/Source/Editor/FoliageEdit/Private/FoliageEdMode.cpp b/Engine/Source/Editor/FoliageEdit/Private/FoliageEdMode.cpp
--- a/Engine/Source/Editor/FoliageEdit/Private/FoliageEdMode.cpp
+++ b/Engine/Source/Editor/FoliageEdit/Private/FoliageEdMode.cpp
@@ -1873,7 +1873,7 @@ bool FEdModeFoliage::InputKey(FEditorViewportClient* ViewportClient, FViewport*
 								Instance.ZOffset = 0.f;
 								Instance.Base = Hit.Component.Get();
 								// We cannot be based on an a blueprint component as these will disappear when the construction script is re-run
-								if (Instance.Base && Instance.Base->bCreatedByConstructionScript)
+								if (Instance.Base && Instance.Base->CreationMethod == EComponentCreationMethod::ConstructionScript)
 								{
 									Instance.Base = nullptr;
 								}
diff --git a/Engine/Source/Editor/Kismet/Private/SSCSEditor.cpp b/Engine/Source/Editor/Kismet/Private/SSCSEditor.cpp
--- a/Engine/Source/Editor/Kismet/Private/SSCSEditor.cpp
+++ b/Engine/Source/Editor/Kismet/Private/SSCSEditor.cpp
@@ -543,7 +543,7 @@ bool FSCSEditorTreeNode::IsUserInstanced() const
 	if(bIsInstanced && !bWasInstancedFromNativeClass)
 	{
 		UActorComponent* ComponentInstance = GetComponentTemplate();
-		return ComponentInstance != nullptr && !ComponentInstance->bCreatedByConstructionScript;
+		return ComponentInstance != nullptr && ComponentInstance->CreationMethod != EComponentCreationMethod::ConstructionScript;
 	}
 
 	return false;
@@ -563,7 +563,7 @@ bool FSCSEditorTreeNode::CanEditDefaults() const
 	else if(bIsInstanced)
 	{
 		// Evaluate to TRUE for all instanced components except for those instanced from the Blueprint-generated class (i.e. during SCS or UCS)
-		bCanEdit = !ComponentTemplate->bCreatedByConstructionScript;
+		bCanEdit = (ComponentTemplate->CreationMethod != EComponentCreationMethod::ConstructionScript);
 	}
 	else if(ComponentTemplate != NULL)
 	{
@@ -1398,7 +1398,7 @@ FSlateColor SSCS_RowWidget::GetColorTint() const
 			UActorComponent* InstancedComponent = NodePtr->GetComponentTemplate();
 			check(InstancedComponent != nullptr);
 
-			if(InstancedComponent->bCreatedByConstructionScript)
+			if(InstancedComponent->CreationMethod == EComponentCreationMethod::ConstructionScript)
 			{
 				return FLinearColor(0.08f,0.35f,0.6f);
 			}
@@ -3419,7 +3419,7 @@ FSCSEditorTreeNodePtrType SSCSEditor::GetNodeFromActorComponent(const UActorComp
 				UClass* OwnerClass = ActorComponent->GetOwner()->GetActorClass();
 
 				// If the given component is one that's created during Blueprint construction
-				if (ActorComponent->bCreatedByConstructionScript)
+				if (ActorComponent->CreationMethod == EComponentCreationMethod::ConstructionScript)
 				{
 					// Get the Blueprint object associated with the owner's class
 					UBlueprint* Blueprint = UBlueprint::GetBlueprintFromClass(OwnerClass);
diff --git a/Engine/Source/Editor/UnrealEd/Private/Editor.cpp b/Engine/Source/Editor/UnrealEd/Private/Editor.cpp
--- a/Engine/Source/Editor/UnrealEd/Private/Editor.cpp
+++ b/Engine/Source/Editor/UnrealEd/Private/Editor.cpp
@@ -7420,7 +7420,7 @@ namespace EditorUtilities
 											{
 												// Ensure that this instance will be included in any undo/redo operations, and record it into the transaction buffer.
 												// Note: We don't do this for components that originate from script, because they will be re-instanced from the template after an undo, so there is no need to record them.
-												if(!ComponentArchetypeInstance->bCreatedByConstructionScript)
+												if(ComponentArchetypeInstance->CreationMethod != EComponentCreationMethod::ConstructionScript)
 												{
 													ComponentArchetypeInstance->SetFlags(RF_Transactional);
 													ComponentArchetypeInstance->Modify();
diff --git a/Engine/Source/Editor/UnrealEd/Private/EditorObject.cpp b/Engine/Source/Editor/UnrealEd/Private/EditorObject.cpp
--- a/Engine/Source/Editor/UnrealEd/Private/EditorObject.cpp
+++ b/Engine/Source/Editor/UnrealEd/Private/EditorObject.cpp
@@ -491,7 +491,7 @@ static const TCHAR* ImportProperties(
 					if (NewFlags & RF_Transactional)
 					{
 						UActorComponent* Component = Cast<UActorComponent>(ComponentTemplate);
-						if (Component && Component->bCreatedByConstructionScript)
+						if (Component && Component->CreationMethod == EComponentCreationMethod::ConstructionScript)
 						{
 							NewFlags &= ~RF_Transactional;
 						}
diff --git a/Engine/Source/Editor/UnrealEd/Private/Kismet2/Kismet2.cpp b/Engine/Source/Editor/UnrealEd/Private/Kismet2/Kismet2.cpp
--- a/Engine/Source/Editor/UnrealEd/Private/Kismet2/Kismet2.cpp
+++ b/Engine/Source/Editor/UnrealEd/Private/Kismet2/Kismet2.cpp
@@ -938,7 +938,7 @@ void FKismetEditorUtilities::AddComponentsToBlueprint(UBlueprint* Blueprint, con
 			UEditorEngine::CopyPropertiesForUnrelatedObjects(ActorComponent,SCSNode->ComponentTemplate);
 
 			// Clear the instance component flag
-			SCSNode->ComponentTemplate->bInstanceComponent = false;
+			SCSNode->ComponentTemplate->CreationMethod = EComponentCreationMethod::Native;
 
 			USceneComponent* SceneComponent = Cast<USceneComponent>(ActorComponent);
 
@@ -956,7 +956,7 @@ void FKismetEditorUtilities::AddComponentsToBlueprint(UBlueprint* Blueprint, con
 					SCS->AddNode(SCSNode);
 				}
 				// If we're attached to a blueprint component look it up as the variable name is the component name
-				else if (SceneComponent->AttachParent->bCreatedByConstructionScript)
+				else if (SceneComponent->AttachParent->CreationMethod == EComponentCreationMethod::ConstructionScript)
 				{
 					USCS_Node* ParentSCSNode = nullptr;
 					for (UBlueprint* Blueprint : ParentBPStack)
diff --git a/Engine/Source/Editor/UnrealEd/Public/Kismet2/ComponentEditorUtils.h b/Engine/Source/Editor/UnrealEd/Public/Kismet2/ComponentEditorUtils.h
--- a/Engine/Source/Editor/UnrealEd/Public/Kismet2/ComponentEditorUtils.h
+++ b/Engine/Source/Editor/UnrealEd/Public/Kismet2/ComponentEditorUtils.h
@@ -75,7 +75,7 @@ public:
 		{
 			// Ensure that this instance will be included in any undo/redo operations, and record it into the transaction buffer.
 			// Note: We don't do this for components that originate from script, because they will be re-instanced from the template after an undo, so there is no need to record them.
-			if(!InSceneComponent->bCreatedByConstructionScript)
+			if(InSceneComponent->CreationMethod != EComponentCreationMethod::ConstructionScript)
 			{
 				InSceneComponent->SetFlags(RF_Transactional);
 				InSceneComponent->Modify();
diff --git a/Engine/Source/Runtime/Engine/Classes/Components/ActorComponent.h b/Engine/Source/Runtime/Engine/Classes/Components/ActorComponent.h
--- a/Engine/Source/Runtime/Engine/Classes/Components/ActorComponent.h
+++ b/Engine/Source/Runtime/Engine/Classes/Components/ActorComponent.h
@@ -8,6 +8,14 @@
 
 struct FReplicationFlags;
 
+UENUM()
+enum class EComponentCreationMethod : uint8
+{
+	Native,
+	ConstructionScript,
+	Instance,
+};
+
 /**
  * ActorComponent is the base class for components that define reusable behavior that can be added to different types of Actors.
  * ActorComponents that have a transform are known as SceneComponents and those that can be rendered are PrimitiveComponents.
@@ -88,11 +96,11 @@ public:
 
 	/** True if this component was created by a construction script, and will be destroyed by DestroyConstructedComponents */
 	UPROPERTY()
-	uint32 bCreatedByConstructionScript:1;
+	uint32 bCreatedByConstructionScript_DEPRECATED:1;
 
 	/** True if this component was created as an instance component */
 	UPROPERTY()
-	uint32 bInstanceComponent:1;
+	uint32 bInstanceComponent_DEPRECATED:1;
 
 	/** Whether to the component is activated at creation or must be explicitly activated. */
 	UPROPERTY(EditAnywhere, BlueprintReadOnly, Category=Activation)
@@ -111,6 +119,9 @@ public:
 	/** Indicates that InitializeComponent has been called, but UninitializeComponent has not yet */
 	uint32 bHasBeenInitialized:1;
 
+	UPROPERTY()
+	EComponentCreationMethod CreationMethod;
+
 	UFUNCTION()
 	void OnRep_IsActive();
 
@@ -460,6 +471,7 @@ public:
 	virtual int32 GetFunctionCallspace( UFunction* Function, void* Parameters, FFrame* Stack ) override;
 	virtual bool CallRemoteFunction( UFunction* Function, void* Parameters, FOutParmRec* OutParms, FFrame* Stack ) override;
 	virtual void PostInitProperties() override;
+	virtual void PostLoad() override;
 	virtual void PostRename(UObject* OldOuter, const FName OldName) override;
 #if WITH_EDITOR
 	virtual void PreEditChange(UProperty* PropertyThatWillChange) override;
diff --git a/Engine/Source/Runtime/Engine/Private/Actor.cpp b/Engine/Source/Runtime/Engine/Private/Actor.cpp
--- a/Engine/Source/Runtime/Engine/Private/Actor.cpp
+++ b/Engine/Source/Runtime/Engine/Private/Actor.cpp
@@ -797,7 +797,7 @@ bool AActor::Modify( bool bAlwaysMarkDirty/*=true*/ )
 			if (!ObjProp->HasAllPropertyFlags(CPF_NonTransactional))
 			{
 				UActorComponent* ActorComponent = Cast<UActorComponent>(ObjProp->GetObjectPropertyValue(ObjProp->ContainerPtrToValuePtr<void>(this)));
-				if (ActorComponent && ActorComponent->bCreatedByConstructionScript)
+				if (ActorComponent && ActorComponent->CreationMethod == EComponentCreationMethod::ConstructionScript)
 				{
 					ObjProp->SetPropertyFlags(CPF_NonTransactional);
 					TemporarilyNonTransactionalProperties.Add(ObjProp);
@@ -814,7 +814,7 @@ bool AActor::Modify( bool bAlwaysMarkDirty/*=true*/ )
 	}
 
 	// If the root component is blueprint constructed we don't save it to the transaction buffer
-	if( RootComponent && !RootComponent->bCreatedByConstructionScript )
+	if( RootComponent && RootComponent->CreationMethod != EComponentCreationMethod::ConstructionScript)
 	{
 		bSavedToTransactionBuffer = RootComponent->Modify( bAlwaysMarkDirty ) || bSavedToTransactionBuffer;
 	}
@@ -2051,7 +2051,7 @@ const TArray<UActorComponent*>& AActor::GetInstanceComponents() const
 
 void AActor::AddInstanceComponent(UActorComponent* Component)
 {
-	Component->bInstanceComponent = true;
+	Component->CreationMethod = EComponentCreationMethod::Instance;
 	InstanceComponents.Add(Component);
 }
 
diff --git a/Engine/Source/Runtime/Engine/Private/ActorComponent.cpp b/Engine/Source/Runtime/Engine/Private/ActorComponent.cpp
--- a/Engine/Source/Runtime/Engine/Private/ActorComponent.cpp
+++ b/Engine/Source/Runtime/Engine/Private/ActorComponent.cpp
@@ -106,6 +106,8 @@ UActorComponent::UActorComponent(const FObjectInitializer& ObjectInitializer)
 	PrimaryComponentTick.bStartWithTickEnabled = true;
 	PrimaryComponentTick.bCanEverTick = false;
 
+	CreationMethod = EComponentCreationMethod::Native;
+
 	bAutoRegister = true;
 	bNetAddressable = false;
 }
@@ -121,6 +123,20 @@ void UActorComponent::PostInitProperties()
 	}
 }
 
+void UActorComponent::PostLoad()
+{
+	Super::PostLoad();
+
+	if (bCreatedByConstructionScript_DEPRECATED)
+	{
+		CreationMethod = EComponentCreationMethod::ConstructionScript;
+	}
+	else if (bInstanceComponent_DEPRECATED)
+	{
+		CreationMethod = EComponentCreationMethod::Instance;
+	}
+}
+
 void UActorComponent::PostRename(UObject* OldOuter, const FName OldName)
 {
 	if (OldOuter != GetOuter())
@@ -628,7 +644,7 @@ void UActorComponent::RegisterComponentWithWorld(UWorld* InWorld)
 	}
 
 	// If this is a blueprint created component and it has component children they can miss getting registered in some scenarios
-	if (bCreatedByConstructionScript)
+	if (CreationMethod == EComponentCreationMethod::ConstructionScript)
 	{
 		TArray<UObject*> Children;
 		GetObjectsWithOuter(this, Children, true, RF_PendingKill);
@@ -698,7 +714,7 @@ void UActorComponent::DestroyComponent()
 	AActor* Owner = GetOwner();
 	if(Owner != NULL)
 	{
-		if (bCreatedByConstructionScript)
+		if (CreationMethod == EComponentCreationMethod::ConstructionScript)
 		{
 			Owner->BlueprintCreatedComponents.Remove(this);
 		}
diff --git a/Engine/Source/Runtime/Engine/Private/ActorConstruction.cpp b/Engine/Source/Runtime/Engine/Private/ActorConstruction.cpp
--- a/Engine/Source/Runtime/Engine/Private/ActorConstruction.cpp
+++ b/Engine/Source/Runtime/Engine/Private/ActorConstruction.cpp
@@ -194,7 +194,7 @@ void AActor::DestroyConstructedComponents()
 		if (Component)
 		{
 			bool bDestroyComponent = false;
-			if (Component->bCreatedByConstructionScript)
+			if (Component->CreationMethod == EComponentCreationMethod::ConstructionScript)
 			{
 				bDestroyComponent = true;
 			}
@@ -203,7 +203,7 @@ void AActor::DestroyConstructedComponents()
 				UActorComponent* OuterComponent = Component->GetTypedOuter<UActorComponent>();
 				while (OuterComponent)
 				{
-					if (OuterComponent->bCreatedByConstructionScript)
+					if (OuterComponent->CreationMethod == EComponentCreationMethod::ConstructionScript)
 					{
 						bDestroyComponent = true;
 						break;
@@ -336,7 +336,7 @@ void AActor::RerunConstructionScripts()
 			{
 				USceneComponent* EachRoot = AttachedActor->GetRootComponent();
 				// If the component we are attached to is about to go away...
-				if (EachRoot && EachRoot->AttachParent && EachRoot->AttachParent->bCreatedByConstructionScript)
+				if (EachRoot && EachRoot->AttachParent && EachRoot->AttachParent->CreationMethod == EComponentCreationMethod::ConstructionScript)
 				{
 					// Save info about actor to reattach
 					FAttachedActorInfo Info;
@@ -355,7 +355,7 @@ void AActor::RerunConstructionScripts()
 		if (bUseRootComponentProperties && RootComponent != nullptr)
 		{
 			// Do not need to detach if root component is not going away
-			if (RootComponent->AttachParent != NULL && RootComponent->bCreatedByConstructionScript)
+			if (RootComponent->AttachParent != NULL && RootComponent->CreationMethod == EComponentCreationMethod::ConstructionScript)
 			{
 				Parent = RootComponent->AttachParent->GetOwner();
 				// Root component should never be attached to another component in the same actor!
@@ -508,7 +508,7 @@ void AActor::ExecuteConstruction(const FTransform& Transform, const FComponentIn
 			{
 				UBillboardComponent* BillboardComponent = NewObject<UBillboardComponent>(this);
 				BillboardComponent->SetFlags(RF_Transactional);
-				BillboardComponent->bCreatedByConstructionScript = true;
+				BillboardComponent->CreationMethod = EComponentCreationMethod::ConstructionScript;
 #if WITH_EDITOR
 				BillboardComponent->Sprite = (UTexture2D*)(StaticLoadObject(UTexture2D::StaticClass(), NULL, TEXT("/Engine/EditorResources/BadBlueprintSprite.BadBlueprintSprite"), NULL, LOAD_None, NULL));
 #endif
@@ -556,7 +556,7 @@ UActorComponent* AActor::CreateComponentFromTemplate(UActorComponent* Template,
 			UObject* ConflictingObject = FindObjectFast<UObject>(this, *InName);
 			if (ConflictingObject)
 			{
-				ensure(ConflictingObject->IsA<UActorComponent>() && CastChecked<UActorComponent>(ConflictingObject)->bInstanceComponent);
+				ensure(ConflictingObject->IsA<UActorComponent>() && CastChecked<UActorComponent>(ConflictingObject)->CreationMethod == EComponentCreationMethod::Instance);
 				
 				// Try and pick a good name
 				FString ConflictingObjectName = ConflictingObject->GetName();
@@ -586,7 +586,7 @@ UActorComponent* AActor::CreateComponentFromTemplate(UActorComponent* Template,
 		NewActorComp = (UActorComponent*)StaticDuplicateObject(Template, this, *InName, RF_AllFlags & ~(RF_ArchetypeObject|RF_Transactional|RF_WasLoaded|RF_Public) );
 		//NewActorComp = ConstructObject<UActorComponent>(Template->GetClass(), this, *InName, RF_NoFlags, Template);
 
-		NewActorComp->bCreatedByConstructionScript = true;
+		NewActorComp->CreationMethod = EComponentCreationMethod::ConstructionScript;
 
 		// Need to do this so component gets saved - Components array is not serialized
 		BlueprintCreatedComponents.Add(NewActorComp);
diff --git a/Engine/Source/Runtime/Engine/Private/ActorEditor.cpp b/Engine/Source/Runtime/Engine/Private/ActorEditor.cpp
--- a/Engine/Source/Runtime/Engine/Private/ActorEditor.cpp
+++ b/Engine/Source/Runtime/Engine/Private/ActorEditor.cpp
@@ -83,7 +83,7 @@ void AActor::PostEditMove(bool bFinished)
 	}
 
 	// If the root component was not just recreated by the construction script - call PostEditComponentMove on it
-	if(RootComponent != NULL && !RootComponent->bCreatedByConstructionScript)
+	if(RootComponent != NULL && RootComponent->CreationMethod != EComponentCreationMethod::ConstructionScript)
 	{
 		// @TODO Should we call on ALL components?
 		RootComponent->PostEditComponentMove(bFinished);
@@ -204,7 +204,7 @@ AActor::FActorTransactionAnnotation::FActorTransactionAnnotation(const AActor* A
 	: ComponentInstanceData(Actor)
 {
 	USceneComponent* RootComponent = Actor->GetRootComponent();
-	if (RootComponent && RootComponent->bCreatedByConstructionScript)
+	if (RootComponent && RootComponent->CreationMethod == EComponentCreationMethod::ConstructionScript)
 	{
 		bRootComponentDataCached = true;
 		RootComponentData.Transform = RootComponent->ComponentToWorld;
diff --git a/Engine/Source/Runtime/Engine/Private/AudioComponent.cpp b/Engine/Source/Runtime/Engine/Private/AudioComponent.cpp
--- a/Engine/Source/Runtime/Engine/Private/AudioComponent.cpp
+++ b/Engine/Source/Runtime/Engine/Private/AudioComponent.cpp
@@ -59,8 +59,7 @@ void UAudioComponent::OnRegister()
 		SpriteComponent->AlwaysLoadOnServer = false;
 		SpriteComponent->SpriteInfo.Category = TEXT("Misc");
 		SpriteComponent->SpriteInfo.DisplayName = NSLOCTEXT( "SpriteCategory", "Misc", "Misc" );
-		SpriteComponent->bCreatedByConstructionScript = bCreatedByConstructionScript;
-		SpriteComponent->bInstanceComponent = bInstanceComponent;
+		SpriteComponent->CreationMethod = CreationMethod;
 		SpriteComponent->bIsScreenSizeScaled = true;
 		SpriteComponent->bUseInEditorScaling = true;
 
diff --git a/Engine/Source/Runtime/Engine/Private/BlueprintGeneratedClass.cpp b/Engine/Source/Runtime/Engine/Private/BlueprintGeneratedClass.cpp
--- a/Engine/Source/Runtime/Engine/Private/BlueprintGeneratedClass.cpp
+++ b/Engine/Source/Runtime/Engine/Private/BlueprintGeneratedClass.cpp
@@ -308,7 +308,7 @@ void UBlueprintGeneratedClass::CreateComponentsForActor(AActor* Actor) const
 
 		FName NewName = *(FString::Printf(TEXT("TimelineComp__%d"), Actor->BlueprintCreatedComponents.Num() ) );
 		UTimelineComponent* NewTimeline = NewNamedObject<UTimelineComponent>(Actor, NewName);
-		NewTimeline->bCreatedByConstructionScript = true; // Indicate it comes from a blueprint so it gets cleared when we rerun construction scripts
+		NewTimeline->CreationMethod = EComponentCreationMethod::ConstructionScript; // Indicate it comes from a blueprint so it gets cleared when we rerun construction scripts
 		Actor->BlueprintCreatedComponents.Add(NewTimeline); // Add to array so it gets saved
 		NewTimeline->SetNetAddressable();	// This component has a stable name that can be referenced for replication
 
diff --git a/Engine/Source/Runtime/Engine/Private/Camera/CameraComponent.cpp b/Engine/Source/Runtime/Engine/Private/Camera/CameraComponent.cpp
--- a/Engine/Source/Runtime/Engine/Private/Camera/CameraComponent.cpp
+++ b/Engine/Source/Runtime/Engine/Private/Camera/CameraComponent.cpp
@@ -51,8 +51,7 @@ void UCameraComponent::OnRegister()
 		ProxyMeshComponent->bHiddenInGame = true;
 		ProxyMeshComponent->CastShadow = false;
 		ProxyMeshComponent->PostPhysicsComponentTick.bCanEverTick = false;
-		ProxyMeshComponent->bCreatedByConstructionScript = bCreatedByConstructionScript;
-		ProxyMeshComponent->bInstanceComponent = bInstanceComponent;
+		ProxyMeshComponent->CreationMethod = CreationMethod;
 		ProxyMeshComponent->RegisterComponentWithWorld(GetWorld());
 	}
 
@@ -62,8 +61,7 @@ void UCameraComponent::OnRegister()
 		DrawFrustum->AttachTo(this);
 		DrawFrustum->AlwaysLoadOnClient = false;
 		DrawFrustum->AlwaysLoadOnServer = false;
-		DrawFrustum->bCreatedByConstructionScript = bCreatedByConstructionScript;
-		DrawFrustum->bInstanceComponent = bInstanceComponent;
+		DrawFrustum->CreationMethod = CreationMethod;
 		DrawFrustum->RegisterComponentWithWorld(GetWorld());
 	}
 
diff --git a/Engine/Source/Runtime/Engine/Private/ComponentInstanceDataCache.cpp b/Engine/Source/Runtime/Engine/Private/ComponentInstanceDataCache.cpp
--- a/Engine/Source/Runtime/Engine/Private/ComponentInstanceDataCache.cpp
+++ b/Engine/Source/Runtime/Engine/Private/ComponentInstanceDataCache.cpp
@@ -77,7 +77,7 @@ FComponentInstanceDataCache::FComponentInstanceDataCache(const AActor* Actor)
 		// Grab per-instance data we want to persist
 		for (UActorComponent* Component : Components)
 		{
-			if(Component->bCreatedByConstructionScript) // Only cache data from 'created by construction script' components
+			if(Component->CreationMethod == EComponentCreationMethod::ConstructionScript) // Only cache data from 'created by construction script' components
 			{
 				FComponentInstanceDataBase* ComponentInstanceData = Component->GetComponentInstanceData();
 				if (ComponentInstanceData)
@@ -108,7 +108,7 @@ void FComponentInstanceDataCache::ApplyToActor(AActor* Actor) const
 		// Apply per-instance data.
 		for (UActorComponent* Component : Components)
 		{
-			if(Component->bCreatedByConstructionScript) // Only try and apply data to 'created by construction script' components
+			if(Component->CreationMethod == EComponentCreationMethod::ConstructionScript) // Only try and apply data to 'created by construction script' components
 			{
 				const FName ComponentInstanceDataType = Component->GetComponentInstanceDataType();
 
diff --git a/Engine/Source/Runtime/Engine/Private/LightComponent.cpp b/Engine/Source/Runtime/Engine/Private/LightComponent.cpp
--- a/Engine/Source/Runtime/Engine/Private/LightComponent.cpp
+++ b/Engine/Source/Runtime/Engine/Private/LightComponent.cpp
@@ -182,8 +182,7 @@ void ULightComponentBase::OnRegister()
 		SpriteComponent->AlwaysLoadOnServer = false;
 		SpriteComponent->SpriteInfo.Category = TEXT("Lighting");
 		SpriteComponent->SpriteInfo.DisplayName = NSLOCTEXT("SpriteCategory", "Lighting", "Lighting");
-		SpriteComponent->bCreatedByConstructionScript = bCreatedByConstructionScript;
-		SpriteComponent->bInstanceComponent = bInstanceComponent;
+		SpriteComponent->CreationMethod = CreationMethod;
 		SpriteComponent->bIsScreenSizeScaled = true;
 
 		SpriteComponent->RegisterComponent();
diff --git a/Engine/Source/Runtime/Engine/Private/PhysicsEngine/PhysicsConstraintComponent.cpp b/Engine/Source/Runtime/Engine/Private/PhysicsEngine/PhysicsConstraintComponent.cpp
--- a/Engine/Source/Runtime/Engine/Private/PhysicsEngine/PhysicsConstraintComponent.cpp
+++ b/Engine/Source/Runtime/Engine/Private/PhysicsEngine/PhysicsConstraintComponent.cpp
@@ -245,8 +245,7 @@ void UPhysicsConstraintComponent::OnRegister()
 		SpriteComponent->AlwaysLoadOnServer = false;
 		SpriteComponent->SpriteInfo.Category = TEXT("Physics");
 		SpriteComponent->SpriteInfo.DisplayName = NSLOCTEXT( "SpriteCategory", "Physics", "Physics" );
-		SpriteComponent->bCreatedByConstructionScript = bCreatedByConstructionScript;
-		SpriteComponent->bInstanceComponent = bInstanceComponent;
+		SpriteComponent->CreationMethod = CreationMethod;
 		SpriteComponent->bIsScreenSizeScaled = true;
 
 		SpriteComponent->RegisterComponent();
diff --git a/Engine/Source/Runtime/Engine/Private/SceneComponent.cpp b/Engine/Source/Runtime/Engine/Private/SceneComponent.cpp
--- a/Engine/Source/Runtime/Engine/Private/SceneComponent.cpp
+++ b/Engine/Source/Runtime/Engine/Private/SceneComponent.cpp
@@ -1013,7 +1013,7 @@ FSceneComponentInstanceData::FSceneComponentInstanceData(const USceneComponent*
 	for (int32 i = SourceComponent->AttachChildren.Num()-1; i >= 0; --i)
 	{
 		USceneComponent* SceneComponent = SourceComponent->AttachChildren[i];
-		if (SceneComponent && !SceneComponent->bCreatedByConstructionScript)
+		if (SceneComponent && SceneComponent->CreationMethod != EComponentCreationMethod::ConstructionScript)
 		{
 			AttachedInstanceComponents.Add(SceneComponent);
 		}
@@ -1046,7 +1046,7 @@ FComponentInstanceDataBase* USceneComponent::GetComponentInstanceData() const
 
 	for (USceneComponent* Child : AttachChildren)
 	{
-		if (Child && !Child->bCreatedByConstructionScript)
+		if (Child && Child->CreationMethod != EComponentCreationMethod::ConstructionScript)
 		{
 			InstanceData = new FSceneComponentInstanceData(this);
 			break;
diff --git a/Engine/Source/Runtime/Engine/Private/SimpleConstructionScript.cpp b/Engine/Source/Runtime/Engine/Private/SimpleConstructionScript.cpp
--- a/Engine/Source/Runtime/Engine/Private/SimpleConstructionScript.cpp
+++ b/Engine/Source/Runtime/Engine/Private/SimpleConstructionScript.cpp
@@ -338,7 +338,7 @@ void USimpleConstructionScript::ExecuteScriptOnActor(AActor* Actor, const FTrans
 				Actor->GetComponents(Components);
 				for (int32 Index = Components.Num()-1; Index >= 0; --Index)
 				{
-					if (Components[Index]->bInstanceComponent)
+					if (Components[Index]->CreationMethod == EComponentCreationMethod::Instance)
 					{
 						Components.RemoveAt(Index);
 					}
@@ -404,7 +404,7 @@ void USimpleConstructionScript::ExecuteScriptOnActor(AActor* Actor, const FTrans
 	{
 		USceneComponent* SceneComp = NewObject<USceneComponent>(Actor);
 		SceneComp->SetFlags(RF_Transactional);
-		SceneComp->bCreatedByConstructionScript = true;
+		SceneComp->CreationMethod = EComponentCreationMethod::ConstructionScript;
 		SceneComp->SetWorldTransform(RootTransform);
 		Actor->SetRootComponent(SceneComp);
 		SceneComp->RegisterComponent();
diff --git a/Engine/Source/Runtime/Foliage/Private/InstancedFoliage.cpp b/Engine/Source/Runtime/Foliage/Private/InstancedFoliage.cpp
--- a/Engine/Source/Runtime/Foliage/Private/InstancedFoliage.cpp
+++ b/Engine/Source/Runtime/Foliage/Private/InstancedFoliage.cpp
@@ -1603,7 +1603,7 @@ void AInstancedFoliageActor::PostLoad()
 				// Clear out the Base for any instances based on blueprint-created components,
 				// as those components will be destroyed when the construction scripts are
 				// re-run, leaving dangling references and causing crashes (woo!)
-				if (Instance.Base && Instance.Base->bCreatedByConstructionScript)
+				if (Instance.Base && Instance.Base->CreationMethod == EComponentCreationMethod::ConstructionScript)
 				{
 					Instance.Base = NULL;
 				}
