[INTEGRATE] Change 2418462 by [REDACTED_PERSON][REDACTED_HANDLE] on 2015/01/26 05:24:57

	UE-7586 Fix tooltip for Level factory

diff --git a/Engine/Source/Editor/UnrealEd/Classes/Factories/WorldFactory.h b/Engine/Source/Editor/UnrealEd/Classes/Factories/WorldFactory.h
--- a/Engine/Source/Editor/UnrealEd/Classes/Factories/WorldFactory.h
+++ b/Engine/Source/Editor/UnrealEd/Classes/Factories/WorldFactory.h
@@ -20,6 +20,9 @@ class UWorldFactory : public UFactory
 	// Begin UFactory Interface
 	virtual bool ConfigureProperties() override;
 	virtual UObject* FactoryCreateNew(UClass* Class,UObject* InParent,FName Name,EObjectFlags Flags,UObject* Context,FFeedbackContext* Warn) override;
+	virtual FText GetToolTip() const;
+	virtual FString GetToolTipDocumentationPage() const;
+	virtual FString GetToolTipDocumentationExcerpt() const;
 	// Begin UFactory Interface	
 };
 
diff --git a/Engine/Source/Editor/UnrealEd/Private/Factories/WorldFactory.cpp b/Engine/Source/Editor/UnrealEd/Private/Factories/WorldFactory.cpp
--- a/Engine/Source/Editor/UnrealEd/Private/Factories/WorldFactory.cpp
+++ b/Engine/Source/Editor/UnrealEd/Private/Factories/WorldFactory.cpp
@@ -1,6 +1,7 @@
 // Copyright 1998-2015 Epic Games, Inc. All Rights Reserved.
 
 #include "UnrealEd.h"
+#include "EditorClassUtils.h"
 
 #define LOCTEXT_NAMESPACE "WorldFactory"
 
@@ -30,4 +31,19 @@ UObject* UWorldFactory::FactoryCreateNew(UClass* Class, UObject* InParent, FName
 	return NewWorld;
 }
 
+FText UWorldFactory::GetToolTip() const
+{
+	return ULevel::StaticClass()->GetToolTipText();
+}
+
+FString UWorldFactory::GetToolTipDocumentationPage() const
+{
+	return FEditorClassUtils::GetDocumentationPage(ULevel::StaticClass());
+}
+
+FString UWorldFactory::GetToolTipDocumentationExcerpt() const
+{
+	return FEditorClassUtils::GetDocumentationExcerpt(ULevel::StaticClass());
+}
+
 #undef LOCTEXT_NAMESPACE
diff --git a/Engine/Source/Runtime/Engine/Classes/Engine/Level.h b/Engine/Source/Runtime/Engine/Classes/Engine/Level.h
--- a/Engine/Source/Runtime/Engine/Classes/Engine/Level.h
+++ b/Engine/Source/Runtime/Engine/Classes/Engine/Level.h
@@ -277,13 +277,14 @@ struct ENGINE_API FLevelSimplificationDetails
 
 //
 // The level object.  Contains the level's actor list, BSP information, and brush list.
+// Every Level has a World as its Outer and can be used as the PersistentLevel, however,
+// when a Level has been streamed in the OwningWorld represents the World that it is a part of.
 //
 
 
 /**
- * A Level contains actors, a blueprint representing the level logic, BSP information, and brushes.
- * Every Level has a World as its Outer and can be used as the PersistentLevel, however, 
- * when a Level has been streamed in the OwningWorld represents the World that it is a part of.
+ * A Level is a collection of Actors (lights, volumes, mesh instances etc.).
+ * Multiple Levels can be loaded and unloaded into the World to create a streaming experience.
  * 
  * @see https://docs.unrealengine.com/latest/INT/Engine/Levels
  * @see UActor
