[INTEGRATE] Change 2418585 by [REDACTED_PERSON][REDACTED_HANDLE] on 2015/01/26 09:52:39

	Don't require classes that inherit AActor and UActorComponent to have an FObjectInitializer constructor
	- We keep the FObjectInitializer constructor for backwards compatibility
	- In the future we'll update more classes to work this way

diff --git a/Engine/Source/Runtime/Engine/Classes/Components/ActorComponent.h b/Engine/Source/Runtime/Engine/Classes/Components/ActorComponent.h
--- a/Engine/Source/Runtime/Engine/Classes/Components/ActorComponent.h
+++ b/Engine/Source/Runtime/Engine/Classes/Components/ActorComponent.h
@@ -33,6 +33,11 @@ public:
 	/**
 	 * Default UObject constructor.
 	 */
+	UActorComponent();
+
+	/**
+	 * UObject constructor that takes an ObjectInitializer
+	 */
 	UActorComponent(const FObjectInitializer& ObjectInitializer);
 
 	virtual void GetLifetimeReplicatedProps(TArray<FLifetimeProperty>& OutLifetimeProps) const override;
diff --git a/Engine/Source/Runtime/Engine/Classes/GameFramework/Actor.h b/Engine/Source/Runtime/Engine/Classes/GameFramework/Actor.h
--- a/Engine/Source/Runtime/Engine/Classes/GameFramework/Actor.h
+++ b/Engine/Source/Runtime/Engine/Classes/GameFramework/Actor.h
@@ -71,7 +71,12 @@ class ENGINE_API AActor : public UObject
 public:
 
 	/**
-	 * Default UObject constructor.
+	 * Default constructor for AActor
+	 */
+	AActor();
+
+	/**
+	 * Constructor for AActor that takes an ObjectInitializer
 	 */
 	AActor(const FObjectInitializer& ObjectInitializer);
 
diff --git a/Engine/Source/Runtime/Engine/Private/Actor.cpp b/Engine/Source/Runtime/Engine/Private/Actor.cpp
--- a/Engine/Source/Runtime/Engine/Private/Actor.cpp
+++ b/Engine/Source/Runtime/Engine/Private/Actor.cpp
@@ -45,8 +45,7 @@ FMakeNoiseDelegate AActor::MakeNoiseDelegate = FMakeNoiseDelegate::CreateStatic(
 FOnProcessEvent AActor::ProcessEventDelegate;
 #endif
 
-AActor::AActor(const FObjectInitializer& ObjectInitializer)
-	: Super(ObjectInitializer)
+AActor::AActor()
 {
 	PrimaryActorTick.TickGroup = TG_PrePhysics;
 	// Default to no tick function, but if we set 'never ticks' to false (so there is a tick function) it is enabled by default
@@ -85,6 +84,12 @@ AActor::AActor(const FObjectInitializer& ObjectInitializer)
 	bAllowReceiveTickEventOnDedicatedServer = true;
 }
 
+
+AActor::AActor(const FObjectInitializer& ObjectInitializer)
+	: AActor()	// Forward to default constructor (we don't use ObjectInitializer for anything, this is for compatibility with inherited classes that call Super( ObjectInitializer )
+{
+}
+
 void FActorTickFunction::ExecuteTick(float DeltaTime, enum ELevelTick TickType, ENamedThreads::Type CurrentThread, const FGraphEventRef& MyCompletionGraphEvent)
 {
 	if (Target && !Target->HasAnyFlags(RF_PendingKill | RF_Unreachable))
diff --git a/Engine/Source/Runtime/Engine/Private/ActorComponent.cpp b/Engine/Source/Runtime/Engine/Private/ActorComponent.cpp
--- a/Engine/Source/Runtime/Engine/Private/ActorComponent.cpp
+++ b/Engine/Source/Runtime/Engine/Private/ActorComponent.cpp
@@ -99,8 +99,7 @@ FGlobalComponentReregisterContext::~FGlobalComponentReregisterContext()
 }
 
 
-UActorComponent::UActorComponent(const FObjectInitializer& ObjectInitializer)
-	: Super(ObjectInitializer)
+UActorComponent::UActorComponent()
 {
 	PrimaryComponentTick.TickGroup = TG_DuringPhysics;
 	PrimaryComponentTick.bStartWithTickEnabled = true;
@@ -112,6 +111,13 @@ UActorComponent::UActorComponent(const FObjectInitializer& ObjectInitializer)
 	bNetAddressable = false;
 }
 
+
+UActorComponent::UActorComponent( const FObjectInitializer& ObjectInitializer )
+	: UActorComponent()	// Forward to default constructor (we don't use ObjectInitializer for anything, this is for compatibility with inherited classes that call Super( ObjectInitializer )
+{
+}
+
+
 void UActorComponent::PostInitProperties()
 {
 	Super::PostInitProperties();
