[INTEGRATE] Change 2418599 by [REDACTED_PERSON][REDACTED_HANDLE] on 2015/01/26 10:04:12

	Don't require classes that inherit USceneComponent to have an FObjectInitializer constructor
	- We keep the FObjectInitializer constructor for backwards compatibility
	- In the future we'll update more classes to work this way

diff --git a/Engine/Source/Runtime/Engine/Classes/Components/SceneComponent.h b/Engine/Source/Runtime/Engine/Classes/Components/SceneComponent.h
--- a/Engine/Source/Runtime/Engine/Classes/Components/SceneComponent.h
+++ b/Engine/Source/Runtime/Engine/Classes/Components/SceneComponent.h
@@ -85,6 +85,11 @@ public:
 	/**
 	 * Default UObject constructor.
 	 */
+	USceneComponent();
+
+	/**
+	 * UObject constructor that takes an ObjectInitializer
+	 */
 	USceneComponent(const FObjectInitializer& ObjectInitializer);
 
 	virtual void GetLifetimeReplicatedProps(TArray<FLifetimeProperty>& OutLifetimeProps) const override;
diff --git a/Engine/Source/Runtime/Engine/Private/SceneComponent.cpp b/Engine/Source/Runtime/Engine/Private/SceneComponent.cpp
--- a/Engine/Source/Runtime/Engine/Private/SceneComponent.cpp
+++ b/Engine/Source/Runtime/Engine/Private/SceneComponent.cpp
@@ -23,8 +23,7 @@ FOverlapInfo::FOverlapInfo(UPrimitiveComponent* InComponent, int32 InBodyIndex)
 	OverlapInfo.Item = InBodyIndex;
 }
 
-USceneComponent::USceneComponent(const FObjectInitializer& ObjectInitializer)
-	: Super(ObjectInitializer)
+USceneComponent::USceneComponent()
 {
 	Mobility = EComponentMobility::Movable;
 	RelativeScale3D = FVector(1.0f,1.0f,1.0f);
@@ -35,6 +34,13 @@ USceneComponent::USceneComponent(const FObjectInitializer& ObjectInitializer)
 	NetUpdateTransform = false;
 }
 
+
+USceneComponent::USceneComponent( const FObjectInitializer& ObjectInitializer )
+	: USceneComponent()	// Forward to default constructor (we don't use ObjectInitializer for anything, this is for compatibility with inherited classes that call Super( ObjectInitializer )
+{
+}
+
+
 FTransform USceneComponent::CalcNewComponentToWorld(const FTransform& NewRelativeTransform, const USceneComponent* Parent) const
 {
 	Parent = Parent ? Parent : AttachParent;
