[INTEGRATE] Change 2418631 by [REDACTED_PERSON][REDACTED_HANDLE] on 2015/01/26 10:51:59

	Improve comments for tick, overlap & hit events to indicate any flags that need to be set

diff --git a/Engine/Source/Runtime/Engine/Classes/Components/ActorComponent.h b/Engine/Source/Runtime/Engine/Classes/Components/ActorComponent.h
--- a/Engine/Source/Runtime/Engine/Classes/Components/ActorComponent.h
+++ b/Engine/Source/Runtime/Engine/Classes/Components/ActorComponent.h
@@ -305,11 +305,12 @@ public:
 	void RegisterAllComponentTickFunctions(bool bRegister);
 
 	/**
-	 * Updates time dependent state for this component.
-	 * Requires component to be registered
+	 * Function called every frame on this ActorComponent. Override this function to implement custom logic to be executed every frame.
+	 * Only executes if the component is registered, and also PrimaryComponentTick.bCanEverTick must be set to true.
+	 *	
 	 * @param DeltaTime - The time since the last tick.
-	 * @param TickType - The kind of tick this is
-	 * @param ThisTickFunction - Tick function that caused this to run
+	 * @param TickType - The kind of tick this is, for example, are we paused, or 'simulating' in the editor
+	 * @param ThisTickFunction - Internal tick function struct that caused this to run
 	 */
 	virtual void TickComponent(float DeltaTime, enum ELevelTick TickType, FActorComponentTickFunction *ThisTickFunction);
 	/** 
diff --git a/Engine/Source/Runtime/Engine/Classes/Components/PrimitiveComponent.h b/Engine/Source/Runtime/Engine/Classes/Components/PrimitiveComponent.h
--- a/Engine/Source/Runtime/Engine/Classes/Components/PrimitiveComponent.h
+++ b/Engine/Source/Runtime/Engine/Classes/Components/PrimitiveComponent.h
@@ -561,15 +561,24 @@ public:
 	 */
 	virtual bool ComponentOverlapMulti(TArray<struct FOverlapResult>& OutOverlaps, const class UWorld* World, const FVector& Pos, const FRotator& Rot, ECollisionChannel TestChannel, const struct FComponentQueryParams& Params, const struct FCollisionObjectQueryParams& ObjectQueryParams = FCollisionObjectQueryParams::DefaultObjectQueryParam) const;
 
-	/** Event called when a component is touched */
+	/** 
+	 *	Event called when a component collides with something (or is collided with). 
+	 *	@note For collisions during physics simulation to generate hit events, 'Simulation Generates Hit Events' must be enabled for this component
+	 */
 	UPROPERTY(BlueprintAssignable, Category="Collision")
 	FComponentHitSignature OnComponentHit;
 
-	/** Event called when something overlaps this component */
+	/** 
+	 *	Event called when something starts to overlaps this component. 
+	 *	@note Both this component and the other one must have bGenerateOverlapEvents set to true to generate overlap events.
+	 */
 	UPROPERTY(BlueprintAssignable, Category="Collision")
 	FComponentBeginOverlapSignature OnComponentBeginOverlap;
 
-	/** Event called when something ends overlapping this component */
+	/** 
+	 *	Event called when something stops overlapping this component 
+	 *	@note Both this component and the other one must have bGenerateOverlapEvents set to true to generate overlap events.
+	 */
 	UPROPERTY(BlueprintAssignable, Category="Collision")
 	FComponentEndOverlapSignature OnComponentEndOverlap;
 
diff --git a/Engine/Source/Runtime/Engine/Classes/GameFramework/Actor.h b/Engine/Source/Runtime/Engine/Classes/GameFramework/Actor.h
--- a/Engine/Source/Runtime/Engine/Classes/GameFramework/Actor.h
+++ b/Engine/Source/Runtime/Engine/Classes/GameFramework/Actor.h
@@ -489,11 +489,17 @@ public:
 	UPROPERTY(BlueprintAssignable, Category="Game|Damage")
 	FTakePointDamageSignature OnTakePointDamage;
 	
-	/** Called when another actor begins to overlap this actor. */
+	/** 
+	 *	Called when another actor begins to overlap this actor. 
+	 *	@note Components on both this and the other Actor must have bGenerateOverlapEvents set to true to generate overlap events.
+	 */
 	UPROPERTY(BlueprintAssignable, Category="Collision")
 	FActorBeginOverlapSignature OnActorBeginOverlap;
 
-	/** Called when another actor ends overlap with this actor. */
+	/** 
+	 *	Called when another actor steps overlapping this actor. 
+	 *	@note Components on both this and the other Actor must have bGenerateOverlapEvents set to true to generate overlap events.
+	 */
 	UPROPERTY(BlueprintAssignable, Category="Collision")
 	FActorEndOverlapSignature OnActorEndOverlap;
 
@@ -529,7 +535,10 @@ public:
 	UPROPERTY(BlueprintAssignable, Category="Input|Touch Input")
 	FActorEndTouchOverSignature OnInputTouchLeave;
 
-	/** Called when this actor is involved in a blocking collision. */
+	/** 
+	 *	Called when this Actor hits (or is hit by) something solid. 
+	 *	@note For collisions during physics simulation to generate hit events, 'Simulation Generates Hit Events' must be enabled.
+	 */
 	UPROPERTY(BlueprintAssignable, Category="Collision")
 	FActorHitSignature OnActorHit;
 
@@ -992,11 +1001,17 @@ public:
 	UFUNCTION(BlueprintImplementableEvent, meta=(FriendlyName = "Tick"))
 	virtual void ReceiveTick(float DeltaSeconds);
 
-	/** Event when this actor overlaps another actor. */
+	/** 
+	 *	Event when this actor overlaps another actor. 
+	 *	@note Components on both this and the other Actor must have bGenerateOverlapEvents set to true to generate overlap events.
+	 */
 	UFUNCTION(BlueprintImplementableEvent, meta=(FriendlyName = "ActorBeginOverlap"), Category="Collision")
 	virtual void ReceiveActorBeginOverlap(AActor* OtherActor);
 
-	/** Event when an actor no longer overlaps another actor, and they have separated. */
+	/** 
+	 *	Event when an actor no longer overlaps another actor, and they have separated. 
+	 *	@note Components on both this and the other Actor must have bGenerateOverlapEvents set to true to generate overlap events.
+	 */
 	UFUNCTION(BlueprintImplementableEvent, meta=(FriendlyName = "ActorEndOverlap"), Category="Collision")
 	virtual void ReceiveActorEndOverlap(AActor* OtherActor);
 
@@ -1048,7 +1063,10 @@ public:
 	UFUNCTION(BlueprintCallable, Category="Collision", meta=(UnsafeDuringActorConstruction="true"))
 	void GetOverlappingComponents(TArray<UPrimitiveComponent*>& OverlappingComponents) const;
 
-	/** Event when this actor bumps into a blocking object, or blocks another actor that bumps into it. */
+	/** 
+	 *	Event when this actor bumps into a blocking object, or blocks another actor that bumps into it. 
+	 *	@note For collisions during physics simulation to generate hit events, 'Simulation Generates Hit Events' must be enabled.
+	 */
 	UFUNCTION(BlueprintImplementableEvent, meta=(FriendlyName = "Hit"), Category="Collision")
 	virtual void ReceiveHit(class UPrimitiveComponent* MyComp, AActor* Other, class UPrimitiveComponent* OtherComp, bool bSelfMoved, FVector HitLocation, FVector HitNormal, FVector NormalImpulse, const FHitResult& Hit);
 
@@ -1595,7 +1613,12 @@ public:
 	/** accessor for the value of bCanEverTick */
 	FORCEINLINE bool CanEverTick() const { return PrimaryActorTick.bCanEverTick; }
 
-	/** Called from main actor tick function to implement custom code at the appropriate point in the tick */
+	/** 
+	 *	Function called every frame on this Actor. Override this function to implement custom logic to be executed every frame.
+	 *	Note that Tick is disabled by default, and you will need to check PrimaryActorTick.bCanEverTick is set to true to enable it.
+	 *
+	 *	@param	DeltaSeconds	Game time elapsed since last call to Tick
+	 */
 	virtual void Tick( float DeltaSeconds );
 
 	/** If true, actor is ticked even if TickType==LEVELTICK_ViewportsOnly	 */
