[INTEGRATE] Change 2418671 by [REDACTED_PERSON][REDACTED_HANDLE] on 2015/01/26 11:22:24

	Reparent instance component to the root when its BP created component parent is deleted from the blueprint.

diff --git a/Engine/Source/Runtime/Engine/Private/ComponentInstanceDataCache.cpp b/Engine/Source/Runtime/Engine/Private/ComponentInstanceDataCache.cpp
--- a/Engine/Source/Runtime/Engine/Private/ComponentInstanceDataCache.cpp
+++ b/Engine/Source/Runtime/Engine/Private/ComponentInstanceDataCache.cpp
@@ -77,7 +77,7 @@ FComponentInstanceDataCache::FComponentInstanceDataCache(const AActor* Actor)
 		// Grab per-instance data we want to persist
 		for (UActorComponent* Component : Components)
 		{
-			if(Component->CreationMethod == EComponentCreationMethod::ConstructionScript) // Only cache data from 'created by construction script' components
+			if (Component->CreationMethod == EComponentCreationMethod::ConstructionScript) // Only cache data from 'created by construction script' components
 			{
 				FComponentInstanceDataBase* ComponentInstanceData = Component->GetComponentInstanceData();
 				if (ComponentInstanceData)
@@ -86,6 +86,17 @@ FComponentInstanceDataCache::FComponentInstanceDataCache(const AActor* Actor)
 					TypeToDataMap.Add(Component->GetComponentInstanceDataType(), ComponentInstanceData);
 				}
 			}
+			else if (Component->CreationMethod == EComponentCreationMethod::Instance)
+			{
+				// If the instance component is attached to a BP component we have to be prepared for the possibility that it will be deleted
+				if (USceneComponent* SceneComponent = Cast<USceneComponent>(Component))
+				{
+					if (SceneComponent->AttachParent && SceneComponent->AttachParent->CreationMethod == EComponentCreationMethod::ConstructionScript)
+					{
+						InstanceComponentTransformToRootMap.Add(SceneComponent, SceneComponent->GetComponentTransform().GetRelativeTransform(Actor->GetRootComponent()->GetComponentTransform()));
+					}
+				}
+			}
 		}
 	}
 }
@@ -128,10 +139,23 @@ void FComponentInstanceDataCache::ApplyToActor(AActor* Actor) const
 				}
 			}
 		}
+
+		// Once we're done attaching, if we have any unattached instance components move them to the root
+		for (auto InstanceTransformPair : InstanceComponentTransformToRootMap)
+		{
+			check(Actor->GetRootComponent());
+
+			USceneComponent* SceneComponent = InstanceTransformPair.Key;
+			if (SceneComponent && (SceneComponent->AttachParent == nullptr || SceneComponent->AttachParent->IsPendingKill()))
+			{
+				SceneComponent->AttachTo(Actor->GetRootComponent());
+				SceneComponent->SetRelativeTransform(InstanceTransformPair.Value);
+			}
+		}
 	}
 }
 
-void FComponentInstanceDataCache::FindAndReplaceInstances(const TMap<UObject*, UObject*>& OldToNewInstanceMap) const
+void FComponentInstanceDataCache::FindAndReplaceInstances(const TMap<UObject*, UObject*>& OldToNewInstanceMap)
 {
 	for (auto ComponentInstanceDataPair : TypeToDataMap)
 	{
@@ -140,4 +164,14 @@ void FComponentInstanceDataCache::FindAndReplaceInstances(const TMap<UObject*, U
 			ComponentInstanceDataPair.Value->FindAndReplaceInstances(OldToNewInstanceMap);
 		}
 	}
+	TArray<USceneComponent*> SceneComponents;
+	InstanceComponentTransformToRootMap.GetKeys(SceneComponents);
+
+	for (USceneComponent* SceneComponent : SceneComponents)
+	{
+		if (UObject* const* NewSceneComponent = OldToNewInstanceMap.Find(SceneComponent))
+		{
+			InstanceComponentTransformToRootMap.Add(CastChecked<USceneComponent>(*NewSceneComponent), InstanceComponentTransformToRootMap.FindAndRemoveChecked(SceneComponent));
+		}
+	}
 }
\ No newline at end of file
diff --git a/Engine/Source/Runtime/Engine/Public/ComponentInstanceDataCache.h b/Engine/Source/Runtime/Engine/Public/ComponentInstanceDataCache.h
--- a/Engine/Source/Runtime/Engine/Public/ComponentInstanceDataCache.h
+++ b/Engine/Source/Runtime/Engine/Public/ComponentInstanceDataCache.h
@@ -54,11 +54,13 @@ public:
 	void ApplyToActor(AActor* Actor) const;
 
 	/** Iterates over components and replaces any object references with the reinstanced information */
-	void FindAndReplaceInstances(const TMap<UObject*, UObject*>& OldToNewInstanceMap) const;
+	void FindAndReplaceInstances(const TMap<UObject*, UObject*>& OldToNewInstanceMap);
 
 	bool HasInstanceData() const { return TypeToDataMap.Num() > 0; }
 
 private:
 	/** Map of data type name to data of that type */
 	TMultiMap< FName, FComponentInstanceDataBase* >	TypeToDataMap;
+
+	TMap< USceneComponent*, FTransform > InstanceComponentTransformToRootMap;
 };
\ No newline at end of file
