[INTEGRATE] Change 2418971 by [REDACTED_PERSON][REDACTED_HANDLE] on 2015/01/26 14:21:35

	EngineMinimal.h no longer includes SlateCore.h or SlateBasics.h
	- This nearly halves the time it takes to recompile small game projects! (~14s -> ~8s)
	- This shrinks the PCH size for EngineMinimal by 161 MB on Windows (420 MB -> 259 MB, Development Editor Win64)
	- Note that Engine.h and EnginePrivate.h still include SlateCore.h and SlateBasics.h (not a clear detriment)
	- IMPORTANT: We need to verify all example games that use this still compile, and add release notes for user games that forgot to include Slate directly
	- This is all part of UE-7203

diff --git a/Engine/Source/Developer/FunctionalTesting/Private/FunctionalTestingManager.cpp b/Engine/Source/Developer/FunctionalTesting/Private/FunctionalTestingManager.cpp
--- a/Engine/Source/Developer/FunctionalTesting/Private/FunctionalTestingManager.cpp
+++ b/Engine/Source/Developer/FunctionalTesting/Private/FunctionalTestingManager.cpp
@@ -5,12 +5,10 @@
 
 //----------------------------------------------------------------------//
 // 6/25 @todo these will be removed once marge from main comes
-#include "Editor/UnrealEdTypes.h"
-#include "Editor/EditorEngine.h"
+#include "UnrealEd.h"
 class UFactory;
 //----------------------------------------------------------------------//
 
-#include "Editor.h"
 #endif
 
 #define LOCTEXT_NAMESPACE "FunctionalTesting"
diff --git a/Engine/Source/Developer/GameplayDebugger/Private/GameplayDebuggerPrivate.h b/Engine/Source/Developer/GameplayDebugger/Private/GameplayDebuggerPrivate.h
--- a/Engine/Source/Developer/GameplayDebugger/Private/GameplayDebuggerPrivate.h
+++ b/Engine/Source/Developer/GameplayDebugger/Private/GameplayDebuggerPrivate.h
@@ -39,7 +39,7 @@
 #include "GameplayDebuggingComponent.h"
 
 #if WITH_EDITOR
-#	include "Editor.h" //right now it's needed here because of ECoordSystem 
+#	include "UnrealEd.h" //right now it's needed here because of ECoordSystem 
 #endif
 
 DECLARE_LOG_CATEGORY_EXTERN(LogGDT, Warning, All);
diff --git a/Engine/Source/Editor/Documentation/Public/IDocumentation.h b/Engine/Source/Editor/Documentation/Public/IDocumentation.h
--- a/Engine/Source/Editor/Documentation/Public/IDocumentation.h
+++ b/Engine/Source/Editor/Documentation/Public/IDocumentation.h
@@ -192,5 +192,5 @@ public:
 
 	virtual TSharedRef< class SToolTip > CreateToolTip( const TAttribute<FText>& Text, const TSharedPtr<SWidget>& OverrideContent, const FString& Link, const FString& ExcerptName ) const = 0;
 
-	virtual TSharedRef< class SToolTip > CreateToolTip(const TAttribute<FText>& Text, const TSharedRef<SWidget>& OverrideContent, const TSharedPtr<SVerticalBox>& DocVerticalBox, const FString& Link, const FString& ExcerptName) const = 0;
+	virtual TSharedRef< class SToolTip > CreateToolTip(const TAttribute<FText>& Text, const TSharedRef<SWidget>& OverrideContent, const TSharedPtr<class SVerticalBox>& DocVerticalBox, const FString& Link, const FString& ExcerptName) const = 0;
 };
diff --git a/Engine/Source/Runtime/AIModule/Private/EnvironmentQuery/EnvQueryManager.cpp b/Engine/Source/Runtime/AIModule/Private/EnvironmentQuery/EnvQueryManager.cpp
--- a/Engine/Source/Runtime/AIModule/Private/EnvironmentQuery/EnvQueryManager.cpp
+++ b/Engine/Source/Runtime/AIModule/Private/EnvironmentQuery/EnvQueryManager.cpp
@@ -10,8 +10,8 @@
 #include "EnvironmentQuery/EQSTestingPawn.h"
 #include "EnvironmentQuery/EnvQueryDebugHelpers.h"
 #if WITH_EDITOR
+#include "UnrealEd.h"
 #include "Engine/Brush.h"
-#include "Editor/EditorEngine.h"
 #include "EngineUtils.h"
 
 extern UNREALED_API UEditorEngine* GEditor;
diff --git a/Engine/Source/Runtime/Engine/Classes/Engine/GameViewportClient.h b/Engine/Source/Runtime/Engine/Classes/Engine/GameViewportClient.h
--- a/Engine/Source/Runtime/Engine/Classes/Engine/GameViewportClient.h
+++ b/Engine/Source/Runtime/Engine/Classes/Engine/GameViewportClient.h
@@ -8,8 +8,6 @@
 #include "DebugDisplayProperty.h"
 #include "TitleSafeZone.h"
 #include "GameViewportDelegates.h"
-#include "SlateBasics.h"
-
 #include "GameViewportClient.generated.h"
 
 
@@ -589,10 +587,7 @@ public:
 		return bHideCursorDuringCapture;
 	}
 
-	virtual TOptional<EPopupMethod> OnQueryPopupMethod() const override
-	{
-		return EPopupMethod::UseCurrentWindow;
-	}
+	virtual TOptional<EPopupMethod> OnQueryPopupMethod() const override;
 
 private:
 	/**
diff --git a/Engine/Source/Runtime/Engine/Classes/GameFramework/PlayerController.h b/Engine/Source/Runtime/Engine/Classes/GameFramework/PlayerController.h
--- a/Engine/Source/Runtime/Engine/Classes/GameFramework/PlayerController.h
+++ b/Engine/Source/Runtime/Engine/Classes/GameFramework/PlayerController.h
@@ -12,8 +12,6 @@
 #include "GameFramework/Controller.h"
 #include "Engine/LatentActionManager.h"
 #include "GenericPlatform/IInputInterface.h"
-#include "SlateCore.h"
-
 #include "PlayerController.generated.h"
 
 
@@ -64,10 +62,10 @@ struct ENGINE_API FInputModeDataBase
 {
 protected:
 	/** Derived classes override this function to apply the necessary settings for the desired input mode */
-	virtual void ApplyInputMode(FReply& SlateOperations, class UGameViewportClient& GameViewportClient) const = 0;
+	virtual void ApplyInputMode(class FReply& SlateOperations, class UGameViewportClient& GameViewportClient) const = 0;
 
 	/** Utility functions for derived classes. */
-	void SetFocusAndLocking(FReply& SlateOperations, TSharedPtr<SWidget> InWidgetToFocus, bool bLockMouseToViewport, TSharedRef<class SViewport> InViewportWidget) const;
+	void SetFocusAndLocking(FReply& SlateOperations, TSharedPtr<class SWidget> InWidgetToFocus, bool bLockMouseToViewport, TSharedRef<class SViewport> InViewportWidget) const;
 
 	friend class APlayerController;
 };
@@ -1524,9 +1522,6 @@ public:
 	UPROPERTY()
 	uint16		LastCompletedSeamlessTravelCount;
 
-	/** Stores the last input mode set */
-	//EInputMode  InputMode; right now we don't save this as it could be a half truth.
-
 	/** FReply used to defer some slate operations. */
-	FReply		SlateOperations;
+	TSharedPtr<class FReply> SlateOperations;
 };
diff --git a/Engine/Source/Runtime/Engine/Private/EnginePrivate.h b/Engine/Source/Runtime/Engine/Private/EnginePrivate.h
--- a/Engine/Source/Runtime/Engine/Private/EnginePrivate.h
+++ b/Engine/Source/Runtime/Engine/Private/EnginePrivate.h
@@ -94,3 +94,5 @@ EnginePrivate.h: Unreal engine private header file.
 #include "CanvasTypes.h"							// Canvas.
 #include "EngineUtils.h"
 #include "TimerManager.h"					// Game play timers
+#include "SlateCore.h"
+#include "SlateBasics.h"
diff --git a/Engine/Source/Runtime/Engine/Private/GameViewportClient.cpp b/Engine/Source/Runtime/Engine/Private/GameViewportClient.cpp
--- a/Engine/Source/Runtime/Engine/Private/GameViewportClient.cpp
+++ b/Engine/Source/Runtime/Engine/Private/GameViewportClient.cpp
@@ -2193,6 +2193,11 @@ bool UGameViewportClient::HandleShowCommand( const TCHAR* Cmd, FOutputDevice& Ar
 	return true;
 }
 
+TOptional<EPopupMethod> UGameViewportClient::OnQueryPopupMethod() const
+{
+	return EPopupMethod::UseCurrentWindow;
+}
+
 void UGameViewportClient::ToggleShowCollision()
 {
 	// special case: for the Engine.Collision flag, we need to un-hide any primitive components that collide so their collision geometry gets rendered
diff --git a/Engine/Source/Runtime/Engine/Private/PlayerController.cpp b/Engine/Source/Runtime/Engine/Private/PlayerController.cpp
--- a/Engine/Source/Runtime/Engine/Private/PlayerController.cpp
+++ b/Engine/Source/Runtime/Engine/Private/PlayerController.cpp
@@ -47,7 +47,7 @@ const float RetryServerCheckSpectatorThrottleTime = 0.25f;
 
 APlayerController::APlayerController(const FObjectInitializer& ObjectInitializer)
 	: Super(ObjectInitializer)
-	, SlateOperations(FReply::Unhandled())
+	, SlateOperations(new FReply( FReply::Unhandled() ))
 {
 	NetPriority = 3.0f;
 	CheatClass = UCheatManager::StaticClass();
@@ -4376,7 +4376,7 @@ void APlayerController::SetInputMode(const FInputModeDataBase& InData)
 	UGameViewportClient* GameViewportClient = GetWorld()->GetGameViewport();
 	if (GameViewportClient != nullptr)
 	{
-		InData.ApplyInputMode(SlateOperations, *GameViewportClient);
+		InData.ApplyInputMode(*SlateOperations, *GameViewportClient);
 	}
 }
 
diff --git a/Engine/Source/Runtime/Engine/Public/Engine.h b/Engine/Source/Runtime/Engine/Public/Engine.h
--- a/Engine/Source/Runtime/Engine/Public/Engine.h
+++ b/Engine/Source/Runtime/Engine/Public/Engine.h
@@ -317,3 +317,5 @@
 #include "CanvasTypes.h"							// Canvas.
 #include "EngineUtils.h"
 #include "TimerManager.h"					// Game play timers
+#include "SlateCore.h"
+#include "SlateBasics.h"
diff --git a/Engine/Source/Runtime/Launch/Private/LaunchEngineLoop.cpp b/Engine/Source/Runtime/Launch/Private/LaunchEngineLoop.cpp
--- a/Engine/Source/Runtime/Launch/Private/LaunchEngineLoop.cpp
+++ b/Engine/Source/Runtime/Launch/Private/LaunchEngineLoop.cpp
@@ -2130,7 +2130,7 @@ void FEngineLoop::ProcessPlayerControllersSlateOperations() const
 					APlayerController* PlayerController = LocalPlayer->PlayerController;
 					if (PlayerController != nullptr)
 					{
-						FReply& TheReply = PlayerController->SlateOperations;
+						FReply& TheReply = *PlayerController->SlateOperations;
 
 						FWidgetPath PathToWidget;
 						SlateApp.GeneratePathToWidgetUnchecked(ViewportWidget.ToSharedRef(), PathToWidget);
diff --git a/Engine/Source/Runtime/SlateCore/Public/Textures/TextureAtlas.h b/Engine/Source/Runtime/SlateCore/Public/Textures/TextureAtlas.h
--- a/Engine/Source/Runtime/SlateCore/Public/Textures/TextureAtlas.h
+++ b/Engine/Source/Runtime/SlateCore/Public/Textures/TextureAtlas.h
@@ -231,7 +231,7 @@ public:
 	virtual FIntPoint GetAtlasPageSize() const = 0;
 
 	/** Get the page resource for the given index (verify with GetNumAtlasPages) */ 
-	virtual FSlateShaderResource* GetAtlasPageResource(const int32 InIndex) const = 0;
+	virtual class FSlateShaderResource* GetAtlasPageResource(const int32 InIndex) const = 0;
 
 	/** Do the atlas page resources only contain alpha information? This affects how the atlas visualizer will sample them */
 	virtual bool IsAtlasPageResourceAlphaOnly() const = 0;
