[INTEGRATE] Change 2419212 by [REDACTED_PERSON][REDACTED_HANDLE] on 2015/01/26 16:56:49

	Avoid unnecessary recursive calls to RegisterComponent

diff --git a/Engine/Source/Runtime/Engine/Private/SimpleConstructionScript.cpp b/Engine/Source/Runtime/Engine/Private/SimpleConstructionScript.cpp
--- a/Engine/Source/Runtime/Engine/Private/SimpleConstructionScript.cpp
+++ b/Engine/Source/Runtime/Engine/Private/SimpleConstructionScript.cpp
@@ -19,15 +19,18 @@ namespace
 	// Helper method to register instanced components post-construction
 	void RegisterInstancedComponent(UActorComponent* InstancedComponent)
 	{
-		InstancedComponent->RegisterComponent();
-
-		// If this is a scene component, recursively register any child components as well
-		USceneComponent* InstancedSceneComponent = Cast<USceneComponent>(InstancedComponent);
-		if(InstancedSceneComponent != nullptr)
+		if (!InstancedComponent->IsRegistered())
 		{
-			for(auto InstancedChildComponent : InstancedSceneComponent->AttachChildren)
+			InstancedComponent->RegisterComponent();
+
+			// If this is a scene component, recursively register any child components as well
+			USceneComponent* InstancedSceneComponent = Cast<USceneComponent>(InstancedComponent);
+			if(InstancedSceneComponent != nullptr)
 			{
-				RegisterInstancedComponent(InstancedChildComponent);
+				for(auto InstancedChildComponent : InstancedSceneComponent->AttachChildren)
+				{
+					RegisterInstancedComponent(InstancedChildComponent);
+				}
 			}
 		}
 	}
