[INTEGRATE] Change 2419357 by [REDACTED_PERSON][REDACTED_HANDLE] on 2015/01/26 17:56:52

	Store the component class in the component instance data to make the serialized index work correctly

diff --git a/Engine/Source/Runtime/Engine/Private/ComponentInstanceDataCache.cpp b/Engine/Source/Runtime/Engine/Private/ComponentInstanceDataCache.cpp
--- a/Engine/Source/Runtime/Engine/Private/ComponentInstanceDataCache.cpp
+++ b/Engine/Source/Runtime/Engine/Private/ComponentInstanceDataCache.cpp
@@ -7,6 +7,7 @@ FComponentInstanceDataBase::FComponentInstanceDataBase(const UActorComponent* So
 {
 	check(SourceComponent);
 	SourceComponentName = SourceComponent->GetFName();
+	SourceComponentClass = SourceComponent->GetClass();
 	SourceComponentTypeSerializedIndex = -1;
 	
 	AActor* ComponentOwner = SourceComponent->GetOwner();
@@ -23,7 +24,7 @@ FComponentInstanceDataBase::FComponentInstanceDataBase(const UActorComponent* So
 					bFound = true;
 					break;
 				}
-				else if (BlueprintCreatedComponent->GetClass() == SourceComponent->GetClass())
+				else if (BlueprintCreatedComponent->GetClass() == SourceComponentClass)
 				{
 					++SourceComponentTypeSerializedIndex;
 				}
@@ -39,7 +40,7 @@ FComponentInstanceDataBase::FComponentInstanceDataBase(const UActorComponent* So
 bool FComponentInstanceDataBase::MatchesComponent(const UActorComponent* Component) const
 {
 	bool bMatches = false;
-	if (Component)
+	if (Component && Component->GetClass() == SourceComponentClass)
 	{
 		if (Component->GetFName() == SourceComponentName)
 		{
@@ -54,7 +55,7 @@ bool FComponentInstanceDataBase::MatchesComponent(const UActorComponent* Compone
 				for (const UActorComponent* BlueprintCreatedComponent : ComponentOwner->BlueprintCreatedComponents)
 				{
 					if (   BlueprintCreatedComponent
-						&& (BlueprintCreatedComponent->GetClass() == Component->GetClass())
+						&& (BlueprintCreatedComponent->GetClass() == SourceComponentClass)
 						&& (++FoundSerializedComponentsOfType == SourceComponentTypeSerializedIndex))
 					{
 						bMatches = (BlueprintCreatedComponent == Component);
diff --git a/Engine/Source/Runtime/Engine/Public/ComponentInstanceDataCache.h b/Engine/Source/Runtime/Engine/Public/ComponentInstanceDataCache.h
--- a/Engine/Source/Runtime/Engine/Public/ComponentInstanceDataCache.h
+++ b/Engine/Source/Runtime/Engine/Public/ComponentInstanceDataCache.h
@@ -10,7 +10,8 @@ class ENGINE_API FComponentInstanceDataBase
 {
 public:
 	FComponentInstanceDataBase()
-		: SourceComponentTypeSerializedIndex(-1)
+		: SourceComponentClass(nullptr)
+		, SourceComponentTypeSerializedIndex(-1)
 	{}
 
 	FComponentInstanceDataBase(const UActorComponent* SourceComponent);
@@ -31,6 +32,9 @@ protected:
 	/** The name of the source component */
 	FName SourceComponentName;
 
+	/** The class type of the source component */
+	UClass* SourceComponentClass;
+
 	/** The index of the source component in its owner's serialized array 
 		when filtered to just that component type */
 	int32 SourceComponentTypeSerializedIndex;
