[INTEGRATE] Change 2419448 by [REDACTED_PERSON][REDACTED_HANDLE] on 2015/01/26 18:39:27

	Fix potential crash in certain cases related to the billboard component of an Audio/PhysicsConstraint instance component

diff --git a/Engine/Source/Runtime/Engine/Private/ComponentInstanceDataCache.cpp b/Engine/Source/Runtime/Engine/Private/ComponentInstanceDataCache.cpp
--- a/Engine/Source/Runtime/Engine/Private/ComponentInstanceDataCache.cpp
+++ b/Engine/Source/Runtime/Engine/Private/ComponentInstanceDataCache.cpp
@@ -172,7 +172,14 @@ void FComponentInstanceDataCache::FindAndReplaceInstances(const TMap<UObject*, U
 	{
 		if (UObject* const* NewSceneComponent = OldToNewInstanceMap.Find(SceneComponent))
 		{
-			InstanceComponentTransformToRootMap.Add(CastChecked<USceneComponent>(*NewSceneComponent), InstanceComponentTransformToRootMap.FindAndRemoveChecked(SceneComponent));
+			if (*NewSceneComponent)
+			{
+				InstanceComponentTransformToRootMap.Add(CastChecked<USceneComponent>(*NewSceneComponent), InstanceComponentTransformToRootMap.FindAndRemoveChecked(SceneComponent));
+			}
+			else
+			{
+				InstanceComponentTransformToRootMap.Remove(SceneComponent);
+			}
 		}
 	}
 }
\ No newline at end of file
diff --git a/Engine/Source/Runtime/Engine/Private/SceneComponent.cpp b/Engine/Source/Runtime/Engine/Private/SceneComponent.cpp
--- a/Engine/Source/Runtime/Engine/Private/SceneComponent.cpp
+++ b/Engine/Source/Runtime/Engine/Private/SceneComponent.cpp
@@ -1031,7 +1031,10 @@ void FSceneComponentInstanceData::ApplyToComponent(UActorComponent* Component)
 	USceneComponent* SceneComponent = CastChecked<USceneComponent>(Component);
 	for (USceneComponent* ChildComponent : AttachedInstanceComponents)
 	{
-		ChildComponent->AttachTo(SceneComponent);
+		if (ChildComponent)
+		{
+			ChildComponent->AttachTo(SceneComponent);
+		}
 	}
 }
 
@@ -1041,7 +1044,7 @@ void FSceneComponentInstanceData::FindAndReplaceInstances(const TMap<UObject*, U
 	{
 		if (UObject* const* NewChildComponent = OldToNewInstanceMap.Find(ChildComponent))
 		{
-			ChildComponent = CastChecked<USceneComponent>(*NewChildComponent);
+			ChildComponent = CastChecked<USceneComponent>(*NewChildComponent, ECastCheckedType::NullAllowed);
 		}
 	}
 }
