#UE4 Fixed a compile error regarding struct parameter alignment in Win32 builds.

diff --git a/Engine/Source/Runtime/Engine/Private/VisualLogger/VisualLoggerTypes.cpp b/Engine/Source/Runtime/Engine/Private/VisualLogger/VisualLoggerTypes.cpp
--- a/Engine/Source/Runtime/Engine/Private/VisualLogger/VisualLoggerTypes.cpp
+++ b/Engine/Source/Runtime/Engine/Private/VisualLogger/VisualLoggerTypes.cpp
@@ -111,9 +111,9 @@ void FVisualLogEntry::AddText(const FString& TextLine, const FName& CategoryName
 	LogLines.Add(FVisualLogLine(CategoryName, Verbosity, TextLine));
 }
 
-void FVisualLogEntry::AddElement(FVisualLogShapeElement Element)
+void FVisualLogEntry::AddElement(const FVisualLogShapeElement& Element)
 {
-	ElementsToDraw.Emplace(Element);
+	ElementsToDraw.Add(Element);
 }
 
 void FVisualLogEntry::AddElement(const TArray<FVector>& Points, const FName& CategoryName, ELogVerbosity::Type Verbosity, const FColor& Color, const FString& Description, uint16 Thickness)
diff --git a/Engine/Source/Runtime/Engine/Public/VisualLogger/VisualLoggerTypes.h b/Engine/Source/Runtime/Engine/Public/VisualLogger/VisualLoggerTypes.h
--- a/Engine/Source/Runtime/Engine/Public/VisualLogger/VisualLoggerTypes.h
+++ b/Engine/Source/Runtime/Engine/Public/VisualLogger/VisualLoggerTypes.h
@@ -190,7 +190,7 @@ struct ENGINE_API FVisualLogEntry
 	// capsule
 	void AddElement(const FVector& Center, float HalfHeight, float Radius, const FQuat & Rotation, const FName& CategoryName, ELogVerbosity::Type Verbosity, const FColor& Color = FColor::White, const FString& Description = TEXT(""));
 	// custom element
-	void AddElement(FVisualLogShapeElement Element);
+	void AddElement(const FVisualLogShapeElement& Element);
 	// histogram sample
 	void AddHistogramData(const FVector2D& DataSample, const FName& CategoryName, ELogVerbosity::Type Verbosity, const FName& GraphName, const FName& DataName);
 	// Custom data block
