[INTEGRATE] Change 2419939 by [REDACTED_PERSON][REDACTED_HANDLE] on 2015/01/27 00:06:35

	Export base implementation of GetAssetRegistryTags to fix build error.

diff --git a/Engine/Source/Runtime/Engine/Classes/Engine/DataTable.h b/Engine/Source/Runtime/Engine/Classes/Engine/DataTable.h
--- a/Engine/Source/Runtime/Engine/Classes/Engine/DataTable.h
+++ b/Engine/Source/Runtime/Engine/Classes/Engine/DataTable.h
@@ -52,7 +52,7 @@ class UDataTable
 	ENGINE_API virtual void Serialize(FArchive& Ar) override;
 	ENGINE_API static void AddReferencedObjects(UObject* InThis, FReferenceCollector& Collector);
 #if WITH_EDITORONLY_DATA
-	virtual void GetAssetRegistryTags(TArray<FAssetRegistryTag>& OutTags) const;
+	ENGINE_API virtual void GetAssetRegistryTags(TArray<FAssetRegistryTag>& OutTags) const;
 #endif
 	// End  UObject interface
 
