Fixed race condition on instanced render data destruction

diff --git a/Engine/Source/Runtime/Engine/Private/InstancedStaticMesh.cpp b/Engine/Source/Runtime/Engine/Private/InstancedStaticMesh.cpp
--- a/Engine/Source/Runtime/Engine/Private/InstancedStaticMesh.cpp
+++ b/Engine/Source/Runtime/Engine/Private/InstancedStaticMesh.cpp
@@ -1147,19 +1147,16 @@ void UInstancedStaticMeshComponent::ReleasePerInstanceRenderData()
 {
 	if (PerInstanceRenderData.IsValid() && !bPerInstanceRenderDataWasPrebuilt)
 	{
-		if (SceneProxy == nullptr)
+		typedef TSharedPtr<FPerInstanceRenderData, ESPMode::ThreadSafe> FPerInstanceRenderDataPtr;
+		// Move shared pointer to a render task, resource will be released by scene proxy or render task, whoever will get executed last
+		ENQUEUE_UNIQUE_RENDER_COMMAND_ONEPARAMETER(
+			FReleasePerInstanceRenderData,
+			FPerInstanceRenderDataPtr,InPerInstanceRenderData, FPerInstanceRenderDataPtr(MoveTemp(PerInstanceRenderData)),
 		{
-			typedef TSharedPtr<FPerInstanceRenderData, ESPMode::ThreadSafe> FPerInstanceRenderDataPtr;
-			// Release PerInstanceRenderData on a render thread
-			ENQUEUE_UNIQUE_RENDER_COMMAND_ONEPARAMETER(
-				FReleasePerInstanceRenderData,
-				FPerInstanceRenderDataPtr,InPerInstanceRenderData,PerInstanceRenderData,
-			{
-				InPerInstanceRenderData.Reset();
-			});
-		}
-		// Will be released by scene proxy or render thread command
-		PerInstanceRenderData.Reset();
+			InPerInstanceRenderData.Reset();
+		});
+		// At this point we should not have a reference to render data
+		check(!PerInstanceRenderData.IsValid());
 	}
 }
 
