[INTEGRATE] Change 2420219 by [REDACTED_PERSON][REDACTED_HANDLE] on 2015/01/27 08:40:26

	Fix for compile errors on XboxOne. VS2012 compiler doesn't support constructor forwarding.

diff --git a/Engine/Source/Runtime/Engine/Classes/Components/ActorComponent.h b/Engine/Source/Runtime/Engine/Classes/Components/ActorComponent.h
--- a/Engine/Source/Runtime/Engine/Classes/Components/ActorComponent.h
+++ b/Engine/Source/Runtime/Engine/Classes/Components/ActorComponent.h
@@ -40,6 +40,11 @@ public:
 	 */
 	UActorComponent(const FObjectInitializer& ObjectInitializer);
 
+private:
+	/** Called from the constructor to initialize the object to its default settings */
+	void InitializeDefaults();
+
+public:
 	virtual void GetLifetimeReplicatedProps(TArray<FLifetimeProperty>& OutLifetimeProps) const override;
 
 	/** Main tick function for the Actor */
diff --git a/Engine/Source/Runtime/Engine/Classes/Components/SceneComponent.h b/Engine/Source/Runtime/Engine/Classes/Components/SceneComponent.h
--- a/Engine/Source/Runtime/Engine/Classes/Components/SceneComponent.h
+++ b/Engine/Source/Runtime/Engine/Classes/Components/SceneComponent.h
@@ -92,6 +92,12 @@ public:
 	 */
 	USceneComponent(const FObjectInitializer& ObjectInitializer);
 
+private:
+	/** Initialize the component to its default settings */
+	void InitializeDefaults();
+
+public:
+
 	virtual void GetLifetimeReplicatedProps(TArray<FLifetimeProperty>& OutLifetimeProps) const override;
 
 	/** Current transform of this component, relative to the world */
diff --git a/Engine/Source/Runtime/Engine/Classes/GameFramework/Actor.h b/Engine/Source/Runtime/Engine/Classes/GameFramework/Actor.h
--- a/Engine/Source/Runtime/Engine/Classes/GameFramework/Actor.h
+++ b/Engine/Source/Runtime/Engine/Classes/GameFramework/Actor.h
@@ -80,6 +80,11 @@ public:
 	 */
 	AActor(const FObjectInitializer& ObjectInitializer);
 
+private:
+	/** Called from the constructor to initialize the class to its default settings */
+	void InitializeDefaults();
+
+public:
 	void GetLifetimeReplicatedProps(TArray<FLifetimeProperty>& OutLifetimeProps) const override;
 
 	/**
diff --git a/Engine/Source/Runtime/Engine/Private/Actor.cpp b/Engine/Source/Runtime/Engine/Private/Actor.cpp
--- a/Engine/Source/Runtime/Engine/Private/Actor.cpp
+++ b/Engine/Source/Runtime/Engine/Private/Actor.cpp
@@ -46,6 +46,18 @@ FOnProcessEvent AActor::ProcessEventDelegate;
 #endif
 
 AActor::AActor()
+{
+	InitializeDefaults();
+}
+
+
+AActor::AActor(const FObjectInitializer& ObjectInitializer)
+{
+	// Forward to default constructor (we don't use ObjectInitializer for anything, this is for compatibility with inherited classes that call Super( ObjectInitializer )
+	InitializeDefaults();
+}
+
+void AActor::InitializeDefaults()
 {
 	PrimaryActorTick.TickGroup = TG_PrePhysics;
 	// Default to no tick function, but if we set 'never ticks' to false (so there is a tick function) it is enabled by default
@@ -84,12 +96,6 @@ AActor::AActor()
 	bAllowReceiveTickEventOnDedicatedServer = true;
 }
 
-
-AActor::AActor(const FObjectInitializer& ObjectInitializer)
-	: AActor()	// Forward to default constructor (we don't use ObjectInitializer for anything, this is for compatibility with inherited classes that call Super( ObjectInitializer )
-{
-}
-
 void FActorTickFunction::ExecuteTick(float DeltaTime, enum ELevelTick TickType, ENamedThreads::Type CurrentThread, const FGraphEventRef& MyCompletionGraphEvent)
 {
 	if (Target && !Target->HasAnyFlags(RF_PendingKill | RF_Unreachable))
diff --git a/Engine/Source/Runtime/Engine/Private/ActorComponent.cpp b/Engine/Source/Runtime/Engine/Private/ActorComponent.cpp
--- a/Engine/Source/Runtime/Engine/Private/ActorComponent.cpp
+++ b/Engine/Source/Runtime/Engine/Private/ActorComponent.cpp
@@ -100,6 +100,18 @@ FGlobalComponentReregisterContext::~FGlobalComponentReregisterContext()
 
 
 UActorComponent::UActorComponent()
+{
+	InitializeDefaults();
+}
+
+
+UActorComponent::UActorComponent( const FObjectInitializer& ObjectInitializer )
+{
+	// Forward to default constructor (we don't use ObjectInitializer for anything, this is for compatibility with inherited classes that call Super( ObjectInitializer )
+	InitializeDefaults();
+}
+
+void UActorComponent::InitializeDefaults()
 {
 	PrimaryComponentTick.TickGroup = TG_DuringPhysics;
 	PrimaryComponentTick.bStartWithTickEnabled = true;
@@ -111,13 +123,6 @@ UActorComponent::UActorComponent()
 	bNetAddressable = false;
 }
 
-
-UActorComponent::UActorComponent( const FObjectInitializer& ObjectInitializer )
-	: UActorComponent()	// Forward to default constructor (we don't use ObjectInitializer for anything, this is for compatibility with inherited classes that call Super( ObjectInitializer )
-{
-}
-
-
 void UActorComponent::PostInitProperties()
 {
 	Super::PostInitProperties();
diff --git a/Engine/Source/Runtime/Engine/Private/SceneComponent.cpp b/Engine/Source/Runtime/Engine/Private/SceneComponent.cpp
--- a/Engine/Source/Runtime/Engine/Private/SceneComponent.cpp
+++ b/Engine/Source/Runtime/Engine/Private/SceneComponent.cpp
@@ -24,6 +24,17 @@ FOverlapInfo::FOverlapInfo(UPrimitiveComponent* InComponent, int32 InBodyIndex)
 }
 
 USceneComponent::USceneComponent()
+{
+	InitializeDefaults();
+}
+
+USceneComponent::USceneComponent( const FObjectInitializer& ObjectInitializer )
+{
+	// Forward to default constructor (we don't use ObjectInitializer for anything, this is for compatibility with inherited classes that call Super( ObjectInitializer )
+	InitializeDefaults();
+}
+
+void USceneComponent::InitializeDefaults()
 {
 	Mobility = EComponentMobility::Movable;
 	RelativeScale3D = FVector(1.0f,1.0f,1.0f);
@@ -34,13 +45,6 @@ USceneComponent::USceneComponent()
 	NetUpdateTransform = false;
 }
 
-
-USceneComponent::USceneComponent( const FObjectInitializer& ObjectInitializer )
-	: USceneComponent()	// Forward to default constructor (we don't use ObjectInitializer for anything, this is for compatibility with inherited classes that call Super( ObjectInitializer )
-{
-}
-
-
 FTransform USceneComponent::CalcNewComponentToWorld(const FTransform& NewRelativeTransform, const USceneComponent* Parent) const
 {
 	Parent = Parent ? Parent : AttachParent;
