UE4 - add show flags for foliage and grass

diff --git a/Engine/Source/Editor/DetailCustomizations/Private/SceneCaptureDetails.cpp b/Engine/Source/Editor/DetailCustomizations/Private/SceneCaptureDetails.cpp
--- a/Engine/Source/Editor/DetailCustomizations/Private/SceneCaptureDetails.cpp
+++ b/Engine/Source/Editor/DetailCustomizations/Private/SceneCaptureDetails.cpp
@@ -78,6 +78,8 @@ void FSceneCaptureDetails::CustomizeDetails( IDetailLayoutBuilder& DetailLayout
 	ShowFlagsToAllowForCaptures.Add(FEngineShowFlags::EShowFlag::SF_Translucency);
 	ShowFlagsToAllowForCaptures.Add(FEngineShowFlags::EShowFlag::SF_DeferredLighting);
 	ShowFlagsToAllowForCaptures.Add(FEngineShowFlags::EShowFlag::SF_InstancedStaticMeshes);
+	ShowFlagsToAllowForCaptures.Add(FEngineShowFlags::EShowFlag::SF_InstancedFoliage);
+	ShowFlagsToAllowForCaptures.Add(FEngineShowFlags::EShowFlag::SF_InstancedGrass);
 	ShowFlagsToAllowForCaptures.Add(FEngineShowFlags::EShowFlag::SF_Paper2DSprites);
 	ShowFlagsToAllowForCaptures.Add(FEngineShowFlags::EShowFlag::SF_TextRender);
 	ShowFlagsToAllowForCaptures.Add(FEngineShowFlags::EShowFlag::SF_AmbientOcclusion);
diff --git a/Engine/Source/Runtime/Engine/Private/HierarchicalInstancedStaticMesh.cpp b/Engine/Source/Runtime/Engine/Private/HierarchicalInstancedStaticMesh.cpp
--- a/Engine/Source/Runtime/Engine/Private/HierarchicalInstancedStaticMesh.cpp
+++ b/Engine/Source/Runtime/Engine/Private/HierarchicalInstancedStaticMesh.cpp
@@ -582,6 +582,8 @@ class FHierarchicalStaticMeshSceneProxy : public FInstancedStaticMeshSceneProxy
 	int32 FirstUnbuiltIndex;
 	int32 LastUnbuiltIndex;
 
+	bool bIsGrass;
+
 #if !(UE_BUILD_SHIPPING || UE_BUILD_TEST)
 	mutable TArray<uint32> SingleDebugRuns[MAX_STATIC_MESH_LODS];
 	mutable int32 SingleDebugTotalInstances[MAX_STATIC_MESH_LODS];
@@ -592,12 +594,13 @@ class FHierarchicalStaticMeshSceneProxy : public FInstancedStaticMeshSceneProxy
 
 public:
 
-	FHierarchicalStaticMeshSceneProxy(UHierarchicalInstancedStaticMeshComponent* InComponent, ERHIFeatureLevel::Type InFeatureLevel)
+	FHierarchicalStaticMeshSceneProxy(bool bInIsGrass, UHierarchicalInstancedStaticMeshComponent* InComponent, ERHIFeatureLevel::Type InFeatureLevel)
 		: FInstancedStaticMeshSceneProxy(InComponent, InFeatureLevel)
 		, ClusterTreePtr(InComponent->ClusterTreePtr)
 		, ClusterTree(*InComponent->ClusterTreePtr)
 		, FirstUnbuiltIndex(InComponent->NumBuiltInstances)
 		, LastUnbuiltIndex(InComponent->PerInstanceSMData.Num()+InComponent->RemovedInstances.Num()-1)
+		, bIsGrass(bInIsGrass)
 #if !(UE_BUILD_SHIPPING || UE_BUILD_TEST)
 		, CaptureTag(0)
 #endif
@@ -605,12 +608,13 @@ public:
 		check(InComponent->InstanceReorderTable.Num() == InComponent->PerInstanceSMData.Num());
 	}
 
-	FHierarchicalStaticMeshSceneProxy(UHierarchicalInstancedStaticMeshComponent* InComponent, ERHIFeatureLevel::Type InFeatureLevel, FStaticMeshInstanceData& Other)
+	FHierarchicalStaticMeshSceneProxy(bool bInIsGrass, UHierarchicalInstancedStaticMeshComponent* InComponent, ERHIFeatureLevel::Type InFeatureLevel, FStaticMeshInstanceData& Other)
 		: FInstancedStaticMeshSceneProxy(InComponent, InFeatureLevel, Other)
 		, ClusterTreePtr(InComponent->ClusterTreePtr)
 		, ClusterTree(*InComponent->ClusterTreePtr)
 		, FirstUnbuiltIndex(InComponent->NumBuiltInstances)
 		, LastUnbuiltIndex(InstancedRenderData.NumInstances+InComponent->RemovedInstances.Num()-1)
+		, bIsGrass(bInIsGrass)
 #if !(UE_BUILD_SHIPPING || UE_BUILD_TEST)
 		, CaptureTag(0)
 #endif
@@ -623,7 +627,7 @@ public:
 	virtual FPrimitiveViewRelevance GetViewRelevance(const FSceneView* View) override
 	{
 		FPrimitiveViewRelevance Result;
-		if(View->Family->EngineShowFlags.InstancedStaticMeshes)
+		if (bIsGrass ? View->Family->EngineShowFlags.InstancedGrass : View->Family->EngineShowFlags.InstancedFoliage)
 		{
 			Result = FStaticMeshSceneProxy::GetViewRelevance(View);
 			Result.bDynamicRelevance = true;
@@ -1896,11 +1900,11 @@ FPrimitiveSceneProxy* UHierarchicalInstancedStaticMeshComponent::CreateSceneProx
 			ProxySize = FStaticMeshInstanceData::GetResourceSize(WriteOncePrebuiltInstanceBuffer.Num());
 			INC_DWORD_STAT_BY(STAT_FoliageInstanceBuffers, ProxySize);
 
-			return ::new FHierarchicalStaticMeshSceneProxy(this, GetWorld()->FeatureLevel, WriteOncePrebuiltInstanceBuffer);
+			return ::new FHierarchicalStaticMeshSceneProxy(true, this, GetWorld()->FeatureLevel, WriteOncePrebuiltInstanceBuffer);
 		}
 		ProxySize = FStaticMeshInstanceData::GetResourceSize(PerInstanceSMData.Num());
 		INC_DWORD_STAT_BY(STAT_FoliageInstanceBuffers, ProxySize);
-		return ::new FHierarchicalStaticMeshSceneProxy(this, GetWorld()->FeatureLevel);
+		return ::new FHierarchicalStaticMeshSceneProxy(false, this, GetWorld()->FeatureLevel);
 	}
 	return NULL;
 }
diff --git a/Engine/Source/Runtime/Engine/Public/ShowFlagsValues.inl b/Engine/Source/Runtime/Engine/Public/ShowFlagsValues.inl
--- a/Engine/Source/Runtime/Engine/Public/ShowFlagsValues.inl
+++ b/Engine/Source/Runtime/Engine/Public/ShowFlagsValues.inl
@@ -128,8 +128,12 @@ SHOWFLAG_ALWAYS_ACCESSIBLE(TexturedLightProfiles, SFG_LightingFeatures, LOCTEXT(
 SHOWFLAG_ALWAYS_ACCESSIBLE(LightFunctions, SFG_LightingFeatures, LOCTEXT("LightFunctionsSF", "Light Functions"))
 /** Hardware Tessellation (DX11 feature) */
 SHOWFLAG_ALWAYS_ACCESSIBLE(Tessellation, SFG_Advanced, LOCTEXT("TessellationSF", "Tessellation"))
-/** Draws instanced static meshes such as foliage. */
+/** Draws instanced static meshes that are not foliage or grass. */
 SHOWFLAG_ALWAYS_ACCESSIBLE(InstancedStaticMeshes, SFG_Advanced, LOCTEXT("InstancedStaticMeshesSF", "Instanced Static Meshes"))
+/** Draws instanced foliage. */
+SHOWFLAG_ALWAYS_ACCESSIBLE(InstancedFoliage, SFG_Advanced, LOCTEXT("InstancedFoliageSF", "Foliage"))
+/** Draws instanced grass. */
+SHOWFLAG_ALWAYS_ACCESSIBLE(InstancedGrass, SFG_Advanced, LOCTEXT("InstancedGrassSF", "Grass"))
 /** non baked shadows */
 SHOWFLAG_ALWAYS_ACCESSIBLE(DynamicShadows, SFG_LightingComponents, LOCTEXT("DynamicShadowsSF", "Dynamic Shadows"))
 /** Particles */
