Made AIController redirect all of its components' vlogs to his own #UE4

While at it added a safety check for vlog redirects in case an object has already been redirected

diff --git a/Engine/Source/Runtime/AIModule/Private/AIController.cpp b/Engine/Source/Runtime/AIModule/Private/AIController.cpp
--- a/Engine/Source/Runtime/AIModule/Private/AIController.cpp
+++ b/Engine/Source/Runtime/AIModule/Private/AIController.cpp
@@ -61,6 +61,15 @@ void AAIController::PostInitializeComponents()
 	{
 		InitPlayerState();
 	}
+
+#if ENABLE_VISUAL_LOG
+	TArray<UActorComponent*> ComponentSet;
+	GetComponents(ComponentSet);
+	for (auto Component : ComponentSet)
+	{
+		REDIRECT_OBJECT_TO_VLOG(Component, this);
+	}
+#endif // ENABLE_VISUAL_LOG
 }
 
 void AAIController::PostRegisterAllComponents()
diff --git a/Engine/Source/Runtime/Engine/Private/VisualLogger/VisualLogger.cpp b/Engine/Source/Runtime/Engine/Private/VisualLogger/VisualLogger.cpp
--- a/Engine/Source/Runtime/Engine/Private/VisualLogger/VisualLogger.cpp
+++ b/Engine/Source/Runtime/Engine/Private/VisualLogger/VisualLogger.cpp
@@ -98,23 +98,26 @@ int32 FVisualLogger::GetUniqueId(float Timestamp)
 
 void FVisualLogger::Redirect(UObject* FromObject, UObject* ToObject)
 {
-	if (FromObject == ToObject)
+	if (FromObject == ToObject || FromObject == nullptr || ToObject == nullptr)
 	{
 		return;
 	}
 
 	UObject* OldRedirection = FindRedirection(FromObject);
-	UObject* NewRedierection = FindRedirection(ToObject);
+	UObject* NewRedirection = FindRedirection(ToObject);
 
-	auto OldArray = RedirectionMap.Find(OldRedirection);
-	if (OldArray)
+	if (OldRedirection != NewRedirection)
 	{
-		OldArray->RemoveSingleSwap(FromObject);
-	}
+		auto OldArray = RedirectionMap.Find(OldRedirection);
+		if (OldArray)
+		{
+			OldArray->RemoveSingleSwap(FromObject);
+		}
 
-	RedirectionMap.FindOrAdd(NewRedierection).AddUnique(FromObject);
+		RedirectionMap.FindOrAdd(NewRedirection).AddUnique(FromObject);
 
-	UE_CVLOG(FromObject != NULL, FromObject, LogVisual, Log, TEXT("Redirected '%s' to '%s'"), *FromObject->GetName(), *NewRedierection->GetName());
+		UE_CVLOG(FromObject != NULL, FromObject, LogVisual, Log, TEXT("Redirected '%s' to '%s'"), *FromObject->GetName(), *NewRedirection->GetName());
+	}
 }
 
 class UObject* FVisualLogger::FindRedirection(const UObject* Object)
