Fixed PathFollowingComponent not getting notified about path updates due to external events (like goal moved or navmesh rebuild) #UE4

- While working on it I've also fixed GameplayDebugger's path corridor gathering to clear porperly and to be able to detect when AI uses the same path object, but the path has been updated since the last time its navigation corridor has been stored by GameplayDebugger.
- also includes minor refactor to navmesh debug drawing

diff --git a/Engine/Source/Developer/GameplayDebugger/Classes/GameplayDebuggingComponent.h b/Engine/Source/Developer/GameplayDebugger/Classes/GameplayDebuggingComponent.h
--- a/Engine/Source/Developer/GameplayDebugger/Classes/GameplayDebuggingComponent.h
+++ b/Engine/Source/Developer/GameplayDebugger/Classes/GameplayDebuggingComponent.h
@@ -246,6 +246,7 @@ protected:
 	virtual void CollectBasicAnimationData(APawn* MyPawn);
 
 	FNavPathWeakPtr CurrentPath;
+	float LastStoredPathTimeStamp;
 
 	uint32 bEnabledTargetSelection : 1;
 #if WITH_EDITOR
diff --git a/Engine/Source/Developer/GameplayDebugger/Private/GameplayDebuggingComponent.cpp b/Engine/Source/Developer/GameplayDebugger/Private/GameplayDebuggingComponent.cpp
--- a/Engine/Source/Developer/GameplayDebugger/Private/GameplayDebuggingComponent.cpp
+++ b/Engine/Source/Developer/GameplayDebugger/Private/GameplayDebuggingComponent.cpp
@@ -136,6 +136,8 @@ UGameplayDebuggingComponent::UGameplayDebuggingComponent(const FObjectInitialize
 	bAutoActivate = false;
 	PrimaryComponentTick.bStartWithTickEnabled = false;
 
+	LastStoredPathTimeStamp = -1.f;
+
 	ShowExtendedInformatiomCounter = 0;
 #if WITH_EDITOR
 	bWasSelectedInEditor = false;
@@ -472,7 +474,7 @@ void UGameplayDebuggingComponent::CollectBasicPathData(APawn* MyPawn)
 	AAIController* MyAIController = Cast<AAIController>(MyPawn->GetController());
 
 	const ANavigationData* NavData = NavSys->GetNavDataForProps(MyAIController->GetNavAgentPropertiesRef());
-	NavDataInfo = NavData->NavDataConfig.Name.ToString();
+	NavDataInfo = NavData->GetConfig().Name.ToString();
 
 	UPathFollowingComponent* PFC = MyAIController->GetPathFollowingComponent();
 	bIsUsingPathFollowing = (PFC != nullptr);
@@ -663,11 +665,13 @@ void UGameplayDebuggingComponent::CollectPathData()
 			NextPathPointIndex = MyController->PathFollowingComponent->GetNextPathIndex();
 
 			const FNavPathSharedPtr& NewPath = MyController->PathFollowingComponent->GetPath();
-			if (!CurrentPath.HasSameObject(NewPath.Get()))
+			if (CurrentPath.HasSameObject(NewPath.Get()) == false || NewPath->GetTimeStamp() > LastStoredPathTimeStamp)
 			{
+				LastStoredPathTimeStamp = NewPath->GetTimeStamp();
+
 				FVisualLogEntry Snapshot;
 				NewPath->DescribeSelfToVisLog(&Snapshot);
-				PathCorridorPolygons.Reserve(Snapshot.ElementsToDraw.Num());
+				PathCorridorPolygons.Reset(Snapshot.ElementsToDraw.Num());
 				for (auto& CurrentShape : Snapshot.ElementsToDraw)
 				{
 					if (CurrentShape.GetType() == EVisualLoggerShapeElement::Polygon)
diff --git a/Engine/Source/Runtime/AIModule/Classes/Navigation/PathFollowingComponent.h b/Engine/Source/Runtime/AIModule/Classes/Navigation/PathFollowingComponent.h
--- a/Engine/Source/Runtime/AIModule/Classes/Navigation/PathFollowingComponent.h
+++ b/Engine/Source/Runtime/AIModule/Classes/Navigation/PathFollowingComponent.h
@@ -273,6 +273,8 @@ class AIMODULE_API UPathFollowingComponent : public UActorComponent, public IAIR
 	virtual bool IsResourceLocked() const override;
 	// IAIResourceInterface end
 
+	void OnPathEvent(FNavigationPath* InvalidatedPath, ENavPathEvent::Type Event);
+
 protected:
 
 	/** associated movement component */
@@ -292,6 +294,9 @@ protected:
 	/** requested path */
 	FNavPathSharedPtr Path;
 
+	/** responsible to handing path's events like "underlying navigation rebuilt" or "goal moved" */
+	FNavigationPath::FPathObserverDelegate::FDelegate PathObserver;
+
 	/** value based on navigation agent's properties that's used for AcceptanceRadius when DefaultAcceptanceRadius is requested */
 	float MyDefaultAcceptanceRadius;
 
diff --git a/Engine/Source/Runtime/AIModule/Private/Navigation/PathFollowingComponent.cpp b/Engine/Source/Runtime/AIModule/Private/Navigation/PathFollowingComponent.cpp
--- a/Engine/Source/Runtime/AIModule/Private/Navigation/PathFollowingComponent.cpp
+++ b/Engine/Source/Runtime/AIModule/Private/Navigation/PathFollowingComponent.cpp
@@ -50,6 +50,8 @@ UPathFollowingComponent::UPathFollowingComponent(const FObjectInitializer& Objec
 
 	bStopOnOverlap = true;
 	Status = EPathFollowingStatus::Idle;
+
+	PathObserver = FNavigationPath::FPathObserverDelegate::FDelegate::CreateUObject(this, &UPathFollowingComponent::OnPathEvent);
 }
 
 void LogPathHelper(AActor* LogOwner, FNavPathSharedPtr Path, const AActor* GoalActor, const int32 CurrentPathPointGoal)
@@ -109,6 +111,25 @@ FString GetPathDescHelper(FNavPathSharedPtr Path)
 		FString::Printf(TEXT("%s:%d"), Path->IsPartial() ? TEXT("partial") : TEXT("complete"), Path->GetPathPoints().Num());
 }
 
+void UPathFollowingComponent::OnPathEvent(FNavigationPath* InvalidatedPath, ENavPathEvent::Type Event)
+{
+	if (InvalidatedPath == nullptr || Path.Get() != InvalidatedPath)
+	{
+		return;
+	}
+
+	switch (Event)
+	{
+		case ENavPathEvent::UpdatedDueToGoalMoved:
+		case ENavPathEvent::UpdatedDueToNavigationChanged:
+		{
+			UpdateMove(Path, GetCurrentRequestId());
+			OnPathUpdated();
+		}
+		break;
+	}
+}
+
 FAIRequestID UPathFollowingComponent::RequestMove(FNavPathSharedPtr InPath, FRequestCompletedSignature OnComplete,
 	const AActor* InDestinationActor, float InAcceptanceRadius, bool bInStopOnOverlap, FCustomMoveSharedPtr InGameData)
 {
@@ -167,6 +188,9 @@ FAIRequestID UPathFollowingComponent::RequestMove(FNavPathSharedPtr InPath, FReq
 
 		// store new data
 		Path = InPath;
+		Path->AddObserver(PathObserver);
+		check(MovementComp->GetOwner() != nullptr);
+		Path->SetSourceActor(*(MovementComp->GetOwner()));
 		OnPathUpdated();
 
 		AcceptanceRadius = InAcceptanceRadius;
@@ -743,6 +767,10 @@ void UPathFollowingComponent::UpdatePathSegment()
 				const FVector AgentLocation = DestinationAgent ? DestinationAgent->GetNavAgentLocation() : DestinationActor->GetActorLocation();
 				const FVector GoalLocation = FRotationTranslationMatrix(DestinationActor->GetActorRotation(), AgentLocation).TransformPosition(MoveOffset);
 				CurrentDestination.Set(NULL, GoalLocation);
+
+				UE_VLOG(this, LogPathFollowing, Log, TEXT("Moving directly to move goal rather than following last path segment"));
+				UE_VLOG_LOCATION(this, LogPathFollowing, VeryVerbose, GoalLocation, 30, FColor::Green, TEXT("Last-segment-to-actor"));
+				UE_VLOG_SEGMENT(this, LogPathFollowing, VeryVerbose, CurrentLocation, GoalLocation, FColor::Green, TEXT_EMPTY);
 			}
 
 			UpdateMoveFocus();
diff --git a/Engine/Source/Runtime/Engine/Classes/AI/Navigation/NavMeshRenderingComponent.h b/Engine/Source/Runtime/Engine/Classes/AI/Navigation/NavMeshRenderingComponent.h
--- a/Engine/Source/Runtime/Engine/Classes/AI/Navigation/NavMeshRenderingComponent.h
+++ b/Engine/Source/Runtime/Engine/Classes/AI/Navigation/NavMeshRenderingComponent.h
@@ -3,6 +3,8 @@
 #pragma once
 #include "NavMeshRenderingComponent.generated.h"
 
+struct FNavMeshSceneProxyData;
+
 UCLASS(hidecategories=Object, editinlinenew)
 class ENGINE_API UNavMeshRenderingComponent : public UPrimitiveComponent
 {
@@ -28,7 +30,7 @@ public:
 	virtual FBoxSphereBounds CalcBounds(const FTransform &LocalToWorld) const override;
 	// End USceneComponent Interface
 
-	void GatherData(struct FNavMeshSceneProxyData*) const;
+	void GatherData(FNavMeshSceneProxyData& DebugDrawData) const;
 
 protected:
 	void TimerFunction();
diff --git a/Engine/Source/Runtime/Engine/Classes/AI/Navigation/NavigationData.h b/Engine/Source/Runtime/Engine/Classes/AI/Navigation/NavigationData.h
--- a/Engine/Source/Runtime/Engine/Classes/AI/Navigation/NavigationData.h
+++ b/Engine/Source/Runtime/Engine/Classes/AI/Navigation/NavigationData.h
@@ -10,6 +10,7 @@
 class UNavigationQueryFilter;
 class FNavDataGenerator; 
 class INavLinkCustomInterface;
+class UPrimitiveComponent;
 
 USTRUCT()
 struct FSupportedAreaData
@@ -169,10 +170,14 @@ struct ENGINE_API FNavigationPath : public TSharedFromThis<FNavigationPath, ESPM
 		ObserverDelegate.Broadcast(this, UpdateType == ENavPathUpdateType::GoalMoved ? ENavPathEvent::UpdatedDueToGoalMoved : ENavPathEvent::UpdatedDueToNavigationChanged);
 	}
 
+	FORCEINLINE float GetTimeStamp() const { return LastUpdateTimeStamp; }
+	FORCEINLINE void SetTimeStamp(float TimeStamp) { LastUpdateTimeStamp = TimeStamp; }
+
 	void Invalidate();
 	void RePathFailed();
 
 	virtual void DebugDraw(const ANavigationData* NavData, FColor PathColor, class UCanvas* Canvas, bool bPersistent, const uint32 NextPathPointIndex = 0) const;
+	
 #if ENABLE_VISUAL_LOG
 	virtual void DescribeSelfToVisLog(struct FVisualLogEntry* Snapshot) const;
 	virtual FString GetDescription() const;
@@ -365,6 +370,9 @@ protected:
 	/** navigation data used to generate this path */
 	TWeakObjectPtr<ANavigationData> NavigationDataUsed;
 
+	/** gets set during path creation and on subsequent path's updates */
+	float LastUpdateTimeStamp;
+
 private:
 	/* if GoalActor is set this is the distance we'll try to keep GoalActor from end of path. If GoalActor
 	* moves more then this from the end of the path we'll recalculate the path */
@@ -398,10 +406,11 @@ UCLASS(config=Engine, defaultconfig, NotBlueprintable, abstract)
 class ENGINE_API ANavigationData : public AActor
 {
 	GENERATED_UCLASS_BODY()
-
+	
 	UPROPERTY(transient, duplicatetransient)
-	class UPrimitiveComponent* RenderingComp;
+	UPrimitiveComponent* RenderingComp;
 
+protected:
 	UPROPERTY()
 	FNavDataConfig NavDataConfig;
 
@@ -430,6 +439,10 @@ class ENGINE_API ANavigationData : public AActor
 	UPROPERTY(EditAnywhere, Category = Runtime, config)
 	float ObservedPathsTickInterval;
 
+	UPROPERTY()
+	UWorld* CachedWorld;
+
+public:
 	//----------------------------------------------------------------------//
 	// Life cycle                                                                
 	//----------------------------------------------------------------------//
@@ -471,7 +484,8 @@ class ENGINE_API ANavigationData : public AActor
 	// Generation & data access                                                      
 	//----------------------------------------------------------------------//
 public:
-	const FNavDataConfig& GetConfig() const { return NavDataConfig; }
+	FORCEINLINE const FNavDataConfig& GetConfig() const { return NavDataConfig; }
+	FORCEINLINE ERuntimeGenerationType GetRuntimeGenerationMode() const { return RuntimeGeneration; }
 	virtual void SetConfig(const FNavDataConfig& Src) { NavDataConfig = Src; }
 
 	void SetSupportsDefaultAgent(bool bIsDefault) { bSupportsDefaultAgent = bIsDefault; SetNavRenderingEnabled(bIsDefault); }
@@ -481,6 +495,8 @@ public:
 
 	FORCEINLINE void MarkAsNeedingUpdate() { bWantsUpdate = true; }
 
+	virtual void RestrictBuildingToActiveTiles(bool InRestrictBuildingToActiveTiles) {}
+
 protected:
 	virtual void FillConfig(FNavDataConfig& Dest) { Dest = NavDataConfig; }
 
@@ -506,7 +522,7 @@ public:
 
 	/** Request navigation data update after changes in nav octree */
 	virtual void RebuildDirtyAreas(const TArray<FNavigationDirtyArea>& DirtyAreas);
-
+	
 	/** releases navigation generator if any has been created */
 protected:
 	/** register self with navigation system as new NavAreaDefinition(s) observer */
@@ -523,6 +539,8 @@ public:
 		FNavPathSharedPtr SharedPath = MakeShareable(new PathType());
 		SharedPath->SetNavigationDataUsed(this);
 		SharedPath->SetQuerier(Querier);
+		const UWorld* World = GetCachedWorld();
+		SharedPath->SetTimeStamp(World ? World->GetTimeSeconds() : 0);
 
 		DECLARE_CYCLE_STAT(TEXT("FSimpleDelegateGraphTask.Adding a path to ActivePaths"),
 			STAT_FSimpleDelegateGraphTask_AddingPathToActivePaths,
@@ -573,6 +591,8 @@ public:
 	//----------------------------------------------------------------------//
 	void DrawDebugPath(FNavigationPath* Path, FColor PathColor = FColor::White, class UCanvas* Canvas = NULL, bool bPersistent = true, const uint32 NextPathPointIndex = 0) const;
 
+	FORCEINLINE bool IsDrawingEnabled() const { return bEnableDrawing; }
+
 	/** @return Total mem counted, including super calls. */
 	virtual uint32 LogMemUsed() const;
 
@@ -696,6 +716,7 @@ public:
 	 *	@NOTE potentially expensive, so use it with caution */
 	virtual ENavigationQueryResult::Type CalcPathLengthAndCost(const FVector& PathStart, const FVector& PathEnd, float& OutPathLength, float& OutPathCost, TSharedPtr<const FNavigationQueryFilter> QueryFilter = NULL, const UObject* Querier = NULL) const PURE_VIRTUAL(ANavigationData::CalcPathLength, return ENavigationQueryResult::Invalid;);
 
+	const UWorld* GetCachedWorld() const { return CachedWorld; }
 	//----------------------------------------------------------------------//
 	// Areas
 	//----------------------------------------------------------------------//
diff --git a/Engine/Source/Runtime/Engine/Classes/AI/Navigation/NavigationTypes.h b/Engine/Source/Runtime/Engine/Classes/AI/Navigation/NavigationTypes.h
--- a/Engine/Source/Runtime/Engine/Classes/AI/Navigation/NavigationTypes.h
+++ b/Engine/Source/Runtime/Engine/Classes/AI/Navigation/NavigationTypes.h
@@ -632,4 +632,32 @@ struct TContainerTraits<TNavStatArray<InElementType> > : public TContainerTraits
 		PLATFORM_COMPILER_HAS_RVALUE_REFERENCES &&
 		TContainerTraits<typename TNavStatArray<InElementType>::Super>::MoveWillEmptyContainer
 	};
+};
+
+//----------------------------------------------------------------------//
+// Active tiles 
+//----------------------------------------------------------------------//
+struct FNavigationInvokerRaw
+{
+	FVector Location;
+	float RadiusMin;
+	float RadiusMax;
+
+	FNavigationInvokerRaw(const FVector& InLocation, float Min, float Max)
+		: Location(InLocation), RadiusMin(Min), RadiusMax(Max)
+	{}
+};
+
+struct FNavigationInvoker
+{
+	TWeakObjectPtr<AActor> Actor;
+
+	/** tiles GenerationRadius away or close will be generated if they're not already present */
+	float GenerationRadius;
+	/** tiles over RemovalRadius will get removed.
+	*	@Note needs to be >= GenerationRadius or will get clampped */
+	float RemovalRadius;
+
+	FNavigationInvoker();
+	FNavigationInvoker(AActor& InActor, float InGenerationRadius, float InRemovalRadius);
 };
\ No newline at end of file
diff --git a/Engine/Source/Runtime/Engine/Private/AI/Navigation/NavLinkRenderingComponent.cpp b/Engine/Source/Runtime/Engine/Private/AI/Navigation/NavLinkRenderingComponent.cpp
--- a/Engine/Source/Runtime/Engine/Private/AI/Navigation/NavLinkRenderingComponent.cpp
+++ b/Engine/Source/Runtime/Engine/Private/AI/Navigation/NavLinkRenderingComponent.cpp
@@ -174,8 +174,8 @@ void FNavLinkRenderingProxy::GetDynamicMeshElements(const TArray<const FSceneVie
 			for(int32 DataIndex = 0; DataIndex < NavSys->NavDataSet.Num(); ++DataIndex)
 			{
 				const ARecastNavMesh* NavMesh = Cast<const ARecastNavMesh>(NavSys->NavDataSet[DataIndex]);
-				AgentMask = NavMesh->bEnableDrawing ? AgentMask | (1 << DataIndex) : AgentMask;
-				if (NavMesh != NULL && NavMesh->AgentMaxStepHeight > 0 && NavMesh->bEnableDrawing)
+				AgentMask = NavMesh->IsDrawingEnabled() ? AgentMask | (1 << DataIndex) : AgentMask;
+				if (NavMesh != NULL && NavMesh->AgentMaxStepHeight > 0 && NavMesh->IsDrawingEnabled())
 				{
 					StepHeights.Add(NavMesh->AgentMaxStepHeight);
 				}
diff --git a/Engine/Source/Runtime/Engine/Private/AI/Navigation/NavigationData.cpp b/Engine/Source/Runtime/Engine/Private/AI/Navigation/NavigationData.cpp
--- a/Engine/Source/Runtime/Engine/Private/AI/Navigation/NavigationData.cpp
+++ b/Engine/Source/Runtime/Engine/Private/AI/Navigation/NavigationData.cpp
@@ -161,9 +161,9 @@ void ANavigationData::PostInitializeComponents()
 {
 	Super::PostInitializeComponents();
 	
-	UWorld* WorldOuter = GetWorld();
+	CachedWorld = GetWorld();
 
-	if (WorldOuter == NULL || WorldOuter->GetNavigationSystem() == NULL)
+	if (CachedWorld == NULL || CachedWorld->GetNavigationSystem() == NULL)
 	{
 		CleanUpAndMarkPendingKill();
 	}
@@ -171,7 +171,7 @@ void ANavigationData::PostInitializeComponents()
 	{
 		// note: this is not a final fix for world composition's issues with navmesh generation
 		// but it's good for now, and navmesh creation is going to get a face-lift soon anyway
-		bWantsUpdate |= WorldOuter->GetWorldSettings()->bEnableWorldComposition;
+		bWantsUpdate |= CachedWorld->GetWorldSettings()->bEnableWorldComposition;
 	}
 }
 
@@ -185,6 +185,8 @@ void ANavigationData::PostLoad()
 	}
 
 	InstantiateAndRegisterRenderingComponent();
+
+	CachedWorld = GetWorld();
 }
 
 void ANavigationData::TickActor(float DeltaTime, enum ELevelTick TickType, FActorTickFunction& ThisTickFunction)
@@ -240,6 +242,8 @@ void ANavigationData::TickActor(float DeltaTime, enum ELevelTick TickType, FActo
 
 	if (RepathRequests.Num() > 0)
 	{
+		const UWorld* World = GetCachedWorld();
+		float TimeStamp = World ? World->GetTimeSeconds() : 0.f;
 		TArray<FNavPathRecalculationRequest> PostponedRequests;
 
 		// @todo batch-process it!
@@ -257,6 +261,9 @@ void ANavigationData::TickActor(float DeltaTime, enum ELevelTick TickType, FActo
 			// @todo consider supplying NavAgentPropertied from path's querier
 			const FPathFindingResult Result = FindPath(FNavAgentProperties(), Query.SetPathInstanceToUpdate(RecalcRequest.Path));
 
+			// update time stamp to give observers any means of telling if it has changed
+			RecalcRequest.Path->SetTimeStamp(TimeStamp);
+
 			// partial paths are still valid and can change to full path when moving goal gets back on navmesh
 			if (Result.IsSuccessful() || Result.IsPartial())
 			{
diff --git a/Engine/Source/Runtime/Engine/Private/NavMeshRenderingComponent.cpp b/Engine/Source/Runtime/Engine/Private/NavMeshRenderingComponent.cpp
--- a/Engine/Source/Runtime/Engine/Private/NavMeshRenderingComponent.cpp
+++ b/Engine/Source/Runtime/Engine/Private/NavMeshRenderingComponent.cpp
@@ -130,7 +130,7 @@ FPrimitiveSceneProxy* UNavMeshRenderingComponent::CreateSceneProxy()
 	{
 		FNavMeshSceneProxyData ProxyData;
 		ProxyData.Reset();
-		GatherData(&ProxyData);
+		GatherData(ProxyData);
 
 		SceneProxy = new FRecastRenderingSceneProxy(this, &ProxyData);
 		ProxyData.Reset();
@@ -191,33 +191,34 @@ FBoxSphereBounds UNavMeshRenderingComponent::CalcBounds(const FTransform& LocalT
 	return FBoxSphereBounds(BoundingBox);
 }
 
-void UNavMeshRenderingComponent::GatherData(struct FNavMeshSceneProxyData* CurrentData) const
+void UNavMeshRenderingComponent::GatherData(FNavMeshSceneProxyData& CurrentData) const
 {
 #if WITH_RECAST
 	const ARecastNavMesh* NavMesh = Cast<ARecastNavMesh>(GetOwner());
 
-	CurrentData->Reset();
-	CurrentData->bEnableDrawing = NavMesh->bEnableDrawing;
-	CurrentData->bNeedsNewData = false;
-
-	if (CurrentData && NavMesh && NavMesh->bEnableDrawing)
+	CurrentData.Reset();
+	
+	if (NavMesh && NavMesh->IsDrawingEnabled())
 	{
+		CurrentData.bEnableDrawing = true;
+		CurrentData.bNeedsNewData = false;
+
 		FHitProxyId HitProxyId = FHitProxyId();
-		CurrentData->bDrawPathCollidingGeometry = NavMesh->bDrawPathCollidingGeometry;
+		CurrentData.bDrawPathCollidingGeometry = NavMesh->bDrawPathCollidingGeometry;
 
-		CurrentData->NavMeshDrawOffset.Z = NavMesh->DrawOffset;
-		CurrentData->NavMeshGeometry.bGatherPolyEdges = NavMesh->bDrawPolyEdges;
-		CurrentData->NavMeshGeometry.bGatherNavMeshEdges = NavMesh->bDrawNavMeshEdges;
+		CurrentData.NavMeshDrawOffset.Z = NavMesh->DrawOffset;
+		CurrentData.NavMeshGeometry.bGatherPolyEdges = NavMesh->bDrawPolyEdges;
+		CurrentData.NavMeshGeometry.bGatherNavMeshEdges = NavMesh->bDrawNavMeshEdges;
 
 		const FNavDataConfig& NavConfig = NavMesh->GetConfig();
-		CurrentData->NavMeshColors[RECAST_DEFAULT_AREA] = NavConfig.Color.DWColor() > 0 ? NavConfig.Color : NavMeshRenderColor_RecastMesh;
+		CurrentData.NavMeshColors[RECAST_DEFAULT_AREA] = NavConfig.Color.DWColor() > 0 ? NavConfig.Color : NavMeshRenderColor_RecastMesh;
 		for (uint8 i = 0; i < RECAST_DEFAULT_AREA; ++i)
 		{
-			CurrentData->NavMeshColors[i] = NavMesh->GetAreaIDColor(i);
+			CurrentData.NavMeshColors[i] = NavMesh->GetAreaIDColor(i);
 		}
 
 		// just a little trick to make sure navmeshes with different sized are not drawn with same offset
-		CurrentData->NavMeshDrawOffset.Z += NavMesh->GetConfig().AgentRadius / 10.f;
+		CurrentData.NavMeshDrawOffset.Z += NavMesh->GetConfig().AgentRadius / 10.f;
 
 		NavMesh->BeginBatchQuery();
 
@@ -231,7 +232,7 @@ void UNavMeshRenderingComponent::GatherData(struct FNavMeshSceneProxyData* Curre
 				{
 					const FNavigationOctree::FNode& CurrentNode = NodeIt.GetCurrentNode();
 					const FOctreeNodeContext& CorrentContext = NodeIt.GetCurrentContext();
-					CurrentData->OctreeBounds.Add(CorrentContext.Bounds);
+					CurrentData.OctreeBounds.Add(CorrentContext.Bounds);
 
 					FOREACH_OCTREE_CHILD_NODE(ChildRef)
 					{
@@ -244,20 +245,20 @@ void UNavMeshRenderingComponent::GatherData(struct FNavMeshSceneProxyData* Curre
 			}
 		}
 
-		NavMesh->GetDebugGeometry(CurrentData->NavMeshGeometry);
-		const TArray<FVector>& MeshVerts = CurrentData->NavMeshGeometry.MeshVerts;
+		NavMesh->GetDebugGeometry(CurrentData.NavMeshGeometry);
+		const TArray<FVector>& MeshVerts = CurrentData.NavMeshGeometry.MeshVerts;
 
 		// @fixme, this is going to double up on lots of interior lines
 		if (NavMesh->bDrawTriangleEdges)
 		{
 			for (int32 AreaIdx = 0; AreaIdx < RECAST_MAX_AREAS; ++AreaIdx)
 			{
-				const TArray<int32>& MeshIndices = CurrentData->NavMeshGeometry.AreaIndices[AreaIdx];
+				const TArray<int32>& MeshIndices = CurrentData.NavMeshGeometry.AreaIndices[AreaIdx];
 				for (int32 Idx=0; Idx<MeshIndices.Num(); Idx += 3)
 				{
-					CurrentData->ThickLineItems.Add(FNavMeshSceneProxyData::FDebugThickLine(MeshVerts[MeshIndices[Idx + 0]] + CurrentData->NavMeshDrawOffset, MeshVerts[MeshIndices[Idx + 1]] + CurrentData->NavMeshDrawOffset, NavMeshRenderColor_Recast_TriangleEdges, DefaultEdges_LineThickness));
-					CurrentData->ThickLineItems.Add(FNavMeshSceneProxyData::FDebugThickLine(MeshVerts[MeshIndices[Idx + 1]] + CurrentData->NavMeshDrawOffset, MeshVerts[MeshIndices[Idx + 2]] + CurrentData->NavMeshDrawOffset, NavMeshRenderColor_Recast_TriangleEdges, DefaultEdges_LineThickness));
-					CurrentData->ThickLineItems.Add(FNavMeshSceneProxyData::FDebugThickLine(MeshVerts[MeshIndices[Idx + 2]] + CurrentData->NavMeshDrawOffset, MeshVerts[MeshIndices[Idx + 0]] + CurrentData->NavMeshDrawOffset, NavMeshRenderColor_Recast_TriangleEdges, DefaultEdges_LineThickness));
+					CurrentData.ThickLineItems.Add(FNavMeshSceneProxyData::FDebugThickLine(MeshVerts[MeshIndices[Idx + 0]] + CurrentData.NavMeshDrawOffset, MeshVerts[MeshIndices[Idx + 1]] + CurrentData.NavMeshDrawOffset, NavMeshRenderColor_Recast_TriangleEdges, DefaultEdges_LineThickness));
+					CurrentData.ThickLineItems.Add(FNavMeshSceneProxyData::FDebugThickLine(MeshVerts[MeshIndices[Idx + 1]] + CurrentData.NavMeshDrawOffset, MeshVerts[MeshIndices[Idx + 2]] + CurrentData.NavMeshDrawOffset, NavMeshRenderColor_Recast_TriangleEdges, DefaultEdges_LineThickness));
+					CurrentData.ThickLineItems.Add(FNavMeshSceneProxyData::FDebugThickLine(MeshVerts[MeshIndices[Idx + 2]] + CurrentData.NavMeshDrawOffset, MeshVerts[MeshIndices[Idx + 0]] + CurrentData.NavMeshDrawOffset, NavMeshRenderColor_Recast_TriangleEdges, DefaultEdges_LineThickness));
 				}
 			}
 		}
@@ -265,45 +266,45 @@ void UNavMeshRenderingComponent::GatherData(struct FNavMeshSceneProxyData* Curre
 		// make lines for tile edges
 		if (NavMesh->bDrawPolyEdges)
 		{
-			const TArray<FVector>& TileEdgeVerts = CurrentData->NavMeshGeometry.PolyEdges;
+			const TArray<FVector>& TileEdgeVerts = CurrentData.NavMeshGeometry.PolyEdges;
 			for (int32 Idx=0; Idx < TileEdgeVerts.Num(); Idx += 2)
 			{
-				CurrentData->TileEdgeLines.Add( FDebugRenderSceneProxy::FDebugLine(TileEdgeVerts[Idx] + CurrentData->NavMeshDrawOffset, TileEdgeVerts[Idx+1] + CurrentData->NavMeshDrawOffset, NavMeshRenderColor_Recast_TileEdges));
+				CurrentData.TileEdgeLines.Add(FDebugRenderSceneProxy::FDebugLine(TileEdgeVerts[Idx] + CurrentData.NavMeshDrawOffset, TileEdgeVerts[Idx + 1] + CurrentData.NavMeshDrawOffset, NavMeshRenderColor_Recast_TileEdges));
 			}
 		}
 
 		// make lines for navmesh edges
 		if (NavMesh->bDrawNavMeshEdges)
 		{
-			const FColor EdgesColor = DarkenColor(CurrentData->NavMeshColors[RECAST_DEFAULT_AREA]);
-			const TArray<FVector>& NavMeshEdgeVerts = CurrentData->NavMeshGeometry.NavMeshEdges;
+			const FColor EdgesColor = DarkenColor(CurrentData.NavMeshColors[RECAST_DEFAULT_AREA]);
+			const TArray<FVector>& NavMeshEdgeVerts = CurrentData.NavMeshGeometry.NavMeshEdges;
 			for (int32 Idx=0; Idx < NavMeshEdgeVerts.Num(); Idx += 2)
 			{
-				CurrentData->NavMeshEdgeLines.Add( FDebugRenderSceneProxy::FDebugLine(NavMeshEdgeVerts[Idx] + CurrentData->NavMeshDrawOffset, NavMeshEdgeVerts[Idx+1] + CurrentData->NavMeshDrawOffset, EdgesColor));
+				CurrentData.NavMeshEdgeLines.Add( FDebugRenderSceneProxy::FDebugLine(NavMeshEdgeVerts[Idx] + CurrentData.NavMeshDrawOffset, NavMeshEdgeVerts[Idx+1] + CurrentData.NavMeshDrawOffset, EdgesColor));
 			}
 		}
 
 		// offset all navigation-link positions
 		if (!NavMesh->bDrawClusters)
 		{
-			for (int32 OffMeshLineIndex = 0; OffMeshLineIndex < CurrentData->NavMeshGeometry.OffMeshLinks.Num(); ++OffMeshLineIndex)
+			for (int32 OffMeshLineIndex = 0; OffMeshLineIndex < CurrentData.NavMeshGeometry.OffMeshLinks.Num(); ++OffMeshLineIndex)
 			{
-				FRecastDebugGeometry::FOffMeshLink& Link = CurrentData->NavMeshGeometry.OffMeshLinks[OffMeshLineIndex];
+				FRecastDebugGeometry::FOffMeshLink& Link = CurrentData.NavMeshGeometry.OffMeshLinks[OffMeshLineIndex];
 				const bool bLinkValid = (Link.ValidEnds & FRecastDebugGeometry::OMLE_Left) && (Link.ValidEnds & FRecastDebugGeometry::OMLE_Right);
 
 				if (NavMesh->bDrawFailedNavLinks || (NavMesh->bDrawNavLinks && bLinkValid))
 				{
-					const FVector V0 = Link.Left + CurrentData->NavMeshDrawOffset;
-					const FVector V1 = Link.Right + CurrentData->NavMeshDrawOffset;
-					const FColor LinkColor = ((Link.Direction && Link.ValidEnds) || (Link.ValidEnds & FRecastDebugGeometry::OMLE_Left)) ? DarkenColor(CurrentData->NavMeshColors[Link.AreaID]) : NavMeshRenderColor_OffMeshConnectionInvalid;
+					const FVector V0 = Link.Left + CurrentData.NavMeshDrawOffset;
+					const FVector V1 = Link.Right + CurrentData.NavMeshDrawOffset;
+					const FColor LinkColor = ((Link.Direction && Link.ValidEnds) || (Link.ValidEnds & FRecastDebugGeometry::OMLE_Left)) ? DarkenColor(CurrentData.NavMeshColors[Link.AreaID]) : NavMeshRenderColor_OffMeshConnectionInvalid;
 
-					CacheArc(CurrentData->NavLinkLines, V0, V1, 0.4f, 4, LinkColor, LinkLines_LineThickness);
+					CacheArc(CurrentData.NavLinkLines, V0, V1, 0.4f, 4, LinkColor, LinkLines_LineThickness);
 
 					const FVector VOffset(0, 0, FVector::Dist(V0, V1) * 1.333f);
-					CacheArrowHead(CurrentData->NavLinkLines, V1, V0+VOffset, 30.f, LinkColor, LinkLines_LineThickness);
+					CacheArrowHead(CurrentData.NavLinkLines, V1, V0+VOffset, 30.f, LinkColor, LinkLines_LineThickness);
 					if (Link.Direction)
 					{
-						CacheArrowHead(CurrentData->NavLinkLines, V0, V1+VOffset, 30.f, LinkColor, LinkLines_LineThickness);
+						CacheArrowHead(CurrentData.NavLinkLines, V0, V1+VOffset, 30.f, LinkColor, LinkLines_LineThickness);
 					}
 
 					// if the connection as a whole is valid check if there are any of ends is invalid
@@ -312,11 +313,11 @@ void UNavMeshRenderingComponent::GatherData(struct FNavMeshSceneProxyData* Curre
 						if (Link.Direction && (Link.ValidEnds & FRecastDebugGeometry::OMLE_Left) == 0)
 						{
 							// left end invalid - mark it
-							DrawWireCylinder(CurrentData->NavLinkLines, V0, FVector(1,0,0), FVector(0,1,0), FVector(0,0,1), NavMeshRenderColor_OffMeshConnectionInvalid, Link.Radius, NavMesh->AgentMaxStepHeight, 16, 0, DefaultEdges_LineThickness);
+							DrawWireCylinder(CurrentData.NavLinkLines, V0, FVector(1,0,0), FVector(0,1,0), FVector(0,0,1), NavMeshRenderColor_OffMeshConnectionInvalid, Link.Radius, NavMesh->AgentMaxStepHeight, 16, 0, DefaultEdges_LineThickness);
 						}
 						if ((Link.ValidEnds & FRecastDebugGeometry::OMLE_Right) == 0)
 						{
-							DrawWireCylinder(CurrentData->NavLinkLines, V1, FVector(1,0,0), FVector(0,1,0), FVector(0,0,1), NavMeshRenderColor_OffMeshConnectionInvalid, Link.Radius, NavMesh->AgentMaxStepHeight, 16, 0, DefaultEdges_LineThickness);
+							DrawWireCylinder(CurrentData.NavLinkLines, V1, FVector(1,0,0), FVector(0,1,0), FVector(0,0,1), NavMeshRenderColor_OffMeshConnectionInvalid, Link.Radius, NavMesh->AgentMaxStepHeight, 16, 0, DefaultEdges_LineThickness);
 						}
 					}
 				}					
@@ -327,7 +328,7 @@ void UNavMeshRenderingComponent::GatherData(struct FNavMeshSceneProxyData* Curre
 		{
 			// calculate appropriate points for displaying debug labels
 			const int32 TilesCount = NavMesh->GetNavMeshTilesCount();
-			CurrentData->DebugLabels.Reserve(TilesCount);
+			CurrentData.DebugLabels.Reserve(TilesCount);
 
 			for (int32 TileIndex = 0; TileIndex < TilesCount; ++TileIndex)
 			{
@@ -346,8 +347,8 @@ void UNavMeshRenderingComponent::GatherData(struct FNavMeshSceneProxyData* Curre
 
 					if (NavMesh->bDrawTileLabels)
 					{
-						CurrentData->DebugLabels.Add(FNavMeshSceneProxyData::FDebugText(
-							/*Location*/NavLocation.Location + CurrentData->NavMeshDrawOffset
+						CurrentData.DebugLabels.Add(FNavMeshSceneProxyData::FDebugText(
+							/*Location*/NavLocation.Location + CurrentData.NavMeshDrawOffset
 							, /*Text*/FString::Printf(TEXT("(%d,%d:%d)"), X, Y, Layer)
 							));
 					}
@@ -368,8 +369,8 @@ void UNavMeshRenderingComponent::GatherData(struct FNavMeshSceneProxyData* Curre
 							{
 								uint32 AreaID = NavMesh->GetPolyAreaID(Polys[k].Ref);
 
-								CurrentData->DebugLabels.Add(FNavMeshSceneProxyData::FDebugText(
-									/*Location*/Polys[k].Center + CurrentData->NavMeshDrawOffset
+								CurrentData.DebugLabels.Add(FNavMeshSceneProxyData::FDebugText(
+									/*Location*/Polys[k].Center + CurrentData.NavMeshDrawOffset
 									, /*Text*/FString::Printf(TEXT("\\%.3f; %.3f\\"), DefaultCosts[AreaID], FixedCosts[AreaID])
 									));
 							}
@@ -382,8 +383,8 @@ void UNavMeshRenderingComponent::GatherData(struct FNavMeshSceneProxyData* Curre
 								uint32 NavTileIndex = 0;
 								NavMesh->GetPolyTileIndex(Polys[k].Ref, NavPolyIndex, NavTileIndex);
 
-								CurrentData->DebugLabels.Add(FNavMeshSceneProxyData::FDebugText(
-									/*Location*/Polys[k].Center + CurrentData->NavMeshDrawOffset
+								CurrentData.DebugLabels.Add(FNavMeshSceneProxyData::FDebugText(
+									/*Location*/Polys[k].Center + CurrentData.NavMeshDrawOffset
 									, /*Text*/FString::Printf(TEXT("[%X:%X]"), NavTileIndex, NavPolyIndex)
 									));
 							}
@@ -399,25 +400,25 @@ void UNavMeshRenderingComponent::GatherData(struct FNavMeshSceneProxyData* Curre
 						FVector UR(TileBox.Max.X, TileBox.Max.Y, DrawZ);
 						FVector UL(LL.X, UR.Y, DrawZ);
 						FVector LR(UR.X, LL.Y, DrawZ);
-						CurrentData->ThickLineItems.Add(FNavMeshSceneProxyData::FDebugThickLine(LL, UL, NavMeshRenderColor_TileBounds, DefaultEdges_LineThickness));
-						CurrentData->ThickLineItems.Add(FNavMeshSceneProxyData::FDebugThickLine(UL, UR, NavMeshRenderColor_TileBounds, DefaultEdges_LineThickness));
-						CurrentData->ThickLineItems.Add(FNavMeshSceneProxyData::FDebugThickLine(UR, LR, NavMeshRenderColor_TileBounds, DefaultEdges_LineThickness));
-						CurrentData->ThickLineItems.Add(FNavMeshSceneProxyData::FDebugThickLine(LR, LL, NavMeshRenderColor_TileBounds, DefaultEdges_LineThickness));
+						CurrentData.ThickLineItems.Add(FNavMeshSceneProxyData::FDebugThickLine(LL, UL, NavMeshRenderColor_TileBounds, DefaultEdges_LineThickness));
+						CurrentData.ThickLineItems.Add(FNavMeshSceneProxyData::FDebugThickLine(UL, UR, NavMeshRenderColor_TileBounds, DefaultEdges_LineThickness));
+						CurrentData.ThickLineItems.Add(FNavMeshSceneProxyData::FDebugThickLine(UR, LR, NavMeshRenderColor_TileBounds, DefaultEdges_LineThickness));
+						CurrentData.ThickLineItems.Add(FNavMeshSceneProxyData::FDebugThickLine(LR, LL, NavMeshRenderColor_TileBounds, DefaultEdges_LineThickness));
 					}
 				}
 			}
 		}
 
-		CurrentData->bSkipDistanceCheck = GIsEditor && (GEngine->GetDebugLocalPlayer() == NULL);
-		CurrentData->bDrawClusters = NavMesh->bDrawClusters;
+		CurrentData.bSkipDistanceCheck = GIsEditor && (GEngine->GetDebugLocalPlayer() == NULL);
+		CurrentData.bDrawClusters = NavMesh->bDrawClusters;
 		NavMesh->FinishBatchQuery();
 
 		// Draw Mesh
 		if (NavMesh->bDrawClusters)
 		{
-			for (int32 Idx = 0; Idx < CurrentData->NavMeshGeometry.Clusters.Num(); ++Idx)
+			for (int32 Idx = 0; Idx < CurrentData.NavMeshGeometry.Clusters.Num(); ++Idx)
 			{
-				const TArray<int32>& MeshIndices = CurrentData->NavMeshGeometry.Clusters[Idx].MeshIndices;
+				const TArray<int32>& MeshIndices = CurrentData.NavMeshGeometry.Clusters[Idx].MeshIndices;
 
 				if (MeshIndices.Num() == 0)
 				{
@@ -428,21 +429,21 @@ void UNavMeshRenderingComponent::GatherData(struct FNavMeshSceneProxyData* Curre
 				DebugMeshData.ClusterColor = GetClusterColor(Idx);
 				for (int32 VertIdx=0; VertIdx < MeshVerts.Num(); ++VertIdx)
 				{
-					AddVertexHelper(DebugMeshData, MeshVerts[VertIdx] + CurrentData->NavMeshDrawOffset, DebugMeshData.ClusterColor);
+					AddVertexHelper(DebugMeshData, MeshVerts[VertIdx] + CurrentData.NavMeshDrawOffset, DebugMeshData.ClusterColor);
 				}
 				for (int32 TriIdx=0; TriIdx < MeshIndices.Num(); TriIdx+=3)
 				{
 					AddTriangleHelper(DebugMeshData, MeshIndices[TriIdx], MeshIndices[TriIdx+1], MeshIndices[TriIdx+2]);
 				}
 
-				CurrentData->MeshBuilders.Add(DebugMeshData);
+				CurrentData.MeshBuilders.Add(DebugMeshData);
 			}
 		}
 		else if (NavMesh->bDrawNavMesh)
 		{			
 			for (int32 AreaType = 0; AreaType < RECAST_MAX_AREAS; ++AreaType)
 			{
-				const TArray<int32>& MeshIndices = CurrentData->NavMeshGeometry.AreaIndices[AreaType];
+				const TArray<int32>& MeshIndices = CurrentData.NavMeshGeometry.AreaIndices[AreaType];
 
 				if (MeshIndices.Num() == 0)
 				{
@@ -452,15 +453,15 @@ void UNavMeshRenderingComponent::GatherData(struct FNavMeshSceneProxyData* Curre
 				FNavMeshSceneProxyData::FDebugMeshData DebugMeshData;
 				for (int32 VertIdx=0; VertIdx < MeshVerts.Num(); ++VertIdx)
 				{
-					AddVertexHelper(DebugMeshData, MeshVerts[VertIdx] + CurrentData->NavMeshDrawOffset, CurrentData->NavMeshColors[AreaType]);
+					AddVertexHelper(DebugMeshData, MeshVerts[VertIdx] + CurrentData.NavMeshDrawOffset, CurrentData.NavMeshColors[AreaType]);
 				}
 				for (int32 TriIdx=0; TriIdx < MeshIndices.Num(); TriIdx+=3)
 				{
 					AddTriangleHelper(DebugMeshData, MeshIndices[TriIdx], MeshIndices[TriIdx+1], MeshIndices[TriIdx+2]);
 				}
 
-				DebugMeshData.ClusterColor = CurrentData->NavMeshColors[AreaType];
-				CurrentData->MeshBuilders.Add(DebugMeshData);
+				DebugMeshData.ClusterColor = CurrentData.NavMeshColors[AreaType];
+				CurrentData.MeshBuilders.Add(DebugMeshData);
 			}
 		}
 
@@ -477,7 +478,7 @@ void UNavMeshRenderingComponent::GatherData(struct FNavMeshSceneProxyData* Curre
 				for (FNavigationOctree::ElementConstIt ElementIt(Node.GetElementIt()); ElementIt; ElementIt++)
 				{
 					const FNavigationOctreeElement& Element = *ElementIt;
-					if (Element.ShouldUseGeometry(&NavMesh->NavDataConfig) && Element.Data.CollisionData.Num())
+					if (Element.ShouldUseGeometry(NavMesh->GetConfig()) && Element.Data.CollisionData.Num())
 					{
 						const FRecastGeometryCache CachedGeometry(Element.Data.CollisionData.GetData());
 						AppendGeometry(PathCollidingGeomVerts, PathCollidingGeomIndices, CachedGeometry.Verts, CachedGeometry.Header.NumVerts, CachedGeometry.Indices, CachedGeometry.Header.NumFaces);
@@ -491,39 +492,39 @@ void UNavMeshRenderingComponent::GatherData(struct FNavMeshSceneProxyData* Curre
 					}
 				}
 			}
-			CurrentData->PathCollidingGeomIndices = PathCollidingGeomIndices;
+			CurrentData.PathCollidingGeomIndices = PathCollidingGeomIndices;
 			for (const auto& Vertex : PathCollidingGeomVerts)
 			{
-				CurrentData->PathCollidingGeomVerts.Add(FDynamicMeshVertex(Vertex));
+				CurrentData.PathCollidingGeomVerts.Add(FDynamicMeshVertex(Vertex));
 			}
 
 		}
 
-		if (CurrentData->NavMeshGeometry.BuiltMeshIndices.Num() > 0)
+		if (CurrentData.NavMeshGeometry.BuiltMeshIndices.Num() > 0)
 		{
 			FNavMeshSceneProxyData::FDebugMeshData DebugMeshData;
 			for (int32 VertIdx=0; VertIdx < MeshVerts.Num(); ++VertIdx)
 			{
-				AddVertexHelper(DebugMeshData, MeshVerts[VertIdx] + CurrentData->NavMeshDrawOffset, NavMeshRenderColor_RecastTileBeingRebuilt);
+				AddVertexHelper(DebugMeshData, MeshVerts[VertIdx] + CurrentData.NavMeshDrawOffset, NavMeshRenderColor_RecastTileBeingRebuilt);
 			}
-			DebugMeshData.Indices.Append(CurrentData->NavMeshGeometry.BuiltMeshIndices);
+			DebugMeshData.Indices.Append(CurrentData.NavMeshGeometry.BuiltMeshIndices);
 			DebugMeshData.ClusterColor = NavMeshRenderColor_RecastTileBeingRebuilt;
-			CurrentData->MeshBuilders.Add(DebugMeshData);
+			CurrentData.MeshBuilders.Add(DebugMeshData);
 
 			// updates should be requested by FRecastNavMeshGenerator::TickAsyncBuild after tiles were refreshed
 		}
 
 		if (NavMesh->bDrawClusters)
 		{
-			for (int i = 0; i < CurrentData->NavMeshGeometry.ClusterLinks.Num(); i++)
+			for (int i = 0; i < CurrentData.NavMeshGeometry.ClusterLinks.Num(); i++)
 			{
-				const FRecastDebugGeometry::FClusterLink& CLink = CurrentData->NavMeshGeometry.ClusterLinks[i];
-				const FVector V0 = CLink.FromCluster + CurrentData->NavMeshDrawOffset;
-				const FVector V1 = CLink.ToCluster + CurrentData->NavMeshDrawOffset + FVector(0,0,20.0f);
+				const FRecastDebugGeometry::FClusterLink& CLink = CurrentData.NavMeshGeometry.ClusterLinks[i];
+				const FVector V0 = CLink.FromCluster + CurrentData.NavMeshDrawOffset;
+				const FVector V1 = CLink.ToCluster + CurrentData.NavMeshDrawOffset + FVector(0,0,20.0f);
 
-				CacheArc(CurrentData->ClusterLinkLines, V0, V1, 0.4f, 4, FColor::Black, ClusterLinkLines_LineThickness);
+				CacheArc(CurrentData.ClusterLinkLines, V0, V1, 0.4f, 4, FColor::Black, ClusterLinkLines_LineThickness);
 				const FVector VOffset(0, 0, FVector::Dist(V0, V1) * 1.333f);
-				CacheArrowHead(CurrentData->ClusterLinkLines, V1, V0+VOffset, 30.f, FColor::Black, ClusterLinkLines_LineThickness);
+				CacheArrowHead(CurrentData.ClusterLinkLines, V1, V0+VOffset, 30.f, FColor::Black, ClusterLinkLines_LineThickness);
 			}
 		}
 
@@ -532,22 +533,22 @@ void UNavMeshRenderingComponent::GatherData(struct FNavMeshSceneProxyData* Curre
 		{
 			for (int32 iArea = 0; iArea < RECAST_MAX_AREAS; iArea++)
 			{
-				const TArray<int32>& Indices = CurrentData->NavMeshGeometry.OffMeshSegmentAreas[iArea];
+				const TArray<int32>& Indices = CurrentData.NavMeshGeometry.OffMeshSegmentAreas[iArea];
 				FNavMeshSceneProxyData::FDebugMeshData DebugMeshData;
 				int32 VertBase = 0;
 
 				for (int32 i = 0; i < Indices.Num(); i++)
 				{
-					FRecastDebugGeometry::FOffMeshSegment& SegInfo = CurrentData->NavMeshGeometry.OffMeshSegments[Indices[i]];
-					const FVector A0 = SegInfo.LeftStart + CurrentData->NavMeshDrawOffset;
-					const FVector A1 = SegInfo.LeftEnd + CurrentData->NavMeshDrawOffset;
-					const FVector B0 = SegInfo.RightStart + CurrentData->NavMeshDrawOffset;
-					const FVector B1 = SegInfo.RightEnd + CurrentData->NavMeshDrawOffset;
+					FRecastDebugGeometry::FOffMeshSegment& SegInfo = CurrentData.NavMeshGeometry.OffMeshSegments[Indices[i]];
+					const FVector A0 = SegInfo.LeftStart + CurrentData.NavMeshDrawOffset;
+					const FVector A1 = SegInfo.LeftEnd + CurrentData.NavMeshDrawOffset;
+					const FVector B0 = SegInfo.RightStart + CurrentData.NavMeshDrawOffset;
+					const FVector B1 = SegInfo.RightEnd + CurrentData.NavMeshDrawOffset;
 					const FVector Edge0 = B0 - A0;
 					const FVector Edge1 = B1 - A1;
 					const float Len0 = Edge0.Size();
 					const float Len1 = Edge1.Size();
-					const FColor SegColor = DarkenColor(CurrentData->NavMeshColors[SegInfo.AreaID]);
+					const FColor SegColor = DarkenColor(CurrentData.NavMeshColors[SegInfo.AreaID]);
 					const FColor ColA = (SegInfo.ValidEnds & FRecastDebugGeometry::OMLE_Left) ? FColor::White : FColor::Black;
 					const FColor ColB = (SegInfo.ValidEnds & FRecastDebugGeometry::OMLE_Right) ? FColor::White : FColor::Black;
 
@@ -583,13 +584,13 @@ void UNavMeshRenderingComponent::GatherData(struct FNavMeshSceneProxyData* Curre
 
 				if (DebugMeshData.Indices.Num())
 				{
-					CurrentData->MeshBuilders.Add(DebugMeshData);
+					CurrentData.MeshBuilders.Add(DebugMeshData);
 				}
 			}
 		}
 
-		CurrentData->NavMeshGeometry.PolyEdges.Empty();
-		CurrentData->NavMeshGeometry.NavMeshEdges.Empty();
+		CurrentData.NavMeshGeometry.PolyEdges.Empty();
+		CurrentData.NavMeshGeometry.NavMeshEdges.Empty();
 	}
 #endif //WITH_RECAST
 }
diff --git a/Engine/Source/Runtime/Engine/Private/NavMeshRenderingHelpers.h b/Engine/Source/Runtime/Engine/Private/NavMeshRenderingHelpers.h
--- a/Engine/Source/Runtime/Engine/Private/NavMeshRenderingHelpers.h
+++ b/Engine/Source/Runtime/Engine/Private/NavMeshRenderingHelpers.h
@@ -43,11 +43,15 @@ struct FNavMeshSceneProxyData : public TSharedFromThis<FNavMeshSceneProxyData, E
 
 	void Reset()
 	{
-		for (int32 Index=0; Index < RECAST_MAX_AREAS; Index++)
+		for (int32 Index = 0; Index < RECAST_MAX_AREAS; Index++)
+		{
 			NavMeshGeometry.AreaIndices[Index].Reset();
+		}
 		NavMeshGeometry.MeshVerts.Reset();
-		for (int32 Index=0; Index < RECAST_MAX_AREAS; Index++)
+		for (int32 Index = 0; Index < RECAST_MAX_AREAS; Index++)
+		{
 			NavMeshGeometry.AreaIndices[Index].Reset();
+		}
 		NavMeshGeometry.BuiltMeshIndices.Reset();
 		NavMeshGeometry.PolyEdges.Reset();
 		NavMeshGeometry.NavMeshEdges.Reset();
@@ -55,8 +59,10 @@ struct FNavMeshSceneProxyData : public TSharedFromThis<FNavMeshSceneProxyData, E
 		NavMeshGeometry.Clusters.Reset();
 		NavMeshGeometry.ClusterLinks.Reset();
 		NavMeshGeometry.OffMeshSegments.Reset();
-		for (int32 Index=0; Index < RECAST_MAX_AREAS; Index++)
+		for (int32 Index = 0; Index < RECAST_MAX_AREAS; Index++)
+		{
 			NavMeshGeometry.OffMeshSegmentAreas[Index].Reset();
+		}
 		TileEdgeLines.Reset();
 		NavMeshEdgeLines.Reset();
 		NavLinkLines.Reset();
diff --git a/Engine/Source/Runtime/Engine/Public/AI/NavigationOctree.h b/Engine/Source/Runtime/Engine/Public/AI/NavigationOctree.h
--- a/Engine/Source/Runtime/Engine/Public/AI/NavigationOctree.h
+++ b/Engine/Source/Runtime/Engine/Public/AI/NavigationOctree.h
@@ -6,6 +6,7 @@
 
 class AActor;
 class UActorComponent;
+struct FNavDataConfig;
 
 struct ENGINE_API FNavigationOctreeFilter
 {
@@ -90,9 +91,9 @@ struct ENGINE_API FNavigationOctreeElement
 		return Data.Modifiers.HasMetaAreas() ? Data.Modifiers.GetInstantiatedMetaModifier(NavAgent, Owner) : Data.Modifiers;
 	}
 
-	FORCEINLINE bool ShouldUseGeometry(const struct FNavDataConfig* NavConfig) const
+	FORCEINLINE bool ShouldUseGeometry(const FNavDataConfig& NavConfig) const
 	{ 
-		return !Data.ShouldUseGeometryDelegate.IsBound() || Data.ShouldUseGeometryDelegate.Execute(NavConfig);
+		return !Data.ShouldUseGeometryDelegate.IsBound() || Data.ShouldUseGeometryDelegate.Execute(&NavConfig);
 	}
 
 	FORCEINLINE int32 GetAllocatedSize() const
