Fixed Engine DisplayDebug functions using YL incorrectly

UE-7050 - Wrapped strings on Canvas from ShowDebug console command are drawn overlapping other lines.

YL is supposed to be an in/out parameter, however the given value of YL was being used to adjust YPos for all text drawn by DisplayDebug. This is wrong, as YL is supposed to be updated for each draw call, as it affects the start of the next line (as an out parameter).

diff --git a/Engine/Source/Runtime/AIModule/Private/AIController.cpp b/Engine/Source/Runtime/AIModule/Private/AIController.cpp
--- a/Engine/Source/Runtime/AIModule/Private/AIController.cpp
+++ b/Engine/Source/Runtime/AIModule/Private/AIController.cpp
@@ -109,7 +109,7 @@ void AAIController::DisplayDebug(UCanvas* Canvas, const FDebugDisplayInfo& Debug
 		AActor* FocusActor = GetFocusActor();
 		if (FocusActor)
 		{
-			Canvas->DrawText(GEngine->GetSmallFont(), FString::Printf(TEXT("      Focus %s"), *FocusActor->GetName()), 4.0f, YPos);
+			YL = Canvas->DrawText(GEngine->GetSmallFont(), FString::Printf(TEXT("      Focus %s"), *FocusActor->GetName()), 4.0f, YPos);
 			YPos += YL;
 		}
 	}
diff --git a/Engine/Source/Runtime/AIModule/Private/Navigation/PathFollowingComponent.cpp b/Engine/Source/Runtime/AIModule/Private/Navigation/PathFollowingComponent.cpp
--- a/Engine/Source/Runtime/AIModule/Private/Navigation/PathFollowingComponent.cpp
+++ b/Engine/Source/Runtime/AIModule/Private/Navigation/PathFollowingComponent.cpp
@@ -1306,7 +1306,7 @@ void UPathFollowingComponent::DisplayDebug(UCanvas* Canvas, const FDebugDisplayI
 	Canvas->SetDrawColor(FColor::Blue);
 	UFont* RenderFont = GEngine->GetSmallFont();
 	FString StatusDesc = FString::Printf(TEXT("  Move status: %s"), *GetStatusDesc());
-	Canvas->DrawText(RenderFont, StatusDesc, 4.0f, YPos);
+	YL = Canvas->DrawText(RenderFont, StatusDesc, 4.0f, YPos);
 	YPos += YL;
 
 	if (Status == EPathFollowingStatus::Moving)
@@ -1315,7 +1315,7 @@ void UPathFollowingComponent::DisplayDebug(UCanvas* Canvas, const FDebugDisplayI
 		FString TargetDesc = FString::Printf(TEXT("  Move target [%d/%d]: %s (%s)"),
 			MoveSegmentEndIndex, NumMoveSegments, *GetCurrentTargetLocation().ToString(), *GetNameSafe(DestinationActor.Get()));
 		
-		Canvas->DrawText(RenderFont, TargetDesc, 4.0f, YPos);
+		YL = Canvas->DrawText(RenderFont, TargetDesc, 4.0f, YPos);
 		YPos += YL;
 	}
 }
diff --git a/Engine/Source/Runtime/Engine/Classes/GameFramework/Actor.h b/Engine/Source/Runtime/Engine/Classes/GameFramework/Actor.h
--- a/Engine/Source/Runtime/Engine/Classes/GameFramework/Actor.h
+++ b/Engine/Source/Runtime/Engine/Classes/GameFramework/Actor.h
@@ -2295,10 +2295,15 @@ public:
 	/**
 	 * Draw important Actor variables on canvas.  HUD will call DisplayDebug() on the current ViewTarget when the ShowDebug exec is used
 	 *
-	 * @param Canvas - Canvas to draw on
-	 * @param DebugDisplay - Contains information about what debug data to display
-	 * @param YL - Height of the current font
-	 * @param YPos - Y position on Canvas. YPos += YL, gives position to draw text for next debug line.
+	 * @param Canvas			Canvas to draw on
+	 *
+	 * @param DebugDisplay		Contains information about what debug data to display
+	 *
+	 * @param YL				[in]	Height of the previously drawn line.
+	 *							[out]	Height of the last line drawn by this function.
+	 *
+	 * @param YPos				[in]	Y position on Canvas for the previously drawn line. YPos += YL, gives position to draw text for next debug line.
+	 *							[out]	Y position on Canvas for the last line drawn by this function.
 	 */
 	virtual void DisplayDebug(class UCanvas* Canvas, const class FDebugDisplayInfo& DebugDisplay, float& YL, float& YPos);
 
diff --git a/Engine/Source/Runtime/Engine/Private/Actor.cpp b/Engine/Source/Runtime/Engine/Private/Actor.cpp
--- a/Engine/Source/Runtime/Engine/Private/Actor.cpp
+++ b/Engine/Source/Runtime/Engine/Private/Actor.cpp
@@ -1794,7 +1794,7 @@ void AActor::DisplayDebug(UCanvas* Canvas, const FDebugDisplayInfo& DebugDisplay
 	UFont* RenderFont = GEngine->GetSmallFont();
 	if( T != "" )
 	{
-		Canvas->DrawText(RenderFont, T, 4.0f, YPos);
+		YL = Canvas->DrawText(RenderFont, T, 4.0f, YPos);
 		YPos += YL;
 	}
 
@@ -1811,17 +1811,17 @@ void AActor::DisplayDebug(UCanvas* Canvas, const FDebugDisplayInfo& DebugDisplay
 			{
 				T = T + FString(TEXT(" Tear Off"));
 			}
-			Canvas->DrawText(RenderFont, T, 4.0f, YPos);
+			YL = Canvas->DrawText(RenderFont, T, 4.0f, YPos);
 			YPos += YL;
 		}
 	}
 
-	Canvas->DrawText(RenderFont, FString::Printf(TEXT("Location: %s Rotation: %s"), *GetActorLocation().ToString(), *GetActorRotation().ToString()), 4.0f, YPos);
+	YL = Canvas->DrawText(RenderFont, FString::Printf(TEXT("Location: %s Rotation: %s"), *GetActorLocation().ToString(), *GetActorRotation().ToString()), 4.0f, YPos);
 	YPos += YL;
 
 	if( DebugDisplay.IsDisplayOn(TEXT("physics")) )
 	{
-		Canvas->DrawText(RenderFont,FString::Printf(TEXT("Velocity: %s Speed: %f Speed2D: %f"), *GetVelocity().ToString(), GetVelocity().Size(), GetVelocity().Size2D()), 4.0f, YPos);
+		YL = Canvas->DrawText(RenderFont,FString::Printf(TEXT("Velocity: %s Speed: %f Speed2D: %f"), *GetVelocity().ToString(), GetVelocity().Size(), GetVelocity().Size2D()), 4.0f, YPos);
 		YPos += YL;
 	}
 
@@ -1830,12 +1830,12 @@ void AActor::DisplayDebug(UCanvas* Canvas, const FDebugDisplayInfo& DebugDisplay
 		Canvas->DrawColor.B = 0;
 		float MyRadius, MyHeight;
 		GetComponentsBoundingCylinder(MyRadius, MyHeight);
-		Canvas->DrawText(RenderFont, FString::Printf(TEXT("Collision Radius: %f Height: %f"), MyRadius, MyHeight), 4.0f, YPos);
+		YL = Canvas->DrawText(RenderFont, FString::Printf(TEXT("Collision Radius: %f Height: %f"), MyRadius, MyHeight), 4.0f, YPos);
 		YPos += YL;
 
 		if ( RootComponent == NULL )
 		{
-			Canvas->DrawText(RenderFont, FString(TEXT("No RootComponent")), 4.0f, YPos );
+			YL = Canvas->DrawText(RenderFont, FString(TEXT("No RootComponent")), 4.0f, YPos );
 			YPos += YL;
 		}
 
@@ -1859,10 +1859,10 @@ void AActor::DisplayDebug(UCanvas* Canvas, const FDebugDisplayInfo& DebugDisplay
 		{
 			T = TEXT("Overlapping nothing");
 		}
-		Canvas->DrawText(RenderFont,T, 4,YPos);
+		YL = Canvas->DrawText(RenderFont,T, 4,YPos);
 		YPos += YL;
 	}
-	Canvas->DrawText( RenderFont,FString::Printf(TEXT(" Instigator: %s Owner: %s"), (Instigator ? *Instigator->GetName() : TEXT("None")),
+	YL = Canvas->DrawText( RenderFont,FString::Printf(TEXT(" Instigator: %s Owner: %s"), (Instigator ? *Instigator->GetName() : TEXT("None")),
 		(Owner ? *Owner->GetName() : TEXT("None"))), 4,YPos);
 	YPos += YL;
 
diff --git a/Engine/Source/Runtime/Engine/Private/Animation/AnimInstance.cpp b/Engine/Source/Runtime/Engine/Private/Animation/AnimInstance.cpp
--- a/Engine/Source/Runtime/Engine/Private/Animation/AnimInstance.cpp
+++ b/Engine/Source/Runtime/Engine/Private/Animation/AnimInstance.cpp
@@ -695,7 +695,7 @@ void OutputCurveMap(TMap<FName, float>& CurveMap, UCanvas* Canvas, UFont* Render
 	for (FName CurveName : Names)
 	{
 		FString CurveEntry = FString::Printf(TEXT("%s: %.3f"), *CurveName.ToString(), CurveMap[CurveName]);
-		Canvas->DrawText(RenderFont, CurveEntry, Indent, YPos, 1.f, 1.f, RenderInfo);
+		YL = Canvas->DrawText(RenderFont, CurveEntry, Indent, YPos, 1.f, 1.f, RenderInfo);
 		YPos += YL;
 	}
 }
@@ -797,7 +797,7 @@ void UAnimInstance::DisplayDebug(class UCanvas* Canvas, const FDebugDisplayInfo&
 	Canvas->SetLinearDrawColor(TextYellow);
 
 	FString Heading = FString::Printf(TEXT("Animation: %s"), *GetName());
-	Canvas->DrawText(RenderFont, Heading, Indent, YPos, 1.f, 1.f, RenderInfo);
+	YL = Canvas->DrawText(RenderFont, Heading, Indent, YPos, 1.f, 1.f, RenderInfo);
 	YPos += YL;
 	if (bShowSyncGroups)
 	{
@@ -805,7 +805,7 @@ void UAnimInstance::DisplayDebug(class UCanvas* Canvas, const FDebugDisplayInfo&
 
 		//Display Sync Groups
 		Heading = FString::Printf(TEXT("SyncGroups: %i"), SyncGroups.Num());
-		Canvas->DrawText(RenderFont, Heading, Indent, YPos, 1.f, 1.f, RenderInfo);
+		YL = Canvas->DrawText(RenderFont, Heading, Indent, YPos, 1.f, 1.f, RenderInfo);
 		YPos += YL;
 
 		for (int32 GroupIndex = 0; GroupIndex < SyncGroups.Num(); ++GroupIndex)
@@ -816,7 +816,7 @@ void UAnimInstance::DisplayDebug(class UCanvas* Canvas, const FDebugDisplayInfo&
 			Canvas->SetLinearDrawColor(TextYellow);
 
 			FString GroupLabel = FString::Printf(TEXT("Group %i - Players %i"), GroupIndex, SyncGroup.ActivePlayers.Num());
-			Canvas->DrawText(RenderFont, GroupLabel, Indent, YPos, 1.f, 1.f, RenderInfo);
+			YL = Canvas->DrawText(RenderFont, GroupLabel, Indent, YPos, 1.f, 1.f, RenderInfo);
 			YPos += YL;
 
 			if (SyncGroup.ActivePlayers.Num() > 0)
@@ -829,7 +829,7 @@ void UAnimInstance::DisplayDebug(class UCanvas* Canvas, const FDebugDisplayInfo&
 		Canvas->SetLinearDrawColor(TextYellow);
 
 		Heading = FString::Printf(TEXT("Ungrouped: %i"), UngroupedActivePlayers.Num());
-		Canvas->DrawText(RenderFont, Heading, Indent, YPos, 1.f, 1.f, RenderInfo);
+		YL = Canvas->DrawText(RenderFont, Heading, Indent, YPos, 1.f, 1.f, RenderInfo);
 		YPos += YL;
 
 		Canvas->SetLinearDrawColor(TextWhite);
@@ -842,7 +842,7 @@ void UAnimInstance::DisplayDebug(class UCanvas* Canvas, const FDebugDisplayInfo&
 		Canvas->SetLinearDrawColor(TextYellow);
 
 		Heading = FString::Printf(TEXT("Montages: %i"), MontageInstances.Num());
-		Canvas->DrawText(RenderFont, Heading, Indent, YPos, 1.f, 1.f, RenderInfo);
+		YL = Canvas->DrawText(RenderFont, Heading, Indent, YPos, 1.f, 1.f, RenderInfo);
 		YPos += YL;
 
 		for (int32 MontageIndex = 0; MontageIndex < MontageInstances.Num(); ++MontageIndex)
@@ -854,7 +854,7 @@ void UAnimInstance::DisplayDebug(class UCanvas* Canvas, const FDebugDisplayInfo&
 			Canvas->SetLinearDrawColor((MontageInstance->IsActive()) ? ActiveColor : TextWhite);
 
 			FString MontageEntry = FString::Printf(TEXT("%i) %s CurrSec: %s NextSec: %s W:%.3f DW:%.3f"), MontageIndex, *MontageInstance->Montage->GetName(), *MontageInstance->GetCurrentSection().ToString(), *MontageInstance->GetNextSection().ToString(), MontageInstance->Weight, MontageInstance->DesiredWeight);
-			Canvas->DrawText(RenderFont, MontageEntry, Indent, YPos, 1.f, 1.f, RenderInfo);
+			YL = Canvas->DrawText(RenderFont, MontageEntry, Indent, YPos, 1.f, 1.f, RenderInfo);
 			YPos += YL;
 		}
 	}
@@ -864,7 +864,7 @@ void UAnimInstance::DisplayDebug(class UCanvas* Canvas, const FDebugDisplayInfo&
 		Canvas->SetLinearDrawColor(TextYellow);
 
 		Heading = FString::Printf(TEXT("Active Notify States: %i"), ActiveAnimNotifyState.Num());
-		Canvas->DrawText(RenderFont, Heading, Indent, YPos, 1.f, 1.f, RenderInfo);
+		YL = Canvas->DrawText(RenderFont, Heading, Indent, YPos, 1.f, 1.f, RenderInfo);
 		YPos += YL;
 
 		Canvas->SetLinearDrawColor(TextWhite);
@@ -876,7 +876,7 @@ void UAnimInstance::DisplayDebug(class UCanvas* Canvas, const FDebugDisplayInfo&
 			const FAnimNotifyEvent& NotifyState = ActiveAnimNotifyState[NotifyIndex];
 
 			FString NotifyEntry = FString::Printf(TEXT("%i) %s Class: %s Dur:%.3f"), NotifyIndex, *NotifyState.NotifyName.ToString(), *NotifyState.NotifyStateClass->GetName(), NotifyState.GetDuration());
-			Canvas->DrawText(RenderFont, NotifyEntry, Indent, YPos, 1.f, 1.f, RenderInfo);
+			YL = Canvas->DrawText(RenderFont, NotifyEntry, Indent, YPos, 1.f, 1.f, RenderInfo);
 			YPos += YL;
 		}
 	}
@@ -885,14 +885,14 @@ void UAnimInstance::DisplayDebug(class UCanvas* Canvas, const FDebugDisplayInfo&
 	{
 		Canvas->SetLinearDrawColor(TextYellow);
 
-		Canvas->DrawText(RenderFont, TEXT("Curves"), Indent, YPos, 1.f, 1.f, RenderInfo);
+		YL = Canvas->DrawText(RenderFont, TEXT("Curves"), Indent, YPos, 1.f, 1.f, RenderInfo);
 		YPos += YL;
 
 		{
 			FIndenter CurveIndent(Indent);
 
 			Heading = FString::Printf(TEXT("Morph Curves: %i"), MorphTargetCurves.Num());
-			Canvas->DrawText(RenderFont, Heading, Indent, YPos, 1.f, 1.f, RenderInfo);
+			YL = Canvas->DrawText(RenderFont, Heading, Indent, YPos, 1.f, 1.f, RenderInfo);
 			YPos += YL;
 
 			Canvas->SetLinearDrawColor(TextWhite);
@@ -905,7 +905,7 @@ void UAnimInstance::DisplayDebug(class UCanvas* Canvas, const FDebugDisplayInfo&
 			Canvas->SetLinearDrawColor(TextYellow);
 
 			Heading = FString::Printf(TEXT("Material Curves: %i"), MaterialParameterCurves.Num());
-			Canvas->DrawText(RenderFont, Heading, Indent, YPos, 1.f, 1.f, RenderInfo);
+			YL = Canvas->DrawText(RenderFont, Heading, Indent, YPos, 1.f, 1.f, RenderInfo);
 			YPos += YL;
 
 			Canvas->SetLinearDrawColor(TextWhite);
@@ -918,7 +918,7 @@ void UAnimInstance::DisplayDebug(class UCanvas* Canvas, const FDebugDisplayInfo&
 			Canvas->SetLinearDrawColor(TextYellow);
 
 			Heading = FString::Printf(TEXT("Event Curves: %i"), EventCurves.Num());
-			Canvas->DrawText(RenderFont, Heading, Indent, YPos, 1.f, 1.f, RenderInfo);
+			YL = Canvas->DrawText(RenderFont, Heading, Indent, YPos, 1.f, 1.f, RenderInfo);
 			YPos += YL;
 
 			Canvas->SetLinearDrawColor(TextWhite);
@@ -935,7 +935,7 @@ void UAnimInstance::DisplayDebug(class UCanvas* Canvas, const FDebugDisplayInfo&
 		Canvas->SetLinearDrawColor(TextYellow);
 
 		YPos += YL;
-		Canvas->DrawText(RenderFont, TEXT("Anim Node Tree"), Indent, YPos, 1.f, 1.f, RenderInfo);
+		YL = Canvas->DrawText(RenderFont, TEXT("Anim Node Tree"), Indent, YPos, 1.f, 1.f, RenderInfo);
 
 		const float NodeIndent = 8.f;
 		const float LineIndent = 4.f;
@@ -951,7 +951,6 @@ void UAnimInstance::DisplayDebug(class UCanvas* Canvas, const FDebugDisplayInfo&
 
 		TArray<float> VerticalLineStarts; // Index represents indent level, track the current starting point for that 
 
-		int32 HalfStep = int32(YL / 2);
 		int32 PrevChainID = -1;
 
 		for (FNodeDebugData::FFlattenedDebugData& Line : FlattenedData)
@@ -965,6 +964,7 @@ void UAnimInstance::DisplayDebug(class UCanvas* Canvas, const FDebugDisplayInfo&
 
 			if (PrevChainID != Line.ChainID)
 			{
+				const int32 HalfStep = int32(YL / 2);
 				YPos += HalfStep; // Extra spacing to delimit different chains, CurrLineYBase now 
 				// roughly represents middle of text line, so we can use it for line drawing
 
@@ -998,7 +998,7 @@ void UAnimInstance::DisplayDebug(class UCanvas* Canvas, const FDebugDisplayInfo&
 			PrevChainID = Line.ChainID;
 			FLinearColor ItemColor = Line.bPoseSource ? PoseSourceColor : ActiveColor;
 			Canvas->SetLinearDrawColor(Line.IsOnActiveBranch() ? ItemColor : InactiveColor);
-			Canvas->DrawText(RenderFont, Line.DebugLine, CurrIndent, YPos, 1.f, 1.f, RenderInfo);
+			YL = Canvas->DrawText(RenderFont, Line.DebugLine, CurrIndent, YPos, 1.f, 1.f, RenderInfo);
 			YPos += YL;
 		}
 	}
diff --git a/Engine/Source/Runtime/Engine/Private/Character.cpp b/Engine/Source/Runtime/Engine/Private/Character.cpp
--- a/Engine/Source/Runtime/Engine/Private/Character.cpp
+++ b/Engine/Source/Runtime/Engine/Private/Character.cpp
@@ -778,7 +778,7 @@ void ACharacter::DisplayDebug(UCanvas* Canvas, const FDebugDisplayInfo& DebugDis
 			BaseString = FString::Printf(TEXT("Based On %s"), *BaseString);
 		}
 		
-		Canvas->DrawText(RenderFont, FString::Printf(TEXT("RelativeLoc: %s Rot: %s %s"), *BasedMovement.Location.ToString(), *BasedMovement.Rotation.ToString(), *BaseString), Indent, YPos);
+		YL = Canvas->DrawText(RenderFont, FString::Printf(TEXT("RelativeLoc: %s Rot: %s %s"), *BasedMovement.Location.ToString(), *BasedMovement.Rotation.ToString(), *BaseString), Indent, YPos);
 		YPos += YL;
 
 		if ( CharacterMovement != NULL )
@@ -787,7 +787,7 @@ void ACharacter::DisplayDebug(UCanvas* Canvas, const FDebugDisplayInfo& DebugDis
 		}
 		const bool Crouched = CharacterMovement && CharacterMovement->IsCrouching();
 		FString T = FString::Printf(TEXT("Crouched %i"), Crouched);
-		Canvas->DrawText(RenderFont, T, Indent, YPos );
+		YL = Canvas->DrawText(RenderFont, T, Indent, YPos );
 		YPos += YL;
 	}
 }
diff --git a/Engine/Source/Runtime/Engine/Private/CharacterMovementComponent.cpp b/Engine/Source/Runtime/Engine/Private/CharacterMovementComponent.cpp
--- a/Engine/Source/Runtime/Engine/Private/CharacterMovementComponent.cpp
+++ b/Engine/Source/Runtime/Engine/Private/CharacterMovementComponent.cpp
@@ -5245,15 +5245,15 @@ void UCharacterMovementComponent::DisplayDebug(UCanvas* Canvas, const FDebugDisp
 	Canvas->SetDrawColor(255,255,255);
 	UFont* RenderFont = GEngine->GetSmallFont();
 	FString T = FString::Printf(TEXT("CHARACTER MOVEMENT Floor %s Crouched %i"), *CurrentFloor.HitResult.ImpactNormal.ToString(), IsCrouching());
-	Canvas->DrawText(RenderFont, T, 4.0f, YPos );
+	YL = Canvas->DrawText(RenderFont, T, 4.0f, YPos );
 	YPos += YL;
 
 	T = FString::Printf(TEXT("Updated Component: %s"), *UpdatedComponent->GetName());
-	Canvas->DrawText(RenderFont, T, 4.0f, YPos );
+	YL = Canvas->DrawText(RenderFont, T, 4.0f, YPos );
 	YPos += YL;
 
 	T = FString::Printf(TEXT("bForceMaxAccel: %i"), bForceMaxAccel);
-	Canvas->DrawText(RenderFont, T, 4.0f, YPos );
+	YL = Canvas->DrawText(RenderFont, T, 4.0f, YPos );
 	YPos += YL;
 
 	APhysicsVolume * PhysicsVolume = GetPhysicsVolume();
@@ -5263,7 +5263,7 @@ void UCharacterMovementComponent::DisplayDebug(UCanvas* Canvas, const FDebugDisp
 
 	T = FString::Printf(TEXT("%s In physicsvolume %s on base %s component %s gravity %f"), *GetMovementName(), (PhysicsVolume ? *PhysicsVolume->GetName() : TEXT("None")),
 		(BaseActor ? *BaseActor->GetName() : TEXT("None")), (BaseComponent ? *BaseComponent->GetName() : TEXT("None")), GetGravityZ());
-	Canvas->DrawText(RenderFont, T, 4.0f, YPos );
+	YL = Canvas->DrawText(RenderFont, T, 4.0f, YPos );
 	YPos += YL;
 }
 
diff --git a/Engine/Source/Runtime/Engine/Private/Controller.cpp b/Engine/Source/Runtime/Engine/Private/Controller.cpp
--- a/Engine/Source/Runtime/Engine/Private/Controller.cpp
+++ b/Engine/Source/Runtime/Engine/Private/Controller.cpp
@@ -458,7 +458,7 @@ void AController::DisplayDebug(UCanvas* Canvas, const FDebugDisplayInfo& DebugDi
 	{
 		if (PlayerState == NULL)
 		{
-			Canvas->DrawText(RenderFont, TEXT("NO PlayerState"), 4.0f, YPos );
+			YL = Canvas->DrawText(RenderFont, TEXT("NO PlayerState"), 4.0f, YPos );
 		}
 		else
 		{
@@ -471,7 +471,7 @@ void AController::DisplayDebug(UCanvas* Canvas, const FDebugDisplayInfo& DebugDi
 	}
 
 	Canvas->SetDrawColor(255,0,0);
-	Canvas->DrawText(RenderFont, FString::Printf(TEXT("CONTROLLER %s Pawn %s"), *GetName(), *Pawn->GetName()), 4.0f, YPos );
+	YL = Canvas->DrawText(RenderFont, FString::Printf(TEXT("CONTROLLER %s Pawn %s"), *GetName(), *Pawn->GetName()), 4.0f, YPos );
 	YPos += YL;
 }
 
diff --git a/Engine/Source/Runtime/Engine/Private/Pawn.cpp b/Engine/Source/Runtime/Engine/Private/Pawn.cpp
--- a/Engine/Source/Runtime/Engine/Private/Pawn.cpp
+++ b/Engine/Source/Runtime/Engine/Private/Pawn.cpp
@@ -745,7 +745,7 @@ void APawn::DisplayDebug(class UCanvas* Canvas, const FDebugDisplayInfo& DebugDi
 // 		UPrimitiveComponent* const OtherComp = It.Value().Get();
 // 		AActor* OtherActor = OtherComp ? OtherComp->GetOwner() : NULL;
 // 
-// 		HUD->Canvas->DrawText(FString::Printf(TEXT("TOUCHING my %s to %s's %s"), *GetNameSafe(MyComp), *GetNameSafe(OtherActor), *GetNameSafe(OtherComp)));
+// 		YL = HUD->Canvas->DrawText(FString::Printf(TEXT("TOUCHING my %s to %s's %s"), *GetNameSafe(MyComp), *GetNameSafe(OtherActor), *GetNameSafe(OtherComp)));
 // 		YPos += YL;
 // 		HUD->Canvas->SetPos(4,YPos);
 // 	}
@@ -753,7 +753,7 @@ void APawn::DisplayDebug(class UCanvas* Canvas, const FDebugDisplayInfo& DebugDi
 	UFont* RenderFont = GEngine->GetSmallFont();
 	if ( PlayerState == NULL )
 	{
-		Canvas->DrawText(RenderFont, TEXT("NO PlayerState"), 4.0f, YPos );
+		YL = Canvas->DrawText(RenderFont, TEXT("NO PlayerState"), 4.0f, YPos );
 		YPos += YL;
 	}
 	else
@@ -768,7 +768,7 @@ void APawn::DisplayDebug(class UCanvas* Canvas, const FDebugDisplayInfo& DebugDi
 
 	if (DebugDisplay.IsDisplayOn(NAME_Camera))
 	{
-		Canvas->DrawText(RenderFont, FString::Printf(TEXT("BaseEyeHeight %f"), BaseEyeHeight), 4.0f, YPos);
+		YL = Canvas->DrawText(RenderFont, FString::Printf(TEXT("BaseEyeHeight %f"), BaseEyeHeight), 4.0f, YPos);
 		YPos += YL;
 	}
 
@@ -776,7 +776,7 @@ void APawn::DisplayDebug(class UCanvas* Canvas, const FDebugDisplayInfo& DebugDi
 	if ( Controller == NULL )
 	{
 		Canvas->SetDrawColor(255,0,0);
-		Canvas->DrawText(RenderFont, TEXT("NO Controller"), 4.0f, YPos);
+		YL = Canvas->DrawText(RenderFont, TEXT("NO Controller"), 4.0f, YPos);
 		YPos += YL;
 		//HUD->PlayerOwner->DisplayDebug(Canvas, DebugDisplay, YL, YPos);
 	}
diff --git a/Engine/Source/Runtime/Engine/Private/PlayerCameraManager.cpp b/Engine/Source/Runtime/Engine/Private/PlayerCameraManager.cpp
--- a/Engine/Source/Runtime/Engine/Private/PlayerCameraManager.cpp
+++ b/Engine/Source/Runtime/Engine/Private/PlayerCameraManager.cpp
@@ -955,14 +955,14 @@ void APlayerCameraManager::DisplayDebug(class UCanvas* Canvas, const FDebugDispl
 	Canvas->SetDrawColor(255,255,255);
 
 	UFont* RenderFont = GEngine->GetSmallFont();
-	Canvas->DrawText(RenderFont, FString::Printf(TEXT("   Camera Style:%s main ViewTarget:%s"), *CameraStyle.ToString(), *ViewTarget.Target->GetName()), 4.0f, YPos );
+	YL = Canvas->DrawText(RenderFont, FString::Printf(TEXT("   Camera Style:%s main ViewTarget:%s"), *CameraStyle.ToString(), *ViewTarget.Target->GetName()), 4.0f, YPos );
 	YPos += YL;
 
 	//@TODO: Print out more information
-	Canvas->DrawText(RenderFont, FString::Printf(TEXT("   CamLoc:%s CamRot:%s FOV:%f"), *CameraCache.POV.Location.ToCompactString(), *CameraCache.POV.Rotation.ToCompactString(), CameraCache.POV.FOV), 4.0f, YPos );
+	YL = Canvas->DrawText(RenderFont, FString::Printf(TEXT("   CamLoc:%s CamRot:%s FOV:%f"), *CameraCache.POV.Location.ToCompactString(), *CameraCache.POV.Rotation.ToCompactString(), CameraCache.POV.FOV), 4.0f, YPos );
 	YPos += YL;
 
-	Canvas->DrawText(RenderFont, FString::Printf(TEXT("   AspectRatio: %1.3f"), CameraCache.POV.AspectRatio), 4.0f, YPos );
+	YL = Canvas->DrawText(RenderFont, FString::Printf(TEXT("   AspectRatio: %1.3f"), CameraCache.POV.AspectRatio), 4.0f, YPos );
 	YPos += YL;
 }
 
diff --git a/Engine/Source/Runtime/Engine/Private/PlayerController.cpp b/Engine/Source/Runtime/Engine/Private/PlayerController.cpp
--- a/Engine/Source/Runtime/Engine/Private/PlayerController.cpp
+++ b/Engine/Source/Runtime/Engine/Private/PlayerController.cpp
@@ -2775,21 +2775,21 @@ void APlayerController::DisplayDebug(class UCanvas* Canvas, const FDebugDisplayI
 
 	Canvas->SetDrawColor(255,255,0);
 	UFont* RenderFont = GEngine->GetSmallFont();
-	Canvas->DrawText(RenderFont, FString::Printf(TEXT("STATE %s"), *GetStateName().ToString()), 4.0f, YPos );
+	YL = Canvas->DrawText(RenderFont, FString::Printf(TEXT("STATE %s"), *GetStateName().ToString()), 4.0f, YPos );
 	YPos += YL;
 
 	if (DebugDisplay.IsDisplayOn(NAME_Camera))
 	{
 		if (PlayerCameraManager != NULL)
 		{
-			Canvas->DrawText(RenderFont, "<<<< CAMERA >>>>", 4.0f, YPos );
+			YL = Canvas->DrawText(RenderFont, "<<<< CAMERA >>>>", 4.0f, YPos );
 			YPos += YL;
 			PlayerCameraManager->DisplayDebug( Canvas, DebugDisplay, YL, YPos );
 		}
 		else
 		{
 			Canvas->SetDrawColor(255,0,0);
-			Canvas->DrawText(RenderFont, "<<<< NO CAMERA >>>>", 4.0f, YPos );
+			YL = Canvas->DrawText(RenderFont, "<<<< NO CAMERA >>>>", 4.0f, YPos );
 			YPos += YL;
 		}
 	}
@@ -2799,7 +2799,7 @@ void APlayerController::DisplayDebug(class UCanvas* Canvas, const FDebugDisplayI
 		BuildInputStack(InputStack);
 
 		Canvas->SetDrawColor(255,255,255);
-		Canvas->DrawText(RenderFont, TEXT("<<<< INPUT STACK >>>"), 4.0f, YPos);
+		YL = Canvas->DrawText(RenderFont, TEXT("<<<< INPUT STACK >>>"), 4.0f, YPos);
 		YPos += YL;
 
 		for(int32 i=InputStack.Num() - 1; i >= 0; --i)
@@ -2808,11 +2808,11 @@ void APlayerController::DisplayDebug(class UCanvas* Canvas, const FDebugDisplayI
 			Canvas->SetDrawColor(255,255,255);
 			if (Owner)
 			{
-				Canvas->DrawText(RenderFont, FString::Printf(TEXT(" %s.%s"), *Owner->GetName(), *InputStack[i]->GetName()), 4.0f, YPos);
+				YL = Canvas->DrawText(RenderFont, FString::Printf(TEXT(" %s.%s"), *Owner->GetName(), *InputStack[i]->GetName()), 4.0f, YPos);
 			}
 			else
 			{
-				Canvas->DrawText(RenderFont, FString::Printf(TEXT(" %s"), *InputStack[i]->GetName()), 4.0f, YPos);
+				YL = Canvas->DrawText(RenderFont, FString::Printf(TEXT(" %s"), *InputStack[i]->GetName()), 4.0f, YPos);
 			}
 			YPos += YL;
 		}
@@ -2824,14 +2824,14 @@ void APlayerController::DisplayDebug(class UCanvas* Canvas, const FDebugDisplayI
 		else
 		{
 			Canvas->SetDrawColor(255,0,0);
-			Canvas->DrawText(RenderFont, "NO INPUT", 4.0f, YPos);
+			YL = Canvas->DrawText(RenderFont, "NO INPUT", 4.0f, YPos);
 			YPos += YL;
 		}
 	}
 	if ( DebugDisplay.IsDisplayOn("ForceFeedback"))
 	{
 		Canvas->SetDrawColor(255, 255, 255);
-		Canvas->DrawText(RenderFont, FString::Printf(TEXT("Force Feedback - Enabled: %s LL: %.2f LS: %.2f RL: %.2f RS: %.2f"), (bForceFeedbackEnabled ? TEXT("true") : TEXT("false")), ForceFeedbackValues.LeftLarge, ForceFeedbackValues.LeftSmall, ForceFeedbackValues.RightLarge, ForceFeedbackValues.RightSmall), 4.0f, YPos);
+		YL = Canvas->DrawText(RenderFont, FString::Printf(TEXT("Force Feedback - Enabled: %s LL: %.2f LS: %.2f RL: %.2f RS: %.2f"), (bForceFeedbackEnabled ? TEXT("true") : TEXT("false")), ForceFeedbackValues.LeftLarge, ForceFeedbackValues.LeftSmall, ForceFeedbackValues.RightLarge, ForceFeedbackValues.RightSmall), 4.0f, YPos);
 		YPos += YL;
 	}
 }
diff --git a/Engine/Source/Runtime/Engine/Private/UserInterface/PlayerInput.cpp b/Engine/Source/Runtime/Engine/Private/UserInterface/PlayerInput.cpp
--- a/Engine/Source/Runtime/Engine/Private/UserInterface/PlayerInput.cpp
+++ b/Engine/Source/Runtime/Engine/Private/UserInterface/PlayerInput.cpp
@@ -1203,7 +1203,7 @@ void UPlayerInput::DisplayDebug(class UCanvas* Canvas, const FDebugDisplayInfo&
 	{
 		Canvas->SetDrawColor(255,0,0);
 		UFont* RenderFont = GEngine->GetSmallFont();
-		Canvas->DrawText(RenderFont, FString::Printf(TEXT("INPUT %s"), *GetName()), 4.0f, YPos );
+		YL = Canvas->DrawText(RenderFont, FString::Printf(TEXT("INPUT %s"), *GetName()), 4.0f, YPos );
 		YPos += YL;
 
 		UWorld* World = GetWorld();
@@ -1224,12 +1224,12 @@ void UPlayerInput::DisplayDebug(class UCanvas* Canvas, const FDebugDisplayInfo&
 					Str += FString::Printf(TEXT(" time: %.2f"), WorldRealTimeSeconds - KeyState->LastUpDownTransitionTime);
 				}
 				Canvas->SetDrawColor(180,255,180);
-				Canvas->DrawText(RenderFont, Str,4.0f, YPos);
+				YL = Canvas->DrawText(RenderFont, Str,4.0f, YPos);
 			}
 			else
 			{
 				Canvas->SetDrawColor(180,180,180);
-				Canvas->DrawText(RenderFont, Str,4.0f, YPos);
+				YL = Canvas->DrawText(RenderFont, Str,4.0f, YPos);
 			}
 			YPos += YL;
 		}
@@ -1237,11 +1237,11 @@ void UPlayerInput::DisplayDebug(class UCanvas* Canvas, const FDebugDisplayInfo&
 		float const DetectedMouseSampleHz = MouseSamples / MouseSamplingTotal;
 
 		Canvas->SetDrawColor(FColor::White);
-		Canvas->DrawText(RenderFont, FString::Printf(TEXT("MouseSampleRate: %.2f"), DetectedMouseSampleHz),4.0f, YPos);
+		YL = Canvas->DrawText(RenderFont, FString::Printf(TEXT("MouseSampleRate: %.2f"), DetectedMouseSampleHz),4.0f, YPos);
 		YPos += YL;
-		Canvas->DrawText(RenderFont, FString::Printf(TEXT("MouseX ZeroTime: %.2f, Smoothed: %.2f"), ZeroTime[0], SmoothedMouse[0]),4.0f, YPos);
+		YL = Canvas->DrawText(RenderFont, FString::Printf(TEXT("MouseX ZeroTime: %.2f, Smoothed: %.2f"), ZeroTime[0], SmoothedMouse[0]),4.0f, YPos);
 		YPos += YL;
-		Canvas->DrawText(RenderFont, FString::Printf(TEXT("MouseY ZeroTime: %.2f, Smoothed: %.2f"), ZeroTime[1], SmoothedMouse[1]),4.0f, YPos);
+		YL = Canvas->DrawText(RenderFont, FString::Printf(TEXT("MouseY ZeroTime: %.2f, Smoothed: %.2f"), ZeroTime[1], SmoothedMouse[1]),4.0f, YPos);
 		YPos += YL;
 
 		if ( (ZeroTime[0] > 2.f && ZeroTime[1] > 2.f) && GetDefault<UInputSettings>()->bEnableMouseSmoothing )
diff --git a/Engine/Source/Runtime/GameplayAbilities/Private/AbilitySystemComponent.cpp b/Engine/Source/Runtime/GameplayAbilities/Private/AbilitySystemComponent.cpp
--- a/Engine/Source/Runtime/GameplayAbilities/Private/AbilitySystemComponent.cpp
+++ b/Engine/Source/Runtime/GameplayAbilities/Private/AbilitySystemComponent.cpp
@@ -1141,11 +1141,13 @@ void UAbilitySystemComponent::DisplayDebug(class UCanvas* Canvas, const class FD
 	GetOwnedGameplayTags(OwnerTags);
 
 	Canvas->SetDrawColor(FColor::White);
-	YPos += Canvas->DrawText(GEngine->GetTinyFont(), FString::Printf(TEXT("Owned Tags: %s"), *OwnerTags.ToStringSimple()), 4.f, YPos);
+	YL = Canvas->DrawText(GEngine->GetTinyFont(), FString::Printf(TEXT("Owned Tags: %s"), *OwnerTags.ToStringSimple()), 4.f, YPos);
 	YPos += YL;
 
 	TSet<FGameplayAttribute> DrawAttributes;
 
+	const float MaxCharHeight = GEngine->GetTinyFont()->GetMaxCharHeight();
+
 	// -------------------------------------------------------------
 
 
@@ -1170,7 +1172,8 @@ void UAbilitySystemComponent::DisplayDebug(class UCanvas* Canvas, const class FD
 				}
 
 				Canvas->SetDrawColor(FColor::White);
-				YPos += Canvas->DrawText(GEngine->GetTinyFont(), AttributeString, 4.f, YPos);
+				YL = Canvas->DrawText(GEngine->GetTinyFont(), AttributeString, 4.f, YPos);
+				YPos += YL;
 
 				DrawAttributes.Add(Attribute);
 
@@ -1191,11 +1194,12 @@ void UAbilitySystemComponent::DisplayDebug(class UCanvas* Canvas, const class FD
 							if (Mod.TargetTagReqs) SrcName += FString::Printf(TEXT("TargetTags: [%s]"), *Mod.TargetTagReqs->ToString());
 						}
 
-						YPos += Canvas->DrawText(GEngine->GetTinyFont(), FString::Printf(TEXT("   %s\t %.2f - %s"), *EGameplayModOpToString(ModOpIdx), Mod.EvaluatedMagnitude, *SrcName), 7.f, YPos);
+						YL = Canvas->DrawText(GEngine->GetTinyFont(), FString::Printf(TEXT("   %s\t %.2f - %s"), *EGameplayModOpToString(ModOpIdx), Mod.EvaluatedMagnitude, *SrcName), 7.f, YPos);
+						YPos += YL;
 
 					}
 				}
-				YPos += YL;
+				YPos += MaxCharHeight;
 			}
 		}
 	}
@@ -1234,13 +1238,15 @@ void UAbilitySystemComponent::DisplayDebug(class UCanvas* Canvas, const class FD
 
 			Canvas->SetDrawColor(ActiveGE.bIsInhibited ? FColor(128, 128, 128): FColor::White );
 
-			YPos += Canvas->DrawText(GEngine->GetTinyFont(), FString::Printf(TEXT("%s %s %s"), *ASC_CleanupName(GetNameSafe(ActiveGE.Spec.Def)), *DurationStr, *StackString ), 4.f, YPos);
+			YL = Canvas->DrawText(GEngine->GetTinyFont(), FString::Printf(TEXT("%s %s %s"), *ASC_CleanupName(GetNameSafe(ActiveGE.Spec.Def)), *DurationStr, *StackString ), 4.f, YPos);
+			YPos += YL;
 
 			FGameplayTagContainer GrantedTags;
 			ActiveGE.Spec.GetAllGrantedTags(GrantedTags);
 			if (GrantedTags.Num() > 0)
 			{
-				YPos += Canvas->DrawText(GEngine->GetTinyFont(), FString::Printf(TEXT("Granted Tags: %s"), *GrantedTags.ToStringSimple() ), 7.f, YPos);
+				YL = Canvas->DrawText(GEngine->GetTinyFont(), FString::Printf(TEXT("Granted Tags: %s"), *GrantedTags.ToStringSimple() ), 7.f, YPos);
+				YPos += YL;
 			}
 
 			for (int32 ModIdx=0; ModIdx < ActiveGE.Spec.Modifiers.Num(); ++ModIdx)
@@ -1262,12 +1268,13 @@ void UAbilitySystemComponent::DisplayDebug(class UCanvas* Canvas, const class FD
 					Canvas->SetDrawColor(FColor(128, 128, 128) );
 				}
 
-				YPos += Canvas->DrawText(GEngine->GetTinyFont(), FString::Printf(TEXT("Mod: %s. %s. %.2f"), *ModInfo.Attribute.GetName(), *EGameplayModOpToString(ModInfo.ModifierOp), ModSpec.GetEvaluatedMagnitude() ), 7.f, YPos);
+				YL = Canvas->DrawText(GEngine->GetTinyFont(), FString::Printf(TEXT("Mod: %s. %s. %.2f"), *ModInfo.Attribute.GetName(), *EGameplayModOpToString(ModInfo.ModifierOp), ModSpec.GetEvaluatedMagnitude() ), 7.f, YPos);
+				YPos += YL;
 
 				Canvas->SetDrawColor(ActiveGE.bIsInhibited ? FColor(128, 128, 128): FColor::White );
 			}
 
-			YPos += YL;
+			YPos += MaxCharHeight;
 		}
 	}
 
@@ -1288,11 +1295,12 @@ void UAbilitySystemComponent::DisplayDebug(class UCanvas* Canvas, const class FD
 				if (Attribute.IsValid())
 				{
 					float Value = GetNumericAttribute(Attribute);
-					YPos += Canvas->DrawText(GEngine->GetTinyFont(), FString::Printf(TEXT("%s %.2f"), *Attribute.GetName(), Value ), 4.f, YPos);
+					YL = Canvas->DrawText(GEngine->GetTinyFont(), FString::Printf(TEXT("%s %.2f"), *Attribute.GetName(), Value ), 4.f, YPos);
+					YPos += YL;
 				}
 			}
 		}
-		YPos += YL;
+		YPos += MaxCharHeight;
 	}
 
 	// -------------------------------------------------------------
@@ -1305,7 +1313,8 @@ void UAbilitySystemComponent::DisplayDebug(class UCanvas* Canvas, const class FD
 				continue;
 
 			Canvas->SetDrawColor(AbilitySpec.IsActive() ? FColor::Yellow : FColor(128, 128, 128));
-			YPos += Canvas->DrawText(GEngine->GetTinyFont(), FString::Printf(TEXT("%s"), *ASC_CleanupName(GetNameSafe(AbilitySpec.Ability))), 4.f, YPos);
+			YL = Canvas->DrawText(GEngine->GetTinyFont(), FString::Printf(TEXT("%s"), *ASC_CleanupName(GetNameSafe(AbilitySpec.Ability))), 4.f, YPos);
+			YPos += YL;
 	
 			if (AbilitySpec.IsActive())
 			{
@@ -1322,20 +1331,24 @@ void UAbilitySystemComponent::DisplayDebug(class UCanvas* Canvas, const class FD
 						UAbilityTask* Task = TaskPtr.Get();
 						if (Task)
 						{
-							YPos += Canvas->DrawText(GEngine->GetTinyFont(), FString::Printf(TEXT("%s"), *Task->GetDebugString()), 7.f, YPos);
+							YL = Canvas->DrawText(GEngine->GetTinyFont(), FString::Printf(TEXT("%s"), *Task->GetDebugString()), 7.f, YPos);
+							YPos += YL;
 						}
 					}
 
 					if (InstanceIdx < Instances.Num() - 2)
 					{
 						Canvas->SetDrawColor(FColor(128, 128, 128));
-						YPos += Canvas->DrawText(GEngine->GetTinyFont(), FString::Printf(TEXT("--------")), 7.f, YPos);
+						YL = Canvas->DrawText(GEngine->GetTinyFont(), FString::Printf(TEXT("--------")), 7.f, YPos);
+						YPos += YL;
 					}
 				}
 			}
 		}
-		YPos += YL;
+		YPos += MaxCharHeight;
 	}
+
+	YL = MaxCharHeight;
 }
 
 #undef LOCTEXT_NAMESPACE
