Navigation's Active Tiles, pass 0.9 #UE4

- still very much experimental!

diff --git a/Engine/Source/Runtime/Engine/Classes/AI/Navigation/NavigationSystem.h b/Engine/Source/Runtime/Engine/Classes/AI/Navigation/NavigationSystem.h
--- a/Engine/Source/Runtime/Engine/Classes/AI/Navigation/NavigationSystem.h
+++ b/Engine/Source/Runtime/Engine/Classes/AI/Navigation/NavigationSystem.h
@@ -175,6 +175,19 @@ public:
 	UPROPERTY(config, EditAnywhere, Category=NavigationSystem)
 	uint32 bSkipAgentHeightCheckWhenPickingNavData:1;
 
+	/** If set to true navigation will be generated only around registered "navigation enforcers"
+	*	This has a range of consequences (including how navigation octree operates) so it needs to
+	*	be a conscious decision.
+	*	Once enabled results in whole world being navigable.
+	*	@see RegisterNavigationInvoker
+	*/
+	UPROPERTY(EditDefaultsOnly, Category = "Navigation", config)
+	uint32 bGenerateNavigationOnlyAroundNavigationInvokers : 1;
+
+	/** Minimal time between active tiles set update */
+	UPROPERTY(EditAnywhere, Category = "Navigation Enforcing", meta = (ClampMin = "0.1", UIMin = "0.1", EditCondition = "bGenerateNavigationOnlyAroundNavigationInvokers"), config)
+	float ActiveTilesUpdateInterval;
+
 	UPROPERTY(config, EditAnywhere, Category = Agents)
 	TArray<FNavDataConfig> SupportedAgents;
 	
@@ -201,13 +214,18 @@ public:
 	
 private:
 	TWeakObjectPtr<UCrowdManager> CrowdManager;
-
+	
 	/** set to true when navigation processing was blocked due to missing nav bounds */
-	uint32 bNavDataRemovedDueToMissingNavBounds:1;
-
+	uint32 bNavDataRemovedDueToMissingNavBounds : 1;
+	
 	/** All areas where we build/have navigation */
 	TSet<FNavigationBounds> RegisteredNavBounds;
 
+	TMap<AActor*, FNavigationInvoker> Invokers;
+
+	float NextInvokersUpdateTime;
+	void UpdateInvokers();
+
 public:
 	//----------------------------------------------------------------------//
 	// Blueprint functions
@@ -261,6 +279,19 @@ public:
 	UFUNCTION(BlueprintCallable, Category="AI|Navigation", meta=(WorldContext="WorldContext" ))
 	static bool NavigationRaycast(UObject* WorldContext, const FVector& RayStart, const FVector& RayEnd, FVector& HitLocation, TSubclassOf<UNavigationQueryFilter> FilterClass = NULL, AController* Querier = NULL);
 
+	/** Registers given actor as a "navigation enforcer" which means navigation system will
+	 *	make sure navigation is being generated in specified radius around it.
+	 *	@note: you need NavigationSystem's GenerateNavigationOnlyAroundNavigationInvokers to be set to true
+	 *		to take advantage of this feature
+	 */
+	UFUNCTION(BlueprintCallable, Category = Navigation)
+	void RegisterNavigationInvoker(AActor* Invoker, float TileGenerationRadius = 3000, float TileRemovalRadius = 5000);
+
+	/** Removes given actor from the list of active navigation enforcers.
+	 *	@see RegisterNavigationInvoker for more details */
+	UFUNCTION(BlueprintCallable, Category = Navigation)
+	void UnregisterNavigationInvoker(AActor* Invoker);
+
 	/** delegate type for events that dirty the navigation data ( Params: const FBox& DirtyBounds ) */
 	DECLARE_MULTICAST_DELEGATE_OneParam(FOnNavigationDirty, const FBox&);
 	/** called after navigation influencing event takes place*/
@@ -448,6 +479,14 @@ public:
 	static bool DoesPathIntersectBox(const FNavigationPath* Path, const FBox& Box, uint32 StartingIndex = 0);
 	static bool DoesPathIntersectBox(const FNavigationPath* Path, const FBox& Box, const FVector& AgentLocation, uint32 StartingIndex = 0);
 
+	//----------------------------------------------------------------------//
+	// Active tiles
+	//----------------------------------------------------------------------//
+	void RegisterInvoker(AActor& Invoker, float TileGenerationRadius, float TileRemovalRadius);
+	void UnregisterInvoker(AActor& Invoker);
+
+	static void RegisterNavigationInvoker(AActor& Invoker, float TileGenerationRadius, float TileRemovalRadius);
+	static void UnregisterNavigationInvoker(AActor& Invoker);
 
 	//----------------------------------------------------------------------//
 	// Bookkeeping 
diff --git a/Engine/Source/Runtime/Engine/Classes/AI/Navigation/NavigationTestingActor.h b/Engine/Source/Runtime/Engine/Classes/AI/Navigation/NavigationTestingActor.h
--- a/Engine/Source/Runtime/Engine/Classes/AI/Navigation/NavigationTestingActor.h
+++ b/Engine/Source/Runtime/Engine/Classes/AI/Navigation/NavigationTestingActor.h
@@ -6,6 +6,7 @@
 #include "NavigationTestingActor.generated.h"
 
 class ANavigationTestingActor;
+class UNavigationInvokerComponent;
 
 struct FNavTestTickHelper : FTickableGameObject
 {
@@ -45,6 +46,13 @@ private_subobject:
 	UPROPERTY()
 	class UNavTestRenderingComponent* EdRenderComp;
 #endif // WITH_EDITORONLY_DATA
+
+	UPROPERTY(EditAnywhere, Category = Navigation, meta=(EditCondition="bActAsNavigationInvoker"))
+	UNavigationInvokerComponent* InvokerComponent;
+
+	UPROPERTY()
+	uint32 bActAsNavigationInvoker : 1;
+
 public:
 
 	/** @todo document */
diff --git a/Engine/Source/Runtime/Engine/Classes/AI/Navigation/RecastNavMesh.h b/Engine/Source/Runtime/Engine/Classes/AI/Navigation/RecastNavMesh.h
--- a/Engine/Source/Runtime/Engine/Classes/AI/Navigation/RecastNavMesh.h
+++ b/Engine/Source/Runtime/Engine/Classes/AI/Navigation/RecastNavMesh.h
@@ -481,10 +481,10 @@ class ENGINE_API ARecastNavMesh : public ANavigationData
 	float CellHeight;
 
 	/** Radius of smallest agent to traverse this navmesh */
-	UPROPERTY(EditAnywhere, Category=Generation, config, meta=(ClampMin = "0.0"))
+	UPROPERTY(EditAnywhere, Category = Generation, config, meta = (ClampMin = "0.0"))
 	float AgentRadius;
 
-	UPROPERTY(EditAnywhere, Category=Generation, config, meta=(ClampMin = "0.0"))
+	UPROPERTY(EditAnywhere, Category = Generation, config, meta = (ClampMin = "0.0"))
 	float AgentHeight;
 
 	/** Size of the tallest agent that will path with this navmesh. */
@@ -495,7 +495,7 @@ class ENGINE_API ARecastNavMesh : public ANavigationData
 	UPROPERTY(EditAnywhere, Category=Generation, config, meta=(ClampMin = "0.0", ClampMax = "89.0", UIMin = "0.0", UIMax = "89.0" ))
 	float AgentMaxSlope;
 
-	UPROPERTY(EditAnywhere, Category=Generation, config, meta=(ClampMin = "0.0"))
+	UPROPERTY(EditAnywhere, Category = Generation, config, meta = (ClampMin = "0.0"))
 	float AgentMaxStepHeight;
 
 	/* The minimum dimension of area. Areas smaller than this will be discarded */
@@ -572,7 +572,7 @@ private:
 	/** indicates how often we will sort navigation tiles to mach players position */
 	UPROPERTY(config)
 	float TileSetUpdateInterval;
-
+	
 	/** contains last available dtPoly's flag bit set (8th bit at the moment of writing) */
 	static FNavPolyFlags NavLinkFlag;
 
@@ -679,6 +679,9 @@ protected:
 	/** Serialization helper. */
 	void SerializeRecastNavMesh(FArchive& Ar, FPImplRecastNavMesh*& NavMesh);
 
+	TArray<FIntPoint>& GetActiveTiles(); 
+	virtual void RestrictBuildingToActiveTiles(bool InRestrictBuildingToActiveTiles) override;
+
 public:
 	/** Whether NavMesh should adjust his tile pool size when NavBounds are changed */
 	bool IsResizable() const;
@@ -883,6 +886,10 @@ public:
 	virtual bool SupportsStreaming() const override;
 	virtual void ConditionalConstructGenerator() override;
 
+	void UpdateActiveTiles(const TArray<FNavigationInvokerRaw>& InvokerLocations);
+	void RemoveTiles(const TArray<FIntPoint>& Tiles);
+	void RebuildTile(const TArray<FIntPoint>& Tiles);
+
 protected:
 	// @todo docuement
 	UPrimitiveComponent* ConstructRenderingComponentImpl();
diff --git a/Engine/Source/Runtime/Engine/Private/AI/Navigation/NavigationPath.cpp b/Engine/Source/Runtime/Engine/Private/AI/Navigation/NavigationPath.cpp
--- a/Engine/Source/Runtime/Engine/Private/AI/Navigation/NavigationPath.cpp
+++ b/Engine/Source/Runtime/Engine/Private/AI/Navigation/NavigationPath.cpp
@@ -26,6 +26,7 @@ FNavigationPath::FNavigationPath()
 	, bIsReady(false)
 	, bIsPartial(false)
 	, bReachedSearchLimit(false)
+	, LastUpdateTimeStamp(-1.f)	// indicates that it has not been set
 	, GoalActorLocationTetherDistanceSq(-1.f)
 {
 
@@ -37,6 +38,7 @@ FNavigationPath::FNavigationPath(const TArray<FVector>& Points, AActor* InBase)
 	, bIsReady(true)
 	, bIsPartial(false)
 	, bReachedSearchLimit(false)
+	, LastUpdateTimeStamp(-1.f)	// indicates that it has not been set
 	, GoalActorLocationTetherDistanceSq(-1.f)
 {
 	Base = InBase;
diff --git a/Engine/Source/Runtime/Engine/Private/AI/Navigation/NavigationSystem.cpp b/Engine/Source/Runtime/Engine/Private/AI/Navigation/NavigationSystem.cpp
--- a/Engine/Source/Runtime/Engine/Private/AI/Navigation/NavigationSystem.cpp
+++ b/Engine/Source/Runtime/Engine/Private/AI/Navigation/NavigationSystem.cpp
@@ -6,6 +6,7 @@
 #include "NavigationOctree.h"
 #include "AI/Navigation/NavMeshBoundsVolume.h"
 #include "AI/Navigation/NavRelevantComponent.h"
+#include "AI/Navigation/NavigationInvokerComponent.h"
 
 #if WITH_RECAST
 #include "RecastNavMeshGenerator.h"
@@ -110,6 +111,25 @@ namespace NavigationDebugDrawing
 	const FVector PathNodeBoxExtent(16.f);
 }
 
+//----------------------------------------------------------------------//
+// FNavigationInvoker
+//----------------------------------------------------------------------//
+FNavigationInvoker::FNavigationInvoker()
+: Actor(nullptr)
+, GenerationRadius(0)
+, RemovalRadius(0)
+{
+
+}
+
+FNavigationInvoker::FNavigationInvoker(AActor& InActor, float InGenerationRadius, float InRemovalRadius)
+: Actor(&InActor)
+, GenerationRadius(InGenerationRadius)
+, RemovalRadius(InRemovalRadius)
+{
+
+}
+
 //----------------------------------------------------------------------//
 // FNavDataConfig
 //----------------------------------------------------------------------//
@@ -222,6 +242,11 @@ UNavigationSystem::UNavigationSystem(const FObjectInitializer& ObjectInitializer
 	NavUpdateLockFlags = 0;
 #endif
 
+	// active tiles
+	NextInvokersUpdateTime = 0.f;
+	ActiveTilesUpdateInterval = 1.f;
+	bGenerateNavigationOnlyAroundNavigationInvokers = false;
+
 	if (HasAnyFlags(RF_ClassDefaultObject) == false)
 	{
 		// reserve some arbitrary size
@@ -437,6 +462,7 @@ void UNavigationSystem::PostEditChangeProperty(FPropertyChangedEvent& PropertyCh
 {
 	static const FName NAME_SupportedAgents = GET_MEMBER_NAME_CHECKED(UNavigationSystem, SupportedAgents);
 	static const FName NAME_NavigationDataClass = GET_MEMBER_NAME_CHECKED(FNavDataConfig, NavigationDataClass);
+	static const FName NAME_EnableActiveTiles = GET_MEMBER_NAME_CHECKED(UNavigationSystem, bGenerateNavigationOnlyAroundNavigationInvokers);
 
 	Super::PostEditChangeProperty(PropertyChangedEvent);
 
@@ -453,6 +479,16 @@ void UNavigationSystem::PostEditChangeProperty(FPropertyChangedEvent& PropertyCh
 				SaveConfig();
 			}
 		}
+		else if (PropName == NAME_EnableActiveTiles)
+		{
+			for (auto NavData : NavDataSet)
+			{
+				if (NavData)
+				{
+					NavData->RestrictBuildingToActiveTiles(bGenerateNavigationOnlyAroundNavigationInvokers);
+				}
+			}
+		}
 	}
 }
 #endif // WITH_EDITOR
@@ -577,6 +613,19 @@ void UNavigationSystem::OnWorldInitDone(FNavigationSystem::EMode Mode)
 			}
 		}
 	}
+
+	if (Mode == FNavigationSystem::EditorMode && bGenerateNavigationOnlyAroundNavigationInvokers)
+	{
+		UWorld* MyWorld = GetWorld();
+		// gather enforcers manually to be able to see the results in editor as well
+		for (TObjectIterator<UNavigationInvokerComponent> It; It; ++It)
+		{
+			if (MyWorld == It->GetWorld())
+			{
+				It->RegisterWithNavigationSystem(*this);
+			}
+		}
+	}
 }
 
 void UNavigationSystem::RegisterNavigationDataInstances()
@@ -655,7 +704,11 @@ void UNavigationSystem::Tick(float DeltaSeconds)
 		}
 		INC_FLOAT_STAT_BY(STAT_Navigation_CumulativeBuildTime,(float)ThisTime*1000);
 	}
-		
+	
+	{
+		UpdateInvokers();
+	}
+
 	{
 		SCOPE_CYCLE_COUNTER(STAT_Navigation_TickMarkDirty);
 
@@ -1003,7 +1056,7 @@ bool UNavigationSystem::ProjectPointToNavigation(const FVector& Point, FNavLocat
 	}
 
 	return NavData != NULL && NavData->ProjectPoint(Point, OutLocation
-		, FNavigationSystem::IsValidExtent(Extent) ? Extent : NavData->NavDataConfig.DefaultQueryExtent
+		, FNavigationSystem::IsValidExtent(Extent) ? Extent : NavData->GetConfig().DefaultQueryExtent
 		, QueryFilter);
 }
 
@@ -1365,7 +1418,7 @@ bool UNavigationSystem::IsNavigationBuilt(const AWorldSettings* Settings) const
 		if (NavData != NULL && NavData->GetWorldSettings() == Settings)
 		{
 			FNavDataGenerator* Generator = NavData->GetGenerator();
-			if ((NavData->RuntimeGeneration != ERuntimeGenerationType::Static 
+			if ((NavData->GetRuntimeGenerationMode() != ERuntimeGenerationType::Static
 #if WITH_EDITOR
 				|| GEditor != NULL
 #endif // WITH_EDITOR
@@ -3265,9 +3318,9 @@ ERuntimeGenerationType UNavigationSystem::GetRuntimeGenerationType() const
 
 	for (ANavigationData* NavData : NavDataSet)
 	{
-		if (NavData && NavData->RuntimeGeneration > RuntimeGenerationType)
+		if (NavData && NavData->GetRuntimeGenerationMode() > RuntimeGenerationType)
 		{
-			RuntimeGenerationType = NavData->RuntimeGeneration;
+			RuntimeGenerationType = NavData->GetRuntimeGenerationMode();
 		}
 	}
 	
@@ -3497,6 +3550,145 @@ bool UNavigationSystem::DoesPathIntersectBox(const FNavigationPath* Path, const
 	return Path != NULL && Path->DoesIntersectBox(Box, AgentLocation, StartingIndex);
 }
 
+//----------------------------------------------------------------------//
+// Active tiles
+//----------------------------------------------------------------------//
+
+void UNavigationSystem::RegisterNavigationInvoker(AActor& Invoker, float TileGenerationRadius, float TileRemovalRadius)
+{
+	UWorld* World = Invoker.GetWorld();
+
+	if (World && Cast<UNavigationSystem>(World->GetNavigationSystem()))
+	{
+		((UNavigationSystem*)(World->GetNavigationSystem()))->RegisterInvoker(Invoker, TileGenerationRadius, TileRemovalRadius);
+	}
+}
+
+void UNavigationSystem::UnregisterNavigationInvoker(AActor& Invoker)
+{
+	UWorld* World = Invoker.GetWorld();
+
+	if (World && Cast<UNavigationSystem>(World->GetNavigationSystem()))
+	{
+		((UNavigationSystem*)(World->GetNavigationSystem()))->UnregisterInvoker(Invoker);
+	}
+}
+
+void UNavigationSystem::RegisterInvoker(AActor& Invoker, float TileGenerationRadius, float TileRemovalRadius)
+{
+	UE_CVLOG(bGenerateNavigationOnlyAroundNavigationInvokers == false, this, LogNavigation, Warning
+		, TEXT("Trying to register %s as enforcer, but NavigationSystem is not set up for enforcer-centric generation. See GenerateNavigationOnlyAroundNavigationInvokers in NavigationSystem's properties")
+		, *Invoker.GetName());
+
+	TileGenerationRadius = FMath::Clamp(TileGenerationRadius, 0.f, BIG_NUMBER);
+	TileRemovalRadius = FMath::Clamp(TileRemovalRadius, TileGenerationRadius, BIG_NUMBER);
+
+	FNavigationInvoker& Data = Invokers.FindOrAdd(&Invoker);
+	Data.Actor = &Invoker;
+	Data.GenerationRadius = TileGenerationRadius;
+	Data.RemovalRadius = TileRemovalRadius;
+
+	UE_VLOG_CYLINDER(this, LogNavigation, Log, Invoker.GetActorLocation(), Invoker.GetActorLocation() + FVector(0, 0, 20), TileGenerationRadius, FColorList::LimeGreen
+		, TEXT("%s %.0f %.0f"), *Invoker.GetName(), TileGenerationRadius, TileRemovalRadius);
+	UE_VLOG_CYLINDER(this, LogNavigation, Log, Invoker.GetActorLocation(), Invoker.GetActorLocation() + FVector(0, 0, 20), TileRemovalRadius, FColorList::IndianRed, TEXT(""));
+}
+
+void UNavigationSystem::UnregisterInvoker(AActor& Invoker)
+{
+	UE_VLOG(this, LogNavigation, Log, TEXT("Removing %s from enforcers list"), *Invoker.GetName());
+	Invokers.Remove(&Invoker);
+}
+
+void UNavigationSystem::UpdateInvokers()
+{
+	UWorld* World = GetWorld();
+	const float CurrentTime = World->GetTimeSeconds();
+	if (CurrentTime >= NextInvokersUpdateTime)
+	{
+		if (Invokers.Num() > 0)
+		{
+			QUICK_SCOPE_CYCLE_COUNTER(STAT_NavSys_Clusterize);
+
+			const double StartTime = FPlatformTime::Seconds();
+
+			TArray<FNavigationInvokerRaw> InvokerLocations;
+			InvokerLocations.Reserve(Invokers.Num());
+
+			for (auto Item = Invokers.CreateConstIterator(); Item; ++Item)
+			{
+				AActor* Actor = Item->Value.Actor.Get();
+				if (Actor != nullptr
+#if WITH_EDITOR
+					//&& (GIsEditor == false || Item.Actor->HasAnyFlags(RF_Transactional | RF_PendingKill) == false)
+#endif //WITH_EDITOR
+					)
+				{
+					InvokerLocations.Add(FNavigationInvokerRaw(Actor->GetActorLocation(), Item->Value.GenerationRadius, Item->Value.RemovalRadius));
+				}
+			}
+
+#if ENABLE_VISUAL_LOG
+			//TArray<FNavigationInvokerRaw> LocationsCopy = InvokerLocations;
+			const double CachingFinishTime = FPlatformTime::Seconds();
+#endif //ENABLE_VISUAL_LOG
+			//InvokerLocations = KMedoidsClusterization::Clusterize(InvokerLocations, MaxClustersCount);
+
+#if ENABLE_VISUAL_LOG
+			//const double ClusterizationFinishTime = FPlatformTime::Seconds();
+
+			UE_VLOG(this, LogNavigation, Log, TEXT("Caching time %fms"), float((CachingFinishTime - StartTime) / 1000));
+			//UE_VLOG(this, LogNavigation, Log, TEXT("Clusterization time %fms for %d clusters"), float(ClusterizationFinishTime - CachingFinishTime) / 1000, MaxClustersCount);
+
+			FVisualLogger& Vlog = FVisualLogger::Get();
+			if (Vlog.IsRecording())
+			{
+				/*FVisualLogEntry* Entry = Vlog.GetEntryToWrite(this, GetWorld()->TimeSeconds);
+				check(Entry);
+
+				FVisualLogShapeElement LogElement(TEXT(""), FColorList::LightBlue, 30, LogNavigation.GetCategoryName());
+				LogElement.Points = LocationsCopy;
+				LogElement.Type = EVisualLoggerShapeElement::SinglePoint;
+				LogElement.Verbosity = ELogVerbosity::Log;
+
+				Entry->AddElement(LogElement);*/
+
+				for (const auto& InvokerData : InvokerLocations)
+				{
+					UE_VLOG_CYLINDER(this, LogNavigation, Log, InvokerData.Location, InvokerData.Location + FVector(0, 0, 20), InvokerData.RadiusMax, FColorList::Blue, TEXT(""));
+					UE_VLOG_CYLINDER(this, LogNavigation, Log, InvokerData.Location, InvokerData.Location + FVector(0, 0, 20), InvokerData.RadiusMin, FColorList::CadetBlue, TEXT(""));
+				}
+			}
+#endif // ENABLE_VISUAL_LOG
+
+			const double UpdateStartTime = FPlatformTime::Seconds();
+			for (TActorIterator<ARecastNavMesh> It(GetWorld()); It; ++It)
+			{
+				It->UpdateActiveTiles(InvokerLocations);
+			}
+			const double UpdateEndTime = FPlatformTime::Seconds();
+			UE_VLOG(this, LogNavigation, Log, TEXT("Marking tiles to update %fms"), float((UpdateStartTime - UpdateEndTime) / 1000));
+		}
+		// once per second
+		NextInvokersUpdateTime = CurrentTime + ActiveTilesUpdateInterval;
+	}
+}
+
+void UNavigationSystem::RegisterNavigationInvoker(AActor* Invoker, float TileGenerationRadius, float TileRemovalRadius)
+{
+	if (Invoker != nullptr)
+	{
+		RegisterInvoker(*Invoker, TileGenerationRadius, TileRemovalRadius);
+	}
+}
+
+void UNavigationSystem::UnregisterNavigationInvoker(AActor* Invoker)
+{
+	if (Invoker != nullptr)
+	{
+		UnregisterInvoker(*Invoker);
+	}
+}
+
 //----------------------------------------------------------------------//
 // DEPRECATED
 //----------------------------------------------------------------------//
diff --git a/Engine/Source/Runtime/Engine/Private/AI/Navigation/NavigationTestingActor.cpp b/Engine/Source/Runtime/Engine/Private/AI/Navigation/NavigationTestingActor.cpp
--- a/Engine/Source/Runtime/Engine/Private/AI/Navigation/NavigationTestingActor.cpp
+++ b/Engine/Source/Runtime/Engine/Private/AI/Navigation/NavigationTestingActor.cpp
@@ -6,6 +6,7 @@
 #endif
 #include "AI/Navigation/NavigationTestingActor.h"
 #include "AI/Navigation/NavTestRenderingComponent.h"
+#include "AI/Navigation/NavigationInvokerComponent.h"
 #include "Components/CapsuleComponent.h"
 
 void FNavTestTickHelper::Tick(float DeltaTime)
@@ -59,6 +60,9 @@ ANavigationTestingActor::ANavigationTestingActor(const FObjectInitializer& Objec
 
 	RootComponent = CapsuleComponent;
 
+	InvokerComponent = CreateDefaultSubobject<UNavigationInvokerComponent>(TEXT("InvokerComponent"));
+	InvokerComponent->bAutoActivate = bActAsNavigationInvoker;
+
 	PathObserver = FNavigationPath::FPathObserverDelegate::FDelegate::CreateUObject(this, &ANavigationTestingActor::OnPathEvent);
 }
 
@@ -104,6 +108,7 @@ void ANavigationTestingActor::PostEditChangeProperty(FPropertyChangedEvent& Prop
 	static const FName NAME_ShouldBeVisibleInGame = GET_MEMBER_NAME_CHECKED(ANavigationTestingActor, bShouldBeVisibleInGame);
 	static const FName NAME_OtherActor = GET_MEMBER_NAME_CHECKED(ANavigationTestingActor, OtherActor);
 	static const FName NAME_IsSearchStart = GET_MEMBER_NAME_CHECKED(ANavigationTestingActor, bSearchStart);
+	static const FName NAME_InvokerComponent = GET_MEMBER_NAME_CHECKED(ANavigationTestingActor, InvokerComponent);
 
 	Super::PostEditChangeProperty(PropertyChangedEvent);
 
@@ -180,6 +185,10 @@ void ANavigationTestingActor::PostEditChangeProperty(FPropertyChangedEvent& Prop
 
 			UpdatePathfinding();
 		}
+		else if (NAME_InvokerComponent == ChangedPropName)
+		{
+			InvokerComponent->SetActive(bActAsNavigationInvoker);
+		}
 	}
 }
 
@@ -202,6 +211,8 @@ void ANavigationTestingActor::PostLoad()
 {
 	Super::PostLoad();
 
+	InvokerComponent->bAutoActivate = bActAsNavigationInvoker;
+
 #if WITH_RECAST && WITH_EDITORONLY_DATA
 	if (GIsEditor)
 	{
diff --git a/Engine/Source/Runtime/Engine/Private/AI/Navigation/RecastNavMeshGenerator.cpp b/Engine/Source/Runtime/Engine/Private/AI/Navigation/RecastNavMeshGenerator.cpp
--- a/Engine/Source/Runtime/Engine/Private/AI/Navigation/RecastNavMeshGenerator.cpp
+++ b/Engine/Source/Runtime/Engine/Private/AI/Navigation/RecastNavMeshGenerator.cpp
@@ -1536,7 +1536,7 @@ void FRecastTileGenerator::GatherGeometry(const FRecastNavMeshGenerator& ParentG
 {
 	const UNavigationSystem*	NavSys = UNavigationSystem::GetCurrent(ParentGenerator.GetWorld());
 	const FNavigationOctree*	NavOctree = NavSys ? NavSys->GetNavOctree() : nullptr;
-	const FNavDataConfig*		NavDataConfig = &ParentGenerator.GetOwner()->NavDataConfig;
+	const FNavDataConfig&		NavDataConfig = ParentGenerator.GetOwner()->GetConfig();
 
 	for (FNavigationOctree::TConstElementBoxIterator<FNavigationOctree::DefaultStackAllocator> It(*NavOctree, ParentGenerator.GrowBoundingBox(TileBB, /*bIncludeAgentHeight*/ false));
 		It.HasPendingElements();
@@ -1580,7 +1580,7 @@ void FRecastTileGenerator::GatherGeometry(const FRecastNavMeshGenerator& ParentG
 				}
 			}
 
-			const FCompositeNavModifier ModifierInstance = Element.GetModifierForAgent(NavDataConfig);
+			const FCompositeNavModifier ModifierInstance = Element.GetModifierForAgent(&NavDataConfig);
 			AppendModifier(ModifierInstance, Element.Data.NavDataPerInstanceTransformDelegate);
 		}
 	}
@@ -2680,7 +2680,7 @@ static int32 CaclulateMaxTilesCount(const TNavStatArray<FBox>& NavigableAreas, f
 // Whether navmesh is static, does not support rebuild from geometry
 static bool IsGameStaticNavMesh(ARecastNavMesh* InNavMesh)
 {
-	return (InNavMesh->GetWorld()->IsGameWorld() && InNavMesh->RuntimeGeneration != ERuntimeGenerationType::Dynamic);
+	return (InNavMesh->GetWorld()->IsGameWorld() && InNavMesh->GetRuntimeGenerationMode() != ERuntimeGenerationType::Dynamic);
 }
 
 //----------------------------------------------------------------------//
@@ -2693,6 +2693,7 @@ FRecastNavMeshGenerator::FRecastNavMeshGenerator(ARecastNavMesh& InDestNavMesh)
 	, AvgLayersPerTile(8.0f)
 	, DestNavMesh(&InDestNavMesh)
 	, bInitialized(false)
+	, bRestrictBuildingToActiveTiles(false)
 	, Version(0)
 {
 	INC_DWORD_STAT_BY(STAT_NavigationMemory, sizeof(*this));
@@ -3093,6 +3094,127 @@ int32 FRecastNavMeshGenerator::FindInclusionBoundEncapsulatingBox(const FBox& Bo
 	return INDEX_NONE;
 }
 
+void FRecastNavMeshGenerator::RestrictBuildingToActiveTiles(bool InRestrictBuildingToActiveTiles) 
+{ 
+	if (bRestrictBuildingToActiveTiles != InRestrictBuildingToActiveTiles)
+	{
+		bRestrictBuildingToActiveTiles = InRestrictBuildingToActiveTiles;
+		if (InRestrictBuildingToActiveTiles)
+		{
+			// gather non-empty tiles and add them to ActiveTiles
+
+			const dtNavMesh* DetourMesh = DestNavMesh->GetRecastNavMeshImpl()->GetRecastMesh();
+
+			if (DetourMesh != nullptr && DetourMesh->isEmpty() == false)
+			{
+				ActiveTiles.Reset();
+				int32 TileCount = DetourMesh->getMaxTiles();
+				for (int32 TileIndex = 0; TileIndex < TileCount; ++TileIndex)
+				{
+					const dtMeshTile* Tile = DetourMesh->getTile(TileIndex);
+					if (Tile != nullptr && Tile->header != nullptr && Tile->header->polyCount > 0)
+					{
+						ActiveTiles.AddUnique(FIntPoint(Tile->header->x, Tile->header->y));
+					}
+				}
+			}
+		}
+	}
+}
+
+bool FRecastNavMeshGenerator::IsInActiveSet(const FIntPoint& Tile) const
+{
+	// @TODO checking if given tile is in active tiles needs to be faster
+	return bRestrictBuildingToActiveTiles == false || ActiveTiles.Find(Tile) != INDEX_NONE;
+}
+
+void FRecastNavMeshGenerator::RemoveTiles(const TArray<FIntPoint>& Tiles)
+{
+	for (const FIntPoint& TileXY : Tiles)
+	{
+		RemoveTileLayers(TileXY.X, TileXY.Y);
+		if (PendingDirtyTiles.Num() > 0)
+		{
+			FPendingTileElement DirtyTile;
+			DirtyTile.Coord = TileXY;
+			PendingDirtyTiles.Remove(DirtyTile);
+		}
+	}
+}
+
+void FRecastNavMeshGenerator::ReAddTiles(const TArray<FIntPoint>& Tiles)
+{
+	static const FVector Expansion(1, 1, BIG_NUMBER);
+	// a little trick here - adding a dirty area so that navmesh building figures it out on its own
+	dtNavMesh* DetourMesh = DestNavMesh->GetRecastNavMeshImpl()->GetRecastMesh();
+	const dtNavMeshParams* SavedNavParams = DestNavMesh->GetRecastNavMeshImpl()->DetourNavMesh->getParams();
+	const float TileDim = Config.tileSize * Config.cs;
+
+	TSet<FPendingTileElement> DirtyTiles;
+
+	// @note we act on assumption all items in Tiles are unique
+	for (const FIntPoint& TileCoords : Tiles)
+	{
+		FPendingTileElement Element;
+		Element.Coord = TileCoords;
+		Element.bRebuildGeometry = true;
+		DirtyTiles.Add(Element);
+	}
+
+	int32 NumTilesMarked = DirtyTiles.Num();
+
+	// Merge all pending tiles into one container
+	for (const FPendingTileElement& Element : PendingDirtyTiles)
+	{
+		FPendingTileElement* ExistingElement = DirtyTiles.Find(Element);
+		if (ExistingElement)
+		{
+			ExistingElement->bRebuildGeometry |= Element.bRebuildGeometry;
+			// Append area bounds to existing list 
+			if (ExistingElement->bRebuildGeometry == false)
+			{
+				ExistingElement->DirtyAreas.Append(Element.DirtyAreas);
+			}
+			else
+			{
+				ExistingElement->DirtyAreas.Empty();
+			}
+		}
+		else
+		{
+			DirtyTiles.Add(Element);
+		}
+	}
+
+	// Dump results into array
+	PendingDirtyTiles.Empty(DirtyTiles.Num());
+	for (const FPendingTileElement& Element : DirtyTiles)
+	{
+		PendingDirtyTiles.Add(Element);
+	}
+
+	// Sort tiles by proximity to players 
+	if (NumTilesMarked > 0)
+	{
+		SortPendingBuildTiles();
+	}
+
+	/*TArray<FNavigationDirtyArea> DirtyAreasContainer;
+	DirtyAreasContainer.Reserve(Tiles.Num());
+
+	TSet<FPendingTileElement> DirtyTiles;
+
+	for (const FIntPoint& TileCoords : Tiles)
+	{
+		const FVector TileCenter = Recast2UnrealPoint(SavedNavParams->orig) + FVector(TileDim * float(TileCoords.X), TileDim * float(TileCoords.Y), 0);
+		
+		FNavigationDirtyArea DirtyArea(FBox(TileCenter - Expansion, TileCenter - 1), ENavigationDirtyFlag::All);
+		DirtyAreasContainer.Add(DirtyArea);
+	}
+
+	MarkDirtyTiles(DirtyAreasContainer);*/
+}
+
 TArray<uint32> FRecastNavMeshGenerator::RemoveTileLayers(const int32 TileX, const int32 TileY)
 {
 	TArray<uint32> ResultTileIndices;
@@ -3144,9 +3266,10 @@ TArray<uint32> FRecastNavMeshGenerator::AddGeneratedTiles(const FRecastTileGener
 		// remove all layers
 		ResultTileIndices = RemoveTileLayers(TileX, TileY);
 	}
+
 	
 	dtNavMesh* DetourMesh = DestNavMesh->GetRecastNavMeshImpl()->GetRecastMesh();
-	if (DetourMesh != nullptr)
+	if (DetourMesh != nullptr && IsInActiveSet(FIntPoint(TileX, TileY)))
 	{
 		TArray<FNavMeshTileData> TileLayers = TileGenerator.GetNavigationData();
 		ResultTileIndices.Reserve(TileLayers.Num());
@@ -3276,8 +3399,7 @@ void FRecastNavMeshGenerator::MarkDirtyTiles(const TArray<FNavigationDirtyArea>&
 	const float TileSizeInWorldUnits = Config.tileSize * Config.cs;
 	check(TileSizeInWorldUnits > 0);
 	const FVector NavMeshOrigin = FVector::ZeroVector;
-	int32 NumTilesMarked = 0;
-
+	
 	const bool bGameStaticNavMesh = IsGameStaticNavMesh(DestNavMesh);
 		
 	// find all tiles that need regeneration
@@ -3324,13 +3446,18 @@ void FRecastNavMeshGenerator::MarkDirtyTiles(const TArray<FNavigationDirtyArea>&
 		const int32 YMin = FMath::FloorToInt((RcAreaBounds.Min.Z - NavMeshOrigin.Z) / TileSizeInWorldUnits);
 		const int32 YMax = FMath::FloorToInt((RcAreaBounds.Max.Z - NavMeshOrigin.Z) / TileSizeInWorldUnits);
 
-		for (int32 y = YMin; y <= YMax; ++y)
+		for (int32 TileY = YMin; TileY <= YMax; ++TileY)
 		{
-			for (int32 x = XMin; x <= XMax; ++x)
+			for (int32 TileX = XMin; TileX <= XMax; ++TileX)
 			{
+				if (IsInActiveSet(FIntPoint(TileX, TileY)) == false)
+				{
+					continue;
+				}
+
 				if (DirtyArea.HasFlag(ENavigationDirtyFlag::NavigationBounds) == false && bDoTileInclusionTest == true)
 				{
-					const FBox TileBox = CalculateTileBounds(x, y, NavMeshOrigin, TotalNavBounds, TileSizeInWorldUnits);
+					const FBox TileBox = CalculateTileBounds(TileX, TileY, NavMeshOrigin, TotalNavBounds, TileSizeInWorldUnits);
 
 					// do per tile check since we can have lots of tiles inbetween navigable bounds volumes
 					if (IntercestBounds(TileBox, InclusionBounds) == false)
@@ -3341,7 +3468,7 @@ void FRecastNavMeshGenerator::MarkDirtyTiles(const TArray<FNavigationDirtyArea>&
 				}
 												
 				FPendingTileElement Element;
-				Element.Coord = FIntPoint(x, y);
+				Element.Coord = FIntPoint(TileX, TileY);
 				Element.bRebuildGeometry = DirtyArea.HasFlag(ENavigationDirtyFlag::Geometry) || DirtyArea.HasFlag(ENavigationDirtyFlag::NavigationBounds);
 				if (Element.bRebuildGeometry == false)
 				{
@@ -3370,7 +3497,7 @@ void FRecastNavMeshGenerator::MarkDirtyTiles(const TArray<FNavigationDirtyArea>&
 		}
 	}
 	
-	NumTilesMarked = DirtyTiles.Num();
+	int32 NumTilesMarked = DirtyTiles.Num();
 
 	// Merge all pending tiles into one container
 	for (const FPendingTileElement& Element : PendingDirtyTiles)
@@ -3705,7 +3832,7 @@ void FRecastNavMeshGenerator::ExportNavigationData(const FString& FileName) cons
 				It.Advance())
 			{
 				const FNavigationOctreeElement& Element = It.GetCurrentElement();
-				const bool bExportGeometry = Element.Data.HasGeometry() && Element.ShouldUseGeometry(&DestNavMesh->NavDataConfig);
+				const bool bExportGeometry = Element.Data.HasGeometry() && Element.ShouldUseGeometry(DestNavMesh->GetConfig());
 
 				if (bExportGeometry && Element.Data.CollisionData.Num())
 				{
diff --git a/Engine/Source/Runtime/Engine/Public/AI/Navigation/RecastNavMeshGenerator.h b/Engine/Source/Runtime/Engine/Public/AI/Navigation/RecastNavMeshGenerator.h
--- a/Engine/Source/Runtime/Engine/Public/AI/Navigation/RecastNavMeshGenerator.h
+++ b/Engine/Source/Runtime/Engine/Public/AI/Navigation/RecastNavMeshGenerator.h
@@ -3,7 +3,7 @@
 
 #if WITH_RECAST
 
-#include "Recast.h"
+#include "Recast/Recast.h"
 #include "DetourNavMesh.h"
 #include "AI/NavDataGenerator.h"
 #include "AI/NavigationModifier.h"
@@ -269,6 +269,11 @@ struct FPendingTileElement
 	{
 	}
 
+	bool operator == (const FIntPoint& Location) const
+	{
+		return Coord == Location;
+	}
+
 	bool operator == (const FPendingTileElement& Other) const
 	{
 		return Coord == Other.Coord;
@@ -420,8 +425,16 @@ protected:
 	/** Adds generated tiles to NavMesh, replacing old ones */
 	TArray<uint32> AddGeneratedTiles(const FRecastTileGenerator& TileGenerator);
 
+public:
 	/** Removes all tiles at specified grid location */
 	TArray<uint32> RemoveTileLayers(const int32 TileX, const int32 TileY);
+
+	void RemoveTiles(const TArray<FIntPoint>& Tiles);
+	void ReAddTiles(const TArray<FIntPoint>& Tiles);
+
+protected:
+	bool IsInActiveSet(const FIntPoint& Tile) const;
+	void RestrictBuildingToActiveTiles(bool InRestrictBuildingToActiveTiles);
 	
 	/** Blocks until build for specified list of tiles is complete and discard results */
 	void DiscardCurrentBuildingTasks();
@@ -434,9 +447,11 @@ protected:
 	virtual uint32 LogMemUsed() const override;
 
 private:
+	friend ARecastNavMesh;
+
 	/** Parameters defining navmesh tiles */
 	struct dtNavMeshParams TiledMeshParams;
-	struct FRecastBuildConfig Config;
+	FRecastBuildConfig Config;
 	
 	int32 NumActiveTiles;
 	int32 MaxTileGeneratorTasks;
@@ -461,6 +476,8 @@ private:
 	/** List of tiles that were recently regenerated */
 	TNavStatArray<FTileTimestamp> RecentlyBuiltTiles;
 #endif// WITH_EDITOR
+	
+	TArray<FIntPoint> ActiveTiles;
 
 	/** */
 	FRecastNavMeshCachedData AdditionalCachedData;
@@ -470,6 +487,8 @@ private:
 
 	uint32 bInitialized:1;
 
+	uint32 bRestrictBuildingToActiveTiles:1;
+
 	/** Runtime generator's version, increased every time all tile generators get invalidated
 	 *	like when navmesh size changes */
 	uint32 Version;
