Files missing in previous submission #UE4

diff --git a/Engine/Source/Runtime/Engine/Classes/AI/Navigation/NavigationInvokerComponent.h b/Engine/Source/Runtime/Engine/Classes/AI/Navigation/NavigationInvokerComponent.h
new file mode 100644
--- /dev/null
+++ b/Engine/Source/Runtime/Engine/Classes/AI/Navigation/NavigationInvokerComponent.h
@@ -0,0 +1,28 @@
+// Copyright 1998-2015 Epic Games, Inc. All Rights Reserved.
+#pragma once
+
+#include "AI/Navigation/NavigationSystem.h"
+#include "NavigationInvokerComponent.generated.h"
+
+UCLASS(ClassGroup = (Navigation), meta = (BlueprintSpawnableComponent), hidecategories = (Activation))
+class UNavigationInvokerComponent : public UActorComponent
+{
+	GENERATED_BODY()
+
+protected:
+
+	UPROPERTY(EditAnywhere, Category = Navigation, meta = (ClampMin = "0.1", UIMin = "0.1"))
+	float TileGenerationRadius;
+
+	UPROPERTY(EditAnywhere, Category = Navigation, meta = (ClampMin = "0.1", UIMin = "0.1"))
+	float TileRemovalRadius;
+
+public:
+	UNavigationInvokerComponent(const FObjectInitializer& ObjectInitializer);
+
+	void RegisterWithNavigationSystem(UNavigationSystem& NavSys);
+
+protected:
+	virtual void Activate(bool bReset = false) override;
+	virtual void Deactivate() override;
+};
\ No newline at end of file
diff --git a/Engine/Source/Runtime/Engine/Private/AI/Navigation/NavigationInvokerComponent.cpp b/Engine/Source/Runtime/Engine/Private/AI/Navigation/NavigationInvokerComponent.cpp
new file mode 100644
--- /dev/null
+++ b/Engine/Source/Runtime/Engine/Private/AI/Navigation/NavigationInvokerComponent.cpp
@@ -0,0 +1,46 @@
+// Copyright 1998-2015 Epic Games, Inc. All Rights Reserved.
+
+#include "EnginePrivate.h"
+#include "AI/Navigation/NavigationInvokerComponent.h"
+
+UNavigationInvokerComponent::UNavigationInvokerComponent(const FObjectInitializer& ObjectInitializer)
+	: Super(ObjectInitializer)
+	, TileGenerationRadius(3000)
+	, TileRemovalRadius(5000)
+{
+	bAutoActivate = true;
+}
+
+void UNavigationInvokerComponent::Activate(bool bReset)
+{
+	Super::Activate(bReset);
+
+	AActor* Owner = GetOwner();
+	if (Owner)
+	{
+		UNavigationSystem::RegisterNavigationInvoker(*Owner, TileGenerationRadius, TileRemovalRadius);
+	}
+}
+
+void UNavigationInvokerComponent::Deactivate()
+{
+	Super::Deactivate();
+
+	AActor* Owner = GetOwner();
+	if (Owner)
+	{
+		UNavigationSystem::UnregisterNavigationInvoker(*Owner);
+	}
+}
+
+void UNavigationInvokerComponent::RegisterWithNavigationSystem(UNavigationSystem& NavSys)
+{
+	if (IsActive())
+	{
+		AActor* Owner = GetOwner();
+		if (Owner)
+		{
+			NavSys.RegisterInvoker(*Owner, TileGenerationRadius, TileRemovalRadius);
+		}
+	}
+}
diff --git a/Engine/Source/Runtime/Engine/Private/AI/Navigation/RecastNavMesh.cpp b/Engine/Source/Runtime/Engine/Private/AI/Navigation/RecastNavMesh.cpp
--- a/Engine/Source/Runtime/Engine/Private/AI/Navigation/RecastNavMesh.cpp
+++ b/Engine/Source/Runtime/Engine/Private/AI/Navigation/RecastNavMesh.cpp
@@ -660,6 +660,22 @@ void ARecastNavMesh::SortAreasForGenerator(TArray<FRecastAreaNavModifierElement>
 	Modifiers.Sort(FNavAreaSortPredicate());
 }
 
+TArray<FIntPoint>& ARecastNavMesh::GetActiveTiles()
+{
+	FRecastNavMeshGenerator* MyGenerator = static_cast<FRecastNavMeshGenerator*>(GetGenerator());
+	check(MyGenerator);
+	return MyGenerator->ActiveTiles;
+}
+
+void ARecastNavMesh::RestrictBuildingToActiveTiles(bool InRestrictBuildingToActiveTiles)
+{
+	FRecastNavMeshGenerator* MyGenerator = static_cast<FRecastNavMeshGenerator*>(GetGenerator());
+	if (MyGenerator)
+	{
+		MyGenerator->RestrictBuildingToActiveTiles(InRestrictBuildingToActiveTiles);
+	}
+}
+
 void ARecastNavMesh::SerializeRecastNavMesh(FArchive& Ar, FPImplRecastNavMesh*& NavMesh)
 {
 	if (!Ar.IsLoading()	&& NavMesh == NULL)
@@ -1910,8 +1926,131 @@ bool ARecastNavMesh::HasValidNavmesh() const
 #endif // WITH_RECAST
 }
 
-#if WITH_RECAST
+//----------------------------------------------------------------------//
+// RecastNavMesh: Active Tiles 
+//----------------------------------------------------------------------//
+void ARecastNavMesh::UpdateActiveTiles(const TArray<FNavigationInvokerRaw>& InvokerLocations)
+{
+	if (HasValidNavmesh() == false)
+	{
+		return;
+	}
+
+	FRecastNavMeshGenerator* MyGenerator = static_cast<FRecastNavMeshGenerator*>(GetGenerator());
+	const dtNavMeshParams* NavParams = GetRecastNavMeshImpl()->DetourNavMesh->getParams();
+	check(NavParams && MyGenerator);
+	const FRecastBuildConfig& Config = MyGenerator->GetConfig();
+	const FVector NavmeshOrigin = Recast2UnrealPoint(NavParams->orig);
+	const float TileDim = Config.tileSize * Config.cs;
+	const FVector TileCenterOffset(TileDim, TileDim, 0);
+
+	TArray<FIntPoint>& ActiveTiles = GetActiveTiles();
+	TArray<FIntPoint> OldActiveSet = ActiveTiles;
+	TArray<FIntPoint> TilesInMinDistance;
+	TArray<FIntPoint> TilesInMaxDistance;
+	TilesInMinDistance.Reserve(ActiveTiles.Num());
+	TilesInMaxDistance.Reserve(ActiveTiles.Num());
+	ActiveTiles.Reset();
+
+	//const int32 TileRadius = FMath::CeilToInt(Radius / TileDim);
+	static const float SqareRootOf2 = FMath::Sqrt(2.f);
+
+	for (const auto& Invoker : InvokerLocations)
+	{
+		const FVector InvokerRelativeLocation = (NavmeshOrigin - Invoker.Location);
+		const float TileCenterDistanceToRemoveSq = FMath::Square(TileDim * SqareRootOf2 / 2 + Invoker.RadiusMax);
+		const float TileCenterDistanceToAddSq = FMath::Square(TileDim * SqareRootOf2 / 2 + Invoker.RadiusMin);
+
+		const int32 MinTileX = FMath::FloorToInt((InvokerRelativeLocation.X - Invoker.RadiusMax) / TileDim);
+		const int32 MaxTileX = FMath::CeilToInt((InvokerRelativeLocation.X + Invoker.RadiusMax) / TileDim);
+		const int32 MinTileY = FMath::FloorToInt((InvokerRelativeLocation.Y - Invoker.RadiusMax) / TileDim);
+		const int32 MaxTileY = FMath::CeilToInt((InvokerRelativeLocation.Y + Invoker.RadiusMax) / TileDim);
+
+		for (int32 X = MinTileX; X <= MaxTileX; ++X)
+		{
+			for (int32 Y = MinTileY; Y <= MaxTileY; ++Y)
+			{
+				const float DistanceSq = (InvokerRelativeLocation - FVector(X * TileDim + TileDim / 2, Y * TileDim + TileDim / 2, 0.f)).SizeSquared2D();
+				if (DistanceSq < TileCenterDistanceToRemoveSq)
+				{
+					TilesInMaxDistance.AddUnique(FIntPoint(X, Y));
+
+					if (DistanceSq < TileCenterDistanceToAddSq)
+					{
+						TilesInMinDistance.AddUnique(FIntPoint(X, Y));
+					}
+				}
+			}
+		}
+	}
+
+	ActiveTiles.Append(TilesInMinDistance);
+
+	TArray<FIntPoint> TilesToRemove;
+	TilesToRemove.Reserve(OldActiveSet.Num());
+	for (int32 Index = OldActiveSet.Num() - 1; Index >= 0; --Index)
+	{
+		if (TilesInMaxDistance.Find(OldActiveSet[Index]) == INDEX_NONE)
+		{
+			TilesToRemove.Add(OldActiveSet[Index]);
+			OldActiveSet.RemoveAtSwap(Index, 1, /*bAllowShrinking=*/false);
+		}
+		else
+		{
+			ActiveTiles.AddUnique(OldActiveSet[Index]);
+		}
+	}
+
+	TArray<FIntPoint> TilesToUpdate;
+	TilesToUpdate.Reserve(ActiveTiles.Num());
+	for (int32 Index = TilesInMinDistance.Num() - 1; Index >= 0; --Index)
+	{
+		// check if it's a new tile
+		if (OldActiveSet.Find(TilesInMinDistance[Index]) == INDEX_NONE)
+		{
+			TilesToUpdate.Add(TilesInMinDistance[Index]);
+		}
+	}
+
+	RemoveTiles(TilesToRemove);
+	RebuildTile(TilesToUpdate);
+
+	if (ActiveTiles.Num() == 0)
+	{
+		UpdateNavMeshDrawing();
+	}
+}
+
+void ARecastNavMesh::RemoveTiles(const TArray<FIntPoint>& Tiles)
+{
+	if (Tiles.Num() > 0)
+	{
+		bWantsUpdate = true;
+		FRecastNavMeshGenerator* MyGenerator = static_cast<FRecastNavMeshGenerator*>(GetGenerator());
+		if (MyGenerator)
+		{
+			MyGenerator->RemoveTiles(Tiles);
+		}
+	}
+}
 
+void ARecastNavMesh::RebuildTile(const TArray<FIntPoint>& Tiles)
+{
+	if (Tiles.Num() > 0)
+	{
+		bWantsUpdate = true;
+		FRecastNavMeshGenerator* MyGenerator = static_cast<FRecastNavMeshGenerator*>(GetGenerator());
+		if (MyGenerator)
+		{
+			MyGenerator->ReAddTiles(Tiles);
+		}
+	}
+}
+
+//----------------------------------------------------------------------//
+// FRecastNavMeshCachedData
+//----------------------------------------------------------------------//
+#if WITH_RECAST
 FRecastNavMeshCachedData FRecastNavMeshCachedData::Construct(const ARecastNavMesh* RecastNavMeshActor)
 {
 	check(RecastNavMeshActor);
@@ -1968,3 +2107,4 @@ void FRecastNavMeshCachedData::OnAreaAdded(const UClass* AreaClass, int32 AreaID
 }
 
 #endif// WITH_RECAST
+
