UE-7934:  Crash in FCanvas::Flush_GameThread.  Updated to not deref the render target on the game thread, and use GetSizeXY instead, as per Nick P.

diff --git a/Engine/Source/Runtime/Engine/Private/UserInterface/Canvas.cpp b/Engine/Source/Runtime/Engine/Private/UserInterface/Canvas.cpp
--- a/Engine/Source/Runtime/Engine/Private/UserInterface/Canvas.cpp
+++ b/Engine/Source/Runtime/Engine/Private/UserInterface/Canvas.cpp
@@ -827,15 +827,16 @@ void FCanvas::Flush_GameThread(bool bForce)
 		}
 	};
 	// sort the array of FCanvasSortElement entries so that higher sort keys render first (back-to-front)
-	SortedElements.Sort( FCompareFCanvasSortElement() );
-	if( ViewRect.Area() <= 0 )
+	SortedElements.Sort(FCompareFCanvasSortElement());
+	
+	const FIntPoint RenderTargetSize = RenderTarget->GetSizeXY();
+	if (ViewRect.Area() <= 0)
 	{
-		ViewRect = FIntRect( FIntPoint::ZeroValue, RenderTarget->GetSizeXY() );
+		ViewRect = FIntRect(FIntPoint::ZeroValue, RenderTargetSize);
 	}
-	if (IsScaledToRenderTarget() && IsValidRef(RenderTarget->GetRenderTargetTexture()))
+	if (IsScaledToRenderTarget())
 	{
-		ViewRect = FIntRect(0, 0, RenderTarget->GetRenderTargetTexture()->GetSizeX(), 
-								  RenderTarget->GetRenderTargetTexture()->GetSizeY());
+		ViewRect = FIntRect(0, 0, RenderTargetSize.X, RenderTargetSize.Y);
 	}
 
 	struct FCanvasFlushParameters
