Fixed non-unity build

diff --git a/Engine/Source/Runtime/Engine/Private/AI/Navigation/NavigationSystem.cpp b/Engine/Source/Runtime/Engine/Private/AI/Navigation/NavigationSystem.cpp
--- a/Engine/Source/Runtime/Engine/Private/AI/Navigation/NavigationSystem.cpp
+++ b/Engine/Source/Runtime/Engine/Private/AI/Navigation/NavigationSystem.cpp
@@ -4,6 +4,7 @@
 #include "MessageLog.h"
 #include "NavDataGenerator.h"
 #include "NavigationOctree.h"
+#include "VisualLogger.h"
 #include "AI/Navigation/NavMeshBoundsVolume.h"
 #include "AI/Navigation/NavRelevantComponent.h"
 #include "AI/Navigation/NavigationInvokerComponent.h"
