UE-1546 - Add support for tracking unmapped objects on client for fast TArray path

diff --git a/Engine/Source/Runtime/CoreUObject/Public/UObject/CoreNet.h b/Engine/Source/Runtime/CoreUObject/Public/UObject/CoreNet.h
--- a/Engine/Source/Runtime/CoreUObject/Public/UObject/CoreNet.h
+++ b/Engine/Source/Runtime/CoreUObject/Public/UObject/CoreNet.h
@@ -307,6 +307,12 @@ struct FNetDeltaSerializeInfo
 		Struct		= NULL;
 
 		NetSerializeCB = NULL;
+
+		bUpdateUnmappedObjects		= false;
+		bOutSomeObjectsWereMapped	= false;
+		bCalledPreNetReceive		= false;
+		bOutHasMoreUnmapped			= false;
+		Object						= NULL;
 	}
 
 	// Used when writing
@@ -325,6 +331,12 @@ struct FNetDeltaSerializeInfo
 
 	INetSerializeCB*				NetSerializeCB;
 
+	bool							bUpdateUnmappedObjects;		// If true, we are wanting to update unmapped objects
+	bool							bOutSomeObjectsWereMapped;
+	bool							bCalledPreNetReceive;
+	bool							bOutHasMoreUnmapped;
+	UObject*						Object;
+
 	// Debugging variables
 	FString							DebugName;
 };
diff --git a/Engine/Source/Runtime/Engine/Classes/Engine/NetSerialization.h b/Engine/Source/Runtime/Engine/Classes/Engine/NetSerialization.h
--- a/Engine/Source/Runtime/Engine/Classes/Engine/NetSerialization.h
+++ b/Engine/Source/Runtime/Engine/Classes/Engine/NetSerialization.h
@@ -329,6 +329,14 @@ struct FFastArraySerializerItem
 	FORCEINLINE void PostReplicatedChange(const struct FFastArraySerializer& InArraySerializer) { }
 };
 
+/** Struct for holding information about a element that has unmapped network guids */
+struct FFastArraySerializerUnmappedItem
+{
+	TArray< FNetworkGUID >		UnmappedGUIDs;		// List of guids that were unmapped so we can quickly check
+	TArray< uint8 >				Buffer;				// Buffer of data to re-serialize when the guids are mapped
+	int32						NumBufferBits;		// Number of bits in the buffer
+};
+
 /** Base struct for wrapping the array used in Fast TArray Replication */
 USTRUCT()
 struct FFastArraySerializer
@@ -336,9 +344,12 @@ struct FFastArraySerializer
 	GENERATED_USTRUCT_BODY()
 
 	FFastArraySerializer() : IDCounter(0), ArrayReplicationKey(0) { }
+
 	TMap<int32, int32>	ItemMap;
-	int32	IDCounter;
-	int32	ArrayReplicationKey;
+	int32				IDCounter;
+	int32				ArrayReplicationKey;
+
+	TMap< int32, FFastArraySerializerUnmappedItem > UnmappedItems;	// List of items that need to be re-serialized when the referenced objects are mapped
 
 	/** This must be called if you add or change an item in the array */
 	void MarkItemDirty(FFastArraySerializerItem & Item)
@@ -367,13 +378,119 @@ struct FFastArraySerializer
 	static bool FastArrayDeltaSerialize( TArray<Type> &Items, FNetDeltaSerializeInfo& Parms, SerializerType& ArraySerializer );
 };
 
-
 /** The function that implements Fast TArray Replication  */
 template< typename Type, typename SerializerType >
 bool FFastArraySerializer::FastArrayDeltaSerialize( TArray<Type> &Items, FNetDeltaSerializeInfo& Parms, SerializerType& ArraySerializer )
 {
 	class UScriptStruct* InnerStruct = Type::StaticStruct();
 
+	if ( Parms.bUpdateUnmappedObjects || Parms.Writer == NULL )
+	{
+		//---------------
+		// Build ItemMap if necessary. This maps ReplicationID to our local index into the Items array.
+		//---------------
+		if (ArraySerializer.ItemMap.Num() != Items.Num())
+		{
+			SCOPE_CYCLE_COUNTER(STAT_NetSerializeFast_Array);
+			UE_LOG(LogNetSerialization, Verbose, TEXT("FastArrayDeltaSerialize: Recreating Items map. Struct: %s, Items.Num: %d Map.Num: %d"), *InnerStruct->GetOwnerStruct()->GetName(), Items.Num(), ArraySerializer.ItemMap.Num() );
+
+			ArraySerializer.ItemMap.Empty();
+			for (int32 i=0; i < Items.Num(); ++i)
+			{
+				if ( Items[i].ReplicationID == INDEX_NONE )
+				{
+					UE_LOG( LogNetSerialization, Warning, TEXT( "FastArrayDeltaSerialize: Item with uninitialized ReplicationID. Struct: %s, ItemIndex: %i" ), *InnerStruct->GetOwnerStruct()->GetName(), i );
+				}
+				ArraySerializer.ItemMap.Add( Items[i].ReplicationID, i );
+			}
+		}
+	}
+
+	if ( Parms.bUpdateUnmappedObjects )
+	{
+		// Loop over each item that has unmapped objects
+		for ( auto It = ArraySerializer.UnmappedItems.CreateIterator(); It; ++It )
+		{
+			// Get the element id
+			const int32 ElementID = It.Key();
+			
+			// Get a reference to the unmapped item itself
+			FFastArraySerializerUnmappedItem& UnmappedItem = It.Value();
+
+			if ( UnmappedItem.UnmappedGUIDs.Num() == 0 || ArraySerializer.ItemMap.Find( ElementID ) == NULL )
+			{
+				// If for some reason the item is gone (or all guids were removed), we don't need to track guids for this item anymore
+				UnmappedItem.UnmappedGUIDs.Empty();
+				It.RemoveCurrent();
+				continue;		// We're done with this unmapped item
+			}
+
+			// Loop over all the guids, and check to see if any of them are loaded yet
+			bool bMappedSomeGUIDs = false;
+
+			for ( int32 i = UnmappedItem.UnmappedGUIDs.Num() - 1; i >= 0 ; i-- )
+			{			
+				const FNetworkGUID& GUID = UnmappedItem.UnmappedGUIDs[i];
+
+				if ( Parms.Map->IsGUIDBroken( GUID, false ) )
+				{
+					// Stop trying to load broken guids
+					UE_LOG( LogNet, Warning, TEXT( "FastArrayDeltaSerialize: Broken GUID. NetGuid: %s" ), *GUID.ToString() );
+					UnmappedItem.UnmappedGUIDs.RemoveAt( i );
+					continue;
+				}
+
+				UObject* Object = Parms.Map->GetObjectFromNetGUID( GUID, false );
+
+				if ( Object != NULL )
+				{
+					// This guid loaded!
+					UnmappedItem.UnmappedGUIDs.RemoveAt( i );
+					bMappedSomeGUIDs = true;
+				}
+			}
+
+			// Check to see if we loaded any guids. If we did, we can serialize the element again which will load it this time
+			if ( bMappedSomeGUIDs )
+			{
+				Parms.bOutSomeObjectsWereMapped = true;
+
+				if ( !Parms.bCalledPreNetReceive )
+				{
+					// Call PreNetReceive if we are going to change a value (some game code will need to think this is an actual replicated value)
+					Parms.Object->PreNetReceive();
+					Parms.bCalledPreNetReceive = true;
+				}
+
+				Type* ThisElement = &Items[ArraySerializer.ItemMap.FindChecked( ElementID )];
+
+				// Initialize the reader with the stored buffer that we need to read from
+				FBitReader Reader( UnmappedItem.Buffer.GetData(), UnmappedItem.NumBufferBits );
+
+				// Read the property (which should serialize any newly mapped objects as well)
+				bool bHasUnmapped = false;
+				Parms.NetSerializeCB->NetSerializeStruct( InnerStruct, Reader, Parms.Map, ThisElement, bHasUnmapped );
+
+				// Let the element know it changed
+				ThisElement->PostReplicatedChange( ArraySerializer );
+			}
+
+			// If we have no more guids, we can remove this item
+			if ( UnmappedItem.UnmappedGUIDs.Num() == 0 )
+			{
+				It.RemoveCurrent();
+			}
+		}
+
+		// If we still have unmapped items, then communicate this to the outside
+		if ( ArraySerializer.UnmappedItems.Num() > 0 )
+		{
+			Parms.bOutHasMoreUnmapped = true;
+		}
+
+		return true;
+	}
+
 	if (Parms.Writer)
 	{
 		//-----------------------------
@@ -473,7 +590,7 @@ bool FFastArraySerializer::FastArrayDeltaSerialize( TArray<Type> &Items, FNetDel
 		}
 
 		// return if nothing changed
-		if ( ChangedElements.Num() == 0  && DeletedElements.Num() == 0)
+		if ( ChangedElements.Num() == 0 && DeletedElements.Num() == 0)
 		{
 			// Nothing changed
 			UE_LOG(LogNetSerialization, Verbose, TEXT("   No Changed Elements in this array - skipping write"));
@@ -512,17 +629,7 @@ bool FFastArraySerializer::FastArrayDeltaSerialize( TArray<Type> &Items, FNetDel
 			UE_LOG(LogNetSerialization, Log, TEXT("   Changed ElementID: %d"), ID);
 
 			bool bHasUnmapped = false;
-
 			Parms.NetSerializeCB->NetSerializeStruct(InnerStruct, Writer, Parms.Map, ThisElement, bHasUnmapped);
-
-			if (bHasUnmapped)
-			{
-				// Set to 0 to mean 'unmapped'. This will force reserialization on next update.
-				*NewMap.Find(ID) = 0;
-				NewState->ArrayReplicationKey = INDEX_NONE;
-				UStruct* StructPtr = InnerStruct;
-				UE_LOG(LogNetSerialization, Log, TEXT("   Property: %s is unmapped. Will reserialize."), *StructPtr->GetName());
-			}
 		}
 	}
 	else
@@ -533,21 +640,6 @@ bool FFastArraySerializer::FastArrayDeltaSerialize( TArray<Type> &Items, FNetDel
 		check(Parms.Reader);
 		FBitReader& Reader = *Parms.Reader;
 
-		//---------------
-		// Build ItemMap if necessary. This maps ReplicationID to our local index into the Items array.
-		//---------------
-		if (ArraySerializer.ItemMap.Num() != Items.Num())
-		{
-			SCOPE_CYCLE_COUNTER(STAT_NetSerializeFast_Array);
-			UE_LOG(LogNetSerialization, Verbose, TEXT("Recreating Items map. Items.Num: %d Map.Num: %d"), Items.Num(), ArraySerializer.ItemMap.Num() );
-
-			ArraySerializer.ItemMap.Empty();
-			for (int32 i=0; i < Items.Num(); ++i)
-			{
-				ArraySerializer.ItemMap.Add( Items[i].ReplicationID, i );
-			}
-		}
-		
 		static const int32 MAX_NUM_CHANGED = 2048;
 		static const int32 MAX_NUM_DELETED = 2048;
 
@@ -588,6 +680,8 @@ bool FFastArraySerializer::FastArrayDeltaSerialize( TArray<Type> &Items, FNetDel
 				int32 ElementID;
 				Reader << ElementID;
 
+				ArraySerializer.UnmappedItems.Remove( ElementID );
+
 				int32* ElementIndexPtr = ArraySerializer.ItemMap.Find(ElementID);
 				if (ElementIndexPtr)
 				{
@@ -640,9 +734,47 @@ bool FFastArraySerializer::FastArrayDeltaSerialize( TArray<Type> &Items, FNetDel
 				ThisElement = &Items[ElementIndex];
 			}
 
+			// Let package map know we want to track and know about any guids that are unmapped during the serialize call
+			Parms.Map->ResetTrackedUnmappedGuids( true );
+
+			// Remember where we started reading from, so that if we have unmapped properties, we can re-deserialize from this data later
+			FBitReaderMark Mark( Reader );
+
 			bool bHasUnmapped = false;
 			Parms.NetSerializeCB->NetSerializeStruct( InnerStruct, Reader, Parms.Map, ThisElement, bHasUnmapped );
 
+			if ( !Reader.IsError() )
+			{
+				// Track unmapped guids
+				const TArray< FNetworkGUID > & TrackedUnmappedGuids = Parms.Map->GetTrackedUnmappedGuids();
+
+				if ( TrackedUnmappedGuids.Num() )
+				{	
+					FFastArraySerializerUnmappedItem& UnmappedItem = ArraySerializer.UnmappedItems.FindOrAdd( ElementID );
+
+					// Copy the unmapped guid list to this unmapped item
+					UnmappedItem.UnmappedGUIDs = TrackedUnmappedGuids;
+					UnmappedItem.Buffer.Empty();
+
+					// Remember the number of bits in the buffer
+					UnmappedItem.NumBufferBits = Reader.GetPosBits() - Mark.GetPos();
+					
+					// Copy the buffer itself
+					Mark.Copy( Reader, UnmappedItem.Buffer );
+
+					// Hijack this property to communicate that we need to be tracked since we have some unmapped guids
+					Parms.bOutHasMoreUnmapped = true;	
+				}
+				else
+				{
+					// If we don't have any unmapped objects, make sure we're no longer tracking this item in the unmapped lists
+					ArraySerializer.UnmappedItems.Remove( ElementID );
+				}
+			}
+
+			// Stop tracking unmapped objects
+			Parms.Map->ResetTrackedUnmappedGuids( false );
+
 			if ( Reader.IsError() )
 			{
 				UE_LOG( LogNetSerialization, Warning, TEXT( "Parms.NetSerializeCB->NetSerializeStruct: Reader.IsError() == true" ) );
@@ -674,6 +806,7 @@ bool FFastArraySerializer::FastArrayDeltaSerialize( TArray<Type> &Items, FNetDel
 			}
 
 			// Clear the map now that the indices are all shifted around. This kind of sucks, we could use slightly better data structures here I think.
+			// This will force the ItemMap to be rebuilt for the current Items array
 			ArraySerializer.ItemMap.Empty();
 		}
 	}
diff --git a/Engine/Source/Runtime/Engine/Private/DataReplication.cpp b/Engine/Source/Runtime/Engine/Private/DataReplication.cpp
--- a/Engine/Source/Runtime/Engine/Private/DataReplication.cpp
+++ b/Engine/Source/Runtime/Engine/Private/DataReplication.cpp
@@ -509,7 +509,7 @@
 				// Pointer to destination.
 				uint8* Data = ReplicatedProp->ContainerPtrToValuePtr<uint8>((uint8*)Object, Element);
 				TArray<uint8>	MetaData;
-				PTRINT Offset = 0;
+				const PTRINT DataOffset = Data - (uint8*)Object;
 
 				// Receive custom delta property.
 				UStructProperty * StructProperty = Cast< UStructProperty >( ReplicatedProp );
@@ -554,6 +554,11 @@
 					return false;
 				}
 
+				if ( Parms.bOutHasMoreUnmapped )
+				{
+					UnmappedCustomProperties.Add( DataOffset, StructProperty );
+				}
+
 				// Successfully received it.
 				UE_LOG( LogNetTraffic, Log, TEXT( " %s - %s" ), *Object->GetName(), *ReplicatedProp->GetName());
 
@@ -743,65 +748,8 @@
 		return;
 	}
 
-	RepLayout->CallRepNotifies( RepState, Object );
-
 	// Call RepNotifies
-	if ( RepNotifies.Num() > 0 )
-	{
-		for (int32 RepNotifyIdx = 0; RepNotifyIdx < RepNotifies.Num(); RepNotifyIdx++)
-		{
-			//UE_LOG(LogNet, Log,  TEXT("Calling Object->%s with %s"), *RepNotifies(RepNotifyIdx)->RepNotifyFunc.ToString(), *RepNotifies(RepNotifyIdx)->GetName()); 						
-			UProperty* RepProperty = RepNotifies[RepNotifyIdx];
-			UFunction* RepNotifyFunc = Object->FindFunctionChecked(RepProperty->RepNotifyFunc);
-						
-			if (RepNotifyFunc->NumParms == 0)
-			{
-				Object->ProcessEvent(RepNotifyFunc, NULL);
-			}
-			else if (RepNotifyFunc->NumParms == 1)
-			{
-				Object->ProcessEvent(RepNotifyFunc, RepProperty->ContainerPtrToValuePtr<uint8>(RepState->StaticBuffer.GetData()) );
-			}
-			else if (RepNotifyFunc->NumParms == 2)
-			{
-				// Fixme: this isn't as safe as it could be. Right now we have two types of parameters: MetaData (a TArray<uint8>)
-				// and the last local value (pointer into the Recent[] array).
-				//
-				// Arrays always expect MetaData. Everything else, including structs, expect last value.
-				// This is enforced with UHT only. If a ::NetSerialize function ever starts producing a MetaData array thats not in UArrayProperty,
-				// we have no static way of catching this and the replication system could pass the wrong thing into ProcessEvent here.
-				//
-				// But this is all sort of an edge case feature anyways, so its not worth tearing things up too much over.
-
-				FMemMark Mark(FMemStack::Get());
-				uint8* Parms = new(FMemStack::Get(),MEM_Zeroed,RepNotifyFunc->ParmsSize)uint8;
-				
-				TFieldIterator<UProperty> Itr(RepNotifyFunc);
-				check(Itr);
-				
-				Itr->CopyCompleteValue( Itr->ContainerPtrToValuePtr<void>(Parms), RepProperty->ContainerPtrToValuePtr<uint8>(RepState->StaticBuffer.GetData()) );
-				++Itr;
-				check(Itr);
-
-				TArray<uint8> *NotifyMetaData = RepNotifyMetaData.Find(RepNotifies[RepNotifyIdx]);
-				check(NotifyMetaData);
-				Itr->CopyCompleteValue( Itr->ContainerPtrToValuePtr<void>(Parms), NotifyMetaData );
-				
-				Object->ProcessEvent(RepNotifyFunc, Parms );
-
-				Mark.Pop();
-			}
- 						
- 			if (Object == NULL || Object->IsPendingKill())
- 			{
- 				// script event destroyed Object
- 				break;
- 			}
-		}
-	}
-
-	RepNotifies.Reset();
-	RepNotifyMetaData.Empty();
+	CallRepNotifies();
 }
 
 static FORCEINLINE FPropertyRetirement ** UpdateAckedRetirements( FPropertyRetirement &	Retire, int32 OutAckPacketId )
@@ -1129,6 +1077,75 @@
 	bLastUpdateEmpty = false; // Ensure we get one more attempt to update properties
 }
 
+void FObjectReplicator::CallRepNotifies()
+{
+	UObject* Object = GetObject();
+
+	if ( Object == NULL || Object->IsPendingKill() )
+	{
+		return;
+	}
+
+	RepLayout->CallRepNotifies( RepState, Object );
+
+	if ( RepNotifies.Num() > 0 )
+	{
+		for (int32 RepNotifyIdx = 0; RepNotifyIdx < RepNotifies.Num(); RepNotifyIdx++)
+		{
+			//UE_LOG(LogNet, Log,  TEXT("Calling Object->%s with %s"), *RepNotifies(RepNotifyIdx)->RepNotifyFunc.ToString(), *RepNotifies(RepNotifyIdx)->GetName()); 						
+			UProperty* RepProperty = RepNotifies[RepNotifyIdx];
+			UFunction* RepNotifyFunc = Object->FindFunctionChecked(RepProperty->RepNotifyFunc);
+
+			if (RepNotifyFunc->NumParms == 0)
+			{
+				Object->ProcessEvent(RepNotifyFunc, NULL);
+			}
+			else if (RepNotifyFunc->NumParms == 1)
+			{
+				Object->ProcessEvent(RepNotifyFunc, RepProperty->ContainerPtrToValuePtr<uint8>(RepState->StaticBuffer.GetData()) );
+			}
+			else if (RepNotifyFunc->NumParms == 2)
+			{
+				// Fixme: this isn't as safe as it could be. Right now we have two types of parameters: MetaData (a TArray<uint8>)
+				// and the last local value (pointer into the Recent[] array).
+				//
+				// Arrays always expect MetaData. Everything else, including structs, expect last value.
+				// This is enforced with UHT only. If a ::NetSerialize function ever starts producing a MetaData array thats not in UArrayProperty,
+				// we have no static way of catching this and the replication system could pass the wrong thing into ProcessEvent here.
+				//
+				// But this is all sort of an edge case feature anyways, so its not worth tearing things up too much over.
+
+				FMemMark Mark(FMemStack::Get());
+				uint8* Parms = new(FMemStack::Get(),MEM_Zeroed,RepNotifyFunc->ParmsSize)uint8;
+
+				TFieldIterator<UProperty> Itr(RepNotifyFunc);
+				check(Itr);
+
+				Itr->CopyCompleteValue( Itr->ContainerPtrToValuePtr<void>(Parms), RepProperty->ContainerPtrToValuePtr<uint8>(RepState->StaticBuffer.GetData()) );
+				++Itr;
+				check(Itr);
+
+				TArray<uint8> *NotifyMetaData = RepNotifyMetaData.Find(RepNotifies[RepNotifyIdx]);
+				check(NotifyMetaData);
+				Itr->CopyCompleteValue( Itr->ContainerPtrToValuePtr<void>(Parms), NotifyMetaData );
+
+				Object->ProcessEvent(RepNotifyFunc, Parms );
+
+				Mark.Pop();
+			}
+
+			if (Object == NULL || Object->IsPendingKill())
+			{
+				// script event destroyed Object
+				break;
+			}
+		}
+	}
+
+	RepNotifies.Reset();
+	RepNotifyMetaData.Empty();
+}
+
 void FObjectReplicator::UpdateUnmappedObjects( bool & bOutHasMoreUnmapped )
 {
 	UObject* Object = GetObject();
@@ -1146,13 +1163,63 @@
 	}
 
 	check( RepState->RepNotifies.Num() == 0 );
+	check( RepNotifies.Num() == 0 );
 
 	bool bSomeObjectsWereMapped = false;
 
+	// Let the rep layout update any unmapped properties
 	RepLayout->UpdateUnmappedObjects( RepState, Connection->PackageMap, Object, bSomeObjectsWereMapped, bOutHasMoreUnmapped );
 
+	// Update unmapped objects for custom properties (currently just fast tarray)
+	for ( auto It = UnmappedCustomProperties.CreateIterator(); It; ++It )
+	{
+		const int32			Offset			= It.Key();
+		UStructProperty*	StructProperty	= It.Value();
+		UScriptStruct*		InnerStruct		= StructProperty->Struct;
+
+		check( InnerStruct->StructFlags & STRUCT_NetDeltaSerializeNative );
+
+		UScriptStruct::ICppStructOps* CppStructOps = InnerStruct->GetCppStructOps();
+
+		check( CppStructOps );
+		check( !InnerStruct->InheritedCppStructOps() );
+
+		FNetDeltaSerializeInfo Parms;
+
+		FNetSerializeCB NetSerializeCB( OwningChannel->Connection->Driver );
+
+		Parms.DebugName			= StructProperty->GetName();
+		Parms.Struct			= InnerStruct;
+		Parms.Map				= Connection->PackageMap;
+		Parms.NetSerializeCB	= &NetSerializeCB;
+
+		Parms.bUpdateUnmappedObjects	= true;
+		Parms.bCalledPreNetReceive		= bSomeObjectsWereMapped;	// RepLayout used this to flag whether PreNetReceive was called
+		Parms.Object					= Object;
+
+		// Call the custom delta serialize function to handle it
+		CppStructOps->NetDeltaSerialize( Parms, (uint8*)Object + Offset );
+
+		// Merge in results
+		bSomeObjectsWereMapped	|= Parms.bOutSomeObjectsWereMapped;
+		bOutHasMoreUnmapped		|= Parms.bOutHasMoreUnmapped;
+
+		if ( Parms.bOutSomeObjectsWereMapped )
+		{
+			// If we mapped a property, call the rep notify
+			TArray<uint8> MetaData;
+			QueuePropertyRepNotify( Object, StructProperty, 0, MetaData );
+		}
+
+		// If this property no longer has unmapped objects, we can stop checking it
+		if ( !Parms.bOutHasMoreUnmapped )
+		{
+			It.RemoveCurrent();
+		}
+	}
+
 	// Call any rep notifies that need to happen when object pointers change
-	RepLayout->CallRepNotifies( RepState, Object );
+	CallRepNotifies();
 
 	if ( bSomeObjectsWereMapped )
 	{
diff --git a/Engine/Source/Runtime/Engine/Public/Net/DataReplication.h b/Engine/Source/Runtime/Engine/Public/Net/DataReplication.h
--- a/Engine/Source/Runtime/Engine/Public/Net/DataReplication.h
+++ b/Engine/Source/Runtime/Engine/Public/Net/DataReplication.h
@@ -72,8 +72,10 @@ public:
 	UNetConnection *								Connection;					// Connection this replicator was created on
 	class UActorChannel	*							OwningChannel;
 
-	TArray< UProperty *,TInlineAllocator< 32 > >	RepNotifies;
-	TMap< UProperty *, TArray<uint8> >				RepNotifyMetaData;
+	TMap< int32, UStructProperty* >					UnmappedCustomProperties;
+
+	TArray< UProperty*,TInlineAllocator< 32 > >		RepNotifies;
+	TMap< UProperty*, TArray<uint8> >				RepNotifyMetaData;
 
 	TSharedPtr< FRepLayout >						RepLayout;
 	FRepState *										RepState;
@@ -125,6 +127,8 @@ public:
 
 	void StartBecomingDormant();
 
+	void CallRepNotifies();
+
 	void UpdateUnmappedObjects( bool & bOutHasMoreUnmapped );
 
 	FORCEINLINE UObject *	GetObject() const { return ObjectPtr.Get(); }
