Guard against no actors in a group when updating interpolation for parent tracks when evaluating camera cuts

diff --git a/Engine/Source/Runtime/Engine/Private/Interpolation.cpp b/Engine/Source/Runtime/Engine/Private/Interpolation.cpp
--- a/Engine/Source/Runtime/Engine/Private/Interpolation.cpp
+++ b/Engine/Source/Runtime/Engine/Private/Interpolation.cpp
@@ -893,30 +893,33 @@ void AMatineeActor::TermInterp()
 void AMatineeActor::UpdateInterpForParentMovementTracks( float Time, UInterpGroupInst* ViewGroupInst )
 {
 	AActor* Actor = ViewGroupInst->GetGroupActor();
-	AActor* Parent = Actor->GetAttachParentActor();
-
-	UInterpGroupInst* ParentInst = FindGroupInst( Parent );
-	if(ParentInst)
+	if( Actor )
 	{
-		TArray<UInterpTrack*> FoundTracks;
+		AActor* Parent = Actor->GetAttachParentActor();
 
-		UInterpTrackInst* ParentTrackInst = nullptr;
-		for (UInterpTrackInst* Inst : ParentInst->TrackInst )
+		UInterpGroupInst* ParentInst = FindGroupInst(Parent);
+		if(ParentInst)
 		{
-			if( Inst->GetGroupActor() == Parent )
+			TArray<UInterpTrack*> FoundTracks;
+
+			UInterpTrackInst* ParentTrackInst = nullptr;
+			for(UInterpTrackInst* Inst : ParentInst->TrackInst)
 			{
-				ParentTrackInst = Inst;
-				break;
+				if(Inst->GetGroupActor() == Parent)
+				{
+					ParentTrackInst = Inst;
+					break;
+				}
 			}
-		}
 
-		if (ParentTrackInst)
-		{
-			TArray<UInterpTrack*> FoundTracks;
-			ParentInst->Group->FindTracksByClass( UInterpTrackMove::StaticClass(), FoundTracks );
-			//Just use the first one, multiple move tracks wouldnt work well anyway
-			UInterpTrackMove* MoveTrack = CastChecked<UInterpTrackMove>(FoundTracks[0]);
-			MoveTrack->ConditionalUpdateTrack(Time, ParentTrackInst, true);
+			if(ParentTrackInst)
+			{
+				TArray<UInterpTrack*> FoundTracks;
+				ParentInst->Group->FindTracksByClass(UInterpTrackMove::StaticClass(), FoundTracks);
+				//Just use the first one, multiple move tracks wouldnt work well anyway
+				UInterpTrackMove* MoveTrack = CastChecked<UInterpTrackMove>(FoundTracks[0]);
+				MoveTrack->ConditionalUpdateTrack(Time, ParentTrackInst, true);
+			}
 		}
 	}
 }
