#UE4 - attempt to fix IsLocalPlayerController problems at creation/destruction

diff --git a/Engine/Source/Runtime/Engine/Classes/GameFramework/GameMode.h b/Engine/Source/Runtime/Engine/Classes/GameFramework/GameMode.h
--- a/Engine/Source/Runtime/Engine/Classes/GameFramework/GameMode.h
+++ b/Engine/Source/Runtime/Engine/Classes/GameFramework/GameMode.h
@@ -412,12 +412,21 @@ public:
 	 * Called to login new players by creating a player controller, overridable by the game
 	 *
 	 * Sets up basic properties of the player (name, unique id, registers with backend, etc) and should not be used to do
-	 * more complicated game logic.  The player controller is not fully initialized within this function as far as networking is conecerned.
+	 * more complicated game logic.  The player controller is not fully initialized within this function as far as networking is concerned.
 	 * Save "game logic" for PostLogin which is called shortly afterward.
 	 *
+	 * @param NewPlayer pointer to the UPlayer object that represents this player (either local or remote)
+	 * @param RemoteRole the remote role this controller has
+	 * @param Portal desired portal location specified by the client
+	 * @param Options game options passed in by the client at login
+	 * @param UniqueId platform specific unique identifier for the logging in player
+	 * @param ErrorMessage [out] error message, if any, why this login will be failing
+	 *
 	 * If login is successful, returns a new PlayerController to associate with this player. Login fails if ErrorMessage string is set. 
+	 *
+	 * @return a new player controller for the logged in player, NULL if login failed for any reason
 	 */
-	virtual APlayerController* Login(class UPlayer* NewPlayer, const FString& Portal, const FString& Options, const TSharedPtr<class FUniqueNetId>& UniqueId, FString& ErrorMessage);
+	virtual APlayerController* Login(class UPlayer* NewPlayer, ENetRole RemoteRole, const FString& Portal, const FString& Options, const TSharedPtr<class FUniqueNetId>& UniqueId, FString& ErrorMessage);
 
 	/** Called after a successful login.  This is the first place it is safe to call replicated functions on the PlayerAController. */
 	virtual void PostLogin( APlayerController* NewPlayer );
@@ -426,8 +435,16 @@ public:
 	UFUNCTION(BlueprintImplementableEvent, Category="Game", meta=(FriendlyName="PostLogin"))
 	virtual void K2_PostLogin( APlayerController* NewPlayer );
 
-	/** Spawns a PlayerController at the specified location; split out from Login()/HandleSeamlessTravelPlayer() for easier overriding */
-	virtual APlayerController* SpawnPlayerController(FVector const& SpawnLocation, FRotator const& SpawnRotation);
+	/** 
+	 * Spawns a PlayerController at the specified location; split out from Login()/HandleSeamlessTravelPlayer() for easier overriding 
+	 *
+	 * @param RemoteRole the role this controller will play remotely
+	 * @param SpawnLocation location in the world to spawn
+	 * @param SpawnRotation rotation to set relative to the world
+	 *
+	 * @return PlayerController for the player, NULL if there is any reason this player shouldn't exist or due to some error
+	 */
+	virtual APlayerController* SpawnPlayerController(ENetRole RemoteRole, FVector const& SpawnLocation, FRotator const& SpawnRotation);
 
 	/** @Returns true if NewPlayerController may only join the server as a spectator. */
 	virtual bool MustSpectate(APlayerController* NewPlayerController) const;
diff --git a/Engine/Source/Runtime/Engine/Classes/GameFramework/PlayerController.h b/Engine/Source/Runtime/Engine/Classes/GameFramework/PlayerController.h
--- a/Engine/Source/Runtime/Engine/Classes/GameFramework/PlayerController.h
+++ b/Engine/Source/Runtime/Engine/Classes/GameFramework/PlayerController.h
@@ -1509,10 +1509,19 @@ public:
 	 */
 	virtual void SendClientAdjustment();
 
+	/**
+	 * Designate this player controller as local (public for GameMode to use, not expected to be called anywhere else)
+	 */
+	void SetAsLocalPlayerController() { bIsLocalPlayerController = true; }
+
 private:
 	/** Used to delay calling ClientRestart() again when it hasn't been appropriately acknowledged. */
 	float		LastRetryPlayerTime;
 
+	/** Set during SpawnActor once and never again to indicate the intent of this controller instance (SERVER ONLY) */
+	UPROPERTY()
+	bool		bIsLocalPlayerController;
+
 public:
 	/** Counter for this players seamless travels (used along with the below value, to restrict ServerNotifyLoadedWorld) */
 	UPROPERTY()
diff --git a/Engine/Source/Runtime/Engine/Private/GameMode.cpp b/Engine/Source/Runtime/Engine/Private/GameMode.cpp
--- a/Engine/Source/Runtime/Engine/Private/GameMode.cpp
+++ b/Engine/Source/Runtime/Engine/Private/GameMode.cpp
@@ -801,7 +801,7 @@ void AGameMode::HandleSeamlessTravelPlayer(AController*& C)
 		if (PC->Player != NULL)
 		{
 			// we need to spawn a new PlayerController to replace the old one
-			APlayerController* NewPC = SpawnPlayerController(PC->GetFocalLocation(), PC->GetControlRotation());
+			APlayerController* NewPC = SpawnPlayerController(PC->IsLocalPlayerController() ? ROLE_SimulatedProxy : ROLE_AutonomousProxy, PC->GetFocalLocation(), PC->GetControlRotation());
 			if (NewPC == NULL)
 			{
 				UE_LOG(LogGameMode, Warning, TEXT("Failed to spawn new PlayerController for %s (old class %s)"), *PC->GetHumanReadableName(), *PC->GetClass()->GetName());
@@ -1025,7 +1025,7 @@ bool AGameMode::MustSpectate(APlayerController* NewPlayerController) const
 	return NewPlayerController->PlayerState->bOnlySpectator;
 }
 
-APlayerController* AGameMode::Login(UPlayer* NewPlayer, const FString& Portal, const FString& Options, const TSharedPtr<FUniqueNetId>& UniqueId, FString& ErrorMessage)
+APlayerController* AGameMode::Login(UPlayer* NewPlayer, ENetRole RemoteRole, const FString& Portal, const FString& Options, const TSharedPtr<FUniqueNetId>& UniqueId, FString& ErrorMessage)
 {
 	ErrorMessage = GameSession->ApproveLogin(Options);
 	if (!ErrorMessage.IsEmpty())
@@ -1033,7 +1033,7 @@ APlayerController* AGameMode::Login(UPlayer* NewPlayer, const FString& Portal, c
 		return NULL;
 	}
 
-	APlayerController* NewPlayerController = SpawnPlayerController(FVector::ZeroVector, FRotator::ZeroRotator);
+	APlayerController* NewPlayerController = SpawnPlayerController(RemoteRole, FVector::ZeroVector, FRotator::ZeroRotator);
 
 	// Handle spawn failure.
 	if (NewPlayerController == NULL)
@@ -1277,12 +1277,25 @@ void AGameMode::PreLogin(const FString& Options, const FString& Address, const T
 	ErrorMessage = GameSession->ApproveLogin(Options);
 }
 
-APlayerController* AGameMode::SpawnPlayerController(FVector const& SpawnLocation, FRotator const& SpawnRotation)
+APlayerController* AGameMode::SpawnPlayerController(ENetRole RemoteRole, FVector const& SpawnLocation, FRotator const& SpawnRotation)
 {
 	FActorSpawnParameters SpawnInfo;
 	SpawnInfo.Instigator = Instigator;	
 	SpawnInfo.ObjectFlags |= RF_Transient;	// We never want to save player controllers into a map
-	return GetWorld()->SpawnActor<APlayerController>(PlayerControllerClass, SpawnLocation, SpawnRotation, SpawnInfo );
+	SpawnInfo.bDeferConstruction = true;
+	APlayerController* NewPC = GetWorld()->SpawnActor<APlayerController>(PlayerControllerClass, SpawnLocation, SpawnRotation, SpawnInfo);
+	if (NewPC)
+	{
+		if (RemoteRole == ROLE_SimulatedProxy)
+		{
+			// This is a local player because it has no authority/autonomous remote role
+			NewPC->SetAsLocalPlayerController();
+		}
+		
+		UGameplayStatics::FinishSpawningActor(NewPC, FTransform(SpawnRotation, SpawnLocation));
+	}
+
+	return NewPC;
 }
 
 UClass* AGameMode::GetDefaultPawnClassForController(AController* InController)
diff --git a/Engine/Source/Runtime/Engine/Private/LevelActor.cpp b/Engine/Source/Runtime/Engine/Private/LevelActor.cpp
--- a/Engine/Source/Runtime/Engine/Private/LevelActor.cpp
+++ b/Engine/Source/Runtime/Engine/Private/LevelActor.cpp
@@ -577,7 +577,7 @@ APlayerController* UWorld::SpawnPlayActor(UPlayer* NewPlayer, ENetRole RemoteRol
 	AGameMode* GameMode = GetAuthGameMode();
 
 	// Give the GameMode a chance to accept the login
-	APlayerController* const NewPlayerController = GameMode->Login(NewPlayer, *InURL.Portal, Options, UniqueId, Error);
+	APlayerController* const NewPlayerController = GameMode->Login(NewPlayer, RemoteRole, *InURL.Portal, Options, UniqueId, Error);
 	if (NewPlayerController == NULL)
 	{
 		UE_LOG(LogSpawn, Warning, TEXT("Login failed: %s"), *Error);
diff --git a/Engine/Source/Runtime/Engine/Private/PlayerController.cpp b/Engine/Source/Runtime/Engine/Private/PlayerController.cpp
--- a/Engine/Source/Runtime/Engine/Private/PlayerController.cpp
+++ b/Engine/Source/Runtime/Engine/Private/PlayerController.cpp
@@ -71,6 +71,8 @@ APlayerController::APlayerController(const FObjectInitializer& ObjectInitializer
 
 	bAutoManageActiveCameraTarget = true;
 
+	bIsLocalPlayerController = false;
+
 	if (RootComponent)
 	{
 		// We want to drive rotation with ControlRotation regardless of attachment state.
@@ -106,29 +108,20 @@ UNetConnection* APlayerController::GetNetConnection()
 
 bool APlayerController::IsLocalController() const
 {
-	if (Player == NULL)
-	{
-		UE_LOG(LogPlayerController, Warning, TEXT("Calling IsLocalController() while Player is NULL is undefined!"));
-	}
-
 	ENetMode NetMode = GetNetMode();
 	if (NetMode == NM_DedicatedServer)
 	{
+		check(!bIsLocalPlayerController);
 		return false;
 	}
 
-	if (Super::IsLocalController())
-	{
-		return true;
-	}
-
 	if (NetMode == NM_Client)
 	{
 		// Clients only receive their own PC. We are not ROLE_AutonomousProxy until after PostInitializeComponents so we can't check that.
 		return true;
 	}
 
-	return false;
+	return bIsLocalPlayerController;
 }
 
 bool APlayerController::IsLocalPlayerController() const
