Fix for Mac compile error

diff --git a/Engine/Source/Runtime/Engine/Classes/Engine/NetSerialization.h b/Engine/Source/Runtime/Engine/Classes/Engine/NetSerialization.h
--- a/Engine/Source/Runtime/Engine/Classes/Engine/NetSerialization.h
+++ b/Engine/Source/Runtime/Engine/Classes/Engine/NetSerialization.h
@@ -435,7 +435,7 @@ bool FFastArraySerializer::FastArrayDeltaSerialize( TArray<Type> &Items, FNetDel
 				if ( Parms.Map->IsGUIDBroken( GUID, false ) )
 				{
 					// Stop trying to load broken guids
-					UE_LOG( LogNet, Warning, TEXT( "FastArrayDeltaSerialize: Broken GUID. NetGuid: %s" ), *GUID.ToString() );
+					UE_LOG( LogNetSerialization, Warning, TEXT( "FastArrayDeltaSerialize: Broken GUID. NetGuid: %s" ), *GUID.ToString() );
 					UnmappedItem.UnmappedGUIDs.RemoveAt( i );
 					continue;
 				}
