#UE4 Fixed a linux compile error caused by the use of UWorld in a header.

diff --git a/Engine/Source/Runtime/Engine/Classes/AI/Navigation/NavigationData.h b/Engine/Source/Runtime/Engine/Classes/AI/Navigation/NavigationData.h
--- a/Engine/Source/Runtime/Engine/Classes/AI/Navigation/NavigationData.h
+++ b/Engine/Source/Runtime/Engine/Classes/AI/Navigation/NavigationData.h
@@ -11,6 +11,7 @@ class UNavigationQueryFilter;
 class FNavDataGenerator; 
 class INavLinkCustomInterface;
 class UPrimitiveComponent;
+class UWorld;
 
 USTRUCT()
 struct FSupportedAreaData
@@ -539,8 +540,7 @@ public:
 		FNavPathSharedPtr SharedPath = MakeShareable(new PathType());
 		SharedPath->SetNavigationDataUsed(this);
 		SharedPath->SetQuerier(Querier);
-		const UWorld* World = GetCachedWorld();
-		SharedPath->SetTimeStamp(World ? World->GetTimeSeconds() : 0);
+		SharedPath->SetTimeStamp( GetWorldTimeStamp() );
 
 		DECLARE_CYCLE_STAT(TEXT("FSimpleDelegateGraphTask.Adding a path to ActivePaths"),
 			STAT_FSimpleDelegateGraphTask_AddingPathToActivePaths,
@@ -716,7 +716,9 @@ public:
 	 *	@NOTE potentially expensive, so use it with caution */
 	virtual ENavigationQueryResult::Type CalcPathLengthAndCost(const FVector& PathStart, const FVector& PathEnd, float& OutPathLength, float& OutPathCost, TSharedPtr<const FNavigationQueryFilter> QueryFilter = NULL, const UObject* Querier = NULL) const PURE_VIRTUAL(ANavigationData::CalcPathLength, return ENavigationQueryResult::Invalid;);
 
-	const UWorld* GetCachedWorld() const { return CachedWorld; }
+	const UWorld* GetCachedWorld() const;
+	float GetWorldTimeStamp() const;
+
 	//----------------------------------------------------------------------//
 	// Areas
 	//----------------------------------------------------------------------//
diff --git a/Engine/Source/Runtime/Engine/Private/AI/Navigation/NavigationData.cpp b/Engine/Source/Runtime/Engine/Private/AI/Navigation/NavigationData.cpp
--- a/Engine/Source/Runtime/Engine/Private/AI/Navigation/NavigationData.cpp
+++ b/Engine/Source/Runtime/Engine/Private/AI/Navigation/NavigationData.cpp
@@ -242,8 +242,7 @@ void ANavigationData::TickActor(float DeltaTime, enum ELevelTick TickType, FActo
 
 	if (RepathRequests.Num() > 0)
 	{
-		const UWorld* World = GetCachedWorld();
-		float TimeStamp = World ? World->GetTimeSeconds() : 0.f;
+		float TimeStamp = GetWorldTimeStamp();
 		TArray<FNavPathRecalculationRequest> PostponedRequests;
 
 		// @todo batch-process it!
@@ -512,6 +511,17 @@ void ANavigationData::DrawDebugPath(FNavigationPath* Path, FColor PathColor, UCa
 	Path->DebugDraw(this, PathColor, Canvas, bPersistent, NextPathPointIndex);
 }
 
+const UWorld* ANavigationData::GetCachedWorld() const
+{
+	return CachedWorld;
+}
+
+float ANavigationData::GetWorldTimeStamp() const
+{
+	const UWorld* World = GetCachedWorld();
+	return World ? World->GetTimeSeconds() : 0.f;
+}
+
 void ANavigationData::OnNavAreaAdded(const UClass* NavAreaClass, int32 AgentIndex)
 {
 	// check if area can be added
