Fix for shader manager crashing on shutdown when bAllowAsynchronousShaderCompiling is false

diff --git a/Engine/Source/Runtime/Engine/Public/ShaderCompiler.h b/Engine/Source/Runtime/Engine/Public/ShaderCompiler.h
--- a/Engine/Source/Runtime/Engine/Public/ShaderCompiler.h
+++ b/Engine/Source/Runtime/Engine/Public/ShaderCompiler.h
@@ -77,7 +77,10 @@ public:
 	virtual uint32 Run();
 	inline void WaitForCompletion() const
 	{
-		Thread->WaitForCompletion();
+		if( Thread )
+		{
+			Thread->WaitForCompletion();
+		}
 	}
 
 	/** Checks the thread's health, and passes on any errors that have occured.  Called by the main thread. */
