Final touches leading to Navigation Active Tiles being v 1.0 #UE4

diff --git a/Engine/Source/Editor/DetailCustomizations/Private/NavLinkStructCustomization.cpp b/Engine/Source/Editor/DetailCustomizations/Private/NavLinkStructCustomization.cpp
--- a/Engine/Source/Editor/DetailCustomizations/Private/NavLinkStructCustomization.cpp
+++ b/Engine/Source/Editor/DetailCustomizations/Private/NavLinkStructCustomization.cpp
@@ -42,7 +42,7 @@ void FNavLinkStructCustomization::CustomizeChildren( TSharedRef<class IPropertyH
 
 	FString AgentPrefix("bSupportsAgent");
 	const UNavigationSystem* DefNavSys = (UNavigationSystem*)(UNavigationSystem::StaticClass()->GetDefaultObject());
-	const int32 NumAgents = FMath::Min(DefNavSys->SupportedAgents.Num(), 16);
+	const int32 NumAgents = FMath::Min(DefNavSys->GetSupportedAgents().Num(), 16);
 
 	for (uint32 i = 0; i < NumChildren; i++)
 	{
@@ -54,7 +54,7 @@ void FNavLinkStructCustomization::CustomizeChildren( TSharedRef<class IPropertyH
 
 			if (AgentIdx >= 0 && AgentIdx < NumAgents)
 			{
-				FText PropName = FText::Format(LOCTEXT("SupportedAgentFmt", "Supports Agent: {0}"), FText::FromName(DefNavSys->SupportedAgents[AgentIdx].Name));
+				FText PropName = FText::Format(LOCTEXT("SupportedAgentFmt", "Supports Agent: {0}"), FText::FromName(DefNavSys->GetSupportedAgents()[AgentIdx].Name));
 				StructBuilder.AddChildContent(PropName)
 					.NameContent()
 					[
diff --git a/Engine/Source/Runtime/Engine/Classes/AI/Navigation/NavigationSystem.h b/Engine/Source/Runtime/Engine/Classes/AI/Navigation/NavigationSystem.h
--- a/Engine/Source/Runtime/Engine/Classes/AI/Navigation/NavigationSystem.h
+++ b/Engine/Source/Runtime/Engine/Classes/AI/Navigation/NavigationSystem.h
@@ -175,22 +175,24 @@ public:
 	UPROPERTY(config, EditAnywhere, Category=NavigationSystem)
 	uint32 bSkipAgentHeightCheckWhenPickingNavData:1;
 
+protected:
 	/** If set to true navigation will be generated only around registered "navigation enforcers"
 	*	This has a range of consequences (including how navigation octree operates) so it needs to
 	*	be a conscious decision.
 	*	Once enabled results in whole world being navigable.
 	*	@see RegisterNavigationInvoker
 	*/
-	UPROPERTY(EditDefaultsOnly, Category = "Navigation", config)
+	UPROPERTY(EditDefaultsOnly, Category = "Navigation Enforcing", config)
 	uint32 bGenerateNavigationOnlyAroundNavigationInvokers : 1;
 
-	/** Minimal time between active tiles set update */
+	/** Minimal time, in seconds, between active tiles set update */
 	UPROPERTY(EditAnywhere, Category = "Navigation Enforcing", meta = (ClampMin = "0.1", UIMin = "0.1", EditCondition = "bGenerateNavigationOnlyAroundNavigationInvokers"), config)
 	float ActiveTilesUpdateInterval;
 
 	UPROPERTY(config, EditAnywhere, Category = Agents)
 	TArray<FNavDataConfig> SupportedAgents;
-	
+
+public:
 	/** update frequency for dirty areas on navmesh */
 	UPROPERTY(config, EditAnywhere, Category=NavigationSystem)
 	float DirtyAreasUpdateFreq;
@@ -292,6 +294,8 @@ public:
 	UFUNCTION(BlueprintCallable, Category = Navigation)
 	void UnregisterNavigationInvoker(AActor* Invoker);
 
+	FORCEINLINE bool IsActiveTilesGenerationEnabled() const{ return bGenerateNavigationOnlyAroundNavigationInvokers; }
+	
 	/** delegate type for events that dirty the navigation data ( Params: const FBox& DirtyBounds ) */
 	DECLARE_MULTICAST_DELEGATE_OneParam(FOnNavigationDirty, const FBox&);
 	/** called after navigation influencing event takes place*/
@@ -465,6 +469,7 @@ public:
 	FORCEINLINE static TSubclassOf<UNavArea> GetDefaultObstacleArea() { return DefaultObstacleArea; }
 
 	FORCEINLINE const FNavDataConfig& GetDefaultSupportedAgentConfig() const { check(SupportedAgents.Num() > 0);  return SupportedAgents[0]; }
+	FORCEINLINE const TArray<FNavDataConfig>& GetSupportedAgents() const { return SupportedAgents; }
 
 	void ApplyWorldOffset(const FVector& InOffset, bool bWorldShift);
 
diff --git a/Engine/Source/Runtime/Engine/Private/AI/Navigation/NavAreas/NavArea.cpp b/Engine/Source/Runtime/Engine/Private/AI/Navigation/NavAreas/NavArea.cpp
--- a/Engine/Source/Runtime/Engine/Private/AI/Navigation/NavAreas/NavArea.cpp
+++ b/Engine/Source/Runtime/Engine/Private/AI/Navigation/NavAreas/NavArea.cpp
@@ -60,11 +60,11 @@ void UNavArea::UpdateAgentConfig()
 	check(DefNavSys);
 
 	const int32 MaxAllowedAgents = 16;
-	const int32 NumAgents = FMath::Min(DefNavSys->SupportedAgents.Num(), MaxAllowedAgents);
-	if (DefNavSys->SupportedAgents.Num() > MaxAllowedAgents)
+	const int32 NumAgents = FMath::Min(DefNavSys->GetSupportedAgents().Num(), MaxAllowedAgents);
+	if (DefNavSys->GetSupportedAgents().Num() > MaxAllowedAgents)
 	{
 		UE_LOG(LogNavigation, Error, TEXT("Navigation system supports %d agents, but only %d can be shown in %s properties!"),
-			DefNavSys->SupportedAgents.Num(), MaxAllowedAgents, *GetClass()->GetName());
+			DefNavSys->GetSupportedAgents().Num(), MaxAllowedAgents, *GetClass()->GetName());
 	}
 
 	const FString CustomNameMeta = TEXT("DisplayName");
@@ -77,7 +77,7 @@ void UNavArea::UpdateAgentConfig()
 		if (i < NumAgents && NumAgents > 1)
 		{
 			Prop->SetPropertyFlags(CPF_Edit);
-			Prop->SetMetaData(*CustomNameMeta, *FString::Printf(TEXT("Supports Agent: %s"), *DefNavSys->SupportedAgents[i].Name.ToString()));
+			Prop->SetMetaData(*CustomNameMeta, *FString::Printf(TEXT("Supports Agent: %s"), *DefNavSys->GetSupportedAgents()[i].Name.ToString()));
 		}
 		else
 		{
diff --git a/Engine/Source/Runtime/Engine/Private/AI/Navigation/NavAreas/NavAreaMeta_SwitchByAgent.cpp b/Engine/Source/Runtime/Engine/Private/AI/Navigation/NavAreas/NavAreaMeta_SwitchByAgent.cpp
--- a/Engine/Source/Runtime/Engine/Private/AI/Navigation/NavAreas/NavAreaMeta_SwitchByAgent.cpp
+++ b/Engine/Source/Runtime/Engine/Private/AI/Navigation/NavAreas/NavAreaMeta_SwitchByAgent.cpp
@@ -62,11 +62,11 @@ void UNavAreaMeta_SwitchByAgent::UpdateAgentConfig()
 	check(DefNavSys);
 
 	const int32 MaxAllowedAgents = 16;
-	const int32 NumAgents = FMath::Min(DefNavSys->SupportedAgents.Num(), MaxAllowedAgents);
-	if (DefNavSys->SupportedAgents.Num() > MaxAllowedAgents)
+	const int32 NumAgents = FMath::Min(DefNavSys->GetSupportedAgents().Num(), MaxAllowedAgents);
+	if (DefNavSys->GetSupportedAgents().Num() > MaxAllowedAgents)
 	{
 		UE_LOG(LogNavigation, Error, TEXT("Navigation system supports %d agents, but only %d can be shown in %s properties!"),
-			DefNavSys->SupportedAgents.Num(), MaxAllowedAgents, *GetClass()->GetName());
+			DefNavSys->GetSupportedAgents().Num(), MaxAllowedAgents, *GetClass()->GetName());
 	}
 
 	const FString CustomNameMeta = TEXT("DisplayName");
@@ -79,7 +79,7 @@ void UNavAreaMeta_SwitchByAgent::UpdateAgentConfig()
 		if (i < NumAgents && NumAgents > 1)
 		{
 			Prop->SetPropertyFlags(CPF_Edit);
-			Prop->SetMetaData(*CustomNameMeta, *FString::Printf(TEXT("Area Class for: %s"), *DefNavSys->SupportedAgents[i].Name.ToString()));
+			Prop->SetMetaData(*CustomNameMeta, *FString::Printf(TEXT("Area Class for: %s"), *DefNavSys->GetSupportedAgents()[i].Name.ToString()));
 		}
 		else
 		{
diff --git a/Engine/Source/Runtime/Engine/Private/AI/Navigation/NavigationSystem.cpp b/Engine/Source/Runtime/Engine/Private/AI/Navigation/NavigationSystem.cpp
--- a/Engine/Source/Runtime/Engine/Private/AI/Navigation/NavigationSystem.cpp
+++ b/Engine/Source/Runtime/Engine/Private/AI/Navigation/NavigationSystem.cpp
@@ -706,6 +706,7 @@ void UNavigationSystem::Tick(float DeltaSeconds)
 		INC_FLOAT_STAT_BY(STAT_Navigation_CumulativeBuildTime,(float)ThisTime*1000);
 	}
 	
+	if (bGenerateNavigationOnlyAroundNavigationInvokers)
 	{
 		UpdateInvokers();
 	}
@@ -3615,49 +3616,32 @@ void UNavigationSystem::UpdateInvokers()
 			TArray<FNavigationInvokerRaw> InvokerLocations;
 			InvokerLocations.Reserve(Invokers.Num());
 
-			for (auto Item = Invokers.CreateConstIterator(); Item; ++Item)
+			for (auto ItemIterator = Invokers.CreateIterator(); ItemIterator; ++ItemIterator)
 			{
-				AActor* Actor = Item->Value.Actor.Get();
+				AActor* Actor = ItemIterator->Value.Actor.Get();
 				if (Actor != nullptr
 #if WITH_EDITOR
+					// Would like to ignore objects in transactional buffer here, but there's no flag for it
 					//&& (GIsEditor == false || Item.Actor->HasAnyFlags(RF_Transactional | RF_PendingKill) == false)
 #endif //WITH_EDITOR
 					)
 				{
-					InvokerLocations.Add(FNavigationInvokerRaw(Actor->GetActorLocation(), Item->Value.GenerationRadius, Item->Value.RemovalRadius));
+					InvokerLocations.Add(FNavigationInvokerRaw(Actor->GetActorLocation(), ItemIterator->Value.GenerationRadius, ItemIterator->Value.RemovalRadius));
+				}
+				else
+				{
+					ItemIterator.RemoveCurrent();
 				}
 			}
 
 #if ENABLE_VISUAL_LOG
-			//TArray<FNavigationInvokerRaw> LocationsCopy = InvokerLocations;
 			const double CachingFinishTime = FPlatformTime::Seconds();
-#endif //ENABLE_VISUAL_LOG
-			//InvokerLocations = KMedoidsClusterization::Clusterize(InvokerLocations, MaxClustersCount);
-
-#if ENABLE_VISUAL_LOG
-			//const double ClusterizationFinishTime = FPlatformTime::Seconds();
-
 			UE_VLOG(this, LogNavigation, Log, TEXT("Caching time %fms"), float((CachingFinishTime - StartTime) / 1000));
-			//UE_VLOG(this, LogNavigation, Log, TEXT("Clusterization time %fms for %d clusters"), float(ClusterizationFinishTime - CachingFinishTime) / 1000, MaxClustersCount);
 
-			FVisualLogger& Vlog = FVisualLogger::Get();
-			if (Vlog.IsRecording())
+			for (const auto& InvokerData : InvokerLocations)
 			{
-				/*FVisualLogEntry* Entry = Vlog.GetEntryToWrite(this, GetWorld()->TimeSeconds);
-				check(Entry);
-
-				FVisualLogShapeElement LogElement(TEXT(""), FColorList::LightBlue, 30, LogNavigation.GetCategoryName());
-				LogElement.Points = LocationsCopy;
-				LogElement.Type = EVisualLoggerShapeElement::SinglePoint;
-				LogElement.Verbosity = ELogVerbosity::Log;
-
-				Entry->AddElement(LogElement);*/
-
-				for (const auto& InvokerData : InvokerLocations)
-				{
-					UE_VLOG_CYLINDER(this, LogNavigation, Log, InvokerData.Location, InvokerData.Location + FVector(0, 0, 20), InvokerData.RadiusMax, FColorList::Blue, TEXT(""));
-					UE_VLOG_CYLINDER(this, LogNavigation, Log, InvokerData.Location, InvokerData.Location + FVector(0, 0, 20), InvokerData.RadiusMin, FColorList::CadetBlue, TEXT(""));
-				}
+				UE_VLOG_CYLINDER(this, LogNavigation, Log, InvokerData.Location, InvokerData.Location + FVector(0, 0, 20), InvokerData.RadiusMax, FColorList::Blue, TEXT(""));
+				UE_VLOG_CYLINDER(this, LogNavigation, Log, InvokerData.Location, InvokerData.Location + FVector(0, 0, 20), InvokerData.RadiusMin, FColorList::CadetBlue, TEXT(""));
 			}
 #endif // ENABLE_VISUAL_LOG
 
diff --git a/Engine/Source/Runtime/Engine/Private/AI/Navigation/RecastNavMesh.cpp b/Engine/Source/Runtime/Engine/Private/AI/Navigation/RecastNavMesh.cpp
--- a/Engine/Source/Runtime/Engine/Private/AI/Navigation/RecastNavMesh.cpp
+++ b/Engine/Source/Runtime/Engine/Private/AI/Navigation/RecastNavMesh.cpp
@@ -316,7 +316,7 @@ ANavigationData* ARecastNavMesh::CreateNavigationInstances(UNavigationSystem* Na
 		return NULL;
 	}
 
-	const TArray<FNavDataConfig>* SupportedAgents = &NavSys->SupportedAgents;
+	const TArray<FNavDataConfig>* SupportedAgents = &NavSys->GetSupportedAgents();
 	const int SupportedAgentsCount = SupportedAgents->Num();
 
 	if (SupportedAgentsCount > 0)
@@ -1883,10 +1883,17 @@ void ARecastNavMesh::ConditionalConstructGenerator()
 {
 	NavDataGenerator.Reset();
 	
-	const bool bRequiresGenerator = SupportsRuntimeGeneration() || !GetWorld()->IsGameWorld();
+	UWorld* World = GetWorld();
+	check(World);
+	const bool bRequiresGenerator = SupportsRuntimeGeneration() || !World->IsGameWorld();
 	if (bRequiresGenerator)
 	{
 		NavDataGenerator.Reset(new FRecastNavMeshGenerator(*this));
+
+		if (World->GetNavigationSystem())
+		{
+			RestrictBuildingToActiveTiles(World->GetNavigationSystem()->IsActiveTilesGenerationEnabled());
+		}
 	}
 }
 
@@ -1937,6 +1944,11 @@ void ARecastNavMesh::UpdateActiveTiles(const TArray<FNavigationInvokerRaw>& Invo
 	}
 
 	FRecastNavMeshGenerator* MyGenerator = static_cast<FRecastNavMeshGenerator*>(GetGenerator());
+	if (MyGenerator == nullptr)
+	{
+		return;
+	}
+
 	const dtNavMeshParams* NavParams = GetRecastNavMeshImpl()->DetourNavMesh->getParams();
 	check(NavParams && MyGenerator);
 	const FRecastBuildConfig& Config = MyGenerator->GetConfig();
