Fixed navigation system not removing generated navmesh when "Active Tiles" are enabled and there are no navigation invokers #UE4

diff --git a/Engine/Source/Runtime/Engine/Private/AI/Navigation/NavigationSystem.cpp b/Engine/Source/Runtime/Engine/Private/AI/Navigation/NavigationSystem.cpp
--- a/Engine/Source/Runtime/Engine/Private/AI/Navigation/NavigationSystem.cpp
+++ b/Engine/Source/Runtime/Engine/Private/AI/Navigation/NavigationSystem.cpp
@@ -3607,13 +3607,14 @@ void UNavigationSystem::UpdateInvokers()
 	const float CurrentTime = World->GetTimeSeconds();
 	if (CurrentTime >= NextInvokersUpdateTime)
 	{
+		TArray<FNavigationInvokerRaw> InvokerLocations;
+
 		if (Invokers.Num() > 0)
 		{
 			QUICK_SCOPE_CYCLE_COUNTER(STAT_NavSys_Clusterize);
 
 			const double StartTime = FPlatformTime::Seconds();
 
-			TArray<FNavigationInvokerRaw> InvokerLocations;
 			InvokerLocations.Reserve(Invokers.Num());
 
 			for (auto ItemIterator = Invokers.CreateIterator(); ItemIterator; ++ItemIterator)
@@ -3644,15 +3645,16 @@ void UNavigationSystem::UpdateInvokers()
 				UE_VLOG_CYLINDER(this, LogNavigation, Log, InvokerData.Location, InvokerData.Location + FVector(0, 0, 20), InvokerData.RadiusMin, FColorList::CadetBlue, TEXT(""));
 			}
 #endif // ENABLE_VISUAL_LOG
+		}
 
-			const double UpdateStartTime = FPlatformTime::Seconds();
-			for (TActorIterator<ARecastNavMesh> It(GetWorld()); It; ++It)
-			{
-				It->UpdateActiveTiles(InvokerLocations);
-			}
-			const double UpdateEndTime = FPlatformTime::Seconds();
-			UE_VLOG(this, LogNavigation, Log, TEXT("Marking tiles to update %fms"), float((UpdateStartTime - UpdateEndTime) / 1000));
+		const double UpdateStartTime = FPlatformTime::Seconds();
+		for (TActorIterator<ARecastNavMesh> It(GetWorld()); It; ++It)
+		{
+			It->UpdateActiveTiles(InvokerLocations);
 		}
+		const double UpdateEndTime = FPlatformTime::Seconds();
+		UE_VLOG(this, LogNavigation, Log, TEXT("Marking tiles to update %fms (%d invokers)"), float((UpdateStartTime - UpdateEndTime) / 1000), InvokerLocations.Num());
+		
 		// once per second
 		NextInvokersUpdateTime = CurrentTime + ActiveTilesUpdateInterval;
 	}
diff --git a/Engine/Source/Runtime/Engine/Private/AI/Navigation/RecastNavMesh.cpp b/Engine/Source/Runtime/Engine/Private/AI/Navigation/RecastNavMesh.cpp
--- a/Engine/Source/Runtime/Engine/Private/AI/Navigation/RecastNavMesh.cpp
+++ b/Engine/Source/Runtime/Engine/Private/AI/Navigation/RecastNavMesh.cpp
@@ -2027,7 +2027,7 @@ void ARecastNavMesh::UpdateActiveTiles(const TArray<FNavigationInvokerRaw>& Invo
 	RemoveTiles(TilesToRemove);
 	RebuildTile(TilesToUpdate);
 
-	if (ActiveTiles.Num() == 0)
+	if (TilesToRemove.Num() > 0 || TilesToUpdate.Num() > 0)
 	{
 		UpdateNavMeshDrawing();
 	}
