Fixed crash when placing BSP.

diff --git a/Engine/Source/Runtime/Engine/Classes/Components/ModelComponent.h b/Engine/Source/Runtime/Engine/Classes/Components/ModelComponent.h
--- a/Engine/Source/Runtime/Engine/Classes/Components/ModelComponent.h
+++ b/Engine/Source/Runtime/Engine/Classes/Components/ModelComponent.h
@@ -18,7 +18,7 @@ class FModelElement;
 UCLASS(MinimalAPI)
 class UModelComponent : public UPrimitiveComponent, public IInterface_CollisionDataProvider
 {
-	GENERATED_UCLASS_BODY()
+	GENERATED_BODY()
 
 private:
 	/** The BSP tree. */
@@ -42,11 +42,14 @@ private:
 
 public:
 
+	/** Default constructor */
+	UModelComponent();
+
 #if WITH_EDITOR
 	/**
 	 * Minimal initialization constructor.
 	 */
-	UModelComponent( const FObjectInitializer& ObjectInitializer,UModel* InModel, uint16 InComponentIndex, uint32 MaskedSurfaceFlags, const TArray<uint16>& InNodes );
+	UModelComponent( UModel* InModel, uint16 InComponentIndex, uint32 MaskedSurfaceFlags, const TArray<uint16>& InNodes );
 #endif // WITH_EDITOR
 
 	/**
diff --git a/Engine/Source/Runtime/Engine/Classes/Components/PrimitiveComponent.h b/Engine/Source/Runtime/Engine/Classes/Components/PrimitiveComponent.h
--- a/Engine/Source/Runtime/Engine/Classes/Components/PrimitiveComponent.h
+++ b/Engine/Source/Runtime/Engine/Classes/Components/PrimitiveComponent.h
@@ -109,8 +109,15 @@ public:
 	/**
 	 * Default UObject constructor.
 	 */
+	UPrimitiveComponent();
+
 	UPrimitiveComponent(const FObjectInitializer& ObjectInitializer);
 
+private:
+	/** Initialize the component to its default settings */
+	void InitializePrimitiveComponentDefaults();
+
+public:
 	// Rendering
 	
 	/**
diff --git a/Engine/Source/Runtime/Engine/Private/Level.cpp b/Engine/Source/Runtime/Engine/Private/Level.cpp
--- a/Engine/Source/Runtime/Engine/Private/Level.cpp
+++ b/Engine/Source/Runtime/Engine/Private/Level.cpp
@@ -823,7 +823,7 @@ void ULevel::CreateModelComponents()
 			Model->Nodes[Nodes[NodeIndex]].ComponentNodeIndex = NodeIndex;
 		}
 
-		UModelComponent* ModelComponent = new(this) UModelComponent(FObjectInitializer(),Model,ModelComponents.Num(),Key.MaskedPolyFlags,Nodes);
+		UModelComponent* ModelComponent = new(this) UModelComponent(Model,ModelComponents.Num(),Key.MaskedPolyFlags,Nodes);
 		ModelComponents.Add(ModelComponent);
 
 		for(int32 NodeIndex = 0;NodeIndex < Nodes.Num();NodeIndex++)
diff --git a/Engine/Source/Runtime/Engine/Private/ModelComponent.cpp b/Engine/Source/Runtime/Engine/Private/ModelComponent.cpp
--- a/Engine/Source/Runtime/Engine/Private/ModelComponent.cpp
+++ b/Engine/Source/Runtime/Engine/Private/ModelComponent.cpp
@@ -51,8 +51,7 @@ FArchive& operator<<(FArchive& Ar,FModelElement& Element)
 	return Ar;
 }
 
-UModelComponent::UModelComponent(const FObjectInitializer& ObjectInitializer)
-	: Super(ObjectInitializer)
+UModelComponent::UModelComponent()
 {
 	CastShadow = true;
 	bUseAsOccluder = true;
@@ -63,8 +62,7 @@ UModelComponent::UModelComponent(const FObjectInitializer& ObjectInitializer)
 }
 
 #if WITH_EDITOR
-UModelComponent::UModelComponent(const FObjectInitializer& ObjectInitializer,UModel* InModel,uint16 InComponentIndex,uint32 MaskedSurfaceFlags,const TArray<uint16>& InNodes):
-	UPrimitiveComponent(ObjectInitializer),
+UModelComponent::UModelComponent(UModel* InModel,uint16 InComponentIndex,uint32 MaskedSurfaceFlags,const TArray<uint16>& InNodes):
 	Model(InModel),
 	ComponentIndex(InComponentIndex),
 	Nodes(InNodes)
diff --git a/Engine/Source/Runtime/Engine/Private/PrimitiveComponent.cpp b/Engine/Source/Runtime/Engine/Private/PrimitiveComponent.cpp
--- a/Engine/Source/Runtime/Engine/Private/PrimitiveComponent.cpp
+++ b/Engine/Source/Runtime/Engine/Private/PrimitiveComponent.cpp
@@ -68,8 +68,18 @@ int32 UPrimitiveComponent::CurrentTag = 2147483647 / 4;
 // 0 is reserved to mean invalid
 uint64 UPrimitiveComponent::NextComponentId = 1;
 
+UPrimitiveComponent::UPrimitiveComponent()
+{
+	InitializePrimitiveComponentDefaults();
+}
+
 UPrimitiveComponent::UPrimitiveComponent(const FObjectInitializer& ObjectInitializer)
 	: Super(ObjectInitializer)
+{
+	InitializePrimitiveComponentDefaults();
+}
+
+void UPrimitiveComponent::InitializePrimitiveComponentDefaults()
 {
 	PostPhysicsComponentTick.bCanEverTick = false;
 	PostPhysicsComponentTick.bStartWithTickEnabled = true;
