Added optional raycast to conform NavMesh walking closer to underlying geometry

Raycast is performed every 0.1 seconds by default & interpolated otherwise. If no geometry is found below the character, we fallback to navmesh height.

diff --git a/Engine/Source/Runtime/Engine/Classes/GameFramework/CharacterMovementComponent.h b/Engine/Source/Runtime/Engine/Classes/GameFramework/CharacterMovementComponent.h
--- a/Engine/Source/Runtime/Engine/Classes/GameFramework/CharacterMovementComponent.h
+++ b/Engine/Source/Runtime/Engine/Classes/GameFramework/CharacterMovementComponent.h
@@ -676,6 +676,10 @@ protected:
 	/** Flag set in pre-physics update to indicate that based movement should be updated post-physics */
 	uint32 bDeferUpdateBasedMovement : 1;
 
+	/** Whether to raycast to underlying geometry to better conform navmesh-walking characters */
+	UPROPERTY(Category = "Character Movement", EditAnywhere, BlueprintReadOnly)
+	uint32 bProjectNavMeshWalking : 1;
+
 	/** forced avoidance velocity, used when AvoidanceLockTimer is > 0 */
 	FVector AvoidanceLockVelocity;
 
@@ -727,6 +731,16 @@ public:
 	/** last known location projected on navmesh, used by NavWalking mode */
 	FNavLocation CachedNavLocation;
 
+	/** Last valid projected hit result from raycast to geometry from navmesh */
+	FHitResult CachedProjectedNavMeshHitResult;
+
+	/** How often we should raycast to project from navmesh to underlying geometry */
+	UPROPERTY(Category = "Character Movement", EditAnywhere, BlueprintReadOnly, meta=(editcondition = "bProjectNavMeshWalking"))
+	float NavMeshProjectionInterval;
+
+	UPROPERTY(Transient)
+	float NavMeshProjectionTimer;
+
 	/** Change avoidance state and registers in RVO manager if needed */
 	UFUNCTION(BlueprintCallable, Category = "Pawn|Components|CharacterMovement", meta = (UnsafeDuringActorConstruction = "true"))
 	void SetAvoidanceEnabled(bool bEnable);
@@ -1111,6 +1125,13 @@ protected:
 	 */
 	virtual bool TryToLeaveNavWalking();
 
+	/** 
+	 * Attempts to better align navmesh walking characters with underlying geometry (sometimes 
+	 * navmesh can differ quite significantly from geometry).
+	 * Updates CachedProjectedNavMeshHitResult, access this for more info about hits.
+	 */
+	void ProjectLocationFromNavMesh(float DeltaSeconds, FVector& InOutLocation);
+
 public:
 
 	/** Called by owning Character upon successful teleport from AActor::TeleportTo(). */
diff --git a/Engine/Source/Runtime/Engine/Private/CharacterMovementComponent.cpp b/Engine/Source/Runtime/Engine/Private/CharacterMovementComponent.cpp
--- a/Engine/Source/Runtime/Engine/Private/CharacterMovementComponent.cpp
+++ b/Engine/Source/Runtime/Engine/Private/CharacterMovementComponent.cpp
@@ -279,6 +279,8 @@ UCharacterMovementComponent::UCharacterMovementComponent(const FObjectInitialize
 
 	OldBaseQuat = FQuat::Identity;
 	OldBaseLocation = FVector::ZeroVector;
+
+	NavMeshProjectionInterval = 0.1f;
 }
 
 void UCharacterMovementComponent::PostLoad()
@@ -3830,8 +3832,10 @@ void UCharacterMovementComponent::PhysNavWalking(float deltaTime, int32 Iteratio
 
 	if (DestNavLocation.NodeRef != INVALID_NAVNODEREF)
 	{
-		const FVector ActorLoc = GetActorLocation();
-		FVector AdjustedDelta = FVector(AdjustedDest.X, AdjustedDest.Y, DestNavLocation.Location.Z) - OldLocation;
+		FVector NewLocation(AdjustedDest.X, AdjustedDest.Y, DestNavLocation.Location.Z);
+		ProjectLocationFromNavMesh(deltaTime, NewLocation);
+
+		FVector AdjustedDelta = NewLocation - OldLocation;
 
 		if (!AdjustedDelta.IsNearlyZero())
 		{
@@ -3852,6 +3856,37 @@ void UCharacterMovementComponent::PhysNavWalking(float deltaTime, int32 Iteratio
 	}
 }
 
+void UCharacterMovementComponent::ProjectLocationFromNavMesh(float DeltaSeconds, FVector& InOutLocation)
+{
+	if(bProjectNavMeshWalking)
+	{
+		NavMeshProjectionTimer -= DeltaSeconds;
+
+		if(NavMeshProjectionTimer < 0.0f)
+		{
+			const FVector RayCastOffset(0.0f, 0.0f, CharacterOwner->GetCapsuleComponent()->GetScaledCapsuleHalfHeight() * 2.0f);
+
+			// raycast to underlying mesh to allow us to more closely follow geometry
+			FCollisionQueryParams Params;
+			FCollisionObjectQueryParams ObjectQueryParams(FCollisionObjectQueryParams::AllStaticObjects);
+			if(GetWorld()->LineTraceSingle(CachedProjectedNavMeshHitResult, InOutLocation + RayCastOffset, InOutLocation - RayCastOffset, Params, ObjectQueryParams))
+			{
+				InOutLocation.Z = CachedProjectedNavMeshHitResult.Location.Z;
+			}
+
+			NavMeshProjectionTimer = NavMeshProjectionInterval;
+		}
+		
+		// interpolate to new height if we had a blocking hit
+		if(CachedProjectedNavMeshHitResult.bBlockingHit)
+		{
+			FPlane Plane(CachedProjectedNavMeshHitResult.Location, CachedProjectedNavMeshHitResult.Normal);
+			FVector ProjectedPoint = FVector::PointPlaneProject(InOutLocation, Plane);
+			InOutLocation.Z = ProjectedPoint.Z;
+		}
+	}
+}
+
 void UCharacterMovementComponent::PhysCustom(float deltaTime, int32 Iterations)
 {
 	if (CharacterOwner)
@@ -3971,6 +4006,8 @@ void UCharacterMovementComponent::SetNavWalkingPhysics(bool bEnable)
 		{
 			UpdatedComponent->SetCollisionResponseToChannel(ECC_WorldStatic, ECR_Ignore);
 			UpdatedComponent->SetCollisionResponseToChannel(ECC_WorldDynamic, ECR_Ignore);
+			CachedProjectedNavMeshHitResult.Reset();
+			NavMeshProjectionTimer = -1.0f;
 		}
 		else
 		{
