Added support to navigation octree for lazy geometry gathering #UE4

Notes:
- still doing the work on game thread
- unloading no longer needed data not implemented yet

diff --git a/Engine/Source/Runtime/Engine/Private/AI/Navigation/NavigationOctree.cpp b/Engine/Source/Runtime/Engine/Private/AI/Navigation/NavigationOctree.cpp
--- a/Engine/Source/Runtime/Engine/Private/AI/Navigation/NavigationOctree.cpp
+++ b/Engine/Source/Runtime/Engine/Private/AI/Navigation/NavigationOctree.cpp
@@ -11,6 +11,7 @@
 //----------------------------------------------------------------------//
 FNavigationOctree::FNavigationOctree(const FVector& Origin, float Radius)
 	: TOctree<FNavigationOctreeElement, FNavigationOctreeSemantics>(Origin, Radius)
+	, bPerformGeometryGatheringLazily(false)
 	, bGatherGeometry(false)
 	, NodesMemory(0)
 {
@@ -23,37 +24,79 @@ FNavigationOctree::~FNavigationOctree()
 	DEC_MEMORY_STAT_BY(STAT_Navigation_CollisionTreeMemory, NodesMemory);
 }
 
+void FNavigationOctree::SetLazyGeometryGatheringEnabled(bool bInEnable)
+{
+	bPerformGeometryGatheringLazily = bInEnable;
+}
+
 void FNavigationOctree::SetNavigableGeometryStoringMode(ENavGeometryStoringMode NavGeometryMode)
 {
 	bGatherGeometry = (NavGeometryMode == FNavigationOctree::StoreNavGeometry);
 }
 
-void FNavigationOctree::AddNode(UObject* ElementOb, INavRelevantInterface* NavElement, const FBox& Bounds, FNavigationOctreeElement& Data)
+void FNavigationOctree::DemandGatheringGeometry(const FNavigationOctreeElement& Element) const
 {
-	Data.Owner = ElementOb;
-	Data.Bounds = Bounds;
-
+	UObject* ElementOb = Element.Owner.Get();
+	if (ElementOb == nullptr || bGatherGeometry == false)
+	{
+		return;
+	}
+		
 	UActorComponent* ActorComp = Cast<UActorComponent>(ElementOb);
-	if (bGatherGeometry && ActorComp)
+	if (ActorComp)
+	{
+		FNavigationOctreeElement* MutableElement = const_cast<FNavigationOctreeElement*>(&Element);
+
+		ComponentExportDelegate.ExecuteIfBound(ActorComp, MutableElement->Data);
+
+		// shrink arrays before counting memory
+		// it will be reallocated when adding to octree and RemoveNode will have different value returned by GetAllocatedSize()
+		MutableElement->Shrink();
+
+		// mark this element as no longer needing geometry gathering
+		MutableElement->Data.bPendingLazyGeometryGathering = false;
+
+		const int32 ElementMemory = MutableElement->Data.GetGeometryAllocatedSize();
+		const_cast<FNavigationOctree*>(this)->NodesMemory += ElementMemory;
+		INC_MEMORY_STAT_BY(STAT_Navigation_CollisionTreeMemory, ElementMemory);
+	}
+}
+
+void FNavigationOctree::AddNode(UObject* ElementOb, INavRelevantInterface* NavElement, const FBox& Bounds, FNavigationOctreeElement& Element)
+{
+	// we assume NavElement is ElementOb already cast
+	//check(ElementOb == NavElement);
+	Element.Owner = ElementOb;
+	Element.Bounds = Bounds;
+
+	if (bGatherGeometry)
 	{
-		ComponentExportDelegate.ExecuteIfBound(ActorComp, Data.Data);
+		UActorComponent* ActorComp = Cast<UActorComponent>(ElementOb);
+		if (ActorComp && (bPerformGeometryGatheringLazily == false))
+		{
+			ComponentExportDelegate.ExecuteIfBound(ActorComp, Element.Data);
+		}
+		else
+		{
+			Element.Data.bPendingLazyGeometryGathering = true;
+		}
 	}
 
 	if (NavElement)
 	{
 		SCOPE_CYCLE_COUNTER(STAT_Navigation_GatheringNavigationModifiersSync);
-		NavElement->GetNavigationData(Data.Data);
+		NavElement->GetNavigationData(Element.Data);
 	}
 
 	// shrink arrays before counting memory
 	// it will be reallocated when adding to octree and RemoveNode will have different value returned by GetAllocatedSize()
-	Data.Shrink();
+	Element.Shrink();
 
-	const int32 ElementMemory = Data.GetAllocatedSize();
+	const int32 ElementMemory = Element.GetAllocatedSize();
 	NodesMemory += ElementMemory;
 	INC_MEMORY_STAT_BY(STAT_Navigation_CollisionTreeMemory, ElementMemory);
 
-	AddElement(Data);
+	AddElement(Element);
 }
 
 void FNavigationOctree::AppendToNode(const FOctreeElementId& Id, INavRelevantInterface* NavElement, const FBox& Bounds, FNavigationOctreeElement& Data)
diff --git a/Engine/Source/Runtime/Engine/Private/AI/Navigation/NavigationSystem.cpp b/Engine/Source/Runtime/Engine/Private/AI/Navigation/NavigationSystem.cpp
--- a/Engine/Source/Runtime/Engine/Private/AI/Navigation/NavigationSystem.cpp
+++ b/Engine/Source/Runtime/Engine/Private/AI/Navigation/NavigationSystem.cpp
@@ -422,6 +422,7 @@ bool UNavigationSystem::ConditionalPopulateNavOctree()
 	if (bSupportRebuilding)
 	{
 		NavOctree = new FNavigationOctree(FVector(0,0,0), 64000);
+		NavOctree->SetLazyGeometryGatheringEnabled(bGenerateNavigationOnlyAroundNavigationInvokers);
 		
 		const ERuntimeGenerationType RuntimeGenerationType = GetRuntimeGenerationType();
 		const bool bStoreNavGeometry = (RuntimeGenerationType == ERuntimeGenerationType::Dynamic);
@@ -482,6 +483,11 @@ void UNavigationSystem::PostEditChangeProperty(FPropertyChangedEvent& PropertyCh
 		}
 		else if (PropName == NAME_EnableActiveTiles)
 		{
+			if (NavOctree)
+			{
+				NavOctree->SetLazyGeometryGatheringEnabled(bGenerateNavigationOnlyAroundNavigationInvokers);
+			}
+
 			for (auto NavData : NavDataSet)
 			{
 				if (NavData)
@@ -3637,7 +3643,7 @@ void UNavigationSystem::UpdateInvokers()
 
 #if ENABLE_VISUAL_LOG
 			const double CachingFinishTime = FPlatformTime::Seconds();
-			UE_VLOG(this, LogNavigation, Log, TEXT("Caching time %fms"), float((CachingFinishTime - StartTime) / 1000));
+			UE_VLOG(this, LogNavigation, Log, TEXT("Caching time %fms"), (CachingFinishTime - StartTime) * 1000);
 
 			for (const auto& InvokerData : InvokerLocations)
 			{
@@ -3653,7 +3659,7 @@ void UNavigationSystem::UpdateInvokers()
 			It->UpdateActiveTiles(InvokerLocations);
 		}
 		const double UpdateEndTime = FPlatformTime::Seconds();
-		UE_VLOG(this, LogNavigation, Log, TEXT("Marking tiles to update %fms (%d invokers)"), float((UpdateStartTime - UpdateEndTime) / 1000), InvokerLocations.Num());
+		UE_VLOG(this, LogNavigation, Log, TEXT("Marking tiles to update %fms (%d invokers)"), (UpdateEndTime - UpdateStartTime) * 1000, InvokerLocations.Num());
 		
 		// once per second
 		NextInvokersUpdateTime = CurrentTime + ActiveTilesUpdateInterval;
diff --git a/Engine/Source/Runtime/Engine/Private/AI/Navigation/RecastNavMeshGenerator.cpp b/Engine/Source/Runtime/Engine/Private/AI/Navigation/RecastNavMeshGenerator.cpp
--- a/Engine/Source/Runtime/Engine/Private/AI/Navigation/RecastNavMeshGenerator.cpp
+++ b/Engine/Source/Runtime/Engine/Private/AI/Navigation/RecastNavMeshGenerator.cpp
@@ -1536,6 +1536,7 @@ void FRecastTileGenerator::GatherGeometry(const FRecastNavMeshGenerator& ParentG
 {
 	const UNavigationSystem*	NavSys = UNavigationSystem::GetCurrent(ParentGenerator.GetWorld());
 	const FNavigationOctree*	NavOctree = NavSys ? NavSys->GetNavOctree() : nullptr;
+	check(NavOctree);
 	const FNavDataConfig&		NavDataConfig = ParentGenerator.GetOwner()->GetConfig();
 
 	for (FNavigationOctree::TConstElementBoxIterator<FNavigationOctree::DefaultStackAllocator> It(*NavOctree, ParentGenerator.GrowBoundingBox(TileBB, /*bIncludeAgentHeight*/ false));
@@ -1546,6 +1547,11 @@ void FRecastTileGenerator::GatherGeometry(const FRecastNavMeshGenerator& ParentG
 		const bool bShouldUse = Element.ShouldUseGeometry(NavDataConfig);
 		if (bShouldUse)
 		{
+			if (Element.Data.IsPendingLazyGeometryGathering())
+			{
+				NavOctree->DemandGatheringGeometry(Element);
+			}
+
 			const bool bExportGeometry = bGeometryChanged && Element.Data.HasGeometry();
 			if (bExportGeometry)
 			{
diff --git a/Engine/Source/Runtime/Engine/Public/AI/NavigationOctree.h b/Engine/Source/Runtime/Engine/Public/AI/NavigationOctree.h
--- a/Engine/Source/Runtime/Engine/Public/AI/NavigationOctree.h
+++ b/Engine/Source/Runtime/Engine/Public/AI/NavigationOctree.h
@@ -45,13 +45,23 @@ struct ENGINE_API FNavigationRelevantData
 	/** additional modifiers: areas and external links */
 	FCompositeNavModifier Modifiers;
 
+	/** get set to true when lazy navigation exporting is enabled and this navigation data has "potential" of 
+	 *	containing geometry data. First access will result in gathering the data and setting this flag back to false.
+	 *	Mind that this flag can go back to 'true' if related data gets cleared out. */
+	uint32 bPendingLazyGeometryGathering : 1;
+
 	// Gathers per instance data for navigation geometry in a specified area box
 	FNavDataPerInstanceTransformDelegate NavDataPerInstanceTransformDelegate;
 
+	FNavigationRelevantData() : bPendingLazyGeometryGathering(false)
+	{}
+
 	FORCEINLINE bool HasGeometry() const { return VoxelData.Num() || CollisionData.Num(); }
 	FORCEINLINE bool HasModifiers() const { return !Modifiers.IsEmpty(); }
+	FORCEINLINE bool IsPendingLazyGeometryGathering() const { return bPendingLazyGeometryGathering; }
 	FORCEINLINE bool IsEmpty() const { return !HasGeometry() && !HasModifiers(); }
 	FORCEINLINE uint32 GetAllocatedSize() const { return CollisionData.GetAllocatedSize() + VoxelData.GetAllocatedSize() + Modifiers.GetAllocatedSize(); }
+	FORCEINLINE uint32 GetGeometryAllocatedSize() const { return CollisionData.GetAllocatedSize() + VoxelData.GetAllocatedSize(); }
 	FORCEINLINE int32 GetDirtyFlag() const
 	{
 		return (HasGeometry() ? ENavigationDirtyFlag::Geometry : 0) |
@@ -167,8 +177,13 @@ public:
 		return bGatherGeometry ? StoreNavGeometry : SkipNavGeometry;
 	}
 
-protected:
+	void SetLazyGeometryGatheringEnabled(bool bInEnable);
+
+	// @hack! TO BE FIXED
+	void DemandGatheringGeometry(const FNavigationOctreeElement& Element) const;
 
+protected:
+	uint32 bPerformGeometryGatheringLazily : 1;
 	uint32 bGatherGeometry : 1;
 	uint32 NodesMemory;
 };
