Track unreliable multicast RPCs in the network profiler.
Since these RPCs usually share a content block header with a property update, the reported header size for these RPCs will be smaller than that of other types of RPCs.

diff --git a/Engine/Source/Runtime/Engine/Private/NetworkDriver.cpp b/Engine/Source/Runtime/Engine/Private/NetworkDriver.cpp
--- a/Engine/Source/Runtime/Engine/Private/NetworkDriver.cpp
+++ b/Engine/Source/Runtime/Engine/Private/NetworkDriver.cpp
@@ -962,36 +962,40 @@ void UNetDriver::InternalProcessRemoteFunction
 	{
 		UE_LOG(LogNetTraffic, Log, TEXT("RPC bunch on closing channel") );
 	}
-	else if (QueueBunch)
-	{
-		// Unreliable multicast functions are queued and sent out during property replication
-		if (LogAsWarning)
-		{
-			UE_LOG(LogNetTraffic, Warning,	TEXT("      Queing unreliable multicast RPC: %s::%s [%.1f bytes]"), *Actor->GetName(), *Function->GetName(), Bunch.GetNumBits() / 8.f );
-		}
-		else
-		{
-			UE_LOG(LogNetTraffic, Log,		TEXT("      Queing unreliable multicast RPC: %s::%s [%.1f bytes]"), *Actor->GetName(), *Function->GetName(), Bunch.GetNumBits() / 8.f );
-		}
-
-		Ch->QueueRemoteFunctionBunch(TargetObj, Function, Bunch);
-	}
 	else
 	{
-		if (LogAsWarning)
+		// Make sure we're tracking all the bits in the bunch
+		check(Bunch.GetNumBits() == HeaderBits + ParameterBits + FooterBits);
+
+		if (QueueBunch)
 		{
-			UE_LOG(LogNetTraffic, Warning,	TEXT("      Sent RPC: %s::%s [%.1f bytes]"), *Actor->GetName(), *Function->GetName(), Bunch.GetNumBits() / 8.f );
+			// Unreliable multicast functions are queued and sent out during property replication
+			if (LogAsWarning)
+			{
+				UE_LOG(LogNetTraffic, Warning,	TEXT("      Queing unreliable multicast RPC: %s::%s [%.1f bytes]"), *Actor->GetName(), *Function->GetName(), Bunch.GetNumBits() / 8.f );
+			}
+			else
+			{
+				UE_LOG(LogNetTraffic, Log,		TEXT("      Queing unreliable multicast RPC: %s::%s [%.1f bytes]"), *Actor->GetName(), *Function->GetName(), Bunch.GetNumBits() / 8.f );
+			}
+
+			NETWORK_PROFILER(GNetworkProfiler.TrackSendRPC(Actor, Function, HeaderBits, ParameterBits, FooterBits));
+			Ch->QueueRemoteFunctionBunch(TargetObj, Function, Bunch);
 		}
 		else
 		{
-			UE_LOG(LogNetTraffic, Log,		TEXT("      Sent RPC: %s::%s [%.1f bytes]"), *Actor->GetName(), *Function->GetName(), Bunch.GetNumBits() / 8.f );
-		}
-
-		// Make sure we're tracking all the bits in the bunch
-		check(Bunch.GetNumBits() == HeaderBits + ParameterBits + FooterBits);
+			if (LogAsWarning)
+			{
+				UE_LOG(LogNetTraffic, Warning,	TEXT("      Sent RPC: %s::%s [%.1f bytes]"), *Actor->GetName(), *Function->GetName(), Bunch.GetNumBits() / 8.f );
+			}
+			else
+			{
+				UE_LOG(LogNetTraffic, Log,		TEXT("      Sent RPC: %s::%s [%.1f bytes]"), *Actor->GetName(), *Function->GetName(), Bunch.GetNumBits() / 8.f );
+			}
 
-		NETWORK_PROFILER(GNetworkProfiler.TrackSendRPC(Actor, Function, HeaderBits, ParameterBits, FooterBits));
-		Ch->SendBunch( &Bunch, 1 );
+			NETWORK_PROFILER(GNetworkProfiler.TrackSendRPC(Actor, Function, HeaderBits, ParameterBits, FooterBits));
+			Ch->SendBunch( &Bunch, 1 );
+		}
 	}
 }
 
