Expanded foliage statistics, allowing box count query and a transform list query (code only)

diff --git a/Engine/Source/Runtime/Engine/Classes/Components/HierarchicalInstancedStaticMeshComponent.h b/Engine/Source/Runtime/Engine/Classes/Components/HierarchicalInstancedStaticMeshComponent.h
--- a/Engine/Source/Runtime/Engine/Classes/Components/HierarchicalInstancedStaticMeshComponent.h
+++ b/Engine/Source/Runtime/Engine/Classes/Components/HierarchicalInstancedStaticMeshComponent.h
@@ -102,6 +102,10 @@ public:
 
 	/** Get the number of instances that overlap a given sphere */
 	int32 GetOverlappingSphereCount(const FSphere& Sphere) const;
+	/** Get the number of instances that overlap a given box */
+	int32 GetOverlappingBoxCount(const FBox& Box) const;
+	/** Get the transforms of instances inside the provided box */
+	void GetOverlappingBoxTransforms(const FBox& Box, TArray<FTransform>& OutTransforms) const;
 
 	virtual bool ShouldCreatePhysicsState() const override;
 
diff --git a/Engine/Source/Runtime/Engine/Private/HierarchicalInstancedStaticMesh.cpp b/Engine/Source/Runtime/Engine/Private/HierarchicalInstancedStaticMesh.cpp
--- a/Engine/Source/Runtime/Engine/Private/HierarchicalInstancedStaticMesh.cpp
+++ b/Engine/Source/Runtime/Engine/Private/HierarchicalInstancedStaticMesh.cpp
@@ -1970,6 +1970,46 @@ int32 UHierarchicalInstancedStaticMeshComponent::GetOverlappingSphereCount(const
 	return Count;
 }
 
+int32 UHierarchicalInstancedStaticMeshComponent::GetOverlappingBoxCount(const FBox& Box) const
+{
+	TArray<FTransform> Transforms;
+	GatherInstanceTransformsInArea(*this, Box, 0, Transforms);
+	
+	int32 Count = 0;
+	const FBoxSphereBounds Bounds = StaticMesh->GetBounds();
+	for(FTransform& T : Transforms)
+	{
+		const FVector Centre = T.GetLocation();
+		const FBox OtherBox(FVector(Centre - Bounds.BoxExtent), FVector(Centre + Bounds.BoxExtent));
+
+		if(Box.Intersect(OtherBox))
+		{
+			Count++;
+		}
+	}
+
+	return Count;
+}
+
+void UHierarchicalInstancedStaticMeshComponent::GetOverlappingBoxTransforms(const FBox& Box, TArray<FTransform>& OutTransforms) const
+{
+	GatherInstanceTransformsInArea(*this, Box, 0, OutTransforms);
+
+	const FBoxSphereBounds Bounds = StaticMesh->GetBounds();
+	int32 NumTransforms = OutTransforms.Num();
+	for(int32 Idx = NumTransforms - 1 ; Idx >= 0 ; --Idx)
+	{
+		FTransform& TM = OutTransforms[Idx];
+		const FVector Centre = TM.GetLocation();
+		const FBox OtherBox(FVector(Centre - Bounds.BoxExtent), FVector(Centre + Bounds.BoxExtent));
+
+		if(!Box.Intersect(OtherBox))
+		{
+			OutTransforms.RemoveAt(Idx);
+		}
+	}
+}
+
 void UHierarchicalInstancedStaticMeshComponent::GetNavigationPerInstanceTransforms(const FBox& AreaBox, TArray<FTransform>& InstanceData) const
 {
 	if (IsTreeFullyBuilt())
diff --git a/Engine/Source/Runtime/Foliage/Private/FoliageStatistics.cpp b/Engine/Source/Runtime/Foliage/Private/FoliageStatistics.cpp
--- a/Engine/Source/Runtime/Foliage/Private/FoliageStatistics.cpp
+++ b/Engine/Source/Runtime/Foliage/Private/FoliageStatistics.cpp
@@ -26,4 +26,35 @@ int32 UFoliageStatistics::FoliageOverlappingSphereCount(UObject* WorldContextObj
 	}
 
 	return Count;
-}
\ No newline at end of file
+}
+
+int32 UFoliageStatistics::FoliageOverlappingBoxCount(UObject* WorldContextObject, const UStaticMesh* StaticMesh, FBox Box)
+{
+	UWorld* World = GEngine->GetWorldFromContextObject(WorldContextObject);
+
+	int32 Count = 0;
+	for(FConstLevelIterator LvlItr = World->GetLevelIterator() ; LvlItr ; ++LvlItr)
+	{
+		ULevel* Level = *LvlItr;
+		AInstancedFoliageActor* IFA = AInstancedFoliageActor::GetInstancedFoliageActorForLevel(Level);
+		const UFoliageType* Type = IFA->GetSettingsForMesh(StaticMesh);
+
+		Count += IFA->GetOverlappingBoxCount(Type, Box);
+	}
+
+	return Count;
+}
+
+void UFoliageStatistics::FoliageOverlappingBoxTransforms(UObject* WorldContextObject, const UStaticMesh* StaticMesh, FBox Box, TArray<FTransform>& OutTransforms)
+{
+	UWorld* World = GEngine->GetWorldFromContextObject(WorldContextObject);
+
+	for(FConstLevelIterator LvlItr = World->GetLevelIterator() ; LvlItr ; ++LvlItr)
+	{
+		ULevel* Level = *LvlItr;
+		AInstancedFoliageActor* IFA = AInstancedFoliageActor::GetInstancedFoliageActorForLevel(Level);
+		const UFoliageType* Type = IFA->GetSettingsForMesh(StaticMesh);
+
+		IFA->GetOverlappingBoxTransforms(Type, Box, OutTransforms);
+	}
+}
diff --git a/Engine/Source/Runtime/Foliage/Private/FoliageStatistics.h b/Engine/Source/Runtime/Foliage/Private/FoliageStatistics.h
--- a/Engine/Source/Runtime/Foliage/Private/FoliageStatistics.h
+++ b/Engine/Source/Runtime/Foliage/Private/FoliageStatistics.h
@@ -20,5 +20,21 @@ class FOLIAGE_API UFoliageStatistics : public UBlueprintFunctionLibrary
 	*/
 	UFUNCTION(BlueprintCallable, Category = "Foliage", meta = (WorldContext = "WorldContextObject", UnsafeDuringActorConstruction = "true"))
 	static int32 FoliageOverlappingSphereCount(UObject* WorldContextObject, const UStaticMesh* StaticMesh, FVector CenterPosition, float Radius);
+
+	/** 
+	 *	Gets the number of instances overlapping a provided box
+	 *	@param StaticMesh Mesh to count
+	 *	@param Box Box to overlap
+	 */
+	UFUNCTION(BlueprintCallable, Category = "Foliage", meta = (WorldContext = "WorldContextObject", UnsafeDuringActorConstruction = "true"))
+	static int32 FoliageOverlappingBoxCount(UObject* WorldContextObject, const UStaticMesh* StaticMesh, FBox Box);
+
+	/** 
+	 *  Get the transform of every instance overlapping the provided FBox
+	 *	@param StaticMesh Mesh to get instances of
+	 *	@param Box Box to use for overlap
+	 *	@param OutTransforms Array to populate with transforms
+	 */
+	static void FoliageOverlappingBoxTransforms(UObject* WorldContextObject, const UStaticMesh* StaticMesh, FBox Box, TArray<FTransform>& OutTransforms);
 };
 
diff --git a/Engine/Source/Runtime/Foliage/Private/InstancedFoliage.cpp b/Engine/Source/Runtime/Foliage/Private/InstancedFoliage.cpp
--- a/Engine/Source/Runtime/Foliage/Private/InstancedFoliage.cpp
+++ b/Engine/Source/Runtime/Foliage/Private/InstancedFoliage.cpp
@@ -950,6 +950,32 @@ int32 AInstancedFoliageActor::GetOverlappingSphereCount(const UFoliageType* Foli
 }
 
 
+int32 AInstancedFoliageActor::GetOverlappingBoxCount(const UFoliageType* FoliageType, const FBox& Box) const
+{
+	if(const FFoliageMeshInfo* MeshInfo = FindMesh(FoliageType))
+	{
+		if(MeshInfo->Component)
+		{
+			return MeshInfo->Component->GetOverlappingBoxCount(Box);
+		}
+	}
+
+	return 0;
+}
+
+
+void AInstancedFoliageActor::GetOverlappingBoxTransforms(const UFoliageType* FoliageType, const FBox& Box, TArray<FTransform>& OutTransforms) const
+{
+	if(const FFoliageMeshInfo* MeshInfo = FindMesh(FoliageType))
+	{
+		if(MeshInfo->Component)
+		{
+			MeshInfo->Component->GetOverlappingBoxTransforms(Box, OutTransforms);
+		}
+	}
+}
+
+
 UFoliageType* AInstancedFoliageActor::GetSettingsForMesh(const UStaticMesh* InMesh, FFoliageMeshInfo** OutMeshInfo)
 {
 	UFoliageType* Type = nullptr;
diff --git a/Engine/Source/Runtime/Foliage/Public/InstancedFoliageActor.h b/Engine/Source/Runtime/Foliage/Public/InstancedFoliageActor.h
--- a/Engine/Source/Runtime/Foliage/Public/InstancedFoliageActor.h
+++ b/Engine/Source/Runtime/Foliage/Public/InstancedFoliageActor.h
@@ -59,6 +59,10 @@ public:
 
 	// Finds the number of instances overlapping with the sphere. 
 	FOLIAGE_API int32 GetOverlappingSphereCount(const UFoliageType* FoliageType, const FSphere& Sphere) const;
+	// Finds the number of instances overlapping with the box. 
+	FOLIAGE_API int32 GetOverlappingBoxCount(const UFoliageType* FoliageType, const FBox& Box) const;
+	// Finds all instances in the provided box and get their transforms
+	FOLIAGE_API void GetOverlappingBoxTransforms(const UFoliageType* FoliageType, const FBox& Box, TArray<FTransform>& OutTransforms) const;
 
 	// Finds a mesh entry
 	FOLIAGE_API FFoliageMeshInfo* FindMesh(const UFoliageType* InType);
