UE4 - Rename setting for velocity rendering on base pass; flip logic to opting-out materials

diff --git a/Engine/Source/Runtime/Engine/Classes/Engine/RendererSettings.h b/Engine/Source/Runtime/Engine/Classes/Engine/RendererSettings.h
--- a/Engine/Source/Runtime/Engine/Classes/Engine/RendererSettings.h
+++ b/Engine/Source/Runtime/Engine/Classes/Engine/RendererSettings.h
@@ -235,8 +235,8 @@ class ENGINE_API URendererSettings
 	TEnumAsByte<EClearSceneOptions::Type> ClearSceneMethod;
 
 	UPROPERTY(config, EditAnywhere, Category=Optimizations, meta=(
-		ConsoleVariable="r.BasePassOutputsVelocity", DisplayName="WPO materials output velocity",
-		ToolTip="Enables materials with time-based World Position Offset attributes to render their velocities for motion blur in the base rendering pass."))
+		ConsoleVariable="r.BasePassOutputsVelocity", DisplayName="Accurate velocities from Vertex Deformation",
+		ToolTip="Enables materials with time-based World Position Offset and/or World Displacement to output accurate velocities. This incurs a performance cost. If this is disabled, those materials will not output velocities."))
 	uint32 bBasePassOutputsVelocity:1;
 
 	UPROPERTY(config, EditAnywhere, Category=Editor, meta=(
diff --git a/Engine/Source/Runtime/Engine/Classes/Materials/Material.h b/Engine/Source/Runtime/Engine/Classes/Materials/Material.h
--- a/Engine/Source/Runtime/Engine/Classes/Materials/Material.h
+++ b/Engine/Source/Runtime/Engine/Classes/Materials/Material.h
@@ -615,8 +615,8 @@ public:
 	uint32 Wireframe:1;
 
 	/** Skips outputting velocity during the base pass. */
-	UPROPERTY(EditAnywhere, Category=Material, AdvancedDisplay)
-	uint32 bSkipVelocityOnBasePass:1;
+	UPROPERTY(EditAnywhere, Category=Material, AdvancedDisplay, meta=(DisplayName = "Support accurate velocities from Vertex Deformation"))
+	uint32 bOutputVelocityOnBasePass:1;
 
 	UPROPERTY()
 	int32 EditorX;
diff --git a/Engine/Source/Runtime/Engine/Classes/Materials/MaterialInterface.h b/Engine/Source/Runtime/Engine/Classes/Materials/MaterialInterface.h
--- a/Engine/Source/Runtime/Engine/Classes/Materials/MaterialInterface.h
+++ b/Engine/Source/Runtime/Engine/Classes/Materials/MaterialInterface.h
@@ -58,7 +58,7 @@ struct ENGINE_API FMaterialRelevance
 	uint32 bSubsurfaceProfile : 1;
 
 	UPROPERTY()
-	uint32 bSkipVelocityRendering : 1;
+	uint32 bOutputsVelocityInBasePass : 1;
 
 	/** Default constructor. */
 	FMaterialRelevance()
@@ -69,7 +69,7 @@ struct ENGINE_API FMaterialRelevance
 		, bNormalTranslucency(false)
 		, bDisableDepthTest(false)
 		, bSubsurfaceProfile(false)
-		, bSkipVelocityRendering(false)
+		, bOutputsVelocityInBasePass(true)
 	{}
 
 	/** Bitwise OR operator.  Sets any relevance bits which are present in either FMaterialRelevance. */
@@ -82,7 +82,7 @@ struct ENGINE_API FMaterialRelevance
 		bNormalTranslucency |= B.bNormalTranslucency;
 		bDisableDepthTest |= B.bDisableDepthTest;
 		bSubsurfaceProfile |= B.bSubsurfaceProfile;
-		bSkipVelocityRendering |= B.bSkipVelocityRendering;
+		bOutputsVelocityInBasePass |= B.bOutputsVelocityInBasePass;
 		return *this;
 	}
 
diff --git a/Engine/Source/Runtime/Engine/Private/Materials/Material.cpp b/Engine/Source/Runtime/Engine/Private/Materials/Material.cpp
--- a/Engine/Source/Runtime/Engine/Private/Materials/Material.cpp
+++ b/Engine/Source/Runtime/Engine/Private/Materials/Material.cpp
@@ -534,6 +534,7 @@ UMaterial::UMaterial(const FObjectInitializer& ObjectInitializer)
 	D3D11TessellationMode = MTM_NoTessellation;
 	bEnableCrackFreeDisplacement = false;
 	bEnableAdaptiveTessellation = true;
+	bOutputVelocityOnBasePass = true;
 	bEnableSeparateTranslucency = true;
 	bEnableResponsiveAA = false;
 	bTangentSpaceNormal = true;
diff --git a/Engine/Source/Runtime/Engine/Private/Materials/MaterialInterface.cpp b/Engine/Source/Runtime/Engine/Private/Materials/MaterialInterface.cpp
--- a/Engine/Source/Runtime/Engine/Private/Materials/MaterialInterface.cpp
+++ b/Engine/Source/Runtime/Engine/Private/Materials/MaterialInterface.cpp
@@ -57,7 +57,7 @@ FMaterialRelevance UMaterialInterface::GetRelevance_Internal(const UMaterial* Ma
 		MaterialRelevance.bNormalTranslucency = bIsTranslucent && !Material->bEnableSeparateTranslucency;
 		MaterialRelevance.bDisableDepthTest = bIsTranslucent && Material->bDisableDepthTest;
 		MaterialRelevance.bSubsurfaceProfile = (Material->MaterialDomain == MD_Surface) && !bIsTranslucent && (ShadingModel == MSM_SubsurfaceProfile);
-		MaterialRelevance.bSkipVelocityRendering = Material->bSkipVelocityOnBasePass;
+		MaterialRelevance.bOutputsVelocityInBasePass = Material->bOutputVelocityOnBasePass;
 
 		return MaterialRelevance;
 	}
diff --git a/Engine/Source/Runtime/Engine/Private/Materials/MaterialShared.cpp b/Engine/Source/Runtime/Engine/Private/Materials/MaterialShared.cpp
--- a/Engine/Source/Runtime/Engine/Private/Materials/MaterialShared.cpp
+++ b/Engine/Source/Runtime/Engine/Private/Materials/MaterialShared.cpp
@@ -790,9 +790,9 @@ bool FMaterialResource::IsFullyRough() const
 	return Material->bFullyRough;
 }
 
-bool FMaterialResource::HasSkipVelocityOnBasePass() const
+bool FMaterialResource::OutputsVelocityOnBasePass() const
 {
-	return Material->bSkipVelocityOnBasePass;
+	return Material->bOutputVelocityOnBasePass;
 }
 
 bool FMaterialResource::IsNonmetal() const
diff --git a/Engine/Source/Runtime/Engine/Public/MaterialShared.h b/Engine/Source/Runtime/Engine/Public/MaterialShared.h
--- a/Engine/Source/Runtime/Engine/Public/MaterialShared.h
+++ b/Engine/Source/Runtime/Engine/Public/MaterialShared.h
@@ -922,7 +922,7 @@ public:
 	virtual bool IsCrackFreeDisplacementEnabled() const { return false; }
 	virtual bool IsAdaptiveTessellationEnabled() const { return false; }
 	virtual bool IsFullyRough() const { return false; }
-	virtual bool HasSkipVelocityOnBasePass() const { return false; }
+	virtual bool OutputsVelocityOnBasePass() const { return true; }
 	virtual bool IsNonmetal() const { return false; }
 	virtual bool UseLmDirectionality() const { return true; }
 	virtual bool IsMasked() const = 0;
@@ -1505,7 +1505,7 @@ public:
 	ENGINE_API virtual bool IsCrackFreeDisplacementEnabled() const;
 	ENGINE_API virtual bool IsAdaptiveTessellationEnabled() const;
 	ENGINE_API virtual bool IsFullyRough() const;
-	ENGINE_API virtual bool HasSkipVelocityOnBasePass() const;
+	ENGINE_API virtual bool OutputsVelocityOnBasePass() const;
 	ENGINE_API virtual bool IsNonmetal() const;
 	ENGINE_API virtual bool UseLmDirectionality() const;
 	ENGINE_API virtual enum EBlendMode GetBlendMode() const;
diff --git a/Engine/Source/Runtime/Renderer/Private/VelocityRendering.cpp b/Engine/Source/Runtime/Renderer/Private/VelocityRendering.cpp
--- a/Engine/Source/Runtime/Renderer/Private/VelocityRendering.cpp
+++ b/Engine/Source/Runtime/Renderer/Private/VelocityRendering.cpp
@@ -369,7 +369,7 @@ bool FVelocityDrawingPolicy::HasVelocityOnBasePass(const FViewInfo& View,const F
 
 	//@todo-rco: Optimize finding WPO!
 	auto* Material = Mesh.MaterialRenderProxy->GetMaterial(Scene->GetFeatureLevel());
-	return Material->MaterialModifiesMeshPosition_RenderThread() && !Material->HasSkipVelocityOnBasePass();
+	return Material->MaterialModifiesMeshPosition_RenderThread() && Material->OutputsVelocityOnBasePass();
 }
 
 FBoundShaderStateInput FVelocityDrawingPolicy::GetBoundShaderStateInput(ERHIFeatureLevel::Type InFeatureLevel)
