Fix crash when duplicating the current landscape level and then opening it

diff --git a/Engine/Source/Runtime/Engine/Private/World.cpp b/Engine/Source/Runtime/Engine/Private/World.cpp
--- a/Engine/Source/Runtime/Engine/Private/World.cpp
+++ b/Engine/Source/Runtime/Engine/Private/World.cpp
@@ -418,6 +418,11 @@ void UWorld::PostDuplicate(bool bDuplicateForPIE)
 		{
 			FArchiveReplaceObjectRef<UObject> ReplaceAr(Obj, ReplacementMap, bNullPrivateRefs, bIgnoreOuterRef, bIgnoreArchetypeRef);
 		}
+		// PostEditChange is required for some objects to react to the change, e.g. update render-thread proxies
+		for (auto* Obj : ObjectsToFixReferences)
+		{
+			Obj->PostEditChange();
+		}
 	}
 #endif // WITH_EDITOR
 }
