Find-in-Blueprint data is no longer gathered during Blueprint regeneration, instead it is delayed until PostLoad. The gathering was forcing level actor nodes to cache their title before the actor was fully loaded, leading to incorrect node titles being displayed.

diff --git a/Engine/Source/Editor/UnrealEd/Private/Kismet2/BlueprintEditorUtils.cpp b/Engine/Source/Editor/UnrealEd/Private/Kismet2/BlueprintEditorUtils.cpp
--- a/Engine/Source/Editor/UnrealEd/Private/Kismet2/BlueprintEditorUtils.cpp
+++ b/Engine/Source/Editor/UnrealEd/Private/Kismet2/BlueprintEditorUtils.cpp
@@ -1235,14 +1235,6 @@ UClass* FBlueprintEditorUtils::RegenerateBlueprintClass(UBlueprint* Blueprint, U
 			Blueprint->Status = BS_UpToDate;
 		}
 
-#if WITH_EDITOR
-		// Do not want to run this code without the editor present nor when running commandlets.
-		if(GEditor && GIsEditor && !IsRunningCommandlet())
-		{
-			FFindInBlueprintSearchManager::Get().AddOrUpdateBlueprintSearchMetadata(Blueprint);
-		}
-#endif
-
 		if (ReplaceBlueprintWithClass)
 		{
 			FEditoronlyBlueprintHelper::ShouldBeFixed(Blueprint, true);
diff --git a/Engine/Source/Runtime/Engine/Private/Blueprint.cpp b/Engine/Source/Runtime/Engine/Private/Blueprint.cpp
--- a/Engine/Source/Runtime/Engine/Private/Blueprint.cpp
+++ b/Engine/Source/Runtime/Engine/Private/Blueprint.cpp
@@ -538,6 +538,16 @@ void UBlueprint::PostLoad()
 	}
 
 	FStructureEditorUtils::RemoveInvalidStructureMemberVariableFromBlueprint(this);
+
+
+#if WITH_EDITOR
+	// Do not want to run this code without the editor present nor when running commandlets.
+	if(GEditor && GIsEditor && !IsRunningCommandlet())
+	{
+		// Gathers Find-in-Blueprint data, makes sure that it is fresh and ready, especially if the asset did not have any available.
+		FFindInBlueprintSearchManager::Get().AddOrUpdateBlueprintSearchMetadata(this);
+	}
+#endif
 }
 
 void UBlueprint::DebuggingWorldRegistrationHelper(UObject* ObjectProvidingWorld, UObject* ValueToRegister)
