Stat for anim game thread time

diff --git a/Engine/Source/Runtime/Engine/Private/Animation/AnimInstance.cpp b/Engine/Source/Runtime/Engine/Private/Animation/AnimInstance.cpp
--- a/Engine/Source/Runtime/Engine/Private/Animation/AnimInstance.cpp
+++ b/Engine/Source/Runtime/Engine/Private/Animation/AnimInstance.cpp
@@ -53,6 +53,8 @@ DEFINE_STAT(STAT_AnimStateMachineEvaluate);
 DEFINE_STAT(STAT_SkinPerPolyVertices);
 DEFINE_STAT(STAT_UpdateTriMeshVertices);
 
+DEFINE_STAT(STAT_AnimGameThreadTime);
+
 // Define AnimNotify
 DEFINE_LOG_CATEGORY(LogAnimNotify);
 
diff --git a/Engine/Source/Runtime/Engine/Private/SkeletalMeshComponent.cpp b/Engine/Source/Runtime/Engine/Private/SkeletalMeshComponent.cpp
--- a/Engine/Source/Runtime/Engine/Private/SkeletalMeshComponent.cpp
+++ b/Engine/Source/Runtime/Engine/Private/SkeletalMeshComponent.cpp
@@ -85,6 +85,7 @@ public:
 
 	void DoTask(ENamedThreads::Type CurrentThread, const FGraphEventRef& MyCompletionGraphEvent)
 	{
+		SCOPE_CYCLE_COUNTER(STAT_AnimGameThreadTime);
 		if (USkeletalMeshComponent* Comp = SkeletalMeshComponent.Get())
 		{
 			Comp->CompleteParallelAnimationEvaluation();
@@ -447,6 +448,7 @@ void USkeletalMeshComponent::LoadedFromAnotherClass(const FName& OldClassName)
 void USkeletalMeshComponent::TickAnimation(float DeltaTime)
 {
 	SCOPE_CYCLE_COUNTER(STAT_AnimTickTime);
+	SCOPE_CYCLE_COUNTER(STAT_AnimGameThreadTime);
 	if (SkeletalMesh != NULL)
 	{
 		if (AnimScriptInstance != NULL)
@@ -923,6 +925,7 @@ void USkeletalMeshComponent::PerformAnimationEvaluation(const USkeletalMesh* InS
 void USkeletalMeshComponent::RefreshBoneTransforms(FActorComponentTickFunction* TickFunction)
 {
 	SCOPE_CYCLE_COUNTER(STAT_RefreshBoneTransforms);
+	SCOPE_CYCLE_COUNTER(STAT_AnimGameThreadTime);
 
 	if (!SkeletalMesh || GetNumSpaceBases() == 0)
 	{
diff --git a/Engine/Source/Runtime/Engine/Public/AnimTree.h b/Engine/Source/Runtime/Engine/Public/AnimTree.h
--- a/Engine/Source/Runtime/Engine/Public/AnimTree.h
+++ b/Engine/Source/Runtime/Engine/Public/AnimTree.h
@@ -36,5 +36,6 @@ DECLARE_CYCLE_STAT_EXTERN(TEXT("Update SkelMesh Bounds"), STAT_UpdateSkelMeshBou
 DECLARE_CYCLE_STAT_EXTERN(TEXT("BlendInPhysics"), STAT_BlendInPhysics, STATGROUP_Anim, );
 DECLARE_CYCLE_STAT_EXTERN(TEXT("SkinPerPolyVertices);"), STAT_SkinPerPolyVertices, STATGROUP_Anim, );
 DECLARE_CYCLE_STAT_EXTERN(TEXT("UpdateTriMeshVertices"), STAT_UpdateTriMeshVertices, STATGROUP_Anim, );
+DECLARE_CYCLE_STAT_EXTERN(TEXT("AnimGameThreadTime"), STAT_AnimGameThreadTime, STATGROUP_Anim, );
 
 
