nicer printout

diff --git a/Engine/Source/Runtime/Engine/Private/StaticMesh.cpp b/Engine/Source/Runtime/Engine/Private/StaticMesh.cpp
--- a/Engine/Source/Runtime/Engine/Private/StaticMesh.cpp
+++ b/Engine/Source/Runtime/Engine/Private/StaticMesh.cpp
@@ -1169,8 +1169,8 @@ void FStaticMeshRenderData::Cache(UStaticMesh* Owner, const FStaticMeshLODSettin
 		GetDerivedDataCacheRef().Put(*DerivedDataKey, DerivedData);
 
 		int32 T1 = FPlatformTime::Cycles();
-		UE_LOG(LogStaticMesh,Log,TEXT("Built static mesh [%fms] %s"),
-			FPlatformTime::ToMilliseconds(T1-T0),
+		UE_LOG(LogStaticMesh,Log,TEXT("Built static mesh [%f.2s] %s"),
+			FPlatformTime::ToMilliseconds(T1-T0) / 1000.0f,
 			*Owner->GetPathName()
 			);
 		FPlatformAtomics::InterlockedAdd(&StaticMeshDerivedDataTimings::BuildCycles, T1-T0);
