Fixed two rare crashes related to navmesh debugging tools #UE4

diff --git a/Engine/Source/Runtime/Engine/Private/AI/Navigation/NavigationTestingActor.cpp b/Engine/Source/Runtime/Engine/Private/AI/Navigation/NavigationTestingActor.cpp
--- a/Engine/Source/Runtime/Engine/Private/AI/Navigation/NavigationTestingActor.cpp
+++ b/Engine/Source/Runtime/Engine/Private/AI/Navigation/NavigationTestingActor.cpp
@@ -197,13 +197,16 @@ void ANavigationTestingActor::PostEditMove(bool bFinished)
 	Super::PostEditMove(bFinished);
 
 	// project location to navmesh
-	FNavLocation NavLoc;
-	bProjectedLocationValid = GetWorld()->GetNavigationSystem()->ProjectPointToNavigation(GetActorLocation(), NavLoc, QueryingExtent, MyNavData);
-	ProjectedLocation = NavLoc.Location;
-
-	if (bSearchStart || (OtherActor != NULL && OtherActor->bSearchStart))
+	if (GetWorld()->GetNavigationSystem())
 	{
-		UpdatePathfinding();
+		FNavLocation NavLoc;
+		bProjectedLocationValid = GetWorld()->GetNavigationSystem()->ProjectPointToNavigation(GetActorLocation(), NavLoc, QueryingExtent, MyNavData);
+		ProjectedLocation = NavLoc.Location;
+
+		if (bSearchStart || (OtherActor != NULL && OtherActor->bSearchStart))
+		{
+			UpdatePathfinding();
+		}
 	}
 }
 
diff --git a/Engine/Source/Runtime/Engine/Private/NavMeshRenderingHelpers.h b/Engine/Source/Runtime/Engine/Private/NavMeshRenderingHelpers.h
--- a/Engine/Source/Runtime/Engine/Private/NavMeshRenderingHelpers.h
+++ b/Engine/Source/Runtime/Engine/Private/NavMeshRenderingHelpers.h
@@ -398,35 +398,41 @@ public:
 				}
 
 				// Draw Mesh
-				for (int32 Index = 0; Index < MeshBatchElements.Num(); ++Index)
+				if (MeshBatchElements.Num())
 				{
-					if (MeshBatchElements[Index].NumPrimitives == 0)
+					for (size_t i = 0; i < length; i++)
 					{
-						continue;
-					}
+						(int32 Index = 0; Index < MeshBatchElements.Num(); ++Index)
+						{
+						if (MeshBatchElements[Index].NumPrimitives == 0)
+						{
+							continue;
+						}
 
-					FMeshBatch& Mesh = Collector.AllocateMesh();
-					FMeshBatchElement& BatchElement = Mesh.Elements[0];
-					BatchElement = MeshBatchElements[Index];
-					BatchElement.PrimitiveUniformBuffer = CreatePrimitiveUniformBufferImmediate(GetLocalToWorld(), GetBounds(), GetLocalBounds(), false, UseEditorDepthTest());
-
-					Mesh.bWireframe = false;
-					Mesh.VertexFactory = &VertexFactory;
-					Mesh.MaterialRenderProxy = &MeshColors[Index];
-					Mesh.ReverseCulling = IsLocalToWorldDeterminantNegative();
-					Mesh.Type = PT_TriangleList;
-					Mesh.DepthPriorityGroup = SDPG_World;
-					Mesh.bCanApplyViewModeOverrides = false;
-					Collector.AddMesh(ViewIndex, Mesh);
-				}
+						FMeshBatch& Mesh = Collector.AllocateMesh();
+						FMeshBatchElement& BatchElement = Mesh.Elements[0];
+						BatchElement = MeshBatchElements[Index];
+						BatchElement.PrimitiveUniformBuffer = CreatePrimitiveUniformBufferImmediate(GetLocalToWorld(), GetBounds(), GetLocalBounds(), false, UseEditorDepthTest());
+
+						Mesh.bWireframe = false;
+						Mesh.VertexFactory = &VertexFactory;
+						Mesh.MaterialRenderProxy = &MeshColors[Index];
+						Mesh.ReverseCulling = IsLocalToWorldDeterminantNegative();
+						Mesh.Type = PT_TriangleList;
+						Mesh.DepthPriorityGroup = SDPG_World;
+						Mesh.bCanApplyViewModeOverrides = false;
+						Collector.AddMesh(ViewIndex, Mesh);
+					}
 
-				if (ProxyData.PathCollidingGeomIndices.Num() > 2)
-				{
-					FDynamicMeshBuilder MeshBuilder;
-					MeshBuilder.AddVertices(ProxyData.PathCollidingGeomVerts);
-					MeshBuilder.AddTriangles(ProxyData.PathCollidingGeomIndices);
+						if (ProxyData.PathCollidingGeomIndices.Num() > 2)
+						{
+							FDynamicMeshBuilder MeshBuilder;
+							MeshBuilder.AddVertices(ProxyData.PathCollidingGeomVerts);
+							MeshBuilder.AddTriangles(ProxyData.PathCollidingGeomIndices);
 
-					MeshBuilder.GetMesh(FMatrix::Identity, &MeshColors[MeshBatchElements.Num()], SDPG_World, false, false, ViewIndex, Collector);
+							MeshBuilder.GetMesh(FMatrix::Identity, &MeshColors[MeshBatchElements.Num()], SDPG_World, false, false, ViewIndex, Collector);
+						}
+					}
 				}
 
 				int32 Num = ProxyData.NavMeshEdgeLines.Num();
