Fix to previous change to navmesh rendering #UE4

diff --git a/Engine/Source/Runtime/Engine/Private/NavMeshRenderingHelpers.h b/Engine/Source/Runtime/Engine/Private/NavMeshRenderingHelpers.h
--- a/Engine/Source/Runtime/Engine/Private/NavMeshRenderingHelpers.h
+++ b/Engine/Source/Runtime/Engine/Private/NavMeshRenderingHelpers.h
@@ -400,10 +400,8 @@ public:
 				// Draw Mesh
 				if (MeshBatchElements.Num())
 				{
-					for (size_t i = 0; i < length; i++)
+					for (int32 Index = 0; Index < MeshBatchElements.Num(); ++Index)
 					{
-						(int32 Index = 0; Index < MeshBatchElements.Num(); ++Index)
-						{
 						if (MeshBatchElements[Index].NumPrimitives == 0)
 						{
 							continue;
@@ -424,14 +422,13 @@ public:
 						Collector.AddMesh(ViewIndex, Mesh);
 					}
 
-						if (ProxyData.PathCollidingGeomIndices.Num() > 2)
-						{
-							FDynamicMeshBuilder MeshBuilder;
-							MeshBuilder.AddVertices(ProxyData.PathCollidingGeomVerts);
-							MeshBuilder.AddTriangles(ProxyData.PathCollidingGeomIndices);
+					if (ProxyData.PathCollidingGeomIndices.Num() > 2)
+					{
+						FDynamicMeshBuilder MeshBuilder;
+						MeshBuilder.AddVertices(ProxyData.PathCollidingGeomVerts);
+						MeshBuilder.AddTriangles(ProxyData.PathCollidingGeomIndices);
 
-							MeshBuilder.GetMesh(FMatrix::Identity, &MeshColors[MeshBatchElements.Num()], SDPG_World, false, false, ViewIndex, Collector);
-						}
+						MeshBuilder.GetMesh(FMatrix::Identity, &MeshColors[MeshBatchElements.Num()], SDPG_World, false, false, ViewIndex, Collector);
 					}
 				}
 
