Improved the compile error presented when using a runtime cached font with a font sampler

UE-7553 - Invalid Font Page error when using a Custom Font

It used to report "Invalid font page 0. Max allowed is 0", but it now reports "Font 'X' is runtime cached, but only offline cached fonts can be sampled".

diff --git a/Engine/Source/Runtime/Engine/Private/Materials/MaterialExpressions.cpp b/Engine/Source/Runtime/Engine/Private/Materials/MaterialExpressions.cpp
--- a/Engine/Source/Runtime/Engine/Private/Materials/MaterialExpressions.cpp
+++ b/Engine/Source/Runtime/Engine/Private/Materials/MaterialExpressions.cpp
@@ -5405,48 +5405,49 @@ int32 UMaterialExpressionFontSample::Compile(class FMaterialCompiler* Compiler,
 		FontTexturePage = 0;
 	}
 #endif
-	if( Font )
+	if( !Font )
+	{
+		Result = CompilerError(Compiler, TEXT("Missing input Font"));
+	}
+	else if( Font->FontCacheType == EFontCacheType::Runtime )
+	{
+		Result = CompilerError(Compiler, *FString::Printf(TEXT("Font '%s' is runtime cached, but only offline cached fonts can be sampled"), *Font->GetName()));
+	}
+	else if( !Font->Textures.IsValidIndex(FontTexturePage) )
 	{
-		if( Font->Textures.IsValidIndex(FontTexturePage) )
+		Result = CompilerError(Compiler, *FString::Printf(TEXT("Invalid font page %d. Max allowed is %d"), FontTexturePage, Font->Textures.Num()));
+	}
+	else
+	{
+		UTexture* Texture = Font->Textures[FontTexturePage];
+		if( !Texture )
 		{
-			UTexture* Texture = Font->Textures[FontTexturePage];
-			if( !Texture )
-			{
-				UE_LOG(LogMaterial, Log, TEXT("Invalid font texture. Using default texture"));
-				Texture = Texture = GEngine->DefaultTexture;
-			}
-			check(Texture);
-
-			EMaterialSamplerType ExpectedSamplerType;
-			if (Texture->CompressionSettings == TC_DistanceFieldFont)
-			{
-				ExpectedSamplerType = SAMPLERTYPE_DistanceFieldFont;
-			}
-			else
-			{
-				ExpectedSamplerType = Texture->SRGB ? SAMPLERTYPE_Color : SAMPLERTYPE_LinearColor;
-			}
-
-			if (!VerifySamplerType(Compiler, (Desc.Len() > 0 ? *Desc : TEXT("FontSample")), Texture, ExpectedSamplerType))
-			{
-				return INDEX_NONE;
-			}
+			UE_LOG(LogMaterial, Log, TEXT("Invalid font texture. Using default texture"));
+			Texture = Texture = GEngine->DefaultTexture;
+		}
+		check(Texture);
 
-			int32 TextureCodeIndex = Compiler->Texture(Texture);
-			Result = Compiler->TextureSample(
-				TextureCodeIndex,
-				Compiler->TextureCoordinate(0, false, false),
-				ExpectedSamplerType
-			);
+		EMaterialSamplerType ExpectedSamplerType;
+		if (Texture->CompressionSettings == TC_DistanceFieldFont)
+		{
+			ExpectedSamplerType = SAMPLERTYPE_DistanceFieldFont;
 		}
 		else
 		{
-			Result = CompilerError(Compiler, *FString::Printf(TEXT("Invalid font page %d. Max allowed is %d"),FontTexturePage,Font->Textures.Num()));
-		}		
-	}
-	else
-	{
-		Result = CompilerError(Compiler, TEXT("Missing input Font"));
+			ExpectedSamplerType = Texture->SRGB ? SAMPLERTYPE_Color : SAMPLERTYPE_LinearColor;
+		}
+
+		if (!VerifySamplerType(Compiler, (Desc.Len() > 0 ? *Desc : TEXT("FontSample")), Texture, ExpectedSamplerType))
+		{
+			return INDEX_NONE;
+		}
+
+		int32 TextureCodeIndex = Compiler->Texture(Texture);
+		Result = Compiler->TextureSample(
+			TextureCodeIndex,
+			Compiler->TextureCoordinate(0, false, false),
+			ExpectedSamplerType
+		);
 	}
 	return Result;
 }
