Marked HierarchicalInstancedStaticMeshComponent as being able to affect navmesh generation (which doesn't mean it will, it just has an option to) #UE4

diff --git a/Engine/Source/Runtime/Engine/Private/HierarchicalInstancedStaticMesh.cpp b/Engine/Source/Runtime/Engine/Private/HierarchicalInstancedStaticMesh.cpp
--- a/Engine/Source/Runtime/Engine/Private/HierarchicalInstancedStaticMesh.cpp
+++ b/Engine/Source/Runtime/Engine/Private/HierarchicalInstancedStaticMesh.cpp
@@ -1384,6 +1384,7 @@ UHierarchicalInstancedStaticMeshComponent::UHierarchicalInstancedStaticMeshCompo
 	, bConcurrentRemoval(false)
 	, AccumulatedNavigationDirtyArea(0)
 {
+	bCanEverAffectNavigation = true;
 }
 
 UHierarchicalInstancedStaticMeshComponent::~UHierarchicalInstancedStaticMeshComponent()
