Fixed reloading snapping bug issue

https://jira.ol.epicgames.net/browse/UE-6189

diff --git a/Engine/Source/Runtime/Engine/Private/Animation/AnimMontage.cpp b/Engine/Source/Runtime/Engine/Private/Animation/AnimMontage.cpp
--- a/Engine/Source/Runtime/Engine/Private/Animation/AnimMontage.cpp
+++ b/Engine/Source/Runtime/Engine/Private/Animation/AnimMontage.cpp
@@ -683,7 +683,7 @@ EAnimEventTriggerOffsets::Type UAnimMontage::CalculateOffsetFromSections(float T
 	for(auto Iter = CompositeSections.CreateConstIterator(); Iter; ++Iter)
 	{
 		float SectionTime = Iter->GetTime();
-		if(SectionTime == Time)
+		if(FMath::IsNearlyEqual(SectionTime,Time))
 		{
 			return EAnimEventTriggerOffsets::OffsetBefore;
 		}
