Camera shakes: Ignore calls to PlayCameraShake with zero or negative scales, and allow shakes that reach zero scale to still update their oscillator so they will eventually be removed [UE-8656]

diff --git a/Engine/Source/Runtime/Engine/Private/Camera/CameraModifier_CameraShake.cpp b/Engine/Source/Runtime/Engine/Private/Camera/CameraModifier_CameraShake.cpp
--- a/Engine/Source/Runtime/Engine/Private/Camera/CameraModifier_CameraShake.cpp
+++ b/Engine/Source/Runtime/Engine/Private/Camera/CameraModifier_CameraShake.cpp
@@ -36,12 +36,7 @@ void UCameraModifier_CameraShake::UpdateCameraShake(float DeltaTime, FCameraShak
 	}
 
 	// this is the base scale for the whole shake, anim and oscillation alike
-	float const BaseShakeScale = GetShakeScale(Shake);
-	// do not update if percentage is null
-	if (BaseShakeScale <= 0.f)
-	{
-		return;
-	}
+	float const BaseShakeScale = FMath::Max<float>(GetShakeScale(Shake), 0.0f);
 
 	// update anims with any desired scaling
 	if (Shake.AnimInst)
diff --git a/Engine/Source/Runtime/Engine/Private/PlayerCameraManager.cpp b/Engine/Source/Runtime/Engine/Private/PlayerCameraManager.cpp
--- a/Engine/Source/Runtime/Engine/Private/PlayerCameraManager.cpp
+++ b/Engine/Source/Runtime/Engine/Private/PlayerCameraManager.cpp
@@ -1062,7 +1062,7 @@ void APlayerCameraManager::ClearCameraLensEffects()
 
 void APlayerCameraManager::PlayCameraShake(TSubclassOf<UCameraShake> Shake, float Scale, ECameraAnimPlaySpace::Type PlaySpace, FRotator UserPlaySpaceRot)
 {
-	if (Shake != NULL)
+	if ((Shake != nullptr) && (Scale > 0.0f))
 	{
 		CameraShakeCamMod->AddCameraShake(Shake, Scale, PlaySpace, UserPlaySpaceRot);
 	}
