UE4 - r.DumpShaderDebugShortNames to shorten dump shader debug paths

diff --git a/Engine/Config/ConsoleVariables.ini b/Engine/Config/ConsoleVariables.ini
--- a/Engine/Config/ConsoleVariables.ini
+++ b/Engine/Config/ConsoleVariables.ini
@@ -24,6 +24,8 @@
 ; Uncomment to dump shaders in the Saved folder
 ; Warning: leaving this on for a while will fill your hard drive with many small files and folders
 ;r.DumpShaderDebugInfo=1
+; When this is enabled, dumped shader paths will get collapsed (in the cases where paths are longer than the OS's max)
+;r.DumpShaderDebugShortNames=1
 ; Uncomment to enable parallel rendering at startup
 ;r.RHICmdBypass=0
 
@@ -31,4 +33,4 @@
 ;r.XGEShaderCompile = 1
 ; Uncomment when running with a graphical debugger (but not when profiling)
 ;r.Shaders.Optimize=0
-;r.Shaders.KeepDebugInfo=1
\ No newline at end of file
+;r.Shaders.KeepDebugInfo=1
diff --git a/Engine/Source/Runtime/Engine/Private/ShaderCompiler/ShaderCompiler.cpp b/Engine/Source/Runtime/Engine/Private/ShaderCompiler/ShaderCompiler.cpp
--- a/Engine/Source/Runtime/Engine/Private/ShaderCompiler/ShaderCompiler.cpp
+++ b/Engine/Source/Runtime/Engine/Private/ShaderCompiler/ShaderCompiler.cpp
@@ -32,6 +32,14 @@ static FAutoConsoleVariableRef CVarDumpShaderDebugInfo(
 	TEXT("On iOS, if the PowerVR graphics SDK is installed to the default path, the PowerVR shader compiler will be called and errors will be reported during the cook.")
 	);
 
+int32 GDumpShaderDebugInfoShort = 0;
+static FAutoConsoleVariableRef CVarDumpShaderDebugShortNames(
+	TEXT("r.DumpShaderDebugShortNames"),
+	GDumpShaderDebugInfoShort,
+	TEXT("Only valid when r.DumpShaderDebugInfo=1.\n")
+	TEXT("When set to 1, will shorten names factory and shader type folder names to avoid issues with long paths.")
+	);
+
 static TAutoConsoleVariable<int32> CVarKeepShaderDebugData(
 	TEXT("r.Shaders.KeepDebugInfo"),
 	0,
@@ -2084,10 +2092,52 @@ void GlobalBeginCompileShader(
 
 		if (VFType)
 		{
-			Input.DumpDebugInfoPath = Input.DumpDebugInfoPath / VFType->GetName();
+			FString VFName = VFType->GetName();
+			if (GDumpShaderDebugInfoShort)
+			{
+				// Shorten vertex factory name
+				if (VFName[0] == TCHAR('F') || VFName[0] == TCHAR('T'))
+				{
+					VFName.RemoveAt(0);
+				}
+				VFName.ReplaceInline(TEXT("VertexFactory"), TEXT("VF"));
+				VFName.ReplaceInline(TEXT("GPUSkinAPEXCloth"), TEXT("APEX"));
+				VFName.ReplaceInline(TEXT("true"), TEXT("_1"));
+				VFName.ReplaceInline(TEXT("false"), TEXT("_0"));
+			}
+			Input.DumpDebugInfoPath = Input.DumpDebugInfoPath / VFName;
 		}
 
-		Input.DumpDebugInfoPath = Input.DumpDebugInfoPath / ShaderType->GetName();
+		{
+			FString ShaderTypeName = ShaderType->GetName();
+			if (GDumpShaderDebugInfoShort)
+			{
+				// Shorten known types
+				if (ShaderTypeName[0] == TCHAR('F') || ShaderTypeName[0] == TCHAR('T'))
+				{
+					ShaderTypeName.RemoveAt(0);
+				}
+				ShaderTypeName.ReplaceInline(TEXT("BasePass"), TEXT("BP"));
+				ShaderTypeName.ReplaceInline(TEXT("ForForward"), TEXT("Fwd"));
+				ShaderTypeName.ReplaceInline(TEXT("Shadow"), TEXT("Shdw"));
+				ShaderTypeName.ReplaceInline(TEXT("LightMap"), TEXT("LM"));
+				ShaderTypeName.ReplaceInline(TEXT("Atmospheric"), TEXT("Atm"));
+				ShaderTypeName.ReplaceInline(TEXT("Perspective"), TEXT("Persp"));
+				ShaderTypeName.ReplaceInline(TEXT("Position"), TEXT("Pos"));
+				ShaderTypeName.ReplaceInline(TEXT("Skylight"), TEXT("Sky"));
+				ShaderTypeName.ReplaceInline(TEXT("LightingPolicy"), TEXT("LP"));
+				ShaderTypeName.ReplaceInline(TEXT("TranslucentLighting"), TEXT("TranslLight"));
+				ShaderTypeName.ReplaceInline(TEXT("DistanceField"), TEXT("DistFiel"));
+				ShaderTypeName.ReplaceInline(TEXT("Indirect"), TEXT("Ind"));
+				ShaderTypeName.ReplaceInline(TEXT("Cached"), TEXT("Cach"));
+				ShaderTypeName.ReplaceInline(TEXT("Dynamic"), TEXT("Dyn"));
+				ShaderTypeName.ReplaceInline(TEXT("Vertex"), TEXT("Vtx"));
+				ShaderTypeName.ReplaceInline(TEXT("Output"), TEXT("Out"));
+				ShaderTypeName.ReplaceInline(TEXT("true"), TEXT("_1"));
+				ShaderTypeName.ReplaceInline(TEXT("false"), TEXT("_0"));
+			}
+			Input.DumpDebugInfoPath = Input.DumpDebugInfoPath / ShaderTypeName;
+		}
 		// Sanitize the name to be used as a path
 		// List mostly comes from set of characters not allowed by windows in a path.  Just try to rename a file and type one of these for the list.
 		Input.DumpDebugInfoPath.ReplaceInline(TEXT("<"), TEXT("("));
