#vr Adding gating to prevent initializing an HMD when running a commandlet

diff --git a/Engine/Source/Runtime/Engine/Private/UnrealEngine.cpp b/Engine/Source/Runtime/Engine/Private/UnrealEngine.cpp
--- a/Engine/Source/Runtime/Engine/Private/UnrealEngine.cpp
+++ b/Engine/Source/Runtime/Engine/Private/UnrealEngine.cpp
@@ -1738,23 +1738,26 @@ public:
 
 bool UEngine::InitializeHMDDevice()
 {
-	if (FParse::Param(FCommandLine::Get(), TEXT("emulatestereo")))
-	{
-		TSharedPtr<FFakeStereoRenderingDevice> FakeStereoDevice(new FFakeStereoRenderingDevice());
-		StereoRenderingDevice = FakeStereoDevice;
-	}
-	// No reason to connect an HMD on a dedicated server.  Also fixes dedicated servers stealing the oculus connection.
-	else if (!HMDDevice.IsValid() && !FParse::Param(FCommandLine::Get(), TEXT("nohmd")) && !IsRunningDedicatedServer())
+	if (!IsRunningCommandlet())
 	{
-		// Get a list of plugins that implement this feature
-		TArray<IHeadMountedDisplayModule*> HMDImplementations = IModularFeatures::Get().GetModularFeatureImplementations<IHeadMountedDisplayModule>(IHeadMountedDisplayModule::GetModularFeatureName());
-		HMDImplementations.Sort(FHMDPluginSorter());
-		for (auto HMDModuleIt = HMDImplementations.CreateIterator(); HMDModuleIt && !HMDDevice.IsValid(); ++HMDModuleIt)
+		if (FParse::Param(FCommandLine::Get(), TEXT("emulatestereo")))
 		{
-			HMDDevice = (*HMDModuleIt)->CreateHeadMountedDisplay();
-			if (HMDDevice.IsValid())
+			TSharedPtr<FFakeStereoRenderingDevice> FakeStereoDevice(new FFakeStereoRenderingDevice());
+			StereoRenderingDevice = FakeStereoDevice;
+		}
+		// No reason to connect an HMD on a dedicated server.  Also fixes dedicated servers stealing the oculus connection.
+		else if (!HMDDevice.IsValid() && !FParse::Param(FCommandLine::Get(), TEXT("nohmd")) && !IsRunningDedicatedServer())
+		{
+			// Get a list of plugins that implement this feature
+			TArray<IHeadMountedDisplayModule*> HMDImplementations = IModularFeatures::Get().GetModularFeatureImplementations<IHeadMountedDisplayModule>(IHeadMountedDisplayModule::GetModularFeatureName());
+			HMDImplementations.Sort(FHMDPluginSorter());
+			for (auto HMDModuleIt = HMDImplementations.CreateIterator(); HMDModuleIt && !HMDDevice.IsValid(); ++HMDModuleIt)
 			{
-				StereoRenderingDevice = HMDDevice;
+				HMDDevice = (*HMDModuleIt)->CreateHeadMountedDisplay();
+				if (HMDDevice.IsValid())
+				{
+					StereoRenderingDevice = HMDDevice;
+				}
 			}
 		}
 	}
