Unshelved from pending changelist '2420730':

Fix for Tim's GPU emitter Matinee bug.

diff --git a/Engine/Source/Runtime/Core/Public/UObject/ObjectVersion.h b/Engine/Source/Runtime/Core/Public/UObject/ObjectVersion.h
--- a/Engine/Source/Runtime/Core/Public/UObject/ObjectVersion.h
+++ b/Engine/Source/Runtime/Core/Public/UObject/ObjectVersion.h
@@ -466,6 +466,8 @@ enum EUnrealEngineObjectUE4Version
 	VER_UE4_AFTER_MERGING_ADD_MODIFIERS_RUNTIME_GENERATION_TO_4_7,
 	// Landscape grass weightmap data is now generated in the editor and serialized.
 	VER_UE4_SERIALIZE_LANDSCAPE_GRASS_DATA,
+	// New property to optionally prevent gpu emitters clearing existing particles on Init().
+	VER_UE4_OPTIONALLY_CLEAR_GPU_EMITTERS_ON_INIT,
 
 	// -----<new versions can be added before this line>-------------------------------------------------
 	// - this needs to be the last line (see note below)
diff --git a/Engine/Source/Runtime/Engine/Classes/Particles/TypeData/ParticleModuleTypeDataGpu.h b/Engine/Source/Runtime/Engine/Classes/Particles/TypeData/ParticleModuleTypeDataGpu.h
--- a/Engine/Source/Runtime/Engine/Classes/Particles/TypeData/ParticleModuleTypeDataGpu.h
+++ b/Engine/Source/Runtime/Engine/Classes/Particles/TypeData/ParticleModuleTypeDataGpu.h
@@ -406,6 +406,9 @@ class UParticleModuleTypeDataGpu : public UParticleModuleTypeDataBase
 	UPROPERTY(EditAnywhere, Category=ParticleModuleTypeDataGpu)
 	float CameraMotionBlurAmount;
 
+	/** When true, all existing partilces are cleared when the emitter is initialized. */
+	UPROPERTY(EditAnywhere, Category = ParticleModuleTypeDataGpu)
+	uint32 bClearExistingParticlesOnInit:1;
 
 	// Begin UObject Interface
 	virtual void PostLoad() override;
diff --git a/Engine/Source/Runtime/Engine/Private/Particles/ParticleGpuSimulation.cpp b/Engine/Source/Runtime/Engine/Private/Particles/ParticleGpuSimulation.cpp
--- a/Engine/Source/Runtime/Engine/Private/Particles/ParticleGpuSimulation.cpp
+++ b/Engine/Source/Runtime/Engine/Private/Particles/ParticleGpuSimulation.cpp
@@ -3042,8 +3042,19 @@ public:
 			*Component->GetName(),*Component->Template->GetName(),(PTRINT)this, AllocatedTiles.Num());
 #endif // #if TRACK_TILE_ALLOCATIONS
 
-		ActiveTiles.Init(false, AllocatedTiles.Num());
-		ClearAllocatedTiles();
+		bool bClearExistingParticles = false;
+		UParticleLODLevel* LODLevel = SpriteTemplate->LODLevels[0];
+		if (LODLevel)
+		{
+			UParticleModuleTypeDataGpu* TypeDataModule = CastChecked<UParticleModuleTypeDataGpu>(LODLevel->TypeDataModule);
+			bClearExistingParticles = TypeDataModule->bClearExistingParticlesOnInit;
+		}
+
+		if (bClearExistingParticles || ActiveTiles.Num() != AllocatedTiles.Num())
+		{
+			ActiveTiles.Init(false, AllocatedTiles.Num());
+			ClearAllocatedTiles();
+		}
 
 		Simulation->bDirty_GameThread = true;
 		FXSystem->AddGPUSimulation(Simulation);
diff --git a/Engine/Source/Runtime/Engine/Private/Particles/ParticleModules.cpp b/Engine/Source/Runtime/Engine/Private/Particles/ParticleModules.cpp
--- a/Engine/Source/Runtime/Engine/Private/Particles/ParticleModules.cpp
+++ b/Engine/Source/Runtime/Engine/Private/Particles/ParticleModules.cpp
@@ -4023,6 +4023,7 @@ bool UDistributionVectorParticleParameter::GetParamValue(UObject* Data, FName Pa
 -----------------------------------------------------------------------------*/
 UParticleModuleTypeDataGpu::UParticleModuleTypeDataGpu(const FObjectInitializer& ObjectInitializer)
 	: Super(ObjectInitializer)
+	, bClearExistingParticlesOnInit(false)
 {
 }
 
@@ -4030,6 +4031,13 @@ void UParticleModuleTypeDataGpu::PostLoad()
 {
 	Super::PostLoad();
 	//EmitterInfo.Resources = BeginCreateGPUSpriteResources( ResourceData );
+
+	if (GetLinkerUE4Version() < VER_UE4_OPTIONALLY_CLEAR_GPU_EMITTERS_ON_INIT)
+	{
+		//Force old emitters to clear their particles on Init() to maintain old behaviour.
+		//New emitters are defaulted to false so they behave like other emitter types.
+		bClearExistingParticlesOnInit = true;
+	}
 }
 
 void UParticleModuleTypeDataGpu::BeginDestroy()
