Analytics for all recorded events in the engine.

diff --git a/Engine/Source/Runtime/Engine/Private/Tests/AnalyticsTest.cpp b/Engine/Source/Runtime/Engine/Private/Tests/AnalyticsTest.cpp
new file mode 100644
--- /dev/null
+++ b/Engine/Source/Runtime/Engine/Private/Tests/AnalyticsTest.cpp
@@ -0,0 +1,70 @@
+// Copyright 1998-2015 Epic Games, Inc. All Rights Reserved.
+
+#include "EnginePrivate.h"
+#include "AutomationTest.h"
+#include "EngineAnalytics.h"
+#include "Runtime/Analytics/Analytics/Public/Interfaces/IAnalyticsProvider.h"
+#include "UnrealEngine.h"
+
+/**
+* Artificial Record Event for analytics - Simulates the engine startup simulation.  
+*/
+IMPLEMENT_SIMPLE_AUTOMATION_TEST(FAnalyticStartUpSimTest, "Engine.Analytic.Record Event - Simulate Program Start", EAutomationTestFlags::ATF_FeatureMask)
+
+bool FAnalyticStartUpSimTest::RunTest(const FString& Parameters)
+{
+	if (FEngineAnalytics::IsAvailable())
+	{
+		//Setup the 'Event Attributes'
+		TArray<FAnalyticsEventAttribute> EventAttributes;
+		EventAttributes.Add(FAnalyticsEventAttribute(TEXT("MachineID"), FPlatformMisc::GetMachineId().ToString(EGuidFormats::Digits).ToLower()));
+		EventAttributes.Add(FAnalyticsEventAttribute(TEXT("AccountID"), FPlatformMisc::GetEpicAccountId()));
+		EventAttributes.Add(FAnalyticsEventAttribute(TEXT("GameName"), FApp::GetGameName()));
+		EventAttributes.Add(FAnalyticsEventAttribute(TEXT("CommandLine"), FCommandLine::Get()));
+		
+		//Record the event with the 'Engine.AutomationTest.Analytics.ProgramStartedEvent' title
+		FEngineAnalytics::GetProvider().RecordEvent(TEXT("Engine.AutomationTest.Analytics.ProgramStartedEvent"), EventAttributes);
+
+		//Get the event strings used
+		FString MachineIDTest		=	FPlatformMisc::GetMachineId().ToString(EGuidFormats::Digits).ToLower();
+		FString AccountIDTest		=	FPlatformMisc::GetEpicAccountId();
+		FString GameNameTest		=	FApp::GetGameName();
+		FString CommandLineArgs		=	FCommandLine::Get();
+
+		//Test the strings to verify they have data.
+		TestFalse(TEXT("'MachineID' is not expected to be empty!"), MachineIDTest.IsEmpty());
+		TestFalse(TEXT("'AccountID' is not expected to be empty!"), AccountIDTest.IsEmpty());
+		TestFalse(TEXT("'GameName' is expected."), GameNameTest.IsEmpty());
+
+		//Verify record event is holding the actual data.  This only triggers if the command line argument of 'AnalyticsDisableCaching' was used.
+		if (CommandLineArgs.Contains(TEXT("AnalyticsDisableCaching")))
+		{
+			for (int32 i = 0; i < ExecutionInfo.LogItems.Num(); i++)
+			{
+				if (ExecutionInfo.LogItems[i].Contains(TEXT("Engine.AutomationTest.Analytics.ProgramStartedEvent")))
+				{
+					TestTrue(TEXT("Recorded event name is expected to be in the sent event."), ExecutionInfo.LogItems[i].Contains(TEXT("Engine.AutomationTest.Analytics.ProgramStartedEvent")));
+					TestTrue(TEXT("'MachineID' is expected to be in the sent event."), ExecutionInfo.LogItems[i].Contains(*MachineIDTest));
+					TestTrue(TEXT("'AccountID' is expected to be in the sent event."), ExecutionInfo.LogItems[i].Contains(*AccountIDTest));
+					TestTrue(TEXT("'GameName' is expected to be in the sent event."), ExecutionInfo.LogItems[i].Contains(*GameNameTest));
+					TestTrue(TEXT("'CommandLine arguments' are expected to be in the sent event."), ExecutionInfo.LogItems[i].Contains(TEXT("AnalyticsDisableCaching")));
+				}
+			}
+		}
+	}
+	return true;	
+}
+
+///**
+//* Engine Startup Analytic unit test - This triggers an artificial event.
+//*/
+//IMPLEMENT_SIMPLE_AUTOMATION_TEST(FAnalyticEngineInit, "Engine.Analytic Record Events.Engine Init Test", EAutomationTestFlags::ATF_FeatureMask)
+//
+//bool FAnalyticEngineInit::RunTest(const FString& Parameters)
+//{
+//	IEngineLoop* InNewEngineLoop = nullptr;
+//
+//	GEngine->Init(InNewEngineLoop);
+//
+//	return true;
+//}
\ No newline at end of file
