Set default LightMapResolution in UStaticMesh to 4.

diff --git a/Engine/Source/Runtime/Engine/Private/StaticMesh.cpp b/Engine/Source/Runtime/Engine/Private/StaticMesh.cpp
--- a/Engine/Source/Runtime/Engine/Private/StaticMesh.cpp
+++ b/Engine/Source/Runtime/Engine/Private/StaticMesh.cpp
@@ -1208,6 +1208,7 @@ UStaticMesh::UStaticMesh(const FObjectInitializer& ObjectInitializer)
 	AutoLODPixelError = 1.0f;
 	bAutoComputeLODScreenSize=true;
 #endif // #if WITH_EDITORONLY_DATA
+	LightMapResolution = 4;
 	LpvBiasMultiplier = 1.0f;
 }
 
