Fixed wrong transformation in hierarchical tree traversing (GatherInstanceTransformsInArea)

diff --git a/Engine/Source/Runtime/Engine/Private/HierarchicalInstancedStaticMesh.cpp b/Engine/Source/Runtime/Engine/Private/HierarchicalInstancedStaticMesh.cpp
--- a/Engine/Source/Runtime/Engine/Private/HierarchicalInstancedStaticMesh.cpp
+++ b/Engine/Source/Runtime/Engine/Private/HierarchicalInstancedStaticMesh.cpp
@@ -1891,12 +1891,8 @@ static void GatherInstanceTransformsInArea(const UHierarchicalInstancedStaticMes
 {
 	const TArray<FClusterNode>& ClusterTree = *Component.ClusterTreePtr;
 	const FClusterNode& ChildNode = ClusterTree[Child];
-	const FVector WorldMinMax[] = {
-		Component.ComponentToWorld.TransformPosition(ChildNode.BoundMin),
-		Component.ComponentToWorld.TransformPosition(ChildNode.BoundMax)
-	};
-	const FBox WorldNodeBox(WorldMinMax, 2);
-
+	const FBox WorldNodeBox = FBox(ChildNode.BoundMin, ChildNode.BoundMax).TransformBy(Component.ComponentToWorld);
+	
 	if (AreaBox.Intersect(WorldNodeBox))
 	{
 		if (ChildNode.FirstChild < 0 || AreaBox.IsInside(WorldNodeBox))
