Fixed instanced mesh nav modifiers exporting

diff --git a/Engine/Source/Runtime/Engine/Private/InstancedStaticMesh.cpp b/Engine/Source/Runtime/Engine/Private/InstancedStaticMesh.cpp
--- a/Engine/Source/Runtime/Engine/Private/InstancedStaticMesh.cpp
+++ b/Engine/Source/Runtime/Engine/Private/InstancedStaticMesh.cpp
@@ -1466,10 +1466,17 @@ bool UInstancedStaticMeshComponent::DoCustomNavigableGeometryExport(struct FNavi
 
 void UInstancedStaticMeshComponent::GetNavigationData(FNavigationRelevantData& Data) const
 {
-	Super::GetNavigationData(Data);
+	if (StaticMesh && StaticMesh->NavCollision)	
+	{
+		UNavCollision* NavCollision = StaticMesh->NavCollision;
+		if (NavCollision->bIsDynamicObstacle)
+		{
+			NavCollision->GetNavigationModifier(Data.Modifiers, FTransform::Identity);
 
-	// Hook per instance transform delegate
-	Data.NavDataPerInstanceTransformDelegate = FNavDataPerInstanceTransformDelegate::CreateUObject(this, &UInstancedStaticMeshComponent::GetNavigationPerInstanceTransforms);
+			// Hook per instance transform delegate
+			Data.NavDataPerInstanceTransformDelegate = FNavDataPerInstanceTransformDelegate::CreateUObject(this, &UInstancedStaticMeshComponent::GetNavigationPerInstanceTransforms);
+		}
+	}
 }
 
 void UInstancedStaticMeshComponent::GetNavigationPerInstanceTransforms(const FBox& AreaBox, TArray<FTransform>& InstanceData) const
