Fixed navmesh terrain conforming artifacts

Project down z-axis to last plane rather than projecting point to plane. Doh.
Moved new properties to new category.

diff --git a/Engine/Source/Runtime/Engine/Classes/GameFramework/CharacterMovementComponent.h b/Engine/Source/Runtime/Engine/Classes/GameFramework/CharacterMovementComponent.h
--- a/Engine/Source/Runtime/Engine/Classes/GameFramework/CharacterMovementComponent.h
+++ b/Engine/Source/Runtime/Engine/Classes/GameFramework/CharacterMovementComponent.h
@@ -677,7 +677,7 @@ protected:
 	uint32 bDeferUpdateBasedMovement : 1;
 
 	/** Whether to raycast to underlying geometry to better conform navmesh-walking characters */
-	UPROPERTY(Category = "Character Movement", EditAnywhere, BlueprintReadOnly)
+	UPROPERTY(Category = "NavMesh Movement", EditAnywhere, BlueprintReadOnly)
 	uint32 bProjectNavMeshWalking : 1;
 
 	/** forced avoidance velocity, used when AvoidanceLockTimer is > 0 */
@@ -735,7 +735,7 @@ public:
 	FHitResult CachedProjectedNavMeshHitResult;
 
 	/** How often we should raycast to project from navmesh to underlying geometry */
-	UPROPERTY(Category = "Character Movement", EditAnywhere, BlueprintReadOnly, meta=(editcondition = "bProjectNavMeshWalking"))
+	UPROPERTY(Category = "NavMesh Movement", EditAnywhere, BlueprintReadOnly, meta=(editcondition = "bProjectNavMeshWalking"))
 	float NavMeshProjectionInterval;
 
 	UPROPERTY(Transient)
diff --git a/Engine/Source/Runtime/Engine/Private/CharacterMovementComponent.cpp b/Engine/Source/Runtime/Engine/Private/CharacterMovementComponent.cpp
--- a/Engine/Source/Runtime/Engine/Private/CharacterMovementComponent.cpp
+++ b/Engine/Source/Runtime/Engine/Private/CharacterMovementComponent.cpp
@@ -3862,26 +3862,30 @@ void UCharacterMovementComponent::ProjectLocationFromNavMesh(float DeltaSeconds,
 	{
 		NavMeshProjectionTimer -= DeltaSeconds;
 
-		if(NavMeshProjectionTimer < 0.0f)
-		{
-			const FVector RayCastOffset(0.0f, 0.0f, CharacterOwner->GetCapsuleComponent()->GetScaledCapsuleHalfHeight() * 2.0f);
+		const FVector RayCastOffset(0.0f, 0.0f, CharacterOwner->GetCapsuleComponent()->GetScaledCapsuleHalfHeight() * 2.0f);
 
+		if(NavMeshProjectionTimer <= 0.0f)
+		{
 			// raycast to underlying mesh to allow us to more closely follow geometry
+			// we use static objects here as a best approximation to accept only objects that
+			// influence navmesh generation
 			FCollisionQueryParams Params;
 			FCollisionObjectQueryParams ObjectQueryParams(FCollisionObjectQueryParams::AllStaticObjects);
-			if(GetWorld()->LineTraceSingle(CachedProjectedNavMeshHitResult, InOutLocation + RayCastOffset, InOutLocation - RayCastOffset, Params, ObjectQueryParams))
+			GetWorld()->LineTraceSingle(CachedProjectedNavMeshHitResult, InOutLocation + RayCastOffset, InOutLocation - RayCastOffset, Params, ObjectQueryParams);
+
+			// discard result if we were already inside something
+			if(CachedProjectedNavMeshHitResult.bStartPenetrating)
 			{
-				InOutLocation.Z = CachedProjectedNavMeshHitResult.Location.Z;
+				CachedProjectedNavMeshHitResult.Reset();
 			}
 
 			NavMeshProjectionTimer = NavMeshProjectionInterval;
 		}
 		
-		// interpolate to new height if we had a blocking hit
+		// project to last plane we found
 		if(CachedProjectedNavMeshHitResult.bBlockingHit)
 		{
-			FPlane Plane(CachedProjectedNavMeshHitResult.Location, CachedProjectedNavMeshHitResult.Normal);
-			FVector ProjectedPoint = FVector::PointPlaneProject(InOutLocation, Plane);
+			FVector ProjectedPoint = FMath::LinePlaneIntersection(InOutLocation, InOutLocation - RayCastOffset, CachedProjectedNavMeshHitResult.Location, CachedProjectedNavMeshHitResult.Normal);
 			InOutLocation.Z = ProjectedPoint.Z;
 		}
 	}
@@ -4007,7 +4011,7 @@ void UCharacterMovementComponent::SetNavWalkingPhysics(bool bEnable)
 			UpdatedComponent->SetCollisionResponseToChannel(ECC_WorldStatic, ECR_Ignore);
 			UpdatedComponent->SetCollisionResponseToChannel(ECC_WorldDynamic, ECR_Ignore);
 			CachedProjectedNavMeshHitResult.Reset();
-			NavMeshProjectionTimer = -1.0f;
+			NavMeshProjectionTimer = 0.0f;
 		}
 		else
 		{
