Fixed UCanvas::MeasureStringInternal sometimes overmeasuring strings

UE-7427 - In game: Console font doesn't match to cursor poistion

UFont::GetCharSize was returning the width as the X advance, plus the horizontal offset, however the horizontal offset is only used as an adjuster when working out the center of a character.

UFont::GetCharSize has been changed to set the width as just the X advance, and I've added UFont::GetCharHorizontalOffset to get the horizontal offset as required when using UCanvas::MeasureStringInternal with ELastCharacterIndexFormat::CharacterAtOffset.

diff --git a/Engine/Source/Runtime/Engine/Classes/Engine/Font.h b/Engine/Source/Runtime/Engine/Classes/Engine/Font.h
--- a/Engine/Source/Runtime/Engine/Classes/Engine/Font.h
+++ b/Engine/Source/Runtime/Engine/Classes/Engine/Font.h
@@ -220,6 +220,15 @@ public:
 	 */
 	ENGINE_API int8 GetCharKerning(TCHAR First, TCHAR Second) const;
 
+	/**
+	 * Gets the horizontal distance from the origin to the left most border of the given character
+	 *
+	 * @param	InCh					the character to get the offset for
+	 *
+	 * @return The offset value
+	 */
+	ENGINE_API int16 GetCharHorizontalOffset(TCHAR InCh) const;
+
 	/**
 	 * Calculate the width of the string using this font's default size and scale.
 	 *
diff --git a/Engine/Source/Runtime/Engine/Private/Font.cpp b/Engine/Source/Runtime/Engine/Private/Font.cpp
--- a/Engine/Source/Runtime/Engine/Private/Font.cpp
+++ b/Engine/Source/Runtime/Engine/Private/Font.cpp
@@ -207,7 +207,7 @@ void UFont::GetCharSize(TCHAR InCh, float& Width, float& Height) const
 			FCharacterList& CharacterList = FontCache->GetCharacterList(LegacyFontInfo, FontScale);
 
 			const FCharacterEntry& Entry = CharacterList[InCh];
-			Width = Entry.HorizontalOffset + Entry.XAdvance;
+			Width = Entry.XAdvance;
 
 			// The height of the character will always be the maximum height of any character in this font
 			Height = CharacterList.GetMaxHeight();
@@ -248,6 +248,23 @@ int8 UFont::GetCharKerning(TCHAR First, TCHAR Second) const
 	return 0;
 }
 
+int16 UFont::GetCharHorizontalOffset(TCHAR InCh) const
+{
+	if(FontCacheType == EFontCacheType::Runtime)
+	{
+		TSharedRef<FSlateFontCache> FontCache = FSlateApplication::Get().GetRenderer()->GetFontCache();
+	
+		const float FontScale = 1.0f;
+		const FSlateFontInfo LegacyFontInfo = GetLegacySlateFontInfo();
+		FCharacterList& CharacterList = FontCache->GetCharacterList(LegacyFontInfo, FontScale);
+
+		const FCharacterEntry& Entry = CharacterList[InCh];
+		return Entry.HorizontalOffset;
+	}
+
+	return 0;
+}
+
 int32 UFont::GetStringSize( const TCHAR *Text ) const
 {
 	int32 Width, Height;
diff --git a/Engine/Source/Runtime/Engine/Private/UserInterface/Canvas.cpp b/Engine/Source/Runtime/Engine/Private/UserInterface/Canvas.cpp
--- a/Engine/Source/Runtime/Engine/Private/UserInterface/Canvas.cpp
+++ b/Engine/Source/Runtime/Engine/Private/UserInterface/Canvas.cpp
@@ -1186,7 +1186,8 @@ void UCanvas::MeasureStringInternal( FTextSizingParameters& Parameters, const TC
 				if( CharIndexFormat == ELastCharacterIndexFormat::CharacterAtOffset )
 				{
 					// Round our test toward the character's center position
-					if( StopAfterHorizontalOffset < Parameters.DrawXL - CharWidth / 2 )
+					const float TotalCharWidth = CharWidth + Parameters.DrawFont->GetCharHorizontalOffset(Ch);
+					if( StopAfterHorizontalOffset < Parameters.DrawXL - TotalCharWidth / 2 )
 					{
 						// We've reached the stopping point, so bail
 						break;
