Change update rate optimization so that when using Root Motion from Everything we dont have to tick every frame

diff --git a/Engine/Source/Editor/UnrealEd/Private/Animation/DebugSkelMeshComponent.cpp b/Engine/Source/Editor/UnrealEd/Private/Animation/DebugSkelMeshComponent.cpp
--- a/Engine/Source/Editor/UnrealEd/Private/Animation/DebugSkelMeshComponent.cpp
+++ b/Engine/Source/Editor/UnrealEd/Private/Animation/DebugSkelMeshComponent.cpp
@@ -172,7 +172,7 @@ void UDebugSkelMeshComponent::ConsumeRootMotion(const FVector& FloorMin, const F
 
 	if (UAnimInstance* AnimInst = GetAnimInstance())
 	{
-		ExtractedRootMotion = AnimInst->ConsumeExtractedRootMotion();
+		ExtractedRootMotion = AnimInst->ConsumeExtractedRootMotion(1.f);
 	}
 
 	if (bPreviewRootMotion)
diff --git a/Engine/Source/Runtime/Engine/Classes/Animation/AnimInstance.h b/Engine/Source/Runtime/Engine/Classes/Animation/AnimInstance.h
--- a/Engine/Source/Runtime/Engine/Classes/Animation/AnimInstance.h
+++ b/Engine/Source/Runtime/Engine/Classes/Animation/AnimInstance.h
@@ -741,7 +741,7 @@ public:
 	FAnimMontageInstance * GetRootMotionMontageInstance() const;
 
 	/** Get current accumulated root motion, removing it from the AnimInstance in the process */
-	FRootMotionMovementParams ConsumeExtractedRootMotion();
+	FRootMotionMovementParams ConsumeExtractedRootMotion(float Alpha);
 
 private:
 	/** Active Root Motion Montage Instance, if any. */
diff --git a/Engine/Source/Runtime/Engine/Classes/Animation/AnimationAsset.h b/Engine/Source/Runtime/Engine/Classes/Animation/AnimationAsset.h
--- a/Engine/Source/Runtime/Engine/Classes/Animation/AnimationAsset.h
+++ b/Engine/Source/Runtime/Engine/Classes/Animation/AnimationAsset.h
@@ -323,6 +323,22 @@ struct FRootMotionMovementParams
 			AccumulateWithBlend(FTransform(), WeightLeft);
 		}
 	}
+
+	FRootMotionMovementParams ConsumeRootMotion(float Alpha)
+	{
+		const ScalarRegister VAlpha(Alpha);
+		FTransform PartialRootMotion = (RootMotionTransform*VAlpha);
+		PartialRootMotion.SetScale3D(FVector(1.f));
+		PartialRootMotion.NormalizeRotation();
+		RootMotionTransform = RootMotionTransform.GetRelativeTransform(PartialRootMotion);
+
+		FRootMotionMovementParams ReturnParams;
+		ReturnParams.Set(PartialRootMotion);
+
+		check(PartialRootMotion.IsRotationNormalized());
+		check(RootMotionTransform.IsRotationNormalized());
+		return ReturnParams;
+	}
 };
 
 // This structure is used to either advance or synchronize animation players
diff --git a/Engine/Source/Runtime/Engine/Classes/Components/SkeletalMeshComponent.h b/Engine/Source/Runtime/Engine/Classes/Components/SkeletalMeshComponent.h
--- a/Engine/Source/Runtime/Engine/Classes/Components/SkeletalMeshComponent.h
+++ b/Engine/Source/Runtime/Engine/Classes/Components/SkeletalMeshComponent.h
@@ -6,6 +6,7 @@
 #include "Components/SkinnedMeshComponent.h"
 #include "PhysicsEngine/PhysicsConstraintComponent.h"
 #include "SkeletalMeshTypes.h"
+#include "Animation/AnimationAsset.h"
 #include "SkeletalMeshComponent.generated.h"
 
 class UAnimInstance;
@@ -854,7 +855,7 @@ public:
 	// Begin USkinnedMeshComponent interface
 	virtual bool UpdateLODStatus() override;
 	virtual void RefreshBoneTransforms( FActorComponentTickFunction* TickFunction = NULL ) override;
-	virtual void TickPose( float DeltaTime ) override;
+	virtual void TickPose(float DeltaTime, bool bNeedsValidRootMotion) override;
 	virtual void UpdateSlaveComponent() override;
 	virtual bool ShouldUpdateTransform(bool bLODHasChanged) const override;
 	virtual bool ShouldTickPose() const override;
@@ -867,6 +868,7 @@ public:
 	virtual FVector GetSkinnedVertexPosition(int32 VertexIndex) const override;
 
 	virtual bool IsPlayingRootMotion() override;
+	virtual bool IsPlayingRootMotionFromEverything() override;
 
 	// End USkinnedMeshComponent interface
 	/** 
@@ -1213,12 +1215,17 @@ private:
 	float DefaultPlayRate_DEPRECATED;
 
 public:
-	/** Keep track of when animation haa been ticked to ensure it is ticked only once per frame. */
+	/** Keep track of when animation has been ticked to ensure it is ticked only once per frame. */
 	UPROPERTY(Transient)
-	float LastTickTime;
+	float LastPoseTickTime;
+
+	/** Checked whether we have already ticked the pose this frame */
+	bool PoseTickedThisFrame() const { return LastPoseTickTime == GetWorld()->TimeSeconds; }
 
 	/** Take extracted RootMotion and convert it from local space to world space. */
 	FTransform ConvertLocalRootMotionToWorld(const FTransform& InTransform);
+
+	FRootMotionMovementParams ConsumeRootMotion();
 };
 
 
diff --git a/Engine/Source/Runtime/Engine/Classes/Components/SkinnedMeshComponent.h b/Engine/Source/Runtime/Engine/Classes/Components/SkinnedMeshComponent.h
--- a/Engine/Source/Runtime/Engine/Classes/Components/SkinnedMeshComponent.h
+++ b/Engine/Source/Runtime/Engine/Classes/Components/SkinnedMeshComponent.h
@@ -526,8 +526,9 @@ public:
 	 * @param DeltaTime DeltaTime
 	 *
 	 * @return	Return true if anything modified. Return false otherwise
+	 * @param bNeedsValidRootMotion - Networked games care more about this, but if false we can do less calculations
 	 */
-	virtual void TickPose( float DeltaTime ) {}
+	virtual void TickPose(float DeltaTime, bool bNeedsValidRootMotion) { TickUpdateRate(DeltaTime, bNeedsValidRootMotion); }
 
 	/** 
 	 * Update Slave Component. This gets called when MasterPoseComponent!=NULL
@@ -644,14 +645,8 @@ public:
 
 private:
 	/** Update Rate Optimization ticking. */
-	void TickUpdateRate();
-
-protected:
-	/** Accumulated DeltaTime skipped by Update Rate Optimization. */
-	UPROPERTY(Transient)
-	float SkippedTickDeltaTime;
-	/** check whether animated or not */
-	bool bPoseTicked;
+	void TickUpdateRate(float DeltaTime, bool bNeedsValidRootMotion);
+	
 public:
 	/**
 	 * Set MasterPoseComponent for this component
@@ -911,14 +906,13 @@ public:
 	/** Animation Update Rate optimization parameters. */
 	struct FAnimUpdateRateParameters* AnimUpdateRateParams;
 
-	/** Aimation Update Rate Tick. */
-	void AnimUpdateRateTick(uint8 UpdateRateShift);
-
 	/** Updates AnimUpdateRateParams, used by SkinnedMeshComponents.
 	* 
 	* @param bRecentlyRendered : true if at least one SkinnedMeshComponent on this Actor has been rendered in the last second.
 	* @param MaxDistanceFactor : Largest SkinnedMeshComponent of this Actor drawn on screen. */
-	void AnimUpdateRateSetParams(uint8 UpdateRateShift, const bool & bRecentlyRendered, const float& MaxDistanceFactor, const bool & bPlayingRootMotion);
+	void AnimUpdateRateSetParams(uint8 UpdateRateShift, float DeltaTime, const bool & bRecentlyRendered, const float& MaxDistanceFactor, const bool & bPlayingRootMotion);
 
 	virtual bool IsPlayingRootMotion(){ return false; }
+
+	virtual bool IsPlayingRootMotionFromEverything(){ return false; }
 };
diff --git a/Engine/Source/Runtime/Engine/Classes/Engine/EngineTypes.h b/Engine/Source/Runtime/Engine/Classes/Engine/EngineTypes.h
--- a/Engine/Source/Runtime/Engine/Classes/Engine/EngineTypes.h
+++ b/Engine/Source/Runtime/Engine/Classes/Engine/EngineTypes.h
@@ -1730,7 +1730,16 @@ struct FAnimUpdateRateParameters
 {
 	GENERATED_USTRUCT_BODY()
 
-private:
+public:
+	enum EOptimizeMode
+	{
+		TrailMode,
+		LookAheadMode,
+	};
+
+	/** Cache which Update Rate Optimization mode we are using */
+	EOptimizeMode OptimizeMode;
+
 	/** How often animation will be updated/ticked. 1 = every frame, 2 = every 2 frames, etc. */
 	UPROPERTY()
 	int32 UpdateRate;
@@ -1752,6 +1761,17 @@ private:
 	UPROPERTY()
 	bool bSkipEvaluation;
 
+	UPROPERTY(Transient)
+	/** Track time we have lost via skipping */
+	float TickedPoseOffestTime;
+
+	UPROPERTY(Transient)
+	/** Total time of the last series of skipped updates */
+	float AdditionalTime;
+
+	/** The delta time of the last tick */
+	float ThisTickDelta;
+
 public:
 
 	/** Default constructor. */
@@ -1761,26 +1781,28 @@ public:
 		, bInterpolateSkippedFrames(false)
 		, bSkipUpdate(false)
 		, bSkipEvaluation(false)
+		, TickedPoseOffestTime(0.f)
+		, AdditionalTime(0.f)
+		, ThisTickDelta(0.f)
 	{ }
 
-	/** Set parameters and verify inputs.
+	/** Set parameters and verify inputs for Trail Mode (original behaviour - skip frames, track skipped time and then catch up afterwards).
 	 * @param : UpdateShiftRate. Shift our update frames so that updates across all skinned components are staggered
 	 * @param : NewUpdateRate. How often animation will be updated/ticked. 1 = every frame, 2 = every 2 frames, etc.
 	 * @param : NewEvaluationRate. How often animation will be evaluated. 1 = every frame, 2 = every 2 frames, etc.
 	 * @param : bNewInterpSkippedFrames. When skipping a frame, should it be interpolated or frozen?
 	 */
-	void Set(uint8 UpdateRateShift, const int32& NewUpdateRate, const int32& NewEvaluationRate, const bool & bNewInterpSkippedFrames);
+	void SetTrailMode(float DeltaTime, uint8 UpdateRateShift, int32 NewUpdateRate, int32 NewEvaluationRate, bool bNewInterpSkippedFrames);
 
-	/* Getter for UpdateRate */
-	int32 GetUpdateRate() const
-	{
-		return UpdateRate;
-	}
+	void SetLookAheadMode(float DeltaTime, uint8 UpdateRateShift, float LookAheadAmount);
+
+	float GetInterpolationAlpha() const;
+
+	float GetRootMotionInterp() const;
 
-	/* Getter for EvaluationRate */
-	int32 GetEvaluationRate() const
+	bool DoEvaluationRateOptimizations() const
 	{
-		return EvaluationRate;
+		return OptimizeMode == LookAheadMode || EvaluationRate > 1;
 	}
 
 	/* Getter for bSkipUpdate */
@@ -1800,6 +1822,34 @@ public:
 	{
 		return bInterpolateSkippedFrames;
 	}
+
+	/** Called when we are ticking a pose to make sure we accumulate all needed time */
+	float GetTimeAdjustment()
+	{
+		return AdditionalTime;
+	}
+
+	FColor GetUpdateRateDebugColor() const
+	{
+		if (OptimizeMode == TrailMode)
+		{
+			/*switch (UpdateRate)
+			{
+			case 1: return FColor::Red;
+			case 2: return FColor::Green;
+			case 3: return FColor::Blue;
+			}*/
+			return FColor::Black;
+		}
+		else
+		{
+			if (bSkipUpdate)
+			{
+				return FColor::Yellow;
+			}
+			return FColor::Green;
+		}
+	}
 };
 
 
diff --git a/Engine/Source/Runtime/Engine/Private/Animation/AnimInstance.cpp b/Engine/Source/Runtime/Engine/Private/Animation/AnimInstance.cpp
--- a/Engine/Source/Runtime/Engine/Private/Animation/AnimInstance.cpp
+++ b/Engine/Source/Runtime/Engine/Private/Animation/AnimInstance.cpp
@@ -2414,11 +2414,22 @@ FAnimMontageInstance * UAnimInstance::GetRootMotionMontageInstance() const
 	return RootMotionMontageInstance;
 }
 
-FRootMotionMovementParams UAnimInstance::ConsumeExtractedRootMotion()
+FRootMotionMovementParams UAnimInstance::ConsumeExtractedRootMotion(float Alpha)
 {
-	FRootMotionMovementParams RootMotion = ExtractedRootMotion;
-	ExtractedRootMotion.Clear();
-	return RootMotion;
+	if (Alpha < ZERO_ANIMWEIGHT_THRESH)
+	{
+		return FRootMotionMovementParams();
+	}
+	else if (Alpha > (1.f - ZERO_ANIMWEIGHT_THRESH))
+	{
+		FRootMotionMovementParams RootMotion = ExtractedRootMotion;
+		ExtractedRootMotion.Clear();
+		return RootMotion;
+	}
+	else
+	{
+		return ExtractedRootMotion.ConsumeRootMotion(Alpha);
+	}
 }
 
 void UAnimInstance::SetMorphTarget(FName MorphTargetName, float Value)
diff --git a/Engine/Source/Runtime/Engine/Private/CharacterMovementComponent.cpp b/Engine/Source/Runtime/Engine/Private/CharacterMovementComponent.cpp
--- a/Engine/Source/Runtime/Engine/Private/CharacterMovementComponent.cpp
+++ b/Engine/Source/Runtime/Engine/Private/CharacterMovementComponent.cpp
@@ -6733,14 +6733,15 @@ void UCharacterMovementComponent::TickCharacterPose(float DeltaTime)
 {
 	check(CharacterOwner && CharacterOwner->GetMesh());
 
-	CharacterOwner->GetMesh()->TickPose(DeltaTime);
+	CharacterOwner->GetMesh()->TickPose(DeltaTime, true);
 
 	// Grab root motion now that we have ticked the pose
 	if (CharacterOwner->IsPlayingRootMotion())
 	{
-		if (UAnimInstance * AnimInstance = CharacterOwner->GetMesh()->GetAnimInstance())
+		FRootMotionMovementParams RootMotion = CharacterOwner->GetMesh()->ConsumeRootMotion();
+		if (RootMotion.bHasRootMotion)
 		{
-			RootMotionParams.Accumulate(AnimInstance->ConsumeExtractedRootMotion());
+			RootMotionParams.Accumulate(RootMotion);
 		}
 	}
 }
diff --git a/Engine/Source/Runtime/Engine/Private/SkeletalMeshComponent.cpp b/Engine/Source/Runtime/Engine/Private/SkeletalMeshComponent.cpp
--- a/Engine/Source/Runtime/Engine/Private/SkeletalMeshComponent.cpp
+++ b/Engine/Source/Runtime/Engine/Private/SkeletalMeshComponent.cpp
@@ -551,14 +551,21 @@ bool USkeletalMeshComponent::ShouldTickPose() const
 {
 	// When we stop root motion we go back to ticking after CharacterMovement. Unfortunately that means that we could tick twice that frame.
 	// So only enforce a single tick per frame.
-	const bool bAlreadyTickedThisFrame = (LastTickTime == GetWorld()->TimeSeconds);
+	const bool bAlreadyTickedThisFrame = PoseTickedThisFrame();
 	return (Super::ShouldTickPose() && IsRegistered() && AnimScriptInstance && !bAutonomousTickPose && !bPauseAnims && GetWorld()->AreActorsInitialized() && !bNoSkeletonUpdate && !bAlreadyTickedThisFrame);
 }
 
-void USkeletalMeshComponent::TickPose(float DeltaTime)
+void USkeletalMeshComponent::TickPose(float DeltaTime, bool bNeedsValidRootMotion)
 {
-	TickAnimation( DeltaTime );
-	LastTickTime = GetWorld()->TimeSeconds;
+	Super::TickPose(DeltaTime, bNeedsValidRootMotion);
+
+	if (!bEnableUpdateRateOptimizations || !AnimUpdateRateParams->ShouldSkipUpdate())
+	{
+		float TimeAdjustment = AnimUpdateRateParams->GetTimeAdjustment();
+		TickAnimation(DeltaTime + TimeAdjustment);
+		LastPoseTickTime = GetWorld()->TimeSeconds;
+	}
+	//return bPoseTicked;
 }
 
 void USkeletalMeshComponent::TickComponent(float DeltaTime, enum ELevelTick TickType, FActorComponentTickFunction *ThisTickFunction)
@@ -922,6 +929,11 @@ void USkeletalMeshComponent::PerformAnimationEvaluation(const USkeletalMesh* InS
 	FillSpaceBases(InSkeletalMesh, OutLocalAtoms, OutSpaceBases);
 }
 
+const TCHAR* B(bool b)
+{
+	return b ? TEXT("true") : TEXT("false");
+}
+
 void USkeletalMeshComponent::RefreshBoneTransforms(FActorComponentTickFunction* TickFunction)
 {
 	SCOPE_CYCLE_COUNTER(STAT_RefreshBoneTransforms);
@@ -945,19 +957,19 @@ void USkeletalMeshComponent::RefreshBoneTransforms(FActorComponentTickFunction*
 	AnimEvaluationContext.AnimInstance = AnimScriptInstance;
 
 	//Handle update rate optimization setup
-	const bool bDoUpdateRateOptimization = bEnableUpdateRateOptimizations && (AnimUpdateRateParams->GetEvaluationRate() > 1);
+	const bool bDoEvaluationRateOptimization = bEnableUpdateRateOptimizations && AnimUpdateRateParams->DoEvaluationRateOptimizations();
 	//Dont mark cache as invalid if we aren't performing optimization anyway
-	const bool bInvalidCachedBones = bDoUpdateRateOptimization &&
+	const bool bInvalidCachedBones = bDoEvaluationRateOptimization &&
 									( (LocalAtoms.Num() != SkeletalMesh->RefSkeleton.GetNum())
 									  || (LocalAtoms.Num() != CachedLocalAtoms.Num())
 									  || (GetNumSpaceBases() != CachedSpaceBases.Num()) );
 
-	AnimEvaluationContext.bDoEvaluation = !bDoUpdateRateOptimization || bInvalidCachedBones || !AnimUpdateRateParams->ShouldSkipEvaluation();
+	AnimEvaluationContext.bDoEvaluation = !bDoEvaluationRateOptimization || bInvalidCachedBones || !AnimUpdateRateParams->ShouldSkipEvaluation();
 	
-	AnimEvaluationContext.bDoInterpolation = bDoUpdateRateOptimization && !bInvalidCachedBones && AnimUpdateRateParams->ShouldInterpolateSkippedFrames();
-	AnimEvaluationContext.bDuplicateToCacheBones = bInvalidCachedBones || (bDoUpdateRateOptimization && AnimEvaluationContext.bDoEvaluation && !AnimEvaluationContext.bDoInterpolation);
+	AnimEvaluationContext.bDoInterpolation = bDoEvaluationRateOptimization && !bInvalidCachedBones && AnimUpdateRateParams->ShouldInterpolateSkippedFrames();
+	AnimEvaluationContext.bDuplicateToCacheBones = bInvalidCachedBones || (bDoEvaluationRateOptimization && AnimEvaluationContext.bDoEvaluation && !AnimEvaluationContext.bDoInterpolation);
 
-	if (!bDoUpdateRateOptimization)
+	if (!bDoEvaluationRateOptimization)
 	{
 		//If we aren't optimizing clear the cached local atoms
 		CachedLocalAtoms.Empty();
@@ -1029,7 +1041,7 @@ void USkeletalMeshComponent::PostAnimEvaluation(FAnimationEvaluationContext& Eva
 	{
 		SCOPE_CYCLE_COUNTER(STAT_InterpolateSkippedFrames);
 
-		const float Alpha = 0.25f + (1.f / float(FMath::Max(AnimUpdateRateParams->GetEvaluationRate(), 2) * 2));
+		const float Alpha = AnimUpdateRateParams->GetInterpolationAlpha();
 		FAnimationRuntime::LerpBoneTransforms(LocalAtoms, CachedLocalAtoms, Alpha, RequiredBones);
 		if (bDoubleBufferedBlendSpaces)
 		{
@@ -1039,6 +1051,7 @@ void USkeletalMeshComponent::PostAnimEvaluation(FAnimationEvaluationContext& Eva
 		}
 		FAnimationRuntime::LerpBoneTransforms(GetEditableSpaceBases(), CachedSpaceBases, Alpha, RequiredBones);
 	}
+
 	bNeedToFlipSpaceBaseBuffers = true;
 
 	// Transforms updated, cached local bounds are now out of date.
@@ -1707,6 +1720,16 @@ FTransform USkeletalMeshComponent::ConvertLocalRootMotionToWorld(const FTransfor
 	return DeltaWorldTransform;
 }
 
+FRootMotionMovementParams USkeletalMeshComponent::ConsumeRootMotion()
+{
+	if (AnimScriptInstance)
+	{
+		float InterpAlpha = bEnableUpdateRateOptimizations ? AnimUpdateRateParams->GetRootMotionInterp() : 1.f;
+		return AnimScriptInstance->ConsumeExtractedRootMotion(InterpAlpha);
+	}
+	return FRootMotionMovementParams();
+}
+
 
 float USkeletalMeshComponent::CalculateMass(FName BoneName)
 {
@@ -1926,6 +1949,11 @@ bool USkeletalMeshComponent::IsPlayingRootMotion()
 	return (AnimScriptInstance ? (AnimScriptInstance->GetRootMotionMontageInstance() != NULL) : false);
 }
 
+bool USkeletalMeshComponent::IsPlayingRootMotionFromEverything()
+{
+	return AnimScriptInstance ? (AnimScriptInstance->RootMotionMode == ERootMotionMode::RootMotionFromEverything) : false;
+}
+
 void USkeletalMeshComponent::SetRootBodyIndex(int32 InBodyIndex)
 {
 	RootBodyData.BodyIndex = InBodyIndex;
diff --git a/Engine/Source/Runtime/Engine/Private/SkeletalMeshComponentPhysics.cpp b/Engine/Source/Runtime/Engine/Private/SkeletalMeshComponentPhysics.cpp
--- a/Engine/Source/Runtime/Engine/Private/SkeletalMeshComponentPhysics.cpp
+++ b/Engine/Source/Runtime/Engine/Private/SkeletalMeshComponentPhysics.cpp
@@ -3619,7 +3619,7 @@ void USkeletalMeshComponent::TickClothing(float DeltaTime)
 {
 #if WITH_APEX_CLOTHING
 	// animated but bone transforms were not updated because it was not rendered
-	if(bPoseTicked && !bRecentlyRendered)
+	if(PoseTickedThisFrame() && !bRecentlyRendered)
 	{
 		ForceClothNextUpdateTeleportAndReset();
 	}
diff --git a/Engine/Source/Runtime/Engine/Private/SkinnedMeshComponent.cpp b/Engine/Source/Runtime/Engine/Private/SkinnedMeshComponent.cpp
--- a/Engine/Source/Runtime/Engine/Private/SkinnedMeshComponent.cpp
+++ b/Engine/Source/Runtime/Engine/Private/SkinnedMeshComponent.cpp
@@ -57,6 +57,13 @@ namespace FAnimUpdateRateManager
 
 			return AnimUpdateRateShiftTag;
 		}
+
+		bool IsHumanControlled() const
+		{
+			AActor* Owner = RegisteredComponents[0]->GetOwner();
+			AController * Controller = Owner ? Owner->GetInstigatorController() : NULL;
+			return Cast<APlayerController>(Controller) != NULL;
+		}
 	};
 
 
@@ -105,7 +112,87 @@ namespace FAnimUpdateRateManager
 		}
 	}
 
-	void TickUpdateRateParameters(USkinnedMeshComponent* SkinnedComponent)
+	void AnimUpdateRateSetParams(FAnimUpdateRateParametersTracker* Tracker, float DeltaTime, bool bRecentlyRendered, float MaxDistanceFactor, bool bNeedsValidRootMotion, bool bUsingRootMotionFromEverything)
+	{
+		// default rules for setting update rates
+
+		// Human controlled characters should be ticked always fully to minimize latency w/ game play events triggered by animation.
+		const bool bHumanControlled = Tracker->IsHumanControlled();
+
+		bool bNeedsEveryFrame = bNeedsValidRootMotion && !bUsingRootMotionFromEverything;
+
+		// Not rendered, including dedicated servers. we can skip the Evaluation part.
+		if (!bRecentlyRendered)
+		{
+			Tracker->UpdateRateParameters.SetTrailMode(DeltaTime, Tracker->GetAnimUpdateRateShiftTag(), (bHumanControlled ? 1 : 4), 4, false);
+		}
+		// Visible controlled characters or playing root motion. Need evaluation and ticking done every frame.
+		else  if (bHumanControlled || bNeedsEveryFrame)
+		{
+			Tracker->UpdateRateParameters.SetTrailMode(DeltaTime, Tracker->GetAnimUpdateRateShiftTag(), 1, 1, false);
+		}
+		else
+		{
+			// For visible meshes, figure out how often bones should be evaluated VS interpolated to previous evaluation.
+			// This is based on screen size and makes sense for a 3D FPS game, different games or Actors can implement different rules.
+			int32 DesiredEvaluationRate;
+			if (MaxDistanceFactor > 0.4f)
+			{
+				DesiredEvaluationRate = 1;
+			}
+			else if (MaxDistanceFactor > .2f)
+			{
+				DesiredEvaluationRate = 2;
+			}
+			else
+			{
+				DesiredEvaluationRate = 3;
+			}
+
+			if (bUsingRootMotionFromEverything && DesiredEvaluationRate > 1)
+			{
+				//Use look ahead mode that allows us to rate limit updates even when using root motion
+				Tracker->UpdateRateParameters.SetLookAheadMode(DeltaTime, Tracker->GetAnimUpdateRateShiftTag(), DeltaTime*DesiredEvaluationRate);
+			}
+			else
+			{
+				Tracker->UpdateRateParameters.SetTrailMode(DeltaTime, Tracker->GetAnimUpdateRateShiftTag(), DesiredEvaluationRate, DesiredEvaluationRate, true);
+			}
+		}
+	}
+
+	void AnimUpdateRateTick(FAnimUpdateRateParametersTracker* Tracker, float DeltaTime, bool bNeedsValidRootMotion)
+	{
+		const float RecentlyRenderedTime = (Tracker->RegisteredComponents[0]->GetWorld()->TimeSeconds - 1.f);
+
+		// Go through components and figure out if they've been recently rendered, and the biggest MaxDistanceFactor
+		bool bRecentlyRendered = false;
+		bool bPlayingRootMotion = false;
+		bool bUsingRootMotionFromEverything = true;
+		float MaxDistanceFactor = 0.f;
+
+		const TArray<USkinnedMeshComponent*>& SkinnedComponents = Tracker->RegisteredComponents;
+
+		for (USkinnedMeshComponent* Component : SkinnedComponents)
+		{
+			bRecentlyRendered |= (Component->LastRenderTime > RecentlyRenderedTime);
+			MaxDistanceFactor = FMath::Max(MaxDistanceFactor, Component->MaxDistanceFactor);
+			bPlayingRootMotion |= Component->IsPlayingRootMotion();
+			bUsingRootMotionFromEverything &= Component->IsPlayingRootMotionFromEverything();
+		}
+
+		bNeedsValidRootMotion &= bPlayingRootMotion;
+
+		// Figure out which update rate should be used.
+		AnimUpdateRateSetParams(Tracker, DeltaTime, bRecentlyRendered, MaxDistanceFactor, bNeedsValidRootMotion, bUsingRootMotionFromEverything);
+	}
+
+	const TCHAR* B(bool b)
+	{
+		return b ? TEXT("true") : TEXT("false");
+	}
+
+	void TickUpdateRateParameters(USkinnedMeshComponent* SkinnedComponent, float DeltaTime, bool bNeedsValidRootMotion)
 	{
 		// Convert current frame counter from 64 to 32 bits.
 		const uint32 CurrentFrame32 = uint32(GFrameCounter % MAX_uint32);
@@ -116,17 +203,9 @@ namespace FAnimUpdateRateManager
 		if (CurrentFrame32 != Tracker->AnimUpdateRateFrameCount)
 		{
 			Tracker->AnimUpdateRateFrameCount = CurrentFrame32;
-			SkinnedComponent->AnimUpdateRateTick(Tracker->GetAnimUpdateRateShiftTag());
+			AnimUpdateRateTick(Tracker, DeltaTime, bNeedsValidRootMotion);
 		}
 	}
-
-	const TArray<USkinnedMeshComponent*> GetSharedMeshComponents(USkinnedMeshComponent* SkinnedComponent)
-	{
-		UObject* TrackerIndex = GetMapIndexForComponent(SkinnedComponent);
-		FAnimUpdateRateParametersTracker* Tracker = ActorToUpdateRateParams.FindChecked(TrackerIndex);
-		return Tracker->RegisteredComponents;
-	}
-
 }
 
 void USkinnedMeshComponent::OnRegister()
@@ -323,7 +402,7 @@ bool USkinnedMeshComponent::ShouldUpdateTransform(bool bLODHasChanged) const
 }
 
 
-void USkinnedMeshComponent::TickUpdateRate()
+void USkinnedMeshComponent::TickUpdateRate(float DeltaTime, bool bNeedsValidRootMotion)
 {
 	SCOPE_CYCLE_COUNTER(STAT_TickUpdateRate);
 	if( bEnableUpdateRateOptimizations )
@@ -331,19 +410,12 @@ void USkinnedMeshComponent::TickUpdateRate()
 		if (GetOwner())
 		{
 			// Tick Owner once per frame. All attached SkinnedMeshComponents will share the same settings.
-			FAnimUpdateRateManager::TickUpdateRateParameters(this);
+			FAnimUpdateRateManager::TickUpdateRateParameters(this, DeltaTime, bNeedsValidRootMotion);
 
 			// debug -- @todo: hook this up to a console command.
 			if (bDisplayDebugUpdateRateOptimizations)
 			{
-				FColor DrawColor;
-				switch (AnimUpdateRateParams->GetUpdateRate())
-				{
-				case 1: DrawColor = FColor::Red; break;
-				case 2: DrawColor = FColor::Green; break;
-				case 3: DrawColor = FColor::Blue; break;
-				default: DrawColor = FColor::Black; break;
-				}
+				FColor DrawColor = AnimUpdateRateParams->GetUpdateRateDebugColor();
 				DrawDebugBox(GetWorld(), Bounds.Origin, Bounds.BoxExtent, FQuat::Identity, DrawColor, false);
 			}
 		}
@@ -359,30 +431,11 @@ void USkinnedMeshComponent::TickComponent(float DeltaTime, enum ELevelTick TickT
 
 	// See if this mesh was rendered recently. This has to happen first because other data will rely on this
 	bRecentlyRendered = (LastRenderTime > GetWorld()->TimeSeconds - 1.0f);
-	TickUpdateRate();
 
-	bPoseTicked = false;
 	// Tick Pose first
-	if( ShouldTickPose() )
+	if (ShouldTickPose())
 	{
-		if (bEnableUpdateRateOptimizations && AnimUpdateRateParams->ShouldSkipUpdate())
-		{
-			SkippedTickDeltaTime += DeltaTime;
-			if( !bRecentlyRendered )
-			{
-				return;
-			}
-		}
-		else
-		{
-			TickPose(DeltaTime + SkippedTickDeltaTime);
-			bPoseTicked = true;
-			SkippedTickDeltaTime = 0.f;
-		}
-	}
-	else
-	{
-		SkippedTickDeltaTime = 0.f;
+		TickPose(DeltaTime, false);
 	}
 
 	// Update component's LOD settings
@@ -1847,69 +1900,6 @@ TArray<FActiveVertexAnim> USkinnedMeshComponent::UpdateActiveVertexAnims(const U
 	return OutVertexAnims;
 }
 
-void USkinnedMeshComponent::AnimUpdateRateTick(uint8 UpdateRateShift)
-{
-	const float RecentlyRenderedTime = (GetWorld()->TimeSeconds - 1.f);
-
-	// Go through components and figure out if they've been recently rendered, and the biggest MaxDistanceFactor
-	bool bRecentlyRendered = false;
-	bool bPlayingRootMotion = false;
-	float MaxDistanceFactor = 0.f;
-
-	const TArray<USkinnedMeshComponent*> SkinnedComponents = FAnimUpdateRateManager::GetSharedMeshComponents(this);
-
-	for (USkinnedMeshComponent* Component : SkinnedComponents)
-	{
-		bRecentlyRendered = bRecentlyRendered || (Component->LastRenderTime > RecentlyRenderedTime);
-		MaxDistanceFactor = FMath::Max(MaxDistanceFactor, Component->MaxDistanceFactor);
-		bPlayingRootMotion = bPlayingRootMotion || Component->IsPlayingRootMotion();
-	}
-
-	// Figure out which update rate should be used.
-	AnimUpdateRateSetParams(UpdateRateShift, bRecentlyRendered, MaxDistanceFactor, bPlayingRootMotion);
-}
-
-void USkinnedMeshComponent::AnimUpdateRateSetParams(uint8 UpdateRateShift, const bool & bRecentlyRendered, const float& MaxDistanceFactor, const bool & bPlayingRootMotion)
-{
-	// default rules for setting update rates
-
-	// Human controlled characters should be ticked always fully to minimize latency w/ game play events triggered by animation.
-	AActor* Owner = GetOwner();
-	AController * Controller = Owner ? Owner->GetInstigatorController() : NULL;
-	const bool bHumanControlled = Controller && (Cast<APlayerController>(Controller) != NULL);
-
-	// Not rendered, including dedicated servers. we can skip the Evaluation part.
-	if (!bRecentlyRendered)
-	{
-		AnimUpdateRateParams->Set(UpdateRateShift, (bHumanControlled ? 1 : 4), 4, false);
-	}
-	// Visible controlled characters or playing root motion. Need evaluation and ticking done every frame.
-	else  if (bHumanControlled || bPlayingRootMotion)
-	{
-		AnimUpdateRateParams->Set(UpdateRateShift, 1, 1, false);
-	}
-	else
-	{
-		// For visible meshes, figure out how often bones should be evaluated VS interpolated to previous evaluation.
-		// This is based on screen size and makes sense for a 3D FPS game, different games or Actors can implement different rules.
-		int32 DesiredEvaluationRate;
-		if (MaxDistanceFactor > 0.4f)
-		{
-			DesiredEvaluationRate = 1;
-		}
-		else if (MaxDistanceFactor > .2f)
-		{
-			DesiredEvaluationRate = 2;
-		}
-		else
-		{
-			DesiredEvaluationRate = 3;
-		}
-
-		AnimUpdateRateParams->Set(UpdateRateShift, DesiredEvaluationRate, DesiredEvaluationRate, true);
-	}
-}
-
 void USkinnedMeshComponent::FlipEditableSpaceBases()
 {
 	if (bNeedToFlipSpaceBaseBuffers)
@@ -1941,8 +1931,11 @@ void USkinnedMeshComponent::SetSpaceBaseDoubleBuffering(bool bInDoubleBufferedBl
 	}
 }
 
-void FAnimUpdateRateParameters::Set(uint8 UpdateRateShift, const int32& NewUpdateRate, const int32& NewEvaluationRate, const bool & bNewInterpSkippedFrames)
+void FAnimUpdateRateParameters::SetTrailMode(float DeltaTime, uint8 UpdateRateShift, int32 NewUpdateRate, int32 NewEvaluationRate, bool bNewInterpSkippedFrames)
 {
+	OptimizeMode = TrailMode;
+	ThisTickDelta = DeltaTime;
+
 	UpdateRate = FMath::Max(NewUpdateRate, 1);
 	// Make sure EvaluationRate is a multiple of UpdateRate.
 	EvaluationRate = FMath::Max((NewEvaluationRate / UpdateRate) * UpdateRate, 1);
@@ -1953,5 +1946,76 @@ void FAnimUpdateRateParameters::Set(uint8 UpdateRateShift, const int32& NewUpdat
 
 	bSkipUpdate = ((Counter % UpdateRate) > 0);
 	bSkipEvaluation = ((Counter % EvaluationRate) > 0);
+
+	AdditionalTime = 0.f;
+
+	if (bSkipUpdate)
+	{
+		TickedPoseOffestTime -= DeltaTime;
+	}
+	else
+	{
+		if (TickedPoseOffestTime < 0.f)
+		{
+			AdditionalTime = -TickedPoseOffestTime;
+			TickedPoseOffestTime = 0.f;
+		}
+	}
 }
 
+void FAnimUpdateRateParameters::SetLookAheadMode(float DeltaTime, uint8 UpdateRateShift, float LookAheadAmount)
+{
+	float OriginalTickedPoseOffestTime = TickedPoseOffestTime;
+	if (OptimizeMode == TrailMode)
+	{
+		TickedPoseOffestTime = 0.f;
+	}
+	OptimizeMode = LookAheadMode;
+	ThisTickDelta = DeltaTime;
+
+	bInterpolateSkippedFrames = true;
+
+	TickedPoseOffestTime -= DeltaTime;
+
+	if (TickedPoseOffestTime < 0.f)
+	{
+		AdditionalTime = LookAheadAmount;
+		TickedPoseOffestTime += LookAheadAmount;
+
+		bool bValid = (TickedPoseOffestTime > 0.f);
+		if (!bValid)
+		{
+			FPlatformMisc::LowLevelOutputDebugStringf(TEXT("TPO Time: %.3f | Orig TPO Time: %.3f | DT: %.3f | LookAheadAmount: %.3f\n"), TickedPoseOffestTime, OriginalTickedPoseOffestTime, DeltaTime, LookAheadAmount);
+		}
+		check(bValid);
+		bSkipUpdate = bSkipEvaluation = false;
+	}
+	else
+	{
+		AdditionalTime = 0.f;
+		bSkipUpdate = bSkipEvaluation = true;
+	}
+}
+
+float FAnimUpdateRateParameters::GetInterpolationAlpha() const
+{
+	if (OptimizeMode == TrailMode)
+	{
+		return 0.25f + (1.f / float(FMath::Max(EvaluationRate, 2) * 2));
+	}
+	else if (OptimizeMode == LookAheadMode)
+	{
+		return FMath::Clamp(ThisTickDelta / (TickedPoseOffestTime + ThisTickDelta), 0.f, 1.f);
+	}
+	check(false); // Unknown mode
+	return 0.f;
+}
+
+float FAnimUpdateRateParameters::GetRootMotionInterp() const
+{
+	if (OptimizeMode == LookAheadMode)
+	{
+		return FMath::Clamp(ThisTickDelta / (TickedPoseOffestTime + ThisTickDelta), 0.f, 1.f);
+	}
+	return 1.f;
+}
\ No newline at end of file
