[INTEGRATE] Change 2420555 by [REDACTED_PERSON][REDACTED_HANDLE] on 2015/01/27 12:08:01

	[UE-8397] Making ActorComponent a Blueprint type (for blueprintable components).

diff --git a/Engine/Source/Runtime/Engine/Classes/Components/ActorComponent.h b/Engine/Source/Runtime/Engine/Classes/Components/ActorComponent.h
--- a/Engine/Source/Runtime/Engine/Classes/Components/ActorComponent.h
+++ b/Engine/Source/Runtime/Engine/Classes/Components/ActorComponent.h
@@ -24,7 +24,7 @@ enum class EComponentCreationMethod : uint8
  * @see USceneComponent
  * @see UPrimitiveComponent
  */
-UCLASS(DefaultToInstanced, abstract, hidecategories=(ComponentReplication))
+UCLASS(DefaultToInstanced, BlueprintType, abstract, hidecategories=(ComponentReplication))
 class ENGINE_API UActorComponent : public UObject, public IInterface_AssetUserData
 {
 	GENERATED_BODY()
