[INTEGRATE] Change 2420826 by [REDACTED_PERSON][REDACTED_HANDLE] on 2015/01/27 14:39:47

	UE-8104 If one or more elements are set in the SensesConfig property of an AIController class containing the AIPerceptionComponent, the elements are reset to None when the Blueprint is compiled.

	EditInline Instanced subobjects in components are supported.

diff --git a/Engine/Source/Runtime/Engine/Private/BlueprintGeneratedClass.cpp b/Engine/Source/Runtime/Engine/Private/BlueprintGeneratedClass.cpp
--- a/Engine/Source/Runtime/Engine/Private/BlueprintGeneratedClass.cpp
+++ b/Engine/Source/Runtime/Engine/Private/BlueprintGeneratedClass.cpp
@@ -41,7 +41,8 @@ void UBlueprintGeneratedClass::PostLoad()
 	for( auto SubObjIt = SubObjects.CreateIterator(); SubObjIt; ++SubObjIt )
 	{
 		UObject* CurrObj = *SubObjIt;
-		if( !CurrObj->IsDefaultSubobject() && !CurrObj->IsRooted() )
+		const bool bComponentChild = (CurrObj->GetOuter() != ClassCDO) && CurrObj->GetOuter()->IsDefaultSubobject();
+		if (!CurrObj->IsDefaultSubobject() && !CurrObj->IsRooted() && !bComponentChild)
 		{
 			CurrObj->MarkPendingKill();
 		}
diff --git a/Engine/Source/Runtime/Engine/Private/UnrealEngine.cpp b/Engine/Source/Runtime/Engine/Private/UnrealEngine.cpp
--- a/Engine/Source/Runtime/Engine/Private/UnrealEngine.cpp
+++ b/Engine/Source/Runtime/Engine/Private/UnrealEngine.cpp
@@ -10260,6 +10260,39 @@ struct FFindInstancedReferenceSubobjectHelper
 			}
 		}
 	}
+
+	static void Duplicate(UObject* OldObject, UObject* NewObject, TMap<UObject*, UObject*>& ReferenceReplacementMap, TArray<UObject*>& DuplicatedObjects)
+	{
+		if (OldObject->GetClass()->HasAnyClassFlags(CLASS_HasInstancedReference) &&
+			NewObject->GetClass()->HasAnyClassFlags(CLASS_HasInstancedReference))
+		{
+			TSet<UObject*> OldEditInlineObjects;
+			Get(OldObject->GetClass(), reinterpret_cast<uint8*>(OldObject), OldEditInlineObjects);
+			if (OldEditInlineObjects.Num())
+			{
+				TSet<UObject*> NewEditInlineObjects;
+				Get(NewObject->GetClass(), reinterpret_cast<uint8*>(NewObject), NewEditInlineObjects);
+				for (auto Obj : NewEditInlineObjects)
+				{
+					const bool bProperOuter = (Obj->GetOuter() == OldObject);
+					const bool bEditInlineNew = Obj->GetClass()->HasAnyClassFlags(CLASS_EditInlineNew | CLASS_DefaultToInstanced);
+					if (bProperOuter && bEditInlineNew)
+					{
+						const bool bKeptByOld = OldEditInlineObjects.Contains(Obj);
+						const bool bNotHandledYet = !ReferenceReplacementMap.Contains(Obj);
+						if (bKeptByOld && bNotHandledYet)
+						{
+							UObject* NewEditInlineSubobject = StaticDuplicateObject(Obj, NewObject, NULL);
+							ReferenceReplacementMap.Add(Obj, NewEditInlineSubobject);
+
+							// We also need to make sure to fixup any properties here
+							DuplicatedObjects.Add(NewEditInlineSubobject);
+						}
+					}
+				}
+			}
+		}
+	}
 };
 
 void UEngine::CopyPropertiesForUnrelatedObjects(UObject* OldObject, UObject* NewObject, FCopyPropertiesForUnrelatedObjectsParams Params)
@@ -10328,6 +10361,8 @@ void UEngine::CopyPropertiesForUnrelatedObjects(UObject* OldObject, UObject* New
 
 	TArray<UObject*> ComponentsOnNewObject;
 	{
+		TArray<UObject*> EditInlineSubobjectsOfComponents;
+
 		// Find all instanced objects of the old CDO, and save off their modified properties to be later applied to the newly instanced objects of the new CDO
 		NewObject->CollectDefaultSubobjects(ComponentsOnNewObject,true);
 
@@ -10376,6 +10411,7 @@ void UEngine::CopyPropertiesForUnrelatedObjects(UObject* OldObject, UObject* New
 					}
 				}
 				FObjectReader Reader(NewInstance, Record.SavedProperties, true, true);
+				FFindInstancedReferenceSubobjectHelper::Duplicate(Record.OldInstance, NewInstance, ReferenceReplacementMap, EditInlineSubobjectsOfComponents);
 			}
 			else
 			{
@@ -10396,38 +10432,11 @@ void UEngine::CopyPropertiesForUnrelatedObjects(UObject* OldObject, UObject* New
 				}
 			}
 		}
+		ComponentsOnNewObject.Append(EditInlineSubobjectsOfComponents);
 	}
 
-	if (OldObject->GetClass()->HasAnyClassFlags(CLASS_HasInstancedReference) &&
-		NewObject->GetClass()->HasAnyClassFlags(CLASS_HasInstancedReference))
-	{
-		TSet<UObject*> OldEditInlineObjects;
-		FFindInstancedReferenceSubobjectHelper::Get(OldObject->GetClass(), reinterpret_cast<uint8*>(OldObject), OldEditInlineObjects);
-		if (OldEditInlineObjects.Num())
-		{
-			TSet<UObject*> NewEditInlineObjects;
-			FFindInstancedReferenceSubobjectHelper::Get(NewObject->GetClass(), reinterpret_cast<uint8*>(NewObject), NewEditInlineObjects);
-			for (auto Obj : NewEditInlineObjects)
-			{
-				const bool bProperOuter = (Obj->GetOuter() == OldObject);
-				const bool bEditInlineNew = Obj->GetClass()->HasAnyClassFlags(CLASS_EditInlineNew | CLASS_DefaultToInstanced);
-				if (bProperOuter && bEditInlineNew)
-				{
-					const bool bKeptByOld = OldEditInlineObjects.Contains(Obj);
-					const bool bNotHandledYet = !ReferenceReplacementMap.Contains(Obj);
-					if (bKeptByOld && bNotHandledYet)
-					{
-						UObject* NewEditInlineSubobject = StaticDuplicateObject(Obj, NewObject, NULL);
-						ReferenceReplacementMap.Add(Obj, NewEditInlineSubobject);
+	FFindInstancedReferenceSubobjectHelper::Duplicate(OldObject, NewObject, ReferenceReplacementMap, ComponentsOnNewObject);
 
-						// We also need to make sure to fixup any properties here
-						ComponentsOnNewObject.Add(NewEditInlineSubobject);
-					}
-				}
-			}
-		}
-	}
-	
 	// Replace anything with an outer of the old object with NULL, unless it already has a replacement
 	TArray<UObject*> ObjectsInOuter;
 	GetObjectsWithOuter(OldObject, ObjectsInOuter, true);
