UE-8593 cl#2423635 merged from UT branch

diff --git a/Engine/Source/Runtime/Engine/Private/ActorConstruction.cpp b/Engine/Source/Runtime/Engine/Private/ActorConstruction.cpp
--- a/Engine/Source/Runtime/Engine/Private/ActorConstruction.cpp
+++ b/Engine/Source/Runtime/Engine/Private/ActorConstruction.cpp
@@ -22,13 +22,14 @@ namespace
 	/** Tracks info for components instanced during UCS execution */
 	struct FUCSComponentInfo
 	{
-		UActorComponent* Component;
 		EComponentMobility::Type Mobility;
+		TWeakObjectPtr<UActorComponent> ComponentPtr;
 
 		FUCSComponentInfo(UActorComponent* InComponent)
-			:Component(InComponent)
+			: ComponentPtr(InComponent)
 		{
-			USceneComponent* SceneComponent = Cast<USceneComponent>(Component);
+			ensure(!InComponent || !InComponent->IsPendingKill());
+			USceneComponent* SceneComponent = Cast<USceneComponent>(InComponent);
 			if(SceneComponent != nullptr)
 			{
 				// Save original mobility
@@ -61,26 +62,22 @@ namespace
 		/** Add a component instance for the given Actor */
 		void AddComponent(const AActor* InActor, UActorComponent* InComponent)
 		{
-			TArray<FUCSComponentInfo>* UCSComponentsList = UCSComponentsMap.Find(InActor);
-			if (UCSComponentsList == nullptr)
-			{
-				UCSComponentsList = &UCSComponentsMap.Add(InActor, TArray<FUCSComponentInfo>());
-			}
-
-			UCSComponentsList->Add(FUCSComponentInfo(InComponent));
+			TArray<FUCSComponentInfo>& UCSComponentsList = UCSComponentsMap.FindOrAdd(InActor);
+			UCSComponentsList.Add(FUCSComponentInfo(InComponent));
 		}
 
 		/** Called after UCS execution has finished for the given Actor */
 		void PostProcessComponents(const AActor* InActor)
 		{
-			TArray<FUCSComponentInfo>* UCSComponentsList = UCSComponentsMap.Find(InActor);
-			if (UCSComponentsList != nullptr)
+			TArray<FUCSComponentInfo> UCSComponentsList;
+			const bool bFound = UCSComponentsMap.RemoveAndCopyValue(InActor, UCSComponentsList);
+			if (bFound)
 			{
-				for (int32 ComponentIndex = 0; ComponentIndex < UCSComponentsList->Num(); ++ComponentIndex)
+				for (int32 ComponentIndex = 0; ComponentIndex < UCSComponentsList.Num(); ++ComponentIndex)
 				{
-					const FUCSComponentInfo& UCSComponentInfo = (*UCSComponentsList)[ComponentIndex];
+					const FUCSComponentInfo& UCSComponentInfo = UCSComponentsList[ComponentIndex];
 
-					USceneComponent* SceneComponent = Cast<USceneComponent>(UCSComponentInfo.Component);
+					USceneComponent* SceneComponent = Cast<USceneComponent>(UCSComponentInfo.ComponentPtr.Get());
 					if(SceneComponent != nullptr)
 					{
 						// Restore original mobility after UCS execution
@@ -102,8 +99,6 @@ namespace
 						}
 					}
 				}
-
-				UCSComponentsMap.Remove(InActor);
 			}
 		}
 
@@ -614,7 +609,7 @@ UActorComponent* AActor::AddComponent(FName TemplateName, bool bManualAttachment
 		// Keep track of the new component during UCS execution
 		if(bRunningUserConstructionScript)
 		{
-			FUCSComponentManager::Get().AddComponent(this, NewActorComp);
+		FUCSComponentManager::Get().AddComponent(this, NewActorComp);
 		}
 
 		// The user has the option of doing attachment manually where they have complete control or via the automatic rule
