[INTEGRATE] Change 2421969 by [REDACTED_PERSON][REDACTED_HANDLE] on 2015/01/28 07:25:38

	Component can have an Instanced subobject that has an Instanced subobject ect.. Better fix for UE-8104.

diff --git a/Engine/Source/Runtime/Engine/Private/BlueprintGeneratedClass.cpp b/Engine/Source/Runtime/Engine/Private/BlueprintGeneratedClass.cpp
--- a/Engine/Source/Runtime/Engine/Private/BlueprintGeneratedClass.cpp
+++ b/Engine/Source/Runtime/Engine/Private/BlueprintGeneratedClass.cpp
@@ -38,10 +38,20 @@ void UBlueprintGeneratedClass::PostLoad()
 	TArray<UObject*> SubObjects;
 	GetObjectsWithOuter(ClassCDO, SubObjects);
 
+	struct FCheckIfComponentChildHelper
+	{
+		static bool IsComponentChild(UObject* CurrObj, const UObject* CDO)
+		{
+			UObject*  OuterObject = CurrObj ? CurrObj->GetOuter() : nullptr;
+			const bool bValidOuter = OuterObject && (OuterObject != CDO);
+			return bValidOuter ? (OuterObject->IsDefaultSubobject() || IsComponentChild(OuterObject, CDO)) : false;
+		};
+	};
+
 	for( auto SubObjIt = SubObjects.CreateIterator(); SubObjIt; ++SubObjIt )
 	{
 		UObject* CurrObj = *SubObjIt;
-		const bool bComponentChild = (CurrObj->GetOuter() != ClassCDO) && CurrObj->GetOuter()->IsDefaultSubobject();
+		const bool bComponentChild = FCheckIfComponentChildHelper::IsComponentChild(CurrObj, ClassCDO);
 		if (!CurrObj->IsDefaultSubobject() && !CurrObj->IsRooted() && !bComponentChild)
 		{
 			CurrObj->MarkPendingKill();
