[INTEGRATE] Change 2422166 by [REDACTED_PERSON][REDACTED_HANDLE] on 2015/01/28 10:05:24

	Strip nulls in AActor::GetComponentsByClass for the UActorComponent case as well

diff --git a/Engine/Source/Runtime/Engine/Private/Actor.cpp b/Engine/Source/Runtime/Engine/Private/Actor.cpp
--- a/Engine/Source/Runtime/Engine/Private/Actor.cpp
+++ b/Engine/Source/Runtime/Engine/Private/Actor.cpp
@@ -2122,13 +2122,20 @@ UActorComponent* AActor::GetComponentByClass(TSubclassOf<UActorComponent> Compon
 
 TArray<UActorComponent*> AActor::GetComponentsByClass(TSubclassOf<UActorComponent> ComponentClass) const
 {
+	TArray<UActorComponent*> ValidComponents;
+
+        // In the UActorComponent case we can skip the IsA checks for a slight performance benefit
 	if (ComponentClass == UActorComponent::StaticClass())
 	{
-		return OwnedComponents;
+		for (UActorComponent* Component : OwnedComponents)
+		{
+			if (Component)
+			{
+				ValidComponents.Add(Component);
+			}
+		}
 	}
-
-	TArray<UActorComponent*> ValidComponents;
-	if (*ComponentClass)
+	else if (*ComponentClass)
 	{
 		for (UActorComponent* Component : OwnedComponents)
 		{
