[INTEGRATE] Change 2422204 by [REDACTED_PERSON][REDACTED_HANDLE] on 2015/01/28 10:27:19

	UE-8358 Ensure when you delete and undo deleting a component after dragging it around first

	After a component was removed and restored (by Undo) it was not listed in owners OwnedComponents. So it wasn't handled by AActor::DestroyConstructedComponents().

diff --git a/Engine/Source/Runtime/Engine/Private/ActorComponent.cpp b/Engine/Source/Runtime/Engine/Private/ActorComponent.cpp
--- a/Engine/Source/Runtime/Engine/Private/ActorComponent.cpp
+++ b/Engine/Source/Runtime/Engine/Private/ActorComponent.cpp
@@ -411,6 +411,15 @@ void UActorComponent::PostEditUndo()
 			EditReregisterContexts.Remove(this);
 		}
 	}
+	else
+	{
+		//Let the component be properly registered, after it was restored.
+		AActor* Owner = GetOwner();
+		if (Owner)
+		{
+			Owner->AddOwnedComponent(this);
+		}
+	}
 	Super::PostEditUndo();
 }
 
