[INTEGRATE] Change 2422209 by [REDACTED_PERSON][REDACTED_HANDLE] on 2015/01/28 10:28:51

	InheritableComponentTemplate flag should not be propagated to instances of templates.

diff --git a/Engine/Source/Runtime/Engine/Private/ActorConstruction.cpp b/Engine/Source/Runtime/Engine/Private/ActorConstruction.cpp
--- a/Engine/Source/Runtime/Engine/Private/ActorConstruction.cpp
+++ b/Engine/Source/Runtime/Engine/Private/ActorConstruction.cpp
@@ -578,7 +578,7 @@ UActorComponent* AActor::CreateComponentFromTemplate(UActorComponent* Template,
 		}
 
 		// Note we aren't copying the the RF_ArchetypeObject flag. Also note the result is non-transactional by default.
-		NewActorComp = (UActorComponent*)StaticDuplicateObject(Template, this, *InName, RF_AllFlags & ~(RF_ArchetypeObject|RF_Transactional|RF_WasLoaded|RF_Public) );
+		NewActorComp = (UActorComponent*)StaticDuplicateObject(Template, this, *InName, RF_AllFlags & ~(RF_ArchetypeObject|RF_Transactional|RF_WasLoaded|RF_Public|RF_InheritableComponentTemplate) );
 		//NewActorComp = ConstructObject<UActorComponent>(Template->GetClass(), this, *InName, RF_NoFlags, Template);
 
 		NewActorComp->CreationMethod = EComponentCreationMethod::ConstructionScript;
