[INTEGRATE] Change 2422323 by [REDACTED_PERSON][REDACTED_HANDLE] on 2015/01/28 11:30:12

	Allow attachment of physics-enabled components during SCS/UCS.

diff --git a/Engine/Source/Runtime/Engine/Private/SceneComponent.cpp b/Engine/Source/Runtime/Engine/Private/SceneComponent.cpp
--- a/Engine/Source/Runtime/Engine/Private/SceneComponent.cpp
+++ b/Engine/Source/Runtime/Engine/Private/SceneComponent.cpp
@@ -791,7 +791,7 @@ void USceneComponent::AttachTo(class USceneComponent* Parent, FName InSocketName
 			//Also physics state may not be created yet so we use bSimulatePhysics to determine if the object has any intention of being physically simulated
 			UPrimitiveComponent * PrimitiveComponent = Cast<UPrimitiveComponent>(this);
 
-			if (PrimitiveComponent && PrimitiveComponent->BodyInstance.bSimulatePhysics && !bWeldSimulatedBodies && GetWorld() && GetWorld()->IsGameWorld())
+			if (PrimitiveComponent && PrimitiveComponent->BodyInstance.bSimulatePhysics && !bWeldSimulatedBodies && GetWorld() && GetWorld()->IsGameWorld() && !GetWorld()->bIsRunningConstructionScript)
 			{
 				//Since the object is physically simulated it can't be the case that it's a child of object A and being attached to object B (at runtime)
 				if (bMaintainWorldPosition == false)	//User tried to attach but physically based so detach. However, if they provided relative coordinates we should still get the correct position
