[INTEGRATE] Change 2422425 by [REDACTED_PERSON][REDACTED_HANDLE] on 2015/01/28 12:41:40

	Improve SkelMeshComp class comment a bit

diff --git a/Engine/Source/Runtime/Engine/Classes/Components/SkeletalMeshComponent.h b/Engine/Source/Runtime/Engine/Classes/Components/SkeletalMeshComponent.h
--- a/Engine/Source/Runtime/Engine/Classes/Components/SkeletalMeshComponent.h
+++ b/Engine/Source/Runtime/Engine/Classes/Components/SkeletalMeshComponent.h
@@ -264,7 +264,7 @@ struct FSkeletalMeshComponentPreClothTickFunction : public FTickFunction
 
 
 /**
- * SkeletalMeshComponent is used to create an instance of a USkeletalMesh.
+ * SkeletalMeshComponent is used to create an instance of an animated SkeletalMesh asset.
  *
  * @see https://docs.unrealengine.com/latest/INT/Engine/Content/Types/SkeletalMeshes/
  * @see USkeletalMesh
