[INTEGRATE] Change 2422471 by [REDACTED_PERSON][REDACTED_HANDLE] on 2015/01/28 13:16:43

	[UE-8401] GetComponentsByClass() now comforms its output type to the picked class type (in blueprints).

diff --git a/Engine/Source/Runtime/Engine/Classes/GameFramework/Actor.h b/Engine/Source/Runtime/Engine/Classes/GameFramework/Actor.h
--- a/Engine/Source/Runtime/Engine/Classes/GameFramework/Actor.h
+++ b/Engine/Source/Runtime/Engine/Classes/GameFramework/Actor.h
@@ -2137,7 +2137,7 @@ public:
 
 	/* Gets all the components that inherit from the given class.
 		Currently returns an array of UActorComponent which must be cast to the correct type. */
-	UFUNCTION(BlueprintCallable, Category="Actor", meta=(ComponentClass="ActorComponent"))
+	UFUNCTION(BlueprintCallable, Category = "Actor", meta = (ComponentClass = "ActorComponent"), meta=(DeterminesOutputType="ComponentClass"))
 	TArray<UActorComponent*> GetComponentsByClass(TSubclassOf<UActorComponent> ComponentClass) const;
 
 	/** Templatized version for syntactic nicety. */
