#ANIM: Allow access of curve data from outside.

Merging

//depot/Partners/Weta/UE4/Engine/Source/Runtime/Engine/Classes/Animation/AnimInstance.h

diff --git a/Engine/Source/Runtime/Engine/Classes/Animation/AnimInstance.h b/Engine/Source/Runtime/Engine/Classes/Animation/AnimInstance.h
--- a/Engine/Source/Runtime/Engine/Classes/Animation/AnimInstance.h
+++ b/Engine/Source/Runtime/Engine/Classes/Animation/AnimInstance.h
@@ -548,7 +548,7 @@ protected:
 
 public:
 	/** Returns the value of a named curve. */
-	UFUNCTION(BlueprintPure, Category="Animation", meta=(BlueprintProtected = "true"))
+	UFUNCTION(BlueprintPure, Category="Animation")
 	float GetCurveValue(FName CurveName);
 
 	/** Returns the length (in seconds) of an animation AnimAsset. */
