[INTEGRATE] Change 2424017 by [REDACTED_PERSON][REDACTED_HANDLE] on 2015/01/29 12:42:21

	When destroying components that create visualizer components, destroy them as well.

diff --git a/Engine/Source/Runtime/Engine/Classes/Camera/CameraComponent.h b/Engine/Source/Runtime/Engine/Classes/Camera/CameraComponent.h
--- a/Engine/Source/Runtime/Engine/Classes/Camera/CameraComponent.h
+++ b/Engine/Source/Runtime/Engine/Classes/Camera/CameraComponent.h
@@ -54,6 +54,7 @@ class UCameraComponent : public USceneComponent
 	ENGINE_API virtual void OnUnregister() override;
 	ENGINE_API virtual void PostLoad() override;
 #if WITH_EDITOR
+	ENGINE_API virtual void OnComponentDestroyed() override;
 	ENGINE_API virtual void CheckForErrors() override;
 	// End of UActorComponent interface
 
diff --git a/Engine/Source/Runtime/Engine/Classes/Components/AudioComponent.h b/Engine/Source/Runtime/Engine/Classes/Components/AudioComponent.h
--- a/Engine/Source/Runtime/Engine/Classes/Components/AudioComponent.h
+++ b/Engine/Source/Runtime/Engine/Classes/Components/AudioComponent.h
@@ -278,6 +278,7 @@ public:
 
 	// Begin ActorComponent interface.
 #if WITH_EDITORONLY_DATA
+	virtual void OnComponentDestroyed() override;
 	virtual void OnRegister() override;
 #endif
 	virtual void OnUnregister() override;
diff --git a/Engine/Source/Runtime/Engine/Classes/PhysicsEngine/PhysicsConstraintComponent.h b/Engine/Source/Runtime/Engine/Classes/PhysicsEngine/PhysicsConstraintComponent.h
--- a/Engine/Source/Runtime/Engine/Classes/PhysicsEngine/PhysicsConstraintComponent.h
+++ b/Engine/Source/Runtime/Engine/Classes/PhysicsEngine/PhysicsConstraintComponent.h
@@ -68,8 +68,9 @@ public:
 	//Begin ActorComponent interface
 #if WITH_EDITOR
 	virtual void CheckForErrors() override;
-#endif // WITH_EDITOR
+	virtual void OnComponentDestroyed() override;
 	virtual void OnRegister() override;
+#endif // WITH_EDITOR
 	virtual void OnUnregister() override;
 	virtual void InitializeComponent() override;
 	//End ActorComponent interface
diff --git a/Engine/Source/Runtime/Engine/Private/AudioComponent.cpp b/Engine/Source/Runtime/Engine/Private/AudioComponent.cpp
--- a/Engine/Source/Runtime/Engine/Private/AudioComponent.cpp
+++ b/Engine/Source/Runtime/Engine/Private/AudioComponent.cpp
@@ -83,6 +83,17 @@ void UAudioComponent::OnUnregister()
 	}
 }
 
+#if WITH_EDITORONLY_DATA
+void UAudioComponent::OnComponentDestroyed()
+{
+	Super::OnComponentDestroyed();
+	if (SpriteComponent)
+	{
+		SpriteComponent->DestroyComponent();
+	}
+}
+#endif
+
 void UAudioComponent::SetSound( USoundBase* NewSound )
 {
 	const bool bPlay = IsPlaying();
diff --git a/Engine/Source/Runtime/Engine/Private/Camera/CameraComponent.cpp b/Engine/Source/Runtime/Engine/Private/Camera/CameraComponent.cpp
--- a/Engine/Source/Runtime/Engine/Private/Camera/CameraComponent.cpp
+++ b/Engine/Source/Runtime/Engine/Private/Camera/CameraComponent.cpp
@@ -39,6 +39,22 @@ UCameraComponent::UCameraComponent(const FObjectInitializer& ObjectInitializer)
 	SetDeprecatedControllerViewRotation(*this, bUsePawnControlRotation);
 }
 
+#if WITH_EDITORONLY_DATA
+void UCameraComponent::OnComponentDestroyed()
+{
+	Super::OnComponentDestroyed();
+
+	if (ProxyMeshComponent)
+	{
+		ProxyMeshComponent->DestroyComponent();
+	}
+	if (DrawFrustum)
+	{
+		DrawFrustum->DestroyComponent();
+	}
+}
+#endif
+
 void UCameraComponent::OnRegister()
 {
 #if WITH_EDITORONLY_DATA
diff --git a/Engine/Source/Runtime/Engine/Private/PhysicsEngine/PhysicsConstraintComponent.cpp b/Engine/Source/Runtime/Engine/Private/PhysicsEngine/PhysicsConstraintComponent.cpp
--- a/Engine/Source/Runtime/Engine/Private/PhysicsEngine/PhysicsConstraintComponent.cpp
+++ b/Engine/Source/Runtime/Engine/Private/PhysicsEngine/PhysicsConstraintComponent.cpp
@@ -230,11 +230,21 @@ void UPhysicsConstraintComponent::InitializeComponent()
 	InitComponentConstraint();
 }
 
+#if WITH_EDITOR
+void UPhysicsConstraintComponent::OnComponentDestroyed()
+{
+	Super::OnComponentDestroyed();
+
+	if (SpriteComponent)
+	{
+		SpriteComponent->DestroyComponent();
+	}
+}
+
 void UPhysicsConstraintComponent::OnRegister()
 {
 	Super::OnRegister();
 
-#if WITH_EDITOR
 	if (SpriteComponent == NULL && GetOwner() && !GetWorld()->IsGameWorld() )
 	{
 		SpriteComponent = ConstructObject<UBillboardComponent>(UBillboardComponent::StaticClass(), GetOwner(), NAME_None, RF_Transactional | RF_TextExportTransient);
@@ -250,8 +260,8 @@ void UPhysicsConstraintComponent::OnRegister()
 
 		SpriteComponent->RegisterComponent();
 	}
-#endif
 }
+#endif
 
 void UPhysicsConstraintComponent::OnUnregister()
 {
