[INTEGRATE] Change 2424450 by [REDACTED_PERSON][REDACTED_HANDLE] on 2015/01/29 16:45:42

	#ue4 -  bOwnerNoSee and bOnlyOwnerSee are now editable in instances, not just defaults.

diff --git a/Engine/Source/Runtime/Engine/Classes/Components/PrimitiveComponent.h b/Engine/Source/Runtime/Engine/Classes/Components/PrimitiveComponent.h
--- a/Engine/Source/Runtime/Engine/Classes/Components/PrimitiveComponent.h
+++ b/Engine/Source/Runtime/Engine/Classes/Components/PrimitiveComponent.h
@@ -214,11 +214,11 @@ public:
 	uint32 bReceivesDecals:1;
 
 	/** If this is True, this component won't be visible when the view actor is the component's owner, directly or indirectly. */
-	UPROPERTY(EditDefaultsOnly, AdvancedDisplay, BlueprintReadOnly, Category=Rendering)
+	UPROPERTY(EditAnywhere, AdvancedDisplay, BlueprintReadOnly, Category = Rendering)
 	uint32 bOwnerNoSee:1;
 
 	/** If this is True, this component will only be visible when the view actor is the component's owner, directly or indirectly. */
-	UPROPERTY(EditDefaultsOnly, AdvancedDisplay, BlueprintReadOnly, Category=Rendering)
+	UPROPERTY(EditAnywhere, AdvancedDisplay, BlueprintReadOnly, Category = Rendering)
 	uint32 bOnlyOwnerSee:1;
 
 	/** Treat this primitive as part of the background for occlusion purposes. This can be used as an optimization to reduce the cost of rendering skyboxes, large ground planes that are part of the vista, etc. */
