[INTEGRATE] Change 2424737 by [REDACTED_PERSON][REDACTED_HANDLE] on 2015/01/29 18:46:37

	#ue4 - Improve warning text when adding movement component to an actor without a primitive component as the root component.

	UE-8602

diff --git a/Engine/Source/Runtime/Engine/Private/MovementComponent.cpp b/Engine/Source/Runtime/Engine/Private/MovementComponent.cpp
--- a/Engine/Source/Runtime/Engine/Private/MovementComponent.cpp
+++ b/Engine/Source/Runtime/Engine/Private/MovementComponent.cpp
@@ -99,10 +99,14 @@ void UMovementComponent::InitializeComponent()
 			NewUpdatedComponent = Cast<UPrimitiveComponent>(MyActor->GetRootComponent());
 			if (!NewUpdatedComponent)
 			{
-				FMessageLog("PIE").Warning(FText::Format(LOCTEXT("NoRootPrimitiveWarning", "Movement component {0} must update a PrimitiveComponent, but owning actor '{1}' does not have a root PrimitiveComponent. Auto registration failed."),
-					FText::FromString(GetName()),
-					FText::FromString(MyActor->GetName())
-					));
+				FMessageLog("PIE").Warning(
+					FText::Format(
+						LOCTEXT("NoRootPrimitiveWarning",
+								"Movement component {0} must update a PrimitiveComponent. Consider changing root component of owning actor '{1}' to a shape (Box, Capsule, etc) or another component that can handle collision."),
+						FText::FromString(GetName()),
+						FText::FromString(MyActor->GetName())
+					)
+				);
 			}
 		}
 	}
