[INTEGRATE] Change 2424739 by [REDACTED_PERSON][REDACTED_HANDLE] on 2015/01/29 18:48:18

	#ue4 - AutoPossesAI is supposed to be ignored if there is an AutoPossessPlayer set on the Pawn.

diff --git a/Engine/Source/Runtime/Engine/Private/Pawn.cpp b/Engine/Source/Runtime/Engine/Private/Pawn.cpp
--- a/Engine/Source/Runtime/Engine/Private/Pawn.cpp
+++ b/Engine/Source/Runtime/Engine/Private/Pawn.cpp
@@ -98,14 +98,17 @@ void APawn::PostInitializeComponents()
 		GetWorld()->AddPawn( this );
 
 		// Automatically add Controller to AI Pawns if we are allowed to.
-		if (AutoPossessAI != EAutoPossessAI::Disabled && Controller == NULL && GetNetMode() != NM_Client)
+		if (AutoPossessPlayer == EAutoReceiveInput::Disabled)
 		{
-			const bool bPlacedInWorld = (GetWorld()->bStartup);
-			if ((AutoPossessAI == EAutoPossessAI::PlacedInWorldOrSpawned) ||
-				(AutoPossessAI == EAutoPossessAI::PlacedInWorld && bPlacedInWorld) ||
-				(AutoPossessAI == EAutoPossessAI::Spawned && !bPlacedInWorld))
+			if (AutoPossessAI != EAutoPossessAI::Disabled && Controller == NULL && GetNetMode() != NM_Client)
 			{
-				SpawnDefaultController();
+				const bool bPlacedInWorld = (GetWorld()->bStartup);
+				if ((AutoPossessAI == EAutoPossessAI::PlacedInWorldOrSpawned) ||
+					(AutoPossessAI == EAutoPossessAI::PlacedInWorld && bPlacedInWorld) ||
+					(AutoPossessAI == EAutoPossessAI::Spawned && !bPlacedInWorld))
+				{
+					SpawnDefaultController();
+				}
 			}
 		}
 
