UpdateDestructibleChunkTM was taking array by value by mistake

diff --git a/Engine/Source/Runtime/Engine/Classes/Components/DestructibleComponent.h b/Engine/Source/Runtime/Engine/Classes/Components/DestructibleComponent.h
--- a/Engine/Source/Runtime/Engine/Classes/Components/DestructibleComponent.h
+++ b/Engine/Source/Runtime/Engine/Classes/Components/DestructibleComponent.h
@@ -177,7 +177,7 @@ public:
 	 *  @param ActiveActors - The array of actors that need their transforms updated
 	 *
      */
-	static void UpdateDestructibleChunkTM(TArray<const physx::PxRigidActor*> ActiveActors);
+	static void UpdateDestructibleChunkTM(const TArray<const physx::PxRigidActor*>& ActiveActors);
 #endif
 
 
diff --git a/Engine/Source/Runtime/Engine/Private/DestructibleComponent.cpp b/Engine/Source/Runtime/Engine/Private/DestructibleComponent.cpp
--- a/Engine/Source/Runtime/Engine/Private/DestructibleComponent.cpp
+++ b/Engine/Source/Runtime/Engine/Private/DestructibleComponent.cpp
@@ -858,7 +858,7 @@ void UDestructibleComponent::SetChunkVisible( int32 ChunkIndex, bool bVisible )
 }
 
 #if WITH_APEX
-void UDestructibleComponent::UpdateDestructibleChunkTM(TArray<const PxRigidActor*> ActiveActors)
+void UDestructibleComponent::UpdateDestructibleChunkTM(const TArray<const PxRigidActor*>& ActiveActors)
 {
 	//We want to consolidate the transforms so that we update each destructible component once by passing it an array of chunks to update.
 	//This helps avoid a lot of duplicated work like marking render dirty, computing inverse world component, etc...
