Fixed centered Canvas text not working when using a FSlateFontInfo override

UETOOL-233 - FCanvasTextItem CentreX/CentreY don't work correctly when providing an FSlateFontInfo override

UCanvas::CanvasStringSize doesn't work with the FSlateFontInfo override that can be given to FCanvasText, so we have to use FSlateFontMeasure in this case.

diff --git a/Engine/Source/Runtime/Engine/Private/UserInterface/CanvasItem.cpp b/Engine/Source/Runtime/Engine/Private/UserInterface/CanvasItem.cpp
--- a/Engine/Source/Runtime/Engine/Private/UserInterface/CanvasItem.cpp
+++ b/Engine/Source/Runtime/Engine/Private/UserInterface/CanvasItem.cpp
@@ -808,11 +808,32 @@ void FCanvasTextItem::Draw( class FCanvas* InCanvas )
 
 	FVector2D DrawPos( Position.X , Position.Y );
 
-	// If we are centering the string or we want to fix stereoscopic rending issues we need to measure the string
+	// If we are centering the string or we want to fix stereoscopic rendering issues we need to measure the string
 	if( ( bCentreX || bCentreY ) || ( !bDontCorrectStereoscopic ) )
 	{
-		FTextSizingParameters Parameters( Font, Scale.X ,Scale.Y );
-		UCanvas::CanvasStringSize( Parameters, *Text.ToString() );
+		FVector2D MeasuredTextSize;
+		switch( GetFontCacheType() )
+		{
+		case EFontCacheType::Offline:
+			{
+				FTextSizingParameters Parameters( Font, Scale.X ,Scale.Y );
+				UCanvas::CanvasStringSize( Parameters, *Text.ToString() );
+				MeasuredTextSize.X = Parameters.DrawXL;
+				MeasuredTextSize.Y = Parameters.DrawYL;
+			}
+			break;
+
+		case EFontCacheType::Runtime:
+			{
+				const FSlateFontInfo LegacyFontInfo = (SlateFontInfo.IsSet()) ? SlateFontInfo.GetValue() : Font->GetLegacySlateFontInfo();
+				const TSharedRef<FSlateFontMeasure> FontMeasure = FSlateApplication::Get().GetRenderer()->GetFontMeasureService();
+				MeasuredTextSize = FontMeasure->Measure( Text, LegacyFontInfo ) * Scale;
+			}
+			break;
+
+		default:
+			break;
+		}
 				
 		// Calculate the offset if we are centering
 		if( bCentreX || bCentreY )
@@ -820,23 +841,23 @@ void FCanvasTextItem::Draw( class FCanvas* InCanvas )
 			// Note we drop the fraction after the length divide or we can end up with coords on 1/2 pixel boundaries
 			if( bCentreX )
 			{
-				DrawPos.X = DrawPos.X - (int)( Parameters.DrawXL / 2 );
+				DrawPos.X -= (int)( MeasuredTextSize.X / 2 );
 			}
 			if( bCentreY )
 			{
-				DrawPos.Y = DrawPos.Y - (int)( Parameters.DrawYL / 2 );
+				DrawPos.Y -= (int)( MeasuredTextSize.Y / 2 );
 			}
 		}
 
 		// Check if we want to correct the stereo3d issues - if we do, render the correction now
-		bool CorrectStereo = !bDontCorrectStereoscopic  && GEngine->IsStereoscopic3D();
+		const bool CorrectStereo = !bDontCorrectStereoscopic  && GEngine->IsStereoscopic3D();
 		if( CorrectStereo )
 		{
-			FVector2D StereoOutlineBoxSize( 2.0f, 2.0f );
+			const FVector2D StereoOutlineBoxSize( 2.0f, 2.0f );
 			TileItem.MaterialRenderProxy = GEngine->RemoveSurfaceMaterial->GetRenderProxy( false );
 			TileItem.Position = DrawPos - StereoOutlineBoxSize;
-			FVector2D CorrectionSize = FVector2D( Parameters.DrawXL, Parameters.DrawYL ) + StereoOutlineBoxSize + StereoOutlineBoxSize;
-			TileItem.Size=  CorrectionSize;
+			const FVector2D CorrectionSize = MeasuredTextSize + StereoOutlineBoxSize + StereoOutlineBoxSize;
+			TileItem.Size = CorrectionSize;
 			TileItem.bFreezeTime = true;
 			TileItem.Draw( InCanvas );
 		}		
