Fixed blueprint spline meshes generating a new entry in the DDC every time their blueprint construction script is run (e.g. loading the map in the editor)
This also allows multiple spline mesh components with the same params to share the same collision DDC, which can only be good :)

diff --git a/Engine/Source/Runtime/Engine/Classes/Components/SplineMeshComponent.h b/Engine/Source/Runtime/Engine/Classes/Components/SplineMeshComponent.h
--- a/Engine/Source/Runtime/Engine/Classes/Components/SplineMeshComponent.h
+++ b/Engine/Source/Runtime/Engine/Classes/Components/SplineMeshComponent.h
@@ -163,6 +163,7 @@ class ENGINE_API USplineMeshComponent : public UStaticMeshComponent, public IInt
 	virtual bool GetPhysicsTriMeshData(struct FTriMeshCollisionData* CollisionData, bool InUseAllTriData) override;
 	virtual bool ContainsPhysicsTriMeshData(bool InUseAllTriData) const override;
 	virtual bool WantsNegXTriMesh() override { return false; }
+	virtual void GetMeshId(FString& OutMeshId) override;
 	// End Interface_CollisionDataProvider Interface
 
 	/** Called when spline params are changed, to notify render thread and possibly collision */
diff --git a/Engine/Source/Runtime/Engine/Private/SplineMeshComponent.cpp b/Engine/Source/Runtime/Engine/Private/SplineMeshComponent.cpp
--- a/Engine/Source/Runtime/Engine/Private/SplineMeshComponent.cpp
+++ b/Engine/Source/Runtime/Engine/Private/SplineMeshComponent.cpp
@@ -604,6 +604,33 @@ bool USplineMeshComponent::ContainsPhysicsTriMeshData(bool InUseAllTriData) cons
 	return false;
 }
 
+void USplineMeshComponent::GetMeshId(FString& OutMeshId)
+{
+	// new method: Same guid as the base mesh but with a unique DDC-id based on the spline params.
+	// This fixes the bug where running a blueprint construction script regenerates the guid and uses
+	// a new DDC slot even if the mesh hasn't changed
+	// If BodySetup is null that means we're *currently* duplicating one, and haven't transformed its data
+	// to fit the spline yet, so just use the data from the base mesh by using a blank MeshId
+	// It would be better if we could stop it building data in that case at all...
+	if (BodySetup != nullptr && BodySetup->BodySetupGuid == CachedMeshBodySetupGuid)
+	{
+		OutMeshId += FString::Printf(TEXT("(%s,%s,%s,%f,%s)_(%s,%s,%s,%f,%s)_%s_%d_%c"),
+			*SplineParams.StartPos.ToString(),
+			*SplineParams.StartTangent.ToString(),
+			*SplineParams.StartScale.ToString(),
+			SplineParams.StartRoll,
+			SplineParams.StartOffset != FVector2D::ZeroVector ? *SplineParams.StartOffset.ToString() : TEXT(""),
+			*SplineParams.EndPos.ToString(),
+			*SplineParams.EndTangent.ToString(),
+			*SplineParams.EndScale.ToString(),
+			SplineParams.EndRoll,
+			SplineParams.EndOffset != FVector2D::ZeroVector ? *SplineParams.EndOffset.ToString() : TEXT(""),
+			SplineUpDir != FVector::UpVector ? *SplineUpDir.ToString() : TEXT(""),
+			(int32)bSmoothInterpRollScale,
+			TEXT("XYZ")[(int32)ForwardAxis.GetValue()]);
+	}
+}
+
 void USplineMeshComponent::CreatePhysicsState()
 {
 #if WITH_EDITOR
@@ -710,6 +737,8 @@ void USplineMeshComponent::RecreateCollision()
 			BodySetup->CopyBodyPropertiesFrom(StaticMesh->BodySetup);
 			BodySetup->CollisionTraceFlag = StaticMesh->BodySetup->CollisionTraceFlag;
 		}
+		BodySetup->BodySetupGuid = StaticMesh->BodySetup->BodySetupGuid;
+		CachedMeshBodySetupGuid = StaticMesh->BodySetup->BodySetupGuid;
 
 		if (BodySetup->CollisionTraceFlag == CTF_UseComplexAsSimple)
 		{
@@ -780,7 +809,6 @@ void USplineMeshComponent::RecreateCollision()
 		}
 
 		BodySetup->CreatePhysicsMeshes();
-		CachedMeshBodySetupGuid = StaticMesh->BodySetup->BodySetupGuid;
 	}
 	else
 	{
