Fix PreViewTranslation being uninitialized in some ortho views.

diff --git a/Engine/Source/Runtime/Engine/Private/SceneView.cpp b/Engine/Source/Runtime/Engine/Private/SceneView.cpp
--- a/Engine/Source/Runtime/Engine/Private/SceneView.cpp
+++ b/Engine/Source/Runtime/Engine/Private/SceneView.cpp
@@ -335,6 +335,7 @@ FSceneView::FSceneView(const FSceneViewInitOptions& InitOptions)
 	// Translate world-space so its origin is at ViewOrigin for improved precision.
 	// Note that this isn't exactly right for orthogonal projections (See the above special case), but we still use ViewOrigin
 	// in that case so the same value may be used in shaders for both the world-space translation and the camera's world position.
+	ViewMatrices.PreViewTranslation = FVector::ZeroVector;
 	if(bApplyPreViewTranslation)
 	{
 		ViewMatrices.PreViewTranslation = -FVector(ViewMatrices.ViewOrigin);
