Removed log spam

diff --git a/Engine/Source/Runtime/Engine/Public/MaterialShared.h b/Engine/Source/Runtime/Engine/Public/MaterialShared.h
--- a/Engine/Source/Runtime/Engine/Public/MaterialShared.h
+++ b/Engine/Source/Runtime/Engine/Public/MaterialShared.h
@@ -1048,7 +1048,6 @@ public:
 	void RemoveOutstandingCompileId(const int32 OldOutstandingCompileShaderMapId )
 	{
 		OutstandingCompileShaderMapIds.Remove( OldOutstandingCompileShaderMapId );
-		UE_LOG(LogMaterial, Display, TEXT("Removing compile shader map id %d"), OldOutstandingCompileShaderMapId);
 	}
 
 	void AddReferencedObjects(FReferenceCollector& Collector);
