Distance field replacement mesh is now postloaded, hopeful fix for a crash

diff --git a/Engine/Source/Runtime/Engine/Private/StaticMesh.cpp b/Engine/Source/Runtime/Engine/Private/StaticMesh.cpp
--- a/Engine/Source/Runtime/Engine/Private/StaticMesh.cpp
+++ b/Engine/Source/Runtime/Engine/Private/StaticMesh.cpp
@@ -1189,6 +1189,13 @@ void FStaticMeshRenderData::Cache(UStaticMesh* Owner, const FStaticMeshLODSettin
 
 		const FMeshBuildSettings& BuildSettings = Owner->SourceModels[0].BuildSettings;
 		UStaticMesh* MeshToGenerateFrom = BuildSettings.DistanceFieldReplacementMesh ? BuildSettings.DistanceFieldReplacementMesh : Owner;
+
+		if (BuildSettings.DistanceFieldReplacementMesh)
+		{
+			// Make sure dependency is postloaded
+			BuildSettings.DistanceFieldReplacementMesh->ConditionalPostLoad();
+		}
+
 		LODResources[0].DistanceFieldData->CacheDerivedData(DistanceFieldKey, Owner, MeshToGenerateFrom, BuildSettings.DistanceFieldResolutionScale, BuildSettings.bGenerateDistanceFieldAsIfTwoSided);
 	}
 }
