Fix TranslatedViewMatrix in the case of certain ortho cameras

diff --git a/Engine/Source/Runtime/Engine/Private/SceneView.cpp b/Engine/Source/Runtime/Engine/Private/SceneView.cpp
--- a/Engine/Source/Runtime/Engine/Private/SceneView.cpp
+++ b/Engine/Source/Runtime/Engine/Private/SceneView.cpp
@@ -332,10 +332,12 @@ FSceneView::FSceneView(const FSceneViewInitOptions& InitOptions)
 		bApplyPreViewTranslation = false;
 	}
 
+	/** The view transform, starting from world-space points translated by -ViewOrigin. */
+	FMatrix TranslatedViewMatrix = InitOptions.ViewRotationMatrix;
+
 	// Translate world-space so its origin is at ViewOrigin for improved precision.
 	// Note that this isn't exactly right for orthogonal projections (See the above special case), but we still use ViewOrigin
 	// in that case so the same value may be used in shaders for both the world-space translation and the camera's world position.
-	ViewMatrices.PreViewTranslation = FVector::ZeroVector;
 	if(bApplyPreViewTranslation)
 	{
 		ViewMatrices.PreViewTranslation = -FVector(ViewMatrices.ViewOrigin);
@@ -359,9 +361,11 @@ FSceneView::FSceneView(const FSceneViewInitOptions& InitOptions)
 		}
 #endif // !(UE_BUILD_SHIPPING || UE_BUILD_TEST)
 	}
-
-	/** The view transform, starting from world-space points translated by -ViewOrigin. */
-	FMatrix TranslatedViewMatrix = InitOptions.ViewRotationMatrix;
+	else
+	{
+		// If not applying PreViewTranslation then we need to use the view matrix directly.
+		TranslatedViewMatrix = ViewMatrices.ViewMatrix;
+	}
 
 	// When the view origin is fudged for faux ortho view position the translations don't cancel out.
 	if (bViewOriginIsFudged)
