UE4 - Add 'requires restart' to render setting

diff --git a/Engine/Source/Runtime/Engine/Classes/Engine/RendererSettings.h b/Engine/Source/Runtime/Engine/Classes/Engine/RendererSettings.h
--- a/Engine/Source/Runtime/Engine/Classes/Engine/RendererSettings.h
+++ b/Engine/Source/Runtime/Engine/Classes/Engine/RendererSettings.h
@@ -236,7 +236,7 @@ class ENGINE_API URendererSettings
 
 	UPROPERTY(config, EditAnywhere, Category=Optimizations, meta=(
 		ConsoleVariable="r.BasePassOutputsVelocity", DisplayName="Accurate velocities from Vertex Deformation",
-		ToolTip="Enables materials with time-based World Position Offset and/or World Displacement to output accurate velocities. This incurs a performance cost. If this is disabled, those materials will not output velocities."))
+		ToolTip="Enables materials with time-based World Position Offset and/or World Displacement to output accurate velocities. This incurs a performance cost. If this is disabled, those materials will not output velocities. Changing this setting requires restarting the editor."))
 	uint32 bBasePassOutputsVelocity:1;
 
 	UPROPERTY(config, EditAnywhere, Category=Editor, meta=(
