Pixel Depth Offset material input
* Allows writing to the depth buffer per-pixel, making use of D3D11's conservative depth writes to maintain ZCull (theoretically, it's up to the driver)
* This is useful for billboard LODs with distance field lighting, as the distance field representation of the mesh used for shadowing is still accurate but the billboard depths used for primary rays is not.

diff --git a/Engine/Shaders/BasePassPixelShader.usf b/Engine/Shaders/BasePassPixelShader.usf
--- a/Engine/Shaders/BasePassPixelShader.usf
+++ b/Engine/Shaders/BasePassPixelShader.usf
@@ -576,7 +576,8 @@ uint CustomAlpha2Coverage(inout float4 InOutColor)
 void Main(
 	FVertexFactoryInterpolantsVSToPS Interpolants,
 	FBasePassInterpolantsVSToPS BasePassInterpolants,
-	in float4 SvPosition : SV_Position,
+	in INPUT_POSITION_QUALIFIERS float4 SvPosition : SV_Position,
+
 	OPTIONAL_IsFrontFace,
 	out float4 OutColor		: SV_Target0
 #if FEATURE_LEVEL >= FEATURE_LEVEL_SM4 && (MATERIALBLENDING_SOLID || MATERIALBLENDING_MASKED) && !NO_GBUFFER
@@ -592,6 +593,8 @@ void Main(
 	#endif
 #endif
 
+	OPTIONAL_OutDepth
+
 #if EDITOR_ALPHA2COVERAGE != 0
    ,in uint InCoverage : SV_Coverage
    ,out uint OutCoverage : SV_Coverage
@@ -615,7 +618,10 @@ void Main(
 #else
 	const bool bEditorWeightedZBuffering = false;
 #endif
-
+		
+#if OUTPUT_PIXEL_DEPTH_OFFSET
+	ApplyPixelDepthOffsetToMaterialParameters(MaterialParameters, OutDepth);
+#endif
 
 	//Clip if the blend mode requires it.
 	if(!bEditorWeightedZBuffering)
@@ -642,7 +648,7 @@ void Main(
 		float4 SubsurfaceData = GetMaterialSubsurfaceData(MaterialParameters);
 
 #if MATERIAL_SHADINGMODEL_SUBSURFACE || MATERIAL_SHADINGMODEL_PREINTEGRATED_SKIN || MATERIAL_SHADINGMODEL_TWOSIDED_FOLIAGE
-		SubsurfaceColor = SubsurfaceData.rgb;
+		SubsurfaceColor = SubsurfaceData.rgb * View.DiffuseOverrideParameter.w + View.DiffuseOverrideParameter.xyz;
 #endif
 		SubsurfaceProfile = SubsurfaceData.a;
 	}
@@ -866,9 +872,6 @@ void Main(
 		float QuantizationBias = PseudoRandom( MaterialParameters.SVPosition.xy ) * 0.5 - 0.5;
 		EncodeGBuffer(GBuffer, OutGBufferA, OutGBufferB, OutGBufferC, OutGBufferD, OutGBufferE, OutGBufferVelocity, QuantizationBias);
 	#endif 
-	
-
-
 
 	if(bEditorWeightedZBuffering)
 	{
diff --git a/Engine/Shaders/DepthOnlyPixelShader.usf b/Engine/Shaders/DepthOnlyPixelShader.usf
--- a/Engine/Shaders/DepthOnlyPixelShader.usf
+++ b/Engine/Shaders/DepthOnlyPixelShader.usf
@@ -9,20 +9,25 @@
 #include "VertexFactory.usf"
 
 void Main(
-#if !MATERIALBLENDING_SOLID
+#if !MATERIALBLENDING_SOLID || OUTPUT_PIXEL_DEPTH_OFFSET
 	FVertexFactoryInterpolantsVSToPS FactoryInterpolants,
 	float4 PixelPosition : TEXCOORD6,
-	in float4 SvPosition : SV_Position,
-	OPTIONAL_IsFrontFace,
+	in INPUT_POSITION_QUALIFIERS float4 SvPosition : SV_Position,
+	OPTIONAL_IsFrontFace
+	OPTIONAL_OutDepth
 #endif
-	out float4 OutColor : SV_Target0
+	,out float4 OutColor : SV_Target0
 	)
 {
-#if !MATERIALBLENDING_SOLID
+#if !MATERIALBLENDING_SOLID || OUTPUT_PIXEL_DEPTH_OFFSET
 
 	FMaterialPixelParameters MaterialParameters = GetMaterialPixelParameters(FactoryInterpolants, PixelPosition);
 	CalcMaterialParameters(MaterialParameters, SvPosition, bIsFrontFace, PixelPosition);
 
+	#if OUTPUT_PIXEL_DEPTH_OFFSET
+		ApplyPixelDepthOffsetToMaterialParameters(MaterialParameters, OutDepth);
+	#endif
+
 	GetMaterialCoverageAndClipping(MaterialParameters);
 #endif
 
diff --git a/Engine/Shaders/DepthOnlyVertexShader.usf b/Engine/Shaders/DepthOnlyVertexShader.usf
--- a/Engine/Shaders/DepthOnlyVertexShader.usf
+++ b/Engine/Shaders/DepthOnlyVertexShader.usf
@@ -11,7 +11,7 @@
 
 struct FDepthOnlyVSToPS
 {
-	#if !MATERIALBLENDING_SOLID
+	#if !MATERIALBLENDING_SOLID || OUTPUT_PIXEL_DEPTH_OFFSET
 		FVertexFactoryInterpolantsVSToPS FactoryInterpolants;
 		float4 PixelPosition : TEXCOORD6;
 	#endif
@@ -51,7 +51,7 @@ struct FDepthOnlyVSToPS
 	{
 		FDepthOnlyVSToPS O;
 		
-		#if !MATERIALBLENDING_SOLID
+		#if !MATERIALBLENDING_SOLID || OUTPUT_PIXEL_DEPTH_OFFSET
 			O.FactoryInterpolants = VertexFactoryAssignInterpolants(Interpolants.FactoryInterpolants);
 		#endif
 
@@ -61,7 +61,7 @@ struct FDepthOnlyVSToPS
 			O.Position = mul(WorldPosition, View.TranslatedWorldToClip);
 		}
 			
-		#if !MATERIALBLENDING_SOLID
+		#if !MATERIALBLENDING_SOLID || OUTPUT_PIXEL_DEPTH_OFFSET
 			O.PixelPosition = WorldPosition;
 		#endif
 		
@@ -104,13 +104,13 @@ void Main(
 	}
 #endif
 
-#if !MATERIALBLENDING_SOLID || USING_TESSELLATION
+#if !MATERIALBLENDING_SOLID || USING_TESSELLATION || OUTPUT_PIXEL_DEPTH_OFFSET
 	// Masked and transparent materials need texture coords to clip, and tessellated
 	// materials need texture coords to displace
 	Output.FactoryInterpolants = VertexFactoryGetInterpolants(Input, VFIntermediates, VertexParameters);
 #endif
 
-#if !MATERIALBLENDING_SOLID
+#if !MATERIALBLENDING_SOLID || OUTPUT_PIXEL_DEPTH_OFFSET
 	#if !USING_TESSELLATION	
 		Output.PixelPosition = WorldPos;
 	#endif
diff --git a/Engine/Shaders/HitProxyPixelShader.usf b/Engine/Shaders/HitProxyPixelShader.usf
--- a/Engine/Shaders/HitProxyPixelShader.usf
+++ b/Engine/Shaders/HitProxyPixelShader.usf
@@ -16,14 +16,19 @@ void Main(
 	float4 InstanceHitProxyId : HIT_PROXY_ID,
 #endif
 	float4 PixelPosition	: TEXCOORD6,
-	in float4 SvPosition	: SV_Position,
-	OPTIONAL_IsFrontFace,
-	out float4 OutColor		: SV_Target0
+	in INPUT_POSITION_QUALIFIERS float4 SvPosition	: SV_Position,
+	OPTIONAL_IsFrontFace
+	OPTIONAL_OutDepth
+	,out float4 OutColor		: SV_Target0
 	)
 {
 	FMaterialPixelParameters MaterialParameters = GetMaterialPixelParameters(Interpolants, PixelPosition);
 	CalcMaterialParameters(MaterialParameters, SvPosition, bIsFrontFace, PixelPosition);
 
+	#if OUTPUT_PIXEL_DEPTH_OFFSET
+		ApplyPixelDepthOffsetToMaterialParameters(MaterialParameters, OutDepth);
+	#endif
+
 	GetMaterialCoverageAndClipping(MaterialParameters);
 
 #if USE_INSTANCING
diff --git a/Engine/Shaders/MaterialTemplate.usf b/Engine/Shaders/MaterialTemplate.usf
--- a/Engine/Shaders/MaterialTemplate.usf
+++ b/Engine/Shaders/MaterialTemplate.usf
@@ -1259,6 +1259,11 @@ void GetMaterialCustomizedUVs(FMaterialVertexParameters Parameters, out float2 O
 }
 #endif
 
+float GetMaterialPixelDepthOffset(FMaterialPixelParameters Parameters)
+{
+%s;
+}
+
 // Programmatically set the line number after all the material inputs which have a variable number of line endings
 // This allows shader error line numbers after this point to be the same regardless of which material is being compiled
 #line %s
@@ -1455,3 +1460,39 @@ half3x3 AssembleTangentToWorld( half3 TangentToWorld0, half4 TangentToWorld2 )
 	// Transform from tangent space to world space
 	return half3x3(TangentToWorld0, TangentToWorld1, TangentToWorld2.xyz);
 }
+
+// Whether the material shader should output pixel depth offset
+#define OUTPUT_PIXEL_DEPTH_OFFSET (WANT_PIXEL_DEPTH_OFFSET && (MATERIALBLENDING_SOLID || MATERIALBLENDING_MASKED))
+
+// Whether to use the hidden d3d11 feature that supports depth writes with ZCull by only pushing into the screen
+//@todo - use for other SM5 platforms
+#define USE_CONSERVATIVE_DEPTH_WRITES (OUTPUT_PIXEL_DEPTH_OFFSET && COMPILER_HLSL && FEATURE_LEVEL >= FEATURE_LEVEL_SM5)
+
+#if USE_CONSERVATIVE_DEPTH_WRITES
+	// Note: for some reason using SV_DepthLessEqual without these interpolation modifiers causes a compile error in d3d
+	#define INPUT_POSITION_QUALIFIERS linear noperspective centroid
+	// Use conservative depth output so we still get Z Cull.  Note, this is a reversed Z depth surface.
+	#define DEPTH_WRITE_SEMANTIC SV_DepthLessEqual
+#else
+	#define INPUT_POSITION_QUALIFIERS 
+	#define DEPTH_WRITE_SEMANTIC SV_DEPTH
+#endif
+
+#if OUTPUT_PIXEL_DEPTH_OFFSET
+	#define OPTIONAL_OutDepth ,out float OutDepth : DEPTH_WRITE_SEMANTIC
+#else
+	#define OPTIONAL_OutDepth
+#endif
+
+void ApplyPixelDepthOffsetToMaterialParameters(inout FMaterialPixelParameters MaterialParameters, out float OutDepth)
+{
+	// Can only push into the scene for conservative depth writes
+	float PixelDepthOffset = max(GetMaterialPixelDepthOffset(MaterialParameters), 0);
+
+	// Update positions used for shading
+	MaterialParameters.ScreenPosition.w += PixelDepthOffset;
+	MaterialParameters.SVPosition.w = MaterialParameters.ScreenPosition.w;
+	MaterialParameters.WorldPosition += MaterialParameters.CameraVector * PixelDepthOffset;
+
+	OutDepth = MaterialParameters.ScreenPosition.z / MaterialParameters.ScreenPosition.w;
+}
diff --git a/Engine/Source/Editor/UnrealEd/Private/MaterialGraph.cpp b/Engine/Source/Editor/UnrealEd/Private/MaterialGraph.cpp
--- a/Engine/Source/Editor/UnrealEd/Private/MaterialGraph.cpp
+++ b/Engine/Source/Editor/UnrealEd/Private/MaterialGraph.cpp
@@ -54,6 +54,7 @@ void UMaterialGraph::RebuildGraph()
 		}
 
 		MaterialInputs.Add( FMaterialInputInfo( LOCTEXT("MaterialAttributes", "Material Attributes"), MP_MaterialAttributes) );
+		MaterialInputs.Add( FMaterialInputInfo( LOCTEXT("PixelDepthOffset", "Pixel Depth Offset"), MP_PixelDepthOffset ) );
 
 		// Add Root Node
 		FGraphNodeCreator<UMaterialGraphNode_Root> NodeCreator(*this);
diff --git a/Engine/Source/Runtime/Engine/Classes/Materials/Material.h b/Engine/Source/Runtime/Engine/Classes/Materials/Material.h
--- a/Engine/Source/Runtime/Engine/Classes/Materials/Material.h
+++ b/Engine/Source/Runtime/Engine/Classes/Materials/Material.h
@@ -383,6 +383,9 @@ public:
 	UPROPERTY()
 	FMaterialAttributesInput MaterialAttributes;
 
+	UPROPERTY()
+	FScalarMaterialInput PixelDepthOffset;
+
 	/** Indicates that the material should be rendered in the SeparateTranslucency Pass (not affected by DOF, requires bAllowSeparateTranslucency to be set in .ini). */
 	UPROPERTY(EditAnywhere, Category=Translucency, meta=(DisplayName = "Separate Translucency"))
 	uint32 bEnableSeparateTranslucency:1;
diff --git a/Engine/Source/Runtime/Engine/Classes/Materials/MaterialExpressionMakeMaterialAttributes.h b/Engine/Source/Runtime/Engine/Classes/Materials/MaterialExpressionMakeMaterialAttributes.h
--- a/Engine/Source/Runtime/Engine/Classes/Materials/MaterialExpressionMakeMaterialAttributes.h
+++ b/Engine/Source/Runtime/Engine/Classes/Materials/MaterialExpressionMakeMaterialAttributes.h
@@ -61,6 +61,9 @@ class UMaterialExpressionMakeMaterialAttributes : public UMaterialExpression
 	UPROPERTY()
 	FExpressionInput CustomizedUVs[8];
 
+	UPROPERTY()
+	FExpressionInput PixelDepthOffset;
+
 	// Begin UMaterialExpression Interface
 	virtual int32 Compile(class FMaterialCompiler* Compiler, int32 OutputIndex, int32 MultiplexIndex) override;
 	virtual void GetCaption(TArray<FString>& OutCaptions) const override;
diff --git a/Engine/Source/Runtime/Engine/Private/Materials/HLSLMaterialTranslator.h b/Engine/Source/Runtime/Engine/Private/Materials/HLSLMaterialTranslator.h
--- a/Engine/Source/Runtime/Engine/Private/Materials/HLSLMaterialTranslator.h
+++ b/Engine/Source/Runtime/Engine/Private/Materials/HLSLMaterialTranslator.h
@@ -198,6 +198,7 @@ protected:
 	uint32 bUsesTransformVector : 1;
 	// True if the current property requires last frame's time
 	uint32 bUsesPreviousFrameTime : 1;
+	uint32 bUsesPixelDepthOffset : 1;
 	/** Tracks the number of texture coordinates used by this material. */
 	uint32 NumUserTexCoords;
 	/** Tracks the number of texture coordinates used by the vertex shader in this material. */
@@ -240,6 +241,7 @@ public:
 	,	bUsesParticleColor(false)
 	,	bUsesTransformVector(false)
 	,	bUsesPreviousFrameTime(false)
+	,	bUsesPixelDepthOffset(false)
 	,	NumUserTexCoords(0)
 	,	NumUserVertexTexCoords(0)
 	{}
@@ -343,31 +345,12 @@ public:
 				}
 			}
 
-			if (Chunk[MP_WorldPositionOffset] == -1)
-			{
-				MaterialCompilationOutput.bModifiesMeshPosition = false;
-			}
-			else
-			{
-				int32 WPOFreq = (int32)GetMaterialPropertyShaderFrequency(MP_WorldPositionOffset);
-				FShaderCodeChunk& WPOChunk = PropertyCodeChunks[MP_WorldPositionOffset][WPOFreq][Chunk[MP_WorldPositionOffset]];
-
-				// Determine whether the world position offset is used. 
-				// If the output chunk has a uniform expression, it is constant, and GetNumberValue returns the default property value then WPO isn't used.
-				MaterialCompilationOutput.bModifiesMeshPosition = true;
-				if( WPOChunk.UniformExpression && WPOChunk.UniformExpression->IsConstant() )
-				{
-					FLinearColor WPOValue;
-					FMaterialRenderContext DummyContext(nullptr, *Material, nullptr);
-					WPOChunk.UniformExpression->GetNumberValue(DummyContext, WPOValue);
-					if (FVector(WPOValue) == FVector::ZeroVector)
-					{
-						MaterialCompilationOutput.bModifiesMeshPosition = false;
-					}
-				}
-			}
-
+			Chunk[MP_PixelDepthOffset] = Material->CompilePropertyAndSetMaterialProperty(MP_PixelDepthOffset,this);
 
+			bUsesPixelDepthOffset = IsMaterialPropertyUsed(MP_PixelDepthOffset, Chunk[MP_PixelDepthOffset], FLinearColor(0, 0, 0, 0), 1);
+			const bool bUsesWorldPositionOffset = IsMaterialPropertyUsed(MP_WorldPositionOffset, Chunk[MP_WorldPositionOffset], FLinearColor(0, 0, 0, 0), 3);
+			MaterialCompilationOutput.bModifiesMeshPosition = bUsesPixelDepthOffset || bUsesWorldPositionOffset;
+			
 			if (Material->GetBlendMode() == BLEND_Modulate && MaterialShadingModel != MSM_Unlit && !Material->IsUsedWithDeferredDecal())
 			{
 				Errorf(TEXT("Dynamically lit translucency is not supported for BLEND_Modulate materials."));
@@ -632,6 +615,7 @@ public:
 		OutEnvironment.SetDefine(TEXT("INTERPOLATE_VERTEX_COLOR"), bUsesVertexColor ? TEXT("1") : TEXT("0")); 
 		OutEnvironment.SetDefine(TEXT("NEEDS_PARTICLE_COLOR"), bUsesParticleColor ? TEXT("1") : TEXT("0")); 
 		OutEnvironment.SetDefine(TEXT("USES_TRANSFORM_VECTOR"), bUsesTransformVector ? TEXT("1") : TEXT("0")); 
+		OutEnvironment.SetDefine(TEXT("WANT_PIXEL_DEPTH_OFFSET"), bUsesPixelDepthOffset ? TEXT("1") : TEXT("0")); 
 		// Distortion uses tangent space transform 
 		OutEnvironment.SetDefine(TEXT("USES_DISTORTION"), Material->IsDistorted() ? TEXT("1") : TEXT("0")); 
 
@@ -708,6 +692,8 @@ public:
 
 		LazyPrintf.PushParam(*CustomUVAssignments);
 	
+		LazyPrintf.PushParam(*GenerateFunctionCode(MP_PixelDepthOffset));
+
 		LazyPrintf.PushParam(*FString::Printf(TEXT("%u"),MaterialTemplateLineNumber));
 
 		return LazyPrintf.GetResultString();
@@ -715,6 +701,42 @@ public:
 
 protected:
 
+	bool IsMaterialPropertyUsed(EMaterialProperty Property, int32 PropertyChunkIndex, const FLinearColor& ReferenceValue, int32 NumComponents)
+	{
+		bool bPropertyUsed = false;
+
+		if (PropertyChunkIndex == -1)
+		{
+			bPropertyUsed = false;
+		}
+		else
+		{
+			int32 Frequency = (int32)GetMaterialPropertyShaderFrequency(Property);
+			FShaderCodeChunk& PropertyChunk = PropertyCodeChunks[Property][Frequency][PropertyChunkIndex];
+
+			// Determine whether the property is used. 
+			// If the output chunk has a uniform expression, it is constant, and GetNumberValue returns the default property value then property isn't used.
+			bPropertyUsed = true;
+
+			if( PropertyChunk.UniformExpression && PropertyChunk.UniformExpression->IsConstant() )
+			{
+				FLinearColor Value;
+				FMaterialRenderContext DummyContext(nullptr, *Material, nullptr);
+				PropertyChunk.UniformExpression->GetNumberValue(DummyContext, Value);
+
+				if ((NumComponents < 1 || Value.R == ReferenceValue.R)
+					&& (NumComponents < 2 || Value.G == ReferenceValue.G)
+					&& (NumComponents < 3 || Value.B == ReferenceValue.B)
+					&& (NumComponents < 4 || Value.A == ReferenceValue.A))
+				{
+					bPropertyUsed = false;
+				}
+			}
+		}
+
+		return bPropertyUsed;
+	}
+
 	// only used by GetMaterialShaderCode()
 	// @param Index ECompiledMaterialProperty or EMaterialProperty
 	FString GenerateFunctionCode(uint32 Index) const
diff --git a/Engine/Source/Runtime/Engine/Private/Materials/Material.cpp b/Engine/Source/Runtime/Engine/Private/Materials/Material.cpp
--- a/Engine/Source/Runtime/Engine/Private/Materials/Material.cpp
+++ b/Engine/Source/Runtime/Engine/Private/Materials/Material.cpp
@@ -3078,6 +3078,7 @@ FExpressionInput* UMaterial::GetExpressionInputForProperty(EMaterialProperty InP
 		case MP_AmbientOcclusion:		return &AmbientOcclusion;
 		case MP_Refraction:				return &Refraction;
 		case MP_MaterialAttributes:		return &MaterialAttributes;
+		case MP_PixelDepthOffset:		return &PixelDepthOffset;
 	}
 
 	if (InProperty >= MP_CustomizedUVs0 && InProperty <= MP_CustomizedUVs7)
@@ -3378,6 +3379,7 @@ int32 UMaterial::CompilePropertyEx( FMaterialCompiler* Compiler, EMaterialProper
 		case MP_Normal:					return Normal.CompileWithDefault(Compiler, Property);
 		case MP_WorldPositionOffset:	return WorldPositionOffset.CompileWithDefault(Compiler, Property);
 		case MP_WorldDisplacement:		return WorldDisplacement.CompileWithDefault(Compiler, Property);
+		case MP_PixelDepthOffset:		return PixelDepthOffset.CompileWithDefault(Compiler, Property);
 
 		default:
 			if (Property >= MP_CustomizedUVs0 && Property <= MP_CustomizedUVs7)
@@ -3663,6 +3665,9 @@ bool UMaterial::IsPropertyActive(EMaterialProperty InProperty) const
 	case MP_WorldPositionOffset:
 		Active = !bUsedWithUI;
 		break;
+	case MP_PixelDepthOffset:
+		Active = !IsTranslucentBlendMode((EBlendMode)BlendMode);
+		break;
 	case MP_MaterialAttributes:
 	default:
 		Active = true;
diff --git a/Engine/Source/Runtime/Engine/Private/Materials/MaterialExpressions.cpp b/Engine/Source/Runtime/Engine/Private/Materials/MaterialExpressions.cpp
--- a/Engine/Source/Runtime/Engine/Private/Materials/MaterialExpressions.cpp
+++ b/Engine/Source/Runtime/Engine/Private/Materials/MaterialExpressions.cpp
@@ -3196,6 +3196,12 @@ int32 UMaterialExpressionMakeMaterialAttributes::Compile(class FMaterialCompiler
 		Ret = CustomizedUVs[MultiplexIndex - CustomUVStart].Compile(Compiler); Expression = CustomizedUVs[MultiplexIndex - CustomUVStart].Expression; 
 	}
 
+	if (MultiplexIndex == CustomUVStart + 2)
+	{
+		Ret = PixelDepthOffset.Compile(Compiler);
+		Expression = PixelDepthOffset.Expression;
+	}
+
 	//If we've connected an expression but its still returned INDEX_NONE, flag the error.
 	if( Expression && INDEX_NONE == Ret )
 	{
@@ -3252,6 +3258,7 @@ UMaterialExpressionBreakMaterialAttributes::UMaterialExpressionBreakMaterialAttr
 	{
 		Outputs.Add(FExpressionOutput(*FString::Printf(TEXT("CustomizedUV%u"), UVIndex), 1, 1, 1, 0, 0));
 	}
+	Outputs.Add(FExpressionOutput(TEXT("PixelDepthOffset"), 1, 1, 1, 1, 0));
 }
 
 int32 UMaterialExpressionBreakMaterialAttributes::Compile(class FMaterialCompiler* Compiler, int32 OutputIndex, int32 MultiplexIndex)
diff --git a/Engine/Source/Runtime/Engine/Private/Materials/MaterialShared.cpp b/Engine/Source/Runtime/Engine/Private/Materials/MaterialShared.cpp
--- a/Engine/Source/Runtime/Engine/Private/Materials/MaterialShared.cpp
+++ b/Engine/Source/Runtime/Engine/Private/Materials/MaterialShared.cpp
@@ -211,6 +211,7 @@ EMaterialValueType GetMaterialPropertyType(EMaterialProperty Property)
 	case MP_AmbientOcclusion: return MCT_Float;
 	case MP_Refraction: return MCT_Float;
 	case MP_MaterialAttributes: return MCT_MaterialAttributes;
+	case MP_PixelDepthOffset: return MCT_Float;
 	};
 
 	if (Property >= MP_CustomizedUVs0 && Property <= MP_CustomizedUVs7)
@@ -1409,7 +1410,6 @@ bool FMaterial::BeginCompileShaderMap(
 		}
 		else
 		{
-			UE_LOG(LogMaterial, Display, TEXT("Add compile shader map id %d"), NewShaderMap->GetCompilingId());
 			OutstandingCompileShaderMapIds.Add( NewShaderMap->GetCompilingId() );
 			// Async compile, use NULL so that rendering will fall back to the default material.
 			OutShaderMap = NULL;
@@ -2157,6 +2157,11 @@ EMaterialProperty GetMaterialPropertyFromInputOutputIndex(int32 Index)
 		return MP_MaterialAttributes;
 	}
 
+	if (Index == UVEnd + 2)
+	{
+		return MP_PixelDepthOffset;
+	}
+
 	return MP_MAX;
 }
 
@@ -2185,9 +2190,17 @@ int32 GetInputOutputIndexFromMaterialProperty(EMaterialProperty Property)
 	case MP_MaterialAttributes: UE_LOG(LogMaterial, Fatal, TEXT("We should never need the IO index of the MaterialAttriubtes property."));
 	};
 
+	const int32 UVStart = 16;
+	const int32 UVEnd = UVStart + MP_CustomizedUVs7 - MP_CustomizedUVs0;
+
 	if (Property >= MP_CustomizedUVs0 && Property <= MP_CustomizedUVs7)
 	{
-		return 16 + Property - MP_CustomizedUVs0;
+		return UVStart + Property - MP_CustomizedUVs0;
+	}
+
+	if (Property == MP_PixelDepthOffset)
+	{
+		return UVEnd + 2;
 	}
 
 	return -1;
@@ -2206,6 +2219,7 @@ int32 GetDefaultExpressionForMaterialProperty(FMaterialCompiler* Compiler, EMate
 		case MP_ClearCoat:				return Compiler->Constant(1.0f);
 		case MP_ClearCoatRoughness:		return Compiler->Constant(0.1f);
 		case MP_AmbientOcclusion:		return Compiler->Constant(1.0f);
+		case MP_PixelDepthOffset:		return Compiler->Constant(0.0f);
 
 		case MP_EmissiveColor:			return Compiler->Constant3(0, 0, 0);
 		case MP_DiffuseColor:			return Compiler->Constant3(0, 0, 0);
@@ -2257,6 +2271,7 @@ FString GetNameOfMaterialProperty(EMaterialProperty Property)
 	case MP_ClearCoatRoughness:		return TEXT("ClearCoatRoughness");
 	case MP_AmbientOcclusion:		return TEXT("AmbientOcclusion");
 	case MP_Refraction:				return TEXT("Refraction");
+	case MP_PixelDepthOffset:		return TEXT("PixelDepthOffset");
 	};
 
 	if (Property >= MP_CustomizedUVs0 && Property <= MP_CustomizedUVs7)
diff --git a/Engine/Source/Runtime/Engine/Public/SceneTypes.h b/Engine/Source/Runtime/Engine/Public/SceneTypes.h
--- a/Engine/Source/Runtime/Engine/Public/SceneTypes.h
+++ b/Engine/Source/Runtime/Engine/Public/SceneTypes.h
@@ -153,6 +153,7 @@ enum EMaterialProperty
 	MP_CustomizedUVs6,
 	MP_CustomizedUVs7,
 	MP_MaterialAttributes,
+	MP_PixelDepthOffset,
 	MP_MAX,
 };
 
