Use the instance component checks as a gate, not an ensure

diff --git a/Engine/Source/Runtime/Engine/Private/ActorConstruction.cpp b/Engine/Source/Runtime/Engine/Private/ActorConstruction.cpp
--- a/Engine/Source/Runtime/Engine/Private/ActorConstruction.cpp
+++ b/Engine/Source/Runtime/Engine/Private/ActorConstruction.cpp
@@ -549,10 +549,8 @@ UActorComponent* AActor::CreateComponentFromTemplate(UActorComponent* Template,
 		if (!InName.IsEmpty())
 		{
 			UObject* ConflictingObject = FindObjectFast<UObject>(this, *InName);
-			if (ConflictingObject)
-			{
-				ensure(ConflictingObject->IsA<UActorComponent>() && CastChecked<UActorComponent>(ConflictingObject)->CreationMethod == EComponentCreationMethod::Instance);
-				
+			if (ConflictingObject && ConflictingObject->IsA<UActorComponent>() && CastChecked<UActorComponent>(ConflictingObject)->CreationMethod == EComponentCreationMethod::Instance)
+			{		
 				// Try and pick a good name
 				FString ConflictingObjectName = ConflictingObject->GetName();
 				int32 CharIndex = ConflictingObjectName.Len()-1;
