Fix brush collision not recreating physics state. Needed for in editor collision checks

diff --git a/Engine/Source/Runtime/Engine/Private/BrushComponent.cpp b/Engine/Source/Runtime/Engine/Private/BrushComponent.cpp
--- a/Engine/Source/Runtime/Engine/Private/BrushComponent.cpp
+++ b/Engine/Source/Runtime/Engine/Private/BrushComponent.cpp
@@ -730,6 +730,7 @@ void UBrushComponent::BuildSimpleBrushCollision()
 #if WITH_EDITOR
 	// Convert collision model into convex hulls.
 	BrushBodySetup->CreateFromModel( Brush, true );
+	RecreatePhysicsState();
 #endif // WITH_EDITOR
 
 	MarkPackageDirty();
