Fixed IrrelevantLights getting cleared per-LOD when it is stored per component, caused lighting to not be fully built with static meshes with LODs that are fully shadowed from a stationary light

diff --git a/Engine/Source/Runtime/Engine/Private/StaticMeshLight.cpp b/Engine/Source/Runtime/Engine/Private/StaticMeshLight.cpp
--- a/Engine/Source/Runtime/Engine/Private/StaticMeshLight.cpp
+++ b/Engine/Source/Runtime/Engine/Private/StaticMeshLight.cpp
@@ -258,8 +258,7 @@ void FStaticMeshStaticLightingTextureMapping::Apply(FQuantizedLightmapData* Quan
 		}
 
 		// Build the list of statically irrelevant lights.
-		// TODO: This should be stored per LOD.
-		StaticMeshComponent->IrrelevantLights.Empty();
+		// IrrelevantLights was cleared in InvalidateLightingCacheDetailed
 
 		for(int32 LightIndex = 0;LightIndex < Mesh->RelevantLights.Num();LightIndex++)
 		{
