Merge CL 2409316

Adds RepNotify to AGameState::ElapsedTime

diff --git a/Engine/Source/Runtime/Engine/Classes/GameFramework/GameState.h b/Engine/Source/Runtime/Engine/Classes/GameFramework/GameState.h
--- a/Engine/Source/Runtime/Engine/Classes/GameFramework/GameState.h
+++ b/Engine/Source/Runtime/Engine/Classes/GameFramework/GameState.h
@@ -67,7 +67,7 @@ protected:
 public:
 
 	/** Elapsed game time since match has started. */
-	UPROPERTY(replicated, BlueprintReadOnly, Category=GameState)
+	UPROPERTY(replicatedUsing=OnRep_ElapsedTime, BlueprintReadOnly, Category = GameState)
 	int32 ElapsedTime;
 
 	/** Array of all PlayerStates, maintained on both server and clients (PlayerStates are always relevant) */
@@ -90,6 +90,10 @@ public:
 	UFUNCTION()
 	virtual void OnRep_MatchState();
 
+	/** Gives clients the chance to do something when time gets updates */
+	UFUNCTION()
+	virtual void OnRep_ElapsedTime();
+
 	// Begin AActor interface
 	virtual void PostInitializeComponents() override;
 	// End AActor interface
diff --git a/Engine/Source/Runtime/Engine/Private/GameState.cpp b/Engine/Source/Runtime/Engine/Private/GameState.cpp
--- a/Engine/Source/Runtime/Engine/Private/GameState.cpp
+++ b/Engine/Source/Runtime/Engine/Private/GameState.cpp
@@ -46,6 +46,10 @@ void AGameState::DefaultTimer()
 	if (IsMatchInProgress())
 	{
 		++ElapsedTime;
+		if (GetNetMode() != NM_DedicatedServer)
+		{
+			OnRep_ElapsedTime();
+		}
 	}
 
 	GetWorldTimerManager().SetTimer(TimerHandle_DefaultTimer, this, &AGameState::DefaultTimer, GetWorldSettings()->GetEffectiveTimeDilation(), true);
@@ -241,6 +245,11 @@ void AGameState::OnRep_MatchState()
 	PreviousMatchState = MatchState;
 }
 
+void AGameState::OnRep_ElapsedTime()
+{
+	//Blank on purpose
+}
+
 bool AGameState::ShouldShowGore() const
 {
 	return true;
@@ -255,4 +264,4 @@ void AGameState::GetLifetimeReplicatedProps( TArray< FLifetimeProperty > & OutLi
 
 	DOREPLIFETIME_CONDITION( AGameState, GameModeClass,	COND_InitialOnly );
 	DOREPLIFETIME_CONDITION( AGameState, ElapsedTime,	COND_InitialOnly );
-}
+}
\ No newline at end of file
