Fix for UE-8450, could not select 'wildcard' for input pins on macros

diff --git a/Engine/Source/Editor/BlueprintGraph/Private/K2Node.cpp b/Engine/Source/Editor/BlueprintGraph/Private/K2Node.cpp
--- a/Engine/Source/Editor/BlueprintGraph/Private/K2Node.cpp
+++ b/Engine/Source/Editor/BlueprintGraph/Private/K2Node.cpp
@@ -267,9 +267,30 @@ void UK2Node::PinConnectionListChanged(UEdGraphPin* Pin)
 	NotifyPinConnectionListChanged(Pin);
 }
 
-void UK2Node::ReallocatePinsDuringReconstruction(TArray<UEdGraphPin*>& /*OldPins*/)
+void UK2Node::ReallocatePinsDuringReconstruction(TArray<UEdGraphPin*>& OldPins)
 {
 	AllocateDefaultPins();
+
+	for (auto OldPin : OldPins)
+	{
+		if (OldPin->ParentPin)
+		{
+			// find the new pin that corresponds to parent, and split it if it isn't already split
+			for (auto NewPin : Pins)
+			{
+				if (FCString::Stricmp(*(NewPin->PinName), *(OldPin->ParentPin->PinName)) == 0)
+				{
+					// Make sure we're not dealing with a menu node
+					UEdGraph* OuterGraph = GetGraph();
+					if (OuterGraph && OuterGraph->Schema && NewPin->SubPins.Num() == 0)
+					{
+						NewPin->PinType = OldPin->ParentPin->PinType;
+						GetSchema()->SplitPin(NewPin);
+					}
+				}
+			}
+		}
+	}
 }
 
 void UK2Node::PostReconstructNode()
diff --git a/Engine/Source/Runtime/Engine/Private/EdGraph/EdGraphPin.cpp b/Engine/Source/Runtime/Engine/Private/EdGraph/EdGraphPin.cpp
--- a/Engine/Source/Runtime/Engine/Private/EdGraph/EdGraphPin.cpp
+++ b/Engine/Source/Runtime/Engine/Private/EdGraph/EdGraphPin.cpp
@@ -152,14 +152,6 @@ void UEdGraphPin::CopyPersistentDataFromOldPin(const UEdGraphPin& SourcePin)
 		DefaultTextValue = SourcePin.DefaultTextValue;
 	}
 
-	// In K2 schemas the wildcard pins need to have their type copied before we get to pin splitting
-	// TODO: Better less hacky way of this?
-	static const FString WildCardText(TEXT("wildcard"));
-	if (PinType.PinCategory == WildCardText)
-	{
-		PinType = SourcePin.PinType;
-	}
-
 	// Copy the links
 	for (int32 LinkIndex = 0; LinkIndex < SourcePin.LinkedTo.Num(); ++LinkIndex)
 	{
@@ -186,8 +178,10 @@ void UEdGraphPin::CopyPersistentDataFromOldPin(const UEdGraphPin& SourcePin)
 		}
 	}
 
-	// If the source pin is split, then split the new one
-	if (SourcePin.SubPins.Num() > 0)
+	// If the source pin is split, then split the new one, but don't split multiple times, typically splitting is done
+	// by UK2Node::ReallocatePinsDuringReconstruction or FBlueprintEditor::OnSplitStructPin, but there are several code
+	// paths into this, and split state should be persistent:
+	if (SourcePin.SubPins.Num() > 0 && SubPins.Num() == 0)
 	{
 		GetSchema()->SplitPin(this);
 	}
