Input settings: Fix for the input settings to ensure we don't add a second console key unless it's necessary.

diff --git a/Engine/Source/Runtime/Engine/Private/UserInterface/InputSettings.cpp b/Engine/Source/Runtime/Engine/Private/UserInterface/InputSettings.cpp
--- a/Engine/Source/Runtime/Engine/Private/UserInterface/InputSettings.cpp
+++ b/Engine/Source/Runtime/Engine/Private/UserInterface/InputSettings.cpp
@@ -66,7 +66,10 @@ void UInputSettings::PostInitProperties()
 			break;
 		}
 
-		ConsoleKeys.Add(DefaultConsoleKey);
+		if (DefaultConsoleKey != EKeys::Tilde && DefaultConsoleKey.IsValid())
+		{
+			ConsoleKeys.Add(DefaultConsoleKey);
+		}
 	}
 #endif
 }
