UECORE-64: Add ENGINE_API to FPS Chart global variables

diff --git a/Engine/Source/Runtime/Engine/Public/ChartCreation.h b/Engine/Source/Runtime/Engine/Public/ChartCreation.h
--- a/Engine/Source/Runtime/Engine/Public/ChartCreation.h
+++ b/Engine/Source/Runtime/Engine/Public/ChartCreation.h
@@ -75,11 +75,13 @@ enum FPSChartStats
 	STAT_FPSChart_TotalUnitGame,
 	STAT_FPSChart_TotalUnitGPU,
 };
+
 /** Start time of current FPS chart.										*/
-extern double			GFPSChartStartTime;
+extern ENGINE_API double GFPSChartStartTime;
 /** FPS chart information. Time spent for each bucket and count.			*/
-extern FFPSChartEntry	GFPSChart[STAT_FPSChartLastBucketStat - STAT_FPSChartFirstStat];
-
+extern ENGINE_API FFPSChartEntry GFPSChart[STAT_FPSChartLastBucketStat - STAT_FPSChartFirstStat];
+/** The total GPU time taken to render all frames. In seconds.				*/
+extern ENGINE_API double GTotalGPUTime;
 
 /** Helper structure for hitch entries. */
 struct FHitchChartEntry
