Fix for UE-8445. Two issues:
Shaders were removed from ShaderJobs but not processed until later causing shaders to think they were finished compilation but didn't have results yet.
Compiling id wasn't being given to some shaders which resulted in shaders thinking they were finished compilation but again didn't have results yet.

diff --git a/Engine/Source/Runtime/Engine/Private/Materials/MaterialShared.cpp b/Engine/Source/Runtime/Engine/Private/Materials/MaterialShared.cpp
--- a/Engine/Source/Runtime/Engine/Private/Materials/MaterialShared.cpp
+++ b/Engine/Source/Runtime/Engine/Private/Materials/MaterialShared.cpp
@@ -617,7 +617,7 @@ void FMaterial::SerializeInlineShaderMap(FArchive& Ar)
 			FinishCompilation();
 
 			bool bValid = GameThreadShaderMap != NULL && GameThreadShaderMap->CompiledSuccessfully();
-
+			
 			Ar << bValid;
 
 			if (bValid)
@@ -1293,6 +1293,7 @@ bool FMaterial::CacheShaders(const FMaterialShaderMapId& ShaderMapId, EShaderPla
 
 	if (GameThreadShaderMap && GameThreadShaderMap->TryToAddToExistingCompilationTask(this))
 	{
+		OutstandingCompileShaderMapIds.Add(GameThreadShaderMap->GetCompilingId());
 		// Reset the shader map so the default material will be used until the compile finishes.
 		GameThreadShaderMap = NULL;
 		bSucceeded = true;
diff --git a/Engine/Source/Runtime/Engine/Private/ShaderCompiler/ShaderCompiler.cpp b/Engine/Source/Runtime/Engine/Private/ShaderCompiler/ShaderCompiler.cpp
--- a/Engine/Source/Runtime/Engine/Private/ShaderCompiler/ShaderCompiler.cpp
+++ b/Engine/Source/Runtime/Engine/Private/ShaderCompiler/ShaderCompiler.cpp
@@ -1928,6 +1928,8 @@ void FShaderCompilingManager::FinishCompilation(const TCHAR* MaterialName, const
 	FScopedSlowTask SlowTask(0, StatusUpdate, GIsEditor && !IsRunningCommandlet());
 
 	TMap<int32, FShaderMapFinalizeResults> CompiledShaderMaps;
+	CompiledShaderMaps.Append( PendingFinalizeShaderMaps );
+	PendingFinalizeShaderMaps.Empty();
 	BlockOnShaderMapCompletion(ShaderMapIdsToFinishCompiling, CompiledShaderMaps);
 
 	bool bRetry = false;
@@ -1951,6 +1953,9 @@ void FShaderCompilingManager::FinishAllCompilation()
 	const double StartTime = FPlatformTime::Seconds();
 
 	TMap<int32, FShaderMapFinalizeResults> CompiledShaderMaps;
+	CompiledShaderMaps.Append( PendingFinalizeShaderMaps );
+	PendingFinalizeShaderMaps.Empty();
+	
 	BlockOnAllShaderMapCompletion(CompiledShaderMaps);
 
 	bool bRetry = false;
