#ANIM: Fixed issue with playing one bone animation -

Merging

//depot/Partners/Weta/UE4/Engine/Source/Runtime/Engine/Private/Animation/AnimSequence.cpp

to //depot/UE4/Engine/Source/Runtime/Engine/Private/Animation/AnimSequence.cpp

diff --git a/Engine/Source/Runtime/Engine/Private/Animation/AnimSequence.cpp b/Engine/Source/Runtime/Engine/Private/Animation/AnimSequence.cpp
--- a/Engine/Source/Runtime/Engine/Private/Animation/AnimSequence.cpp
+++ b/Engine/Source/Runtime/Engine/Private/Animation/AnimSequence.cpp
@@ -1023,14 +1023,17 @@ void UAnimSequence::GetBonePose(FTransformArrayA2& OutAtoms, const FBoneContaine
 		RetargetBoneTransform(RootAtom, 0, 0, RequiredBones);
 	}
 
-	// get the remaining bone atoms
-	AnimationFormat_GetAnimationPose(	
-		*((FTransformArray*)&OutAtoms), //@TODO:@ANIMATION: Nasty hack
-		RotationScalePairs,
-		TranslationPairs,
-		RotationScalePairs, 
-		*this,
-		ExtractionContext.CurrentTime);
+	if (RotationScalePairs.Num() > 0)
+	{
+		// get the remaining bone atoms
+		AnimationFormat_GetAnimationPose(
+			*((FTransformArray*)&OutAtoms), //@TODO:@ANIMATION: Nasty hack
+			RotationScalePairs,
+			TranslationPairs,
+			RotationScalePairs,
+			*this,
+			ExtractionContext.CurrentTime);
+	}
 
 	// Once pose has been extracted, snap root bone back to first frame if we are extracting root motion.
 	if( ExtractionContext.bExtractRootMotion && bEnableRootMotion)
