Merging Linux dedicated server fix

//depot/UE4-UT/...

to //depot/UE4/...

diff --git a/Engine/Source/Runtime/Engine/Private/GameUserSettings.cpp b/Engine/Source/Runtime/Engine/Private/GameUserSettings.cpp
--- a/Engine/Source/Runtime/Engine/Private/GameUserSettings.cpp
+++ b/Engine/Source/Runtime/Engine/Private/GameUserSettings.cpp
@@ -211,6 +211,10 @@ void UGameUserSettings::ApplyNonResolutionSettings()
 
 void UGameUserSettings::ApplyResolutionSettings(bool bCheckForCommandLineOverrides)
 {
+#if UE_SERVER
+	return;
+#endif
+
 	ValidateSettings();
 
 	EWindowMode::Type NewFullscreenMode = GetFullscreenMode();
