Trim down the amount of UObject constructors.

diff --git a/Engine/Source/Runtime/Engine/Private/BrushComponent.cpp b/Engine/Source/Runtime/Engine/Private/BrushComponent.cpp
--- a/Engine/Source/Runtime/Engine/Private/BrushComponent.cpp
+++ b/Engine/Source/Runtime/Engine/Private/BrushComponent.cpp
@@ -718,7 +718,7 @@ void UBrushComponent::BuildSimpleBrushCollision()
 
 	if(BrushBodySetup == NULL)
 	{
-		BrushBodySetup = ConstructObject<UBodySetup>(UBodySetup::StaticClass(), this);
+		BrushBodySetup = NewObject<UBodySetup>(this);
 		check(BrushBodySetup);
 	}
 
diff --git a/Engine/Source/Runtime/Engine/Private/Camera/CameraComponent.cpp b/Engine/Source/Runtime/Engine/Private/Camera/CameraComponent.cpp
--- a/Engine/Source/Runtime/Engine/Private/Camera/CameraComponent.cpp
+++ b/Engine/Source/Runtime/Engine/Private/Camera/CameraComponent.cpp
@@ -60,7 +60,7 @@ void UCameraComponent::OnRegister()
 #if WITH_EDITORONLY_DATA
 	if (ProxyMeshComponent == NULL)
 	{
-		ProxyMeshComponent = ConstructObject<UStaticMeshComponent>(UStaticMeshComponent::StaticClass(), GetOuter(), NAME_None, RF_Transactional);
+		ProxyMeshComponent = NewObject<UStaticMeshComponent>(GetOuter(), NAME_None, RF_Transactional);
 		ProxyMeshComponent->AttachTo(this);
 		ProxyMeshComponent->StaticMesh = CameraMesh;
 		ProxyMeshComponent->SetCollisionProfileName(UCollisionProfile::NoCollision_ProfileName);
@@ -73,7 +73,7 @@ void UCameraComponent::OnRegister()
 
 	if (DrawFrustum == NULL)
 	{
-		DrawFrustum = ConstructObject<UDrawFrustumComponent>(UDrawFrustumComponent::StaticClass(), GetOuter(), NAME_None, RF_Transactional);
+		DrawFrustum = NewObject<UDrawFrustumComponent>(GetOuter(), NAME_None, RF_Transactional);
 		DrawFrustum->AttachTo(this);
 		DrawFrustum->AlwaysLoadOnClient = false;
 		DrawFrustum->AlwaysLoadOnServer = false;
diff --git a/Engine/Source/Runtime/Engine/Private/CanvasRenderTarget2D.cpp b/Engine/Source/Runtime/Engine/Private/CanvasRenderTarget2D.cpp
--- a/Engine/Source/Runtime/Engine/Private/CanvasRenderTarget2D.cpp
+++ b/Engine/Source/Runtime/Engine/Private/CanvasRenderTarget2D.cpp
@@ -26,7 +26,7 @@ void UCanvasRenderTarget2D::UpdateResource()
 	UCanvas* Canvas = (UCanvas*)StaticFindObjectFast(UCanvas::StaticClass(), GetTransientPackage(), CanvasName );
 	if (Canvas == nullptr)
 	{
-		Canvas = ConstructObject<UCanvas>(UCanvas::StaticClass(), GetTransientPackage(), CanvasName );
+		Canvas = NewObject<UCanvas>(GetTransientPackage(), CanvasName);
 		Canvas->AddToRoot();
 	}
 
@@ -85,7 +85,7 @@ UCanvasRenderTarget2D* UCanvasRenderTarget2D::CreateCanvasRenderTarget2D(UObject
 {
 	if ((Width > 0) && (Height > 0) && (CanvasRenderTarget2DClass != NULL))
 	{
-		UCanvasRenderTarget2D* NewCanvasRenderTarget = ConstructObject<UCanvasRenderTarget2D>(CanvasRenderTarget2DClass, GetTransientPackage());
+		UCanvasRenderTarget2D* NewCanvasRenderTarget = NewObject<UCanvasRenderTarget2D>(GetTransientPackage(), CanvasRenderTarget2DClass);
 		if (NewCanvasRenderTarget)
 		{
 			NewCanvasRenderTarget->World = GEngine->GetWorldFromContextObject(WorldContextObject);
diff --git a/Engine/Source/Runtime/Engine/Private/CapsuleComponent.cpp b/Engine/Source/Runtime/Engine/Private/CapsuleComponent.cpp
--- a/Engine/Source/Runtime/Engine/Private/CapsuleComponent.cpp
+++ b/Engine/Source/Runtime/Engine/Private/CapsuleComponent.cpp
@@ -154,7 +154,7 @@ void UCapsuleComponent::UpdateBodySetup()
 {
 	if (ShapeBodySetup == NULL || ShapeBodySetup->IsPendingKill())
 	{
-		ShapeBodySetup = ConstructObject<UBodySetup>(UBodySetup::StaticClass(), this);
+		ShapeBodySetup = NewObject<UBodySetup>(this);
 		ShapeBodySetup->CollisionTraceFlag = CTF_UseSimpleAsComplex;
 		ShapeBodySetup->AggGeom.SphylElems.Add(FKSphylElem());
 	}
diff --git a/Engine/Source/Runtime/Engine/Private/DataChannel.cpp b/Engine/Source/Runtime/Engine/Private/DataChannel.cpp
--- a/Engine/Source/Runtime/Engine/Private/DataChannel.cpp
+++ b/Engine/Source/Runtime/Engine/Private/DataChannel.cpp
@@ -2443,7 +2443,7 @@ UObject* UActorChannel::ReadContentBlockHeader(FInBunch & Bunch, bool& bObjectDe
 		// Construct the sub-object
 		UE_LOG( LogNetTraffic, Log, TEXT( "UActorChannel::ReadContentBlockHeader: Instantiating sub-object. Class: %s, Actor: %s" ), *SubObjClass->GetName(), *Actor->GetName() );
 
-		SubObj = ConstructObject< UObject >( SubObjClass, Actor );
+		SubObj = NewObject< UObject >(Actor, SubObjClass);
 
 		// Sanity check some things
 		check( SubObj != NULL );
diff --git a/Engine/Source/Runtime/Engine/Private/Debug/DebugDrawService.cpp b/Engine/Source/Runtime/Engine/Private/Debug/DebugDrawService.cpp
--- a/Engine/Source/Runtime/Engine/Private/Debug/DebugDrawService.cpp
+++ b/Engine/Source/Runtime/Engine/Private/Debug/DebugDrawService.cpp
@@ -78,7 +78,7 @@ void UDebugDrawService::Draw(const FEngineShowFlags Flags, FViewport* Viewport,
 	UCanvas* CanvasObject = FindObject<UCanvas>(GetTransientPackage(),TEXT("DebugCanvasObject"));
 	if (CanvasObject == NULL)
 	{
-		CanvasObject = ConstructObject<UCanvas>(UCanvas::StaticClass(),GetTransientPackage(),TEXT("DebugCanvasObject"));
+		CanvasObject = NewObject<UCanvas>(GetTransientPackage(), TEXT("DebugCanvasObject"));
 		CanvasObject->AddToRoot();
 	}
 
diff --git a/Engine/Source/Runtime/Engine/Private/DemoNetDriver.cpp b/Engine/Source/Runtime/Engine/Private/DemoNetDriver.cpp
--- a/Engine/Source/Runtime/Engine/Private/DemoNetDriver.cpp
+++ b/Engine/Source/Runtime/Engine/Private/DemoNetDriver.cpp
@@ -156,7 +156,7 @@ bool UDemoNetDriver::InitConnect( FNetworkNotify* InNotify, const FURL& ConnectU
 	}
 
 	// Playback, local machine is a client, and the demo stream acts "as if" it's the server.
-	ServerConnection = ConstructObject<UNetConnection>( UDemoNetConnection::StaticClass() );
+	ServerConnection = NewObject<UNetConnection>(GetTransientPackage(), UDemoNetConnection::StaticClass());
 	ServerConnection->InitConnection( this, USOCK_Pending, ConnectURL, 1000000 );
 
 	// use the same byte format regardless of platform so that the demos are cross platform
@@ -289,7 +289,7 @@ bool UDemoNetDriver::InitListen( FNetworkNotify* InNotify, FURL& ListenURL, bool
 	}
 
 	// Recording, local machine is server, demo stream acts "as if" it's a client.
-	UDemoNetConnection* Connection = ConstructObject<UDemoNetConnection>( UDemoNetConnection::StaticClass() );
+	UDemoNetConnection* Connection = NewObject<UDemoNetConnection>();
 	Connection->InitConnection( this, USOCK_Open, ListenURL, 1000000 );
 	Connection->InitSendBuffer();
 
diff --git a/Engine/Source/Runtime/Engine/Private/DeviceProfiles/DeviceProfileManager.cpp b/Engine/Source/Runtime/Engine/Private/DeviceProfiles/DeviceProfileManager.cpp
--- a/Engine/Source/Runtime/Engine/Private/DeviceProfiles/DeviceProfileManager.cpp
+++ b/Engine/Source/Runtime/Engine/Private/DeviceProfiles/DeviceProfileManager.cpp
@@ -135,7 +135,7 @@ UDeviceProfileManager::UDeviceProfileManager( const FObjectInitializer& ObjectIn
 
 UDeviceProfile* UDeviceProfileManager::CreateProfile( const FString& ProfileName )
 {
-	UDeviceProfile* NewProfile = ConstructObject<UDeviceProfile>( UDeviceProfile::StaticClass(), GetTransientPackage(), *ProfileName, RF_Transient|RF_Public);
+	UDeviceProfile* NewProfile = NewObject<UDeviceProfile>(GetTransientPackage(), *ProfileName, RF_Transient | RF_Public);
 	Profiles.Add( NewProfile );
 
 	// Inform the UI that the device list has changed
@@ -150,7 +150,7 @@ UDeviceProfile* UDeviceProfileManager::CreateProfile( const FString& ProfileName
 	UDeviceProfile* DeviceProfile = FindObject<UDeviceProfile>( GetTransientPackage(), *ProfileName );
 	if( DeviceProfile == NULL )
 	{
-		DeviceProfile = ConstructObject<UDeviceProfile>( UDeviceProfile::StaticClass(), GetTransientPackage(), *ProfileName, RF_Transient|RF_Public);
+		DeviceProfile = NewObject<UDeviceProfile>(GetTransientPackage(), *ProfileName, RF_Transient | RF_Public);
 		DeviceProfile->LoadConfig( UDeviceProfile::StaticClass(), *DeviceProfileFileName );
 		DeviceProfile->BaseProfileName = ParentName != TEXT("") ? ParentName : DeviceProfile->BaseProfileName;
 		DeviceProfile->DeviceType = ProfileType;
diff --git a/Engine/Source/Runtime/Engine/Private/DialogueWave.cpp b/Engine/Source/Runtime/Engine/Private/DialogueWave.cpp
--- a/Engine/Source/Runtime/Engine/Private/DialogueWave.cpp
+++ b/Engine/Source/Runtime/Engine/Private/DialogueWave.cpp
@@ -228,7 +228,7 @@ void UDialogueWave::UpdateMappingProxy(FDialogueContextMapping& ContextMapping)
 	{
 		if (!ContextMapping.Proxy)
 		{
-			ContextMapping.Proxy = ConstructObject<UDialogueSoundWaveProxy>(UDialogueSoundWaveProxy::StaticClass());
+			ContextMapping.Proxy = NewObject<UDialogueSoundWaveProxy>();
 		}
 	}
 	else
diff --git a/Engine/Source/Runtime/Engine/Private/EdGraph/EdGraph.cpp b/Engine/Source/Runtime/Engine/Private/EdGraph/EdGraph.cpp
--- a/Engine/Source/Runtime/Engine/Private/EdGraph/EdGraph.cpp
+++ b/Engine/Source/Runtime/Engine/Private/EdGraph/EdGraph.cpp
@@ -123,7 +123,7 @@ void UEdGraph::RemoveOnGraphChangedHandler( FDelegateHandle Handle )
 
 UEdGraphNode* UEdGraph::CreateNode( TSubclassOf<UEdGraphNode> NewNodeClass, bool bSelectNewNode/* = true*/ )
 {
-	UEdGraphNode* NewNode = ConstructObject<UEdGraphNode>(NewNodeClass, this, NAME_None, RF_Transactional);
+	UEdGraphNode* NewNode = NewObject<UEdGraphNode>(this, NewNodeClass, NAME_None, RF_Transactional);
 	AddNode(NewNode, false, bSelectNewNode );
 	return NewNode;
 }
diff --git a/Engine/Source/Runtime/Engine/Private/GameEngine.cpp b/Engine/Source/Runtime/Engine/Private/GameEngine.cpp
--- a/Engine/Source/Runtime/Engine/Private/GameEngine.cpp
+++ b/Engine/Source/Runtime/Engine/Private/GameEngine.cpp
@@ -416,7 +416,7 @@ void UGameEngine::Init(IEngineLoop* InEngineLoop)
 		FStringClassReference GameInstanceClassName = GetDefault<UGameMapsSettings>()->GameInstanceClass;
 		UClass* GameInstanceClass = (GameInstanceClassName.IsValid() ? LoadObject<UClass>(NULL, *GameInstanceClassName.ToString()) : UGameInstance::StaticClass());
 		
-		GameInstance = ConstructObject<UGameInstance>(GameInstanceClass, this);
+		GameInstance = NewObject<UGameInstance>(this, GameInstanceClass);
 
 		GameInstance->InitializeStandalone();
 	}
@@ -428,7 +428,7 @@ void UGameEngine::Init(IEngineLoop* InEngineLoop)
 	UGameViewportClient* ViewportClient = NULL;
 	if(GIsClient)
 	{
-		ViewportClient = ConstructObject<UGameViewportClient>(GameViewportClientClass,this);
+		ViewportClient = NewObject<UGameViewportClient>(this, GameViewportClientClass);
 		ViewportClient->Init(*GameInstance->GetWorldContext(), GameInstance);
 		GameViewport = ViewportClient;
 		GameInstance->GetWorldContext()->GameViewport = ViewportClient;
diff --git a/Engine/Source/Runtime/Engine/Private/GameViewportClient.cpp b/Engine/Source/Runtime/Engine/Private/GameViewportClient.cpp
--- a/Engine/Source/Runtime/Engine/Private/GameViewportClient.cpp
+++ b/Engine/Source/Runtime/Engine/Private/GameViewportClient.cpp
@@ -689,7 +689,7 @@ static UCanvas* GetCanvasByName(FName CanvasName)
 		UCanvas* CanvasObject = FindObject<UCanvas>(GetTransientPackage(),*CanvasName.ToString());
 		if( !CanvasObject )
 		{
-			CanvasObject = NewNamedObject<UCanvas>(GetTransientPackage(), CanvasName);
+			CanvasObject = NewObject<UCanvas>(GetTransientPackage(), CanvasName);
 			CanvasObject->AddToRoot();
 		}
 
diff --git a/Engine/Source/Runtime/Engine/Private/GameplayStatics.cpp b/Engine/Source/Runtime/Engine/Private/GameplayStatics.cpp
--- a/Engine/Source/Runtime/Engine/Private/GameplayStatics.cpp
+++ b/Engine/Source/Runtime/Engine/Private/GameplayStatics.cpp
@@ -532,7 +532,7 @@ void UGameplayStatics::PlayWorldCameraShake(UObject* WorldContextObject, TSubcla
 
 UParticleSystemComponent* CreateParticleSystem(UParticleSystem* EmitterTemplate, UWorld* World, AActor* Actor, bool bAutoDestroy)
 {
-	UParticleSystemComponent* PSC = ConstructObject<UParticleSystemComponent>(UParticleSystemComponent::StaticClass(), (Actor ? Actor : (UObject*)World) );
+	UParticleSystemComponent* PSC = NewObject<UParticleSystemComponent>((Actor ? Actor : (UObject*)World));
 	PSC->bAutoDestroy = bAutoDestroy;
 	PSC->SecondsBeforeInactive = 0.0f;
 	PSC->bAutoActivate = false;
@@ -899,7 +899,7 @@ UDecalComponent* CreateDecalComponent(class UMaterialInterface* DecalMaterial, F
 {
 	const FMatrix DecalInternalTransform = FRotationMatrix(FRotator(0.f, 90.0f, -90.0f));
 
-	UDecalComponent* DecalComp = ConstructObject<UDecalComponent>(UDecalComponent::StaticClass(), (Actor ? Actor : (UObject*)GetTransientPackage()));
+	UDecalComponent* DecalComp = NewObject<UDecalComponent>((Actor ? Actor : (UObject*)GetTransientPackage()));
 	DecalComp->DecalMaterial = DecalMaterial;
 	DecalComp->RelativeScale3D = DecalInternalTransform.TransformVector(DecalSize);
 	DecalComp->bAbsoluteScale = true;
diff --git a/Engine/Source/Runtime/Engine/Private/Interpolation.cpp b/Engine/Source/Runtime/Engine/Private/Interpolation.cpp
--- a/Engine/Source/Runtime/Engine/Private/Interpolation.cpp
+++ b/Engine/Source/Runtime/Engine/Private/Interpolation.cpp
@@ -572,7 +572,7 @@ void AMatineeActor::AddPlayerToDirectorTracks(class APlayerController* PC)
 					if( IsMatineeCompatibleWithPlayer( PC ) )
 					{
 						// create a new instance with this player
-						UInterpGroupInstDirector* NewGroupInstDir = ConstructObject<UInterpGroupInstDirector>(UInterpGroupInstDirector::StaticClass(), this, NAME_None, RF_Transactional);
+						UInterpGroupInstDirector* NewGroupInstDir = NewObject<UInterpGroupInstDirector>(this, NAME_None, RF_Transactional);
 						GroupInst.Add(NewGroupInstDir);
 
 						// and initialize the instance
@@ -794,7 +794,7 @@ void AMatineeActor::InitInterp()
 						if (IsMatineeCompatibleWithPlayer( PC ) )
 						{
 							// create a new instance with this player
-							UInterpGroupInstDirector* NewGroupInstDir = ConstructObject<UInterpGroupInstDirector>(UInterpGroupInstDirector::StaticClass(), this, NAME_None, RF_Transactional);
+							UInterpGroupInstDirector* NewGroupInstDir = NewObject<UInterpGroupInstDirector>(this, NAME_None, RF_Transactional);
 							GroupInst.Add(NewGroupInstDir);
 
 							// and initialize the instance
@@ -807,7 +807,7 @@ void AMatineeActor::InitInterp()
 				else
 				{
 					// In the editor always create a director group instance with a NULL group actor since there are no player controllers.
-					UInterpGroupInstDirector* NewGroupInstDir = ConstructObject<UInterpGroupInstDirector>(UInterpGroupInstDirector::StaticClass(), this, NAME_None, RF_Transactional);
+					UInterpGroupInstDirector* NewGroupInstDir = NewObject<UInterpGroupInstDirector>(this, NAME_None, RF_Transactional);
 					GroupInst.Add(NewGroupInstDir);
 
 					// and initialize the instance
@@ -827,7 +827,7 @@ void AMatineeActor::InitInterp()
 						{
 							AActor* Actor = GroupInfo->Actors[ ActorIndex ];
 
-							UInterpGroupInst* NewGroupInst = ConstructObject<UInterpGroupInst>(UInterpGroupInst::StaticClass(), this, NAME_None, RF_Transactional);
+							UInterpGroupInst* NewGroupInst = NewObject<UInterpGroupInst>(this, NAME_None, RF_Transactional);
 							GroupInst.Add(NewGroupInst);
 
 							NewGroupInst->InitGroupInst(Group, Actor);
@@ -837,7 +837,7 @@ void AMatineeActor::InitInterp()
 					{
 						// we need to create groupinst when actor does not exist. 
 						// Create new InterpGroupInst
-						UInterpGroupInst* NewGroupInst = ConstructObject<UInterpGroupInst>( UInterpGroupInst::StaticClass(), this, NAME_None, RF_Transactional );
+						UInterpGroupInst* NewGroupInst = NewObject<UInterpGroupInst>(this, NAME_None, RF_Transactional);
 						GroupInst.Add(NewGroupInst);
 						// Initialise group instance, saving ref to actor it works on.
 						NewGroupInst->InitGroupInst(Group, NULL);
@@ -1791,36 +1791,36 @@ UInterpData::UInterpData(const FObjectInitializer& ObjectInitializer)
 #if WITH_EDITORONLY_DATA
 void UInterpData::CreateDefaultFilters()
 {
-	UInterpFilter* FilterAll = NewNamedObject<UInterpFilter>(this, TEXT("FilterAll"), RF_Transient | RF_TextExportTransient);
+	UInterpFilter* FilterAll = NewObject<UInterpFilter>(this, TEXT("FilterAll"), RF_Transient | RF_TextExportTransient);
 	FilterAll->Caption = TEXT("All");
 	DefaultFilters.Add(FilterAll);
 
-	UInterpFilter_Classes* FilterCameras = NewNamedObject<UInterpFilter_Classes>(this, TEXT("FilterCameras"), RF_Transient | RF_TextExportTransient);
+	UInterpFilter_Classes* FilterCameras = NewObject<UInterpFilter_Classes>(this, TEXT("FilterCameras"), RF_Transient | RF_TextExportTransient);
 	FilterCameras->Caption = TEXT("Cameras");
 	FilterCameras->ClassToFilterBy = ACameraActor::StaticClass();
 	DefaultFilters.Add(FilterCameras);
 
-	UInterpFilter_Classes* FilterSkeletalMeshes = NewNamedObject<UInterpFilter_Classes>(this, TEXT("FilterSkeletalMeshes"), RF_Transient | RF_TextExportTransient);
+	UInterpFilter_Classes* FilterSkeletalMeshes = NewObject<UInterpFilter_Classes>(this, TEXT("FilterSkeletalMeshes"), RF_Transient | RF_TextExportTransient);
 	FilterSkeletalMeshes->Caption = TEXT("Skeletal Meshes");
 	FilterSkeletalMeshes->ClassToFilterBy = ASkeletalMeshActor::StaticClass();
 	DefaultFilters.Add(FilterSkeletalMeshes);
 
-	UInterpFilter_Classes* FilterLighting = NewNamedObject<UInterpFilter_Classes>(this, TEXT("FilterLighting"), RF_Transient | RF_TextExportTransient);
+	UInterpFilter_Classes* FilterLighting = NewObject<UInterpFilter_Classes>(this, TEXT("FilterLighting"), RF_Transient | RF_TextExportTransient);
 	FilterLighting->Caption = TEXT("Lights");
 	FilterLighting->ClassToFilterBy = ALight::StaticClass();
 	DefaultFilters.Add(FilterLighting);
 
-	UInterpFilter_Classes* FilterEmitters = NewNamedObject<UInterpFilter_Classes>(this, TEXT("FilterEmitters"), RF_Transient | RF_TextExportTransient);
+	UInterpFilter_Classes* FilterEmitters = NewObject<UInterpFilter_Classes>(this, TEXT("FilterEmitters"), RF_Transient | RF_TextExportTransient);
 	FilterEmitters->Caption = TEXT("Particles");
 	FilterEmitters->ClassToFilterBy = AEmitter::StaticClass();
 	DefaultFilters.Add(FilterEmitters);
 
-	UInterpFilter_Classes* FilterSounds = NewNamedObject<UInterpFilter_Classes>(this, TEXT("FilterSounds"), RF_Transient | RF_TextExportTransient);
+	UInterpFilter_Classes* FilterSounds = NewObject<UInterpFilter_Classes>(this, TEXT("FilterSounds"), RF_Transient | RF_TextExportTransient);
 	FilterSounds->Caption = TEXT("Sounds");
 	FilterSounds->TrackClasses.Add(UInterpTrackSound::StaticClass());
 	DefaultFilters.Add(FilterSounds);
 
-	UInterpFilter_Classes* FilterEvents = NewNamedObject<UInterpFilter_Classes>(this, TEXT("FilterEvents"), RF_Transient | RF_TextExportTransient);
+	UInterpFilter_Classes* FilterEvents = NewObject<UInterpFilter_Classes>(this, TEXT("FilterEvents"), RF_Transient | RF_TextExportTransient);
 	FilterEvents->Caption = TEXT("Events");
 	FilterEvents->TrackClasses.Add(UInterpTrackEvent::StaticClass());
 	DefaultFilters.Add(FilterEvents);
@@ -2512,7 +2512,7 @@ void UInterpGroupInst::InitGroupInst(UInterpGroup* InGroup, AActor* InGroupActor
 	for(int32 i=0; i<InGroup->InterpTracks.Num(); i++)
 	{
 		// Construct Track instance object
-		UInterpTrackInst* TrInst = ConstructObject<UInterpTrackInst>( InGroup->InterpTracks[i]->TrackInstClass, this, NAME_None, RF_Transactional );
+		UInterpTrackInst* TrInst = NewObject<UInterpTrackInst>(this, InGroup->InterpTracks[i]->TrackInstClass, NAME_None, RF_Transactional);
 		TrackInst.Add(TrInst);
 
 		TrInst->InitTrackInst( InGroup->InterpTracks[i] );
@@ -4743,7 +4743,7 @@ void UInterpTrackMove::CreateSubTracks( bool bCopy )
 		check( TrackDef && TrackDef->bSubTrackOnly );
 
 		UInterpTrack* NewSubTrack = NULL;
-		NewSubTrack = ConstructObject<UInterpTrack>( SubTrackInfo.SupportedClass, this, NAME_None, RF_Transactional );
+		NewSubTrack = NewObject<UInterpTrack>(this, SubTrackInfo.SupportedClass, NAME_None, RF_Transactional);
 		check( NewSubTrack );
 
 		int32 NewTrackIndex = SubTracks.Add( NewSubTrack );
@@ -7983,7 +7983,7 @@ void UInterpTrackSound::PreviewUpdateTrack(float NewPosition, UInterpTrackInst*
 		FSoundTrackKey& SoundTrackKey = GetSoundTrackKeyAtPosition(NewPosition);
 		const bool bIsInRange = NewPosition >= SoundTrackKey.Time && NewPosition <= ( SoundTrackKey.Time + SoundTrackKey.Sound->Duration );
 
-		USoundCue* TempPlaybackAudioCue = ConstructObject<USoundCue>(USoundCue::StaticClass());
+		USoundCue* TempPlaybackAudioCue = NewObject<USoundCue>();
 		UAudioComponent* Component = FAudioDevice::CreateComponent(TempPlaybackAudioCue, NULL, NULL, false, false);
 
 		if ( bIsInRange && Component )
diff --git a/Engine/Source/Runtime/Engine/Private/LevelActor.cpp b/Engine/Source/Runtime/Engine/Private/LevelActor.cpp
--- a/Engine/Source/Runtime/Engine/Private/LevelActor.cpp
+++ b/Engine/Source/Runtime/Engine/Private/LevelActor.cpp
@@ -307,7 +307,7 @@ AActor* UWorld::SpawnActor( UClass* Class, FVector const* Location, FRotator con
 		// Spawn in the same level as the owner if we have one. @warning: this relies on the outer of an actor being the level.
 		LevelToSpawnIn = (SpawnParameters.Owner != NULL) ? CastChecked<ULevel>(SpawnParameters.Owner->GetOuter()) : CurrentLevel;
 	}
-	AActor* Actor = ConstructObject<AActor>( Class, LevelToSpawnIn, SpawnParameters.Name, SpawnParameters.ObjectFlags, Template );
+	AActor* Actor = NewObject<AActor>(LevelToSpawnIn, Class, SpawnParameters.Name, SpawnParameters.ObjectFlags, Template);
 	check(Actor);
 
 #if WITH_EDITOR
diff --git a/Engine/Source/Runtime/Engine/Private/LevelScriptActor.cpp b/Engine/Source/Runtime/Engine/Private/LevelScriptActor.cpp
--- a/Engine/Source/Runtime/Engine/Private/LevelScriptActor.cpp
+++ b/Engine/Source/Runtime/Engine/Private/LevelScriptActor.cpp
@@ -58,7 +58,7 @@ void ALevelScriptActor::PreInitializeComponents()
 	if(BGClass != NULL)
 	{
 		// create an InputComponent object so that the level script actor can bind key events
-		InputComponent = ConstructObject<UInputComponent>(UInputComponent::StaticClass(), this);
+		InputComponent = NewObject<UInputComponent>(this);
 		InputComponent->RegisterComponent();
 
 		UInputDelegateBinding::BindInputDelegates(BGClass, InputComponent);
diff --git a/Engine/Source/Runtime/Engine/Private/LightComponent.cpp b/Engine/Source/Runtime/Engine/Private/LightComponent.cpp
--- a/Engine/Source/Runtime/Engine/Private/LightComponent.cpp
+++ b/Engine/Source/Runtime/Engine/Private/LightComponent.cpp
@@ -175,7 +175,7 @@ void ULightComponentBase::OnRegister()
 #if WITH_EDITOR
 	if (SpriteComponent == NULL && GetOwner() && !GetWorld()->IsGameWorld())
 	{
-		SpriteComponent = ConstructObject<UBillboardComponent>(UBillboardComponent::StaticClass(), GetOwner(), NAME_None, RF_Transactional | RF_TextExportTransient);
+		SpriteComponent = NewObject<UBillboardComponent>(GetOwner(), NAME_None, RF_Transactional | RF_TextExportTransient);
 
 		SpriteComponent->AttachTo(this);
 		SpriteComponent->AlwaysLoadOnClient = false;
diff --git a/Engine/Source/Runtime/Engine/Private/LocalPlayer.cpp b/Engine/Source/Runtime/Engine/Private/LocalPlayer.cpp
--- a/Engine/Source/Runtime/Engine/Private/LocalPlayer.cpp
+++ b/Engine/Source/Runtime/Engine/Private/LocalPlayer.cpp
@@ -186,7 +186,7 @@ void ULocalPlayer::InitOnlineSession()
 	if (OnlineSession == NULL)
 	{
 		UClass* SpawnClass = GetOnlineSessionClass();
-		OnlineSession = ConstructObject<UOnlineSession>(SpawnClass, this);
+		OnlineSession = NewObject<UOnlineSession>(this, SpawnClass);
 		if (OnlineSession)
 		{
 			UWorld* World = GetWorld();
diff --git a/Engine/Source/Runtime/Engine/Private/Materials/Material.cpp b/Engine/Source/Runtime/Engine/Private/Materials/Material.cpp
--- a/Engine/Source/Runtime/Engine/Private/Materials/Material.cpp
+++ b/Engine/Source/Runtime/Engine/Private/Materials/Material.cpp
@@ -1789,7 +1789,7 @@ void UMaterial::BackwardsCompatibilityInputConversion()
 
 			check( GConvertFromDiffSpecMaterialFunction );
 
-			UMaterialExpressionMaterialFunctionCall* FunctionExpression = ConstructObject< UMaterialExpressionMaterialFunctionCall >( UMaterialExpressionMaterialFunctionCall::StaticClass(), this );
+			UMaterialExpressionMaterialFunctionCall* FunctionExpression = NewObject<UMaterialExpressionMaterialFunctionCall>(this);
 			Expressions.Add( FunctionExpression );
 
 			FunctionExpression->MaterialExpressionEditorX += 200;
diff --git a/Engine/Source/Runtime/Engine/Private/Materials/MaterialInstanceDynamic.cpp b/Engine/Source/Runtime/Engine/Private/Materials/MaterialInstanceDynamic.cpp
--- a/Engine/Source/Runtime/Engine/Private/Materials/MaterialInstanceDynamic.cpp
+++ b/Engine/Source/Runtime/Engine/Private/Materials/MaterialInstanceDynamic.cpp
@@ -15,10 +15,7 @@ UMaterialInstanceDynamic::UMaterialInstanceDynamic(const FObjectInitializer& Obj
 UMaterialInstanceDynamic* UMaterialInstanceDynamic::Create(UMaterialInterface* ParentMaterial, UObject* InOuter)
 {
 	UObject* Outer = InOuter ? InOuter : GetTransientPackage();
-	UMaterialInstanceDynamic* MID = ConstructObject<UMaterialInstanceDynamic>(
-		UMaterialInstanceDynamic::StaticClass(),
-		Outer
-		);
+	UMaterialInstanceDynamic* MID = NewObject<UMaterialInstanceDynamic>(Outer);
 	MID->SetParentInternal(ParentMaterial);
 	return MID;
 }
diff --git a/Engine/Source/Runtime/Engine/Private/ModelComponent.cpp b/Engine/Source/Runtime/Engine/Private/ModelComponent.cpp
--- a/Engine/Source/Runtime/Engine/Private/ModelComponent.cpp
+++ b/Engine/Source/Runtime/Engine/Private/ModelComponent.cpp
@@ -502,7 +502,7 @@ void UModelComponent::CreateModelBodySetup()
 {
 	if(ModelBodySetup == NULL)
 	{
-		ModelBodySetup = ConstructObject<UBodySetup>(UBodySetup::StaticClass(), this);
+		ModelBodySetup = NewObject<UBodySetup>(this);
 		check(ModelBodySetup);
 		ModelBodySetup->BodySetupGuid = FGuid::NewGuid();
 	}
diff --git a/Engine/Source/Runtime/Engine/Private/NetConnection.cpp b/Engine/Source/Runtime/Engine/Private/NetConnection.cpp
--- a/Engine/Source/Runtime/Engine/Private/NetConnection.cpp
+++ b/Engine/Source/Runtime/Engine/Private/NetConnection.cpp
@@ -1285,7 +1285,7 @@ UChannel* UNetConnection::CreateChannel( EChannelType ChType, bool bOpenedLocall
 	check(Channels[ChIndex]==NULL);
 
 	// Create channel.
-	UChannel* Channel = ConstructObject<UChannel>( UChannel::ChannelClasses[ChType] );
+	UChannel* Channel = NewObject<UChannel>(GetTransientPackage(), UChannel::ChannelClasses[ChType]);
 	Channel->Init( this, ChIndex, bOpenedLocally );
 	Channels[ChIndex] = Channel;
 	OpenChannels.Add(Channel);
diff --git a/Engine/Source/Runtime/Engine/Private/ObjectLibrary.cpp b/Engine/Source/Runtime/Engine/Private/ObjectLibrary.cpp
--- a/Engine/Source/Runtime/Engine/Private/ObjectLibrary.cpp
+++ b/Engine/Source/Runtime/Engine/Private/ObjectLibrary.cpp
@@ -69,7 +69,7 @@ void UObjectLibrary::PostEditChangeProperty(FPropertyChangedEvent& PropertyChang
 
 class UObjectLibrary* UObjectLibrary::CreateLibrary(UClass* InBaseClass, bool bInHasBlueprintClasses, bool bInUseWeak)
 {
-	UObjectLibrary *NewLibrary = ConstructObject<UObjectLibrary>(UObjectLibrary::StaticClass());
+	UObjectLibrary *NewLibrary = NewObject<UObjectLibrary>();
 
 	NewLibrary->ObjectBaseClass = InBaseClass;
 	NewLibrary->bHasBlueprintClasses = bInHasBlueprintClasses;
diff --git a/Engine/Source/Runtime/Engine/Private/Particles/ParticleBeamModules.cpp b/Engine/Source/Runtime/Engine/Private/Particles/ParticleBeamModules.cpp
--- a/Engine/Source/Runtime/Engine/Private/Particles/ParticleBeamModules.cpp
+++ b/Engine/Source/Runtime/Engine/Private/Particles/ParticleBeamModules.cpp
@@ -58,21 +58,21 @@ void UParticleModuleTypeDataBeam2::InitializeDefaults()
 {
 	if (!Distance.Distribution)
 	{
-		UDistributionFloatConstant* DistributionDistance = NewNamedObject<UDistributionFloatConstant>(this, TEXT("DistributionDistance"));
+		UDistributionFloatConstant* DistributionDistance = NewObject<UDistributionFloatConstant>(this, TEXT("DistributionDistance"));
 		DistributionDistance->Constant = 25.0f;
 		Distance.Distribution = DistributionDistance;
 	}
 
 	if (!TaperFactor.Distribution)
 	{
-		UDistributionFloatConstant* DistributionTaperFactor = NewNamedObject<UDistributionFloatConstant>(this, TEXT("DistributionTaperFactor"));
+		UDistributionFloatConstant* DistributionTaperFactor = NewObject<UDistributionFloatConstant>(this, TEXT("DistributionTaperFactor"));
 		DistributionTaperFactor->Constant = 1.0f;
 		TaperFactor.Distribution = DistributionTaperFactor;
 	}
 
 	if (!TaperScale.Distribution)
 	{
-		UDistributionFloatConstant* DistributionTaperScale = NewNamedObject<UDistributionFloatConstant>(this, TEXT("DistributionTaperScale"));
+		UDistributionFloatConstant* DistributionTaperScale = NewObject<UDistributionFloatConstant>(this, TEXT("DistributionTaperScale"));
 		DistributionTaperScale->Constant = 1.0f;
 		TaperScale.Distribution = DistributionTaperScale;
 	}
@@ -908,21 +908,21 @@ void UParticleModuleBeamModifier::InitializeDefaults()
 {
 	if(!Position.Distribution)
 	{
-		UDistributionVectorConstant* DistributionPosition = NewNamedObject<UDistributionVectorConstant>(this, TEXT("DistributionPosition"));
+		UDistributionVectorConstant* DistributionPosition = NewObject<UDistributionVectorConstant>(this, TEXT("DistributionPosition"));
 		DistributionPosition->Constant = FVector(0.0f, 0.0f, 0.0f);
 		Position.Distribution = DistributionPosition;
 	}
 	
 	if(!Tangent.Distribution)
 	{
-		UDistributionVectorConstant* DistributionTangent = NewNamedObject<UDistributionVectorConstant>(this, TEXT("DistributionTangent"));
+		UDistributionVectorConstant* DistributionTangent = NewObject<UDistributionVectorConstant>(this, TEXT("DistributionTangent"));
 		DistributionTangent->Constant = FVector(0.0f, 0.0f, 0.0f);
 		Tangent.Distribution = DistributionTangent;
 	}
 
 	if(!Strength.Distribution)
 	{
-		UDistributionFloatConstant* DistributionStrength = NewNamedObject<UDistributionFloatConstant>(this, TEXT("DistributionStrength"));
+		UDistributionFloatConstant* DistributionStrength = NewObject<UDistributionFloatConstant>(this, TEXT("DistributionStrength"));
 		DistributionStrength->Constant = 0.0f;
 		Strength.Distribution = DistributionStrength;
 	}
@@ -933,15 +933,15 @@ void UParticleModuleBeamModifier::PostInitProperties()
 	Super::PostInitProperties();
 	if (!HasAnyFlags(RF_ClassDefaultObject | RF_NeedLoad))
 	{
-		UDistributionVectorConstant* DistributionPosition = NewNamedObject<UDistributionVectorConstant>(this, TEXT("DistributionPosition"));
+		UDistributionVectorConstant* DistributionPosition = NewObject<UDistributionVectorConstant>(this, TEXT("DistributionPosition"));
 		DistributionPosition->Constant = FVector(0.0f, 0.0f, 0.0f);
 		Position.Distribution = DistributionPosition;
 
-		UDistributionVectorConstant* DistributionTangent = NewNamedObject<UDistributionVectorConstant>(this, TEXT("DistributionTangent"));
+		UDistributionVectorConstant* DistributionTangent = NewObject<UDistributionVectorConstant>(this, TEXT("DistributionTangent"));
 		DistributionTangent->Constant = FVector(0.0f, 0.0f, 0.0f);
 		Tangent.Distribution = DistributionTangent;
 
-		UDistributionFloatConstant* DistributionStrength = NewNamedObject<UDistributionFloatConstant>(this, TEXT("DistributionStrength"));
+		UDistributionFloatConstant* DistributionStrength = NewObject<UDistributionFloatConstant>(this, TEXT("DistributionStrength"));
 		DistributionStrength->Constant = 0.0f;
 		Strength.Distribution = DistributionStrength;
 	}
@@ -1174,35 +1174,35 @@ void UParticleModuleBeamNoise::InitializeDefaults()
 {
 	if (!NoiseSpeed.Distribution)
 	{
-		UDistributionVectorConstant* DistributionNoiseSpeed = NewNamedObject<UDistributionVectorConstant>(this, TEXT("DistributionNoiseSpeed"));
+		UDistributionVectorConstant* DistributionNoiseSpeed = NewObject<UDistributionVectorConstant>(this, TEXT("DistributionNoiseSpeed"));
 		DistributionNoiseSpeed->Constant = FVector(50.0f, 50.0f, 50.0f);
 		NoiseSpeed.Distribution = DistributionNoiseSpeed;
 	}
 
 	if (!NoiseRange.Distribution)
 	{
-		UDistributionVectorConstant* DistributionNoiseRange = NewNamedObject<UDistributionVectorConstant>(this, TEXT("DistributionNoiseRange"));
+		UDistributionVectorConstant* DistributionNoiseRange = NewObject<UDistributionVectorConstant>(this, TEXT("DistributionNoiseRange"));
 		DistributionNoiseRange->Constant = FVector(50.0f, 50.0f, 50.0f);
 		NoiseRange.Distribution = DistributionNoiseRange;
 	}
 
 	if (!NoiseRangeScale.Distribution)
 	{
-		UDistributionFloatConstant* DistributionNoiseRangeScale = NewNamedObject<UDistributionFloatConstant>(this, TEXT("DistributionNoiseRangeScale"));
+		UDistributionFloatConstant* DistributionNoiseRangeScale = NewObject<UDistributionFloatConstant>(this, TEXT("DistributionNoiseRangeScale"));
 		DistributionNoiseRangeScale->Constant = 1.0f;
 		NoiseRangeScale.Distribution = DistributionNoiseRangeScale; 
 	}
 
 	if (!NoiseTangentStrength.Distribution)
 	{
-		UDistributionFloatConstant* DistributionNoiseTangentStrength = NewNamedObject<UDistributionFloatConstant>(this, TEXT("DistributionNoiseTangentStrength"));
+		UDistributionFloatConstant* DistributionNoiseTangentStrength = NewObject<UDistributionFloatConstant>(this, TEXT("DistributionNoiseTangentStrength"));
 		DistributionNoiseTangentStrength->Constant = 250.0f;
 		NoiseTangentStrength.Distribution = DistributionNoiseTangentStrength; 
 	}
 
 	if (!NoiseScale.Distribution)
 	{
-		NoiseScale.Distribution = NewNamedObject<UDistributionFloatConstantCurve>(this, TEXT("DistributionNoiseScale"));
+		NoiseScale.Distribution = NewObject<UDistributionFloatConstantCurve>(this, TEXT("DistributionNoiseScale"));
 	}
 }
 
@@ -1211,23 +1211,23 @@ void UParticleModuleBeamNoise::PostInitProperties()
 	Super::PostInitProperties();
 	if (!HasAnyFlags(RF_ClassDefaultObject | RF_NeedLoad))
 	{
-		UDistributionVectorConstant* DistributionNoiseSpeed = NewNamedObject<UDistributionVectorConstant>(this, TEXT("DistributionNoiseSpeed"));
+		UDistributionVectorConstant* DistributionNoiseSpeed = NewObject<UDistributionVectorConstant>(this, TEXT("DistributionNoiseSpeed"));
 		DistributionNoiseSpeed->Constant = FVector(50.0f, 50.0f, 50.0f);
 		NoiseSpeed.Distribution = DistributionNoiseSpeed;
 
-		UDistributionVectorConstant* DistributionNoiseRange = NewNamedObject<UDistributionVectorConstant>(this, TEXT("DistributionNoiseRange"));
+		UDistributionVectorConstant* DistributionNoiseRange = NewObject<UDistributionVectorConstant>(this, TEXT("DistributionNoiseRange"));
 		DistributionNoiseRange->Constant = FVector(50.0f, 50.0f, 50.0f);
 		NoiseRange.Distribution = DistributionNoiseRange;
 
-		UDistributionFloatConstant* DistributionNoiseRangeScale = NewNamedObject<UDistributionFloatConstant>(this, TEXT("DistributionNoiseRangeScale"));
+		UDistributionFloatConstant* DistributionNoiseRangeScale = NewObject<UDistributionFloatConstant>(this, TEXT("DistributionNoiseRangeScale"));
 		DistributionNoiseRangeScale->Constant = 1.0f;
 		NoiseRangeScale.Distribution = DistributionNoiseRangeScale; 
 
-		UDistributionFloatConstant* DistributionNoiseTangentStrength = NewNamedObject<UDistributionFloatConstant>(this, TEXT("DistributionNoiseTangentStrength"));
+		UDistributionFloatConstant* DistributionNoiseTangentStrength = NewObject<UDistributionFloatConstant>(this, TEXT("DistributionNoiseTangentStrength"));
 		DistributionNoiseTangentStrength->Constant = 250.0f;
 		NoiseTangentStrength.Distribution = DistributionNoiseTangentStrength; 
 
-		NoiseScale.Distribution = NewNamedObject<UDistributionFloatConstantCurve>(this, TEXT("DistributionNoiseScale"));
+		NoiseScale.Distribution = NewObject<UDistributionFloatConstantCurve>(this, TEXT("DistributionNoiseScale"));
 	}
 }
 
@@ -1513,21 +1513,21 @@ void UParticleModuleBeamSource::InitializeDefaults()
 {
 	if (!Source.Distribution)
 	{
-		UDistributionVectorConstant* DistributionSource = NewNamedObject<UDistributionVectorConstant>(this, TEXT("DistributionSource"));
+		UDistributionVectorConstant* DistributionSource = NewObject<UDistributionVectorConstant>(this, TEXT("DistributionSource"));
 		DistributionSource->Constant = FVector(50.0f, 50.0f, 50.0f);
 		Source.Distribution = DistributionSource; 
 	}
 
 	if (!SourceTangent.Distribution)
 	{
-		UDistributionVectorConstant* DistributionSourceTangent = NewNamedObject<UDistributionVectorConstant>(this, TEXT("DistributionSourceTangent")); 
+		UDistributionVectorConstant* DistributionSourceTangent = NewObject<UDistributionVectorConstant>(this, TEXT("DistributionSourceTangent"));
 		DistributionSourceTangent->Constant = FVector(1.0f, 0.0f, 0.0f);
 		SourceTangent.Distribution = DistributionSourceTangent; 
 	}
 
 	if (!SourceStrength.Distribution)
 	{
-		UDistributionFloatConstant* DistributionSourceStrength = NewNamedObject<UDistributionFloatConstant>(this, TEXT("DistributionSourceStrength"));
+		UDistributionFloatConstant* DistributionSourceStrength = NewObject<UDistributionFloatConstant>(this, TEXT("DistributionSourceStrength"));
 		DistributionSourceStrength->Constant = 25.0f;
 		SourceStrength.Distribution = DistributionSourceStrength; 
 	}
@@ -1934,21 +1934,21 @@ void UParticleModuleBeamTarget::InitializeDefaults()
 {
 	if (!Target.Distribution)
 	{
-		UDistributionVectorConstant* DistributionTarget = NewNamedObject<UDistributionVectorConstant>(this, TEXT("DistributionTarget"));
+		UDistributionVectorConstant* DistributionTarget = NewObject<UDistributionVectorConstant>(this, TEXT("DistributionTarget"));
 		DistributionTarget->Constant = FVector(50.0f, 50.0f, 50.0f);
 		Target.Distribution = DistributionTarget;
 	}
 
 	if (!TargetTangent.Distribution)
 	{
-		UDistributionVectorConstant* DistributionTargetTangent = NewNamedObject<UDistributionVectorConstant>(this, TEXT("DistributionTargetTangent"));
+		UDistributionVectorConstant* DistributionTargetTangent = NewObject<UDistributionVectorConstant>(this, TEXT("DistributionTargetTangent"));
 		DistributionTargetTangent->Constant = FVector(1.0f, 0.0f, 0.0f);
 		TargetTangent.Distribution = DistributionTargetTangent; 
 	}
 
 	if (!TargetStrength.Distribution)
 	{
-		UDistributionFloatConstant* DistributionTargetStrength = NewNamedObject<UDistributionFloatConstant>(this, TEXT("DistributionTargetStrength"));
+		UDistributionFloatConstant* DistributionTargetStrength = NewObject<UDistributionFloatConstant>(this, TEXT("DistributionTargetStrength"));
 		DistributionTargetStrength->Constant = 25.0;
 		TargetStrength.Distribution = DistributionTargetStrength;
 	}
diff --git a/Engine/Source/Runtime/Engine/Private/Particles/ParticleComponents.cpp b/Engine/Source/Runtime/Engine/Private/Particles/ParticleComponents.cpp
--- a/Engine/Source/Runtime/Engine/Private/Particles/ParticleComponents.cpp
+++ b/Engine/Source/Runtime/Engine/Private/Particles/ParticleComponents.cpp
@@ -812,7 +812,7 @@ void UParticleLODLevel::ConvertToSpawnModule()
 	}
 
 	UParticleEmitter* EmitterOuter = CastChecked<UParticleEmitter>(GetOuter());
-	SpawnModule = ConstructObject<UParticleModuleSpawn>(UParticleModuleSpawn::StaticClass(), EmitterOuter->GetOuter());
+	SpawnModule = NewObject<UParticleModuleSpawn>(EmitterOuter->GetOuter());
 	check(SpawnModule);
 
 	UDistributionFloat* SourceDist = RequiredModule->SpawnRate.Distribution;
@@ -1454,7 +1454,7 @@ int32 UParticleEmitter::CreateLODLevel(int32 LODLevel, bool bGenerateModuleData)
 	}
 
 	// Create a ParticleLODLevel
-	CreatedLODLevel = ConstructObject<UParticleLODLevel>(UParticleLODLevel::StaticClass(), this);
+	CreatedLODLevel = NewObject<UParticleLODLevel>(this);
 	check(CreatedLODLevel);
 
 	CreatedLODLevel->Level = LODLevel;
@@ -1487,7 +1487,7 @@ int32 UParticleEmitter::CreateLODLevel(int32 LODLevel, bool bGenerateModuleData)
 	else
 	{
 		// Create the RequiredModule
-		UParticleModuleRequired* RequiredModule	= ConstructObject<UParticleModuleRequired>(UParticleModuleRequired::StaticClass(), GetOuter());
+		UParticleModuleRequired* RequiredModule = NewObject<UParticleModuleRequired>(GetOuter());
 		check(RequiredModule);
 		RequiredModule->SetToSensibleDefaults(this);
 		CreatedLODLevel->RequiredModule	= RequiredModule;
@@ -1513,7 +1513,7 @@ int32 UParticleEmitter::CreateLODLevel(int32 LODLevel, bool bGenerateModuleData)
 		RequiredModule->LODValidity = (1 << LODLevel);
 
 		// There must be a spawn module as well...
-		UParticleModuleSpawn* SpawnModule = ConstructObject<UParticleModuleSpawn>(UParticleModuleSpawn::StaticClass(), GetOuter());
+		UParticleModuleSpawn* SpawnModule = NewObject<UParticleModuleSpawn>(GetOuter());
 		check(SpawnModule);
 		CreatedLODLevel->SpawnModule = SpawnModule;
 		SpawnModule->LODValidity = (1 << LODLevel);
@@ -1603,7 +1603,7 @@ bool UParticleEmitter::AutogenerateLowestLODLevel(bool bDuplicateHighest)
 	{
 		// We need to generate it...
 		LODLevels.InsertZeroed(1, 1);
-		UParticleLODLevel* LODLevel = ConstructObject<UParticleLODLevel>(UParticleLODLevel::StaticClass(), this);
+		UParticleLODLevel* LODLevel = NewObject<UParticleLODLevel>(this);
 		check(LODLevel);
 		LODLevels[1]					= LODLevel;
 		LODLevel->Level					= 1;
@@ -1887,7 +1887,7 @@ void UParticleSpriteEmitter::SetToSensibleDefaults()
 	// Create basic set of modules
 
 	// Lifetime module
-	UParticleModuleLifetime* LifetimeModule = ConstructObject<UParticleModuleLifetime>(UParticleModuleLifetime::StaticClass(), GetOuter());
+	UParticleModuleLifetime* LifetimeModule = NewObject<UParticleModuleLifetime>(GetOuter());
 	UDistributionFloatUniform* LifetimeDist = Cast<UDistributionFloatUniform>(LifetimeModule->Lifetime.Distribution);
 	if (LifetimeDist)
 	{
@@ -1899,7 +1899,7 @@ void UParticleSpriteEmitter::SetToSensibleDefaults()
 	LODLevel->Modules.Add(LifetimeModule);
 
 	// Size module
-	UParticleModuleSize* SizeModule = ConstructObject<UParticleModuleSize>(UParticleModuleSize::StaticClass(), GetOuter());
+	UParticleModuleSize* SizeModule = NewObject<UParticleModuleSize>(GetOuter());
 	UDistributionVectorUniform* SizeDist = Cast<UDistributionVectorUniform>(SizeModule->StartSize.Distribution);
 	if (SizeDist)
 	{
@@ -1911,7 +1911,7 @@ void UParticleSpriteEmitter::SetToSensibleDefaults()
 	LODLevel->Modules.Add(SizeModule);
 
 	// Initial velocity module
-	UParticleModuleVelocity* VelModule = ConstructObject<UParticleModuleVelocity>(UParticleModuleVelocity::StaticClass(), GetOuter());
+	UParticleModuleVelocity* VelModule = NewObject<UParticleModuleVelocity>(GetOuter());
 	UDistributionVectorUniform* VelDist = Cast<UDistributionVectorUniform>(VelModule->StartVelocity.Distribution);
 	if (VelDist)
 	{
@@ -1923,7 +1923,7 @@ void UParticleSpriteEmitter::SetToSensibleDefaults()
 	LODLevel->Modules.Add(VelModule);
 
 	// Color over life module
-	UParticleModuleColorOverLife* ColorModule = ConstructObject<UParticleModuleColorOverLife>(UParticleModuleColorOverLife::StaticClass(), GetOuter());
+	UParticleModuleColorOverLife* ColorModule = NewObject<UParticleModuleColorOverLife>(GetOuter());
 	UDistributionVectorConstantCurve* ColorCurveDist = Cast<UDistributionVectorConstantCurve>(ColorModule->ColorOverLife.Distribution);
 	if (ColorCurveDist)
 	{
@@ -3417,7 +3417,7 @@ FParticleDynamicData* UParticleSystemComponent::CreateDynamicData()
 			if( ReplayData == NULL )
 			{
 				// Create a new replay clip!
-				ReplayData = ConstructObject< UParticleSystemReplay >( UParticleSystemReplay::StaticClass(), this );
+				ReplayData = NewObject<UParticleSystemReplay>(this);
 
 				// Set the clip ID number
 				ReplayData->ClipIDNumber = ReplayClipIDNumber;
diff --git a/Engine/Source/Runtime/Engine/Private/Particles/ParticleModules.cpp b/Engine/Source/Runtime/Engine/Private/Particles/ParticleModules.cpp
--- a/Engine/Source/Runtime/Engine/Private/Particles/ParticleModules.cpp
+++ b/Engine/Source/Runtime/Engine/Private/Particles/ParticleModules.cpp
@@ -878,7 +878,7 @@ void UParticleModuleSourceMovement::PostInitProperties()
 	Super::PostInitProperties();
 	if (!HasAnyFlags(RF_ClassDefaultObject | RF_NeedLoad))
 	{
-		UDistributionVectorConstant* DistributionSourceMovementScale = NewNamedObject<UDistributionVectorConstant>(this, TEXT("DistributionSourceMovementScale"));
+		UDistributionVectorConstant* DistributionSourceMovementScale = NewObject<UDistributionVectorConstant>(this, TEXT("DistributionSourceMovementScale"));
 		DistributionSourceMovementScale->Constant = FVector(1.0f, 1.0f, 1.0f);
 		SourceMovementScale.Distribution = DistributionSourceMovementScale;
 	}
@@ -998,7 +998,7 @@ void UParticleModuleRequired::InitializeDefaults()
 {
 	if (!SpawnRate.Distribution)
 	{
-		SpawnRate.Distribution = NewNamedObject<UDistributionFloatConstant>(this, TEXT("RequiredDistributionSpawnRate"));
+		SpawnRate.Distribution = NewObject<UDistributionFloatConstant>(this, TEXT("RequiredDistributionSpawnRate"));
 	}
 }
 
@@ -1147,7 +1147,7 @@ void UParticleModuleMeshRotation::PostInitProperties()
 	Super::PostInitProperties();
 	if (!HasAnyFlags(RF_ClassDefaultObject | RF_NeedLoad))
 	{
-		UDistributionVectorUniform* DistributionStartRotation = NewNamedObject<UDistributionVectorUniform>(this, TEXT("DistributionStartRotation"));
+		UDistributionVectorUniform* DistributionStartRotation = NewObject<UDistributionVectorUniform>(this, TEXT("DistributionStartRotation"));
 		DistributionStartRotation->Min = FVector(0.0f, 0.0f, 0.0f);
 		DistributionStartRotation->Max = FVector(1.0f, 1.0f, 1.0f);
 		StartRotation.Distribution = DistributionStartRotation;
@@ -1245,7 +1245,7 @@ void UParticleModuleMeshRotationRate::PostInitProperties()
 	Super::PostInitProperties();
 	if (!HasAnyFlags(RF_ClassDefaultObject | RF_NeedLoad))
 	{
-		UDistributionVectorUniform* DistributionStartRotationRate = NewNamedObject<UDistributionVectorUniform>(this, TEXT("DistributionStartRotationRate"));
+		UDistributionVectorUniform* DistributionStartRotationRate = NewObject<UDistributionVectorUniform>(this, TEXT("DistributionStartRotationRate"));
 		DistributionStartRotationRate->Min = FVector(0.0f, 0.0f, 0.0f);
 		DistributionStartRotationRate->Max = FVector(360.0f, 360.0f, 360.0f);
 		StartRotationRate.Distribution = DistributionStartRotationRate;
@@ -1340,7 +1340,7 @@ void UParticleModuleMeshRotationRateMultiplyLife::PostInitProperties()
 	Super::PostInitProperties();
 	if (!HasAnyFlags(RF_ClassDefaultObject | RF_NeedLoad))
 	{
-		LifeMultiplier.Distribution = NewNamedObject<UDistributionVectorConstant>(this, TEXT("DistributionLifeMultiplier"));
+		LifeMultiplier.Distribution = NewObject<UDistributionVectorConstant>(this, TEXT("DistributionLifeMultiplier"));
 	}
 }
 
@@ -1399,7 +1399,7 @@ void UParticleModuleMeshRotationRateOverLife::PostInitProperties()
 	Super::PostInitProperties();
 	if (!HasAnyFlags(RF_ClassDefaultObject | RF_NeedLoad))
 	{
-		RotRate.Distribution = NewNamedObject<UDistributionVectorConstantCurve>(this, TEXT("DistributionRotRate"));
+		RotRate.Distribution = NewObject<UDistributionVectorConstantCurve>(this, TEXT("DistributionRotRate"));
 	}
 }
 
@@ -1488,7 +1488,7 @@ void UParticleModuleRotation::InitializeDefaults()
 {
 	if (!StartRotation.Distribution)
 	{
-		UDistributionFloatUniform* DistributionStartRotation = NewNamedObject<UDistributionFloatUniform>(this, TEXT("DistributionStartRotation"));
+		UDistributionFloatUniform* DistributionStartRotation = NewObject<UDistributionFloatUniform>(this, TEXT("DistributionStartRotation"));
 		DistributionStartRotation->Min = 0.0f;
 		DistributionStartRotation->Max = 1.0f;
 		StartRotation.Distribution = DistributionStartRotation;
@@ -1574,7 +1574,7 @@ void UParticleModuleRotationRate::InitializeDefaults()
 {
 	if (!StartRotationRate.Distribution)
 	{
-		StartRotationRate.Distribution = NewNamedObject<UDistributionFloatConstant>(this, TEXT("DistributionStartRotationRate"));
+		StartRotationRate.Distribution = NewObject<UDistributionFloatConstant>(this, TEXT("DistributionStartRotationRate"));
 	}
 }
 
@@ -1685,7 +1685,7 @@ void UParticleModuleRotationOverLifetime::InitializeDefaults()
 {
 	if (!RotationOverLife.Distribution)
 	{
-		RotationOverLife.Distribution = NewNamedObject<UDistributionFloatConstantCurve>(this, TEXT("DistributionRotOverLife"));
+		RotationOverLife.Distribution = NewObject<UDistributionFloatConstantCurve>(this, TEXT("DistributionRotOverLife"));
 	}
 }
 
@@ -1694,7 +1694,7 @@ void UParticleModuleRotationOverLifetime::PostInitProperties()
 	Super::PostInitProperties();
 	if (!HasAnyFlags(RF_ClassDefaultObject | RF_NeedLoad))
 	{
-		RotationOverLife.Distribution = NewNamedObject<UDistributionFloatConstantCurve>(this, TEXT("DistributionRotOverLife"));
+		RotationOverLife.Distribution = NewObject<UDistributionFloatConstantCurve>(this, TEXT("DistributionRotOverLife"));
 	}
 }
 
@@ -1752,7 +1752,7 @@ void UParticleModuleSubUV::InitializeDefaults()
 {
 	if (!SubImageIndex.Distribution)
 	{
-		SubImageIndex.Distribution = NewNamedObject<UDistributionFloatConstant>(this, TEXT("DistributionSubImage"));
+		SubImageIndex.Distribution = NewObject<UDistributionFloatConstant>(this, TEXT("DistributionSubImage"));
 	}
 }
 
@@ -1761,7 +1761,7 @@ void UParticleModuleSubUV::PostInitProperties()
 	Super::PostInitProperties();
 	if (!HasAnyFlags(RF_ClassDefaultObject | RF_NeedLoad))
 	{
-		SubImageIndex.Distribution = NewNamedObject<UDistributionFloatConstant>(this, TEXT("DistributionSubImage"));
+		SubImageIndex.Distribution = NewObject<UDistributionFloatConstant>(this, TEXT("DistributionSubImage"));
 	}
 }
 
@@ -1961,7 +1961,7 @@ void UParticleModuleSubUVMovie::InitializeDefaults()
 {
 	if (!FrameRate.Distribution)
 	{
-		UDistributionFloatConstant* DistributionFrameRate = NewNamedObject<UDistributionFloatConstant>(this, TEXT("DistributionFrameRate"));
+		UDistributionFloatConstant* DistributionFrameRate = NewObject<UDistributionFloatConstant>(this, TEXT("DistributionFrameRate"));
 		DistributionFrameRate->Constant = 30.0f;
 		FrameRate.Distribution = DistributionFrameRate;
 	}
@@ -2128,7 +2128,7 @@ void UParticleModuleRotationRateMultiplyLife::InitializeDefaults()
 {
 	if (!LifeMultiplier.Distribution)
 	{
-		LifeMultiplier.Distribution = NewNamedObject<UDistributionFloatConstant>(this, TEXT("DistributionLifeMultiplier"));
+		LifeMultiplier.Distribution = NewObject<UDistributionFloatConstant>(this, TEXT("DistributionLifeMultiplier"));
 	}
 }
 
@@ -2289,7 +2289,7 @@ void UParticleModuleAccelerationDrag::InitializeDefaults()
 {
 	if (!DragCoefficient)
 	{
-		UDistributionFloatConstant* DistributionDragCoefficient = NewNamedObject<UDistributionFloatConstant>(this, TEXT("DistributionDragCoefficient"));
+		UDistributionFloatConstant* DistributionDragCoefficient = NewObject<UDistributionFloatConstant>(this, TEXT("DistributionDragCoefficient"));
 		DistributionDragCoefficient->Constant = 1.0f;
 		DragCoefficient = DistributionDragCoefficient;
 	}
@@ -2354,7 +2354,7 @@ void UParticleModuleAccelerationDragScaleOverLife::InitializeDefaults()
 {
 	if (!DragScale)
 	{
-		UDistributionFloatConstant* DistributionDragScale = NewNamedObject<UDistributionFloatConstant>(this, TEXT("DistributionDragScale"));
+		UDistributionFloatConstant* DistributionDragScale = NewObject<UDistributionFloatConstant>(this, TEXT("DistributionDragScale"));
 		DistributionDragScale->Constant = 1.0f;
 		DragScale = DistributionDragScale;
 	}
@@ -2410,7 +2410,7 @@ void UParticleModuleAttractorPointGravity::InitializeDefaults()
 {
 	if (!Strength)
 	{
-		UDistributionFloatConstant* DistributionStrength = NewNamedObject<UDistributionFloatConstant>(this, TEXT("DistributionStrength"));
+		UDistributionFloatConstant* DistributionStrength = NewObject<UDistributionFloatConstant>(this, TEXT("DistributionStrength"));
 		DistributionStrength->Constant = 1.0f;
 		Strength = DistributionStrength;
 	}
@@ -2467,7 +2467,7 @@ void UParticleModuleAcceleration::InitializeDefaults()
 {
 	if (!Acceleration.Distribution)
 	{
-		Acceleration.Distribution = NewNamedObject<UDistributionVectorUniform>(this, TEXT("DistributionAcceleration"));
+		Acceleration.Distribution = NewObject<UDistributionVectorUniform>(this, TEXT("DistributionAcceleration"));
 	}
 }
 
@@ -2598,7 +2598,7 @@ void UParticleModuleAccelerationOverLifetime::InitializeDefaults()
 {
 	if (!AccelOverLife.Distribution)
 	{
-		AccelOverLife.Distribution = NewNamedObject<UDistributionVectorConstantCurve>(this, TEXT("DistributionAccelOverLife"));
+		AccelOverLife.Distribution = NewObject<UDistributionVectorConstantCurve>(this, TEXT("DistributionAccelOverLife"));
 	}
 }
 
@@ -2670,22 +2670,22 @@ void UParticleModuleLight::InitializeDefaults()
 {
 	if (!ColorScaleOverLife.Distribution)
 	{
-		ColorScaleOverLife.Distribution = NewNamedObject<UDistributionVectorConstant>(this, TEXT("DistributionColorScaleOverLife"));
+		ColorScaleOverLife.Distribution = NewObject<UDistributionVectorConstant>(this, TEXT("DistributionColorScaleOverLife"));
 	}
 
 	if (!BrightnessOverLife.Distribution)
 	{
-		BrightnessOverLife.Distribution = NewNamedObject<UDistributionFloatConstant>(this, TEXT("DistributionBrightnessOverLife"));
+		BrightnessOverLife.Distribution = NewObject<UDistributionFloatConstant>(this, TEXT("DistributionBrightnessOverLife"));
 	}
 
 	if (!RadiusScale.Distribution)
 	{
-		RadiusScale.Distribution = NewNamedObject<UDistributionFloatConstant>(this, TEXT("DistributionRadiusScale"));
+		RadiusScale.Distribution = NewObject<UDistributionFloatConstant>(this, TEXT("DistributionRadiusScale"));
 	}
 
 	if (!LightExponent.Distribution)
 	{
-		LightExponent.Distribution = NewNamedObject<UDistributionFloatConstant>(this, TEXT("DistributionLightExponent"));
+		LightExponent.Distribution = NewObject<UDistributionFloatConstant>(this, TEXT("DistributionLightExponent"));
 	}
 }
 
@@ -2926,7 +2926,7 @@ void UParticleModuleTypeDataMesh::CreateDistribution()
 {
 	if (!RollPitchYawRange.Distribution)
 	{
-		RollPitchYawRange.Distribution = NewNamedObject<UDistributionVectorUniform>(this, TEXT("DistributionRollPitchYaw"));
+		RollPitchYawRange.Distribution = NewObject<UDistributionVectorUniform>(this, TEXT("DistributionRollPitchYaw"));
 	}
 }
 
@@ -2982,12 +2982,12 @@ void UParticleModuleKillBox::InitializeDefaults()
 {
 	if (!LowerLeftCorner.Distribution)
 	{
-		LowerLeftCorner.Distribution = NewNamedObject<UDistributionVectorConstant>(this, TEXT("DistributionLowerLeftCorner"));
+		LowerLeftCorner.Distribution = NewObject<UDistributionVectorConstant>(this, TEXT("DistributionLowerLeftCorner"));
 	}
 
 	if (!UpperRightCorner.Distribution)
 	{
-		UpperRightCorner.Distribution = NewNamedObject<UDistributionVectorConstant>(this, TEXT("DistributionUpperRightCorner"));
+		UpperRightCorner.Distribution = NewObject<UDistributionVectorConstant>(this, TEXT("DistributionUpperRightCorner"));
 	}
 }
 
@@ -3132,7 +3132,7 @@ void UParticleModuleKillHeight::InitializeDefaults()
 {
 	if (!Height.Distribution)
 	{
-		Height.Distribution = NewNamedObject<UDistributionFloatConstant>(this, TEXT("DistributionHeight"));
+		Height.Distribution = NewObject<UDistributionFloatConstant>(this, TEXT("DistributionHeight"));
 	}
 }
 
@@ -3255,7 +3255,7 @@ void UParticleModuleLifetime::InitializeDefaults()
 {
 	if(!Lifetime.Distribution)
 	{
-		Lifetime.Distribution = NewNamedObject<UDistributionFloatUniform>(this, TEXT("DistributionLifetime"));
+		Lifetime.Distribution = NewObject<UDistributionFloatUniform>(this, TEXT("DistributionLifetime"));
 	}
 }
 
@@ -3405,12 +3405,12 @@ void UParticleModuleAttractorLine::InitializeDefaults()
 {
 	if(!Strength.Distribution)
 	{
-		Strength.Distribution = NewNamedObject<UDistributionFloatConstant>(this, TEXT("DistributionStrength"));
+		Strength.Distribution = NewObject<UDistributionFloatConstant>(this, TEXT("DistributionStrength"));
 	}
 
 	if(!Range.Distribution)
 	{
-		Range.Distribution = NewNamedObject<UDistributionFloatConstant>(this, TEXT("DistributionRange"));
+		Range.Distribution = NewObject<UDistributionFloatConstant>(this, TEXT("DistributionRange"));
 	}
 }
 
@@ -3553,12 +3553,12 @@ void UParticleModuleAttractorParticle::InitializeDefaults()
 {
 	if(!Range.Distribution)
 	{
-		Range.Distribution =  NewNamedObject<UDistributionFloatConstant>(this, TEXT("DistributionRange"));
+		Range.Distribution = NewObject<UDistributionFloatConstant>(this, TEXT("DistributionRange"));
 	}
 
 	if(!Strength.Distribution)
 	{
-		Strength.Distribution = NewNamedObject<UDistributionFloatConstant>(this, TEXT("DistributionStrength"));
+		Strength.Distribution = NewObject<UDistributionFloatConstant>(this, TEXT("DistributionStrength"));
 	}
 }
 
@@ -3796,17 +3796,17 @@ void UParticleModuleAttractorPoint::InitializeDefaults()
 {
 	if(!Position.Distribution)
 	{
-		Position.Distribution = NewNamedObject<UDistributionVectorConstant>(this, TEXT("DistributionPosition"));
+		Position.Distribution = NewObject<UDistributionVectorConstant>(this, TEXT("DistributionPosition"));
 	}
 	
 	if(!Range.Distribution)
 	{
-		Range.Distribution = NewNamedObject<UDistributionFloatConstant>(this, TEXT("DistributionRange"));
+		Range.Distribution = NewObject<UDistributionFloatConstant>(this, TEXT("DistributionRange"));
 	}
 
 	if(!Strength.Distribution)
 	{	
-		Strength.Distribution = NewNamedObject<UDistributionFloatConstant>(this, TEXT("DistributionStrength"));
+		Strength.Distribution = NewObject<UDistributionFloatConstant>(this, TEXT("DistributionStrength"));
 	}
 }
 
diff --git a/Engine/Source/Runtime/Engine/Private/Particles/ParticleModules_Camera.cpp b/Engine/Source/Runtime/Engine/Private/Particles/ParticleModules_Camera.cpp
--- a/Engine/Source/Runtime/Engine/Private/Particles/ParticleModules_Camera.cpp
+++ b/Engine/Source/Runtime/Engine/Private/Particles/ParticleModules_Camera.cpp
@@ -38,7 +38,7 @@ void UParticleModuleCameraOffset::InitializeDefaults()
 {
 	if (!CameraOffset.Distribution)
 	{
-		UDistributionFloatConstant* DistributionCameraOffset = NewNamedObject<UDistributionFloatConstant>(this, TEXT("DistributionCameraOffset")); 
+		UDistributionFloatConstant* DistributionCameraOffset = NewObject<UDistributionFloatConstant>(this, TEXT("DistributionCameraOffset"));
 		DistributionCameraOffset->Constant = 1.0f;
 		CameraOffset.Distribution = DistributionCameraOffset; 
 	}
diff --git a/Engine/Source/Runtime/Engine/Private/Particles/ParticleModules_Collision.cpp b/Engine/Source/Runtime/Engine/Private/Particles/ParticleModules_Collision.cpp
--- a/Engine/Source/Runtime/Engine/Private/Particles/ParticleModules_Collision.cpp
+++ b/Engine/Source/Runtime/Engine/Private/Particles/ParticleModules_Collision.cpp
@@ -57,31 +57,31 @@ void UParticleModuleCollision::InitializeDefaults()
 {
 	if (!DampingFactor.Distribution)
 	{
-		DampingFactor.Distribution = NewNamedObject<UDistributionVectorUniform>(this, TEXT("DistributionDampingFactor"));
+		DampingFactor.Distribution = NewObject<UDistributionVectorUniform>(this, TEXT("DistributionDampingFactor"));
 	}
 
 	if (!DampingFactorRotation.Distribution)
 	{
-		UDistributionVectorConstant* DistributionDampingFactorRotation = NewNamedObject<UDistributionVectorConstant>(this, TEXT("DistributionDampingFactorRotation"));
+		UDistributionVectorConstant* DistributionDampingFactorRotation = NewObject<UDistributionVectorConstant>(this, TEXT("DistributionDampingFactorRotation"));
 		DistributionDampingFactorRotation->Constant = FVector(1.0f, 1.0f, 1.0f);
 		DampingFactorRotation.Distribution = DistributionDampingFactorRotation; 
 	}
 
 	if (!MaxCollisions.Distribution)
 	{
-		MaxCollisions.Distribution = NewNamedObject<UDistributionFloatUniform>(this, TEXT("DistributionMaxCollisions"));
+		MaxCollisions.Distribution = NewObject<UDistributionFloatUniform>(this, TEXT("DistributionMaxCollisions"));
 	}
 
 	if (!ParticleMass.Distribution)
 	{
-		UDistributionFloatConstant* DistributionParticleMass = NewNamedObject<UDistributionFloatConstant>(this, TEXT("DistributionParticleMass")); 
+		UDistributionFloatConstant* DistributionParticleMass = NewObject<UDistributionFloatConstant>(this, TEXT("DistributionParticleMass"));
 		DistributionParticleMass->Constant = 0.1f;
 		ParticleMass.Distribution = DistributionParticleMass;
 	}
 
 	if (!DelayAmount.Distribution)
 	{
-		UDistributionFloatConstant* DistributionDelayAmount = NewNamedObject<UDistributionFloatConstant>(this, TEXT("DistributionDelayAmount"));
+		UDistributionFloatConstant* DistributionDelayAmount = NewObject<UDistributionFloatConstant>(this, TEXT("DistributionDelayAmount"));
 		DistributionDelayAmount->Constant = 0.0f;
 		DelayAmount.Distribution = DistributionDelayAmount;
 	}
@@ -581,14 +581,14 @@ void UParticleModuleCollisionGPU::InitializeDefaults()
 {
 	if (!Resilience.Distribution)
 	{
-		UDistributionFloatConstant* ResilienceDistribution = NewNamedObject<UDistributionFloatConstant>(this, TEXT("ResilienceDistribution"));
+		UDistributionFloatConstant* ResilienceDistribution = NewObject<UDistributionFloatConstant>(this, TEXT("ResilienceDistribution"));
 		ResilienceDistribution->Constant = 0.5f;
 		Resilience.Distribution = ResilienceDistribution;
 	}
 
 	if (!ResilienceScaleOverLife.Distribution)
 	{
-		UDistributionFloatConstant* ResilienceScaleOverLifeDistribution = NewNamedObject<UDistributionFloatConstant>(this, TEXT("ResilienceScaleOverLifeDistribution"));
+		UDistributionFloatConstant* ResilienceScaleOverLifeDistribution = NewObject<UDistributionFloatConstant>(this, TEXT("ResilienceScaleOverLifeDistribution"));
 		ResilienceScaleOverLifeDistribution->Constant = 1.0f;
 		ResilienceScaleOverLife.Distribution = ResilienceScaleOverLifeDistribution;
 	}
diff --git a/Engine/Source/Runtime/Engine/Private/Particles/ParticleModules_Color.cpp b/Engine/Source/Runtime/Engine/Private/Particles/ParticleModules_Color.cpp
--- a/Engine/Source/Runtime/Engine/Private/Particles/ParticleModules_Color.cpp
+++ b/Engine/Source/Runtime/Engine/Private/Particles/ParticleModules_Color.cpp
@@ -47,12 +47,12 @@ void UParticleModuleColor::InitializeDefaults()
 {
 	if (!StartColor.Distribution)
 	{
-		StartColor.Distribution = NewNamedObject<UDistributionVectorConstant>(this, TEXT("DistributionStartColor"));
+		StartColor.Distribution = NewObject<UDistributionVectorConstant>(this, TEXT("DistributionStartColor"));
 	}
 
 	if (!StartAlpha.Distribution)
 	{
-		UDistributionFloatConstant* DistributionStartAlpha = NewNamedObject<UDistributionFloatConstant>(this, TEXT("DistributionStartAlpha"));
+		UDistributionFloatConstant* DistributionStartAlpha = NewObject<UDistributionFloatConstant>(this, TEXT("DistributionStartAlpha"));
 		DistributionStartAlpha->Constant = 1.0f;
 		StartAlpha.Distribution = DistributionStartAlpha;
 	}
@@ -224,12 +224,12 @@ void UParticleModuleColorOverLife::InitializeDefaults()
 {
 	if (!ColorOverLife.Distribution)
 	{
-		ColorOverLife.Distribution = NewNamedObject<UDistributionVectorConstantCurve>(this, TEXT("DistributionColorOverLife"));
+		ColorOverLife.Distribution = NewObject<UDistributionVectorConstantCurve>(this, TEXT("DistributionColorOverLife"));
 	}
 
 	if (!AlphaOverLife.Distribution)
 	{
-		UDistributionFloatConstant* DistributionAlphaOverLife = NewNamedObject<UDistributionFloatConstant>(this, TEXT("DistributionAlphaOverLife"));
+		UDistributionFloatConstant* DistributionAlphaOverLife = NewObject<UDistributionFloatConstant>(this, TEXT("DistributionAlphaOverLife"));
 		DistributionAlphaOverLife->Constant = 1.0f;
 		AlphaOverLife.Distribution = DistributionAlphaOverLife;
 	}
@@ -240,9 +240,9 @@ void UParticleModuleColorOverLife::PostInitProperties()
 	Super::PostInitProperties();
 	if (!HasAnyFlags(RF_ClassDefaultObject | RF_NeedLoad))
 	{
-		ColorOverLife.Distribution = NewNamedObject<UDistributionVectorConstantCurve>(this, TEXT("DistributionColorOverLife"));
+		ColorOverLife.Distribution = NewObject<UDistributionVectorConstantCurve>(this, TEXT("DistributionColorOverLife"));
 
-		UDistributionFloatConstant* DistributionAlphaOverLife = NewNamedObject<UDistributionFloatConstant>(this, TEXT("DistributionAlphaOverLife"));
+		UDistributionFloatConstant* DistributionAlphaOverLife = NewObject<UDistributionFloatConstant>(this, TEXT("DistributionAlphaOverLife"));
 		DistributionAlphaOverLife->Constant = 1.0f;
 		AlphaOverLife.Distribution = DistributionAlphaOverLife;
 	}
@@ -474,12 +474,12 @@ void UParticleModuleColorScaleOverLife::InitializeDefaults()
 {
 	if (!ColorScaleOverLife.Distribution)
 	{
-		ColorScaleOverLife.Distribution = NewNamedObject<UDistributionVectorConstantCurve>(this, TEXT("DistributionColorScaleOverLife"));
+		ColorScaleOverLife.Distribution = NewObject<UDistributionVectorConstantCurve>(this, TEXT("DistributionColorScaleOverLife"));
 	}
 
 	if (!AlphaScaleOverLife.Distribution)
 	{
-		UDistributionFloatConstant* DistributionAlphaScaleOverLife = NewNamedObject<UDistributionFloatConstant>(this, TEXT("DistributionAlphaScaleOverLife"));
+		UDistributionFloatConstant* DistributionAlphaScaleOverLife = NewObject<UDistributionFloatConstant>(this, TEXT("DistributionAlphaScaleOverLife"));
 		DistributionAlphaScaleOverLife->Constant = 1.0f;
 		AlphaScaleOverLife.Distribution = DistributionAlphaScaleOverLife;
 	}
diff --git a/Engine/Source/Runtime/Engine/Private/Particles/ParticleModules_Event.cpp b/Engine/Source/Runtime/Engine/Private/Particles/ParticleModules_Event.cpp
--- a/Engine/Source/Runtime/Engine/Private/Particles/ParticleModules_Event.cpp
+++ b/Engine/Source/Runtime/Engine/Private/Particles/ParticleModules_Event.cpp
@@ -297,14 +297,14 @@ void UParticleModuleEventReceiverSpawn::InitializeDefaults()
 {
 	if (!SpawnCount.Distribution)
 	{
-		UDistributionFloatConstant* RequiredDistributionSpawnCount = NewNamedObject<UDistributionFloatConstant>(this, TEXT("RequiredDistributionSpawnCount"));
+		UDistributionFloatConstant* RequiredDistributionSpawnCount = NewObject<UDistributionFloatConstant>(this, TEXT("RequiredDistributionSpawnCount"));
 		RequiredDistributionSpawnCount->Constant = 0.0f;
 		SpawnCount.Distribution = RequiredDistributionSpawnCount;
 	}
 
 	if (!InheritVelocityScale.Distribution)
 	{
-		UDistributionVectorConstant* RequiredDistributionInheritVelocityScale = NewNamedObject<UDistributionVectorConstant>(this, TEXT("RequiredDistributionInheritVelocityScale"));
+		UDistributionVectorConstant* RequiredDistributionInheritVelocityScale = NewObject<UDistributionVectorConstant>(this, TEXT("RequiredDistributionInheritVelocityScale"));
 		RequiredDistributionInheritVelocityScale->Constant = FVector(1.0f, 1.0f, 1.0f);
 		InheritVelocityScale.Distribution = RequiredDistributionInheritVelocityScale;
 	}
diff --git a/Engine/Source/Runtime/Engine/Private/Particles/ParticleModules_Location.cpp b/Engine/Source/Runtime/Engine/Private/Particles/ParticleModules_Location.cpp
--- a/Engine/Source/Runtime/Engine/Private/Particles/ParticleModules_Location.cpp
+++ b/Engine/Source/Runtime/Engine/Private/Particles/ParticleModules_Location.cpp
@@ -57,7 +57,7 @@ void UParticleModuleLocation::InitializeDefaults()
 {
 	if (!StartLocation.Distribution)
 	{
-		StartLocation.Distribution = NewNamedObject<UDistributionVectorUniform>(this, TEXT("DistributionStartLocation"));
+		StartLocation.Distribution = NewObject<UDistributionVectorUniform>(this, TEXT("DistributionStartLocation"));
 	}
 }
 
@@ -280,23 +280,23 @@ void UParticleModuleLocationDirect::InitializeDefaults()
 {
 	if (!Location.Distribution)
 	{
-		Location.Distribution = NewNamedObject<UDistributionVectorUniform>(this, TEXT("DistributionLocation"));
+		Location.Distribution = NewObject<UDistributionVectorUniform>(this, TEXT("DistributionLocation"));
 	}
 
 	if (!LocationOffset.Distribution)
 	{
-		UDistributionVectorConstant* DistributionLocationOffset = NewNamedObject<UDistributionVectorConstant>(this, TEXT("DistributionLocationOffset"));
+		UDistributionVectorConstant* DistributionLocationOffset = NewObject<UDistributionVectorConstant>(this, TEXT("DistributionLocationOffset"));
 		DistributionLocationOffset->Constant = FVector(0.0f, 0.0f, 0.0f);
 		LocationOffset.Distribution = DistributionLocationOffset;
 	}
 
 	if (!Direction.Distribution)
 	{
-		UDistributionVectorConstant* DistributionScaleFactor = NewNamedObject<UDistributionVectorConstant>(this, TEXT("DistributionScaleFactor"));
+		UDistributionVectorConstant* DistributionScaleFactor = NewObject<UDistributionVectorConstant>(this, TEXT("DistributionScaleFactor"));
 		DistributionScaleFactor->Constant = FVector(1.0f, 1.0f, 1.0f);
 		ScaleFactor.Distribution = DistributionScaleFactor;
 
-		Direction.Distribution = NewNamedObject<UDistributionVectorUniform>(this, TEXT("DistributionDirection"));
+		Direction.Distribution = NewObject<UDistributionVectorUniform>(this, TEXT("DistributionDirection"));
 	}
 }
 
@@ -648,14 +648,14 @@ void UParticleModuleLocationPrimitiveBase::InitializeDefaults()
 {
 	if (!VelocityScale.Distribution)
 	{
-		UDistributionFloatConstant* DistributionVelocityScale = NewNamedObject<UDistributionFloatConstant>(this, TEXT("DistributionVelocityScale"));
+		UDistributionFloatConstant* DistributionVelocityScale = NewObject<UDistributionFloatConstant>(this, TEXT("DistributionVelocityScale"));
 		DistributionVelocityScale->Constant = 1.0f;
 		VelocityScale.Distribution = DistributionVelocityScale;
 	}
 
 	if (!StartLocation.Distribution)
 	{
-		UDistributionVectorConstant* DistributionStartLocation = NewNamedObject<UDistributionVectorConstant>(this, TEXT("DistributionStartLocation"));
+		UDistributionVectorConstant* DistributionStartLocation = NewObject<UDistributionVectorConstant>(this, TEXT("DistributionStartLocation"));
 		DistributionStartLocation->Constant = FVector(0.0f, 0.0f, 0.0f);
 		StartLocation.Distribution = DistributionStartLocation;
 	}
@@ -764,28 +764,28 @@ void UParticleModuleLocationPrimitiveTriangle::InitializeDefaults()
 {
 	if (!StartOffset.Distribution)
 	{
-		UDistributionVectorConstant* DistributionOffset = NewNamedObject<UDistributionVectorConstant>(this, TEXT("DistributionOffset"));
+		UDistributionVectorConstant* DistributionOffset = NewObject<UDistributionVectorConstant>(this, TEXT("DistributionOffset"));
 		DistributionOffset->Constant = FVector::ZeroVector;
 		StartOffset.Distribution = DistributionOffset;
 	}
 
 	if (!Height.Distribution)
 	{
-		UDistributionFloatConstant* DistributionHeight = NewNamedObject<UDistributionFloatConstant>(this, TEXT("DistributionHeight"));
+		UDistributionFloatConstant* DistributionHeight = NewObject<UDistributionFloatConstant>(this, TEXT("DistributionHeight"));
 		DistributionHeight->Constant = 50.0f;
 		Height.Distribution = DistributionHeight;
 	}
 
 	if (!Angle.Distribution)
 	{
-		UDistributionFloatConstant* DistributionAngle = NewNamedObject<UDistributionFloatConstant>(this, TEXT("DistributionAngle"));
+		UDistributionFloatConstant* DistributionAngle = NewObject<UDistributionFloatConstant>(this, TEXT("DistributionAngle"));
 		DistributionAngle->Constant = 90.0f;
 		Angle.Distribution = DistributionAngle;
 	}
 
 	if (!Thickness.Distribution)
 	{
-		UDistributionFloatConstant* DistributionThickness = NewNamedObject<UDistributionFloatConstant>(this, TEXT("DistributionThickness"));
+		UDistributionFloatConstant* DistributionThickness = NewObject<UDistributionFloatConstant>(this, TEXT("DistributionThickness"));
 		DistributionThickness->Constant = 0.0f;
 		Thickness.Distribution = DistributionThickness;
 	}
@@ -931,14 +931,14 @@ void UParticleModuleLocationPrimitiveCylinder::InitializeDefaults()
 {
 	if (!StartRadius.Distribution)
 	{
-		UDistributionFloatConstant* DistributionStartRadius = NewNamedObject<UDistributionFloatConstant>(this, TEXT("DistributionStartRadius"));
+		UDistributionFloatConstant* DistributionStartRadius = NewObject<UDistributionFloatConstant>(this, TEXT("DistributionStartRadius"));
 		DistributionStartRadius->Constant = 50.0f;
 		StartRadius.Distribution = DistributionStartRadius;
 	}
 
 	if (!StartHeight.Distribution)
 	{
-		UDistributionFloatConstant* DistributionStartHeight = NewNamedObject<UDistributionFloatConstant>(this, TEXT("DistributionStartHeight"));
+		UDistributionFloatConstant* DistributionStartHeight = NewObject<UDistributionFloatConstant>(this, TEXT("DistributionStartHeight"));
 		DistributionStartHeight->Constant = 50.0f;
 		StartHeight.Distribution = DistributionStartHeight;
 	}
@@ -1237,7 +1237,7 @@ void UParticleModuleLocationPrimitiveSphere::InitializeDefaults()
 {
 	if (!StartRadius.Distribution)
 	{
-		UDistributionFloatConstant* DistributionStartRadius = NewNamedObject<UDistributionFloatConstant>(this, TEXT("DistributionStartRadius"));
+		UDistributionFloatConstant* DistributionStartRadius = NewObject<UDistributionFloatConstant>(this, TEXT("DistributionStartRadius"));
 		DistributionStartRadius->Constant = 50.0f;
 		StartRadius.Distribution = DistributionStartRadius;
 	}
diff --git a/Engine/Source/Runtime/Engine/Private/Particles/ParticleModules_Orbit.cpp b/Engine/Source/Runtime/Engine/Private/Particles/ParticleModules_Orbit.cpp
--- a/Engine/Source/Runtime/Engine/Private/Particles/ParticleModules_Orbit.cpp
+++ b/Engine/Source/Runtime/Engine/Private/Particles/ParticleModules_Orbit.cpp
@@ -35,7 +35,7 @@ void UParticleModuleOrbit::InitializeDefaults()
 {
 	if (!OffsetAmount.Distribution)
 	{
-		UDistributionVectorUniform* DistributionOffsetAmount = NewNamedObject<UDistributionVectorUniform>(this, TEXT("DistributionOffsetAmount"));
+		UDistributionVectorUniform* DistributionOffsetAmount = NewObject<UDistributionVectorUniform>(this, TEXT("DistributionOffsetAmount"));
 		DistributionOffsetAmount->Min = FVector(0.0f, 0.0f, 0.0f);
 		DistributionOffsetAmount->Max = FVector(0.0f, 50.0f, 0.0f);
 		OffsetAmount.Distribution = DistributionOffsetAmount;
@@ -43,7 +43,7 @@ void UParticleModuleOrbit::InitializeDefaults()
 
 	if (!RotationAmount.Distribution)
 	{
-		UDistributionVectorUniform* DistributionRotationAmount = NewNamedObject<UDistributionVectorUniform>(this, TEXT("DistributionRotationAmount"));
+		UDistributionVectorUniform* DistributionRotationAmount = NewObject<UDistributionVectorUniform>(this, TEXT("DistributionRotationAmount"));
 		DistributionRotationAmount->Min = FVector(0.0f, 0.0f, 0.0f);
 		DistributionRotationAmount->Max = FVector(1.0f, 1.0f, 1.0f);
 		RotationAmount.Distribution = DistributionRotationAmount;
@@ -51,7 +51,7 @@ void UParticleModuleOrbit::InitializeDefaults()
 
 	if (!RotationRateAmount.Distribution)
 	{
-		UDistributionVectorUniform* DistributionRotationRateAmount = NewNamedObject<UDistributionVectorUniform>(this, TEXT("DistributionRotationRateAmount"));
+		UDistributionVectorUniform* DistributionRotationRateAmount = NewObject<UDistributionVectorUniform>(this, TEXT("DistributionRotationRateAmount"));
 		DistributionRotationRateAmount->Min = FVector(0.0f, 0.0f, 0.0f);
 		DistributionRotationRateAmount->Max = FVector(1.0f, 1.0f, 1.0f);
 		RotationRateAmount.Distribution = DistributionRotationRateAmount;
diff --git a/Engine/Source/Runtime/Engine/Private/Particles/ParticleModules_Parameter.cpp b/Engine/Source/Runtime/Engine/Private/Particles/ParticleModules_Parameter.cpp
--- a/Engine/Source/Runtime/Engine/Private/Particles/ParticleModules_Parameter.cpp
+++ b/Engine/Source/Runtime/Engine/Private/Particles/ParticleModules_Parameter.cpp
@@ -42,7 +42,7 @@ void UParticleModuleParameterDynamic::InitializeDefaults()
 	{
 		if(!DynamicParams[ParamIdx].ParamValue.Distribution)
 		{
-			DynamicParams[ParamIdx].ParamValue.Distribution = NewNamedObject<UDistributionFloatConstant>(this, TEXT("DistributionParam1"));
+			DynamicParams[ParamIdx].ParamValue.Distribution = NewObject<UDistributionFloatConstant>(this, TEXT("DistributionParam1"));
 		}
 	}
 }
@@ -63,10 +63,10 @@ void UParticleModuleParameterDynamic::PostInitProperties()
 		};
 		static FConstructorStatics ConstructorStatics;
 
-		DynamicParams.Add(FEmitterDynamicParameter(ConstructorStatics.NAME_None, false, EDPV_UserSet, NewNamedObject<UDistributionFloatConstant>(this, TEXT("DistributionParam1"))));
-		DynamicParams.Add(FEmitterDynamicParameter(ConstructorStatics.NAME_None, false, EDPV_UserSet, NewNamedObject<UDistributionFloatConstant>(this, TEXT("DistributionParam2"))));
-		DynamicParams.Add(FEmitterDynamicParameter(ConstructorStatics.NAME_None, false, EDPV_UserSet, NewNamedObject<UDistributionFloatConstant>(this, TEXT("DistributionParam3"))));
-		DynamicParams.Add(FEmitterDynamicParameter(ConstructorStatics.NAME_None, false, EDPV_UserSet, NewNamedObject<UDistributionFloatConstant>(this, TEXT("DistributionParam4"))));
+		DynamicParams.Add(FEmitterDynamicParameter(ConstructorStatics.NAME_None, false, EDPV_UserSet, NewObject<UDistributionFloatConstant>(this, TEXT("DistributionParam1"))));
+		DynamicParams.Add(FEmitterDynamicParameter(ConstructorStatics.NAME_None, false, EDPV_UserSet, NewObject<UDistributionFloatConstant>(this, TEXT("DistributionParam2"))));
+		DynamicParams.Add(FEmitterDynamicParameter(ConstructorStatics.NAME_None, false, EDPV_UserSet, NewObject<UDistributionFloatConstant>(this, TEXT("DistributionParam3"))));
+		DynamicParams.Add(FEmitterDynamicParameter(ConstructorStatics.NAME_None, false, EDPV_UserSet, NewObject<UDistributionFloatConstant>(this, TEXT("DistributionParam4"))));
 	}
 }
 
diff --git a/Engine/Source/Runtime/Engine/Private/Particles/ParticleModules_Size.cpp b/Engine/Source/Runtime/Engine/Private/Particles/ParticleModules_Size.cpp
--- a/Engine/Source/Runtime/Engine/Private/Particles/ParticleModules_Size.cpp
+++ b/Engine/Source/Runtime/Engine/Private/Particles/ParticleModules_Size.cpp
@@ -42,7 +42,7 @@ void UParticleModuleSize::InitializeDefaults()
 {
 	if (!StartSize.Distribution)
 	{
-		UDistributionVectorUniform* DistributionStartSize = NewNamedObject<UDistributionVectorUniform>(this, TEXT("DistributionStartSize"));
+		UDistributionVectorUniform* DistributionStartSize = NewObject<UDistributionVectorUniform>(this, TEXT("DistributionStartSize"));
 		DistributionStartSize->Min = FVector(1.0f, 1.0f, 1.0f);
 		DistributionStartSize->Max = FVector(1.0f, 1.0f, 1.0f);
 		StartSize.Distribution = DistributionStartSize;
@@ -146,7 +146,7 @@ void UParticleModuleSizeMultiplyLife::InitializeDefaults()
 {
 	if (!LifeMultiplier.Distribution)
 	{
-		LifeMultiplier.Distribution = NewNamedObject<UDistributionVectorConstant>(this, TEXT("DistributionLifeMultiplier"));
+		LifeMultiplier.Distribution = NewObject<UDistributionVectorConstant>(this, TEXT("DistributionLifeMultiplier"));
 	}
 }
 
@@ -336,7 +336,7 @@ void UParticleModuleSizeScale::InitializeDefaults()
 {
 	if (!SizeScale.Distribution)
 	{
-		SizeScale.Distribution = NewNamedObject<UDistributionVectorConstant>(this, TEXT("DistributionSizeScale"));
+		SizeScale.Distribution = NewObject<UDistributionVectorConstant>(this, TEXT("DistributionSizeScale"));
 	}
 }
 
diff --git a/Engine/Source/Runtime/Engine/Private/Particles/ParticleModules_Spawn.cpp b/Engine/Source/Runtime/Engine/Private/Particles/ParticleModules_Spawn.cpp
--- a/Engine/Source/Runtime/Engine/Private/Particles/ParticleModules_Spawn.cpp
+++ b/Engine/Source/Runtime/Engine/Private/Particles/ParticleModules_Spawn.cpp
@@ -38,21 +38,21 @@ void UParticleModuleSpawn::InitializeDefaults()
 {
 	if(!Rate.Distribution)
 	{
-		UDistributionFloatConstant* RequiredDistributionSpawnRate = NewNamedObject<UDistributionFloatConstant>(this, TEXT("RequiredDistributionSpawnRate"));
+		UDistributionFloatConstant* RequiredDistributionSpawnRate = NewObject<UDistributionFloatConstant>(this, TEXT("RequiredDistributionSpawnRate"));
 		RequiredDistributionSpawnRate->Constant = 20.0f;
 		Rate.Distribution = RequiredDistributionSpawnRate;
 	}
 
 	if(!RateScale.Distribution)
 	{
-		UDistributionFloatConstant* RequiredDistributionSpawnRateScale = NewNamedObject<UDistributionFloatConstant>(this, TEXT("RequiredDistributionSpawnRateScale"));
+		UDistributionFloatConstant* RequiredDistributionSpawnRateScale = NewObject<UDistributionFloatConstant>(this, TEXT("RequiredDistributionSpawnRateScale"));
 		RequiredDistributionSpawnRateScale->Constant = 1.0f;
 		RateScale.Distribution = RequiredDistributionSpawnRateScale;
 	}
 
 	if(!BurstScale.Distribution)
 	{
-		UDistributionFloatConstant* BurstScaleDistribution = NewNamedObject<UDistributionFloatConstant>(this, TEXT("BurstScaleDistribution"));
+		UDistributionFloatConstant* BurstScaleDistribution = NewObject<UDistributionFloatConstant>(this, TEXT("BurstScaleDistribution"));
 		BurstScaleDistribution->Constant = 1.0f;
 		BurstScale.Distribution = BurstScaleDistribution;
 	}
@@ -266,7 +266,7 @@ void UParticleModuleSpawnPerUnit::InitializeDefaults()
 {
 	if (!SpawnPerUnit.Distribution)
 	{
-		UDistributionFloatConstant* RequiredDistributionSpawnPerUnit = NewNamedObject<UDistributionFloatConstant>(this, TEXT("RequiredDistributionSpawnPerUnit"));
+		UDistributionFloatConstant* RequiredDistributionSpawnPerUnit = NewObject<UDistributionFloatConstant>(this, TEXT("RequiredDistributionSpawnPerUnit"));
 		RequiredDistributionSpawnPerUnit->Constant = 0.0f;
 		SpawnPerUnit.Distribution = RequiredDistributionSpawnPerUnit;
 	}
diff --git a/Engine/Source/Runtime/Engine/Private/Particles/ParticleModules_VectorField.cpp b/Engine/Source/Runtime/Engine/Private/Particles/ParticleModules_VectorField.cpp
--- a/Engine/Source/Runtime/Engine/Private/Particles/ParticleModules_VectorField.cpp
+++ b/Engine/Source/Runtime/Engine/Private/Particles/ParticleModules_VectorField.cpp
@@ -48,7 +48,7 @@ void UParticleModuleVectorFieldScale::PostInitProperties()
 	Super::PostInitProperties();
 	if (!HasAnyFlags(RF_ClassDefaultObject | RF_NeedLoad))
 	{
-		UDistributionFloatConstant* DistributionVectorFieldScale = NewNamedObject<UDistributionFloatConstant>(this, TEXT("DistributionVectorFieldScale"));
+		UDistributionFloatConstant* DistributionVectorFieldScale = NewObject<UDistributionFloatConstant>(this, TEXT("DistributionVectorFieldScale"));
 		DistributionVectorFieldScale->Constant = 1.0f;
 		VectorFieldScale = DistributionVectorFieldScale;
 	}
@@ -89,7 +89,7 @@ void UParticleModuleVectorFieldScaleOverLife::PostInitProperties()
 	Super::PostInitProperties();
 	if (!HasAnyFlags(RF_ClassDefaultObject | RF_NeedLoad))
 	{
-		UDistributionFloatConstant* DistributionVectorFieldScaleOverLife = NewNamedObject<UDistributionFloatConstant>(this, TEXT("DistributionVectorFieldScaleOverLife"));
+		UDistributionFloatConstant* DistributionVectorFieldScaleOverLife = NewObject<UDistributionFloatConstant>(this, TEXT("DistributionVectorFieldScaleOverLife"));
 		DistributionVectorFieldScaleOverLife->Constant = 1.0f;
 		VectorFieldScaleOverLife = DistributionVectorFieldScaleOverLife;
 	}
diff --git a/Engine/Source/Runtime/Engine/Private/Particles/ParticleModules_Velocity.cpp b/Engine/Source/Runtime/Engine/Private/Particles/ParticleModules_Velocity.cpp
--- a/Engine/Source/Runtime/Engine/Private/Particles/ParticleModules_Velocity.cpp
+++ b/Engine/Source/Runtime/Engine/Private/Particles/ParticleModules_Velocity.cpp
@@ -46,12 +46,12 @@ void UParticleModuleVelocity::InitializeDefaults()
 {
 	if (!StartVelocity.Distribution)
 	{
-		StartVelocity.Distribution = NewNamedObject<UDistributionVectorUniform>(this, TEXT("DistributionStartVelocity"));
+		StartVelocity.Distribution = NewObject<UDistributionVectorUniform>(this, TEXT("DistributionStartVelocity"));
 	}
 
 	if (!StartVelocityRadial.Distribution)
 	{
-		StartVelocityRadial.Distribution = NewNamedObject<UDistributionFloatUniform>(this, TEXT("DistributionStartVelocityRadial"));
+		StartVelocityRadial.Distribution = NewObject<UDistributionFloatUniform>(this, TEXT("DistributionStartVelocityRadial"));
 	}
 }
 
@@ -163,7 +163,7 @@ void UParticleModuleVelocityInheritParent::InitializeDefaults()
 {
 	if (!Scale.Distribution)
 	{
-		UDistributionVectorConstant* DistributionScale = NewNamedObject<UDistributionVectorConstant>(this, TEXT("DistributionScale"));
+		UDistributionVectorConstant* DistributionScale = NewObject<UDistributionVectorConstant>(this, TEXT("DistributionScale"));
 		DistributionScale->Constant = FVector(1.0f, 1.0f, 1.0f);
 		Scale.Distribution = DistributionScale;
 	}
@@ -227,7 +227,7 @@ void UParticleModuleVelocityOverLifetime::InitializeDefaults()
 {
 	if (!VelOverLife.Distribution)
 	{
-		VelOverLife.Distribution = NewNamedObject<UDistributionVectorConstantCurve>(this, TEXT("DistributionVelOverLife"));
+		VelOverLife.Distribution = NewObject<UDistributionVectorConstantCurve>(this, TEXT("DistributionVelOverLife"));
 	}
 }
 
@@ -387,12 +387,12 @@ void UParticleModuleVelocityCone::InitializeDefaults()
 {
 	if (!Angle.Distribution)
 	{
-		Angle.Distribution = NewNamedObject<UDistributionFloatUniform>(this, TEXT("DistributionAngle"));
+		Angle.Distribution = NewObject<UDistributionFloatUniform>(this, TEXT("DistributionAngle"));
 	}
 
 	if (!Velocity.Distribution)
 	{
-		Velocity.Distribution = NewNamedObject<UDistributionFloatUniform>(this, TEXT("DistributionVelocity"));
+		Velocity.Distribution = NewObject<UDistributionFloatUniform>(this, TEXT("DistributionVelocity"));
 	}
 }
 
diff --git a/Engine/Source/Runtime/Engine/Private/Particles/ParticleTrailModules.cpp b/Engine/Source/Runtime/Engine/Private/Particles/ParticleTrailModules.cpp
--- a/Engine/Source/Runtime/Engine/Private/Particles/ParticleTrailModules.cpp
+++ b/Engine/Source/Runtime/Engine/Private/Particles/ParticleTrailModules.cpp
@@ -50,7 +50,7 @@ void UParticleModuleTrailSource::InitializeDefaults()
 {
 	if (!SourceStrength.Distribution)
 	{
-		UDistributionFloatConstant* DistributionSourceStrength = NewNamedObject<UDistributionFloatConstant>(this, TEXT("DistributionSourceStrength"));
+		UDistributionFloatConstant* DistributionSourceStrength = NewObject<UDistributionFloatConstant>(this, TEXT("DistributionSourceStrength"));
 		DistributionSourceStrength->Constant = 100.0f;
 		SourceStrength.Distribution = DistributionSourceStrength;
 	}
diff --git a/Engine/Source/Runtime/Engine/Private/Pawn.cpp b/Engine/Source/Runtime/Engine/Private/Pawn.cpp
--- a/Engine/Source/Runtime/Engine/Private/Pawn.cpp
+++ b/Engine/Source/Runtime/Engine/Private/Pawn.cpp
@@ -556,7 +556,7 @@ class UPlayer* APawn::GetNetOwningPlayer()
 UInputComponent* APawn::CreatePlayerInputComponent()
 {
 	static const FName InputComponentName(TEXT("PawnInputComponent0"));
-	return ConstructObject<UInputComponent>(UInputComponent::StaticClass(), this, InputComponentName);
+	return NewObject<UInputComponent>(this, InputComponentName);
 }
 
 void APawn::DestroyPlayerInputComponent()
diff --git a/Engine/Source/Runtime/Engine/Private/PlayerController.cpp b/Engine/Source/Runtime/Engine/Private/PlayerController.cpp
--- a/Engine/Source/Runtime/Engine/Private/PlayerController.cpp
+++ b/Engine/Source/Runtime/Engine/Private/PlayerController.cpp
@@ -631,7 +631,7 @@ void APlayerController::InitInputSystem()
 {
 	if (PlayerInput == NULL)
 	{
-		PlayerInput = ConstructObject<UPlayerInput>(UPlayerInput::StaticClass(), this);
+		PlayerInput = NewObject<UPlayerInput>(this);
 	}
 
 	SetupInputComponent();
@@ -2254,7 +2254,7 @@ void APlayerController::SetupInputComponent()
 	// A subclass could create a different InputComponent class but still want the default bindings
 	if (InputComponent == NULL)
 	{
-		InputComponent = ConstructObject<UInputComponent>(UInputComponent::StaticClass(), this, TEXT("PC_InputComponent0"));
+		InputComponent = NewObject<UInputComponent>(this, TEXT("PC_InputComponent0"));
 		InputComponent->RegisterComponent();
 	}
 
@@ -3978,7 +3978,7 @@ void APlayerController::UpdateStateInputComponents()
 		if (InactiveStateInputComponent == NULL)
 		{
 			static const FName InactiveStateInputComponentName(TEXT("PC_InactiveStateInputComponent0"));
-			InactiveStateInputComponent = ConstructObject<UInputComponent>(UInputComponent::StaticClass(), this, InactiveStateInputComponentName);
+			InactiveStateInputComponent = NewObject<UInputComponent>(this, InactiveStateInputComponentName);
 			SetupInactiveStateInputComponent(InactiveStateInputComponent);
 			InactiveStateInputComponent->RegisterComponent();
 			PushInputComponent(InactiveStateInputComponent);
diff --git a/Engine/Source/Runtime/Engine/Private/PreviewScene.cpp b/Engine/Source/Runtime/Engine/Private/PreviewScene.cpp
--- a/Engine/Source/Runtime/Engine/Private/PreviewScene.cpp
+++ b/Engine/Source/Runtime/Engine/Private/PreviewScene.cpp
@@ -34,12 +34,12 @@ FPreviewScene::FPreviewScene(FPreviewScene::ConstructionValues CVS)
 
 	GetScene()->UpdateDynamicSkyLight(FLinearColor::White * CVS.SkyBrightness, FLinearColor::Black);
 
-	DirectionalLight = ConstructObject<UDirectionalLightComponent>(UDirectionalLightComponent::StaticClass());
+	DirectionalLight = NewObject<UDirectionalLightComponent>(GetTransientPackage(), TEXT("DirectionalLight"));
 	DirectionalLight->Intensity = CVS.LightBrightness;
 	DirectionalLight->LightColor = FColor::White;
 	AddComponent(DirectionalLight, FTransform(CVS.LightRotation));
 
-	LineBatcher = ConstructObject<ULineBatchComponent>(ULineBatchComponent::StaticClass());
+	LineBatcher = NewObject<ULineBatchComponent>(GetTransientPackage(), TEXT("LineBatcher"));
 	AddComponent(LineBatcher, FTransform::Identity);
 }
 
diff --git a/Engine/Source/Runtime/Engine/Private/ScriptPlatformInterface.cpp b/Engine/Source/Runtime/Engine/Private/ScriptPlatformInterface.cpp
--- a/Engine/Source/Runtime/Engine/Private/ScriptPlatformInterface.cpp
+++ b/Engine/Source/Runtime/Engine/Private/ScriptPlatformInterface.cpp
@@ -142,7 +142,7 @@ void UPlatformInterfaceBase::ClearDelegate(int32 DelegateType, FPlatformInterfac
 			} \
  \
 			/* make the singleton object */ \
-			Singleton = ConstructObject<Class>(SingletonClass); \
+			Singleton = NewObject<Class>(GetTransientPackage(), SingletonClass); \
 			check(Singleton); \
  \
 			/* initialize it */ \
diff --git a/Engine/Source/Runtime/Engine/Private/SkeletalMesh.cpp b/Engine/Source/Runtime/Engine/Private/SkeletalMesh.cpp
--- a/Engine/Source/Runtime/Engine/Private/SkeletalMesh.cpp
+++ b/Engine/Source/Runtime/Engine/Private/SkeletalMesh.cpp
@@ -2351,7 +2351,7 @@ void USkeletalMesh::Serialize( FArchive& Ar )
 	{
 		if ( AssetImportData == NULL )
 		{
-			AssetImportData = ConstructObject<UAssetImportData>(UAssetImportData::StaticClass(), this);
+			AssetImportData = NewObject<UAssetImportData>(this);
 		}
 
 		AssetImportData->SourceFilePath = SourceFilePath_DEPRECATED;
@@ -3330,7 +3330,7 @@ void USkeletalMesh::CreateBodySetup()
 {
 	if (BodySetup == nullptr)
 	{
-		BodySetup = ConstructObject<UBodySetup>(UBodySetup::StaticClass(), this);
+		BodySetup = NewObject<UBodySetup>(this);
 		BodySetup->bSharedCookedData = true;
 	}
 }
diff --git a/Engine/Source/Runtime/Engine/Private/SkeletalMeshComponentPhysics.cpp b/Engine/Source/Runtime/Engine/Private/SkeletalMeshComponentPhysics.cpp
--- a/Engine/Source/Runtime/Engine/Private/SkeletalMeshComponentPhysics.cpp
+++ b/Engine/Source/Runtime/Engine/Private/SkeletalMeshComponentPhysics.cpp
@@ -369,7 +369,7 @@ void USkeletalMeshComponent::CreateBodySetup()
 {
 	if (BodySetup == NULL && SkeletalMesh)
 	{
-		BodySetup = ConstructObject<UBodySetup>(UBodySetup::StaticClass(), this);
+		BodySetup = NewObject<UBodySetup>(this);
 	}
 
 	UBodySetup* OriginalBodySetup = SkeletalMesh->GetBodySetup();
diff --git a/Engine/Source/Runtime/Engine/Private/StaticMesh.cpp b/Engine/Source/Runtime/Engine/Private/StaticMesh.cpp
--- a/Engine/Source/Runtime/Engine/Private/StaticMesh.cpp
+++ b/Engine/Source/Runtime/Engine/Private/StaticMesh.cpp
@@ -1832,7 +1832,7 @@ void UStaticMesh::Serialize(FArchive& Ar)
 	{
 		if ( AssetImportData == NULL )
 		{
-			AssetImportData = ConstructObject<UAssetImportData>(UAssetImportData::StaticClass(), this);
+			AssetImportData = NewObject<UAssetImportData>(this);
 		}
 
 		AssetImportData->SourceFilePath = SourceFilePath_DEPRECATED;
@@ -2096,7 +2096,7 @@ void UStaticMesh::CreateBodySetup()
 {
 	if (BodySetup==NULL)
 	{
-		BodySetup = ConstructObject<UBodySetup>(UBodySetup::StaticClass(), this);
+		BodySetup = NewObject<UBodySetup>(this);
 	}
 }
 
@@ -2104,7 +2104,7 @@ void UStaticMesh::CreateNavCollision()
 {
 	if (NavCollision == NULL && BodySetup != NULL)
 	{
-		NavCollision = ConstructObject<UNavCollision>(UNavCollision::StaticClass(), this);
+		NavCollision = NewObject<UNavCollision>(this);
 		NavCollision->Setup(BodySetup);
 	}
 }
diff --git a/Engine/Source/Runtime/Engine/Private/Texture2D.cpp b/Engine/Source/Runtime/Engine/Private/Texture2D.cpp
--- a/Engine/Source/Runtime/Engine/Private/Texture2D.cpp
+++ b/Engine/Source/Runtime/Engine/Private/Texture2D.cpp
@@ -925,8 +925,7 @@ UTexture2D* UTexture2D::CreateTransient(int32 InSizeX, int32 InSizeY, EPixelForm
 		(InSizeX % GPixelFormats[InFormat].BlockSizeX) == 0 &&
 		(InSizeY % GPixelFormats[InFormat].BlockSizeY) == 0)
 	{
-		NewTexture = ConstructObject<UTexture2D>(
-			UTexture2D::StaticClass(),
+		NewTexture = NewObject<UTexture2D>(
 			GetTransientPackage(),
 			NAME_None,
 			RF_Transient
diff --git a/Engine/Source/Runtime/Engine/Private/UnrealEngine.cpp b/Engine/Source/Runtime/Engine/Private/UnrealEngine.cpp
--- a/Engine/Source/Runtime/Engine/Private/UnrealEngine.cpp
+++ b/Engine/Source/Runtime/Engine/Private/UnrealEngine.cpp
@@ -1441,7 +1441,7 @@ void UEngine::InitializeObjectReferences()
 
 		checkf(SingletonClass != NULL, TEXT("Engine config value GameSingletonClassName is not a valid class name."));
 
-		GameSingleton = ConstructObject<UObject>(SingletonClass, this);
+		GameSingleton = NewObject<UObject>(this, SingletonClass);
 	}
 
 	if (DefaultTireType == NULL && DefaultTireTypeName.ToString().Len())
@@ -7608,7 +7608,7 @@ bool CreateNamedNetDriver_Local(UEngine *Engine, FWorldContext &Context, FName N
 				}
 
 				// Try to create network driver.
-				NetDriver = ConstructObject<UNetDriver>(NetDriverClass);
+				NetDriver = NewObject<UNetDriver>(GetTransientPackage(), NetDriverClass);
 				check(NetDriver);
 				NetDriver->NetDriverName = NetDriverName;
 				
@@ -9985,8 +9985,7 @@ bool UEngine::CommitMapChange( FWorldContext &Context )
 		check( FakePersistentLevel );
 
 		// Construct a new ULevelStreamingPersistent for the new persistent level.
-		ULevelStreamingPersistent* LevelStreamingPersistent = ConstructObject<ULevelStreamingPersistent>(
-			ULevelStreamingPersistent::StaticClass(),
+		ULevelStreamingPersistent* LevelStreamingPersistent = NewObject<ULevelStreamingPersistent>(
 			GetTransientPackage(),
 			*FString::Printf(TEXT("LevelStreamingPersistent_%s"), *FakePersistentLevel->GetOutermost()->GetName()) );
 
@@ -10154,7 +10153,7 @@ FURL& UEngine::LastURLFromWorld(UWorld *World)
 void UEngine::CreateGameUserSettings()
 {
 	UGameUserSettings::LoadConfigIni();
-	GameUserSettings = ConstructObject<UGameUserSettings>(GEngine->GameUserSettingsClass);
+	GameUserSettings = NewObject<UGameUserSettings>(GetTransientPackage(), GEngine->GameUserSettingsClass);
 	GameUserSettings->LoadSettings();
 }
 
@@ -10520,7 +10519,7 @@ UDeviceProfileManager* UEngine::GetDeviceProfileManager()
 {
 	if(DeviceProfileManager == NULL)
 	{
-		DeviceProfileManager = ConstructObject< UDeviceProfileManager >( UDeviceProfileManager::StaticClass(), GetTransientPackage(), TEXT("GlobalDeviceProfileManager"), RF_Public|RF_Transient );
+		DeviceProfileManager = NewObject<UDeviceProfileManager>(GetTransientPackage(), TEXT("GlobalDeviceProfileManager"), RF_Public | RF_Transient);
 	}
 
 	return DeviceProfileManager;
diff --git a/Engine/Source/Runtime/Engine/Private/UnrealExporter.cpp b/Engine/Source/Runtime/Engine/Private/UnrealExporter.cpp
--- a/Engine/Source/Runtime/Engine/Private/UnrealExporter.cpp
+++ b/Engine/Source/Runtime/Engine/Private/UnrealExporter.cpp
@@ -77,13 +77,14 @@ UExporter* UExporter::FindExporter( UObject* Object, const TCHAR* FileType )
 	}
 
 	UClass** E;
+	auto TransientPackage = GetTransientPackage();
 	for (UClass* TempClass = Object->GetClass(); TempClass != NULL; TempClass = TempClass->GetSuperClass())
 	{
 		const bool bFoundExporter = ((E = Exporters.Find( TempClass )) != NULL);
 
 		if( bFoundExporter )
 		{
-			return ConstructObject<UExporter>( *E );
+			return NewObject<UExporter>(TransientPackage, *E);
 		}
 	}
 		
diff --git a/Engine/Source/Runtime/Engine/Private/World.cpp b/Engine/Source/Runtime/Engine/Private/World.cpp
--- a/Engine/Source/Runtime/Engine/Private/World.cpp
+++ b/Engine/Source/Runtime/Engine/Private/World.cpp
@@ -555,7 +555,7 @@ void UWorld::PostLoad()
 		// Make sure thumbnail info exists
 		if ( !ThumbnailInfo )
 		{
-			ThumbnailInfo = ConstructObject<UWorldThumbnailInfo>(UWorldThumbnailInfo::StaticClass(), this);
+			ThumbnailInfo = NewObject<UWorldThumbnailInfo>(this);
 		}
 	}
 #endif
@@ -884,7 +884,7 @@ void UWorld::InitWorld(const InitializationValues IVS)
 			// Spawn the default brush.
 			DefaultBrush = SpawnBrush();
 			check(DefaultBrush->GetBrushComponent());
-			DefaultBrush->Brush = NewNamedObject<UModel>(DefaultBrush->GetOuter(), TEXT("Brush"));
+			DefaultBrush->Brush = NewObject<UModel>(DefaultBrush->GetOuter(), TEXT("Brush"));
 			DefaultBrush->Brush->Initialize(DefaultBrush, 1);
 			DefaultBrush->GetBrushComponent()->Brush = DefaultBrush->Brush;
 			DefaultBrush->SetNotForClientOrServer();
@@ -955,7 +955,7 @@ void UWorld::InitializeNewWorld(const InitializationValues IVS)
 		ClearFlags(RF_Transactional);
 	}
 
-	PersistentLevel = NewNamedObject<ULevel>(this, TEXT("PersistentLevel"));
+	PersistentLevel = NewObject<ULevel>(this, TEXT("PersistentLevel"));
 	PersistentLevel->Initialize(FURL(nullptr));
 	PersistentLevel->Model = NewObject<UModel>(PersistentLevel);
 	PersistentLevel->Model->Initialize(nullptr, 1);
@@ -1051,7 +1051,7 @@ UWorld* UWorld::CreateWorld(const EWorldType::Type InWorldType, bool bInformEngi
 
 	// Create new UWorld, ULevel and UModel.
 	const FString WorldNameString = (WorldName != NAME_None) ? WorldName.ToString() : TEXT("NewWorld");
-	UWorld* NewWorld = NewNamedObject<UWorld>(WorldPackage, *WorldNameString);
+	UWorld* NewWorld = NewObject<UWorld>(WorldPackage, *WorldNameString);
 	NewWorld->SetFlags(RF_Transactional);
 	NewWorld->WorldType = InWorldType;
 	NewWorld->FeatureLevel = InFeatureLevel;
@@ -1176,7 +1176,7 @@ void UWorld::UpdateWorldComponents(bool bRerunConstructionScripts, bool bCurrent
 	{
 		if(!LineBatcher)
 		{
-			LineBatcher = ConstructObject<ULineBatchComponent>(ULineBatchComponent::StaticClass());
+			LineBatcher = NewObject<ULineBatchComponent>();
 		}
 
 		if(!LineBatcher->IsRegistered())
@@ -1186,7 +1186,7 @@ void UWorld::UpdateWorldComponents(bool bRerunConstructionScripts, bool bCurrent
 
 		if(!PersistentLineBatcher)
 		{
-			PersistentLineBatcher = ConstructObject<ULineBatchComponent>(ULineBatchComponent::StaticClass());
+			PersistentLineBatcher = NewObject<ULineBatchComponent>();
 		}
 
 		if(!PersistentLineBatcher->IsRegistered())	
@@ -1196,7 +1196,7 @@ void UWorld::UpdateWorldComponents(bool bRerunConstructionScripts, bool bCurrent
 
 		if(!ForegroundLineBatcher)
 		{
-			ForegroundLineBatcher = ConstructObject<ULineBatchComponent>(ULineBatchComponent::StaticClass());
+			ForegroundLineBatcher = NewObject<ULineBatchComponent>();
 		}
 
 		if(!ForegroundLineBatcher->IsRegistered())	
