Made Scene Previews not construct Navigation System #UE4

- also added a sanity check for _any_ navigation data existance while calling generic pathfinding code

diff --git a/Engine/Source/Runtime/Engine/Private/AI/Navigation/NavigationSystem.cpp b/Engine/Source/Runtime/Engine/Private/AI/Navigation/NavigationSystem.cpp
--- a/Engine/Source/Runtime/Engine/Private/AI/Navigation/NavigationSystem.cpp
+++ b/Engine/Source/Runtime/Engine/Private/AI/Navigation/NavigationSystem.cpp
@@ -1210,7 +1210,7 @@ UNavigationPath* UNavigationSystem::FindPathToLocationSynchronously(UObject* Wor
 
 	UNavigationPath* ResultPath = NULL;
 
-	if (World != NULL && World->GetNavigationSystem() != NULL)
+	if (World != NULL && World->GetNavigationSystem() != NULL && World->GetNavigationSystem()->GetMainNavData() != nullptr)
 	{
 		UNavigationSystem* NavSys = World->GetNavigationSystem();
 
diff --git a/Engine/Source/Runtime/Engine/Private/PreviewScene.cpp b/Engine/Source/Runtime/Engine/Private/PreviewScene.cpp
--- a/Engine/Source/Runtime/Engine/Private/PreviewScene.cpp
+++ b/Engine/Source/Runtime/Engine/Private/PreviewScene.cpp
@@ -28,6 +28,8 @@ FPreviewScene::FPreviewScene(FPreviewScene::ConstructionValues CVS)
 										.AllowAudioPlayback(CVS.bAllowAudioPlayback)
 										.CreatePhysicsScene(CVS.bCreatePhysicsScene)
 										.RequiresHitProxies(false)
+										.CreateNavigation(false)
+										.CreateAISystem(false)
 										.ShouldSimulatePhysics(CVS.bShouldSimulatePhysics)
 										.SetTransactional(CVS.bTransactional));
 	PreviewWorld->InitializeActorsForPlay(FURL());
