fix for UE-8462 Drag-dropping a child-attached actor in the scene outliner to the World gives a bogus error message and does not allow you to reparent it sensibly.  Also fixed related scene outliner/parenting issues.

diff --git a/Engine/Source/Editor/SceneOutliner/Private/FolderTreeItem.cpp b/Engine/Source/Editor/SceneOutliner/Private/FolderTreeItem.cpp
--- a/Engine/Source/Editor/SceneOutliner/Private/FolderTreeItem.cpp
+++ b/Engine/Source/Editor/SceneOutliner/Private/FolderTreeItem.cpp
@@ -75,7 +75,7 @@ FDragValidationInfo FFolderDropTarget::ValidateDrop(FDragDropPayload& DraggedObj
 		{
 			AActor* Actor = WeakActor.Get();
 
-			if (Actor->GetFolderPath() == DestinationPath)
+			if (Actor->GetFolderPath() == DestinationPath && !Actor->GetAttachParentActor())
 			{
 				FFormatNamedArguments Args;
 				Args.Add(TEXT("SourceName"), FText::FromString(Actor->GetActorLabel()));
diff --git a/Engine/Source/Editor/SceneOutliner/Public/SceneOutlinerFilters.h b/Engine/Source/Editor/SceneOutliner/Public/SceneOutlinerFilters.h
--- a/Engine/Source/Editor/SceneOutliner/Public/SceneOutlinerFilters.h
+++ b/Engine/Source/Editor/SceneOutliner/Public/SceneOutlinerFilters.h
@@ -174,7 +174,7 @@ namespace SceneOutliner
 		/** Test whether this tree item passes all filters, and set its interactive state according to the filter it failed (if applicable) */
 		bool TestAndSetInteractiveState(ITreeItem& InItem) const
 		{
-			bool bPassed = false;
+			bool bPassed = true;
 
 			// Default to interactive
 			InItem.Flags.bInteractive = true;
diff --git a/Engine/Source/Editor/UnrealEd/Private/EditorActorFolders.cpp b/Engine/Source/Editor/UnrealEd/Private/EditorActorFolders.cpp
--- a/Engine/Source/Editor/UnrealEd/Private/EditorActorFolders.cpp
+++ b/Engine/Source/Editor/UnrealEd/Private/EditorActorFolders.cpp
@@ -430,7 +430,7 @@ bool FActorFolders::RenameFolderInWorld(UWorld& World, FName OldPath, FName NewP
 		if (Actor->GetFolderPath() == OldPath || PathIsChildOf(OldActorPath.ToString(), OldPathString))
 		{
 			RenamedFolders.Add(OldActorPath);
-			ActorIt->SetFolderPath(OldPathToNewPath(OldPathString, NewPathString, OldActorPath.ToString()));
+			ActorIt->SetFolderPath(OldPathToNewPath(OldPathString, NewPathString, OldActorPath.ToString()), false);
 		}
 	}
 
diff --git a/Engine/Source/Runtime/Engine/Classes/GameFramework/Actor.h b/Engine/Source/Runtime/Engine/Classes/GameFramework/Actor.h
--- a/Engine/Source/Runtime/Engine/Classes/GameFramework/Actor.h
+++ b/Engine/Source/Runtime/Engine/Classes/GameFramework/Actor.h
@@ -1434,9 +1434,10 @@ public:
 
 	/**
 	 * Assigns a new folder to this actor. Actor folder paths are only available in development builds.
-	 * @param	NewFolderPath	The new folder to assign to the actor.
+	 * @param	NewFolderPath	   The new folder to assign to the actor.
+	 * @param   bDetachFromParent  whether to detach this actor from its parent when setting the folder path
 	 */
-	void SetFolderPath(const FName& NewFolderPath);
+	void SetFolderPath(const FName& NewFolderPath, bool bDetachFromParent = true);
 
 	/**
 	 * Used by the "Sync to Content Browser" right-click menu option in the editor.
diff --git a/Engine/Source/Runtime/Engine/Private/ActorEditor.cpp b/Engine/Source/Runtime/Engine/Private/ActorEditor.cpp
--- a/Engine/Source/Runtime/Engine/Private/ActorEditor.cpp
+++ b/Engine/Source/Runtime/Engine/Private/ActorEditor.cpp
@@ -576,7 +576,7 @@ const FName& AActor::GetFolderPath() const
 	return FolderPath;
 }
 
-void AActor::SetFolderPath(const FName& NewFolderPath)
+void AActor::SetFolderPath(const FName& NewFolderPath, bool bDetachFromParent)
 {
 	// Detach the actor if it is attached
 	USceneComponent* RootComp = GetRootComponent();
@@ -593,7 +593,7 @@ void AActor::SetFolderPath(const FName& NewFolderPath)
 	FolderPath = NewFolderPath;
 	
 	// Detach the actor if it is attached
-	if (bIsAttached)
+	if (bIsAttached && bDetachFromParent)
 	{
 		AActor* OldParentActor = RootComp->AttachParent->GetOwner();
 		OldParentActor->Modify();
@@ -605,6 +605,19 @@ void AActor::SetFolderPath(const FName& NewFolderPath)
 	{
 		GEngine->BroadcastLevelActorFolderChanged(this, OldPath);
 	}
+
+	//recursively change folder path for children (but do not detach so they remain as child)
+	if(RootComp)
+	{
+		for(auto ChildSceneComponent : RootComp->AttachChildren)
+		{
+			AActor* ChildActor = ChildSceneComponent ? ChildSceneComponent->GetOwner() : nullptr;
+			if(ChildActor)
+			{
+				ChildActor->SetFolderPath(NewFolderPath, false);
+			}
+		}
+	}
 }
 
 void AActor::CheckForDeprecated()
