Added early-out to USceneComponent::AttachTo() for the case where we're already attached, to fix "FBodyInstance::UpdateBodyScale NegX does not exist" warning caused by USceneComponent calling AttachTo in its OnRegister when it may already be attached (in this case, due to a component reregister context reregistering it without detatching it first) which then ends up causing a spurious detach.
Change to scary code approved by Ori, MarcA and PhillipK

diff --git a/Engine/Source/Runtime/Engine/Private/SceneComponent.cpp b/Engine/Source/Runtime/Engine/Private/SceneComponent.cpp
--- a/Engine/Source/Runtime/Engine/Private/SceneComponent.cpp
+++ b/Engine/Source/Runtime/Engine/Private/SceneComponent.cpp
@@ -815,6 +815,12 @@ void USceneComponent::AttachTo(class USceneComponent* Parent, FName InSocketName
 {
 	if(Parent != NULL)
 	{
+		if (Parent == AttachParent && InSocketName == AttachSocketName && Parent->AttachChildren.Contains(this))
+		{
+			// already attached!
+			return;
+		}
+
 		if(Parent == this)
 		{
 			FMessageLog("PIE").Warning(FText::Format(LOCTEXT("AttachToSelfWarning", "AttachTo: '{0}' cannot be attached to itself. Aborting."), 
