Moved OnTeleported from UCharacterMovementComponent to UMovementComponent, since teleportation is in the Actor API.

diff --git a/Engine/Source/Runtime/Engine/Classes/GameFramework/CharacterMovementComponent.h b/Engine/Source/Runtime/Engine/Classes/GameFramework/CharacterMovementComponent.h
--- a/Engine/Source/Runtime/Engine/Classes/GameFramework/CharacterMovementComponent.h
+++ b/Engine/Source/Runtime/Engine/Classes/GameFramework/CharacterMovementComponent.h
@@ -1135,7 +1135,7 @@ protected:
 public:
 
 	/** Called by owning Character upon successful teleport from AActor::TeleportTo(). */
-	virtual void OnTeleported();
+	virtual void OnTeleported() override;
 
 	/**
 	 * Checks if new capsule size fits (no encroachment), and call CharacterOwner->OnStartCrouch() if successful.
diff --git a/Engine/Source/Runtime/Engine/Classes/GameFramework/MovementComponent.h b/Engine/Source/Runtime/Engine/Classes/GameFramework/MovementComponent.h
--- a/Engine/Source/Runtime/Engine/Classes/GameFramework/MovementComponent.h
+++ b/Engine/Source/Runtime/Engine/Classes/GameFramework/MovementComponent.h
@@ -384,6 +384,9 @@ public:
 	UFUNCTION(BlueprintCallable, Category="Components|Movement|Planar")
 	virtual void SnapUpdatedComponentToPlane();
 
+	/** Called by owning Actor upon successful teleport from AActor::TeleportTo(). */
+	virtual void OnTeleported() {};
+
 private:
 
 	/** Transient flag indicating whether we are executing OnRegister(). */
