Expose FCanvasWordWrapper::Execute for linking in game projects

diff --git a/Engine/Source/Runtime/Engine/Public/CanvasTypes.h b/Engine/Source/Runtime/Engine/Public/CanvasTypes.h
--- a/Engine/Source/Runtime/Engine/Public/CanvasTypes.h
+++ b/Engine/Source/Runtime/Engine/Public/CanvasTypes.h
@@ -152,7 +152,7 @@ public:
 	* @param InWrapWidth The width available.
 	* @param OutWrappedLineData An optional array to fill with the indices from the source string marking the begin and end points of the wrapped lines
 	*/
-	void Execute(const TCHAR* const InString, const FTextSizingParameters& InParameters, TArray<FWrappedStringElement>& OutStrings, FWrappedLineData* const OutWrappedLineData);
+	ENGINE_API void Execute(const TCHAR* const InString, const FTextSizingParameters& InParameters, TArray<FWrappedStringElement>& OutStrings, FWrappedLineData* const OutWrappedLineData);
 
 private:
 	/**
