Revert previous changelist. Word wrapping function did not need to be exposed

diff --git a/Engine/Source/Runtime/Engine/Public/CanvasTypes.h b/Engine/Source/Runtime/Engine/Public/CanvasTypes.h
--- a/Engine/Source/Runtime/Engine/Public/CanvasTypes.h
+++ b/Engine/Source/Runtime/Engine/Public/CanvasTypes.h
@@ -152,7 +152,7 @@ public:
 	* @param InWrapWidth The width available.
 	* @param OutWrappedLineData An optional array to fill with the indices from the source string marking the begin and end points of the wrapped lines
 	*/
-	ENGINE_API void Execute(const TCHAR* const InString, const FTextSizingParameters& InParameters, TArray<FWrappedStringElement>& OutStrings, FWrappedLineData* const OutWrappedLineData);
+	void Execute(const TCHAR* const InString, const FTextSizingParameters& InParameters, TArray<FWrappedStringElement>& OutStrings, FWrappedLineData* const OutWrappedLineData);
 
 private:
 	/**
