Added more info to a texture log warning

diff --git a/Engine/Source/Runtime/Engine/Private/Texture2D.cpp b/Engine/Source/Runtime/Engine/Private/Texture2D.cpp
--- a/Engine/Source/Runtime/Engine/Private/Texture2D.cpp
+++ b/Engine/Source/Runtime/Engine/Private/Texture2D.cpp
@@ -804,7 +804,7 @@ FTextureResource* UTexture2D::CreateResource()
 		}
 		else if (bIncompatibleTexture)
 		{
-			UE_LOG(LogTexture, Error, TEXT("%s contains no miplevels! Please delete."), *GetFullName());
+			UE_LOG(LogTexture, Error, TEXT("%s contains no miplevels! Please delete. (Format: %d)"), *GetFullName(), (int)GetPixelFormat());
 		}
 	}
 	else
