Export FCanvasWordWrapper

diff --git a/Engine/Source/Runtime/Engine/Public/CanvasTypes.h b/Engine/Source/Runtime/Engine/Public/CanvasTypes.h
--- a/Engine/Source/Runtime/Engine/Public/CanvasTypes.h
+++ b/Engine/Source/Runtime/Engine/Public/CanvasTypes.h
@@ -111,7 +111,7 @@ struct FWrappedStringElement
 	
 };
 
-class FCanvasWordWrapper
+class ENGINE_API FCanvasWordWrapper
 {
 public:
 	/** Array of indices where the wrapped lines begin and end in the source string */
