#ANIM: Fix Matinee preview of material parameter change from animation

- Merged...

diff --git a/Engine/Source/Runtime/Engine/Classes/Components/SkeletalMeshComponent.h b/Engine/Source/Runtime/Engine/Classes/Components/SkeletalMeshComponent.h
--- a/Engine/Source/Runtime/Engine/Classes/Components/SkeletalMeshComponent.h
+++ b/Engine/Source/Runtime/Engine/Classes/Components/SkeletalMeshComponent.h
@@ -930,6 +930,9 @@ public:
 	UFUNCTION(BlueprintCallable, Category="Components|SkeletalMesh")
 	void SetAllBodiesSimulatePhysics(bool bNewSimulate);
 
+	/** Update Material Parameters based on AnimInstance */
+	void UpdateMaterialParameters();
+
 	/** This is global set up for setting physics blend weight
 	 * This does multiple things automatically
 	 * If PhysicsBlendWeight == 1.f, it will enable Simulation, and if PhysicsBlendWeight == 0.f, it will disable Simulation. 
diff --git a/Engine/Source/Runtime/Engine/Private/Animation/AnimSingleNodeInstance.cpp b/Engine/Source/Runtime/Engine/Private/Animation/AnimSingleNodeInstance.cpp
--- a/Engine/Source/Runtime/Engine/Private/Animation/AnimSingleNodeInstance.cpp
+++ b/Engine/Source/Runtime/Engine/Private/Animation/AnimSingleNodeInstance.cpp
@@ -547,6 +547,7 @@ void UAnimSingleNodeInstance::SetPosition(float InPosition, bool bFireNotifies)
 			// since this is singlenode instance, if position changes, we can't keep old morphtarget curves
 			// we clear it and evaluate curve here with new asset. 
 			MorphTargetCurves.Empty();
+			MaterialParameterCurves.Empty();
 			// Evaluate Curve data now - even if time did not move, we still need to return curve if it exists
 			SequenceBase->EvaluateCurveData(this, CurrentTime, 1.0);
 		}
diff --git a/Engine/Source/Runtime/Engine/Private/SkeletalMesh.cpp b/Engine/Source/Runtime/Engine/Private/SkeletalMesh.cpp
--- a/Engine/Source/Runtime/Engine/Private/SkeletalMesh.cpp
+++ b/Engine/Source/Runtime/Engine/Private/SkeletalMesh.cpp
@@ -3698,6 +3698,7 @@ void ASkeletalMeshActor::PreviewSetAnimPosition(FName SlotName, int32 ChannelInd
 	}
 
 	// Update space bases so new animation position has an effect.
+	SkeletalMeshComponent->UpdateMaterialParameters();
 	SkeletalMeshComponent->RefreshBoneTransforms();
 	SkeletalMeshComponent->RefreshSlaveComponents();
 	SkeletalMeshComponent->UpdateComponentToWorld();
diff --git a/Engine/Source/Runtime/Engine/Private/SkeletalMeshComponent.cpp b/Engine/Source/Runtime/Engine/Private/SkeletalMeshComponent.cpp
--- a/Engine/Source/Runtime/Engine/Private/SkeletalMeshComponent.cpp
+++ b/Engine/Source/Runtime/Engine/Private/SkeletalMeshComponent.cpp
@@ -479,34 +479,38 @@ void USkeletalMeshComponent::TickAnimation(float DeltaTime)
 				}
 
 				//Update material parameters
-				if(AnimScriptInstance->MaterialParameterCurves.Num() > 0)
+				UpdateMaterialParameters();
+			}
+		}
+	}
+}
+
+void USkeletalMeshComponent::UpdateMaterialParameters()
+{
+	if(AnimScriptInstance->MaterialParameterCurves.Num() > 0)
+	{
+		for(auto Iter = AnimScriptInstance->MaterialParameterCurves.CreateConstIterator(); Iter; ++Iter)
+		{
+			FName ParameterName = Iter.Key();
+			float ParameterValue = Iter.Value();
+
+			UE_LOG(LogAnimation, Verbose, TEXT("Material Parameter change by Animation (%s : %0.2f)"), *ParameterName.ToString(), ParameterValue);
+			for(int32 MaterialIndex = 0; MaterialIndex < GetNumMaterials(); ++MaterialIndex)
+			{
+				UMaterialInterface* MaterialInterface = GetMaterial(MaterialIndex);
+				if(MaterialInterface)
 				{
-					for( auto Iter = AnimScriptInstance->MaterialParameterCurves.CreateConstIterator(); Iter; ++Iter )
+					float TestValue; //not used but needed for GetScalarParameterValue call
+					if(MaterialInterface->GetScalarParameterValue(ParameterName, TestValue))
 					{
-						FName ParameterName = Iter.Key();
-						float ParameterValue = Iter.Value();
-
-						for(int32 MaterialIndex = 0; MaterialIndex < GetNumMaterials(); ++MaterialIndex)
+						UMaterialInstanceDynamic* DynamicMaterial = Cast<UMaterialInstanceDynamic>(MaterialInterface);
+						if(!DynamicMaterial) //Is it already a UMaterialInstanceDynamic (ie we used it last tick)
 						{
-							UMaterialInterface* MaterialInterface = GetMaterial(MaterialIndex);
-							if (MaterialInterface)
-							{
-								float TestValue; //not used but needed for GetScalarParameterValue call
-								if(MaterialInterface->GetScalarParameterValue(ParameterName,TestValue))
-								{
-									UMaterialInstanceDynamic* DynamicMaterial = Cast<UMaterialInstanceDynamic>(MaterialInterface);
-									if(!DynamicMaterial) //Is it already a UMaterialInstanceDynamic (ie we used it last tick)
-									{
-										DynamicMaterial = CreateAndSetMaterialInstanceDynamic(MaterialIndex);
-									}
-									DynamicMaterial->SetScalarParameterValue(ParameterName, ParameterValue);
-								
-									//Assume that we only set the parameter on one of the materials, remove this break
-									//if that is no longer desired
-									break;
-								}
-							}
+							DynamicMaterial = CreateAndSetMaterialInstanceDynamic(MaterialIndex);
 						}
+						DynamicMaterial->SetScalarParameterValue(ParameterName, ParameterValue);
+
+						// we don't break here because we can have multiple materials wanted to be driven by same parameter
 					}
 				}
 			}
