2nd fix for race condition cleaning up FPerInstanceRenderData (merged from branch)

diff --git a/Engine/Source/Runtime/Engine/Private/InstancedStaticMesh.cpp b/Engine/Source/Runtime/Engine/Private/InstancedStaticMesh.cpp
--- a/Engine/Source/Runtime/Engine/Private/InstancedStaticMesh.cpp
+++ b/Engine/Source/Runtime/Engine/Private/InstancedStaticMesh.cpp
@@ -1193,15 +1193,19 @@ void UInstancedStaticMeshComponent::ReleasePerInstanceRenderData()
 	if (PerInstanceRenderData.IsValid() && !bPerInstanceRenderDataWasPrebuilt)
 	{
 		typedef TSharedPtr<FPerInstanceRenderData, ESPMode::ThreadSafe> FPerInstanceRenderDataPtr;
-		// Move shared pointer to a render task, resource will be released by scene proxy or render task, whoever will get executed last
+
+		// Make shared pointer object on the heap
+		FPerInstanceRenderDataPtr* CleanupRenderDataPtr = new FPerInstanceRenderDataPtr(PerInstanceRenderData);
+		PerInstanceRenderData.Reset();
+
 		ENQUEUE_UNIQUE_RENDER_COMMAND_ONEPARAMETER(
 			FReleasePerInstanceRenderData,
-			FPerInstanceRenderDataPtr,InPerInstanceRenderData, FPerInstanceRenderDataPtr(MoveTemp(PerInstanceRenderData)),
+			FPerInstanceRenderDataPtr*, InCleanupRenderDataPtr, CleanupRenderDataPtr,
 		{
-			InPerInstanceRenderData.Reset();
+			// Destroy the shared pointer object we allocated on the heap.
+			// Resource will either be released here or by scene proxy on the render thread, whoever gets executed last
+			delete InCleanupRenderDataPtr;
 		});
-		// At this point we should not have a reference to render data
-		check(!PerInstanceRenderData.IsValid());
 	}
 }
 
