merged bunch of fixes for direct paths (abstract nav data)
- #2429852:  fixed crash in path following when trying to optimize segments of direct path (abstract nav data)
- #2429880: fixed direct navigation paths not using dedicated class (required for detection in path following)
- #2429955: fixed crash on accessing navigation filters for abstract nav data (direct paths)
- #2431342: removed warning when navigation filter can't apply area costs, it's now a valid case (abstract nav data)

#ue4

diff --git a/Engine/Source/Runtime/AIModule/Private/Navigation/PathFollowingComponent.cpp b/Engine/Source/Runtime/AIModule/Private/Navigation/PathFollowingComponent.cpp
--- a/Engine/Source/Runtime/AIModule/Private/Navigation/PathFollowingComponent.cpp
+++ b/Engine/Source/Runtime/AIModule/Private/Navigation/PathFollowingComponent.cpp
@@ -639,6 +639,13 @@ int32 UPathFollowingComponent::OptimizeSegmentVisibility(int32 StartIndex)
 		return StartIndex + 1;
 	}
 
+	const bool bIsDirect = (Path->CastPath<FAbstractNavigationPath>() != NULL);
+	if (bIsDirect)
+	{
+		// can't optimize anything without real corridor
+		return StartIndex + 1;
+	}
+
 	const APawn* MyPawn = MyAI->GetPawn();
 	const float PawnHalfHeight = (MyAI->GetCharacter() && MyAI->GetCharacter()->GetCapsuleComponent()) ? MyAI->GetCharacter()->GetCapsuleComponent()->GetScaledCapsuleHalfHeight() : 0.0f;
 	const FVector PawnLocation = MyPawn->GetActorLocation();
diff --git a/Engine/Source/Runtime/Engine/Classes/AI/Navigation/AbstractNavData.h b/Engine/Source/Runtime/Engine/Classes/AI/Navigation/AbstractNavData.h
--- a/Engine/Source/Runtime/Engine/Classes/AI/Navigation/AbstractNavData.h
+++ b/Engine/Source/Runtime/Engine/Classes/AI/Navigation/AbstractNavData.h
@@ -14,6 +14,25 @@ struct ENGINE_API FAbstractNavigationPath : public FNavigationPath
 	static const FNavPathType Type;
 };
 
+class ENGINE_API FAbstractQueryFilter : public INavigationQueryFilterInterface
+{
+public:
+	virtual void Reset() override {}
+	virtual void SetAreaCost(uint8 AreaType, float Cost) override {}
+	virtual void SetFixedAreaEnteringCost(uint8 AreaType, float Cost) override {}
+	virtual void SetExcludedArea(uint8 AreaType) override {}
+	virtual void SetAllAreaCosts(const float* CostArray, const int32 Count) override {}
+	virtual void GetAllAreaCosts(float* CostArray, float* FixedCostArray, const int32 Count) const override {}
+	virtual void SetBacktrackingEnabled(const bool bBacktracking) override {}
+	virtual bool IsBacktrackingEnabled() const override { return false; }
+	virtual bool IsEqual(const INavigationQueryFilterInterface* Other) const override { return true; }
+	virtual void SetIncludeFlags(uint16 Flags) override {}
+	virtual uint16 GetIncludeFlags() const override { return 0; }
+	virtual void SetExcludeFlags(uint16 Flags) override {}
+	virtual uint16 GetExcludeFlags() const override { return 0; }
+	virtual INavigationQueryFilterInterface* CreateCopy() const override;
+};
+
 UCLASS()
 class AAbstractNavData : public ANavigationData
 {
diff --git a/Engine/Source/Runtime/Engine/Private/AI/Navigation/AbstractNavData.cpp b/Engine/Source/Runtime/Engine/Private/AI/Navigation/AbstractNavData.cpp
--- a/Engine/Source/Runtime/Engine/Private/AI/Navigation/AbstractNavData.cpp
+++ b/Engine/Source/Runtime/Engine/Private/AI/Navigation/AbstractNavData.cpp
@@ -10,6 +10,11 @@ FAbstractNavigationPath::FAbstractNavigationPath()
 	PathType = FAbstractNavigationPath::Type;
 }
 
+INavigationQueryFilterInterface* FAbstractQueryFilter::CreateCopy() const
+{
+	return new FAbstractQueryFilter();
+}
+
 AAbstractNavData::AAbstractNavData(const FObjectInitializer& ObjectInitializer) : Super(ObjectInitializer)
 {
 	if (HasAnyFlags(RF_ClassDefaultObject) == false)
@@ -21,6 +26,8 @@ AAbstractNavData::AAbstractNavData(const FObjectInitializer& ObjectInitializer)
 		TestHierarchicalPathImplementation = TestPathAbstract;
 
 		RaycastImplementation = RaycastAbstract;
+
+		DefaultQueryFilter->SetFilterType<FAbstractQueryFilter>();
 	}
 }
 
@@ -33,7 +40,7 @@ FPathFindingResult AAbstractNavData::FindPathAbstract(const FNavAgentProperties&
 	}
 
 	FPathFindingResult Result;
-	Result.Path = Query.PathInstanceToFill.IsValid() ? Query.PathInstanceToFill : Self->CreatePathInstance<FNavigationPath>(Query.Owner.Get());
+	Result.Path = Query.PathInstanceToFill.IsValid() ? Query.PathInstanceToFill : Self->CreatePathInstance<FAbstractNavigationPath>(Query.Owner.Get());
 
 	Result.Path->GetPathPoints().Add(FNavPathPoint(Query.StartLocation));
 	Result.Path->GetPathPoints().Add(FNavPathPoint(Query.EndLocation));
diff --git a/Engine/Source/Runtime/Engine/Private/AI/Navigation/NavFilters/NavigationQueryFilter.cpp b/Engine/Source/Runtime/Engine/Private/AI/Navigation/NavFilters/NavigationQueryFilter.cpp
--- a/Engine/Source/Runtime/Engine/Private/AI/Navigation/NavFilters/NavigationQueryFilter.cpp
+++ b/Engine/Source/Runtime/Engine/Private/AI/Navigation/NavFilters/NavigationQueryFilter.cpp
@@ -142,9 +142,6 @@ void UNavigationQueryFilter::InitializeFilter(const ANavigationData* NavData, FN
 		const int32 AreaId = NavData->GetAreaID(AreaData.AreaClass);
 		if (AreaId == INDEX_NONE)
 		{
-			UE_LOG(LogNavigation, Warning, TEXT("Can't find area '%s' in '%s' for filter: %s"),
-				*GetNameSafe(AreaData.AreaClass), *GetNameSafe(NavData), *GetName());
-
 			continue;
 		}
 
