Fix logspam caused by duplicate OnLevelAddedToWorld notification

diff --git a/Engine/Source/Runtime/Engine/Private/World.cpp b/Engine/Source/Runtime/Engine/Private/World.cpp
--- a/Engine/Source/Runtime/Engine/Private/World.cpp
+++ b/Engine/Source/Runtime/Engine/Private/World.cpp
@@ -2255,18 +2255,17 @@ void UWorld::UpdateLevelStreamingInner(ULevelStreaming* StreamingLevel)
 			if (!Level->bIsVisible)
 			{
 				AddToWorld(Level, StreamingLevel->LevelTransform);
-			}
-
-			// In case we have finished making level visible
-			if (Level->bIsVisible)
-			{
-				// immediately discard previous level
-				StreamingLevel->DiscardPendingUnloadLevel(this);
-
-				if (Scene)
+				// In case we have finished making level visible
+				if (Level->bIsVisible)
 				{
-					// Notify the new level has been added after the old has been discarded
-					Scene->OnLevelAddedToWorld(Level->GetOutermost()->GetFName());
+					// immediately discard previous level
+					StreamingLevel->DiscardPendingUnloadLevel(this);
+
+					if (Scene)
+					{
+						// Notify the new level has been added after the old has been discarded
+						Scene->OnLevelAddedToWorld(Level->GetOutermost()->GetFName());
+					}
 				}
 			}
 		}
