Made navigation system not call clean-up on navigation data on it's destruction #UE4

This was not necessary and could result in redundant work. Also, old behavior was resulting in a crash on map rebuilding if someone saved his level with bug reported in UE-8906

diff --git a/Engine/Source/Runtime/Engine/Classes/AI/Navigation/NavigationData.h b/Engine/Source/Runtime/Engine/Classes/AI/Navigation/NavigationData.h
--- a/Engine/Source/Runtime/Engine/Classes/AI/Navigation/NavigationData.h
+++ b/Engine/Source/Runtime/Engine/Classes/AI/Navigation/NavigationData.h
@@ -3,6 +3,7 @@
 #pragma once
 #include "AI/Navigation/NavFilters/NavigationQueryFilter.h"
 #include "AI/Navigation/NavigationTypes.h"
+#include "AI/NavDataGenerator.h"
 #include "GameFramework/Actor.h"
 #include "UniquePtr.h"
 #include "NavigationData.generated.h"
@@ -447,8 +448,6 @@ public:
 	//----------------------------------------------------------------------//
 	// Life cycle                                                                
 	//----------------------------------------------------------------------//
-	/** Dtor */
-	virtual ~ANavigationData();
 
 	// Begin UObject/AActor Interface
 	virtual void PostInitProperties() override;
@@ -457,7 +456,7 @@ public:
 #if WITH_EDITOR
 	virtual void PostEditUndo() override;
 #endif // WITH_EDITOR
-	virtual void Destroyed() override;
+	virtual void BeginDestroy() override;
 	// End UObject Interface
 		
 	virtual void CleanUp();
diff --git a/Engine/Source/Runtime/Engine/Private/AI/Navigation/NavigationData.cpp b/Engine/Source/Runtime/Engine/Private/AI/Navigation/NavigationData.cpp
--- a/Engine/Source/Runtime/Engine/Private/AI/Navigation/NavigationData.cpp
+++ b/Engine/Source/Runtime/Engine/Private/AI/Navigation/NavigationData.cpp
@@ -115,11 +115,6 @@ ANavigationData::ANavigationData(const FObjectInitializer& ObjectInitializer)
 	ObservedPathsTickInterval = 0.5;
 }
 
-ANavigationData::~ANavigationData()
-{
-	CleanUp();
-}
-
 uint16 ANavigationData::GetNextUniqueID()
 {
 	check(IsInGameThread());
@@ -364,7 +359,7 @@ bool ANavigationData::DoesSupportAgent(const FNavAgentProperties& AgentProps) co
 	return NavDataConfig.IsEquivalent(AgentProps);
 }
 
-void ANavigationData::Destroyed()
+void ANavigationData::BeginDestroy()
 {
 	UWorld* WorldOuter = GetWorld();
 
@@ -373,10 +368,10 @@ void ANavigationData::Destroyed()
 	{
 		WorldOuter->GetNavigationSystem()->UnregisterNavData(this);
 	}
-	
+
 	CleanUp();
 
-	Super::Destroyed();
+	Super::BeginDestroy();
 }
 
 void ANavigationData::CleanUp()
diff --git a/Engine/Source/Runtime/Engine/Private/AI/Navigation/NavigationSystem.cpp b/Engine/Source/Runtime/Engine/Private/AI/Navigation/NavigationSystem.cpp
--- a/Engine/Source/Runtime/Engine/Private/AI/Navigation/NavigationSystem.cpp
+++ b/Engine/Source/Runtime/Engine/Private/AI/Navigation/NavigationSystem.cpp
@@ -3257,14 +3257,6 @@ void UNavigationSystem::CleanUp(ECleanupMode Mode)
 
 	DestroyNavOctree();
 	
-	for (int32 NavDataIndex = 0; NavDataIndex < NavDataSet.Num(); ++NavDataIndex)
-	{
-		if (NavDataSet[NavDataIndex] != NULL)
-		{
-			NavDataSet[NavDataIndex]->CleanUp();
-		}
-	}	
-
 	SetCrowdManager(NULL);
 
 	NavDataSet.Reset();
