Fix naming objects in simple construction script.

diff --git a/Engine/Source/Runtime/Engine/Private/SimpleConstructionScript.cpp b/Engine/Source/Runtime/Engine/Private/SimpleConstructionScript.cpp
--- a/Engine/Source/Runtime/Engine/Private/SimpleConstructionScript.cpp
+++ b/Engine/Source/Runtime/Engine/Private/SimpleConstructionScript.cpp
@@ -779,7 +779,7 @@ USCS_Node* USimpleConstructionScript::CreateNode(UActorComponent* NewComponentTe
 	{
 		// Create a node for the script, and save a pointer to the template
 		// NewNamedObject to work around the fact we shouldn't use NewObject for default subobjects
-		auto NewNode = NewObject<USCS_Node>(this);
+		auto NewNode = NewObject<USCS_Node>(this, MakeUniqueObjectName(this, USCS_Node::StaticClass()));
 		NewNode->SetFlags(RF_Transactional);
 		NewNode->ComponentTemplate = NewComponentTemplate;
 
