Limited valid restitution values to 0...1.

diff --git a/Engine/Source/Runtime/Engine/Classes/PhysicalMaterials/PhysicalMaterial.h b/Engine/Source/Runtime/Engine/Classes/PhysicalMaterials/PhysicalMaterial.h
--- a/Engine/Source/Runtime/Engine/Classes/PhysicalMaterials/PhysicalMaterial.h
+++ b/Engine/Source/Runtime/Engine/Classes/PhysicalMaterials/PhysicalMaterial.h
@@ -54,8 +54,8 @@ class ENGINE_API UPhysicalMaterial : public UObject
 	UPROPERTY(EditAnywhere, BlueprintReadWrite, Category = PhysicalMaterial)
 	bool bOverrideFrictionCombineMode;
 
-	/** Resitution or 'bouncyness' of this surface */
-	UPROPERTY(EditAnywhere, BlueprintReadOnly, Category = PhysicalMaterial)
+	/** Restitution or 'bouncyness' of this surface. */
+	UPROPERTY(EditAnywhere, BlueprintReadOnly, Category = PhysicalMaterial, meta = (UIMin=0, UIMax=1, ClampMin=0, ClampMax=1))
 	float Restitution;
 
 	//
