Merging required headers for class resolution

//depot/UE4-UT/...

to //depot/UE4/...

diff --git a/Engine/Source/Runtime/Engine/Public/StaticBoundShaderState.h b/Engine/Source/Runtime/Engine/Public/StaticBoundShaderState.h
--- a/Engine/Source/Runtime/Engine/Public/StaticBoundShaderState.h
+++ b/Engine/Source/Runtime/Engine/Public/StaticBoundShaderState.h
@@ -7,6 +7,9 @@
 #ifndef __STATICBOUNDSHADERSTATE_H__
 #define __STATICBOUNDSHADERSTATE_H__
 
+#include "RenderResource.h"
+#include "Shader.h"
+
 /**
  * FGlobalBoundShaderState
  * 
