UMovementComponent::OnTeleported now called for all Pawns that teleport,  rather than just characters who teleport.

diff --git a/Engine/Source/Runtime/Engine/Classes/GameFramework/Character.h b/Engine/Source/Runtime/Engine/Classes/GameFramework/Character.h
--- a/Engine/Source/Runtime/Engine/Classes/GameFramework/Character.h
+++ b/Engine/Source/Runtime/Engine/Classes/GameFramework/Character.h
@@ -344,7 +344,6 @@ public:
 	float JumpMaxHoldTime;
 
 	// Begin AActor Interface.
-	virtual void TeleportSucceeded(bool bIsATest) override;
 	virtual void ClearCrossLevelReferences() override;
 	virtual void PreNetReceive() override;
 	virtual void PostNetReceive() override;
diff --git a/Engine/Source/Runtime/Engine/Classes/GameFramework/Pawn.h b/Engine/Source/Runtime/Engine/Classes/GameFramework/Pawn.h
--- a/Engine/Source/Runtime/Engine/Classes/GameFramework/Pawn.h
+++ b/Engine/Source/Runtime/Engine/Classes/GameFramework/Pawn.h
@@ -231,6 +231,7 @@ public:
 	virtual void BecomeViewTarget(APlayerController* PC) override;
 	virtual void EnableInput(APlayerController* PlayerController) override;
 	virtual void DisableInput(APlayerController* PlayerController) override;
+	virtual void TeleportSucceeded(bool bIsATest) override;
 
 	/** Overridden to defer to the RootComponent's CanCharacterStepUpOn setting if it is explicitly Yes or No. If set to Owner, will return Super::CanBeBaseForCharacter(). */
 	virtual bool CanBeBaseForCharacter(APawn* APawn) const override;
diff --git a/Engine/Source/Runtime/Engine/Private/Character.cpp b/Engine/Source/Runtime/Engine/Private/Character.cpp
--- a/Engine/Source/Runtime/Engine/Private/Character.cpp
+++ b/Engine/Source/Runtime/Engine/Private/Character.cpp
@@ -370,17 +370,6 @@ void ACharacter::ApplyDamageMomentum(float DamageTaken, FDamageEvent const& Dama
 	}
 }
 
-void ACharacter::TeleportSucceeded(bool bIsATest)
-{
-	if (!bIsATest && CharacterMovement)
-	{
-		CharacterMovement->OnTeleported();
-	}
-
-	Super::TeleportSucceeded(bIsATest);
-}
-
-
 void ACharacter::ClearCrossLevelReferences()
 {
 	if( BasedMovement.MovementBase != NULL && GetOutermost() != BasedMovement.MovementBase->GetOutermost() )
@@ -391,8 +380,6 @@ void ACharacter::ClearCrossLevelReferences()
 	Super::ClearCrossLevelReferences();
 }
 
-
-
 namespace MovementBaseUtility
 {
 	bool IsDynamicBase(const UPrimitiveComponent* MovementBase)
diff --git a/Engine/Source/Runtime/Engine/Private/Pawn.cpp b/Engine/Source/Runtime/Engine/Private/Pawn.cpp
--- a/Engine/Source/Runtime/Engine/Private/Pawn.cpp
+++ b/Engine/Source/Runtime/Engine/Private/Pawn.cpp
@@ -946,6 +946,18 @@ void APawn::DisableInput(class APlayerController* PlayerController)
 	}
 }
 
+void APawn::TeleportSucceeded(bool bIsATest)
+{
+	if (bIsATest == false)
+	{
+		UMovementComponent* const MoveComponent = GetMovementComponent();
+		if (MoveComponent)
+		{
+			MoveComponent->OnTeleported();
+		}
+	}
+}
+
 void APawn::GetMoveGoalReachTest(AActor* MovingActor, const FVector& MoveOffset, FVector& GoalOffset, float& GoalRadius, float& GoalHalfHeight) const 
 {
 	GoalOffset = FVector::ZeroVector;
@@ -1067,7 +1079,7 @@ void APawn::GetLifetimeReplicatedProps( TArray< FLifetimeProperty > & OutLifetim
 
 void APawn::MoveIgnoreActorAdd(AActor* ActorToIgnore)
 {
-	UPrimitiveComponent * RootPrimitiveComponent = Cast<UPrimitiveComponent>(GetRootComponent());
+	UPrimitiveComponent* RootPrimitiveComponent = Cast<UPrimitiveComponent>(GetRootComponent());
 	if( RootPrimitiveComponent )
 	{
 		RootPrimitiveComponent->IgnoreActorWhenMoving(ActorToIgnore, true);
@@ -1076,7 +1088,7 @@ void APawn::MoveIgnoreActorAdd(AActor* ActorToIgnore)
 
 void APawn::MoveIgnoreActorRemove(AActor* ActorToIgnore)
 {
-	UPrimitiveComponent * RootPrimitiveComponent = Cast<UPrimitiveComponent>(GetRootComponent());
+	UPrimitiveComponent* RootPrimitiveComponent = Cast<UPrimitiveComponent>(GetRootComponent());
 	if( RootPrimitiveComponent )
 	{
 		RootPrimitiveComponent->IgnoreActorWhenMoving(ActorToIgnore, false);
