Make sure DoesSocketExist work for bone.

diff --git a/Engine/Source/Runtime/Engine/Private/SkinnedMeshComponent.cpp b/Engine/Source/Runtime/Engine/Private/SkinnedMeshComponent.cpp
--- a/Engine/Source/Runtime/Engine/Private/SkinnedMeshComponent.cpp
+++ b/Engine/Source/Runtime/Engine/Private/SkinnedMeshComponent.cpp
@@ -1178,7 +1178,7 @@ class USkeletalMeshSocket const* USkinnedMeshComponent::GetSocketByName(FName In
 
 bool USkinnedMeshComponent::DoesSocketExist(FName InSocketName) const
 {
-	return (GetSocketByName(InSocketName) != NULL);
+	return (GetSocketBoneName(InSocketName) != NULL);
 }
 
 FName USkinnedMeshComponent::GetSocketBoneName(FName InSocketName)
