Fix build issue

diff --git a/Engine/Source/Runtime/Engine/Classes/Components/SkinnedMeshComponent.h b/Engine/Source/Runtime/Engine/Classes/Components/SkinnedMeshComponent.h
--- a/Engine/Source/Runtime/Engine/Classes/Components/SkinnedMeshComponent.h
+++ b/Engine/Source/Runtime/Engine/Classes/Components/SkinnedMeshComponent.h
@@ -386,7 +386,7 @@ public:
 	 * @return	bone name
 	 */
 	UFUNCTION(BlueprintCallable, Category="Components|SkinnedMesh")
-	FName GetSocketBoneName(FName InSocketName);
+	FName GetSocketBoneName(FName InSocketName) const;
 
 	/** 
 	 * Change the SkeletalMesh that is rendered for this Component. Will re-initialize the animation tree etc. 
diff --git a/Engine/Source/Runtime/Engine/Private/SkinnedMeshComponent.cpp b/Engine/Source/Runtime/Engine/Private/SkinnedMeshComponent.cpp
--- a/Engine/Source/Runtime/Engine/Private/SkinnedMeshComponent.cpp
+++ b/Engine/Source/Runtime/Engine/Private/SkinnedMeshComponent.cpp
@@ -1178,10 +1178,10 @@ class USkeletalMeshSocket const* USkinnedMeshComponent::GetSocketByName(FName In
 
 bool USkinnedMeshComponent::DoesSocketExist(FName InSocketName) const
 {
-	return (GetSocketBoneName(InSocketName) != NULL);
+	return (GetSocketBoneName(InSocketName) != NAME_None);
 }
 
-FName USkinnedMeshComponent::GetSocketBoneName(FName InSocketName)
+FName USkinnedMeshComponent::GetSocketBoneName(FName InSocketName) const
 {
 	if(!SkeletalMesh)
 	{
