Added missing Super call in TeleportSucceeded()

diff --git a/Engine/Source/Runtime/Engine/Private/Pawn.cpp b/Engine/Source/Runtime/Engine/Private/Pawn.cpp
--- a/Engine/Source/Runtime/Engine/Private/Pawn.cpp
+++ b/Engine/Source/Runtime/Engine/Private/Pawn.cpp
@@ -956,6 +956,8 @@ void APawn::TeleportSucceeded(bool bIsATest)
 			MoveComponent->OnTeleported();
 		}
 	}
+
+	Super::TeleportSucceeded(bIsATest);
 }
 
 void APawn::GetMoveGoalReachTest(AActor* MovingActor, const FVector& MoveOffset, FVector& GoalOffset, float& GoalRadius, float& GoalHalfHeight) const 
