#ue4 - Remove component from OverlappingComponents array before firing OnComponentEndOverlap events.

Jira: UE-7000

diff --git a/Engine/Source/Runtime/Engine/Private/PrimitiveComponent.cpp b/Engine/Source/Runtime/Engine/Private/PrimitiveComponent.cpp
--- a/Engine/Source/Runtime/Engine/Private/PrimitiveComponent.cpp
+++ b/Engine/Source/Runtime/Engine/Private/PrimitiveComponent.cpp
@@ -1978,28 +1978,28 @@ void UPrimitiveComponent::EndComponentOverlap(const FOverlapInfo& OtherOverlap,
 
 	//	UE_LOG(LogActor, Log, TEXT("END OVERLAP! Self=%s SelfComp=%s, Other=%s, OtherComp=%s"), *GetNameSafe(this), *GetNameSafe(MyComp), *GetNameSafe(OtherActor), *GetNameSafe(OtherComp));
 
-	if ( (OtherActor != NULL) && bDoNotifies && IsOverlappingComponent(OtherOverlap) )
+	const int32 OverlapIdx = OverlappingComponents.Find(OtherOverlap);
+	if (OverlapIdx != INDEX_NONE)
 	{
-		if (!bNoNotifySelf && IsPrimCompValidAndAlive(this))
+		OverlappingComponents.RemoveAtSwap(OverlapIdx, 1, false);
+	
+		if ((OtherActor != NULL) && bDoNotifies)
 		{
-			OnComponentEndOverlap.Broadcast(OtherActor, OtherComp, OtherOverlap.GetBodyIndex());
-		}
+			if (!bNoNotifySelf && IsPrimCompValidAndAlive(this))
+			{
+				OnComponentEndOverlap.Broadcast(OtherActor, OtherComp, OtherOverlap.GetBodyIndex());
+			}
 
-		if(IsPrimCompValidAndAlive(OtherComp))
-		{
-			OtherComp->OnComponentEndOverlap.Broadcast(MyActor, this, INDEX_NONE);
+			if (IsPrimCompValidAndAlive(OtherComp))
+			{
+				OtherComp->OnComponentEndOverlap.Broadcast(MyActor, this, INDEX_NONE);
+			}
 		}
 	}
 
-	int32 OverlapIdx = OverlappingComponents.Find(OtherOverlap);
-	if (OverlapIdx != INDEX_NONE)
-	{
-		OverlappingComponents.RemoveAtSwap(OverlapIdx, 1, false);
-	}
-
 	if (OtherComp)
 	{
-		int32 OtherOverlapIdx = OtherComp->OverlappingComponents.Find(FOverlapInfo(this, INDEX_NONE));
+		const int32 OtherOverlapIdx = OtherComp->OverlappingComponents.Find(FOverlapInfo(this, INDEX_NONE));
 		if (OtherOverlapIdx != INDEX_NONE)
 		{
 			OtherComp->OverlappingComponents.RemoveAtSwap(OtherOverlapIdx, 1, false);
