#ANIM: Fixed nan issue from invalid input
 - still shouldn't expose invalid input

diff --git a/Engine/Source/Runtime/Engine/Private/Animation/AnimInterpFilter.cpp b/Engine/Source/Runtime/Engine/Private/Animation/AnimInterpFilter.cpp
--- a/Engine/Source/Runtime/Engine/Private/Animation/AnimInterpFilter.cpp
+++ b/Engine/Source/Runtime/Engine/Private/Animation/AnimInterpFilter.cpp
@@ -216,6 +216,6 @@ float FFIRFilterTimeBased::CalculateFilteredOutput()
 // 		return SumInputs;
 // 	}
 
-	return SumInputs/SumCoefficient;
+	return (SumCoefficient>0.f)?SumInputs/SumCoefficient : 0.f;
 }
 
