UE4 - Default Speedtree previous velocity during base pass to off

diff --git a/Engine/Source/Runtime/Engine/Private/Materials/MaterialExpressions.cpp b/Engine/Source/Runtime/Engine/Private/Materials/MaterialExpressions.cpp
--- a/Engine/Source/Runtime/Engine/Private/Materials/MaterialExpressions.cpp
+++ b/Engine/Source/Runtime/Engine/Private/Materials/MaterialExpressions.cpp
@@ -8845,7 +8845,7 @@ UMaterialExpressionSpeedTree::UMaterialExpressionSpeedTree(const FObjectInitiali
 	WindType = STW_None;
 	LODType = STLOD_Pop;
 	BillboardThreshold = 0.9f;
-	bOutputPreviousPosition = true;
+	bOutputPreviousPosition = false;
 
 	MenuCategories.Add(ConstructorStatics.NAME_SpeedTree);
 }
