[INTEGRATE] Change 2425384 by [REDACTED_PERSON][REDACTED_HANDLE] on 2015/01/30 04:17:28

	Fix async RB physics scene initialization for static components created/registered during UCS execution.

	change summary:
	- added a 'PhysicsMobility' variable to shadow the base Mobility setting in UPrimitiveComponent
	- added OnComponentCreated() and SetMobility() overrides to UPrimitiveComponent to keep PhysicsMobliity in sync
	- modified FBodyInstance::InitBody() to now look at PhysicsMobility instead of Mobility to determine static or dynamic (as that setting remains unchanged during UCS execution)
	- modified AActor::AddComponent() to call OnComponentCreated() prior to adding the instance to the FUCSComponentManager (so that PhysicsMobility reflects the unmodified setting)

diff --git a/Engine/Source/Runtime/Engine/Classes/Components/PrimitiveComponent.h b/Engine/Source/Runtime/Engine/Classes/Components/PrimitiveComponent.h
--- a/Engine/Source/Runtime/Engine/Classes/Components/PrimitiveComponent.h
+++ b/Engine/Source/Runtime/Engine/Classes/Components/PrimitiveComponent.h
@@ -400,6 +400,12 @@ public:
 	/** Cached navigation relevancy flag for collision updates */
 	uint32 bNavigationRelevant : 1;
 
+	/** 
+	 * Shadowed mobility setting used e.g. to assist with physics scene initialization at runtime.
+	 * It will usually match the value of the 'Mobility' property, except e.g. during UCS execution.
+	 */
+	TEnumAsByte<EComponentMobility::Type> PhysicsMobility;
+
 protected:
 
 	virtual void UpdateNavigationData() override;
@@ -943,6 +949,7 @@ public:
 #endif
 
 	// Begin UActorComponent Interface
+	virtual void OnComponentCreated() override;
 	virtual void InvalidateLightingCacheDetailed(bool bInvalidateBuildEnqueuedLighting, bool bTranslationOnly) override;
 	virtual bool IsEditorOnly() const override;
 	virtual bool ShouldCreatePhysicsState() const override;
@@ -1293,6 +1300,7 @@ public:
 	virtual ECollisionChannel GetCollisionObjectType() const override;
 	virtual const FCollisionResponseContainer& GetCollisionResponseToChannels() const override;
 	virtual FVector GetComponentVelocity() const override;
+	virtual void SetMobility(EComponentMobility::Type NewMobility) override;
 	//End USceneComponent Interface
 
 	/**
diff --git a/Engine/Source/Runtime/Engine/Private/ActorConstruction.cpp b/Engine/Source/Runtime/Engine/Private/ActorConstruction.cpp
--- a/Engine/Source/Runtime/Engine/Private/ActorConstruction.cpp
+++ b/Engine/Source/Runtime/Engine/Private/ActorConstruction.cpp
@@ -604,7 +604,10 @@ UActorComponent* AActor::AddComponent(FName TemplateName, bool bManualAttachment
 	UActorComponent* NewActorComp = CreateComponentFromTemplate(Template);
 	if(NewActorComp != nullptr)
 	{
-		// Keep track of the new component during UCS execution
+		// Call function to notify component it has been created
+		NewActorComp->OnComponentCreated();
+
+		// Keep track of the new component during UCS execution. This also does a temporary mobility swap during UCS execution in order to allow attachment and other things to succeed within that context.
 		if(bRunningUserConstructionScript)
 		{
 		FUCSComponentManager::Get().AddComponent(this, NewActorComp);
@@ -630,9 +633,6 @@ UActorComponent* AActor::AddComponent(FName TemplateName, bool bManualAttachment
 			NewSceneComp->SetRelativeTransform(RelativeTransform);
 		}
 
-		// Call function to notify component it has been created
-		NewActorComp->OnComponentCreated();
-
 		// Register component, which will create physics/rendering state, now component is in correct position
 		NewActorComp->RegisterComponent();
 	}
diff --git a/Engine/Source/Runtime/Engine/Private/PhysicsEngine/BodyInstance.cpp b/Engine/Source/Runtime/Engine/Private/PhysicsEngine/BodyInstance.cpp
--- a/Engine/Source/Runtime/Engine/Private/PhysicsEngine/BodyInstance.cpp
+++ b/Engine/Source/Runtime/Engine/Private/PhysicsEngine/BodyInstance.cpp
@@ -1003,8 +1003,8 @@ void FBodyInstance::InitBody_Legacy(UBodySetup* Setup, const FTransform& Transfo
 		}
 	}
 
-	// See if we are 'static'
-	const bool bPhysicsStatic = (OwnerComponent == NULL) || (OwnerComponent->Mobility != EComponentMobility::Movable);
+	// See if we are 'static' - note that we use the PhysicsMobility setting here which "shadows" the actual Mobility setting in most cases but also will remain constant e.g. during UCS execution
+	const bool bPhysicsStatic = (OwnerComponent == NULL) || (OwnerComponent->PhysicsMobility != EComponentMobility::Movable);
 
 	// In skeletal case, we need both our bone and skelcomponent flag to be true.
 	// This might be 'and'ing us with ourself, but thats fine.
diff --git a/Engine/Source/Runtime/Engine/Private/PrimitiveComponent.cpp b/Engine/Source/Runtime/Engine/Private/PrimitiveComponent.cpp
--- a/Engine/Source/Runtime/Engine/Private/PrimitiveComponent.cpp
+++ b/Engine/Source/Runtime/Engine/Private/PrimitiveComponent.cpp
@@ -137,6 +137,14 @@ void UPrimitiveComponent::GetLightAndShadowMapMemoryUsage( int32& LightMapMemory
 	return;
 }
 
+void UPrimitiveComponent::OnComponentCreated()
+{
+	Super::OnComponentCreated();
+
+	// Shadow the current mobility setting for physics scene initialization (since e.g. during UCS execution we temporarily override the base mobility setting)
+	PhysicsMobility = Mobility;
+}
+
 void UPrimitiveComponent::InvalidateLightingCacheDetailed(bool bInvalidateBuildEnqueuedLighting, bool bTranslationOnly) 
 {
 	bHasCachedStaticLighting = false;
@@ -932,6 +940,14 @@ bool UPrimitiveComponent::IsWorldGeometry() const
 	return Mobility != EComponentMobility::Movable && GetCollisionObjectType()==ECC_WorldStatic;
 }
 
+void UPrimitiveComponent::SetMobility(EComponentMobility::Type NewMobility)
+{
+	Super::SetMobility(NewMobility);
+
+	// Update the shadowed mobility setting for physics
+	PhysicsMobility = Mobility;
+}
+
 ECollisionChannel UPrimitiveComponent::GetCollisionObjectType() const
 {
 	return ECollisionChannel(BodyInstance.GetObjectType());
