[INTEGRATE] Change 2425884 by [REDACTED_PERSON][REDACTED_HANDLE] on 2015/01/30 11:13:24

	Allow SceneComponents to have a visualization sprite
	Fix bug where AudioComponent icon would not change when auto activate was toggled
	Blueprint classes DefaultSceneRoot now display a visualization icon
	Fix bug where Empty Actor sprite component would not go away when default scene root was replaced
	Creating a class blueprint from an Actor with an Instance Default Scene Root will no longer have a Billboard component in the created class

diff --git a/Engine/Source/Editor/UnrealEd/Private/Factories/ActorFactory.cpp b/Engine/Source/Editor/UnrealEd/Private/Factories/ActorFactory.cpp
--- a/Engine/Source/Editor/UnrealEd/Private/Factories/ActorFactory.cpp
+++ b/Engine/Source/Editor/UnrealEd/Private/Factories/ActorFactory.cpp
@@ -1186,13 +1186,13 @@ AActor* UActorFactoryEmptyActor::SpawnActor( UObject* Asset, ULevel* InLevel, co
 
 		USceneComponent* RootComponent = NewObject<USceneComponent>(NewActor, USceneComponent::GetDefaultSceneRootVariableName(), RF_Transactional);
 		RootComponent->Mobility = EComponentMobility::Movable;
+		RootComponent->bVisualizeComponent = true;
 		RootComponent->SetWorldLocationAndRotation(Location, Rotation);
-		NewActor->SetRootComponent(RootComponent);
-
-		UBillboardComponent* BillboardComponent	= CreateEditorOnlyBillboardComponent(NewActor, RootComponent);
 
+		NewActor->SetRootComponent(RootComponent);
 		NewActor->AddInstanceComponent(RootComponent);
-		NewActor->AddInstanceComponent(BillboardComponent);
+
+		RootComponent->RegisterComponent();
 	}
 
 	return NewActor;
diff --git a/Engine/Source/Runtime/Engine/Classes/Components/AudioComponent.h b/Engine/Source/Runtime/Engine/Classes/Components/AudioComponent.h
--- a/Engine/Source/Runtime/Engine/Classes/Components/AudioComponent.h
+++ b/Engine/Source/Runtime/Engine/Classes/Components/AudioComponent.h
@@ -257,9 +257,6 @@ public:
 	/** Whether or not this sound class forces sounds to the center channel */
 	uint32 bCenterChannelOnly:1;
 
-	/** Whether or not the audio component should display an icon in the editor */
-	uint32 bVisualizeComponent:1;
-
 	/** Used by the subtitle manager to prioritize subtitles wave instances spawned by this component. */
 	float SubtitlePriority;
 
@@ -278,7 +275,6 @@ public:
 
 	// Begin ActorComponent interface.
 #if WITH_EDITORONLY_DATA
-	virtual void OnComponentDestroyed() override;
 	virtual void OnRegister() override;
 #endif
 	virtual void OnUnregister() override;
@@ -293,10 +289,6 @@ public:
 private:
 	
 #if WITH_EDITORONLY_DATA
-	/** Editor only component used to display the sprite so as to be able to see the location of the Audio Component  */
-	UPROPERTY(transient)
-	class UBillboardComponent* SpriteComponent;
-
 	/** Utility function that updates which texture is displayed on the sprite dependent on the properties of the Audio Component. */
 	void UpdateSpriteTexture();
 #endif
diff --git a/Engine/Source/Runtime/Engine/Classes/Components/LightComponentBase.h b/Engine/Source/Runtime/Engine/Classes/Components/LightComponentBase.h
--- a/Engine/Source/Runtime/Engine/Classes/Components/LightComponentBase.h
+++ b/Engine/Source/Runtime/Engine/Classes/Components/LightComponentBase.h
@@ -151,21 +151,16 @@ class ENGINE_API ULightComponentBase : public USceneComponent
 	 */
 	bool HasStaticShadowing() const;
 
+#if WITH_EDITOR
 	/** UActorComponent Interface */
 	virtual void OnRegister() override;
+#endif
 
 	/** We return a small bounds to allow us to non-interpenetrates when placing lights in the level. */
 	virtual bool ShouldCollideWhenPlacing() const override;
 
 	/** Get the extent used when placing this component in the editor, used for 'pulling back' hit. */
 	virtual FBoxSphereBounds GetPlacementExtent() const override;
-
-protected:
-#if WITH_EDITORONLY_DATA
-	UPROPERTY(transient)
-	class UBillboardComponent* SpriteComponent;
-#endif
-
 };
 
 
diff --git a/Engine/Source/Runtime/Engine/Classes/Components/SceneComponent.h b/Engine/Source/Runtime/Engine/Classes/Components/SceneComponent.h
--- a/Engine/Source/Runtime/Engine/Classes/Components/SceneComponent.h
+++ b/Engine/Source/Runtime/Engine/Classes/Components/SceneComponent.h
@@ -155,8 +155,18 @@ protected:
 	// Transient flag that temporarily disables UpdateOverlaps within DetachFromParent().
 	uint32 bDisableDetachmentUpdateOverlaps:1;
 
+#if WITH_EDITOR
+protected:
+	/** Editor only component used to display the sprite so as to be able to see the location of the Audio Component  */
+	UPROPERTY(transient)
+	class UBillboardComponent* SpriteComponent;
+
 public:
+	UPROPERTY()
+	uint32 bVisualizeComponent:1;
+#endif
 
+public:
 	/** How often this component is allowed to move, used to make various optimizations. Only safe to set in constructor, use SetMobility() during runtime. */
 	UPROPERTY(Category=Mobility, EditAnywhere, BlueprintReadOnly)
 	TEnumAsByte<EComponentMobility::Type> Mobility;
diff --git a/Engine/Source/Runtime/Engine/Classes/PhysicsEngine/PhysicsConstraintComponent.h b/Engine/Source/Runtime/Engine/Classes/PhysicsEngine/PhysicsConstraintComponent.h
--- a/Engine/Source/Runtime/Engine/Classes/PhysicsEngine/PhysicsConstraintComponent.h
+++ b/Engine/Source/Runtime/Engine/Classes/PhysicsEngine/PhysicsConstraintComponent.h
@@ -68,7 +68,6 @@ public:
 	//Begin ActorComponent interface
 #if WITH_EDITOR
 	virtual void CheckForErrors() override;
-	virtual void OnComponentDestroyed() override;
 	virtual void OnRegister() override;
 #endif // WITH_EDITOR
 	virtual void OnUnregister() override;
@@ -259,10 +258,7 @@ public:
 	UFUNCTION(BlueprintCallable, Category = "Physics|Components|PhysicsConstraint")
 	ENGINE_API void GetConstraintForce(FVector& OutLinearForce, FVector& OutAngularForce);
 
-#if WITH_EDITORONLY_DATA
-	UPROPERTY(transient)
-	class UBillboardComponent* SpriteComponent;
-
+#if WITH_EDITOR
 	ENGINE_API void UpdateSpriteTexture();
 #endif
 
diff --git a/Engine/Source/Runtime/Engine/Private/AudioComponent.cpp b/Engine/Source/Runtime/Engine/Private/AudioComponent.cpp
--- a/Engine/Source/Runtime/Engine/Private/AudioComponent.cpp
+++ b/Engine/Source/Runtime/Engine/Private/AudioComponent.cpp
@@ -48,23 +48,7 @@ void UAudioComponent::OnRegister()
 {
 	Super::OnRegister();
 
-	if ( bVisualizeComponent && SpriteComponent == NULL && GetOwner() && !GetWorld()->IsGameWorld() )
-	{
-		SpriteComponent = NewObject<UBillboardComponent>(GetOwner(), NAME_None, RF_Transactional | RF_TextExportTransient);
-
-		UpdateSpriteTexture();
-		SpriteComponent->RelativeScale3D = FVector(0.5f, 0.5f, 0.5f);
-		SpriteComponent->AttachTo(this);
-		SpriteComponent->AlwaysLoadOnClient = false;
-		SpriteComponent->AlwaysLoadOnServer = false;
-		SpriteComponent->SpriteInfo.Category = TEXT("Misc");
-		SpriteComponent->SpriteInfo.DisplayName = NSLOCTEXT( "SpriteCategory", "Misc", "Misc" );
-		SpriteComponent->CreationMethod = CreationMethod;
-		SpriteComponent->bIsScreenSizeScaled = true;
-		SpriteComponent->bUseInEditorScaling = true;
-
-		SpriteComponent->RegisterComponent();
-	}
+	UpdateSpriteTexture();
 }
 #endif
 
@@ -83,17 +67,6 @@ void UAudioComponent::OnUnregister()
 	}
 }
 
-#if WITH_EDITORONLY_DATA
-void UAudioComponent::OnComponentDestroyed()
-{
-	Super::OnComponentDestroyed();
-	if (SpriteComponent)
-	{
-		SpriteComponent->DestroyComponent();
-	}
-}
-#endif
-
 void UAudioComponent::SetSound( USoundBase* NewSound )
 {
 	const bool bPlay = IsPlaying();
@@ -324,11 +297,11 @@ void UAudioComponent::UpdateSpriteTexture()
 	{
 		if (bAutoActivate)
 		{
-			SpriteComponent->Sprite = LoadObject<UTexture2D>(NULL, TEXT("/Engine/EditorResources/AudioIcons/S_AudioComponent_AutoActivate.S_AudioComponent_AutoActivate"));
+			SpriteComponent->SetSprite(LoadObject<UTexture2D>(NULL, TEXT("/Engine/EditorResources/AudioIcons/S_AudioComponent_AutoActivate.S_AudioComponent_AutoActivate")));
 		}
 		else
 		{
-			SpriteComponent->Sprite = LoadObject<UTexture2D>(NULL, TEXT("/Engine/EditorResources/AudioIcons/S_AudioComponent.S_AudioComponent"));
+			SpriteComponent->SetSprite(LoadObject<UTexture2D>(NULL, TEXT("/Engine/EditorResources/AudioIcons/S_AudioComponent.S_AudioComponent")));
 		}
 	}
 }
diff --git a/Engine/Source/Runtime/Engine/Private/LightComponent.cpp b/Engine/Source/Runtime/Engine/Private/LightComponent.cpp
--- a/Engine/Source/Runtime/Engine/Private/LightComponent.cpp
+++ b/Engine/Source/Runtime/Engine/Private/LightComponent.cpp
@@ -168,29 +168,20 @@ bool ULightComponentBase::HasStaticShadowing() const
 	return Owner && (Mobility != EComponentMobility::Movable);
 }
 
+#if WITH_EDITOR
 void ULightComponentBase::OnRegister()
 {
 	Super::OnRegister();
 
-#if WITH_EDITOR
-	if (SpriteComponent == NULL && GetOwner() && !GetWorld()->IsGameWorld())
+	if (SpriteComponent)
 	{
-		SpriteComponent = NewObject<UBillboardComponent>(GetOwner(), NAME_None, RF_Transactional | RF_TextExportTransient);
-
-		SpriteComponent->AttachTo(this);
-		SpriteComponent->AlwaysLoadOnClient = false;
-		SpriteComponent->AlwaysLoadOnServer = false;
 		SpriteComponent->SpriteInfo.Category = TEXT("Lighting");
 		SpriteComponent->SpriteInfo.DisplayName = NSLOCTEXT("SpriteCategory", "Lighting", "Lighting");
-		SpriteComponent->CreationMethod = CreationMethod;
-		SpriteComponent->bIsScreenSizeScaled = true;
-
-		SpriteComponent->RegisterComponent();
 
 		UpdateLightSpriteTexture();
 	}
-#endif
 }
+#endif
 
 bool ULightComponentBase::ShouldCollideWhenPlacing() const
 {
@@ -302,6 +293,7 @@ ULightComponentBase::ULightComponentBase(const FObjectInitializer& ObjectInitial
 	CastStaticShadows = true;
 	CastDynamicShadows = true;
 	bPrecomputedLightingIsValid = true;
+	bVisualizeComponent = true;
 }
 
 /**
diff --git a/Engine/Source/Runtime/Engine/Private/PhysicsEngine/PhysicsConstraintComponent.cpp b/Engine/Source/Runtime/Engine/Private/PhysicsEngine/PhysicsConstraintComponent.cpp
--- a/Engine/Source/Runtime/Engine/Private/PhysicsEngine/PhysicsConstraintComponent.cpp
+++ b/Engine/Source/Runtime/Engine/Private/PhysicsEngine/PhysicsConstraintComponent.cpp
@@ -13,6 +13,7 @@ UPhysicsConstraintComponent::UPhysicsConstraintComponent(const FObjectInitialize
 	: Super(ObjectInitializer)
 {
 	bWantsInitializeComponent = true;
+	bVisualizeComponent = true;
 }
 
 
@@ -231,34 +232,15 @@ void UPhysicsConstraintComponent::InitializeComponent()
 }
 
 #if WITH_EDITOR
-void UPhysicsConstraintComponent::OnComponentDestroyed()
-{
-	Super::OnComponentDestroyed();
-
-	if (SpriteComponent)
-	{
-		SpriteComponent->DestroyComponent();
-	}
-}
-
 void UPhysicsConstraintComponent::OnRegister()
 {
 	Super::OnRegister();
 
-	if (SpriteComponent == NULL && GetOwner() && !GetWorld()->IsGameWorld() )
+	if (SpriteComponent)
 	{
-		SpriteComponent = NewObject<UBillboardComponent>(GetOwner(), NAME_None, RF_Transactional | RF_TextExportTransient);
-
 		UpdateSpriteTexture();
-		SpriteComponent->AttachTo(this);
-		SpriteComponent->AlwaysLoadOnClient = false;
-		SpriteComponent->AlwaysLoadOnServer = false;
 		SpriteComponent->SpriteInfo.Category = TEXT("Physics");
 		SpriteComponent->SpriteInfo.DisplayName = NSLOCTEXT( "SpriteCategory", "Physics", "Physics" );
-		SpriteComponent->CreationMethod = CreationMethod;
-		SpriteComponent->bIsScreenSizeScaled = true;
-
-		SpriteComponent->RegisterComponent();
 	}
 }
 #endif
diff --git a/Engine/Source/Runtime/Engine/Private/SceneComponent.cpp b/Engine/Source/Runtime/Engine/Private/SceneComponent.cpp
--- a/Engine/Source/Runtime/Engine/Private/SceneComponent.cpp
+++ b/Engine/Source/Runtime/Engine/Private/SceneComponent.cpp
@@ -119,6 +119,28 @@ void USceneComponent::OnRegister()
 	AttachTo(AttachParent, AttachSocketName);
 	
 	Super::OnRegister();
+
+#if WITH_EDITOR
+	if (bVisualizeComponent && SpriteComponent == nullptr && GetOwner() && !GetWorld()->IsGameWorld() )
+	{
+		// Create a new billboard component to serve as a visualization of the actor until there is another primitive component
+		SpriteComponent = ConstructObject<UBillboardComponent>(UBillboardComponent::StaticClass(), GetOwner(), NAME_None, RF_Transactional | RF_TextExportTransient);
+
+		SpriteComponent->Sprite = LoadObject<UTexture2D>(nullptr, TEXT("/Engine/EditorResources/EmptyActor.EmptyActor"));
+		SpriteComponent->RelativeScale3D = FVector(0.5f, 0.5f, 0.5f);
+		SpriteComponent->Mobility = EComponentMobility::Movable;
+		SpriteComponent->AlwaysLoadOnClient = false;
+		SpriteComponent->AlwaysLoadOnServer = false;
+		SpriteComponent->SpriteInfo.Category = TEXT("Misc");
+		SpriteComponent->SpriteInfo.DisplayName = NSLOCTEXT( "SpriteCategory", "Misc", "Misc" );
+		SpriteComponent->CreationMethod = CreationMethod;
+		SpriteComponent->bIsScreenSizeScaled = true;
+		SpriteComponent->bUseInEditorScaling = true;
+
+		SpriteComponent->AttachTo(this);
+		SpriteComponent->RegisterComponent();
+	}
+#endif
 }
 
 void USceneComponent::UpdateComponentToWorld(bool bSkipPhysicsMove)
@@ -298,6 +320,8 @@ void USceneComponent::DestroyComponent(bool bPromoteChildren/*= false*/)
 					auto NewRootComponent = NewObject<USceneComponent>(Owner, USceneComponent::GetDefaultSceneRootVariableName(), RF_Transactional);
 					NewRootComponent->Mobility = Mobility;
 					NewRootComponent->SetWorldLocationAndRotation(GetComponentLocation(), GetComponentRotation());
+					NewRootComponent->bVisualizeComponent = true;
+					NewRootComponent->RegisterComponent();
 
 					// Designate the new default root as the child we're promoting
 					ChildToPromote = NewRootComponent;
@@ -354,6 +378,11 @@ void USceneComponent::OnComponentDestroyed()
 {
 	Super::OnComponentDestroyed();
 
+	if (SpriteComponent)
+	{
+		SpriteComponent->DestroyComponent();
+	}
+
 	ScopedMovementStack.Reset();
 
 	// Ensure we are detached before destroying
diff --git a/Engine/Source/Runtime/Engine/Private/SimpleConstructionScript.cpp b/Engine/Source/Runtime/Engine/Private/SimpleConstructionScript.cpp
--- a/Engine/Source/Runtime/Engine/Private/SimpleConstructionScript.cpp
+++ b/Engine/Source/Runtime/Engine/Private/SimpleConstructionScript.cpp
@@ -651,6 +651,7 @@ void USimpleConstructionScript::ValidateSceneRootNodes()
 			&& Blueprint->BlueprintType != BPTYPE_MacroLibrary)
 		{
 			DefaultSceneRootNode = CreateNode(USceneComponent::StaticClass(), USceneComponent::GetDefaultSceneRootVariableName());
+			CastChecked<USceneComponent>(DefaultSceneRootNode->ComponentTemplate)->bVisualizeComponent = true;
 		}
 	}
 
