[INTEGRATE] Change 2426007 by [REDACTED_PERSON][REDACTED_HANDLE] on 2015/01/30 12:18:10

	Make newly created default scene root an instance component

diff --git a/Engine/Source/Runtime/Engine/Private/SceneComponent.cpp b/Engine/Source/Runtime/Engine/Private/SceneComponent.cpp
--- a/Engine/Source/Runtime/Engine/Private/SceneComponent.cpp
+++ b/Engine/Source/Runtime/Engine/Private/SceneComponent.cpp
@@ -321,6 +321,7 @@ void USceneComponent::DestroyComponent(bool bPromoteChildren/*= false*/)
 					NewRootComponent->Mobility = Mobility;
 					NewRootComponent->SetWorldLocationAndRotation(GetComponentLocation(), GetComponentRotation());
 					NewRootComponent->bVisualizeComponent = true;
+					Owner->AddInstanceComponent(NewRootComponent);
 					NewRootComponent->RegisterComponent();
 
 					// Designate the new default root as the child we're promoting
