[INTEGRATE] Change 2426127 by [REDACTED_PERSON][REDACTED_HANDLE] on 2015/01/30 14:04:42

	Fix WITH_EDITOR = 0 compile errors

diff --git a/Engine/Source/Runtime/Engine/Private/AudioComponent.cpp b/Engine/Source/Runtime/Engine/Private/AudioComponent.cpp
--- a/Engine/Source/Runtime/Engine/Private/AudioComponent.cpp
+++ b/Engine/Source/Runtime/Engine/Private/AudioComponent.cpp
@@ -17,7 +17,9 @@ UAudioComponent::UAudioComponent(const FObjectInitializer& ObjectInitializer)
 	bAllowSpatialization = true;
 	bStopWhenOwnerDestroyed = true;
 	bNeverNeedsRenderUpdate = true;
+#if WITH_EDITOR
 	bVisualizeComponent = true;
+#endif
 	VolumeMultiplier = 1.f;
 	PitchMultiplier = 1.f;
 	VolumeModulationMin = 1.f;
diff --git a/Engine/Source/Runtime/Engine/Private/LightComponent.cpp b/Engine/Source/Runtime/Engine/Private/LightComponent.cpp
--- a/Engine/Source/Runtime/Engine/Private/LightComponent.cpp
+++ b/Engine/Source/Runtime/Engine/Private/LightComponent.cpp
@@ -293,7 +293,9 @@ ULightComponentBase::ULightComponentBase(const FObjectInitializer& ObjectInitial
 	CastStaticShadows = true;
 	CastDynamicShadows = true;
 	bPrecomputedLightingIsValid = true;
+#if WITH_EDITOR
 	bVisualizeComponent = true;
+#endif
 }
 
 /**
diff --git a/Engine/Source/Runtime/Engine/Private/PhysicsEngine/PhysicsConstraintComponent.cpp b/Engine/Source/Runtime/Engine/Private/PhysicsEngine/PhysicsConstraintComponent.cpp
--- a/Engine/Source/Runtime/Engine/Private/PhysicsEngine/PhysicsConstraintComponent.cpp
+++ b/Engine/Source/Runtime/Engine/Private/PhysicsEngine/PhysicsConstraintComponent.cpp
@@ -13,7 +13,9 @@ UPhysicsConstraintComponent::UPhysicsConstraintComponent(const FObjectInitialize
 	: Super(ObjectInitializer)
 {
 	bWantsInitializeComponent = true;
+#if WITH_EDITOR
 	bVisualizeComponent = true;
+#endif
 }
 
 
diff --git a/Engine/Source/Runtime/Engine/Private/SceneComponent.cpp b/Engine/Source/Runtime/Engine/Private/SceneComponent.cpp
--- a/Engine/Source/Runtime/Engine/Private/SceneComponent.cpp
+++ b/Engine/Source/Runtime/Engine/Private/SceneComponent.cpp
@@ -320,7 +320,9 @@ void USceneComponent::DestroyComponent(bool bPromoteChildren/*= false*/)
 					auto NewRootComponent = NewObject<USceneComponent>(Owner, USceneComponent::GetDefaultSceneRootVariableName(), RF_Transactional);
 					NewRootComponent->Mobility = Mobility;
 					NewRootComponent->SetWorldLocationAndRotation(GetComponentLocation(), GetComponentRotation());
+#if WITH_EDITOR
 					NewRootComponent->bVisualizeComponent = true;
+#endif
 					Owner->AddInstanceComponent(NewRootComponent);
 					NewRootComponent->RegisterComponent();
 
