[INTEGRATE] Change 2426130 by [REDACTED_PERSON][REDACTED_HANDLE] on 2015/01/30 14:06:51

	Further fix for WITH_EDITOR = 0

diff --git a/Engine/Source/Runtime/Engine/Private/SceneComponent.cpp b/Engine/Source/Runtime/Engine/Private/SceneComponent.cpp
--- a/Engine/Source/Runtime/Engine/Private/SceneComponent.cpp
+++ b/Engine/Source/Runtime/Engine/Private/SceneComponent.cpp
@@ -381,10 +381,12 @@ void USceneComponent::OnComponentDestroyed()
 {
 	Super::OnComponentDestroyed();
 
+#if WITH_EDITOR
 	if (SpriteComponent)
 	{
 		SpriteComponent->DestroyComponent();
 	}
+#endif
 
 	ScopedMovementStack.Reset();
 
