[INTEGRATE] Change 2426334 by [REDACTED_PERSON][REDACTED_HANDLE] on 2015/01/30 16:09:19

	Blueprints: Clarified the comment for delay and retriggerable delay nodes to indicate that the duration is in seconds

diff --git a/Engine/Source/Runtime/Engine/Classes/Kismet/KismetSystemLibrary.h b/Engine/Source/Runtime/Engine/Classes/Kismet/KismetSystemLibrary.h
--- a/Engine/Source/Runtime/Engine/Classes/Kismet/KismetSystemLibrary.h
+++ b/Engine/Source/Runtime/Engine/Classes/Kismet/KismetSystemLibrary.h
@@ -232,16 +232,22 @@ class UKismetSystemLibrary : public UBlueprintFunctionLibrary
 	// Latent Actions
 
 	/** 
-	 * Perform a latent action with a delay.
+	 * Perform a latent action with a delay (specified in seconds).  Calling again while it is counting down will be ignored.
 	 * 
 	 * @param WorldContext	World context.
-	 * @param Duration 		length of delay.
+	 * @param Duration 		length of delay (in seconds).
 	 * @param LatentInfo 	The latent action.
 	 */
 	UFUNCTION(BlueprintCallable, Category="Utilities|FlowControl", meta=(Latent, WorldContext="WorldContextObject", LatentInfo="LatentInfo", Duration="0.2"))
 	static void	Delay(UObject* WorldContextObject, float Duration, struct FLatentActionInfo LatentInfo );
 
-	/** Delay execution by Duration seconds; Calling again before the delay has expired will reset the countdown to Duration. */
+	/** 
+	 * Perform a latent action with a retriggerable delay (specified in seconds).  Calling again while it is counting down will reset the countdown to Duration.
+	 * 
+	 * @param WorldContext	World context.
+	 * @param Duration 		length of delay (in seconds).
+	 * @param LatentInfo 	The latent action.
+	 */
 	UFUNCTION(BlueprintCallable, meta=(Latent, LatentInfo="LatentInfo", WorldContext="WorldContextObject", Duration="0.2"), Category="Utilities|FlowControl")
 	static void RetriggerableDelay(UObject* WorldContextObject, float Duration, FLatentActionInfo LatentInfo);
 
