[INTEGRATE] Change 2426411 by [REDACTED_PERSON][REDACTED_HANDLE] on 2015/01/30 16:53:10

	Fix for being able to delete the root component. UE-8711

diff --git a/Engine/Source/Editor/UnrealEd/Private/EditorActor.cpp b/Engine/Source/Editor/UnrealEd/Private/EditorActor.cpp
--- a/Engine/Source/Editor/UnrealEd/Private/EditorActor.cpp
+++ b/Engine/Source/Editor/UnrealEd/Private/EditorActor.cpp
@@ -563,6 +563,10 @@ bool UUnrealEdEngine::edactDeleteSelected( UWorld* InWorld, bool bVerifyDeletion
 				bComponentWasDeleted = true;
 			}
 		}
+		if (bComponentWasDeleted)
+		{
+			NoteSelectionChange();
+		}
 		return bComponentWasDeleted;
 	}
 
diff --git a/Engine/Source/Runtime/Engine/Private/SceneComponent.cpp b/Engine/Source/Runtime/Engine/Private/SceneComponent.cpp
--- a/Engine/Source/Runtime/Engine/Private/SceneComponent.cpp
+++ b/Engine/Source/Runtime/Engine/Private/SceneComponent.cpp
@@ -316,6 +316,8 @@ void USceneComponent::DestroyComponent(bool bPromoteChildren/*= false*/)
 				// We only promote non editor-only components to root in instanced mode
 				if (ChildToPromote == nullptr || ChildToPromote->IsEditorOnly())
 				{
+					Rename(NULL, GetOuter(), REN_DoNotDirty | REN_DontCreateRedirectors);
+
 					// Construct a new default root component
 					auto NewRootComponent = NewObject<USceneComponent>(Owner, USceneComponent::GetDefaultSceneRootVariableName(), RF_Transactional);
 					NewRootComponent->Mobility = Mobility;
