[INTEGRATE] Change 2426422 by [REDACTED_PERSON][REDACTED_HANDLE] on 2015/01/30 16:58:02

	Fixed IntelliSense compile warning

diff --git a/Engine/Source/Runtime/Engine/Classes/Components/WindDirectionalSourceComponent.h b/Engine/Source/Runtime/Engine/Classes/Components/WindDirectionalSourceComponent.h
--- a/Engine/Source/Runtime/Engine/Classes/Components/WindDirectionalSourceComponent.h
+++ b/Engine/Source/Runtime/Engine/Classes/Components/WindDirectionalSourceComponent.h
@@ -17,7 +17,7 @@ class UWindDirectionalSourceComponent : public USceneComponent
 
 
 public:
-	FWindSourceSceneProxy* SceneProxy;
+	class FWindSourceSceneProxy* SceneProxy;
 
 protected:
 	// Begin UActorComponent interface.
