UE-9143
If we fail to load back into Entry after a travel failure, log a FatalError with more information (rather than just asserting).

diff --git a/Engine/Source/Runtime/Engine/Private/UnrealEngine.cpp b/Engine/Source/Runtime/Engine/Private/UnrealEngine.cpp
--- a/Engine/Source/Runtime/Engine/Private/UnrealEngine.cpp
+++ b/Engine/Source/Runtime/Engine/Private/UnrealEngine.cpp
@@ -8178,14 +8178,17 @@ EBrowseReturnVal::Type UEngine::Browse( FWorldContext& WorldContext, FURL URL, F
 		
 		const UGameMapsSettings* GameMapsSettings = GetDefault<UGameMapsSettings>();
 		bool LoadSucces = LoadMap(WorldContext, FURL(&URL, *(GameMapsSettings->GetGameDefaultMap() + GameMapsSettings->LocalMapOptions), TRAVEL_Partial), NULL, Error);
-		check(LoadSucces);
+		if (LoadSucces==false)
+		{
+			UE_LOG(LogNet, Fatal, TEXT("Failed to load default map (%s). Error: (%s)"), *(GameMapsSettings->GetGameDefaultMap() + GameMapsSettings->LocalMapOptions), *Error);
+		}
 
 		CollectGarbage( GARBAGE_COLLECTION_KEEPFLAGS );
 
 		// now remove "failed" and "closed" options from LastURL so it doesn't get copied on to future URLs
 		WorldContext.LastURL.RemoveOption(TEXT("failed"));
 		WorldContext.LastURL.RemoveOption(TEXT("closed"));
-		return EBrowseReturnVal::Success;
+		return (LoadSucces ? EBrowseReturnVal::Success : EBrowseReturnVal::Failure);
 	}
 	else if( URL.HasOption(TEXT("restart")) )
 	{
