[INTEGRATE] Change 2426593 by [REDACTED_PERSON][REDACTED_HANDLE] on 2015/01/30 18:04:14

	Use EDITOR_ONLY_DATA instead of EDITOR_ONLY to avoid non-editor build compile issues.

diff --git a/Engine/Source/Runtime/Engine/Classes/Components/SceneComponent.h b/Engine/Source/Runtime/Engine/Classes/Components/SceneComponent.h
--- a/Engine/Source/Runtime/Engine/Classes/Components/SceneComponent.h
+++ b/Engine/Source/Runtime/Engine/Classes/Components/SceneComponent.h
@@ -155,7 +155,7 @@ protected:
 	// Transient flag that temporarily disables UpdateOverlaps within DetachFromParent().
 	uint32 bDisableDetachmentUpdateOverlaps:1;
 
-#if WITH_EDITOR
+#if WITH_EDITORONLY_DATA
 protected:
 	/** Editor only component used to display the sprite so as to be able to see the location of the Audio Component  */
 	UPROPERTY(transient)
diff --git a/Engine/Source/Runtime/Engine/Private/AudioComponent.cpp b/Engine/Source/Runtime/Engine/Private/AudioComponent.cpp
--- a/Engine/Source/Runtime/Engine/Private/AudioComponent.cpp
+++ b/Engine/Source/Runtime/Engine/Private/AudioComponent.cpp
@@ -17,7 +17,7 @@ UAudioComponent::UAudioComponent(const FObjectInitializer& ObjectInitializer)
 	bAllowSpatialization = true;
 	bStopWhenOwnerDestroyed = true;
 	bNeverNeedsRenderUpdate = true;
-#if WITH_EDITOR
+#if WITH_EDITORONLY_DATA
 	bVisualizeComponent = true;
 #endif
 	VolumeMultiplier = 1.f;
diff --git a/Engine/Source/Runtime/Engine/Private/AudioDevice.cpp b/Engine/Source/Runtime/Engine/Private/AudioDevice.cpp
--- a/Engine/Source/Runtime/Engine/Private/AudioDevice.cpp
+++ b/Engine/Source/Runtime/Engine/Private/AudioDevice.cpp
@@ -2059,7 +2059,9 @@ UAudioComponent* FAudioDevice::CreateComponent( USoundBase* Sound, UWorld* World
 			AudioComponent->bIsUISound = false;
 			AudioComponent->bAutoDestroy = bPlay;
 			AudioComponent->bStopWhenOwnerDestroyed = bStopWhenOwnerDestroyed;
+#if WITH_EDITORONLY_DATA
 			AudioComponent->bVisualizeComponent	= false;
+#endif
 			AudioComponent->AttenuationSettings = AttenuationSettings;
 			if (Location)
 			{
diff --git a/Engine/Source/Runtime/Engine/Private/LightComponent.cpp b/Engine/Source/Runtime/Engine/Private/LightComponent.cpp
--- a/Engine/Source/Runtime/Engine/Private/LightComponent.cpp
+++ b/Engine/Source/Runtime/Engine/Private/LightComponent.cpp
@@ -293,7 +293,7 @@ ULightComponentBase::ULightComponentBase(const FObjectInitializer& ObjectInitial
 	CastStaticShadows = true;
 	CastDynamicShadows = true;
 	bPrecomputedLightingIsValid = true;
-#if WITH_EDITOR
+#if WITH_EDITORONLY_DATA
 	bVisualizeComponent = true;
 #endif
 }
diff --git a/Engine/Source/Runtime/Engine/Private/PhysicsEngine/PhysicsConstraintComponent.cpp b/Engine/Source/Runtime/Engine/Private/PhysicsEngine/PhysicsConstraintComponent.cpp
--- a/Engine/Source/Runtime/Engine/Private/PhysicsEngine/PhysicsConstraintComponent.cpp
+++ b/Engine/Source/Runtime/Engine/Private/PhysicsEngine/PhysicsConstraintComponent.cpp
@@ -13,7 +13,7 @@ UPhysicsConstraintComponent::UPhysicsConstraintComponent(const FObjectInitialize
 	: Super(ObjectInitializer)
 {
 	bWantsInitializeComponent = true;
-#if WITH_EDITOR
+#if WITH_EDITORONLY_DATA
 	bVisualizeComponent = true;
 #endif
 }
diff --git a/Engine/Source/Runtime/Engine/Private/SceneComponent.cpp b/Engine/Source/Runtime/Engine/Private/SceneComponent.cpp
--- a/Engine/Source/Runtime/Engine/Private/SceneComponent.cpp
+++ b/Engine/Source/Runtime/Engine/Private/SceneComponent.cpp
@@ -120,7 +120,7 @@ void USceneComponent::OnRegister()
 	
 	Super::OnRegister();
 
-#if WITH_EDITOR
+#if WITH_EDITORONLY_DATA
 	if (bVisualizeComponent && SpriteComponent == nullptr && GetOwner() && !GetWorld()->IsGameWorld() )
 	{
 		// Create a new billboard component to serve as a visualization of the actor until there is another primitive component
@@ -322,7 +322,7 @@ void USceneComponent::DestroyComponent(bool bPromoteChildren/*= false*/)
 					auto NewRootComponent = NewObject<USceneComponent>(Owner, USceneComponent::GetDefaultSceneRootVariableName(), RF_Transactional);
 					NewRootComponent->Mobility = Mobility;
 					NewRootComponent->SetWorldLocationAndRotation(GetComponentLocation(), GetComponentRotation());
-#if WITH_EDITOR
+#if WITH_EDITORONLY_DATA
 					NewRootComponent->bVisualizeComponent = true;
 #endif
 					Owner->AddInstanceComponent(NewRootComponent);
@@ -383,7 +383,7 @@ void USceneComponent::OnComponentDestroyed()
 {
 	Super::OnComponentDestroyed();
 
-#if WITH_EDITOR
+#if WITH_EDITORONLY_DATA
 	if (SpriteComponent)
 	{
 		SpriteComponent->DestroyComponent();
