AnalyticsTest.cpp
Added an unit test for the FAnalyticsEventAttribute struct.

diff --git a/Engine/Source/Runtime/Engine/Private/Tests/AnalyticsTest.cpp b/Engine/Source/Runtime/Engine/Private/Tests/AnalyticsTest.cpp
--- a/Engine/Source/Runtime/Engine/Private/Tests/AnalyticsTest.cpp
+++ b/Engine/Source/Runtime/Engine/Private/Tests/AnalyticsTest.cpp
@@ -4,6 +4,7 @@
 #include "AutomationTest.h"
 #include "EngineAnalytics.h"
 #include "Runtime/Analytics/Analytics/Public/Interfaces/IAnalyticsProvider.h"
+#include "Runtime/Analytics/Analytics/Public/AnalyticsEventAttribute.h"
 #include "UnrealEngine.h"
 
 /**
@@ -21,6 +22,9 @@ bool FAnalyticStartUpSimTest::RunTest(const FString& Parameters)
 		EventAttributes.Add(FAnalyticsEventAttribute(TEXT("AccountID"), FPlatformMisc::GetEpicAccountId()));
 		EventAttributes.Add(FAnalyticsEventAttribute(TEXT("GameName"), FApp::GetGameName()));
 		EventAttributes.Add(FAnalyticsEventAttribute(TEXT("CommandLine"), FCommandLine::Get()));
+
+		//@todo - Set the user ID and check that
+		//@todo - Record OS ID as well
 		
 		//Record the event with the 'Engine.AutomationTest.Analytics.ProgramStartedEvent' title
 		FEngineAnalytics::GetProvider().RecordEvent(TEXT("Engine.AutomationTest.Analytics.ProgramStartedEvent"), EventAttributes);
@@ -51,20 +55,71 @@ bool FAnalyticStartUpSimTest::RunTest(const FString& Parameters)
 				}
 			}
 		}
+		return true;
 	}
+
+	ExecutionInfo.LogItems.Add(TEXT("SKIPPED 'FAnalyticStartUpSimTest' test.  EngineAnalytics are not currently available."));
 	return true;	
 }
 
-///**
-//* Engine Startup Analytic unit test - This triggers an artificial event.
-//*/
-//IMPLEMENT_SIMPLE_AUTOMATION_TEST(FAnalyticEngineInit, "Engine.Analytic Record Events.Engine Init Test", EAutomationTestFlags::ATF_FeatureMask)
-//
-//bool FAnalyticEngineInit::RunTest(const FString& Parameters)
-//{
-//	IEngineLoop* InNewEngineLoop = nullptr;
-//
-//	GEngine->Init(InNewEngineLoop);
-//
-//	return true;
-//}
\ No newline at end of file
+
+/**
+* FAnalyticsEventAttribute Unit Test.
+*/
+IMPLEMENT_SIMPLE_AUTOMATION_TEST(FAnalyticsEventAttributeUnitTest, "Engine.Analytic.EventAttribute Struct Unit Test", EAutomationTestFlags::ATF_FeatureMask)
+
+bool FAnalyticsEventAttributeUnitTest::RunTest(const FString& Parameters)
+{
+	if (FEngineAnalytics::IsAvailable())
+	{
+		FString AttributeName;
+		FString AttributeValue;
+
+		AttributeName = "Test of";
+		AttributeValue = "FAnalyticsEventAttribute '(const FString InName, const FString& InVaue)'";
+
+		//Setup the 'Event Attributes'
+		FAnalyticsEventAttribute EventAttributesFStringTest(AttributeName, AttributeValue);
+		TestTrue(TEXT("Expected to take in these type of values '(const FString&, const FString&)'"), (EventAttributesFStringTest.AttrName == TEXT("Test of") && EventAttributesFStringTest.AttrValue == "FAnalyticsEventAttribute '(const FString InName, const FString& InVaue)'"));
+
+		FAnalyticsEventAttribute EventAttributesTCHARTest(AttributeName, TEXT("FAnalyticsEventAttribute '(const FString InName, const TCHAR* InValue)'"));
+		TestTrue(TEXT("Expected to take in these type of values '(const FString&, const TCHAR*')"), (EventAttributesTCHARTest.AttrName == TEXT("Test of") && EventAttributesTCHARTest.AttrValue == TEXT("FAnalyticsEventAttribute '(const FString InName, const TCHAR* InValue)'")));
+
+		bool bIsBoolTest = true;
+		FAnalyticsEventAttribute EventAttributesBoolTest(AttributeName, bIsBoolTest);
+		TestTrue(TEXT("Expected to take in these types of values '(const FString&, bool)'"), (EventAttributesBoolTest.AttrName == TEXT("Test of") && EventAttributesBoolTest.AttrValue == "true"));
+
+		FGuid GuidTest(FGuid::NewGuid());
+		FAnalyticsEventAttribute EventAttributesGuidTest(AttributeName, GuidTest);
+		TestTrue(TEXT("Expected to take in these type of values '(const FString&, FGuid)'"), (EventAttributesGuidTest.AttrName == TEXT("Test of") && EventAttributesGuidTest.AttrValue == GuidTest.ToString()));
+
+		int32 InNumericType = 42;
+		FAnalyticsEventAttribute EventAttributesTInValueTest(AttributeName, InNumericType);
+		TestTrue(TEXT("Expected to take in a arithmetic type (example int32)"), (EventAttributesTInValueTest.AttrName == TEXT("Test of") && EventAttributesTInValueTest.AttrValue == TEXT("42")));
+
+		TArray<int32> InNumericArray;
+		InNumericArray.Add(0);
+		InNumericArray.Add(1);
+		InNumericArray.Add(2);
+		FAnalyticsEventAttribute EventAttributesTArray(AttributeName, InNumericArray);
+		TestTrue(TEXT("Expected to take in am arithmetic TArray"), (EventAttributesTArray.AttrName == TEXT("Test of") && EventAttributesTArray.AttrValue == TEXT("0,1,2")));
+
+		FString TMapKey1 = "TestKey 1";
+		FString TMapKey2 = "TestKey 2";
+		FString TMapKey3 = "TestKey 3";
+		int32 NumericValue1 = 0;
+		int32 NumericValue2 = 1;
+		int32 NumericValue3 = 99;
+		TMap<FString, int32> InTMap;
+		InTMap.Add(TMapKey1, NumericValue1);
+		InTMap.Add(TMapKey2, NumericValue2);
+		InTMap.Add(TMapKey3, NumericValue3);
+		FAnalyticsEventAttribute EventAttributesTMap(AttributeName, InTMap);
+		TestTrue(TEXT("Expected to take in a TMap "), (EventAttributesTMap.AttrName == TEXT("Test of") && EventAttributesTMap.AttrValue == TEXT("TestKey 1:0,TestKey 2:1,TestKey 3:99")));
+
+		return true;
+	}
+
+	ExecutionInfo.Warnings.Add(TEXT("Skipped 'FAnalyticsEventAttributeUnitTest' test.  EngineAnalytics are not currently available."));
+	return true;
+}
\ No newline at end of file
