UE4 - Rename some variables for clarity (re:velocity rendering i nbase pass)

diff --git a/Engine/Source/Runtime/Engine/Classes/Materials/MaterialExpressionSpeedTree.h b/Engine/Source/Runtime/Engine/Classes/Materials/MaterialExpressionSpeedTree.h
--- a/Engine/Source/Runtime/Engine/Classes/Materials/MaterialExpressionSpeedTree.h
+++ b/Engine/Source/Runtime/Engine/Classes/Materials/MaterialExpressionSpeedTree.h
@@ -50,8 +50,8 @@ class UMaterialExpressionSpeedTree : public UMaterialExpression
 	UPROPERTY(EditAnywhere, Category=MaterialExpressionSpeedTree, meta=(DisplayName = "Billboard Threshold", ToolTip="The threshold for triangles to be removed from the bilboard mesh when not facing the camera (0 = none pass, 1 = all pass).", ClampMin = "0.0", ClampMax = "1.0", UIMin = "0.0", UIMax = "1.0"))
 	float BillboardThreshold;
 
-	UPROPERTY(EditAnywhere, Category=MaterialExpressionSpeedTree, AdvancedDisplay, meta=(DisplayName = "Support accurate velocities from wind"))
-	bool bOutputPreviousPosition;
+	UPROPERTY(EditAnywhere, Category=MaterialExpressionSpeedTree, AdvancedDisplay, meta=(DisplayName = "Support accurate velocities from wind. This will incur extra cost per vertex."))
+	bool bAccurateWindVelocities;
 
 	// Begin UObject Interface
 	virtual void Serialize(FArchive& Ar) override;
@@ -65,5 +65,3 @@ class UMaterialExpressionSpeedTree : public UMaterialExpression
 	virtual void GetCaption(TArray<FString>& OutCaptions) const override;
 	// End UMaterialExpression Interface
 };
-
-
diff --git a/Engine/Source/Runtime/Engine/Private/Materials/HLSLMaterialTranslator.h b/Engine/Source/Runtime/Engine/Private/Materials/HLSLMaterialTranslator.h
--- a/Engine/Source/Runtime/Engine/Private/Materials/HLSLMaterialTranslator.h
+++ b/Engine/Source/Runtime/Engine/Private/Materials/HLSLMaterialTranslator.h
@@ -196,8 +196,8 @@ protected:
 	/** true if the material reads particle color in the pixel shader. */
 	uint32 bUsesParticleColor : 1;
 	uint32 bUsesTransformVector : 1;
-	// True if the current property requires last frame's time
-	uint32 bUsesPreviousFrameTime : 1;
+	// True if the current property requires last frame's information
+	uint32 bCompilingPreviousFrame : 1;
 	uint32 bUsesPixelDepthOffset : 1;
 	/** True if material will output accurate velocities during base pass rendering. */
 	uint32 bOutputsBasePassVelocities : 1;
@@ -242,7 +242,7 @@ public:
 	,	bUsesVertexColor(false)
 	,	bUsesParticleColor(false)
 	,	bUsesTransformVector(false)
-	,	bUsesPreviousFrameTime(false)
+	,	bCompilingPreviousFrame(false)
 	,	bUsesPixelDepthOffset(false)
 	,	bOutputsBasePassVelocities(true)
 	,	NumUserTexCoords(0)
@@ -1097,7 +1097,7 @@ protected:
 		{
 			if (CodeChunk.UniformExpression->IsChangingPerFrame())
 			{
-				if (bUsesPreviousFrameTime)
+				if (bCompilingPreviousFrame)
 				{
 					const int32 ScalarInputIndex = MaterialCompilationOutput.UniformExpressionSet.PerFramePrevUniformScalarExpressions.AddUnique(CodeChunk.UniformExpression);
 					FCString::Sprintf(FormattedCode, TEXT("UE_Material_PerFramePrevScalarExpression%u"), ScalarInputIndex);
@@ -1130,7 +1130,7 @@ protected:
 
 			if (CodeChunk.UniformExpression->IsChangingPerFrame())
 			{
-				if (bUsesPreviousFrameTime)
+				if (bCompilingPreviousFrame)
 				{
 					const int32 VectorInputIndex = MaterialCompilationOutput.UniformExpressionSet.PerFramePrevUniformVectorExpressions.AddUnique(CodeChunk.UniformExpression);
 					FCString::Sprintf(FormattedCode, TEXT("UE_Material_PerFramePrevVectorExpression%u%s"), VectorInputIndex, Mask);
@@ -1292,7 +1292,7 @@ protected:
 			ShaderFrequency = GetMaterialPropertyShaderFrequency(InProperty);
 		}
 
-		bUsesPreviousFrameTime = bUsePreviousFrameTime;
+		bCompilingPreviousFrame = bUsePreviousFrameTime;
 
 		CurrentScopeChunks = &PropertyCodeChunks[MaterialProperty][ShaderFrequency];
 	}
@@ -1631,7 +1631,7 @@ protected:
 	{
 		if (!bPeriodic)
 		{
-			if (bUsesPreviousFrameTime)
+			if (bCompilingPreviousFrame)
 			{
 				return AddInlinedCodeChunk(MCT_Float, TEXT("View.PrevFrameGameTime"));
 			}
@@ -1656,7 +1656,7 @@ protected:
 	{
 		if (!bPeriodic)
 		{
-			if (bUsesPreviousFrameTime)
+			if (bCompilingPreviousFrame)
 			{
 				return AddInlinedCodeChunk(MCT_Float, TEXT("View.PrevFrameRealTime"));
 			}
@@ -3758,7 +3758,7 @@ protected:
 	*
 	* @return	Code index
 	*/
-	virtual int32 SpeedTree(ESpeedTreeGeometryType GeometryType, ESpeedTreeWindType WindType, ESpeedTreeLODType LODType, float BillboardThreshold, bool bOutputPreviousPosition) override 
+	virtual int32 SpeedTree(ESpeedTreeGeometryType GeometryType, ESpeedTreeWindType WindType, ESpeedTreeLODType LODType, float BillboardThreshold, bool bAccurateWindVelocities) override 
 	{ 
 		if (ErrorUnlessFeatureLevelSupported(ERHIFeatureLevel::SM4) == INDEX_NONE)
 		{
@@ -3775,7 +3775,7 @@ protected:
 
 			NumUserVertexTexCoords = FMath::Max<uint32>(NumUserVertexTexCoords, 8);
 			// Only generate previous frame's computations if required and opted-in
-			const bool bEnablePreviousFrameInformation = bUsesPreviousFrameTime && bOutputPreviousPosition;
+			const bool bEnablePreviousFrameInformation = bCompilingPreviousFrame && bAccurateWindVelocities;
 			return AddCodeChunk(MCT_Float3, TEXT("GetSpeedTreeVertexOffset(Parameters, %d, %d, %d, %g, %s)"), GeometryType, WindType, LODType, BillboardThreshold, bEnablePreviousFrameInformation ? TEXT("true") : TEXT("false"));
 		}
 	}
diff --git a/Engine/Source/Runtime/Engine/Private/Materials/MaterialExpressions.cpp b/Engine/Source/Runtime/Engine/Private/Materials/MaterialExpressions.cpp
--- a/Engine/Source/Runtime/Engine/Private/Materials/MaterialExpressions.cpp
+++ b/Engine/Source/Runtime/Engine/Private/Materials/MaterialExpressions.cpp
@@ -8845,14 +8845,14 @@ UMaterialExpressionSpeedTree::UMaterialExpressionSpeedTree(const FObjectInitiali
 	WindType = STW_None;
 	LODType = STLOD_Pop;
 	BillboardThreshold = 0.9f;
-	bOutputPreviousPosition = false;
+	bAccurateWindVelocities = false;
 
 	MenuCategories.Add(ConstructorStatics.NAME_SpeedTree);
 }
 
 int32 UMaterialExpressionSpeedTree::Compile(class FMaterialCompiler* Compiler, int32 OutputIndex, int32 MultiplexIndex)
 {
-	return Compiler->SpeedTree(GeometryType, WindType, LODType, BillboardThreshold, bOutputPreviousPosition);
+	return Compiler->SpeedTree(GeometryType, WindType, LODType, BillboardThreshold, bAccurateWindVelocities);
 }
 
 void UMaterialExpressionSpeedTree::GetCaption(TArray<FString>& OutCaptions) const
diff --git a/Engine/Source/Runtime/Engine/Public/MaterialCompiler.h b/Engine/Source/Runtime/Engine/Public/MaterialCompiler.h
--- a/Engine/Source/Runtime/Engine/Public/MaterialCompiler.h
+++ b/Engine/Source/Runtime/Engine/Public/MaterialCompiler.h
@@ -172,7 +172,7 @@ public:
 	virtual int32 BlackBody( int32 Temp ) = 0;
 	virtual int32 DepthOfFieldFunction(int32 Depth, int32 FunctionValueIndex) = 0;
 	virtual int32 AtmosphericFogColor(int32 WorldPosition) = 0;
-	virtual int32 SpeedTree(ESpeedTreeGeometryType GeometryType, ESpeedTreeWindType WindType, ESpeedTreeLODType LODType, float BillboardThreshold, bool bOutputPreviousPosition) = 0;
+	virtual int32 SpeedTree(ESpeedTreeGeometryType GeometryType, ESpeedTreeWindType WindType, ESpeedTreeLODType LODType, float BillboardThreshold, bool bAccurateWindVelocities) = 0;
 	virtual int32 TextureCoordinateOffset() = 0;
 	virtual int32 EyeAdaptation() = 0;
 };
@@ -335,9 +335,9 @@ public:
 		return Compiler->DepthOfFieldFunction(Depth, FunctionValueIndex);
 	}
 
-	virtual int32 SpeedTree(ESpeedTreeGeometryType GeometryType, ESpeedTreeWindType WindType, ESpeedTreeLODType LODType, float BillboardThreshold, bool bOutputPreviousPosition) override 
+	virtual int32 SpeedTree(ESpeedTreeGeometryType GeometryType, ESpeedTreeWindType WindType, ESpeedTreeLODType LODType, float BillboardThreshold, bool bAccurateWindVelocities) override 
 	{ 
-		return Compiler->SpeedTree(GeometryType, WindType, LODType, BillboardThreshold, bOutputPreviousPosition); 
+		return Compiler->SpeedTree(GeometryType, WindType, LODType, BillboardThreshold, bAccurateWindVelocities); 
 	}
 
 	virtual int32 AtmosphericFogColor(int32 WorldPosition) override
