Don't serialize SourceRawAnimationData while cooking. Fixes issue with layer edited animations causing load errors in cooked builds.

diff --git a/Engine/Source/Runtime/Engine/Private/Animation/AnimSequence.cpp b/Engine/Source/Runtime/Engine/Private/Animation/AnimSequence.cpp
--- a/Engine/Source/Runtime/Engine/Private/Animation/AnimSequence.cpp
+++ b/Engine/Source/Runtime/Engine/Private/Animation/AnimSequence.cpp
@@ -221,9 +221,12 @@ void UAnimSequence::Serialize(FArchive& Ar)
 	{
 		Ar << RawAnimationData;
 #if WITH_EDITORONLY_DATA
-		if (Ar.UE4Ver() >= VER_UE4_ANIMATION_ADD_TRACKCURVES)
+		if (!Ar.IsCooking())
 		{
-			Ar << SourceRawAnimationData;
+			if (Ar.UE4Ver() >= VER_UE4_ANIMATION_ADD_TRACKCURVES)
+			{
+				Ar << SourceRawAnimationData;
+			}
 		}
 #endif // WITH_EDITORONLY_DATA
 	}
