Modifications to network replay system to use new replay streamers- work in progress

diff --git a/Engine/Source/Runtime/Engine/Classes/Engine/DemoNetDriver.h b/Engine/Source/Runtime/Engine/Classes/Engine/DemoNetDriver.h
--- a/Engine/Source/Runtime/Engine/Classes/Engine/DemoNetDriver.h
+++ b/Engine/Source/Runtime/Engine/Classes/Engine/DemoNetDriver.h
@@ -4,7 +4,6 @@
 
 #include "DemoNetDriver.generated.h"
 
-
 /**
  * Simulated network driver for recording and playing back game sessions.
  */
@@ -17,9 +16,6 @@ class UDemoNetDriver
 	/** Name of the file to read/write from */
 	FString DemoFilename;
 
-	/** Handle to the archive that will read/write network packets */
-	FArchive* FileAr;
-
 	/** Current record/playback frame number */
 	int32 DemoFrameNum;
 
@@ -53,6 +49,8 @@ class UDemoNetDriver
 	UPROPERTY(config)
 	FString DemoSpectatorClass;
 
+	TSharedPtr< class INetworkReplayStreamer > ReplayStreamer;
+
 public:
 
 	// UNetDriver interface.
@@ -67,6 +65,8 @@ public:
 	virtual void ProcessRemoteFunction( class AActor* Actor, class UFunction* Function, void* Parameters, struct FOutParmRec* OutParms, struct FFrame* Stack, class UObject* SubObject = nullptr );
 	virtual bool IsAvailable() const override { return true; }
 	
+	bool InitConnectInternal( FString& Error );
+
 public:
 
 	// FExec interface
@@ -91,4 +91,6 @@ public:
 	void ResetDemoState();
 
 	void StopDemo();
+
+	void ReplayStreamingReady( bool bSuccess, bool bRecord );
 };
diff --git a/Engine/Source/Runtime/Engine/Engine.Build.cs b/Engine/Source/Runtime/Engine/Engine.Build.cs
--- a/Engine/Source/Runtime/Engine/Engine.Build.cs
+++ b/Engine/Source/Runtime/Engine/Engine.Build.cs
@@ -171,6 +171,8 @@ public class Engine : ModuleRules
 			new string[] {
 				"Analytics",
 				"AnalyticsET",
+				"NetworkReplayStreaming",
+				"NullNetworkReplayStreaming",
 				"OnlineSubsystem", 
 				"OnlineSubsystemUtils",
 				"Advertising"
diff --git a/Engine/Source/Runtime/Engine/Private/DemoNetDriver.cpp b/Engine/Source/Runtime/Engine/Private/DemoNetDriver.cpp
--- a/Engine/Source/Runtime/Engine/Private/DemoNetDriver.cpp
+++ b/Engine/Source/Runtime/Engine/Private/DemoNetDriver.cpp
@@ -15,6 +15,7 @@
 #include "GameFramework/PlayerStart.h"
 #include "GameFramework/SpectatorPawnMovement.h"
 #include "Engine/GameInstance.h"
+#include "Runtime/NetworkReplayStreaming/NetworkReplayStreaming/Public/NetworkReplayStreaming.h"	// Why do I need the full path here?
 
 DEFINE_LOG_CATEGORY_STATIC( LogDemo, Log, All );
 
@@ -46,6 +47,8 @@ bool UDemoNetDriver::InitBase( bool bInitAsClient, FNetworkNotify* InNotify, con
 
 		ResetDemoState();
 
+		ReplayStreamer = FNetworkReplayStreaming::Get().GetFactory().CreateReplayStreamer();
+
 		return true;
 	}
 
@@ -133,19 +136,29 @@ bool UDemoNetDriver::InitConnect( FNetworkNotify* InNotify, const FURL& ConnectU
 		return false;
 	}
 
-	UGameInstance* GameInstance = GetWorld()->GetGameInstance();
-
 	// handle default initialization
 	if ( !InitBase( true, InNotify, ConnectURL, false, Error ) )
 	{
-		GameInstance->HandleDemoPlaybackFailure( EDemoPlayFailure::Generic, FString( TEXT( "InitBase FAILED" ) ) );
+		GetWorld()->GetGameInstance()->HandleDemoPlaybackFailure( EDemoPlayFailure::Generic, FString( TEXT( "InitBase FAILED" ) ) );
 		return false;
 	}
 
+	// Playback, local machine is a client, and the demo stream acts "as if" it's the server.
+	ServerConnection = NewObject<UNetConnection>(GetTransientPackage(), UDemoNetConnection::StaticClass());
+	ServerConnection->InitConnection( this, USOCK_Pending, ConnectURL, 1000000 );
+
+	ReplayStreamer->StartStreaming( DemoFilename, false, FOnStreamReadyDelegate::CreateUObject( this, &UDemoNetDriver::ReplayStreamingReady ) );
+
+	return true;
+}
+
+bool UDemoNetDriver::InitConnectInternal( FString& Error )
+{
+	UGameInstance* GameInstance = GetWorld()->GetGameInstance();
+
 	ResetDemoState();
 
-	// open the pre-recorded demo file
-	FileAr = IFileManager::Get().CreateFileReader( *DemoFilename );
+	FArchive* FileAr = ReplayStreamer->GetStreamingArchive();
 
 	if ( !FileAr )
 	{
@@ -155,10 +168,6 @@ bool UDemoNetDriver::InitConnect( FNetworkNotify* InNotify, const FURL& ConnectU
 		return false;
 	}
 
-	// Playback, local machine is a client, and the demo stream acts "as if" it's the server.
-	ServerConnection = NewObject<UNetConnection>(GetTransientPackage(), UDemoNetConnection::StaticClass());
-	ServerConnection->InitConnection( this, USOCK_Pending, ConnectURL, 1000000 );
-
 	// use the same byte format regardless of platform so that the demos are cross platform
 	// DEMO_FIXME: This is messing up for some reason, investigate
 	//FileAr->SetByteSwapping( true );
@@ -280,7 +289,7 @@ bool UDemoNetDriver::InitListen( FNetworkNotify* InNotify, FURL& ListenURL, bool
 
 	check( World != NULL );
 
-	class AWorldSettings * WorldSettings = World->GetWorldSettings();
+	class AWorldSettings * WorldSettings = World->GetWorldSettings(); 
 
 	if ( !WorldSettings )
 	{
@@ -292,10 +301,12 @@ bool UDemoNetDriver::InitListen( FNetworkNotify* InNotify, FURL& ListenURL, bool
 	UDemoNetConnection* Connection = NewObject<UDemoNetConnection>();
 	Connection->InitConnection( this, USOCK_Open, ListenURL, 1000000 );
 	Connection->InitSendBuffer();
-
-	FileAr = IFileManager::Get().CreateFileWriter( *DemoFilename );
 	ClientConnections.Add( Connection );
 
+	ReplayStreamer->StartStreaming( DemoFilename, true, FOnStreamReadyDelegate::CreateUObject( this, &UDemoNetDriver::ReplayStreamingReady ) );
+
+	FArchive* FileAr = ReplayStreamer->GetStreamingArchive();
+
 	if( !FileAr )
 	{
 		Error = FString::Printf( TEXT("Couldn't open demo file %s for writing"), *DemoFilename );//@todo demorec: localize
@@ -329,6 +340,8 @@ void UDemoNetDriver::TickFlush( float DeltaSeconds )
 {
 	Super::TickFlush( DeltaSeconds );
 
+	FArchive* FileAr = ReplayStreamer->GetStreamingArchive();
+
 	if ( FileAr )
 	{
 		if ( ClientConnections.Num() > 0 )
@@ -416,6 +429,8 @@ void UDemoNetDriver::StopDemo()
 
 	if ( !ServerConnection )
 	{
+		FArchive* FileAr = ReplayStreamer->GetStreamingArchive();
+
 		// Finish writing the header and other information that goes at the end
 		if ( FileAr != NULL && World != NULL )
 		{
@@ -491,8 +506,7 @@ void UDemoNetDriver::StopDemo()
 		ServerConnection = NULL;
 	}
 
-	delete FileAr;
-	FileAr = NULL;
+	ReplayStreamer->StopStreaming();
 
 	check( ClientConnections.Num() == 0 );
 	check( ServerConnection == NULL );
@@ -549,6 +563,8 @@ void UDemoNetDriver::TickDemoRecord( float DeltaSeconds )
 		return;
 	}
 
+	FArchive* FileAr = ReplayStreamer->GetStreamingArchive();
+
 	if ( FileAr == NULL )
 	{
 		return;
@@ -625,6 +641,8 @@ void UDemoNetDriver::TickDemoRecord( float DeltaSeconds )
 
 bool UDemoNetDriver::ReadDemoFrame()
 {
+	FArchive* FileAr = ReplayStreamer->GetStreamingArchive();
+
 	if ( FileAr->IsError() )
 	{
 		StopDemo();
@@ -830,6 +848,21 @@ void UDemoNetDriver::SpawnDemoRecSpectator( UNetConnection* Connection )
 	Controller->SetPlayer( Connection );
 }
 
+void UDemoNetDriver::ReplayStreamingReady( bool bSuccess, bool bRecord )
+{
+	if ( !bSuccess )
+	{
+		GetWorld()->GetGameInstance()->HandleDemoPlaybackFailure( EDemoPlayFailure::DemoNotFound, FString( EDemoPlayFailure::ToString( EDemoPlayFailure::DemoNotFound ) ) );
+		return;
+	}
+
+	if ( !bRecord )
+	{
+		FString Error;
+		InitConnectInternal( Error );
+	}
+}
+
 /*-----------------------------------------------------------------------------
 	UDemoNetConnection.
 -----------------------------------------------------------------------------*/
@@ -869,7 +902,9 @@ void UDemoNetConnection::LowLevelSend( void* Data, int32 Count )
 		UE_LOG( LogDemo, Fatal, TEXT( "UDemoNetConnection::LowLevelSend: Count > MAX_DEMO_READ_WRITE_BUFFER." ) );
 	}
 
-	if ( !GetDriver()->ServerConnection && GetDriver()->FileAr )
+	FArchive* FileAr = GetDriver()->ReplayStreamer->GetStreamingArchive();
+
+	if ( !GetDriver()->ServerConnection && FileAr )
 	{
 		// If we're outside of an official demo frame, we need to queue this up or it will throw off the stream
 		if ( !GetDriver()->bIsRecordingDemoFrame )
@@ -880,12 +915,12 @@ void UDemoNetConnection::LowLevelSend( void* Data, int32 Count )
 			return;
 		}
 
-		*GetDriver()->FileAr << Count;
-		GetDriver()->FileAr->Serialize( Data, Count );
+		*FileAr << Count;
+		FileAr->Serialize( Data, Count );
 		
 #if DEMO_CHECKSUMS == 1
 		uint32 Checksum = FCrc::MemCrc32( Data, Count, 0 );
-		*GetDriver()->FileAr << Checksum;
+		*FileAr << Checksum;
 #endif
 	}
 }
