Fix for crash creating bodies for ISMC if the StaticMesh is null.

diff --git a/Engine/Source/Runtime/Engine/Private/InstancedStaticMesh.cpp b/Engine/Source/Runtime/Engine/Private/InstancedStaticMesh.cpp
--- a/Engine/Source/Runtime/Engine/Private/InstancedStaticMesh.cpp
+++ b/Engine/Source/Runtime/Engine/Private/InstancedStaticMesh.cpp
@@ -929,29 +929,33 @@ void UInstancedStaticMeshComponent::InitInstanceBody(int32 InstanceIdx, FBodyIns
 
 void UInstancedStaticMeshComponent::CreateAllInstanceBodies()
 {
-	int32 NumBodies = PerInstanceSMData.Num();
-	InstanceBodies.Init(NumBodies);
+	UBodySetup* BodySetup = GetBodySetup();
+	if (BodySetup)
+	{
+		int32 NumBodies = PerInstanceSMData.Num();
+		InstanceBodies.Init(NumBodies);
 
-	TArray<FTransform> Transforms;
-	Transforms.Reserve(NumBodies);
+		TArray<FTransform> Transforms;
+		Transforms.Reserve(NumBodies);
 
-	for (int32 i = 0; i < NumBodies; ++i)
-	{
-		InstanceBodies[i] = new FBodyInstance;
-		FBodyInstance* Instance = InstanceBodies[i];
-		Transforms.Add(FTransform(PerInstanceSMData[i].Transform) * ComponentToWorld);
+		for (int32 i = 0; i < NumBodies; ++i)
+		{
+			InstanceBodies[i] = new FBodyInstance;
+			FBodyInstance* Instance = InstanceBodies[i];
+			Transforms.Add(FTransform(PerInstanceSMData[i].Transform) * ComponentToWorld);
 
-		Instance->CopyBodyInstancePropertiesFrom(&BodyInstance);
-		Instance->InstanceBodyIndex = i; // Set body index 
-		Instance->bAutoWeld = false;
+			Instance->CopyBodyInstancePropertiesFrom(&BodyInstance);
+			Instance->InstanceBodyIndex = i; // Set body index 
+			Instance->bAutoWeld = false;
 
-		// make sure we never enable bSimulatePhysics for ISMComps
-		Instance->bSimulatePhysics = false;
-	}
+			// make sure we never enable bSimulatePhysics for ISMComps
+			Instance->bSimulatePhysics = false;
+		}
 
-	if(NumBodies > 0)
-	{
-		FBodyInstance::InitBodies(InstanceBodies, Transforms, GetBodySetup(), this, GetWorld()->GetPhysicsScene(), nullptr, true);
+		if (NumBodies > 0)
+		{
+			FBodyInstance::InitBodies(InstanceBodies, Transforms, BodySetup, this, GetWorld()->GetPhysicsScene(), nullptr, true);
+		}
 	}
 }
 
