[INTEGRATE] Change 2426889 by [REDACTED_PERSON][REDACTED_HANDLE] on 2015/01/30 21:50:12

	Rendering: Exposed TranslucentSortAxis as a configurable property instead of being hardcoded to -Y [UE-8829]

diff --git a/Engine/Source/Runtime/Engine/Classes/Engine/RendererSettings.h b/Engine/Source/Runtime/Engine/Classes/Engine/RendererSettings.h
--- a/Engine/Source/Runtime/Engine/Classes/Engine/RendererSettings.h
+++ b/Engine/Source/Runtime/Engine/Classes/Engine/RendererSettings.h
@@ -174,6 +174,11 @@ class ENGINE_API URendererSettings
 		ToolTip="The sort mode for translucent primitives, affecting how they are ordered and how they change order as the camera moves."))
 	TEnumAsByte<ETranslucentSortPolicy::Type> TranslucentSortPolicy;
 
+	UPROPERTY(config, EditAnywhere, Category=Translucency, meta=(
+		DisplayName="Translucent Sort Axis",
+		ToolTip="The axis that sorting will occur along when Translucent Sort Policy is set to SortAlongAxis."))
+	FVector TranslucentSortAxis;
+
 	UPROPERTY(config, EditAnywhere, Category=Postprocessing, meta=(
 		ConsoleVariable="r.CustomDepth",DisplayName="Custom Depth Pass",
 		ToolTip="Whether the custom depth pass for tagging primitives for postprocessing passes is enabled. Enabling it on demand can save memory but may cause a hitch the first time the feature is used."))
diff --git a/Engine/Source/Runtime/Engine/Private/RenderingSettings.cpp b/Engine/Source/Runtime/Engine/Private/RenderingSettings.cpp
--- a/Engine/Source/Runtime/Engine/Private/RenderingSettings.cpp
+++ b/Engine/Source/Runtime/Engine/Private/RenderingSettings.cpp
@@ -6,6 +6,7 @@
 URendererSettings::URendererSettings(const FObjectInitializer& ObjectInitializer)
 	: Super(ObjectInitializer)
 {
+	TranslucentSortAxis = FVector(0.0f, -1.0f, 0.0f);
 }
 
 static FName ConsoleVariableFName(TEXT("ConsoleVariable"));
diff --git a/Engine/Source/Runtime/Engine/Private/SceneView.cpp b/Engine/Source/Runtime/Engine/Private/SceneView.cpp
--- a/Engine/Source/Runtime/Engine/Private/SceneView.cpp
+++ b/Engine/Source/Runtime/Engine/Private/SceneView.cpp
@@ -16,6 +16,7 @@
 #include "Interfaces/Interface_PostProcessVolume.h"
 #include "Engine/TextureCube.h"
 #include "IHeadMountedDisplay.h"
+#include "Classes/Engine/RendererSettings.h"
 
 DEFINE_LOG_CATEGORY(LogBufferVisualization);
 
@@ -452,8 +453,7 @@ FSceneView::FSceneView(const FSceneViewInitOptions& InitOptions)
 	static TConsoleVariableData<int32>* SortPolicyCvar = IConsoleManager::Get().FindTConsoleVariableDataInt(TEXT("r.TranslucentSortPolicy"));
 	TranslucentSortPolicy = static_cast<ETranslucentSortPolicy::Type>(SortPolicyCvar->GetValueOnAnyThread());
 
-	//@TODO: PAPER2D: Should come from the config value if/when axis switching is allowed
-	TranslucentSortAxis = FVector(0.0f, -1.0f, 0.0f);
+	TranslucentSortAxis = GetDefault<URendererSettings>()->TranslucentSortAxis;
 
 	// As the world is only accessable from the game thread, bIsGameView should be explicitly
 	// set on any other thread.
