[INTEGRATE] Change 2426930 by [REDACTED_PERSON][REDACTED_HANDLE] on 2015/01/30 23:32:21

	Fix shutdown crash in opengl

diff --git a/Engine/Source/Runtime/Engine/Private/Slate/SceneViewport.cpp b/Engine/Source/Runtime/Engine/Private/Slate/SceneViewport.cpp
--- a/Engine/Source/Runtime/Engine/Private/Slate/SceneViewport.cpp
+++ b/Engine/Source/Runtime/Engine/Private/Slate/SceneViewport.cpp
@@ -1147,7 +1147,7 @@ void FSceneViewport::EnqueueBeginRenderFrame()
 	// If we dont have the ViewportRHI then we need to get it before rendering
 	// Note, we need ViewportRHI even if bUseSeparateRenderTarget is true when stereo rendering
 	// is enabled.
-	if( !IsValidRef(ViewportRHI) )
+	if( !IsValidRef(ViewportRHI) && (!bUseSeparateRenderTarget || (GEngine->StereoRenderingDevice.IsValid())) )
 	{
 		// Get the viewport for this window from the renderer so we can render directly to the backbuffer
 		TSharedPtr<FSlateRenderer> Renderer = FSlateApplication::Get().GetRenderer();
