[INTEGRATE] Change 2427071 by [REDACTED_PERSON][REDACTED_HANDLE] on 2015/01/31 11:57:17

	Fixed Scene Outliner showing ugly "_C" suffixes on Blueprint class actor instances

diff --git a/Engine/Source/Runtime/Engine/Private/ActorEditor.cpp b/Engine/Source/Runtime/Engine/Private/ActorEditor.cpp
--- a/Engine/Source/Runtime/Engine/Private/ActorEditor.cpp
+++ b/Engine/Source/Runtime/Engine/Private/ActorEditor.cpp
@@ -471,6 +471,13 @@ const FString& AActor::GetActorLabel() const
 		// NOTE: Calling GetName() is actually fairly slow (does ANSI->Wide conversion, lots of copies, etc.)
 		FString DefaultActorLabel = ActorClass->GetName();
 
+		// Strip off the ugly "_C" suffix for Blueprint class actor instances
+		UBlueprint* GeneratedByClassBlueprint = Cast<UBlueprint>( ActorClass->ClassGeneratedBy );
+		if( GeneratedByClassBlueprint != nullptr && DefaultActorLabel.EndsWith( TEXT( "_C" ) ) )
+		{
+			DefaultActorLabel.RemoveFromEnd( TEXT( "_C" ) );
+		}
+
 		// We want the actor's label to be initially unique, if possible, so we'll use the number of the
 		// actor's FName when creating the initially.  It doesn't actually *need* to be unique, this is just
 		// an easy way to tell actors apart when observing them in a list.  The user can always go and rename
