Add ability to control triggering of notifies, either through random chance or linked to skeletal meshes lod

diff --git a/Engine/Source/Editor/Persona/Private/AnimNotifyDetails.cpp b/Engine/Source/Editor/Persona/Private/AnimNotifyDetails.cpp
--- a/Engine/Source/Editor/Persona/Private/AnimNotifyDetails.cpp
+++ b/Engine/Source/Editor/Persona/Private/AnimNotifyDetails.cpp
@@ -109,6 +109,40 @@ void FAnimNotifyDetails::CustomizeDetails( IDetailLayoutBuilder& DetailBuilder )
 			}
 		}
 	}
+
+	struct FPropVisPair
+	{
+		const TCHAR* NotifyName;
+		TAttribute<EVisibility> Visibility;
+	};
+
+	TriggerFilterModeHandle = DetailBuilder.GetProperty(TEXT("Event.NotifyFilterType"));
+
+	FPropVisPair TriggerSettingNames[] = { { TEXT("Event.NotifyTriggerChance"), TAttribute<EVisibility>(EVisibility::Visible) }
+										 , { TEXT("Event.NotifyFilterType"), TAttribute<EVisibility>(EVisibility::Visible) }
+										 , { TEXT("Event.NotifyFilterLOD"), TAttribute<EVisibility>(this, &FAnimNotifyDetails::VisibilityForLODFilterMode) } };
+
+
+	IDetailCategoryBuilder& TriggerSettingCategory = DetailBuilder.EditCategory(TEXT("Trigger Settings"), FText::GetEmpty(), ECategoryPriority::TypeSpecific);
+
+	for (FPropVisPair& NotifyPair : TriggerSettingNames)
+	{
+		TSharedRef<IPropertyHandle> NotifyPropHandle = DetailBuilder.GetProperty(NotifyPair.NotifyName);
+		DetailBuilder.HideProperty(NotifyPropHandle);
+		TriggerSettingCategory.AddProperty(NotifyPropHandle).Visibility(NotifyPair.Visibility);
+	}
+}
+
+EVisibility FAnimNotifyDetails::VisibilityForLODFilterMode() const
+{
+	uint8 FilterModeValue = 0;
+	FPropertyAccess::Result Ret = TriggerFilterModeHandle.Get()->GetValue(FilterModeValue);
+	if (Ret == FPropertyAccess::Result::Success)
+	{
+		return (FilterModeValue == ENotifyFilterType::LOD) ? EVisibility::Visible : EVisibility::Hidden;
+	}
+
+	return EVisibility::Hidden; //Hidden if we get fail or MultipleValues from the property
 }
 
 void FAnimNotifyDetails::AddBoneNameProperty(IDetailCategoryBuilder& CategoryBuilder, UObject* Notify,  TSharedPtr<IPropertyHandle> Property)
diff --git a/Engine/Source/Editor/Persona/Private/AnimNotifyDetails.h b/Engine/Source/Editor/Persona/Private/AnimNotifyDetails.h
--- a/Engine/Source/Editor/Persona/Private/AnimNotifyDetails.h
+++ b/Engine/Source/Editor/Persona/Private/AnimNotifyDetails.h
@@ -81,4 +81,10 @@ private:
 	 *	@param Index Index of the new slot
 	 */
 	void OnSlotSelected(TSharedPtr<FString> SlotName, ESelectInfo::Type SelectInfo, TSharedPtr<IPropertyHandle> Property);
+
+	/** Return whether or not the LOD filter mode should be visible */
+	EVisibility VisibilityForLODFilterMode() const;
+
+	/** Caches the Filter Mode handle so we can look up its value after customization has finished */
+	TSharedPtr<IPropertyHandle> TriggerFilterModeHandle;
 };
\ No newline at end of file
diff --git a/Engine/Source/Runtime/Engine/Classes/Animation/AnimInstance.h b/Engine/Source/Runtime/Engine/Classes/Animation/AnimInstance.h
--- a/Engine/Source/Runtime/Engine/Classes/Animation/AnimInstance.h
+++ b/Engine/Source/Runtime/Engine/Classes/Animation/AnimInstance.h
@@ -713,6 +713,12 @@ public:
 	/** Add anim notifier **/
 	void AddAnimNotifies(const TArray<const FAnimNotifyEvent*>& NewNotifies, const float InstanceWeight);
 
+	/** Should the notifies current filtering mode stop it from triggering */
+	bool PassesFiltering(const FAnimNotifyEvent* Notify) const;
+
+	/** Work out whether this notify should be triggered based on its chance of triggering value */
+	bool PassesChanceOfTriggering(const FAnimNotifyEvent* Event) const;
+
 	/** Queues an Anim Notify from the shared list on our generated class */
 	void AddAnimNotifyFromGeneratedClass(int32 NotifyIndex);
 
diff --git a/Engine/Source/Runtime/Engine/Classes/Animation/AnimSequenceBase.h b/Engine/Source/Runtime/Engine/Classes/Animation/AnimSequenceBase.h
--- a/Engine/Source/Runtime/Engine/Classes/Animation/AnimSequenceBase.h
+++ b/Engine/Source/Runtime/Engine/Classes/Animation/AnimSequenceBase.h
@@ -42,6 +42,20 @@ namespace EMontageNotifyTickType
 	};
 }
 
+/** Filtering method for deciding whether to trigger a notify */
+UENUM()
+namespace ENotifyFilterType
+{
+	enum Type
+	{
+		/** No filtering*/
+		NoFiltering,
+
+		/** Filter By Skeletal Mesh LOD */
+		LOD,
+	};
+}
+
 /**
  * Triggers an animation notify.  Each AnimNotifyEvent contains an AnimNotify object
  * which has its Notify method called and passed to the animation.
@@ -89,6 +103,18 @@ struct FAnimNotifyEvent : public FAnimLinkableElement
 	UPROPERTY(EditAnywhere, BlueprintReadWrite, Category=AnimNotifyEvent)
 	TEnumAsByte<EMontageNotifyTickType::Type> MontageTickType;
 
+	/** Defines the chance of of this notify triggering, 0 = No Chance, 1 = Always triggers */
+	UPROPERTY(EditAnywhere, BlueprintReadWrite, Category = AnimNotifyTriggerSettings, meta = (ClampMin = "0.0", ClampMax = "1.0"))
+	float NotifyTriggerChance;
+
+	/** Defines a method for filtering notifies (stopping them triggering) e.g. by looking at the meshes current LOD */
+	UPROPERTY(EditAnywhere, BlueprintReadWrite, Category = AnimNotifyTriggerSettings)
+	TEnumAsByte<ENotifyFilterType::Type> NotifyFilterType;
+
+	/** LOD to start filtering this notify from.*/
+	UPROPERTY(EditAnywhere, BlueprintReadWrite, Category = AnimNotifyTriggerSettings, meta = (ClampMin = "0"))
+	int32 NotifyFilterLOD;
+
 #if WITH_EDITORONLY_DATA
 	/** Color of Notify in editor */
 	UPROPERTY()
@@ -112,6 +138,9 @@ struct FAnimNotifyEvent : public FAnimLinkableElement
 		, Duration(0)
 		, bConvertedFromBranchingPoint(false)
 		, MontageTickType(EMontageNotifyTickType::Queued)
+		, NotifyTriggerChance(1.f)
+		, NotifyFilterType(ENotifyFilterType::NoFiltering)
+		, NotifyFilterLOD(0)
 #if WITH_EDITORONLY_DATA
 		, TrackIndex(0)
 #endif // WITH_EDITORONLY_DATA
diff --git a/Engine/Source/Runtime/Engine/Private/Animation/AnimInstance.cpp b/Engine/Source/Runtime/Engine/Private/Animation/AnimInstance.cpp
--- a/Engine/Source/Runtime/Engine/Private/Animation/AnimInstance.cpp
+++ b/Engine/Source/Runtime/Engine/Private/Animation/AnimInstance.cpp
@@ -1164,18 +1164,44 @@ bool UAnimInstance::CanTransitionSignature() const
 void UAnimInstance::AddAnimNotifies(const TArray<const FAnimNotifyEvent*>& NewNotifies, const float InstanceWeight)
 {
 	// for now there is no filter whatsoever, it just adds everything requested
-	for (auto Iter=NewNotifies.CreateConstIterator(); Iter; ++Iter)
+	for (const FAnimNotifyEvent* Notify : NewNotifies)
 	{
 		// only add if it is over TriggerWeightThreshold
-		if ((*Iter)->TriggerWeightThreshold <= InstanceWeight)
+		if (Notify->TriggerWeightThreshold <= InstanceWeight && PassesFiltering(Notify) && PassesChanceOfTriggering(Notify))
 		{
 			// Only add unique AnimNotifyState instances just once. We can get multiple triggers if looping over an animation.
 			// It is the same state, so just report it once.
-			(*Iter)->NotifyStateClass ? AnimNotifies.AddUnique(*Iter) : AnimNotifies.Add(*Iter);
+			Notify->NotifyStateClass ? AnimNotifies.AddUnique(Notify) : AnimNotifies.Add(Notify);
 		}
 	}
 }
 
+bool UAnimInstance::PassesFiltering(const FAnimNotifyEvent* Notify) const
+{
+	switch (Notify->NotifyFilterType)
+	{
+		case ENotifyFilterType::NoFiltering:
+		{
+			return true;
+		}
+		case ENotifyFilterType::LOD:
+		{
+			USkeletalMeshComponent* Comp = GetSkelMeshComponent();
+			return Comp ? Notify->NotifyFilterLOD > Comp->PredictedLODLevel : true;
+		}
+		default:
+		{
+			ensure(false); // Unknown Filter Type
+		}
+	}
+	return true;
+}
+
+bool UAnimInstance::PassesChanceOfTriggering(const FAnimNotifyEvent* Event) const
+{
+	return Event->NotifyStateClass ? true : FMath::FRandRange(0.f, 1.f) < Event->NotifyTriggerChance;
+}
+
 void UAnimInstance::AddAnimNotifyFromGeneratedClass(int32 NotifyIndex)
 {
 	if(NotifyIndex==INDEX_NONE)
