[INTEGRATE] Change 2428132 by [REDACTED_PERSON][REDACTED_HANDLE] on 2015/02/02 10:10:38

	const function

diff --git a/Engine/Source/Runtime/Engine/Classes/Engine/Blueprint.h b/Engine/Source/Runtime/Engine/Classes/Engine/Blueprint.h
--- a/Engine/Source/Runtime/Engine/Classes/Engine/Blueprint.h
+++ b/Engine/Source/Runtime/Engine/Classes/Engine/Blueprint.h
@@ -480,7 +480,7 @@ public:
 	static bool ValidateGeneratedClass(const UClass* InClass);
 
 	/** Find the object in the TemplateObjects array with the supplied name */
-	UActorComponent* FindTemplateByName(const FName& TemplateName);
+	UActorComponent* FindTemplateByName(const FName& TemplateName) const;
 
 	/** Find a timeline by name */
 	class UTimelineTemplate* FindTimelineTemplateByVariableName(const FName& TimelineName);	
diff --git a/Engine/Source/Runtime/Engine/Private/Blueprint.cpp b/Engine/Source/Runtime/Engine/Private/Blueprint.cpp
--- a/Engine/Source/Runtime/Engine/Private/Blueprint.cpp
+++ b/Engine/Source/Runtime/Engine/Private/Blueprint.cpp
@@ -747,7 +747,7 @@ struct FBlueprintInnerHelper
 	}
 };
 
-UActorComponent* UBlueprint::FindTemplateByName(const FName& TemplateName)
+UActorComponent* UBlueprint::FindTemplateByName(const FName& TemplateName) const
 {
 	return FBlueprintInnerHelper::FindObjectByName<UActorComponent>(ComponentTemplates, TemplateName);
 }
