[INTEGRATE] Change 2428412 by [REDACTED_PERSON][REDACTED_HANDLE] on 2015/02/02 13:44:16

	Set PhysicsMobility's default to Mobility's default.
	Fixes assert crash in vehicle code where physics mobility is set to static

	UE-8826

diff --git a/Engine/Source/Runtime/Engine/Private/PrimitiveComponent.cpp b/Engine/Source/Runtime/Engine/Private/PrimitiveComponent.cpp
--- a/Engine/Source/Runtime/Engine/Private/PrimitiveComponent.cpp
+++ b/Engine/Source/Runtime/Engine/Private/PrimitiveComponent.cpp
@@ -114,6 +114,8 @@ UPrimitiveComponent::UPrimitiveComponent(const FObjectInitializer& ObjectInitial
 
 	bCachedAllCollideableDescendantsRelative = false;
 	LastCheckedAllCollideableDescendantsTime = 0.f;
+
+	PhysicsMobility = Mobility;
 }
 
 bool UPrimitiveComponent::UsesOnlyUnlitMaterials() const
