[INTEGRATE] Change 2428748 by [REDACTED_PERSON][REDACTED_HANDLE] on 2015/02/02 17:04:30

	Fixed being able to move components created in a blueprint.

	Introduced an FSelectedEditableComponentIterator to iterate over only the movable components

diff --git a/Engine/Source/Editor/LevelEditor/Private/LevelEditorActions.cpp b/Engine/Source/Editor/LevelEditor/Private/LevelEditorActions.cpp
--- a/Engine/Source/Editor/LevelEditor/Private/LevelEditorActions.cpp
+++ b/Engine/Source/Editor/LevelEditor/Private/LevelEditorActions.cpp
@@ -2472,7 +2472,7 @@ void FLevelEditorActionCallbacks::SnapTo_Clicked( const bool InAlign, const bool
 	bool bSnappedComponents = false;
 	if( GEditor->GetSelectedComponentCount() > 0 )
 	{
-		for(FSelectionIterator It(GEditor->GetSelectedComponentIterator()); It; ++It)
+		for(FSelectedEditableComponentIterator It(GEditor->GetSelectedEditableComponentIterator()); It; ++It)
 		{
 			USceneComponent* SceneComponent = Cast<USceneComponent>(*It);
 			if(SceneComponent)
diff --git a/Engine/Source/Editor/UnrealEd/Classes/Editor/EditorEngine.h b/Engine/Source/Editor/UnrealEd/Classes/Editor/EditorEngine.h
--- a/Engine/Source/Editor/UnrealEd/Classes/Editor/EditorEngine.h
+++ b/Engine/Source/Editor/UnrealEd/Classes/Editor/EditorEngine.h
@@ -705,7 +705,6 @@ public:
 	bool	HandleMapCommand( const TCHAR* Str, FOutputDevice& Ar, UWorld* InWorld );
 	bool	HandleSelectCommand( const TCHAR* Str, FOutputDevice& Ar, UWorld* InWorld );
 	bool	HandleDeleteCommand( const TCHAR* Str, FOutputDevice& Ar, UWorld* InWorld );
-	bool	HandlePrefabCommand( const TCHAR* Str, FOutputDevice& Ar );
 	bool	HandleLightmassDebugCommand( const TCHAR* Str, FOutputDevice& Ar );
 	bool	HandleLightmassStatsCommand( const TCHAR* Str, FOutputDevice& Ar );
 	bool	HandleSwarmDistributionCommand( const TCHAR* Str, FOutputDevice& Ar );
@@ -1265,15 +1264,6 @@ public:
 	 */
 	virtual bool ShouldAbortActorDeletion() const { return false; }
 
-	/** Create a prefab from the selected actors, and replace those actors with an instance of that prefab. */
-	virtual void edactPrefabSelected() {};
-
-	/** Add the selected prefab at the clicked location. */
-	virtual void edactAddPrefab() {};
-
-	/** Select all Actors that make up the selected PrefabInstance. */
-	virtual void edactSelectPrefabActors() {};
-
 	/**
 	*
 	* Rebuild the level's Bsp from the level's CSG brushes.
@@ -1761,6 +1751,8 @@ public:
 	*/
 	class FSelectionIterator GetSelectedComponentIterator() const;
 
+	class FSelectedEditableComponentIterator GetSelectedEditableComponentIterator() const;
+
 	/**
 	* Returns the number of currently selected components.
 	*/
diff --git a/Engine/Source/Editor/UnrealEd/Private/EdMode.cpp b/Engine/Source/Editor/UnrealEd/Private/EdMode.cpp
--- a/Engine/Source/Editor/UnrealEd/Private/EdMode.cpp
+++ b/Engine/Source/Editor/UnrealEd/Private/EdMode.cpp
@@ -435,7 +435,7 @@ bool FEdMode::ShouldDrawWidget() const
 	bool bDrawWidget = false;
 	if (GEditor->GetSelectedComponentCount() > 0)
 	{
-		for (FSelectionIterator It(GEditor->GetSelectedComponentIterator()); It; ++It)
+		for (FSelectedEditableComponentIterator It(GEditor->GetSelectedEditableComponentIterator()); It; ++It)
 		{
 			if (It->IsA<USceneComponent>())
 			{
diff --git a/Engine/Source/Editor/UnrealEd/Private/Editor.cpp b/Engine/Source/Editor/UnrealEd/Private/Editor.cpp
--- a/Engine/Source/Editor/UnrealEd/Private/Editor.cpp
+++ b/Engine/Source/Editor/UnrealEd/Private/Editor.cpp
@@ -336,6 +336,11 @@ FSelectionIterator UEditorEngine::GetSelectedComponentIterator() const
 	return FSelectionIterator(*GetSelectedComponents());
 };
 
+FSelectedEditableComponentIterator UEditorEngine::GetSelectedEditableComponentIterator() const
+{
+	return FSelectedEditableComponentIterator(*GetSelectedComponents());
+}
+
 USelection* UEditorEngine::GetSelectedComponents() const
 {
 	return PrivateGetSelectedComponents();
diff --git a/Engine/Source/Editor/UnrealEd/Private/EditorActor.cpp b/Engine/Source/Editor/UnrealEd/Private/EditorActor.cpp
--- a/Engine/Source/Editor/UnrealEd/Private/EditorActor.cpp
+++ b/Engine/Source/Editor/UnrealEd/Private/EditorActor.cpp
@@ -554,7 +554,7 @@ bool UUnrealEdEngine::edactDeleteSelected( UWorld* InWorld, bool bVerifyDeletion
 	{
 		bool bComponentWasDeleted = false;
 		// Delete selected components if they are an Instance component
-		for (FSelectionIterator It(GetSelectedComponentIterator()); It; ++It)
+		for (FSelectedEditableComponentIterator It(GetSelectedEditableComponentIterator()); It; ++It)
 		{
 			UActorComponent* CurActorComp = CastChecked<UActorComponent>(*It);
 			if (CurActorComp->CreationMethod == EComponentCreationMethod::Instance)
diff --git a/Engine/Source/Editor/UnrealEd/Private/EditorSelectUtils.cpp b/Engine/Source/Editor/UnrealEd/Private/EditorSelectUtils.cpp
--- a/Engine/Source/Editor/UnrealEd/Private/EditorSelectUtils.cpp
+++ b/Engine/Source/Editor/UnrealEd/Private/EditorSelectUtils.cpp
@@ -277,7 +277,7 @@ void UUnrealEdEngine::UpdatePivotLocationForSelection( bool bOnChange )
 
 	if (GetSelectedComponentCount() > 0)
 	{
-		for (FSelectionIterator It(GetSelectedComponentIterator()); It; ++It)
+		for (FSelectedEditableComponentIterator It(*GetSelectedComponents()); It; ++It)
 		{
 			UActorComponent* Component = CastChecked<UActorComponent>(*It);
 			AActor* ComponentOwner = Component->GetOwner();
diff --git a/Engine/Source/Editor/UnrealEd/Private/EditorServer.cpp b/Engine/Source/Editor/UnrealEd/Private/EditorServer.cpp
--- a/Engine/Source/Editor/UnrealEd/Private/EditorServer.cpp
+++ b/Engine/Source/Editor/UnrealEd/Private/EditorServer.cpp
@@ -5254,13 +5254,6 @@ bool UEditorEngine::Exec( UWorld* InWorld, const TCHAR* Stream, FOutputDevice& A
 		return CommandIsDeprecated( CommandTemp, Ar );
 	}
 	//------------------------------------------------------------------------------------
-	// LEVEL
-	//
-	if( FParse::Command(&Str,TEXT("PREFAB")) )
-	{
-		return HandlePrefabCommand( Str, Ar );
-	}
-	//------------------------------------------------------------------------------------
 	// PARTICLE: Particle system-related commands
 	//
 	else if (FParse::Command(&Str,TEXT("PARTICLE")))
@@ -5705,16 +5698,6 @@ bool UEditorEngine::HandleDeleteCommand( const TCHAR* Str, FOutputDevice& Ar, UW
 	return true;
 }
 
-bool UEditorEngine::HandlePrefabCommand( const TCHAR* Str, FOutputDevice& Ar )
-{
-	if( FParse::Command(&Str,TEXT("SELECTACTORSINPREFABS")) )
-	{
-		edactSelectPrefabActors();
-		return true;
-	}
-	return false;	
-}
-
 bool UEditorEngine::HandleLightmassDebugCommand( const TCHAR* Str, FOutputDevice& Ar )
 {
 	extern UNREALED_API bool GLightmassDebugMode;
diff --git a/Engine/Source/Editor/UnrealEd/Private/LevelEditorViewport.cpp b/Engine/Source/Editor/UnrealEd/Private/LevelEditorViewport.cpp
--- a/Engine/Source/Editor/UnrealEd/Private/LevelEditorViewport.cpp
+++ b/Engine/Source/Editor/UnrealEd/Private/LevelEditorViewport.cpp
@@ -1257,7 +1257,7 @@ bool FLevelEditorViewportClient::DropObjectsAtCoordinates(int32 MouseX, int32 Mo
 						if (ComponentSelection->IsSelected(TargetComponent))
 						{
 							// The target component is selected, so try applying the object to every selected component
-							for (FSelectionIterator It(GEditor->GetSelectedComponentIterator()); It; ++It)
+							for (FSelectedEditableComponentIterator It(GEditor->GetSelectedEditableComponentIterator()); It; ++It)
 							{
 								auto SceneComponent = Cast<USceneComponent>(*It);
 								AttemptApplyObjToComponent(DroppedObjects[0], SceneComponent, TargetMaterialSlot, bCreateDropPreview);
@@ -3023,7 +3023,7 @@ void FLevelEditorViewportClient::ApplyDeltaToActors(const FVector& InDrag,
 				// Only move the parent-most component(s) that are selected 
 				// Otherwise, if both a parent and child are selected and the delta is applied to both, the child will actually move 2x delta
 				TInlineComponentArray<USceneComponent*> ComponentsToMove;
-				for (FSelectionIterator It(GEditor->GetSelectedComponentIterator()); It; ++It)
+				for (FSelectedEditableComponentIterator It(GEditor->GetSelectedEditableComponentIterator()); It; ++It)
 				{
 					USceneComponent* SceneComponent = CastChecked<USceneComponent>(*It);
 					if (SceneComponent)
diff --git a/Engine/Source/Runtime/Engine/Classes/Engine/Selection.h b/Engine/Source/Runtime/Engine/Classes/Engine/Selection.h
--- a/Engine/Source/Runtime/Engine/Classes/Engine/Selection.h
+++ b/Engine/Source/Runtime/Engine/Classes/Engine/Selection.h
@@ -16,7 +16,8 @@ private:
 	typedef TArray<TWeakObjectPtr<UObject> >	ObjectArray;
 	typedef TMRUArray<UClass*>	ClassArray;
 
-	friend class FSelectionIterator;
+	template<typename SelectionFilter>
+	friend class TSelectionIterator;
 
 public:
 	/** Params: UObject* NewSelection */
@@ -351,15 +352,28 @@ private:
 	}
 };
 
+
+/** A filter for generic selection sets.  Simply allows objects which are non-null */
+class FGenericSelectionFilter
+{
+public:
+	bool IsObjectValid( const UObject* InObject ) const
+	{
+		return InObject != nullptr;
+	}
+};
+
 /**
  * Manages selections of objects.  Used in the editor for selecting
  * objects in the various browser windows.
  */
-class FSelectionIterator
+template<typename SelectionFilter>
+class TSelectionIterator
 {
 public:
-	FSelectionIterator(USelection& InSelection)
+	TSelectionIterator(USelection& InSelection)
 		: Selection( InSelection )
+		, Filter( SelectionFilter() )
 	{
 		Reset();
 	}
@@ -424,7 +438,7 @@ private:
 
 	bool IsObjectValid() const
 	{
-		return GetCurrentObject() != NULL;
+		return Filter.IsObjectValid( GetCurrentObject() );
 	}
 
 	bool IsIndexValid() const
@@ -433,5 +447,37 @@ private:
 	}
 
 	USelection&	Selection;
+	SelectionFilter Filter;
 	int32			Index;
 };
+
+
+class FSelectionIterator : public TSelectionIterator<FGenericSelectionFilter>
+{
+public:
+	FSelectionIterator(USelection& InSelection)
+		: TSelectionIterator<FGenericSelectionFilter>( InSelection )
+	{}
+};
+
+/** A filter for only iterating through editable components */
+class FSelectedEditableComponentFilter
+{
+public:
+	bool IsObjectValid(const UObject* Object) const
+	{
+		const UActorComponent* Comp = Cast<UActorComponent>( Object );
+		return Comp && Comp->CreationMethod != EComponentCreationMethod::ConstructionScript;
+	}
+};
+
+/**
+ * An iterator used to iterate through selected components which can be editable (not created in a blueprint)
+ */
+class FSelectedEditableComponentIterator : public TSelectionIterator<FSelectedEditableComponentFilter>
+{
+public:
+	FSelectedEditableComponentIterator(USelection& InSelection)
+		: TSelectionIterator<FSelectedEditableComponentFilter>(InSelection)
+	{}
+};
