Network demo metadata that was written out at the end of recording is now written to a separate file to facilitate streaming the main demo file. Demo version bumped!
The DemoNetDriver is resilient to missing metadata files, and will still attempt to play the stream.
Each demo is now stored in its own directory with its metadata file.
NullNetworkReplayStreaming now handles generating the full path to the demo files, instead of the demorec and demoplay command handlers.
Updated ShooterGame to work with new directory structure.

diff --git a/Engine/Source/Runtime/Engine/Classes/Engine/DemoNetDriver.h b/Engine/Source/Runtime/Engine/Classes/Engine/DemoNetDriver.h
--- a/Engine/Source/Runtime/Engine/Classes/Engine/DemoNetDriver.h
+++ b/Engine/Source/Runtime/Engine/Classes/Engine/DemoNetDriver.h
@@ -34,9 +34,6 @@ class UDemoNetDriver
 	/** Total number of frames in the demo */
 	int32 DemoTotalFrames;
 
-	/** during playback, set to offset of where the stream ends (we don't want to continue reading into the meta data section) */
-	int32 EndOfStreamOffset;
-
 	/** True if we're in the middle of recording a frame */
 	bool bIsRecordingDemoFrame;
 
diff --git a/Engine/Source/Runtime/Engine/Private/DemoNetDriver.cpp b/Engine/Source/Runtime/Engine/Private/DemoNetDriver.cpp
--- a/Engine/Source/Runtime/Engine/Private/DemoNetDriver.cpp
+++ b/Engine/Source/Runtime/Engine/Private/DemoNetDriver.cpp
@@ -43,7 +43,6 @@ bool UDemoNetDriver::InitBase( bool bInitAsClient, FNetworkNotify* InNotify, con
 		Time					= 0;
 		bIsRecordingDemoFrame	= false;
 		bDemoPlaybackDone		= false;
-		EndOfStreamOffset		= 0;
 
 		ResetDemoState();
 
@@ -74,8 +73,11 @@ FString UDemoNetDriver::LowLevelGetNetworkNumber()
 	return FString( TEXT( "" ) );
 }
 
-#define NETWORK_DEMO_MAGIC			( 0x2CF5A13D )
-#define NETWORK_DEMO_VERSION		( 0 )
+static const uint32 NETWORK_DEMO_MAGIC				= 0x2CF5A13D;
+static const uint32 NETWORK_DEMO_VERSION			= 1;
+
+static const uint32 NETWORK_DEMO_METADATA_MAGIC		= 0x3D06B24E;
+static const uint32 NETWORK_DEMO_METADATA_VERSION	= 0;
 
 struct FNetworkDemoHeader
 {
@@ -83,29 +85,47 @@ struct FNetworkDemoHeader
 	uint32	Version;				// Version number to detect version mismatches.
 	uint32	EngineNetVersion;		// Version of engine networking format
 	FString LevelName;				// Name of level loaded for demo
+	
+	FNetworkDemoHeader() : 
+		Magic( NETWORK_DEMO_MAGIC ), 
+		Version( NETWORK_DEMO_VERSION ),
+		EngineNetVersion( GEngineNetVersion )
+	{}
+
+	friend FArchive& operator << ( FArchive& Ar, FNetworkDemoHeader& Header )
+	{
+		Ar << Header.Magic;
+		Ar << Header.Version;
+		Ar << Header.LevelName;
+
+		return Ar;
+	}
+};
+
+struct FNetworkDemoMetadataHeader
+{
+	uint32	Magic;					// Magic to ensure we're opening the right file.
+	uint32	Version;				// Version number to detect version mismatches.
+	uint32	EngineNetVersion;		// Version of engine networking format
 	int32	NumFrames;				// Number of total frames in the demo
 	float	TotalTime;				// Number of total time in seconds in demo
-	int32	MetaDataOffset;			// Offset into the file where the meta data is stored for extra information that wasn't known at start of demo
 	int32	NumStreamingLevels;		// Number of streaming levels
 
-	FNetworkDemoHeader() : 
-		Magic( NETWORK_DEMO_MAGIC ), 
-		Version( NETWORK_DEMO_VERSION ),
+	FNetworkDemoMetadataHeader() : 
+		Magic( NETWORK_DEMO_METADATA_MAGIC ), 
+		Version( NETWORK_DEMO_METADATA_VERSION ),
 		EngineNetVersion( GEngineNetVersion ),
 		NumFrames( 0 ),
 		TotalTime( 0 ),
-		MetaDataOffset( 0 ),
 		NumStreamingLevels( 0 )
 	{}
-
-	friend FArchive& operator << ( FArchive& Ar, FNetworkDemoHeader& Header )
+	
+	friend FArchive& operator << ( FArchive& Ar, FNetworkDemoMetadataHeader& Header )
 	{
 		Ar << Header.Magic;
 		Ar << Header.Version;
-		Ar << Header.LevelName;
 		Ar << Header.NumFrames;
 		Ar << Header.TotalTime;
-		Ar << Header.MetaDataOffset;
 		Ar << Header.NumStreamingLevels;
 
 		return Ar;
@@ -197,11 +217,43 @@ bool UDemoNetDriver::InitConnectInternal( FString& Error )
 	// Create fake control channel
 	ServerConnection->CreateChannel( CHTYPE_Control, 1 );
 
-	DemoTotalFrames = DemoHeader.NumFrames;
-	DemoTotalTime	= DemoHeader.TotalTime;
+	// Attempt to read metadata if it exists
+	FArchive* MetadataAr = ReplayStreamer->GetMetadataArchive();
+
+	FNetworkDemoMetadataHeader MetadataHeader;
+
+	if ( MetadataAr != nullptr )
+	{
+		(*MetadataAr) << MetadataHeader;
 
-	UE_LOG( LogDemo, Log, TEXT( "Starting demo playback with demo. Filename: %s, Frames: %i, Version %i" ), *DemoFilename, DemoTotalFrames, DemoHeader.Version );
+		// Check metadata magic value
+		if ( MetadataHeader.Magic != NETWORK_DEMO_METADATA_MAGIC )
+		{
+			Error = FString( TEXT( "Demo metadata file is corrupt" ) );
+			UE_LOG( LogDemo, Error, TEXT( "UDemoNetDriver::InitConnect: %s" ), *Error );
+			GameInstance->HandleDemoPlaybackFailure( EDemoPlayFailure::Corrupt, Error );
+			return false;
+		}
 
+		// Check version
+		if ( MetadataHeader.Version != NETWORK_DEMO_METADATA_VERSION )
+		{
+			Error = FString( TEXT( "Demo metadata file version is incorrect" ) );
+			UE_LOG( LogDemo, Error, TEXT( "UDemoNetDriver::InitConnect: %s" ), *Error );
+			GameInstance->HandleDemoPlaybackFailure( EDemoPlayFailure::InvalidVersion, Error );
+			return false;
+		}
+
+		DemoTotalFrames = MetadataHeader.NumFrames;
+		DemoTotalTime	= MetadataHeader.TotalTime;
+
+		UE_LOG( LogDemo, Log, TEXT( "Starting demo playback with full demo and metadata. Filename: %s, Frames: %i, Version %i" ), *DemoFilename, DemoTotalFrames, DemoHeader.Version );
+	}
+	else
+	{
+		UE_LOG( LogDemo, Log, TEXT( "Starting demo playback with streaming demo, metadata file not found. Filename: %s, Version %i" ), *DemoFilename, DemoHeader.Version );
+	}
+	
 	// Bypass UDemoPendingNetLevel
 	FString LoadMapError;
 
@@ -239,14 +291,8 @@ bool UDemoNetDriver::InitConnectInternal( FString& Error )
 	WorldContext->World()->DemoNetDriver = this;
 	WorldContext->PendingNetGame = NULL;
 
-	// Remember where we are
-	const int32 OldPos = FileAr->Tell();
-
-	// Jump to meta data
-	FileAr->Seek( DemoHeader.MetaDataOffset );
-
-	// Read meta data
-	for ( int32 i = 0; i < DemoHeader.NumStreamingLevels; ++i )
+	// Read meta data, if it exists
+	for ( int32 i = 0; i < MetadataHeader.NumStreamingLevels; ++i )
 	{
 		ULevelStreamingKismet* StreamingLevel = static_cast<ULevelStreamingKismet*>(StaticConstructObject(ULevelStreamingKismet::StaticClass(), GetWorld(), NAME_None, RF_NoFlags, NULL ) );
 
@@ -259,9 +305,9 @@ bool UDemoNetDriver::InitConnectInternal( FString& Error )
 		FString PackageName;
 		FString PackageNameToLoad;
 
-		(*FileAr) << PackageName;
-		(*FileAr) << PackageNameToLoad;
-		(*FileAr) << StreamingLevel->LevelTransform;
+		(*MetadataAr) << PackageName;
+		(*MetadataAr) << PackageNameToLoad;
+		(*MetadataAr) << StreamingLevel->LevelTransform;
 
 		StreamingLevel->PackageNameToLoad = FName( *PackageNameToLoad );
 		StreamingLevel->SetWorldAssetByPackageName( FName( *PackageName ) );
@@ -271,12 +317,6 @@ bool UDemoNetDriver::InitConnectInternal( FString& Error )
 		UE_LOG( LogDemo, Log, TEXT( "  Loading streamingLevel: %s, %s" ), *PackageName, *PackageNameToLoad );
 	}
 
-	// Jump back to start of stream
-	FileAr->Seek( OldPos );
-
-	// Remember where the meta data is, this is where we must stop reading the demo stream
-	EndOfStreamOffset = DemoHeader.MetaDataOffset;
-
 	return true;
 }
 
@@ -319,10 +359,9 @@ bool UDemoNetDriver::InitListen( FNetworkNotify* InNotify, FURL& ListenURL, bool
 
 	FNetworkDemoHeader DemoHeader;
 
-	// NOTE - This must be the SAME string we write when we close the demo or it will throw the header off
 	DemoHeader.LevelName = World->GetCurrentLevel()->GetOutermost()->GetName();
 
-	// Write the initial header (a lot of the fields will be placeholder until we fill them in later)
+	// Write the header
 	(*FileAr) << DemoHeader;
 
 	// Spawn the demo recording spectator.
@@ -429,15 +468,15 @@ void UDemoNetDriver::StopDemo()
 
 	if ( !ServerConnection )
 	{
-		FArchive* FileAr = ReplayStreamer->GetStreamingArchive();
+		FArchive* MetadataAr = ReplayStreamer->GetMetadataArchive();
 
-		// Finish writing the header and other information that goes at the end
-		if ( FileAr != NULL && World != NULL )
+		// Finish writing the metadata
+		if ( MetadataAr != NULL && World != NULL )
 		{
 			DemoTotalFrames = DemoFrameNum;
 			DemoTotalTime	= DemoCurrentTime;
 
-			// Get the number of streaming levels so we can update the header with the correct info
+			// Get the number of streaming levels so we can update the metadata with the correct info
 			int32 NumStreamingLevels = 0;
 
 			for ( int32 i = 0; i < World->StreamingLevels.Num(); ++i )
@@ -448,23 +487,18 @@ void UDemoNetDriver::StopDemo()
 				}
 			}
 
-			// Make a new header with updated info
-			FNetworkDemoHeader DemoHeader;
+			// Make a header for the metadata
+			FNetworkDemoMetadataHeader DemoHeader;
 
-			DemoHeader.LevelName			= World->GetCurrentLevel()->GetOutermost()->GetName();
 			DemoHeader.NumFrames			= DemoTotalFrames;
 			DemoHeader.TotalTime			= DemoTotalTime;
-			DemoHeader.MetaDataOffset		= FileAr->Tell();
 			DemoHeader.NumStreamingLevels	= NumStreamingLevels;
 
 			// Seek to beginning
-			FileAr->Seek( 0 );
-
-			// Re-write header with new info
-			(*FileAr) << DemoHeader;
+			MetadataAr->Seek( 0 );
 
-			// Restore file position to end of stream
-			FileAr->Seek( DemoHeader.MetaDataOffset );
+			// Write header
+			(*MetadataAr) << DemoHeader;
 
 			//
 			// Write meta data
@@ -481,9 +515,9 @@ void UDemoNetDriver::StopDemo()
 
 					UE_LOG( LogDemo, Log, TEXT( "  StreamingLevel: %s, %s" ), *PackageName, *PackageNameToLoad );
 
-					(*FileAr) << PackageName;
-					(*FileAr) << PackageNameToLoad;
-					(*FileAr) << World->StreamingLevels[i]->LevelTransform;
+					(*MetadataAr) << PackageName;
+					(*MetadataAr) << PackageNameToLoad;
+					(*MetadataAr) << World->StreamingLevels[i]->LevelTransform;
 				}
 			}
 		}
@@ -649,7 +683,7 @@ bool UDemoNetDriver::ReadDemoFrame()
 		return false;
 	}
 
-	if ( FileAr->AtEnd() || FileAr->Tell() >= EndOfStreamOffset )
+	if ( FileAr->AtEnd() )
 	{
 		bDemoPlaybackDone = true;
 
diff --git a/Engine/Source/Runtime/Engine/Private/World.cpp b/Engine/Source/Runtime/Engine/Private/World.cpp
--- a/Engine/Source/Runtime/Engine/Private/World.cpp
+++ b/Engine/Source/Runtime/Engine/Private/World.cpp
@@ -2528,11 +2528,6 @@ bool UWorld::HandleLogActorCountsCommand( const TCHAR* Cmd, FOutputDevice& Ar, U
 
 bool UWorld::HandleDemoRecordCommand( const TCHAR* Cmd, FOutputDevice& Ar, UWorld* InWorld )
 {
-	// Attempt to make the dir if it doesn't exist.
-	const FString DemoDir = FPaths::GameSavedDir() + TEXT( "Demos" );
-
-	IFileManager::Get().MakeDirectory( *DemoDir, true );
-
 	FURL DemoURL;
 	FString DemoName;
 
@@ -2544,23 +2539,8 @@ bool UWorld::HandleDemoRecordCommand( const TCHAR* Cmd, FOutputDevice& Ar, UWorl
 
 	DemoName.ReplaceInline( TEXT( "%m" ), *GetMapName() );
 
-	int32 Year, Month, DayOfWeek, Day, Hour, Min, Sec, MSec;
-	FPlatformTime::SystemTime( Year, Month, DayOfWeek, Day, Hour, Min, Sec, MSec );
-
-	DemoName.ReplaceInline( TEXT( "%td" ), *FDateTime::Now().ToString() );
-	DemoName.ReplaceInline( TEXT( "%d" ), *FString::Printf( TEXT( "%i-%i-%i" ), Month, Day, Year ) );
-	DemoName.ReplaceInline( TEXT( "%t" ), *FString::Printf( TEXT( "%i" ), ( ( Hour * 3600 ) + ( Min * 60 ) + Sec ) * 1000 + MSec ) );
-	DemoName.ReplaceInline( TEXT( "%v" ), *FString::Printf( TEXT( "%i" ), GEngineVersion.GetChangelist() ) );
-
-	// replace bad characters with underscores
-	DemoName.ReplaceInline( TEXT( "\\" ),	TEXT( "_" ) );
-	DemoName.ReplaceInline( TEXT( "/" ),	TEXT( "_" ) );
-	DemoName.ReplaceInline( TEXT( "." ),	TEXT( "_" ) );
-	DemoName.ReplaceInline( TEXT( " " ),	TEXT( "_" ) );
-	DemoName.ReplaceInline( TEXT( "%" ),	TEXT( "_" ) );
-
 	// replace the current URL's map with a demo extension
-	DemoURL.Map = DemoDir + TEXT( "/" ) + DemoName + TEXT( ".demo" );
+	DemoURL.Map = DemoName;
 
 	DestroyDemoNetDriver();
 
@@ -2604,9 +2584,7 @@ bool UWorld::HandleDemoPlayCommand( const TCHAR* Cmd, FOutputDevice& Ar, UWorld*
 		FURL DemoURL;
 		UE_LOG( LogWorld, Log, TEXT( "Attempting to play demo %s" ), *Temp );
 
-		const FString DemoDir = FPaths::GameSavedDir() + TEXT( "Demos" );
-
-		DemoURL.Map = DemoDir + TEXT( "/" ) + Temp + TEXT( ".demo" );
+		DemoURL.Map = Temp;
 
 #if 0
 		UPendingNetGame * NewPendingNetGame = new UDemoPendingNetGame( FObjectInitializer(), DemoURL );
diff --git a/Engine/Source/Runtime/NetworkReplayStreaming/HttpNetworkReplayStreaming/Private/HttpNetworkReplayStreaming.cpp b/Engine/Source/Runtime/NetworkReplayStreaming/HttpNetworkReplayStreaming/Private/HttpNetworkReplayStreaming.cpp
--- a/Engine/Source/Runtime/NetworkReplayStreaming/HttpNetworkReplayStreaming/Private/HttpNetworkReplayStreaming.cpp
+++ b/Engine/Source/Runtime/NetworkReplayStreaming/HttpNetworkReplayStreaming/Private/HttpNetworkReplayStreaming.cpp
@@ -119,6 +119,11 @@ FArchive* FHttpNetworkReplayStreamer::GetStreamingArchive()
 	return ( Archive.IsSaving() || Archive.Buffer.Num() >  0 ) ? &Archive : NULL;
 }
 
+FArchive* FHttpNetworkReplayStreamer::GetMetadataArchive()
+{
+	return NULL;
+}
+
 void FHttpNetworkReplayStreamer::HttpRequestFinished( FHttpRequestPtr HttpRequest, FHttpResponsePtr HttpResponse, bool bSucceeded )
 {
 	if ( bSucceeded )
diff --git a/Engine/Source/Runtime/NetworkReplayStreaming/HttpNetworkReplayStreaming/Private/HttpNetworkReplayStreaming.h b/Engine/Source/Runtime/NetworkReplayStreaming/HttpNetworkReplayStreaming/Private/HttpNetworkReplayStreaming.h
--- a/Engine/Source/Runtime/NetworkReplayStreaming/HttpNetworkReplayStreaming/Private/HttpNetworkReplayStreaming.h
+++ b/Engine/Source/Runtime/NetworkReplayStreaming/HttpNetworkReplayStreaming/Private/HttpNetworkReplayStreaming.h
@@ -32,6 +32,7 @@ public:
 	virtual void StartStreaming( FString& StreamName, bool bRecord, const FOnStreamReadyDelegate& Delegate ) override;
 	virtual void StopStreaming() override;
 	virtual FArchive* GetStreamingArchive() override;
+	virtual FArchive* GetMetadataArchive() override;
 
 	void HttpRequestFinished( FHttpRequestPtr HttpRequest, FHttpResponsePtr HttpResponse, bool bSucceeded );
 
diff --git a/Engine/Source/Runtime/NetworkReplayStreaming/NetworkReplayStreaming/Public/NetworkReplayStreaming.h b/Engine/Source/Runtime/NetworkReplayStreaming/NetworkReplayStreaming/Public/NetworkReplayStreaming.h
--- a/Engine/Source/Runtime/NetworkReplayStreaming/NetworkReplayStreaming/Public/NetworkReplayStreaming.h
+++ b/Engine/Source/Runtime/NetworkReplayStreaming/NetworkReplayStreaming/Public/NetworkReplayStreaming.h
@@ -18,6 +18,7 @@ public:
 	virtual void StartStreaming( FString& StreamName, bool bRecord, const FOnStreamReadyDelegate& Delegate ) = 0;
 	virtual void StopStreaming() = 0;
 	virtual FArchive* GetStreamingArchive() = 0;
+	virtual FArchive* GetMetadataArchive() = 0;
 };
 
 /** Replay streamer factory */
diff --git a/Engine/Source/Runtime/NetworkReplayStreaming/NullNetworkReplayStreaming/Private/NullNetworkReplayStreaming.cpp b/Engine/Source/Runtime/NetworkReplayStreaming/NullNetworkReplayStreaming/Private/NullNetworkReplayStreaming.cpp
--- a/Engine/Source/Runtime/NetworkReplayStreaming/NullNetworkReplayStreaming/Private/NullNetworkReplayStreaming.cpp
+++ b/Engine/Source/Runtime/NetworkReplayStreaming/NullNetworkReplayStreaming/Private/NullNetworkReplayStreaming.cpp
@@ -1,6 +1,8 @@
 // Copyright 1998-2015 Epic Games, Inc. All Rights Reserved.
 
 #include "NullNetworkReplayStreaming.h"
+#include "Paths.h"
+#include "EngineVersion.h"
 
 /**
  * Very basic implementation of network replay streaming using the file system
@@ -11,15 +13,45 @@
 
 void FNullNetworkReplayStreamer::StartStreaming( FString& StreamName, bool bRecord, const FOnStreamReadyDelegate& Delegate )
 {
+	int32 Year, Month, DayOfWeek, Day, Hour, Min, Sec, MSec;
+	FPlatformTime::SystemTime( Year, Month, DayOfWeek, Day, Hour, Min, Sec, MSec );
+
+	FString DemoName = StreamName;
+
+	DemoName.ReplaceInline( TEXT( "%td" ), *FDateTime::Now().ToString() );
+	DemoName.ReplaceInline( TEXT( "%d" ), *FString::Printf( TEXT( "%i-%i-%i" ), Month, Day, Year ) );
+	DemoName.ReplaceInline( TEXT( "%t" ), *FString::Printf( TEXT( "%i" ), ( ( Hour * 3600 ) + ( Min * 60 ) + Sec ) * 1000 + MSec ) );
+	DemoName.ReplaceInline( TEXT( "%v" ), *FString::Printf( TEXT( "%i" ), GEngineVersion.GetChangelist() ) );
+
+	// replace bad characters with underscores
+	DemoName.ReplaceInline( TEXT( "\\" ),	TEXT( "_" ) );
+	DemoName.ReplaceInline( TEXT( "/" ),	TEXT( "_" ) );
+	DemoName.ReplaceInline( TEXT( "." ),	TEXT( "_" ) );
+	DemoName.ReplaceInline( TEXT( " " ),	TEXT( "_" ) );
+	DemoName.ReplaceInline( TEXT( "%" ),	TEXT( "_" ) );
+
+	// Create a directory for this demo
+	const FString DemoDir = FPaths::Combine(*FPaths::GameSavedDir(), TEXT( "Demos" ), *DemoName);
+
+	IFileManager::Get().MakeDirectory( *DemoDir, true );
+
+	// Demo filename without extension
+	const FString FullDemoBaseName = DemoDir + TEXT("/") + DemoName;
+
+	const FString FullDemoFilename = FullDemoBaseName + TEXT(".demo");
+	const FString FullMetadataFilename = FullDemoBaseName + TEXT(".metadata");
+
 	if ( !bRecord )
 	{
-		// Opne file for reading
-		FileAr = IFileManager::Get().CreateFileReader( *StreamName );
+		// Open file for reading
+		FileAr = IFileManager::Get().CreateFileReader( *FullDemoFilename );
+		MetadataFileAr = IFileManager::Get().CreateFileReader( *FullMetadataFilename );
 	}
 	else
 	{
 		// Open file for writing
-		FileAr = IFileManager::Get().CreateFileWriter( *StreamName );
+		FileAr = IFileManager::Get().CreateFileWriter( *FullDemoFilename );
+		MetadataFileAr = IFileManager::Get().CreateFileWriter( *FullMetadataFilename );
 	}
 
 	// Notify immediately
@@ -30,6 +62,9 @@ void FNullNetworkReplayStreamer::StopStreaming()
 {
 	delete FileAr;
 	FileAr = NULL;
+
+	delete MetadataFileAr;
+	MetadataFileAr = NULL;
 }
 
 FArchive* FNullNetworkReplayStreamer::GetStreamingArchive()
@@ -37,6 +72,11 @@ FArchive* FNullNetworkReplayStreamer::GetStreamingArchive()
 	return FileAr;
 }
 
+FArchive* FNullNetworkReplayStreamer::GetMetadataArchive()
+{
+	return MetadataFileAr;
+}
+
 IMPLEMENT_MODULE( FNullNetworkReplayStreamingFactory, NullNetworkReplayStreaming )
 
 TSharedPtr< INetworkReplayStreamer > FNullNetworkReplayStreamingFactory::CreateReplayStreamer() 
diff --git a/Engine/Source/Runtime/NetworkReplayStreaming/NullNetworkReplayStreaming/Public/NullNetworkReplayStreaming.h b/Engine/Source/Runtime/NetworkReplayStreaming/NullNetworkReplayStreaming/Public/NullNetworkReplayStreaming.h
--- a/Engine/Source/Runtime/NetworkReplayStreaming/NullNetworkReplayStreaming/Public/NullNetworkReplayStreaming.h
+++ b/Engine/Source/Runtime/NetworkReplayStreaming/NullNetworkReplayStreaming/Public/NullNetworkReplayStreaming.h
@@ -10,14 +10,18 @@
 class FNullNetworkReplayStreamer : public INetworkReplayStreamer
 {
 public:
-	FNullNetworkReplayStreamer() : FileAr( NULL ) {}
+	FNullNetworkReplayStreamer() : FileAr( NULL ), MetadataFileAr( NULL ) {}
 	virtual void StartStreaming( FString& StreamName, bool bRecord, const FOnStreamReadyDelegate& Delegate ) override;
 	virtual void StopStreaming() override;
 	virtual FArchive* GetStreamingArchive() override;
+	virtual FArchive* GetMetadataArchive() override;
 
 private:
 	/** Handle to the archive that will read/write network packets */
 	FArchive* FileAr;
+
+	/* Handle to the archive that will read/write metadata */
+	FArchive* MetadataFileAr;
 };
 
 class FNullNetworkReplayStreamingFactory : public INetworkReplayStreamingFactory
