#ANIM: Add Animation Compression Setting to project

diff --git a/Engine/Config/BaseEngine.ini b/Engine/Config/BaseEngine.ini
--- a/Engine/Config/BaseEngine.ini
+++ b/Engine/Config/BaseEngine.ini
@@ -1501,21 +1501,6 @@ bRecompileOnLoad=true
 // no good default blueprint, it's up to the platform .ini to override if desired (see WindowsEngine.ini)
 ConfigClass=
 
-[AnimationCompression]
-CompressCommandletVersion=2	// Bump this up to trigger full recompression. Otherwise only new animations imported will be recompressed.
-DefaultCompressionAlgorithm=AnimCompress_RemoveLinearKeys
-TranslationCompressionFormat=0
-RotationCompressionFormat=1
-AlternativeCompressionThreshold=1.f
-ForceRecompression=False
-bOnlyCheckForMissingSkeletalMeshes=False
-+KeyEndEffectorsMatchName=IK
-+KeyEndEffectorsMatchName=eye
-+KeyEndEffectorsMatchName=weapon
-+KeyEndEffectorsMatchName=hand
-+KeyEndEffectorsMatchName=attach
-+KeyEndEffectorsMatchName=camera
-
 [CustomStats]
 +LD=Streaming fudge factor
 +LD=FrameTime
diff --git a/Engine/Source/Editor/ProjectSettingsViewer/Private/ProjectSettingsViewerModule.cpp b/Engine/Source/Editor/ProjectSettingsViewer/Private/ProjectSettingsViewerModule.cpp
--- a/Engine/Source/Editor/ProjectSettingsViewer/Private/ProjectSettingsViewerModule.cpp
+++ b/Engine/Source/Editor/ProjectSettingsViewer/Private/ProjectSettingsViewerModule.cpp
@@ -17,6 +17,7 @@
 #include "Runtime/Engine/Classes/Engine/UserInterfaceSettings.h"
 #include "Runtime/Engine/Classes/Sound/AudioSettings.h"
 #include "Runtime/AIModule/Classes/Navigation/CrowdManager.h"
+#include "Runtime/Engine/Classes/Animation/AnimationSettings.h"
 #include "AISystem.h"
 
 #define LOCTEXT_NAMESPACE "FProjectSettingsViewerModule"
@@ -158,6 +159,13 @@ protected:
 			GetMutableDefault<UCrowdManager>()
 			);
 
+		// Physics settings
+		SettingsModule.RegisterSettings("Project", "Engine", "Animation",
+			LOCTEXT("EngineAnimationSettingsName", "Animation"),
+			LOCTEXT("ProjectAnimationSettingsDescription", "Default animation settings."),
+			GetMutableDefault<UAnimationSettings>()
+		);
+
 		/*
 		// network settings
 		SettingsModule.RegisterSettings("Project", "Engine", "NetworkManager",
diff --git a/Engine/Source/Editor/UnrealEd/Private/Commandlets/PackageUtilities.cpp b/Engine/Source/Editor/UnrealEd/Private/Commandlets/PackageUtilities.cpp
--- a/Engine/Source/Editor/UnrealEd/Private/Commandlets/PackageUtilities.cpp
+++ b/Engine/Source/Editor/UnrealEd/Private/Commandlets/PackageUtilities.cpp
@@ -14,6 +14,7 @@
 
 #include "AnimationUtils.h"
 #include "AnimationCompression.h"
+#include "Animation/AnimationSettings.h"
 
 #include "CollectionManagerModule.h"
 #include "GameFramework/WorldSettings.h"
@@ -1510,9 +1511,7 @@ struct CompressAnimationsFunctor
 		}
 
 		// Get version number. Bump this up every time you want to recompress all animations.
-		int32 CompressCommandletVersion = 0;
-		GConfig->GetInt( TEXT("AnimationCompression"), TEXT("CompressCommandletVersion"), (int32&)CompressCommandletVersion, GEngineIni );
-
+		const int32 CompressCommandletVersion = UAnimationSettings::Get()->CompressCommandletVersion;
 
 		// Count the number of animations to provide some limited progress indication
 		int32 NumAnimationsInPackage = 0;
diff --git a/Engine/Source/Runtime/Engine/Classes/Animation/AnimCompress.h b/Engine/Source/Runtime/Engine/Classes/Animation/AnimCompress.h
--- a/Engine/Source/Runtime/Engine/Classes/Animation/AnimCompress.h
+++ b/Engine/Source/Runtime/Engine/Classes/Animation/AnimCompress.h
@@ -16,7 +16,7 @@ class UAnimCompress : public UObject
 	GENERATED_UCLASS_BODY()
 
 	/** A human-readable name for this modifier; appears in editor UI. */
-	UPROPERTY()
+	UPROPERTY(Category=Compression, VisibleAnywhere)
 	FString Description;
 
 	/** Compression algorithms requiring a skeleton should set this value to true. */
diff --git a/Engine/Source/Runtime/Engine/Classes/Animation/AnimSequence.h b/Engine/Source/Runtime/Engine/Classes/Animation/AnimSequence.h
--- a/Engine/Source/Runtime/Engine/Classes/Animation/AnimSequence.h
+++ b/Engine/Source/Runtime/Engine/Classes/Animation/AnimSequence.h
@@ -27,7 +27,7 @@ enum AnimationCompressionFormat
 	ACF_Fixed32NoW,
 	ACF_Float32NoW,
 	ACF_Identity,
-	ACF_MAX,
+	ACF_MAX UMETA(Hidden),
 };
 
 /**
diff --git a/Engine/Source/Runtime/Engine/Classes/Animation/AnimationSettings.h b/Engine/Source/Runtime/Engine/Classes/Animation/AnimationSettings.h
new file mode 100644
--- /dev/null
+++ b/Engine/Source/Runtime/Engine/Classes/Animation/AnimationSettings.h
@@ -0,0 +1,72 @@
+// Copyright 1998-2015 Epic Games, Inc. All Rights Reserved.
+
+/*=============================================================================
+	AnimationSettings.h: Declares the AnimationSettings class.
+=============================================================================*/
+
+#pragma once
+
+#include "AnimationSettings.generated.h"
+
+/**
+ * Implements project settings for the animation settings
+ */
+UCLASS(config=Engine, defaultconfig)
+class ENGINE_API UAnimationSettings
+	: public UObject
+{
+	GENERATED_UCLASS_BODY()
+
+	// compression upgrade version
+	UPROPERTY(config, VisibleAnywhere, Category = Compression)
+	int32 CompressCommandletVersion;
+
+	UPROPERTY(config, EditAnywhere, Category = Compression)
+	TArray<FString> KeyEndEffectorsMatchNameArray;
+
+	UPROPERTY(config, EditAnywhere, Category = Compression)
+	TSubclassOf<class UAnimCompress>  DefaultCompressionAlgorithm; 
+
+	UPROPERTY(config, EditAnywhere, Category = Compression)
+	TEnumAsByte<AnimationCompressionFormat> RotationCompressionFormat;
+
+	UPROPERTY(config, EditAnywhere, Category = Compression)
+	TEnumAsByte<AnimationCompressionFormat> TranslationCompressionFormat;
+
+	UPROPERTY(config, EditAnywhere, Category = Compression, meta=(ClampMin = "0", UIMin = "0"))
+	float AlternativeCompressionThreshold;
+
+	UPROPERTY(config, EditAnywhere, Category = Compression)
+	bool ForceRecompression;
+
+	UPROPERTY(config, EditAnywhere, Category = Compression)
+	bool bOnlyCheckForMissingSkeletalMeshes;
+	
+	UPROPERTY(config, EditAnywhere, Category = Compression)
+	bool bForceBelowThreshold;
+
+	UPROPERTY(config, EditAnywhere, Category = Compression)
+	bool bFirstRecompressUsingCurrentOrDefault;
+
+	UPROPERTY(config, EditAnywhere, Category = Compression)
+	bool bRaiseMaxErrorToExisting;
+
+	UPROPERTY(config, EditAnywhere, Category = Compression)
+	bool bTryFixedBitwiseCompression;
+
+	UPROPERTY(config, EditAnywhere, Category = Compression)
+	bool bTryPerTrackBitwiseCompression;
+
+	UPROPERTY(config, EditAnywhere, Category = Compression)
+	bool bTryLinearKeyRemovalCompression;
+
+	UPROPERTY(config, EditAnywhere, Category = Compression)
+	bool bTryIntervalKeyRemoval;
+
+public:
+	static UAnimationSettings * Get() { return CastChecked<UAnimationSettings>(UAnimationSettings::StaticClass()->GetDefaultObject()); }
+
+#if WITH_EDITOR
+	virtual void PostEditChangeProperty(struct FPropertyChangedEvent& PropertyChangedEvent) override;
+#endif // WITH_EDITOR
+};
diff --git a/Engine/Source/Runtime/Engine/Private/Animation/AnimationSettings.cpp b/Engine/Source/Runtime/Engine/Private/Animation/AnimationSettings.cpp
new file mode 100644
--- /dev/null
+++ b/Engine/Source/Runtime/Engine/Private/Animation/AnimationSettings.cpp
@@ -0,0 +1,51 @@
+// Copyright 1998-2015 Epic Games, Inc. All Rights Reserved.
+
+#include "EnginePrivate.h"
+#include "Animation/AnimationSettings.h" 
+#include "Animation/AnimCompress_BitwiseCompressOnly.h"
+
+UAnimationSettings::UAnimationSettings(const FObjectInitializer& ObjectInitializer)
+	: Super(ObjectInitializer)
+	, CompressCommandletVersion(2)	// Bump this up to trigger full recompression. Otherwise only new animations imported will be recompressed.
+	, DefaultCompressionAlgorithm(UAnimCompress_BitwiseCompressOnly::StaticClass())
+	, TranslationCompressionFormat(ACF_None)
+	, RotationCompressionFormat(ACF_Float96NoW)
+	, AlternativeCompressionThreshold(1.f)
+	, ForceRecompression(false)
+	, bOnlyCheckForMissingSkeletalMeshes(false)
+	, bForceBelowThreshold(false)
+	, bFirstRecompressUsingCurrentOrDefault(true)
+	, bRaiseMaxErrorToExisting(false)
+	, bTryFixedBitwiseCompression(true)
+	, bTryPerTrackBitwiseCompression(true)
+	, bTryLinearKeyRemovalCompression(true)
+	, bTryIntervalKeyRemoval(true)
+
+{
+	KeyEndEffectorsMatchNameArray.Add(TEXT("IK"));
+	KeyEndEffectorsMatchNameArray.Add(TEXT("eye"));
+	KeyEndEffectorsMatchNameArray.Add(TEXT("weapon"));
+	KeyEndEffectorsMatchNameArray.Add(TEXT("hand"));
+	KeyEndEffectorsMatchNameArray.Add(TEXT("attach"));
+	KeyEndEffectorsMatchNameArray.Add(TEXT("camera"));
+}
+
+
+#if WITH_EDITOR
+void UAnimationSettings::PostEditChangeProperty(struct FPropertyChangedEvent& PropertyChangedEvent)
+{
+	Super::PostEditChangeProperty(PropertyChangedEvent);
+
+// 	const FName PropertyName = PropertyChangedEvent.Property ? PropertyChangedEvent.Property->GetFName() : NAME_None;
+// 	if (PropertyName == GET_MEMBER_NAME_CHECKED(UAnimationSettings, FrictionCombineMode))
+// 	{
+// 		UPhysicalMaterial::RebuildPhysicalMaterials();
+// 	}
+// 	else if (PropertyName == GET_MEMBER_NAME_CHECKED(UAnimationSettings, LockedAxis))
+// 	{
+// 		UMovementComponent::PhysicsLockedAxisSettingChanged();
+// 	}
+}
+
+
+#endif	// WITH_EDITOR
\ No newline at end of file
diff --git a/Engine/Source/Runtime/Engine/Private/Animation/AnimationUtils.cpp b/Engine/Source/Runtime/Engine/Private/Animation/AnimationUtils.cpp
--- a/Engine/Source/Runtime/Engine/Private/Animation/AnimationUtils.cpp
+++ b/Engine/Source/Runtime/Engine/Private/Animation/AnimationUtils.cpp
@@ -10,6 +10,7 @@
 #include "Animation/AnimCompress_LeastDestructive.h"
 #include "Animation/AnimCompress_RemoveEverySecondKey.h"
 #include "Animation/AnimSet.h"
+#include "Animation/AnimationSettings.h"
 #include "AnimationUtils.h"
 #include "AnimationCompression.h"
 #include "AnimEncoding.h"
@@ -37,8 +38,7 @@ void FAnimationUtils::BuildSkeletonMetaData(USkeleton* Skeleton, TArray<FBoneDat
 	OutBoneData.Empty();
 	OutBoneData.AddZeroed( NumBones );
 
-	TArray<FString> KeyEndEffectorsMatchNameArray;
-	GConfig->GetArray( TEXT("AnimationCompression"), TEXT("KeyEndEffectorsMatchName"), KeyEndEffectorsMatchNameArray, GEngineIni );
+	const TArray<FString>& KeyEndEffectorsMatchNameArray = UAnimationSettings::Get()->KeyEndEffectorsMatchNameArray;
 
 	for (int32 BoneIndex = 0 ; BoneIndex<NumBones; ++BoneIndex )
 	{
@@ -367,38 +367,19 @@ namespace {
 static inline UAnimCompress* ConstructDefaultCompressionAlgorithm()
 {
 	// Algorithm.
-	FString DefaultCompressionAlgorithm( UAnimCompress_BitwiseCompressOnly::StaticClass()->GetName() );
-	GConfig->GetString( TEXT("AnimationCompression"), TEXT("DefaultCompressionAlgorithm"), DefaultCompressionAlgorithm, GEngineIni );
+	const UAnimationSettings* AnimSettings = UAnimationSettings::Get();
 
 	// Rotation compression format.
-	AnimationCompressionFormat RotationCompressionFormat = ACF_Float96NoW;
-	GConfig->GetInt( TEXT("AnimationCompression"), TEXT("RotationCompressionFormat"), (int32&)RotationCompressionFormat, GEngineIni );
-	RotationCompressionFormat = FMath::Clamp( RotationCompressionFormat, ACF_None, (AnimationCompressionFormat)(ACF_MAX-1) );
-
+	AnimationCompressionFormat RotationCompressionFormat = AnimSettings->RotationCompressionFormat;
 	// Translation compression format.
-	AnimationCompressionFormat TranslationCompressionFormat = ACF_None;
-	GConfig->GetInt( TEXT("AnimationCompression"), TEXT("TranslationCompressionFormat"), (int32&)TranslationCompressionFormat, GEngineIni );
-	TranslationCompressionFormat = FMath::Clamp( TranslationCompressionFormat, ACF_None, (AnimationCompressionFormat)(ACF_MAX-1) );
-
-	// Find a class that inherits
-	UClass* CompressionAlgorithmClass = NULL;
-	for ( TObjectIterator<UClass> It ; It ; ++It )
-	{
-		UClass* Class = *It;
-		if( !Class->HasAnyClassFlags(CLASS_Abstract | CLASS_Deprecated) )
-		{
-			if ( Class->IsChildOf(UAnimCompress::StaticClass()) && DefaultCompressionAlgorithm == Class->GetName() )
-			{
-				CompressionAlgorithmClass = Class;
-				break;
-
-			}
-		}
-	}
+	AnimationCompressionFormat TranslationCompressionFormat = AnimSettings->TranslationCompressionFormat;
 
+	UClass* CompressionAlgorithmClass = AnimSettings->DefaultCompressionAlgorithm;
 	if ( !CompressionAlgorithmClass )
 	{
-		UE_LOG(LogAnimation, Fatal, TEXT("Couldn't find animation compression algorithm named %s"), *DefaultCompressionAlgorithm );
+		// if can't find back out to bitwise
+		CompressionAlgorithmClass = UAnimCompress_BitwiseCompressOnly::StaticClass();
+		UE_LOG(LogAnimation, Warning, TEXT("Couldn't find animation compression, default to AnimCompress_BitwiseCompressOnly") );
 	}
 
 	UAnimCompress* NewAlgorithm = NewObject<UAnimCompress>(GetTransientPackage(), CompressionAlgorithmClass);
@@ -430,10 +411,7 @@ UAnimCompress* FAnimationUtils::GetDefaultAnimationCompressionAlgorithm()
 float FAnimationUtils::GetAlternativeCompressionThreshold()
 {
 	// Allow the Engine INI file to provide a new override
-	float AlternativeCompressionThreshold = 0.0f;
-	GConfig->GetFloat( TEXT("AnimationCompression"), TEXT("AlternativeCompressionThreshold"), (float&)AlternativeCompressionThreshold, GEngineIni );
-
-	return AlternativeCompressionThreshold;
+	return UAnimationSettings::Get()->AlternativeCompressionThreshold;
 }
 
 /**
@@ -571,35 +549,27 @@ void FAnimationUtils::CompressAnimSequence(UAnimSequence* AnimSeq,  bool bAllowA
 		// get the master tolerance we will use to guide recompression
 		float MasterTolerance = GetAlternativeCompressionThreshold(); 
 
-		bool bOnlyCheckForMissingSkeletalMeshes = false;
-		GConfig->GetBool( TEXT("AnimationCompression"), TEXT("bOnlyCheckForMissingSkeletalMeshes"), (bool&)bOnlyCheckForMissingSkeletalMeshes, GEngineIni );
-
+		bool bOnlyCheckForMissingSkeletalMeshes = UAnimationSettings::Get()->bOnlyCheckForMissingSkeletalMeshes;
 		if (bOnlyCheckForMissingSkeletalMeshes)
 		{
 			TestForMissingMeshes(AnimSeq);
 		}
 		else
 		{
-			bool bForceBelowThreshold = false;
-			bool bFirstRecompressUsingCurrentOrDefault = true;
-			bool bRaiseMaxErrorToExisting = false;
-			GConfig->GetBool( TEXT("AnimationCompression"), TEXT("bForceBelowThreshold"), (bool&)bForceBelowThreshold, GEngineIni );
-			GConfig->GetBool( TEXT("AnimationCompression"), TEXT("bFirstRecompressUsingCurrentOrDefault"), (bool&)bFirstRecompressUsingCurrentOrDefault, GEngineIni );
+			UAnimationSettings* AnimSetting = UAnimationSettings::Get();
+			bool bForceBelowThreshold = AnimSetting->bForceBelowThreshold;
+			bool bFirstRecompressUsingCurrentOrDefault = AnimSetting->bFirstRecompressUsingCurrentOrDefault;
+			bool bRaiseMaxErrorToExisting = AnimSetting->bRaiseMaxErrorToExisting;
 			// If we don't allow alternate compressors, and just want to recompress with default/existing, then make sure we do so.
 			if( !bAllowAlternateCompressor )
 			{
 				bFirstRecompressUsingCurrentOrDefault = true;
 			}
-			GConfig->GetBool( TEXT("AnimationCompression"), TEXT("bRaiseMaxErrorToExisting"), (bool&)bRaiseMaxErrorToExisting, GEngineIni );
 
-			bool bTryFixedBitwiseCompression = true;
-			bool bTryPerTrackBitwiseCompression = true;
-			bool bTryLinearKeyRemovalCompression = true;
-			bool bTryIntervalKeyRemoval = true;
-			GConfig->GetBool( TEXT("AnimationCompression"), TEXT("bTryFixedBitwiseCompression"), bTryFixedBitwiseCompression, GEngineIni );
-			GConfig->GetBool( TEXT("AnimationCompression"), TEXT("bTryPerTrackBitwiseCompression"), bTryPerTrackBitwiseCompression, GEngineIni );
-			GConfig->GetBool( TEXT("AnimationCompression"), TEXT("bTryLinearKeyRemovalCompression"), bTryLinearKeyRemovalCompression, GEngineIni );
-			GConfig->GetBool( TEXT("AnimationCompression"), TEXT("bTryIntervalKeyRemoval"), bTryIntervalKeyRemoval, GEngineIni );
+			bool bTryFixedBitwiseCompression = AnimSetting->bTryFixedBitwiseCompression;
+			bool bTryPerTrackBitwiseCompression = AnimSetting->bTryPerTrackBitwiseCompression;
+			bool bTryLinearKeyRemovalCompression = AnimSetting->bTryLinearKeyRemovalCompression;
+			bool bTryIntervalKeyRemoval = AnimSetting->bTryIntervalKeyRemoval;
 
 			CompressAnimSequenceExplicit(
 				AnimSeq,
@@ -632,13 +602,13 @@ void FAnimationUtils::CompressAnimSequenceExplicit(
 	UAnimSequence* AnimSeq,
 	float MasterTolerance,
 	bool bOutput,
-	bool bFirstRecompressUsingCurrentOrDefault,
-	bool bForceBelowThreshold,
-	bool bRaiseMaxErrorToExisting,
-	bool bTryFixedBitwiseCompression,
-	bool bTryPerTrackBitwiseCompression,
-	bool bTryLinearKeyRemovalCompression,
-	bool bTryIntervalKeyRemoval)
+	const bool bFirstRecompressUsingCurrentOrDefault,
+	const bool bForceBelowThreshold,
+	const bool bRaiseMaxErrorToExisting,
+	const bool bTryFixedBitwiseCompression,
+	const bool bTryPerTrackBitwiseCompression,
+	const bool bTryLinearKeyRemovalCompression,
+	const bool bTryIntervalKeyRemoval)
 {
 #if WITH_EDITORONLY_DATA
 	if( GDisableAnimationRecompression )
