change comment

diff --git a/Engine/Source/Runtime/Engine/Classes/Animation/AnimCompress.h b/Engine/Source/Runtime/Engine/Classes/Animation/AnimCompress.h
--- a/Engine/Source/Runtime/Engine/Classes/Animation/AnimCompress.h
+++ b/Engine/Source/Runtime/Engine/Classes/Animation/AnimCompress.h
@@ -15,7 +15,7 @@ class UAnimCompress : public UObject
 {
 	GENERATED_UCLASS_BODY()
 
-	/** A human-readable name for this modifier; appears in editor UI. */
+	/** Default Compression Scheme name */
 	UPROPERTY(Category=Compression, VisibleAnywhere)
 	FString Description;
 
