Fix comment again

diff --git a/Engine/Source/Runtime/Engine/Classes/Animation/AnimCompress.h b/Engine/Source/Runtime/Engine/Classes/Animation/AnimCompress.h
--- a/Engine/Source/Runtime/Engine/Classes/Animation/AnimCompress.h
+++ b/Engine/Source/Runtime/Engine/Classes/Animation/AnimCompress.h
@@ -15,7 +15,7 @@ class UAnimCompress : public UObject
 {
 	GENERATED_UCLASS_BODY()
 
-	/** Default Compression Scheme name */
+	/** Name of Compression Scheme used for this asset*/
 	UPROPERTY(Category=Compression, VisibleAnywhere)
 	FString Description;
 
