Light class comments

diff --git a/Engine/Source/Runtime/Engine/Public/SceneManagement.h b/Engine/Source/Runtime/Engine/Public/SceneManagement.h
--- a/Engine/Source/Runtime/Engine/Public/SceneManagement.h
+++ b/Engine/Source/Runtime/Engine/Public/SceneManagement.h
@@ -618,7 +618,11 @@ public:
 };
 
 
-/** Encapsulates the data which is used to render a light parallel to the game thread. */
+/** 
+ * Encapsulates the data which is used to render a light by the rendering thread. 
+ * The constructor is called from the game thread, and after that the rendering thread owns the object.
+ * FLightSceneProxy is in the engine module and is subclassed to implement various types of lights.
+ */
 class ENGINE_API FLightSceneProxy
 {
 public:
diff --git a/Engine/Source/Runtime/Renderer/Private/LightSceneInfo.h b/Engine/Source/Runtime/Renderer/Private/LightSceneInfo.h
--- a/Engine/Source/Runtime/Renderer/Private/LightSceneInfo.h
+++ b/Engine/Source/Runtime/Renderer/Private/LightSceneInfo.h
@@ -88,6 +88,7 @@ typedef TOctree<FLightSceneInfoCompact,struct FLightOctreeSemantics> FSceneLight
 
 /**
  * The information used to render a light.  This is the rendering thread's mirror of the game thread's ULightComponent.
+ * FLightSceneInfo is internal to the renderer module and contains internal scene state.
  */
 class FLightSceneInfo : public FRenderResource
 {
