[INTEGRATE] Change 2429226 by [REDACTED_PERSON][REDACTED_HANDLE] on 2015/02/02 22:24:26

	Rendering: Allow UTextRenderComponent to have multi-line text using actual \\n characters in addition to <br>, and annotated the property to be multi-line, allowing the \\n to be entered in the level editor on instances [UE-8833]

diff --git a/Engine/Source/Runtime/Engine/Classes/Components/TextRenderComponent.h b/Engine/Source/Runtime/Engine/Classes/Components/TextRenderComponent.h
--- a/Engine/Source/Runtime/Engine/Classes/Components/TextRenderComponent.h
+++ b/Engine/Source/Runtime/Engine/Classes/Components/TextRenderComponent.h
@@ -38,7 +38,7 @@ class ENGINE_API UTextRenderComponent : public UPrimitiveComponent
 	GENERATED_UCLASS_BODY()
 
 	/** Text content, can be multi line using <br> as line separator */
-	UPROPERTY(EditAnywhere, BlueprintReadOnly, Category=Text)
+	UPROPERTY(EditAnywhere, BlueprintReadOnly, Category=Text, meta=(MultiLine=true))
 	FString Text;
 
 	/** Text material */
diff --git a/Engine/Source/Runtime/Engine/Private/TextRenderComponent.cpp b/Engine/Source/Runtime/Engine/Private/TextRenderComponent.cpp
--- a/Engine/Source/Runtime/Engine/Private/TextRenderComponent.cpp
+++ b/Engine/Source/Runtime/Engine/Private/TextRenderComponent.cpp
@@ -73,6 +73,10 @@ struct FTextIterator
 		{
 			CurrentPosition += 4;
 		}
+		else if (CurrentPosition[0] == '\n')
+		{
+			++CurrentPosition;
+		}
 		else
 		{
 			Ch = *CurrentPosition;
@@ -86,7 +90,7 @@ struct FTextIterator
 	bool Peek(int32& Ch)
 	{
 		check(CurrentPosition);
-		if ( CurrentPosition[0] =='\0' || (CurrentPosition[0] == '<' && CurrentPosition[1] == 'b' && CurrentPosition[2] == 'r' && CurrentPosition[3] == '>'))
+		if ( CurrentPosition[0] =='\0' || (CurrentPosition[0] == '<' && CurrentPosition[1] == 'b' && CurrentPosition[2] == 'r' && CurrentPosition[3] == '>') || (CurrentPosition[0] == '\n') )
 		{
 			return false;
 		}
