[INTEGRATE] Change 2429228 by [REDACTED_PERSON][REDACTED_HANDLE] on 2015/02/02 22:25:32

	Blueprints: Allow Blueprint asset descriptions to be multi-line

diff --git a/Engine/Source/Runtime/Engine/Classes/Engine/Blueprint.h b/Engine/Source/Runtime/Engine/Classes/Engine/Blueprint.h
--- a/Engine/Source/Runtime/Engine/Classes/Engine/Blueprint.h
+++ b/Engine/Source/Runtime/Engine/Classes/Engine/Blueprint.h
@@ -303,7 +303,7 @@ class ENGINE_API UBlueprint : public UBlueprintCore
 	uint32 bGenerateConstClass : 1;
 
 	/**shows up in the content browser when the blueprint is hovered */
-	UPROPERTY(EditAnywhere, Category=BlueprintOptions)
+	UPROPERTY(EditAnywhere, Category=BlueprintOptions, meta=(MultiLine=true))
 	FString BlueprintDescription;
 
 	/** The category of the Blueprint, used to organize this Blueprint class when displayed in palette windows */
