Fix incorrect destructible defaults for ImpactDepthDamage and ImpactResistance.
Change default impact damage to be non-zero since we already have a flag turning the whole thing off by default

diff --git a/Engine/Source/Runtime/Engine/Classes/Engine/DestructibleMesh.h b/Engine/Source/Runtime/Engine/Classes/Engine/DestructibleMesh.h
--- a/Engine/Source/Runtime/Engine/Classes/Engine/DestructibleMesh.h
+++ b/Engine/Source/Runtime/Engine/Classes/Engine/DestructibleMesh.h
@@ -188,7 +188,7 @@ struct FDestructibleDamageParameters
 	bool bEnableImpactDamage;
 
 	/**
-		Controls how much damage is applied upon collision. Damage = ImpactDamage * ImpactForce. The default value is zero, which means impact damage is disabled.
+		Controls how much damage is applied upon collision. Damage = ImpactDamage * ImpactForce.
 		@see DepthParameters for per level control of ImpactDamage
 	*/
 	UPROPERTY(EditAnywhere, Category=DestructibleDamageParameters, meta = (editcondition = "bEnableImpactDamage") )
@@ -220,10 +220,10 @@ struct FDestructibleDamageParameters
 		: DamageThreshold(1.0f)
 		, DamageSpread(0.1f)
 		, bEnableImpactDamage(false)
-		, ImpactDamage(0.0f)
-		, DefaultImpactDamageDepth(-1)
+		, ImpactDamage(0.1f)
+		, DefaultImpactDamageDepth(0.f)
 		, bCustomImpactResistance(false)
-		, ImpactResistance(0.0f)
+		, ImpactResistance(1.f)
 	{ }
 
 #if WITH_APEX
