UE4 - support for hack UT fix relating to preview worlds in game

diff --git a/Engine/Source/Runtime/Engine/Classes/Engine/World.h b/Engine/Source/Runtime/Engine/Classes/Engine/World.h
--- a/Engine/Source/Runtime/Engine/Classes/Engine/World.h
+++ b/Engine/Source/Runtime/Engine/Classes/Engine/World.h
@@ -871,6 +871,9 @@ public:
 	/** If true, this is a preview world used for editor tools, and not an actual loaded map world */
 	TEnumAsByte<EWorldType::Type>				WorldType;
 
+	/** Force UsesGameHiddenFlags to return true. */
+	bool										bHack_Force_UsesGameHiddenFlags_True;
+
 	/** If true this world is in the process of running the construction script for an actor */
 	bool										bIsRunningConstructionScript;
 
@@ -2459,6 +2462,9 @@ public:
 	/** Returns true if this world is any kind of game world (including PIE worlds) */
 	bool IsGameWorld() const;
 
+	/** Returns true if this world should look at game hidden flags instead of editor hidden flags for the purposes of rendering */
+	bool UsesGameHiddenFlags() const;
+
 	// Return the URL of this level, which may possibly
 	// exist on a remote machine.
 	virtual FString GetAddressURL() const;
diff --git a/Engine/Source/Runtime/Engine/Private/SceneComponent.cpp b/Engine/Source/Runtime/Engine/Private/SceneComponent.cpp
--- a/Engine/Source/Runtime/Engine/Private/SceneComponent.cpp
+++ b/Engine/Source/Runtime/Engine/Private/SceneComponent.cpp
@@ -1653,12 +1653,12 @@ bool USceneComponent::ShouldRender() const
 	AActor* Owner = GetOwner();
 	const bool bShowInEditor = 
 #if WITH_EDITOR
-		(!Owner || !Owner->IsHiddenEd());
+		GIsEditor ? (!Owner || !Owner->IsHiddenEd()) : false;
 #else
 		false;
 #endif
-
-	const bool bInGameWorld = GetWorld() && GetWorld()->IsGameWorld();
+	UWorld *World = GetWorld();
+	const bool bInGameWorld = GetWorld() && GetWorld()->UsesGameHiddenFlags();
 
 	const bool bShowInGame = IsVisible() && (!Owner || !Owner->bHidden);
 	return ((bInGameWorld && bShowInGame) || (!bInGameWorld && bShowInEditor)) && bVisible == true;
@@ -1669,12 +1669,12 @@ bool USceneComponent::CanEverRender() const
 	AActor* Owner = GetOwner();
 	const bool bShowInEditor =
 #if WITH_EDITOR
-		(!Owner || !Owner->IsHiddenEd());
+		GIsEditor ? (!Owner || !Owner->IsHiddenEd()) : false;
 #else
 		false;
 #endif
-
-	const bool bInGameWorld = GetWorld() && GetWorld()->IsGameWorld();
+	UWorld *World = GetWorld();
+	const bool bInGameWorld = World && World->UsesGameHiddenFlags();
 
 	const bool bShowInGame = (!Owner || !Owner->bHidden);
 	return ((bInGameWorld && bShowInGame) || (!bInGameWorld && bShowInEditor));
diff --git a/Engine/Source/Runtime/Engine/Private/World.cpp b/Engine/Source/Runtime/Engine/Private/World.cpp
--- a/Engine/Source/Runtime/Engine/Private/World.cpp
+++ b/Engine/Source/Runtime/Engine/Private/World.cpp
@@ -4982,6 +4982,11 @@ bool UWorld::IsGameWorld() const
 	return WorldType == EWorldType::Game || WorldType == EWorldType::PIE;
 }
 
+bool UWorld::UsesGameHiddenFlags() const
+{
+	return IsGameWorld() || bHack_Force_UsesGameHiddenFlags_True;
+}
+
 FString UWorld::GetAddressURL() const
 {
 	return FString::Printf( TEXT("%s:%i"), *URL.Host, URL.Port );
