Fix destructibles ignoring damage parameters that are imported from APB

diff --git a/Engine/Source/Runtime/Engine/Private/DestructibleMesh.cpp b/Engine/Source/Runtime/Engine/Private/DestructibleMesh.cpp
--- a/Engine/Source/Runtime/Engine/Private/DestructibleMesh.cpp
+++ b/Engine/Source/Runtime/Engine/Private/DestructibleMesh.cpp
@@ -457,7 +457,7 @@ void UDestructibleMesh::LoadDefaultDestructibleParametersFromApexAsset()
 		DefaultDestructibleParameters.DebrisParameters.LoadDefaultDestructibleParametersFromApexAsset(Params);
 		DefaultDestructibleParameters.SpecialHierarchyDepths.LoadDefaultDestructibleParametersFromApexAsset(Params);
 		DefaultDestructibleParameters.AdvancedParameters.LoadDefaultDestructibleParametersFromApexAsset(Params);
-		DefaultDestructibleParameters.DebrisParameters.LoadDefaultDestructibleParametersFromApexAsset(Params);
+		DefaultDestructibleParameters.DamageParameters.LoadDefaultDestructibleParametersFromApexAsset(Params);
 		DefaultDestructibleParameters.Flags.LoadDefaultDestructibleParametersFromApexAsset(Params);
 
 		// Depth parameters
