Fixed non-unity

diff --git a/Engine/Source/Runtime/Engine/Classes/Animation/AnimationSettings.h b/Engine/Source/Runtime/Engine/Classes/Animation/AnimationSettings.h
--- a/Engine/Source/Runtime/Engine/Classes/Animation/AnimationSettings.h
+++ b/Engine/Source/Runtime/Engine/Classes/Animation/AnimationSettings.h
@@ -6,6 +6,7 @@
 
 #pragma once
 
+#include "AnimSequence.h"
 #include "AnimationSettings.generated.h"
 
 /**
