Expose restitution combine mode and BounceThresholdVelocity

diff --git a/Engine/Source/Runtime/Engine/Classes/PhysicalMaterials/PhysicalMaterial.h b/Engine/Source/Runtime/Engine/Classes/PhysicalMaterials/PhysicalMaterial.h
--- a/Engine/Source/Runtime/Engine/Classes/PhysicalMaterials/PhysicalMaterial.h
+++ b/Engine/Source/Runtime/Engine/Classes/PhysicalMaterials/PhysicalMaterial.h
@@ -58,6 +58,14 @@ class ENGINE_API UPhysicalMaterial : public UObject
 	UPROPERTY(EditAnywhere, BlueprintReadOnly, Category = PhysicalMaterial, meta = (UIMin=0, UIMax=1, ClampMin=0, ClampMax=1))
 	float Restitution;
 
+	/** Restitution combine mode, controls how restitution is computed for multiple materials. */
+	UPROPERTY(EditAnywhere, BlueprintReadOnly, Category = PhysicalMaterial, meta = (editcondition = "bOverrideRestitutionCombineMode"))
+	TEnumAsByte<EFrictionCombineMode::Type> RestitutionCombineMode;
+
+	/** If set we will use the RestitutionCombineMode of this material, instead of the RestitutionCombineMode found in the project settings. */
+	UPROPERTY(EditAnywhere, BlueprintReadWrite, Category = PhysicalMaterial)
+	bool bOverrideRestitutionCombineMode;
+
 	//
 	// Object properties.
 	//
diff --git a/Engine/Source/Runtime/Engine/Classes/PhysicsEngine/PhysicsSettings.h b/Engine/Source/Runtime/Engine/Classes/PhysicsEngine/PhysicsSettings.h
--- a/Engine/Source/Runtime/Engine/Classes/PhysicsEngine/PhysicsSettings.h
+++ b/Engine/Source/Runtime/Engine/Classes/PhysicsEngine/PhysicsSettings.h
@@ -37,13 +37,13 @@ namespace EFrictionCombineMode
 {
 	enum Type
 	{
-		//Uses the average friction of materials touching: (a+b) /2
+		//Uses the average value of the materials touching: (a+b) /2
 		Average = 0,	
-		//Uses the minimum friction of materials touching: min(a,b)
+		//Uses the minimum value of the materials touching: min(a,b)
 		Min = 1,		
-		//Uses the product of the friction of materials touching: a*b
+		//Uses the product of the values of the materials touching: a*b
 		Multiply = 2,	
-		//Uses the maximum friction of materials touching: max(a,b)
+		//Uses the maximum value of materials touching: max(a,b)
 		Max = 3
 	};
 }
@@ -106,11 +106,18 @@ class ENGINE_API UPhysicsSettings
 	UPROPERTY(config, EditAnywhere, Category = Simulation)
 	TEnumAsByte<ESettingsLockedAxis::Type> LockedAxis;
 	
+	/** Minimum relative velocity required for an object to bounce. A typical value for simulation stability is about 0.2 * gravity*/
+	UPROPERTY(config, EditAnywhere, Category = Simulation, meta = (ClampMin = "0", UIMin = "0"))
+	float BounceThresholdVelocity;
 
 	/** Friction combine mode, controls how friction is computed for multiple materials. */
 	UPROPERTY(config, EditAnywhere, Category=Simulation)
 	TEnumAsByte<EFrictionCombineMode::Type> FrictionCombineMode;
 
+	/** Restitution combine mode, controls how restitution is computed for multiple materials. */
+	UPROPERTY(config, EditAnywhere, Category = Simulation)
+	TEnumAsByte<EFrictionCombineMode::Type> RestitutionCombineMode;
+
 	/** Max velocity which may be used to depenetrate simulated physics objects. 0 means no maximum. */
 	UPROPERTY(config, EditAnywhere, Category = Simulation)
 	float MaxDepenetrationVelocity;
diff --git a/Engine/Source/Runtime/Engine/Private/PhysicsEngine/PhysScene.cpp b/Engine/Source/Runtime/Engine/Private/PhysicsEngine/PhysScene.cpp
--- a/Engine/Source/Runtime/Engine/Private/PhysicsEngine/PhysScene.cpp
+++ b/Engine/Source/Runtime/Engine/Private/PhysicsEngine/PhysScene.cpp
@@ -27,13 +27,6 @@
 static int32 PhysXSceneCount = 1;
 static const int PhysXSlowRebuildRate = 10;
 
-// Cvars
-static TAutoConsoleVariable<float> CVarBounceThresholdVelocity(
-	TEXT("p.BounceThresholdVelocity"),
-	200.f,
-	TEXT("A contact with a relative velocity below this will not bounce. Default: 200"),
-	ECVF_Default);
-
 FORCEINLINE EPhysicsSceneType SceneType(const FBodyInstance* BodyInstance)
 {
 #if WITH_PHYSX
@@ -1206,7 +1199,7 @@ void FPhysScene::InitPhysScene(uint32 SceneType)
 	PSceneDesc.flags |= PxSceneFlag::eENABLE_KINEMATIC_PAIRS;
 
 	// Set bounce threshold
-	PSceneDesc.bounceThresholdVelocity = CVarBounceThresholdVelocity.GetValueOnGameThread();
+	PSceneDesc.bounceThresholdVelocity = UPhysicsSettings::Get()->BounceThresholdVelocity;
 
 	// Possibly set flags in async scene for better behavior with piles
 #if USE_ADAPTIVE_FORCES_FOR_ASYNC_SCENE
diff --git a/Engine/Source/Runtime/Engine/Private/PhysicsEngine/PhysicalMaterial.cpp b/Engine/Source/Runtime/Engine/Private/PhysicsEngine/PhysicalMaterial.cpp
--- a/Engine/Source/Runtime/Engine/Private/PhysicsEngine/PhysicalMaterial.cpp
+++ b/Engine/Source/Runtime/Engine/Private/PhysicsEngine/PhysicalMaterial.cpp
@@ -93,9 +93,13 @@ void UPhysicalMaterial::UpdatePhysXMaterial()
 	{
 		PMaterial->setStaticFriction(Friction);
 		PMaterial->setDynamicFriction(Friction);
+		const uint32 UseFrictionCombineMode = (bOverrideFrictionCombineMode ? FrictionCombineMode.GetValue() : UPhysicsSettings::Get()->FrictionCombineMode.GetValue());
+		PMaterial->setFrictionCombineMode(static_cast<physx::PxCombineMode::Enum>(UseFrictionCombineMode));
+
 		PMaterial->setRestitution(Restitution);
-		uint32 CombineMode = (bOverrideFrictionCombineMode ? FrictionCombineMode.GetValue() : UPhysicsSettings::Get()->FrictionCombineMode.GetValue());
-		PMaterial->setFrictionCombineMode(static_cast<physx::PxCombineMode::Enum>(CombineMode));
+		const uint32 UseRestitutionCombineMode = (bOverrideRestitutionCombineMode ? RestitutionCombineMode.GetValue() : UPhysicsSettings::Get()->RestitutionCombineMode.GetValue());
+		PMaterial->setRestitutionCombineMode(static_cast<physx::PxCombineMode::Enum>(UseRestitutionCombineMode));
+
 #if WITH_VEHICLE
 		FPhysXVehicleManager::UpdateTireFrictionTable();
 #endif
@@ -110,8 +114,12 @@ PxMaterial* UPhysicalMaterial::GetPhysXMaterial()
 	if((PMaterial == NULL) && GPhysXSDK)
 	{
 		PMaterial = GPhysXSDK->createMaterial(Friction, Friction, Restitution);
-		uint32 CombineMode = (bOverrideFrictionCombineMode ? FrictionCombineMode.GetValue() : UPhysicsSettings::Get()->FrictionCombineMode.GetValue());
-		PMaterial->setFrictionCombineMode(static_cast<physx::PxCombineMode::Enum>(CombineMode));
+		const uint32 UseFrictionCombineMode = (bOverrideFrictionCombineMode ? FrictionCombineMode.GetValue() : UPhysicsSettings::Get()->FrictionCombineMode.GetValue());
+		PMaterial->setFrictionCombineMode(static_cast<physx::PxCombineMode::Enum>(UseFrictionCombineMode));
+		
+		const uint32 UseRestitutionCombineMode = (bOverrideRestitutionCombineMode ? RestitutionCombineMode.GetValue() : UPhysicsSettings::Get()->RestitutionCombineMode.GetValue());
+		PMaterial->setRestitutionCombineMode(static_cast<physx::PxCombineMode::Enum>(UseRestitutionCombineMode));
+
 		PMaterial->userData = &PhysxUserData;
 
 		UpdatePhysXMaterial();
diff --git a/Engine/Source/Runtime/Engine/Private/PhysicsEngine/PhysicsSettings.cpp b/Engine/Source/Runtime/Engine/Private/PhysicsEngine/PhysicsSettings.cpp
--- a/Engine/Source/Runtime/Engine/Private/PhysicsEngine/PhysicsSettings.cpp
+++ b/Engine/Source/Runtime/Engine/Private/PhysicsEngine/PhysicsSettings.cpp
@@ -14,6 +14,7 @@ UPhysicsSettings::UPhysicsSettings(const FObjectInitializer& ObjectInitializer)
 	, TriangleMeshTriangleMinAreaThreshold(0.25f)
 	, bEnableAsyncScene(false)
 	, bEnable2DPhysics(false)
+	, BounceThresholdVelocity(200.f)
 	, bSimulateSkeletalMeshOnDedicatedServer(true)
 	, MaxPhysicsDeltaTime(1.f / 30.f)
 	, bSubstepping(false)
@@ -55,7 +56,7 @@ void UPhysicsSettings::PostEditChangeProperty(struct FPropertyChangedEvent& Prop
 	Super::PostEditChangeProperty(PropertyChangedEvent);
 
 	const FName PropertyName = PropertyChangedEvent.Property ? PropertyChangedEvent.Property->GetFName() : NAME_None;
-	if (PropertyName == GET_MEMBER_NAME_CHECKED(UPhysicsSettings, FrictionCombineMode))
+	if (PropertyName == GET_MEMBER_NAME_CHECKED(UPhysicsSettings, FrictionCombineMode) || PropertyName == GET_MEMBER_NAME_CHECKED(UPhysicsSettings, RestitutionCombineMode))
 	{
 		UPhysicalMaterial::RebuildPhysicalMaterials();
 	}
