[INTEGRATE] Change 2429312 by [REDACTED_PERSON][REDACTED_HANDLE] on 2015/02/03 00:31:41

	Framework: Mark camera component proxy mesh as NotForClient and NotForServer (it is protected and is isn't constructed in non-editor builds and is invisible in-game already) [UE-8929]

diff --git a/Engine/Source/Runtime/Engine/Private/Camera/CameraComponent.cpp b/Engine/Source/Runtime/Engine/Private/Camera/CameraComponent.cpp
--- a/Engine/Source/Runtime/Engine/Private/Camera/CameraComponent.cpp
+++ b/Engine/Source/Runtime/Engine/Private/Camera/CameraComponent.cpp
@@ -62,6 +62,8 @@ void UCameraComponent::OnRegister()
 	{
 		ProxyMeshComponent = NewObject<UStaticMeshComponent>(GetOuter(), NAME_None, RF_Transactional);
 		ProxyMeshComponent->AttachTo(this);
+		ProxyMeshComponent->AlwaysLoadOnClient = false;
+		ProxyMeshComponent->AlwaysLoadOnServer = false;
 		ProxyMeshComponent->StaticMesh = CameraMesh;
 		ProxyMeshComponent->SetCollisionProfileName(UCollisionProfile::NoCollision_ProfileName);
 		ProxyMeshComponent->bHiddenInGame = true;
