[INTEGRATE] Change 2429920 by [REDACTED_PERSON][REDACTED_HANDLE] on 2015/02/03 10:16:06

	UE-8824 Crash closing packaged project

diff --git a/Engine/Source/Runtime/CoreUObject/Public/UObject/Class.h b/Engine/Source/Runtime/CoreUObject/Public/UObject/Class.h
--- a/Engine/Source/Runtime/CoreUObject/Public/UObject/Class.h
+++ b/Engine/Source/Runtime/CoreUObject/Public/UObject/Class.h
@@ -1980,7 +1980,7 @@ public:
 		return NULL;
 	}
 
-	virtual void CreatePersistentUberGraphFrame(UObject* Obj) const
+	virtual void CreatePersistentUberGraphFrame(UObject* Obj, bool bCreateOnlyIfEmpty = false) const
 	{
 	}
 
diff --git a/Engine/Source/Runtime/Engine/Classes/Engine/BlueprintGeneratedClass.h b/Engine/Source/Runtime/Engine/Classes/Engine/BlueprintGeneratedClass.h
--- a/Engine/Source/Runtime/Engine/Classes/Engine/BlueprintGeneratedClass.h
+++ b/Engine/Source/Runtime/Engine/Classes/Engine/BlueprintGeneratedClass.h
@@ -470,6 +470,7 @@ public:
 	virtual void CreateComponentsForActor(AActor* Actor) const;
 
 	// UObject interface
+	virtual void Serialize(FArchive& Ar) override;
 	virtual void PostLoad() override;
 	virtual void PostInitProperties() override;
 	// End UObject interface
@@ -481,7 +482,7 @@ public:
 #endif //WITH_EDITOR
 	virtual bool IsFunctionImplementedInBlueprint(FName InFunctionName) const override;
 	virtual uint8* GetPersistentUberGraphFrame(UObject* Obj, UFunction* FuncToCheck) const override;
-	virtual void CreatePersistentUberGraphFrame(UObject* Obj) const override;
+	virtual void CreatePersistentUberGraphFrame(UObject* Obj, bool bCreateOnlyIfEmpty = false) const override;
 	virtual void DestroyPersistentUberGraphFrame(UObject* Obj) const override;
 	virtual void Link(FArchive& Ar, bool bRelinkExistingProperties) override;
 	virtual void PurgeClass(bool bRecompilingOnLoad) override;
diff --git a/Engine/Source/Runtime/Engine/Private/BlueprintGeneratedClass.cpp b/Engine/Source/Runtime/Engine/Private/BlueprintGeneratedClass.cpp
--- a/Engine/Source/Runtime/Engine/Private/BlueprintGeneratedClass.cpp
+++ b/Engine/Source/Runtime/Engine/Private/BlueprintGeneratedClass.cpp
@@ -458,36 +458,40 @@ uint8* UBlueprintGeneratedClass::GetPersistentUberGraphFrame(UObject* Obj, UFunc
 	return ParentClass->GetPersistentUberGraphFrame(Obj, FuncToCheck);
 }
 
-void UBlueprintGeneratedClass::CreatePersistentUberGraphFrame(UObject* Obj) const
+void UBlueprintGeneratedClass::CreatePersistentUberGraphFrame(UObject* Obj, bool bCreateOnlyIfEmpty) const
 {
 	checkSlow(!UberGraphFramePointerProperty == !UberGraphFunction);
 	if (Obj && UsePersistentUberGraphFrame() && UberGraphFramePointerProperty && UberGraphFunction)
 	{
-		uint8* FrameMemory = NULL;
-		const bool bUberGraphFunctionIsReady = !WITH_EDITOR // no cyclic dependency problems
-			|| UberGraphFunction->HasAllFlags(RF_LoadCompleted); // is fully loaded
-		if (bUberGraphFunctionIsReady)
+		auto PointerToUberGraphFrame = UberGraphFramePointerProperty->ContainerPtrToValuePtr<FPointerToUberGraphFrame>(Obj);
+		check(PointerToUberGraphFrame);
+		check(bCreateOnlyIfEmpty || !PointerToUberGraphFrame->RawPointer);
+		
+		if (!PointerToUberGraphFrame->RawPointer)
 		{
-			FrameMemory = (uint8*)FMemory::Malloc(UberGraphFunction->GetStructureSize());
-			FMemory::Memzero(FrameMemory, UberGraphFunction->GetStructureSize());
-			for (UProperty* Property = UberGraphFunction->PropertyLink; Property; Property = Property->PropertyLinkNext)
+			uint8* FrameMemory = NULL;
+			const bool bUberGraphFunctionIsReady = UberGraphFunction->HasAllFlags(RF_LoadCompleted); // is fully loaded
+			if (bUberGraphFunctionIsReady)
 			{
-				Property->InitializeValue_InContainer(FrameMemory);
+				FrameMemory = (uint8*)FMemory::Malloc(UberGraphFunction->GetStructureSize());
+				FMemory::Memzero(FrameMemory, UberGraphFunction->GetStructureSize());
+				for (UProperty* Property = UberGraphFunction->PropertyLink; Property; Property = Property->PropertyLinkNext)
+				{
+					Property->InitializeValue_InContainer(FrameMemory);
+				}
 			}
+			else
+			{
+				UE_LOG(LogBlueprint, Warning, TEXT("Function '%s' is not ready to create frame for '%s'"),
+					*GetPathNameSafe(UberGraphFunction), *GetPathNameSafe(Obj));
+			}
+			PointerToUberGraphFrame->RawPointer = FrameMemory;
 		}
-		else
-		{
-			UE_LOG(LogBlueprint, Log, TEXT("Function '%s' is not ready to create frame."), *GetPathNameSafe(UberGraphFunction));
-		}
-		
-		auto PointerToUberGraphFrame = UberGraphFramePointerProperty->ContainerPtrToValuePtr<FPointerToUberGraphFrame>(Obj);
-		checkSlow(PointerToUberGraphFrame && !PointerToUberGraphFrame->RawPointer);
-		PointerToUberGraphFrame->RawPointer = FrameMemory;
 	}
 
 	auto ParentClass = GetSuperClass();
 	checkSlow(ParentClass);
-	return ParentClass->CreatePersistentUberGraphFrame(Obj);
+	return ParentClass->CreatePersistentUberGraphFrame(Obj, bCreateOnlyIfEmpty);
 }
 
 void UBlueprintGeneratedClass::DestroyPersistentUberGraphFrame(UObject* Obj) const
@@ -621,4 +625,14 @@ bool UBlueprintGeneratedClass::UsePersistentUberGraphFrame()
 #else
 	return false;
 #endif
+}
+
+void UBlueprintGeneratedClass::Serialize(FArchive& Ar)
+{
+	Super::Serialize(Ar);
+
+	if (Ar.IsLoading() && 0 == (Ar.GetPortFlags() & PPF_Duplicate))
+	{
+		CreatePersistentUberGraphFrame(ClassDefaultObject, true);
+	}
 }
\ No newline at end of file
