[INTEGRATE] Change 2430429 by [REDACTED_PERSON][REDACTED_HANDLE] on 2015/02/03 15:44:06

	UE-8100 Increase tolerance when re-applying preocomputed lighting to static meshes after BP construction runs

diff --git a/Engine/Source/Runtime/Engine/Private/StaticMeshComponent.cpp b/Engine/Source/Runtime/Engine/Private/StaticMeshComponent.cpp
--- a/Engine/Source/Runtime/Engine/Private/StaticMeshComponent.cpp
+++ b/Engine/Source/Runtime/Engine/Private/StaticMeshComponent.cpp
@@ -1586,7 +1586,7 @@ void UStaticMeshComponent::ApplyComponentInstanceData(FStaticMeshComponentInstan
 	}
 
 	// See if data matches current state
-	if(	StaticMeshInstanceData->bHasCachedStaticLighting && StaticMeshInstanceData->CachedStaticLighting.Transform.Equals(ComponentToWorld) )
+	if(	StaticMeshInstanceData->bHasCachedStaticLighting && StaticMeshInstanceData->CachedStaticLighting.Transform.Equals(ComponentToWorld, 1.e-3f) )
 	{
 		const int32 NumLODLightMaps = StaticMeshInstanceData->CachedStaticLighting.LODDataLightMap.Num();
 		SetLODDataCount(NumLODLightMaps, NumLODLightMaps);
