[INTEGRATE] Change 2430509 by [REDACTED_PERSON][REDACTED_HANDLE] on 2015/02/03 16:24:46

	Component Tags are more accessible, added GetComponentsByTag to AActor (GetComponentsByTag could not be implemented in blueprints without an exec pin, macro implementation also required improvements to type propagation)

diff --git a/Engine/Source/Runtime/Engine/Classes/Components/ActorComponent.h b/Engine/Source/Runtime/Engine/Classes/Components/ActorComponent.h
--- a/Engine/Source/Runtime/Engine/Classes/Components/ActorComponent.h
+++ b/Engine/Source/Runtime/Engine/Classes/Components/ActorComponent.h
@@ -43,7 +43,7 @@ public:
 	struct FActorComponentTickFunction PrimaryComponentTick;
 
 	/** Array of tags that can be used for grouping and categorizing. Can also be accessed from scripting. */
-	UPROPERTY(EditDefaultsOnly, AdvancedDisplay, BlueprintReadWrite, Category=Component)
+	UPROPERTY(EditAnywhere, BlueprintReadWrite, Category=Tags)
 	TArray<FName> ComponentTags;
 
 protected:
diff --git a/Engine/Source/Runtime/Engine/Classes/GameFramework/Actor.h b/Engine/Source/Runtime/Engine/Classes/GameFramework/Actor.h
--- a/Engine/Source/Runtime/Engine/Classes/GameFramework/Actor.h
+++ b/Engine/Source/Runtime/Engine/Classes/GameFramework/Actor.h
@@ -2141,6 +2141,10 @@ public:
 	UFUNCTION(BlueprintCallable, Category = "Actor", meta = (ComponentClass = "ActorComponent"), meta=(DeterminesOutputType="ComponentClass"))
 	TArray<UActorComponent*> GetComponentsByClass(TSubclassOf<UActorComponent> ComponentClass) const;
 
+	/* Gets all the components that inherit from the given class with a given tag. */
+	UFUNCTION(BlueprintCallable, Category = "Actor", meta = (ComponentClass = "ActorComponent"), meta = (DeterminesOutputType = "ComponentClass"))
+	TArray<UActorComponent*> GetComponentsByTag(TSubclassOf<UActorComponent> ComponentClass, FName Tag) const;
+
 	/** Templatized version for syntactic nicety. */
 	template<class T>
 	T* FindComponentByClass() const
diff --git a/Engine/Source/Runtime/Engine/Private/Actor.cpp b/Engine/Source/Runtime/Engine/Private/Actor.cpp
--- a/Engine/Source/Runtime/Engine/Private/Actor.cpp
+++ b/Engine/Source/Runtime/Engine/Private/Actor.cpp
@@ -2149,6 +2149,23 @@ TArray<UActorComponent*> AActor::GetComponentsByClass(TSubclassOf<UActorComponen
 	return ValidComponents;
 }
 
+TArray<UActorComponent*> AActor::GetComponentsByTag(TSubclassOf<UActorComponent> ComponentClass, FName Tag) const
+{
+	TArray<UActorComponent*> ComponentsByClass = GetComponentsByClass(ComponentClass);
+
+	TArray<UActorComponent*> ComponentsByTag;
+	ComponentsByTag.Reserve(ComponentsByClass.Num());
+	for (int i = 0; i < ComponentsByClass.Num(); ++i)
+	{
+		if (ComponentsByClass[i]->ComponentHasTag(Tag))
+		{
+			ComponentsByTag.Push(ComponentsByClass[i]);
+		}
+	}
+
+	return ComponentsByTag;
+}
+
 void AActor::DisableComponentsSimulatePhysics()
 {
 	TInlineComponentArray<UPrimitiveComponent*> Components;
