[INTEGRATE] Change 2430611 by [REDACTED_PERSON][REDACTED_HANDLE] on 2015/02/03 17:11:12

	Make GetArchetype() for ConstructionScript created components find the component from the SimpleConstructionScript

diff --git a/Engine/Source/Runtime/CoreUObject/Private/UObject/UObjectArchetype.cpp b/Engine/Source/Runtime/CoreUObject/Private/UObject/UObjectArchetype.cpp
--- a/Engine/Source/Runtime/CoreUObject/Private/UObject/UObjectArchetype.cpp
+++ b/Engine/Source/Runtime/CoreUObject/Private/UObject/UObjectArchetype.cpp
@@ -31,22 +31,23 @@ UObject* UObject::GetArchetypeFromRequiredInfo(UClass* Class, UObject* Outer, FN
 			{
 				Result = MyArchetype; // found that my outers archetype had a matching component, that must be my archetype
 			}
-		}
-
-		if (!Result 
-			&& !!(ObjectFlags&RF_InheritableComponentTemplate)
-			&& Outer && Outer->IsA<UClass>())
-		{
-			for (auto SuperClassArchetype = static_cast<UClass*>(Outer)->GetSuperClass();
-				SuperClassArchetype && SuperClassArchetype->HasAllClassFlags(CLASS_CompiledFromBlueprint);
-				SuperClassArchetype = SuperClassArchetype->GetSuperClass())
+			else if (!!(ObjectFlags&RF_InheritableComponentTemplate) && Outer->IsA<UClass>())
 			{
-				Result = static_cast<UObject*>(FindObjectWithOuter(SuperClassArchetype, Class, Name));
-				if (Result)
+				for (auto SuperClassArchetype = static_cast<UClass*>(Outer)->GetSuperClass();
+					SuperClassArchetype && SuperClassArchetype->HasAllClassFlags(CLASS_CompiledFromBlueprint);
+					SuperClassArchetype = SuperClassArchetype->GetSuperClass())
 				{
-					break;
+					Result = static_cast<UObject*>(FindObjectWithOuter(SuperClassArchetype, Class, Name));
+					if (Result)
+					{
+						break;
+					}
 				}
 			}
+			else
+			{
+				Result = ArchetypeToSearch->GetClass()->FindArchetype(Class, Name);
+			}
 		}
 
 		if (!Result)
diff --git a/Engine/Source/Runtime/CoreUObject/Public/UObject/Class.h b/Engine/Source/Runtime/CoreUObject/Public/UObject/Class.h
--- a/Engine/Source/Runtime/CoreUObject/Public/UObject/Class.h
+++ b/Engine/Source/Runtime/CoreUObject/Public/UObject/Class.h
@@ -2186,6 +2186,8 @@ public:
 	 */
 	virtual bool HasProperty(UProperty* InProperty) const;
 
+	virtual UObject* FindArchetype(UClass* ArchetypeClass, const FName ArchetypeName) const { return nullptr; }
+
 	/**
 	 * On save, we order a package's exports in class dependency order (so that
 	 * on load, we create the class dependencies before we create the class).  
diff --git a/Engine/Source/Runtime/Engine/Classes/Engine/BlueprintGeneratedClass.h b/Engine/Source/Runtime/Engine/Classes/Engine/BlueprintGeneratedClass.h
--- a/Engine/Source/Runtime/Engine/Classes/Engine/BlueprintGeneratedClass.h
+++ b/Engine/Source/Runtime/Engine/Classes/Engine/BlueprintGeneratedClass.h
@@ -464,7 +464,7 @@ public:
 	static bool GetGeneratedClassesHierarchy(const UClass* InClass, TArray<const UBlueprintGeneratedClass*>& OutBPGClasses);
 
 	/** Find the object in the TemplateObjects array with the supplied name */
-	UActorComponent* FindComponentTemplateByName(const FName& TemplateName);
+	UActorComponent* FindComponentTemplateByName(const FName& TemplateName) const;
 
 	/** Create Timeline objects for this Actor based on the Timelines array*/
 	virtual void CreateComponentsForActor(AActor* Actor) const;
@@ -488,6 +488,7 @@ public:
 	virtual void PurgeClass(bool bRecompilingOnLoad) override;
 	virtual void Bind() override;
 	virtual void GetRequiredPreloadDependencies(TArray<UObject*>& DependenciesOut) override;
+	virtual UObject* FindArchetype(UClass* ArchetypeClass, const FName ArchetypeName) const override;
 	// End UClass interface
 
 	static void AddReferencedObjectsInUbergraphFrame(UObject* InThis, FReferenceCollector& Collector);
diff --git a/Engine/Source/Runtime/Engine/Classes/Engine/SimpleConstructionScript.h b/Engine/Source/Runtime/Engine/Classes/Engine/SimpleConstructionScript.h
--- a/Engine/Source/Runtime/Engine/Classes/Engine/SimpleConstructionScript.h
+++ b/Engine/Source/Runtime/Engine/Classes/Engine/SimpleConstructionScript.h
@@ -16,6 +16,8 @@ class USimpleConstructionScript : public UObject
 	virtual void PostLoad() override;
 	// End UObject Interface
 
+	void PreloadChain();
+
 	/** Ensures that all root node parent references are still valid and clears the reference if not */
 	ENGINE_API void FixupRootNodeParentReferences();
 
@@ -59,10 +61,10 @@ class USimpleConstructionScript : public UObject
 	ENGINE_API USCS_Node* FindParentNode(USCS_Node* InNode) const;
 
 	/** Find the SCS_Node node by name and return it if found */
-	ENGINE_API USCS_Node* FindSCSNode(FName InName);
+	ENGINE_API USCS_Node* FindSCSNode(const FName InName) const;
 
 	/** Find the SCS_Node node by name and return it if found */
-	ENGINE_API USCS_Node* FindSCSNodeByGuid(FGuid Guid);
+	ENGINE_API USCS_Node* FindSCSNodeByGuid(const FGuid Guid) const;
 
 	/** Checks the root node set for scene components and ensures that it is valid (e.g. after a removal) */
 	ENGINE_API void ValidateSceneRootNodes();
diff --git a/Engine/Source/Runtime/Engine/Private/BlueprintGeneratedClass.cpp b/Engine/Source/Runtime/Engine/Private/BlueprintGeneratedClass.cpp
--- a/Engine/Source/Runtime/Engine/Private/BlueprintGeneratedClass.cpp
+++ b/Engine/Source/Runtime/Engine/Private/BlueprintGeneratedClass.cpp
@@ -4,6 +4,8 @@
 #include "BlueprintUtilities.h"
 #include "Engine/InputDelegateBinding.h"
 #include "Engine/TimelineTemplate.h"
+#include "Engine/SimpleConstructionScript.h"
+#include "Engine/SCS_Node.h"
 
 #if WITH_EDITOR
 #include "BlueprintEditorUtils.h"
@@ -253,6 +255,28 @@ bool UBlueprintGeneratedClass::IsFunctionImplementedInBlueprint(FName InFunction
 	return Function && Function->GetOuter() && Function->GetOuter()->IsA(UBlueprintGeneratedClass::StaticClass());
 }
 
+UObject* UBlueprintGeneratedClass::FindArchetype(UClass* ArchetypeClass, const FName ArchetypeName) const
+{
+	UObject* Archetype = nullptr;
+
+	if (SimpleConstructionScript->HasAllFlags(RF_NeedLoad))
+	{
+		SimpleConstructionScript->PreloadChain();
+	}
+
+	USCS_Node* SCSNode = SimpleConstructionScript->FindSCSNode(ArchetypeName);
+	if (SCSNode)
+	{
+		Archetype = SCSNode->ComponentTemplate;
+	}
+	else
+	{
+		Archetype = FindComponentTemplateByName(ArchetypeName);
+	}
+
+	return Archetype;
+}
+
 UDynamicBlueprintBinding* UBlueprintGeneratedClass::GetDynamicBindingObject(UClass* Class) const
 {
 	UDynamicBlueprintBinding* DynamicBindingObject = NULL;
@@ -307,7 +331,7 @@ bool UBlueprintGeneratedClass::GetGeneratedClassesHierarchy(const UClass* InClas
 	return bNoErrors;
 }
 
-UActorComponent* UBlueprintGeneratedClass::FindComponentTemplateByName(const FName& TemplateName)
+UActorComponent* UBlueprintGeneratedClass::FindComponentTemplateByName(const FName& TemplateName) const
 {
 	for(int32 i = 0; i < ComponentTemplates.Num(); i++)
 	{
diff --git a/Engine/Source/Runtime/Engine/Private/SCS_Node.cpp b/Engine/Source/Runtime/Engine/Private/SCS_Node.cpp
--- a/Engine/Source/Runtime/Engine/Private/SCS_Node.cpp
+++ b/Engine/Source/Runtime/Engine/Private/SCS_Node.cpp
@@ -188,6 +188,11 @@ void USCS_Node::PreloadChain()
 		GetLinker()->Preload(this);
 	}
 
+	if (ComponentTemplate && ComponentTemplate->HasAnyFlags(RF_NeedLoad))
+	{
+		ComponentTemplate->GetLinker()->Preload(ComponentTemplate);
+	}
+
 	for( TArray<USCS_Node*>::TIterator ChildIt(ChildNodes); ChildIt; ++ChildIt )
 	{
 		USCS_Node* CurrentChild = *ChildIt;
diff --git a/Engine/Source/Runtime/Engine/Private/SimpleConstructionScript.cpp b/Engine/Source/Runtime/Engine/Private/SimpleConstructionScript.cpp
--- a/Engine/Source/Runtime/Engine/Private/SimpleConstructionScript.cpp
+++ b/Engine/Source/Runtime/Engine/Private/SimpleConstructionScript.cpp
@@ -135,6 +135,16 @@ void USimpleConstructionScript::Serialize(FArchive& Ar)
 	}
 }
 
+void USimpleConstructionScript::PreloadChain()
+{
+	GetLinker()->Preload(this);
+
+	for (USCS_Node* Node : RootNodes)
+	{
+		Node->PreloadChain();
+	}
+}
+
 void USimpleConstructionScript::PostLoad()
 {
 	Super::PostLoad();
@@ -590,7 +600,7 @@ USCS_Node* USimpleConstructionScript::FindParentNode(USCS_Node* InNode) const
 	return NULL;
 }
 
-USCS_Node* USimpleConstructionScript::FindSCSNode(FName InName)
+USCS_Node* USimpleConstructionScript::FindSCSNode(const FName InName) const
 {
 	TArray<USCS_Node*> AllNodes = GetAllNodes();
 	USCS_Node* ReturnSCSNode = nullptr;
@@ -606,7 +616,7 @@ USCS_Node* USimpleConstructionScript::FindSCSNode(FName InName)
 	return ReturnSCSNode;
 }
 
-USCS_Node* USimpleConstructionScript::FindSCSNodeByGuid(FGuid Guid)
+USCS_Node* USimpleConstructionScript::FindSCSNodeByGuid(const FGuid Guid) const
 {
 	TArray<USCS_Node*> AllNodes = GetAllNodes();
 	USCS_Node* ReturnSCSNode = nullptr;
