fixed navigation updates after adding new objects with lazy geometry loading
#ue4

diff --git a/Engine/Source/Runtime/Engine/Public/AI/NavigationOctree.h b/Engine/Source/Runtime/Engine/Public/AI/NavigationOctree.h
--- a/Engine/Source/Runtime/Engine/Public/AI/NavigationOctree.h
+++ b/Engine/Source/Runtime/Engine/Public/AI/NavigationOctree.h
@@ -64,7 +64,7 @@ struct ENGINE_API FNavigationRelevantData
 	FORCEINLINE uint32 GetGeometryAllocatedSize() const { return CollisionData.GetAllocatedSize() + VoxelData.GetAllocatedSize(); }
 	FORCEINLINE int32 GetDirtyFlag() const
 	{
-		return (HasGeometry() ? ENavigationDirtyFlag::Geometry : 0) |
+		return ((HasGeometry() || IsPendingLazyGeometryGathering()) ? ENavigationDirtyFlag::Geometry : 0) |
 			(HasModifiers() ? ENavigationDirtyFlag::DynamicModifier : 0) |
 			(Modifiers.HasAgentHeightAdjust() ? ENavigationDirtyFlag::UseAgentHeight : 0);
 	}
