[INTEGRATE] Change 2430768 by [REDACTED_PERSON][REDACTED_HANDLE] on 2015/02/03 18:28:06

	Differentiate ConstructionScript CreationMethods in to SimpleConstructionScript and UserConstructionScript (or really other)
	Added IsCreatedByConstructionScript utility function on UActorComponent

diff --git a/Engine/Config/BaseEngine.ini b/Engine/Config/BaseEngine.ini
--- a/Engine/Config/BaseEngine.ini
+++ b/Engine/Config/BaseEngine.ini
@@ -471,6 +471,8 @@ bUseStreamingPause=false
 +ActiveStructRedirects=(OldStructName="KeyboardEvent",NewStructName="KeyEvent")
 +ActiveStructRedirects=(OldStructName="KeyboardFocusEvent",NewStructName="FocusEvent")
 
++EnumRedirects=(EnumName="EComponentCreationMethod", OldEnumEntry="EComponentCreationMethod::ConstructionScript", NewEnumEntry="EComponentCreationMethod::SimpleConstructionScript")
+
 ; Property name redirects
 +TaggedPropertyRedirects=(ClassName="Actor",OldPropertyName="LifeSpan",NewPropertyName="InitialLifeSpan")
 +TaggedPropertyRedirects=(ClassName="ActorComponent",OldPropertyName="bWantsBeginPlayComponent",NewPropertyName="bWantsInitializeComponent")
diff --git a/Engine/Source/Editor/DetailCustomizations/Private/ComponentMaterialCategory.cpp b/Engine/Source/Editor/DetailCustomizations/Private/ComponentMaterialCategory.cpp
--- a/Engine/Source/Editor/DetailCustomizations/Private/ComponentMaterialCategory.cpp
+++ b/Engine/Source/Editor/DetailCustomizations/Private/ComponentMaterialCategory.cpp
@@ -55,7 +55,7 @@ public:
 				UPrimitiveComponent* PrimitiveComp = Cast<UPrimitiveComponent>( CurComponent );
 				UDecalComponent* DecalComponent = (PrimitiveComp ? NULL : Cast<UDecalComponent>( CurComponent ));
 
-				if( CurComponent->CreationMethod != EComponentCreationMethod::ConstructionScript)
+				if (!CurComponent->IsCreatedByConstructionScript())
 				{
 					if( PrimitiveComp )
 					{
@@ -186,7 +186,7 @@ void FComponentMaterialCategory::Create( IDetailLayoutBuilder& DetailBuilder )
 	{	
 		UActorComponent* CurrentComponent = It.GetComponent();
 
-		if( !bAnyMaterialsToDisplay && CurrentComponent->CreationMethod != EComponentCreationMethod::ConstructionScript )
+		if( !bAnyMaterialsToDisplay && !CurrentComponent->IsCreatedByConstructionScript() )
 		{
 			bAnyMaterialsToDisplay = true;
 			break;
@@ -215,7 +215,7 @@ void FComponentMaterialCategory::OnGetMaterialsForView( IMaterialListBuilder& Ma
 
 		UActorComponent* CurrentComponent = It.GetComponent();
 
-		if( CurrentComponent && CurrentComponent->CreationMethod != EComponentCreationMethod::ConstructionScript )
+		if( CurrentComponent && !CurrentComponent->IsCreatedByConstructionScript() )
 		{
 			UMaterialInterface* Material = It.GetMaterial();
 
diff --git a/Engine/Source/Editor/FoliageEdit/Private/FoliageEdMode.cpp b/Engine/Source/Editor/FoliageEdit/Private/FoliageEdMode.cpp
--- a/Engine/Source/Editor/FoliageEdit/Private/FoliageEdMode.cpp
+++ b/Engine/Source/Editor/FoliageEdit/Private/FoliageEdMode.cpp
@@ -2001,7 +2001,7 @@ bool FEdModeFoliage::SnapInstanceToGround(AInstancedFoliageActor* InIFA, float A
 		UPrimitiveComponent* HitComponent = Hit.Component.Get();
 
 		// We cannot be based on an a blueprint component as these will disappear when the construction script is re-run
-		if (HitComponent->CreationMethod == EComponentCreationMethod::ConstructionScript)
+		if (HitComponent->IsCreatedByConstructionScript())
 		{
 			HitComponent = nullptr;
 		}
diff --git a/Engine/Source/Editor/Kismet/Private/SSCSEditor.cpp b/Engine/Source/Editor/Kismet/Private/SSCSEditor.cpp
--- a/Engine/Source/Editor/Kismet/Private/SSCSEditor.cpp
+++ b/Engine/Source/Editor/Kismet/Private/SSCSEditor.cpp
@@ -470,7 +470,7 @@ FSCSEditorTreeNodePtrType FSCSEditorTreeNode::AddChildFromComponent(UActorCompon
 // Tries to find a SCS node that was likely responsible for creating the specified instance component.  Note: This is not always possible to do!
 USCS_Node* FSCSEditorTreeNode::FindSCSNodeForInstance(UActorComponent* InstanceComponent, UClass* ClassToSearch)
 { 
-	if ((ClassToSearch != nullptr) && (InstanceComponent->CreationMethod == EComponentCreationMethod::ConstructionScript))
+	if ((ClassToSearch != nullptr) && InstanceComponent->IsCreatedByConstructionScript())
 	{
 		for (UClass* TestClass = ClassToSearch; TestClass->ClassGeneratedBy != nullptr; TestClass = TestClass->GetSuperClass())
 		{
@@ -1704,7 +1704,7 @@ FText SSCS_RowWidget::GetIntroducedInToolTipText() const
 
 				if (BestClass == nullptr)
 				{
-					ensure(ComponentTemplate->CreationMethod == EComponentCreationMethod::ConstructionScript);
+					ensure(ComponentTemplate->IsCreatedByConstructionScript());
 					IntroducedInTooltip = LOCTEXT("IntroducedInUnknownError", "Unknown Blueprint Class (via an Add Component call)");
 				}
 				else
@@ -3911,7 +3911,7 @@ FSCSEditorTreeNodePtrType SSCSEditor::GetNodeFromActorComponent(const UActorComp
 				UClass* OwnerClass = ActorComponent->GetOwner()->GetActorClass();
 
 				// If the given component is one that's created during Blueprint construction
-				if (ActorComponent->CreationMethod == EComponentCreationMethod::ConstructionScript)
+				if (ActorComponent->IsCreatedByConstructionScript())
 				{
 					// Get the Blueprint object associated with the owner's class
 					UBlueprint* Blueprint = UBlueprint::GetBlueprintFromClass(OwnerClass);
diff --git a/Engine/Source/Editor/LevelEditor/Private/SActorDetails.cpp b/Engine/Source/Editor/LevelEditor/Private/SActorDetails.cpp
--- a/Engine/Source/Editor/LevelEditor/Private/SActorDetails.cpp
+++ b/Engine/Source/Editor/LevelEditor/Private/SActorDetails.cpp
@@ -424,7 +424,7 @@ bool SActorDetails::IsPropertyEditingEnabled() const
 		UActorComponent* ActorComp = Cast<UActorComponent>(Object.Get());
 		if(ActorComp)
 		{
-			bIsEditable = ActorComp->CreationMethod != EComponentCreationMethod::ConstructionScript;
+			bIsEditable = !ActorComp->IsCreatedByConstructionScript();
 			if(!bIsEditable)
 			{
 				break;
diff --git a/Engine/Source/Editor/UnrealEd/Private/Editor.cpp b/Engine/Source/Editor/UnrealEd/Private/Editor.cpp
--- a/Engine/Source/Editor/UnrealEd/Private/Editor.cpp
+++ b/Engine/Source/Editor/UnrealEd/Private/Editor.cpp
@@ -7429,7 +7429,7 @@ namespace EditorUtilities
 											{
 												// Ensure that this instance will be included in any undo/redo operations, and record it into the transaction buffer.
 												// Note: We don't do this for components that originate from script, because they will be re-instanced from the template after an undo, so there is no need to record them.
-												if(ComponentArchetypeInstance->CreationMethod != EComponentCreationMethod::ConstructionScript)
+												if (!ComponentArchetypeInstance->IsCreatedByConstructionScript())
 												{
 													ComponentArchetypeInstance->SetFlags(RF_Transactional);
 													ComponentArchetypeInstance->Modify();
diff --git a/Engine/Source/Editor/UnrealEd/Private/EditorObject.cpp b/Engine/Source/Editor/UnrealEd/Private/EditorObject.cpp
--- a/Engine/Source/Editor/UnrealEd/Private/EditorObject.cpp
+++ b/Engine/Source/Editor/UnrealEd/Private/EditorObject.cpp
@@ -489,7 +489,7 @@ static const TCHAR* ImportProperties(
 					if (NewFlags & RF_Transactional)
 					{
 						UActorComponent* Component = Cast<UActorComponent>(ComponentTemplate);
-						if (Component && Component->CreationMethod == EComponentCreationMethod::ConstructionScript)
+						if (Component && Component->IsCreatedByConstructionScript())
 						{
 							NewFlags &= ~RF_Transactional;
 						}
diff --git a/Engine/Source/Editor/UnrealEd/Private/Kismet2/Kismet2.cpp b/Engine/Source/Editor/UnrealEd/Private/Kismet2/Kismet2.cpp
--- a/Engine/Source/Editor/UnrealEd/Private/Kismet2/Kismet2.cpp
+++ b/Engine/Source/Editor/UnrealEd/Private/Kismet2/Kismet2.cpp
@@ -957,7 +957,7 @@ void FKismetEditorUtilities::AddComponentsToBlueprint(UBlueprint* Blueprint, con
 					SCS->AddNode(SCSNode);
 				}
 				// If we're attached to a blueprint component look it up as the variable name is the component name
-				else if (SceneComponent->AttachParent->CreationMethod == EComponentCreationMethod::ConstructionScript)
+				else if (SceneComponent->AttachParent->IsCreatedByConstructionScript())
 				{
 					USCS_Node* ParentSCSNode = nullptr;
 					for (UBlueprint* Blueprint : ParentBPStack)
diff --git a/Engine/Source/Editor/UnrealEd/Private/LevelEditorViewport.cpp b/Engine/Source/Editor/UnrealEd/Private/LevelEditorViewport.cpp
--- a/Engine/Source/Editor/UnrealEd/Private/LevelEditorViewport.cpp
+++ b/Engine/Source/Editor/UnrealEd/Private/LevelEditorViewport.cpp
@@ -505,7 +505,7 @@ static bool AttemptApplyObjToComponent(UObject* ObjToUse, USceneComponent* Compo
 {
 	bool bResult = false;
 
-	if (ComponentToApplyTo && ComponentToApplyTo->CreationMethod != EComponentCreationMethod::ConstructionScript)
+	if (ComponentToApplyTo && !ComponentToApplyTo->IsCreatedByConstructionScript())
 	{
 		// MESH/DECAL
 		auto MeshComponent = Cast<UMeshComponent>(ComponentToApplyTo);
diff --git a/Engine/Source/Editor/UnrealEd/Public/Kismet2/ComponentEditorUtils.h b/Engine/Source/Editor/UnrealEd/Public/Kismet2/ComponentEditorUtils.h
--- a/Engine/Source/Editor/UnrealEd/Public/Kismet2/ComponentEditorUtils.h
+++ b/Engine/Source/Editor/UnrealEd/Public/Kismet2/ComponentEditorUtils.h
@@ -96,7 +96,7 @@ public:
 		{
 			// Ensure that this instance will be included in any undo/redo operations, and record it into the transaction buffer.
 			// Note: We don't do this for components that originate from script, because they will be re-instanced from the template after an undo, so there is no need to record them.
-			if(InSceneComponent->CreationMethod != EComponentCreationMethod::ConstructionScript)
+			if (!InSceneComponent->IsCreatedByConstructionScript())
 			{
 				InSceneComponent->SetFlags(RF_Transactional);
 				InSceneComponent->Modify();
diff --git a/Engine/Source/Runtime/Engine/Classes/Components/ActorComponent.h b/Engine/Source/Runtime/Engine/Classes/Components/ActorComponent.h
--- a/Engine/Source/Runtime/Engine/Classes/Components/ActorComponent.h
+++ b/Engine/Source/Runtime/Engine/Classes/Components/ActorComponent.h
@@ -12,7 +12,8 @@ UENUM()
 enum class EComponentCreationMethod : uint8
 {
 	Native,
-	ConstructionScript,
+	SimpleConstructionScript,
+	UserConstructionScript,
 	Instance,
 };
 
@@ -124,6 +125,8 @@ public:
 	UPROPERTY()
 	EComponentCreationMethod CreationMethod;
 
+	bool IsCreatedByConstructionScript() const;
+
 	UFUNCTION()
 	void OnRep_IsActive();
 
diff --git a/Engine/Source/Runtime/Engine/Classes/Engine/Selection.h b/Engine/Source/Runtime/Engine/Classes/Engine/Selection.h
--- a/Engine/Source/Runtime/Engine/Classes/Engine/Selection.h
+++ b/Engine/Source/Runtime/Engine/Classes/Engine/Selection.h
@@ -467,7 +467,7 @@ public:
 	bool IsObjectValid(const UObject* Object) const
 	{
 		const UActorComponent* Comp = Cast<UActorComponent>( Object );
-		return Comp && Comp->CreationMethod != EComponentCreationMethod::ConstructionScript;
+		return Comp && !Comp->IsCreatedByConstructionScript();
 	}
 };
 
diff --git a/Engine/Source/Runtime/Engine/Private/Actor.cpp b/Engine/Source/Runtime/Engine/Private/Actor.cpp
--- a/Engine/Source/Runtime/Engine/Private/Actor.cpp
+++ b/Engine/Source/Runtime/Engine/Private/Actor.cpp
@@ -808,7 +808,7 @@ bool AActor::Modify( bool bAlwaysMarkDirty/*=true*/ )
 			if (!ObjProp->HasAllPropertyFlags(CPF_NonTransactional))
 			{
 				UActorComponent* ActorComponent = Cast<UActorComponent>(ObjProp->GetObjectPropertyValue(ObjProp->ContainerPtrToValuePtr<void>(this)));
-				if (ActorComponent && ActorComponent->CreationMethod == EComponentCreationMethod::ConstructionScript)
+				if (ActorComponent && ActorComponent->IsCreatedByConstructionScript())
 				{
 					ObjProp->SetPropertyFlags(CPF_NonTransactional);
 					TemporarilyNonTransactionalProperties.Add(ObjProp);
@@ -825,7 +825,7 @@ bool AActor::Modify( bool bAlwaysMarkDirty/*=true*/ )
 	}
 
 	// If the root component is blueprint constructed we don't save it to the transaction buffer
-	if( RootComponent && RootComponent->CreationMethod != EComponentCreationMethod::ConstructionScript)
+	if( RootComponent && !RootComponent->IsCreatedByConstructionScript())
 	{
 		bSavedToTransactionBuffer = RootComponent->Modify( bAlwaysMarkDirty ) || bSavedToTransactionBuffer;
 	}
diff --git a/Engine/Source/Runtime/Engine/Private/ActorComponent.cpp b/Engine/Source/Runtime/Engine/Private/ActorComponent.cpp
--- a/Engine/Source/Runtime/Engine/Private/ActorComponent.cpp
+++ b/Engine/Source/Runtime/Engine/Private/ActorComponent.cpp
@@ -9,6 +9,7 @@
 #include "UObjectToken.h"
 #include "MapErrors.h"
 #include "ComponentReregisterContext.h"
+#include "Engine/SimpleConstructionScript.h"
 
 #define LOCTEXT_NAMESPACE "ActorComponent"
 
@@ -126,15 +127,25 @@ void UActorComponent::PostLoad()
 {
 	Super::PostLoad();
 
+	// TODO: Wrap all this up with an engine version
 	if (bCreatedByConstructionScript_DEPRECATED)
 	{
-		CreationMethod = EComponentCreationMethod::ConstructionScript;
+		CreationMethod = EComponentCreationMethod::SimpleConstructionScript;
 	}
 	else if (bInstanceComponent_DEPRECATED)
 	{
 		CreationMethod = EComponentCreationMethod::Instance;
 	}
 
+	if (CreationMethod == EComponentCreationMethod::SimpleConstructionScript)
+	{
+		UBlueprintGeneratedClass* Class = CastChecked<UBlueprintGeneratedClass>(GetOuter()->GetClass());
+		if (Class->SimpleConstructionScript->FindSCSNode(GetFName()) == nullptr)
+		{
+			CreationMethod = EComponentCreationMethod::UserConstructionScript;
+		}
+	}
+
 	if (!HasAllFlags(RF_Public) && GetOuter()->IsA<UBlueprintGeneratedClass>())
 	{
 		SetFlags(RF_Public);
@@ -163,6 +174,11 @@ void UActorComponent::PostRename(UObject* OldOuter, const FName OldName)
 	}
 }
 
+bool UActorComponent::IsCreatedByConstructionScript() const
+{
+	return ((CreationMethod == EComponentCreationMethod::SimpleConstructionScript) || (CreationMethod == EComponentCreationMethod::UserConstructionScript));
+}
+
 #if WITH_EDITOR
 void UActorComponent::CheckForErrors()
 {
@@ -665,7 +681,7 @@ void UActorComponent::RegisterComponentWithWorld(UWorld* InWorld)
 	}
 
 	// If this is a blueprint created component and it has component children they can miss getting registered in some scenarios
-	if (CreationMethod == EComponentCreationMethod::ConstructionScript)
+	if (IsCreatedByConstructionScript())
 	{
 		TArray<UObject*> Children;
 		GetObjectsWithOuter(this, Children, true, RF_PendingKill);
@@ -735,7 +751,7 @@ void UActorComponent::DestroyComponent(bool bPromoteChildren/*= false*/)
 	AActor* Owner = GetOwner();
 	if(Owner != NULL)
 	{
-		if (CreationMethod == EComponentCreationMethod::ConstructionScript)
+		if (IsCreatedByConstructionScript())
 		{
 			Owner->BlueprintCreatedComponents.Remove(this);
 		}
diff --git a/Engine/Source/Runtime/Engine/Private/ActorConstruction.cpp b/Engine/Source/Runtime/Engine/Private/ActorConstruction.cpp
--- a/Engine/Source/Runtime/Engine/Private/ActorConstruction.cpp
+++ b/Engine/Source/Runtime/Engine/Private/ActorConstruction.cpp
@@ -189,7 +189,7 @@ void AActor::DestroyConstructedComponents()
 		if (Component)
 		{
 			bool bDestroyComponent = false;
-			if (Component->CreationMethod == EComponentCreationMethod::ConstructionScript)
+			if (Component->IsCreatedByConstructionScript())
 			{
 				bDestroyComponent = true;
 			}
@@ -198,7 +198,7 @@ void AActor::DestroyConstructedComponents()
 				UActorComponent* OuterComponent = Component->GetTypedOuter<UActorComponent>();
 				while (OuterComponent)
 				{
-					if (OuterComponent->CreationMethod == EComponentCreationMethod::ConstructionScript)
+					if (OuterComponent->IsCreatedByConstructionScript())
 					{
 						bDestroyComponent = true;
 						break;
@@ -331,7 +331,7 @@ void AActor::RerunConstructionScripts()
 			{
 				USceneComponent* EachRoot = AttachedActor->GetRootComponent();
 				// If the component we are attached to is about to go away...
-				if (EachRoot && EachRoot->AttachParent && EachRoot->AttachParent->CreationMethod == EComponentCreationMethod::ConstructionScript)
+				if (EachRoot && EachRoot->AttachParent && EachRoot->AttachParent->IsCreatedByConstructionScript())
 				{
 					// Save info about actor to reattach
 					FAttachedActorInfo Info;
@@ -350,7 +350,7 @@ void AActor::RerunConstructionScripts()
 		if (bUseRootComponentProperties && RootComponent != nullptr)
 		{
 			// Do not need to detach if root component is not going away
-			if (RootComponent->AttachParent != NULL && RootComponent->CreationMethod == EComponentCreationMethod::ConstructionScript)
+			if (RootComponent->AttachParent != NULL && RootComponent->IsCreatedByConstructionScript())
 			{
 				Parent = RootComponent->AttachParent->GetOwner();
 				// Root component should never be attached to another component in the same actor!
@@ -503,7 +503,7 @@ void AActor::ExecuteConstruction(const FTransform& Transform, const FComponentIn
 			{
 				UBillboardComponent* BillboardComponent = NewObject<UBillboardComponent>(this);
 				BillboardComponent->SetFlags(RF_Transactional);
-				BillboardComponent->CreationMethod = EComponentCreationMethod::ConstructionScript;
+				BillboardComponent->CreationMethod = EComponentCreationMethod::SimpleConstructionScript;
 #if WITH_EDITOR
 				BillboardComponent->Sprite = (UTexture2D*)(StaticLoadObject(UTexture2D::StaticClass(), NULL, TEXT("/Engine/EditorResources/BadBlueprintSprite.BadBlueprintSprite"), NULL, LOAD_None, NULL));
 #endif
@@ -579,7 +579,7 @@ UActorComponent* AActor::CreateComponentFromTemplate(UActorComponent* Template,
 		NewActorComp = (UActorComponent*)StaticDuplicateObject(Template, this, *InName, RF_AllFlags & ~(RF_ArchetypeObject|RF_Transactional|RF_WasLoaded|RF_Public|RF_InheritableComponentTemplate) );
 		//NewActorComp = ConstructObject<UActorComponent>(Template->GetClass(), this, *InName, RF_NoFlags, Template);
 
-		NewActorComp->CreationMethod = EComponentCreationMethod::ConstructionScript;
+		NewActorComp->CreationMethod = EComponentCreationMethod::UserConstructionScript;
 
 		// Need to do this so component gets saved - Components array is not serialized
 		BlueprintCreatedComponents.Add(NewActorComp);
diff --git a/Engine/Source/Runtime/Engine/Private/ActorEditor.cpp b/Engine/Source/Runtime/Engine/Private/ActorEditor.cpp
--- a/Engine/Source/Runtime/Engine/Private/ActorEditor.cpp
+++ b/Engine/Source/Runtime/Engine/Private/ActorEditor.cpp
@@ -83,7 +83,7 @@ void AActor::PostEditMove(bool bFinished)
 	}
 
 	// If the root component was not just recreated by the construction script - call PostEditComponentMove on it
-	if(RootComponent != NULL && RootComponent->CreationMethod != EComponentCreationMethod::ConstructionScript)
+	if(RootComponent != NULL && !RootComponent->IsCreatedByConstructionScript())
 	{
 		// @TODO Should we call on ALL components?
 		RootComponent->PostEditComponentMove(bFinished);
@@ -204,7 +204,7 @@ AActor::FActorTransactionAnnotation::FActorTransactionAnnotation(const AActor* A
 	: ComponentInstanceData(Actor)
 {
 	USceneComponent* RootComponent = Actor->GetRootComponent();
-	if (RootComponent && RootComponent->CreationMethod == EComponentCreationMethod::ConstructionScript)
+	if (RootComponent && RootComponent->IsCreatedByConstructionScript())
 	{
 		bRootComponentDataCached = true;
 		RootComponentData.Transform = RootComponent->ComponentToWorld;
diff --git a/Engine/Source/Runtime/Engine/Private/ActorEditorUtils.cpp b/Engine/Source/Runtime/Engine/Private/ActorEditorUtils.cpp
--- a/Engine/Source/Runtime/Engine/Private/ActorEditorUtils.cpp
+++ b/Engine/Source/Runtime/Engine/Private/ActorEditorUtils.cpp
@@ -40,7 +40,7 @@ namespace FActorEditorUtils
 		InActor->GetComponents(InstanceComponents);
 		for (auto Component : InstanceComponents)
 		{
-			if (Component->CreationMethod != EComponentCreationMethod::ConstructionScript)
+			if (!Component->IsCreatedByConstructionScript())
 			{
 				OutEditableComponents.Add(Component);
 			}
diff --git a/Engine/Source/Runtime/Engine/Private/BlueprintGeneratedClass.cpp b/Engine/Source/Runtime/Engine/Private/BlueprintGeneratedClass.cpp
--- a/Engine/Source/Runtime/Engine/Private/BlueprintGeneratedClass.cpp
+++ b/Engine/Source/Runtime/Engine/Private/BlueprintGeneratedClass.cpp
@@ -364,7 +364,7 @@ void UBlueprintGeneratedClass::CreateComponentsForActor(AActor* Actor) const
 
 		FName NewName = *(FString::Printf(TEXT("TimelineComp__%d"), Actor->BlueprintCreatedComponents.Num() ) );
 		UTimelineComponent* NewTimeline = NewObject<UTimelineComponent>(Actor, NewName);
-		NewTimeline->CreationMethod = EComponentCreationMethod::ConstructionScript; // Indicate it comes from a blueprint so it gets cleared when we rerun construction scripts
+		NewTimeline->CreationMethod = EComponentCreationMethod::UserConstructionScript; // Indicate it comes from a blueprint so it gets cleared when we rerun construction scripts
 		Actor->BlueprintCreatedComponents.Add(NewTimeline); // Add to array so it gets saved
 		NewTimeline->SetNetAddressable();	// This component has a stable name that can be referenced for replication
 
diff --git a/Engine/Source/Runtime/Engine/Private/ComponentInstanceDataCache.cpp b/Engine/Source/Runtime/Engine/Private/ComponentInstanceDataCache.cpp
--- a/Engine/Source/Runtime/Engine/Private/ComponentInstanceDataCache.cpp
+++ b/Engine/Source/Runtime/Engine/Private/ComponentInstanceDataCache.cpp
@@ -78,7 +78,7 @@ FComponentInstanceDataCache::FComponentInstanceDataCache(const AActor* Actor)
 		// Grab per-instance data we want to persist
 		for (UActorComponent* Component : Components)
 		{
-			if (Component->CreationMethod == EComponentCreationMethod::ConstructionScript) // Only cache data from 'created by construction script' components
+			if (Component->IsCreatedByConstructionScript()) // Only cache data from 'created by construction script' components
 			{
 				FComponentInstanceDataBase* ComponentInstanceData = Component->GetComponentInstanceData();
 				if (ComponentInstanceData)
@@ -92,7 +92,7 @@ FComponentInstanceDataCache::FComponentInstanceDataCache(const AActor* Actor)
 				// If the instance component is attached to a BP component we have to be prepared for the possibility that it will be deleted
 				if (USceneComponent* SceneComponent = Cast<USceneComponent>(Component))
 				{
-					if (SceneComponent->AttachParent && SceneComponent->AttachParent->CreationMethod == EComponentCreationMethod::ConstructionScript)
+					if (SceneComponent->AttachParent && SceneComponent->AttachParent->IsCreatedByConstructionScript())
 					{
 						InstanceComponentTransformToRootMap.Add(SceneComponent, SceneComponent->GetComponentTransform().GetRelativeTransform(Actor->GetRootComponent()->GetComponentTransform()));
 					}
@@ -120,7 +120,7 @@ void FComponentInstanceDataCache::ApplyToActor(AActor* Actor) const
 		// Apply per-instance data.
 		for (UActorComponent* Component : Components)
 		{
-			if(Component->CreationMethod == EComponentCreationMethod::ConstructionScript) // Only try and apply data to 'created by construction script' components
+			if(Component->IsCreatedByConstructionScript()) // Only try and apply data to 'created by construction script' components
 			{
 				const FName ComponentInstanceDataType = Component->GetComponentInstanceDataType();
 
diff --git a/Engine/Source/Runtime/Engine/Private/SCS_Node.cpp b/Engine/Source/Runtime/Engine/Private/SCS_Node.cpp
--- a/Engine/Source/Runtime/Engine/Private/SCS_Node.cpp
+++ b/Engine/Source/Runtime/Engine/Private/SCS_Node.cpp
@@ -52,6 +52,7 @@ UActorComponent* USCS_Node::ExecuteNodeOnActor(AActor* Actor, USceneComponent* P
 	UActorComponent* NewActorComp = Actor->CreateComponentFromTemplate(ActualComponentTemplate, VariableName.ToString());
 	if(NewActorComp != nullptr)
 	{
+		NewActorComp->CreationMethod = EComponentCreationMethod::SimpleConstructionScript;
 		// SCS created components are net addressable
 		NewActorComp->SetNetAddressable();
 
diff --git a/Engine/Source/Runtime/Engine/Private/SceneComponent.cpp b/Engine/Source/Runtime/Engine/Private/SceneComponent.cpp
--- a/Engine/Source/Runtime/Engine/Private/SceneComponent.cpp
+++ b/Engine/Source/Runtime/Engine/Private/SceneComponent.cpp
@@ -1156,7 +1156,7 @@ FSceneComponentInstanceData::FSceneComponentInstanceData(const USceneComponent*
 	for (int32 i = SourceComponent->AttachChildren.Num()-1; i >= 0; --i)
 	{
 		USceneComponent* SceneComponent = SourceComponent->AttachChildren[i];
-		if (SceneComponent && SceneComponent->CreationMethod != EComponentCreationMethod::ConstructionScript)
+		if (SceneComponent && !SceneComponent->IsCreatedByConstructionScript())
 		{
 			AttachedInstanceComponents.Add(SceneComponent);
 		}
@@ -1192,7 +1192,7 @@ FComponentInstanceDataBase* USceneComponent::GetComponentInstanceData() const
 
 	for (USceneComponent* Child : AttachChildren)
 	{
-		if (Child && Child->CreationMethod != EComponentCreationMethod::ConstructionScript)
+		if (Child && !Child->IsCreatedByConstructionScript())
 		{
 			InstanceData = new FSceneComponentInstanceData(this);
 			break;
diff --git a/Engine/Source/Runtime/Engine/Private/SimpleConstructionScript.cpp b/Engine/Source/Runtime/Engine/Private/SimpleConstructionScript.cpp
--- a/Engine/Source/Runtime/Engine/Private/SimpleConstructionScript.cpp
+++ b/Engine/Source/Runtime/Engine/Private/SimpleConstructionScript.cpp
@@ -401,7 +401,7 @@ void USimpleConstructionScript::ExecuteScriptOnActor(AActor* Actor, const FTrans
 	{
 		USceneComponent* SceneComp = NewObject<USceneComponent>(Actor);
 		SceneComp->SetFlags(RF_Transactional);
-		SceneComp->CreationMethod = EComponentCreationMethod::ConstructionScript;
+		SceneComp->CreationMethod = EComponentCreationMethod::SimpleConstructionScript;
 		SceneComp->SetWorldTransform(RootTransform);
 		Actor->SetRootComponent(SceneComp);
 		SceneComp->RegisterComponent();
diff --git a/Engine/Source/Runtime/Foliage/Private/FoliageInstanceBase.cpp b/Engine/Source/Runtime/Foliage/Private/FoliageInstanceBase.cpp
--- a/Engine/Source/Runtime/Foliage/Private/FoliageInstanceBase.cpp
+++ b/Engine/Source/Runtime/Foliage/Private/FoliageInstanceBase.cpp
@@ -85,7 +85,7 @@ FArchive& operator << (FArchive& Ar, FFoliageInstanceBaseCache& InstanceBaseCach
 FFoliageInstanceBaseId FFoliageInstanceBaseCache::AddInstanceBaseId(UActorComponent* InComponent)
 {
 	FFoliageInstanceBaseId BaseId = FFoliageInstanceBaseCache::InvalidBaseId;
-	if (InComponent && InComponent->CreationMethod != EComponentCreationMethod::ConstructionScript)
+	if (InComponent && !InComponent->IsCreatedByConstructionScript())
 	{
 		BaseId = GetInstanceBaseId(InComponent);
 		if (BaseId == FFoliageInstanceBaseCache::InvalidBaseId)
