[INTEGRATE] Change 2430787 by [REDACTED_PERSON][REDACTED_HANDLE] on 2015/02/03 18:43:45

	Allow per instance blueprint created component customization for simple construction script components

diff --git a/Engine/Source/Editor/LevelEditor/Private/SActorDetails.cpp b/Engine/Source/Editor/LevelEditor/Private/SActorDetails.cpp
--- a/Engine/Source/Editor/LevelEditor/Private/SActorDetails.cpp
+++ b/Engine/Source/Editor/LevelEditor/Private/SActorDetails.cpp
@@ -424,7 +424,7 @@ bool SActorDetails::IsPropertyEditingEnabled() const
 		UActorComponent* ActorComp = Cast<UActorComponent>(Object.Get());
 		if(ActorComp)
 		{
-			bIsEditable = !ActorComp->IsCreatedByConstructionScript();
+			bIsEditable = ActorComp->CreationMethod != EComponentCreationMethod::UserConstructionScript;
 			if(!bIsEditable)
 			{
 				break;
diff --git a/Engine/Source/Runtime/Engine/Classes/Engine/Selection.h b/Engine/Source/Runtime/Engine/Classes/Engine/Selection.h
--- a/Engine/Source/Runtime/Engine/Classes/Engine/Selection.h
+++ b/Engine/Source/Runtime/Engine/Classes/Engine/Selection.h
@@ -467,7 +467,7 @@ public:
 	bool IsObjectValid(const UObject* Object) const
 	{
 		const UActorComponent* Comp = Cast<UActorComponent>( Object );
-		return Comp && !Comp->IsCreatedByConstructionScript();
+		return Comp && Comp->CreationMethod != EComponentCreationMethod::UserConstructionScript;
 	}
 };
 
diff --git a/Engine/Source/Runtime/Engine/Private/ComponentInstanceDataCache.cpp b/Engine/Source/Runtime/Engine/Private/ComponentInstanceDataCache.cpp
--- a/Engine/Source/Runtime/Engine/Private/ComponentInstanceDataCache.cpp
+++ b/Engine/Source/Runtime/Engine/Private/ComponentInstanceDataCache.cpp
@@ -35,6 +35,27 @@ FComponentInstanceDataBase::FComponentInstanceDataBase(const UActorComponent* So
 			SourceComponentTypeSerializedIndex = -1;
 		}
 	}
+
+	if (SourceComponent->CreationMethod == EComponentCreationMethod::SimpleConstructionScript)
+	{
+		class FComponentPropertyWriter : public FObjectWriter
+		{
+		public:
+			FComponentPropertyWriter(TArray<uint8>& InBytes)
+				: FObjectWriter(InBytes)
+			{
+			}
+
+			virtual bool ShouldSkipProperty(const UProperty* InProperty) const override
+			{
+				return (    InProperty->HasAnyPropertyFlags(CPF_Transient | CPF_ContainsInstancedReference | CPF_InstancedReference)
+						|| !InProperty->HasAnyPropertyFlags(CPF_Edit | CPF_Interp));
+			}
+
+		} ComponentPropertyWriter(SavedProperties);
+
+		SourceComponentClass->SerializeTaggedProperties(ComponentPropertyWriter, (uint8*)SourceComponent, SourceComponentClass, (uint8*)SourceComponent->GetArchetype());
+	}
 }
 
 bool FComponentInstanceDataBase::MatchesComponent(const UActorComponent* Component) const
@@ -68,6 +89,26 @@ bool FComponentInstanceDataBase::MatchesComponent(const UActorComponent* Compone
 	return bMatches;
 }
 
+void FComponentInstanceDataBase::ApplyToComponent(UActorComponent* Component)
+{
+	if (SavedProperties.Num() > 0)
+	{
+		class FComponentPropertyReader : public FObjectReader
+		{
+		public:
+			FComponentPropertyReader(TArray<uint8>& InBytes)
+				: FObjectReader(InBytes)
+			{
+			}
+		} ComponentPropertyReader(SavedProperties);
+
+		UObject* ArchetypeToSearch = Component->GetOuter()->GetArchetype();
+		UClass* Class = Component->GetClass();
+
+		Class->SerializeTaggedProperties(ComponentPropertyReader, (uint8*)Component, Class, nullptr);
+	}
+}
+
 FComponentInstanceDataCache::FComponentInstanceDataCache(const AActor* Actor)
 {
 	if(Actor != NULL)
diff --git a/Engine/Source/Runtime/Engine/Private/SceneComponent.cpp b/Engine/Source/Runtime/Engine/Private/SceneComponent.cpp
--- a/Engine/Source/Runtime/Engine/Private/SceneComponent.cpp
+++ b/Engine/Source/Runtime/Engine/Private/SceneComponent.cpp
@@ -1165,6 +1165,8 @@ FSceneComponentInstanceData::FSceneComponentInstanceData(const USceneComponent*
 
 void FSceneComponentInstanceData::ApplyToComponent(UActorComponent* Component)
 {
+	FComponentInstanceDataBase::ApplyToComponent(Component);
+
 	USceneComponent* SceneComponent = CastChecked<USceneComponent>(Component);
 	for (USceneComponent* ChildComponent : AttachedInstanceComponents)
 	{
diff --git a/Engine/Source/Runtime/Engine/Public/ComponentInstanceDataCache.h b/Engine/Source/Runtime/Engine/Public/ComponentInstanceDataCache.h
--- a/Engine/Source/Runtime/Engine/Public/ComponentInstanceDataCache.h
+++ b/Engine/Source/Runtime/Engine/Public/ComponentInstanceDataCache.h
@@ -23,7 +23,7 @@ public:
 	bool MatchesComponent(const UActorComponent* Component) const;
 
 	/** Applies this component instance data to the supplied component */
-	virtual void ApplyToComponent(UActorComponent* Component) = 0;
+	virtual void ApplyToComponent(UActorComponent* Component);
 
 	/** Replaces any references to old instances during Actor reinstancing */
 	virtual void FindAndReplaceInstances(const TMap<UObject*, UObject*>& OldToNewInstanceMap) { };
@@ -38,6 +38,8 @@ protected:
 	/** The index of the source component in its owner's serialized array 
 		when filtered to just that component type */
 	int32 SourceComponentTypeSerializedIndex;
+
+	TArray<uint8> SavedProperties;
 };
 
 /** 
