[INTEGRATE] Change 2431193 by [REDACTED_PERSON][REDACTED_HANDLE] on 2015/02/04 00:36:26

	Fix crashes with blueprints that do not have a SimpleConstructionScript, but do use AddComponent nodes.

diff --git a/Engine/Source/Runtime/Engine/Private/ActorComponent.cpp b/Engine/Source/Runtime/Engine/Private/ActorComponent.cpp
--- a/Engine/Source/Runtime/Engine/Private/ActorComponent.cpp
+++ b/Engine/Source/Runtime/Engine/Private/ActorComponent.cpp
@@ -139,8 +139,8 @@ void UActorComponent::PostLoad()
 
 	if (CreationMethod == EComponentCreationMethod::SimpleConstructionScript)
 	{
-		UBlueprintGeneratedClass* Class = CastChecked<UBlueprintGeneratedClass>(GetOuter()->GetClass());
-		if (Class->SimpleConstructionScript->FindSCSNode(GetFName()) == nullptr)
+		USimpleConstructionScript* SCS = CastChecked<UBlueprintGeneratedClass>(GetOuter()->GetClass())->SimpleConstructionScript;
+		if (SCS == nullptr || SCS->FindSCSNode(GetFName()) == nullptr)
 		{
 			CreationMethod = EComponentCreationMethod::UserConstructionScript;
 		}
diff --git a/Engine/Source/Runtime/Engine/Private/BlueprintGeneratedClass.cpp b/Engine/Source/Runtime/Engine/Private/BlueprintGeneratedClass.cpp
--- a/Engine/Source/Runtime/Engine/Private/BlueprintGeneratedClass.cpp
+++ b/Engine/Source/Runtime/Engine/Private/BlueprintGeneratedClass.cpp
@@ -259,17 +259,21 @@ UObject* UBlueprintGeneratedClass::FindArchetype(UClass* ArchetypeClass, const F
 {
 	UObject* Archetype = nullptr;
 
-	if (SimpleConstructionScript->HasAllFlags(RF_NeedLoad))
+	if (SimpleConstructionScript)
 	{
-		SimpleConstructionScript->PreloadChain();
-	}
+		if (SimpleConstructionScript->HasAllFlags(RF_NeedLoad))
+		{
+			SimpleConstructionScript->PreloadChain();
+		}
 
-	USCS_Node* SCSNode = SimpleConstructionScript->FindSCSNode(ArchetypeName);
-	if (SCSNode)
-	{
-		Archetype = SCSNode->ComponentTemplate;
+		USCS_Node* SCSNode = SimpleConstructionScript->FindSCSNode(ArchetypeName);
+		if (SCSNode)
+		{
+			Archetype = SCSNode->ComponentTemplate;
+		}
 	}
-	else
+
+	if (Archetype == nullptr)
 	{
 		Archetype = FindComponentTemplateByName(ArchetypeName);
 	}
