Update the ComponentToWorld of destructible component when a non-fractured destructible component is simulating

diff --git a/Engine/Source/Runtime/Engine/Classes/Components/DestructibleComponent.h b/Engine/Source/Runtime/Engine/Classes/Components/DestructibleComponent.h
--- a/Engine/Source/Runtime/Engine/Classes/Components/DestructibleComponent.h
+++ b/Engine/Source/Runtime/Engine/Classes/Components/DestructibleComponent.h
@@ -200,6 +200,8 @@ public:
 	void OnVisibilityEvent(const physx::apex::NxApexChunkStateEventData & InDamageEvent);
 #endif // WITH_APEX
 
+	bool IsFractured() const;
+
 	// End DestructibleComponent interface.
 
 	virtual bool DoCustomNavigableGeometryExport(struct FNavigableGeometryExport* GeomExport) const override;
diff --git a/Engine/Source/Runtime/Engine/Private/DestructibleComponent.cpp b/Engine/Source/Runtime/Engine/Private/DestructibleComponent.cpp
--- a/Engine/Source/Runtime/Engine/Private/DestructibleComponent.cpp
+++ b/Engine/Source/Runtime/Engine/Private/DestructibleComponent.cpp
@@ -681,6 +681,28 @@ void UDestructibleComponent::OnVisibilityEvent(const NxApexChunkStateEventData &
 }
 #endif // WITH_APEX
 
+bool UDestructibleComponent::IsFractured() const
+{
+	bool bFractured = false;
+#if WITH_APEX
+	if (ApexDestructibleActor)
+	{
+		//If we have only one chunk and its index is 0 we are NOT fractured. Otherwise we must have fractured
+		const physx::PxU32 VisibleChunkCount = ApexDestructibleActor->getNumVisibleChunks();
+		if (VisibleChunkCount == 1)
+		{
+			const physx::PxU16* VisibleChunks = ApexDestructibleActor->getVisibleChunks();
+			bFractured = *VisibleChunks != 0;
+		}
+		else
+		{
+			bFractured = true;
+		}
+	}
+#endif
+	return bFractured;
+}
+
 void UDestructibleComponent::RefreshBoneTransforms(FActorComponentTickFunction* TickFunction)
 {
 }
@@ -904,7 +926,15 @@ void UDestructibleComponent::UpdateDestructibleChunkTM(const TArray<const PxRigi
 	{
 		UDestructibleComponent* DestructibleComponent = It.Key();
 		TArray<FUpdateChunksInfo>& UpdateInfos = It.Value();
-		DestructibleComponent->SetChunksWorldTM(UpdateInfos);	
+		if (DestructibleComponent->IsFractured())
+		{
+			DestructibleComponent->SetChunksWorldTM(UpdateInfos);
+		}
+		else
+		{
+			//if we haven't fractured it must mean that we're simulating a destructible and so we should update our ComponentToWorld based on the single rigid body
+			DestructibleComponent->SyncComponentToRBPhysics();
+		}
 	}
 
 }
