[INTEGRATE] Change 2432567 by [REDACTED_PERSON][REDACTED_HANDLE] on 2015/02/04 17:49:38

	Reenable FindArchetype
	Prevent components based on LevelScriptActor from using the construction script templates as archetypes
	Fix bug improperly converting CreationType of components created by the SimpleConstructionScript of a parent class of the Actor type to  UserConstructionScript.

diff --git a/Engine/Source/Runtime/CoreUObject/Private/UObject/UObjectArchetype.cpp b/Engine/Source/Runtime/CoreUObject/Private/UObject/UObjectArchetype.cpp
--- a/Engine/Source/Runtime/CoreUObject/Private/UObject/UObjectArchetype.cpp
+++ b/Engine/Source/Runtime/CoreUObject/Private/UObject/UObjectArchetype.cpp
@@ -46,7 +46,7 @@ UObject* UObject::GetArchetypeFromRequiredInfo(UClass* Class, UObject* Outer, FN
 			}
 			else
 			{
-				//Result = ArchetypeToSearch->GetClass()->FindArchetype(Class, Name);
+				Result = ArchetypeToSearch->GetClass()->FindArchetype(Class, Name);
 			}
 		}
 
diff --git a/Engine/Source/Runtime/Engine/Private/ActorComponent.cpp b/Engine/Source/Runtime/Engine/Private/ActorComponent.cpp
--- a/Engine/Source/Runtime/Engine/Private/ActorComponent.cpp
+++ b/Engine/Source/Runtime/Engine/Private/ActorComponent.cpp
@@ -139,10 +139,22 @@ void UActorComponent::PostLoad()
 
 	if (CreationMethod == EComponentCreationMethod::SimpleConstructionScript)
 	{
-		USimpleConstructionScript* SCS = CastChecked<UBlueprintGeneratedClass>(GetOuter()->GetClass())->SimpleConstructionScript;
-		if (SCS == nullptr || SCS->FindSCSNode(GetFName()) == nullptr)
+		UBlueprintGeneratedClass* Class = CastChecked<UBlueprintGeneratedClass>(GetOuter()->GetClass());
+		while (Class)
 		{
-			CreationMethod = EComponentCreationMethod::UserConstructionScript;
+			USimpleConstructionScript* SCS = Class->SimpleConstructionScript;
+			if (SCS != nullptr && SCS->FindSCSNode(GetFName()))
+			{
+				break;
+			}
+			else
+			{
+				Class = Cast<UBlueprintGeneratedClass>(Class->GetSuperClass());
+				if (Class == nullptr)
+				{
+					CreationMethod = EComponentCreationMethod::UserConstructionScript;
+				}
+			}
 		}
 	}
 
diff --git a/Engine/Source/Runtime/Engine/Private/BlueprintGeneratedClass.cpp b/Engine/Source/Runtime/Engine/Private/BlueprintGeneratedClass.cpp
--- a/Engine/Source/Runtime/Engine/Private/BlueprintGeneratedClass.cpp
+++ b/Engine/Source/Runtime/Engine/Private/BlueprintGeneratedClass.cpp
@@ -6,6 +6,7 @@
 #include "Engine/TimelineTemplate.h"
 #include "Engine/SimpleConstructionScript.h"
 #include "Engine/SCS_Node.h"
+#include "Engine/LevelScriptActor.h"
 
 #if WITH_EDITOR
 #include "BlueprintEditorUtils.h"
@@ -266,24 +267,33 @@ UObject* UBlueprintGeneratedClass::FindArchetype(UClass* ArchetypeClass, const F
 {
 	UObject* Archetype = nullptr;
 
-	if (SimpleConstructionScript)
+	// There are some rogue LevelScriptActors that still have a SimpleConstructionScript
+	// and since preloading the SCS of a script in a world package is bad news, we need to filter them out
+	if (SimpleConstructionScript && !IsChildOf<ALevelScriptActor>())
 	{
 		if (SimpleConstructionScript->HasAllFlags(RF_NeedLoad))
 		{
 			SimpleConstructionScript->PreloadChain();
 		}
 
-		USCS_Node* SCSNode = SimpleConstructionScript->FindSCSNode(ArchetypeName);
-		if (SCSNode)
+		UBlueprintGeneratedClass* Class = const_cast<UBlueprintGeneratedClass*>(this);
+		while (Class)
 		{
-			Archetype = SCSNode->ComponentTemplate;
+			USCS_Node* SCSNode = Class->SimpleConstructionScript->FindSCSNode(ArchetypeName);
+			if (SCSNode)
+			{
+				Archetype = SCSNode->ComponentTemplate;
+				Class = nullptr;
+			}
+
+			if (Archetype == nullptr)
+			{
+				Archetype = static_cast<UObject*>(FindObjectWithOuter(Class, ArchetypeClass, ArchetypeName));
+				Class = (Archetype ? nullptr : Cast<UBlueprintGeneratedClass>(Class->GetSuperClass()));
+			}
 		}
 	}
 
-	if (Archetype == nullptr)
-	{
-		Archetype = FindComponentTemplateByName(ArchetypeName);
-	}
 
 	return Archetype;
 }
