[INTEGRATE] Change 2432634 by [REDACTED_PERSON][REDACTED_HANDLE] on 2015/02/04 18:27:32

	Fix duplication of components not having sprite components

diff --git a/Engine/Source/Runtime/Engine/Classes/Components/SceneComponent.h b/Engine/Source/Runtime/Engine/Classes/Components/SceneComponent.h
--- a/Engine/Source/Runtime/Engine/Classes/Components/SceneComponent.h
+++ b/Engine/Source/Runtime/Engine/Classes/Components/SceneComponent.h
@@ -158,7 +158,6 @@ protected:
 #if WITH_EDITORONLY_DATA
 protected:
 	/** Editor only component used to display the sprite so as to be able to see the location of the Audio Component  */
-	UPROPERTY(transient)
 	class UBillboardComponent* SpriteComponent;
 
 public:
@@ -525,6 +524,7 @@ public:
 	virtual void Serialize(FArchive& Ar) override;
 	virtual void PostInterpChange(UProperty* PropertyThatChanged) override;
 	virtual void BeginDestroy() override;
+	static void AddReferencedObjects(UObject* InThis, FReferenceCollector& Collector);
 	// End UObject Interface
 
 protected:
diff --git a/Engine/Source/Runtime/Engine/Private/SceneComponent.cpp b/Engine/Source/Runtime/Engine/Private/SceneComponent.cpp
--- a/Engine/Source/Runtime/Engine/Private/SceneComponent.cpp
+++ b/Engine/Source/Runtime/Engine/Private/SceneComponent.cpp
@@ -42,6 +42,13 @@ USceneComponent::USceneComponent(const FObjectInitializer& ObjectInitializer /*=
 	NetUpdateTransform = false;
 }
 
+void USceneComponent::AddReferencedObjects(UObject* InThis, FReferenceCollector& Collector)
+{
+	USceneComponent* This = CastChecked<USceneComponent>(InThis);
+	Collector.AddReferencedObject(This->SpriteComponent);
+	Super::AddReferencedObjects(InThis, Collector);
+}
+
 FTransform USceneComponent::CalcNewComponentToWorld(const FTransform& NewRelativeTransform, const USceneComponent* Parent) const
 {
 	Parent = Parent ? Parent : AttachParent;
