[INTEGRATE] Change 2432850 by [REDACTED_PERSON][REDACTED_HANDLE] on 2015/02/04 21:02:51

	Engine: Add short tooltips to various gameplay classes

diff --git a/Engine/Source/Runtime/Engine/Classes/Components/ActorComponent.h b/Engine/Source/Runtime/Engine/Classes/Components/ActorComponent.h
--- a/Engine/Source/Runtime/Engine/Classes/Components/ActorComponent.h
+++ b/Engine/Source/Runtime/Engine/Classes/Components/ActorComponent.h
@@ -25,7 +25,7 @@ enum class EComponentCreationMethod : uint8
  * @see USceneComponent
  * @see UPrimitiveComponent
  */
-UCLASS(DefaultToInstanced, BlueprintType, abstract, hidecategories=(ComponentReplication))
+UCLASS(DefaultToInstanced, BlueprintType, abstract, hideCategories=(ComponentReplication), meta=(ShortTooltip="An ActorComponent is a reusable component that can be added to any actor."))
 class ENGINE_API UActorComponent : public UObject, public IInterface_AssetUserData
 {
 	GENERATED_BODY()
diff --git a/Engine/Source/Runtime/Engine/Classes/Components/SceneComponent.h b/Engine/Source/Runtime/Engine/Classes/Components/SceneComponent.h
--- a/Engine/Source/Runtime/Engine/Classes/Components/SceneComponent.h
+++ b/Engine/Source/Runtime/Engine/Classes/Components/SceneComponent.h
@@ -76,7 +76,7 @@ DECLARE_DYNAMIC_MULTICAST_DELEGATE_OneParam(FPhysicsVolumeChanged, class APhysic
  * Useful as a 'dummy' component in the hierarchy to offset others.
  * @see [Scene Components](https://docs.unrealengine.com/latest/INT/Programming/UnrealArchitecture/Actors/Components/index.html#scenecomponents)
  */
-UCLASS(ClassGroup=Utility, BlueprintType, HideCategories=(Trigger, PhysicsVolume), meta=(BlueprintSpawnableComponent))
+UCLASS(ClassGroup=Utility, BlueprintType, hideCategories=(Trigger, PhysicsVolume), meta=(BlueprintSpawnableComponent, ShortTooltip="A Scene Component is a component that has a scene transform and can be attached to other scene components."))
 class ENGINE_API USceneComponent : public UActorComponent
 {
 	GENERATED_BODY()
diff --git a/Engine/Source/Runtime/Engine/Classes/GameFramework/Actor.h b/Engine/Source/Runtime/Engine/Classes/GameFramework/Actor.h
--- a/Engine/Source/Runtime/Engine/Classes/GameFramework/Actor.h
+++ b/Engine/Source/Runtime/Engine/Classes/GameFramework/Actor.h
@@ -54,7 +54,7 @@ DECLARE_CYCLE_STAT_EXTERN(TEXT("GetComponentsTime"),STAT_GetComponentsTime,STATG
  * @see https://docs.unrealengine.com/latest/INT/Programming/UnrealArchitecture/Actors/
  * @see UActorComponent
  */
-UCLASS(BlueprintType, Blueprintable, config=Engine)
+UCLASS(BlueprintType, Blueprintable, config=Engine, meta=(ShortTooltip="An Actor is an object that can be placed or spawned in the world."))
 class ENGINE_API AActor : public UObject
 {
 	/**
diff --git a/Engine/Source/Runtime/Engine/Classes/GameFramework/Character.h b/Engine/Source/Runtime/Engine/Classes/GameFramework/Character.h
--- a/Engine/Source/Runtime/Engine/Classes/GameFramework/Character.h
+++ b/Engine/Source/Runtime/Engine/Classes/GameFramework/Character.h
@@ -196,7 +196,7 @@ typedef FBasedMovementInfo FRepRelativeMovement;
  * @see https://docs.unrealengine.com/latest/INT/Gameplay/Framework/Pawn/Character/
  */ 
 
-UCLASS(config=Game, BlueprintType, hidecategories=("Pawn|Character|InternalEvents"))
+UCLASS(config=Game, BlueprintType, hideCategories=("Pawn|Character|InternalEvents"), meta=(ShortTooltip="A character is a type of Pawn that includes the ability to walk around."))
 class ENGINE_API ACharacter : public APawn
 {
 	GENERATED_BODY()
diff --git a/Engine/Source/Runtime/Engine/Classes/GameFramework/GameMode.h b/Engine/Source/Runtime/Engine/Classes/GameFramework/GameMode.h
--- a/Engine/Source/Runtime/Engine/Classes/GameFramework/GameMode.h
+++ b/Engine/Source/Runtime/Engine/Classes/GameFramework/GameMode.h
@@ -59,7 +59,7 @@ struct FGameClassShortName
  *
  * @see https://docs.unrealengine.com/latest/INT/Gameplay/Framework/GameMode/index.html
  */
-UCLASS(config=Game, notplaceable, BlueprintType, Blueprintable, Transient, hidecategories=(Info, Rendering, MovementReplication, Replication, Actor))
+UCLASS(config=Game, notplaceable, BlueprintType, Blueprintable, Transient, hideCategories=(Info, Rendering, MovementReplication, Replication, Actor), meta=(ShortTooltip="Game Mode defines the game being played, its rules, scoring, and other facets of the game type."))
 class ENGINE_API AGameMode : public AInfo
 {
 	GENERATED_UCLASS_BODY()
diff --git a/Engine/Source/Runtime/Engine/Classes/GameFramework/Pawn.h b/Engine/Source/Runtime/Engine/Classes/GameFramework/Pawn.h
--- a/Engine/Source/Runtime/Engine/Classes/GameFramework/Pawn.h
+++ b/Engine/Source/Runtime/Engine/Classes/GameFramework/Pawn.h
@@ -27,7 +27,7 @@ struct FNavAgentProperties;
  *
  * @see https://docs.unrealengine.com/latest/INT/Gameplay/Framework/Pawn/
  */
-UCLASS(config=Game, BlueprintType, Blueprintable, hidecategories=(Navigation))
+UCLASS(config=Game, BlueprintType, Blueprintable, hideCategories=(Navigation), meta=(ShortTooltip="A Pawn is an actor that can be 'possessed' and receieve input from a controller."))
 class ENGINE_API APawn : public AActor, public INavAgentInterface
 {
 	GENERATED_BODY()
diff --git a/Engine/Source/Runtime/Engine/Classes/GameFramework/PlayerController.h b/Engine/Source/Runtime/Engine/Classes/GameFramework/PlayerController.h
--- a/Engine/Source/Runtime/Engine/Classes/GameFramework/PlayerController.h
+++ b/Engine/Source/Runtime/Engine/Classes/GameFramework/PlayerController.h
@@ -143,7 +143,7 @@ protected:
  * @see https://docs.unrealengine.com/latest/INT/Gameplay/Framework/Controller/PlayerController/
  */
 
-UCLASS(config=Game, BlueprintType, Blueprintable)
+UCLASS(config=Game, BlueprintType, Blueprintable, meta=(ShortTooltip="A Player Controller is an actor responsible for controlling a Pawn used by the player."))
 class ENGINE_API APlayerController : public AController
 {
 	GENERATED_UCLASS_BODY()
