[INTEGRATE] Change 2432942 by [REDACTED_PERSON][REDACTED_HANDLE] on 2015/02/04 22:59:51

	Engine: Fix non-editor build failure due to referencing of USceneComponent::SpriteComponent

diff --git a/Engine/Source/Runtime/Engine/Private/SceneComponent.cpp b/Engine/Source/Runtime/Engine/Private/SceneComponent.cpp
--- a/Engine/Source/Runtime/Engine/Private/SceneComponent.cpp
+++ b/Engine/Source/Runtime/Engine/Private/SceneComponent.cpp
@@ -44,8 +44,10 @@ USceneComponent::USceneComponent(const FObjectInitializer& ObjectInitializer /*=
 
 void USceneComponent::AddReferencedObjects(UObject* InThis, FReferenceCollector& Collector)
 {
+#if WITH_EDITORONLY_DATA
 	USceneComponent* This = CastChecked<USceneComponent>(InThis);
 	Collector.AddReferencedObject(This->SpriteComponent);
+#endif
 	Super::AddReferencedObjects(InThis, Collector);
 }
 
