[INTEGRATE] Change 2432982 by [REDACTED_PERSON][REDACTED_HANDLE] on 2015/02/05 00:12:34

	Fix Mac and Linux builds due missing includes exposed by changed include order in CL# 2431692

diff --git a/Engine/Source/Runtime/Engine/Public/StaticBoundShaderState.h b/Engine/Source/Runtime/Engine/Public/StaticBoundShaderState.h
--- a/Engine/Source/Runtime/Engine/Public/StaticBoundShaderState.h
+++ b/Engine/Source/Runtime/Engine/Public/StaticBoundShaderState.h
@@ -4,8 +4,7 @@
 	StaticBoundShaderState.h: Static bound shader state definitions.
 =============================================================================*/
 
-#ifndef __STATICBOUNDSHADERSTATE_H__
-#define __STATICBOUNDSHADERSTATE_H__
+#pragma once
 
 #include "RenderResource.h"
 #include "Shader.h"
@@ -113,5 +112,3 @@ ENGINE_API void SetGlobalBoundShaderState(
 	FShader* PixelShader,
 	FShader* GeometryShader = nullptr
 	);
-
-#endif
