Fixed issue which was causing UBrushComponent::ComponentIsTouchingSelectionBox to return false positives.

diff --git a/Engine/Source/Runtime/Engine/Private/BrushComponent.cpp b/Engine/Source/Runtime/Engine/Private/BrushComponent.cpp
--- a/Engine/Source/Runtime/Engine/Private/BrushComponent.cpp
+++ b/Engine/Source/Runtime/Engine/Private/BrushComponent.cpp
@@ -619,9 +619,9 @@ bool UBrushComponent::ComponentIsTouchingSelectionBox(const FBox& InSelBBox, con
 	{
 		TArray<FVector> Vertices;
 
-		for (const auto& Poly : Brush->Polys->Element)
+		if (!bMustEncompassEntireComponent)
 		{
-			if (!bMustEncompassEntireComponent)
+			for (const auto& Poly : Brush->Polys->Element)
 			{
 				// Just an intersection will do...
 				Vertices.Empty(Poly.Vertices.Num());
@@ -637,7 +637,13 @@ bool UBrushComponent::ComponentIsTouchingSelectionBox(const FBox& InSelBBox, con
 					return true;
 				}
 			}
-			else
+
+			// No poly intersected with the bounding box
+			return false;
+		}
+		else
+		{
+			for (const auto& Poly : Brush->Polys->Element)
 			{
 				// The component must be entirely within the bounding box...
 				for (const auto& Vertex : Poly.Vertices)
@@ -653,11 +659,10 @@ bool UBrushComponent::ComponentIsTouchingSelectionBox(const FBox& InSelBBox, con
 					}
 				}
 			}
-		}
 
-		// If the selection box has to encompass all of the component and none of the component's verts failed the intersection test, this component
-		// is considered touching
-		return true;
+			// All points lay within the selection box
+			return true;
+		}
 	}
 
 	return false;
