Merging  TimerHandle_LifeSpanExpired going to protected

//depot/UE4-UT/...

to //depot/UE4/...

diff --git a/Engine/Source/Runtime/Engine/Classes/GameFramework/Actor.h b/Engine/Source/Runtime/Engine/Classes/GameFramework/Actor.h
--- a/Engine/Source/Runtime/Engine/Classes/GameFramework/Actor.h
+++ b/Engine/Source/Runtime/Engine/Classes/GameFramework/Actor.h
@@ -364,8 +364,6 @@ protected:
 	UPROPERTY(EditAnywhere, BlueprintReadOnly, Category=Actor)
 	float InitialLifeSpan; 
 
-private:
-
 	/** Handle for efficient management of LifeSpanExpired timer */
 	FTimerHandle TimerHandle_LifeSpanExpired;
 
