#UE4 PIE net connection timeout is now properly disabled even in the first few ticks of the connection. This prevents long load times from causing a timeout before the net connection has been given a world.

diff --git a/Engine/Source/Runtime/Engine/Private/NetConnection.cpp b/Engine/Source/Runtime/Engine/Private/NetConnection.cpp
--- a/Engine/Source/Runtime/Engine/Private/NetConnection.cpp
+++ b/Engine/Source/Runtime/Engine/Private/NetConnection.cpp
@@ -1388,7 +1388,8 @@ void UNetConnection::Tick()
 	}
 	bool bUseTimeout = true;
 #if WITH_EDITOR
-	bUseTimeout = !(Driver->GetWorld() && Driver->GetWorld()->WorldType == EWorldType::PIE);
+	// Do not time out in PIE since the server is local.
+	bUseTimeout = !(GEditor && GEditor->PlayWorld);
 #endif
 	if ( bUseTimeout && Driver->Time - LastReceiveTime > Timeout )
 	{
