Fix destructible not turning on hit events if an impact depth override is used

diff --git a/Engine/Source/Runtime/Engine/Private/DestructibleComponent.cpp b/Engine/Source/Runtime/Engine/Private/DestructibleComponent.cpp
--- a/Engine/Source/Runtime/Engine/Private/DestructibleComponent.cpp
+++ b/Engine/Source/Runtime/Engine/Private/DestructibleComponent.cpp
@@ -249,9 +249,15 @@ void UDestructibleComponent::CreatePhysicsState()
 		SmallChunkCollisionResponse.SetAllChannels(ECR_Ignore);
 	}
 
+	bool bEnableImpactDamage = TheDestructibleMesh->DefaultDestructibleParameters.DamageParameters.bEnableImpactDamage;
+	for (const FDestructibleDepthParameters& DepthParams : TheDestructibleMesh->DefaultDestructibleParameters.DepthParameters)
+	{
+		bEnableImpactDamage |= DepthParams.ImpactDamageOverride == EImpactDamageOverride::IDO_On;
+	}
+
 	// Passing AssetInstanceID = 0 so we'll have self-collision
 	AActor* Owner = GetOwner();
-	CreateShapeFilterData(MoveChannel, GetUniqueID(), CollResponse, 0, 0, PQueryFilterData, PSimFilterData, BodyInstance.bUseCCD, TheDestructibleMesh->DefaultDestructibleParameters.DamageParameters.bEnableImpactDamage, false);
+	CreateShapeFilterData(MoveChannel, GetUniqueID(), CollResponse, 0, 0, PQueryFilterData, PSimFilterData, BodyInstance.bUseCCD, bEnableImpactDamage, false);
 
 	// Build filterData variations for complex and simple
 	PSimFilterData.word3 |= EPDF_SimpleCollision | EPDF_ComplexCollision;
