#UE4 Fix non-unity error

diff --git a/Engine/Source/Runtime/Engine/Private/NetConnection.cpp b/Engine/Source/Runtime/Engine/Private/NetConnection.cpp
--- a/Engine/Source/Runtime/Engine/Private/NetConnection.cpp
+++ b/Engine/Source/Runtime/Engine/Private/NetConnection.cpp
@@ -13,6 +13,10 @@
 #include "Engine/PackageMapClient.h"
 #include "GameFramework/GameMode.h"
 
+#if WITH_EDITOR
+#include "UnrealEd.h"
+#endif
+
 /*-----------------------------------------------------------------------------
 	UNetConnection implementation.
 -----------------------------------------------------------------------------*/
