Fix issue where Atmospheric Fog parameters were never being precomputed

diff --git a/Engine/Source/Runtime/Engine/Private/Atmosphere/Atmosphere.cpp b/Engine/Source/Runtime/Engine/Private/Atmosphere/Atmosphere.cpp
--- a/Engine/Source/Runtime/Engine/Private/Atmosphere/Atmosphere.cpp
+++ b/Engine/Source/Runtime/Engine/Private/Atmosphere/Atmosphere.cpp
@@ -502,13 +502,14 @@ void UAtmosphericFogComponent::PostEditChangeProperty(FPropertyChangedEvent& Pro
 {
 	const FName CategoryName = FObjectEditorUtils::GetCategoryFName(PropertyChangedEvent.Property);
 
+	bool bNeedsPrecompute = false;
 	if( CategoryName == FName(TEXT("AtmosphereParam")) )
 	{
 		// Recompute when precompute parameters were changed
 		PrecomputeParams.DensityHeight = FMath::Clamp(PrecomputeParams.DensityHeight, 0.1f, 1.f);
 		PrecomputeParams.MaxScatteringOrder = FMath::Clamp(PrecomputeParams.MaxScatteringOrder, 1, 4);
 		PrecomputeParams.InscatterAltitudeSampleNum = FMath::Clamp(PrecomputeParams.InscatterAltitudeSampleNum, 2, 32);
-		StartPrecompute();
+		bNeedsPrecompute = true;
 	}
 	else
 	{
@@ -524,6 +525,11 @@ void UAtmosphericFogComponent::PostEditChangeProperty(FPropertyChangedEvent& Pro
 	}
 
 	Super::PostEditChangeProperty(PropertyChangedEvent);
+
+	if (bNeedsPrecompute)
+	{
+		StartPrecompute();
+	}
 }
 #endif // WITH_EDITOR
 
