UE4 - Integrate 2439067
Add FMaterialVertexParameters.PrevWorldPosition to material nodes using world position

diff --git a/Engine/Shaders/GpuSkinVertexFactory.usf b/Engine/Shaders/GpuSkinVertexFactory.usf
--- a/Engine/Shaders/GpuSkinVertexFactory.usf
+++ b/Engine/Shaders/GpuSkinVertexFactory.usf
@@ -207,6 +207,7 @@ FMaterialVertexParameters GetMaterialVertexParameters(FVertexFactoryInput Input,
 {
 	FMaterialVertexParameters Result = (FMaterialVertexParameters)0;
 	Result.WorldPosition = WorldPosition - View.PreViewTranslation.xyz;
+	Result.PrevWorldPosition = WorldPosition - View.PrevPreViewTranslation.xyz;
 	Result.VertexColor = Intermediates.Color;
 	Result.TangentToWorld = mul(TangentToLocal, GetLocalToWorld3x3());
 
diff --git a/Engine/Shaders/LandscapeVertexFactory.usf b/Engine/Shaders/LandscapeVertexFactory.usf
--- a/Engine/Shaders/LandscapeVertexFactory.usf
+++ b/Engine/Shaders/LandscapeVertexFactory.usf
@@ -373,6 +373,7 @@ FMaterialVertexParameters GetMaterialVertexParameters(FVertexFactoryInput Input,
 {
 	FMaterialVertexParameters Result = (FMaterialVertexParameters)0;
 	Result.WorldPosition = WorldPosition - View.PreViewTranslation;
+	Result.PrevWorldPosition = WorldPosition - View.PrevPreViewTranslation.xyz;
 	Result.VertexColor = float4(1,1,1,1);
 	Result.TangentToWorld = mul(TangentToLocal, GetLocalToWorld3x3()); 
 
diff --git a/Engine/Shaders/LocalVertexFactory.usf b/Engine/Shaders/LocalVertexFactory.usf
--- a/Engine/Shaders/LocalVertexFactory.usf
+++ b/Engine/Shaders/LocalVertexFactory.usf
@@ -283,6 +283,7 @@ FMaterialVertexParameters GetMaterialVertexParameters(FVertexFactoryInput Input,
 {
 	FMaterialVertexParameters Result = (FMaterialVertexParameters)0;
 	Result.WorldPosition = WorldPosition - View.PreViewTranslation;
+	Result.PrevWorldPosition = WorldPosition - View.PrevPreViewTranslation.xyz;
 	Result.VertexColor = Intermediates.Color;
 
 	// does not handle instancing!
diff --git a/Engine/Shaders/MaterialTemplate.usf b/Engine/Shaders/MaterialTemplate.usf
--- a/Engine/Shaders/MaterialTemplate.usf
+++ b/Engine/Shaders/MaterialTemplate.usf
@@ -226,6 +226,7 @@ struct FMaterialTessellationParameters
 struct FMaterialVertexParameters
 {
 	float3 WorldPosition;
+	float3 PrevWorldPosition;
 	// TangentToWorld[2] is WorldVertexNormal
 	half3x3 TangentToWorld;
 #if USE_INSTANCING
diff --git a/Engine/Shaders/MeshParticleVertexFactory.usf b/Engine/Shaders/MeshParticleVertexFactory.usf
--- a/Engine/Shaders/MeshParticleVertexFactory.usf
+++ b/Engine/Shaders/MeshParticleVertexFactory.usf
@@ -219,6 +219,7 @@ FMaterialVertexParameters GetMaterialVertexParameters(FVertexFactoryInput Input,
 {
 	FMaterialVertexParameters Result = (FMaterialVertexParameters)0;
 	Result.WorldPosition = WorldPosition - View.PreViewTranslation;
+	Result.PrevWorldPosition = WorldPosition - View.PrevPreViewTranslation.xyz;
 	Result.VertexColor = Intermediates.VertexColor;
 
 	Result.Particle.PositionAndSize.xyz = Intermediates.ParticleWorldPosition.xyz;
diff --git a/Engine/Shaders/ParticleBeamTrailVertexFactory.usf b/Engine/Shaders/ParticleBeamTrailVertexFactory.usf
--- a/Engine/Shaders/ParticleBeamTrailVertexFactory.usf
+++ b/Engine/Shaders/ParticleBeamTrailVertexFactory.usf
@@ -92,6 +92,7 @@ FMaterialVertexParameters GetMaterialVertexParameters(FVertexFactoryInput Input,
 {
 	FMaterialVertexParameters Result = (FMaterialVertexParameters)0;
 	Result.WorldPosition = WorldPosition - View.PreViewTranslation.xyz;
+	Result.PrevWorldPosition = WorldPosition - View.PrevPreViewTranslation.xyz;
 	Result.VertexColor = Input.Color;
 	Result.TangentToWorld = mul(TangentToLocal, GetLocalToWorld3x3()); 
 	Result.Particle.Color = Intermediates.Color;
diff --git a/Engine/Shaders/ParticleGPUSpriteVertexFactory.usf b/Engine/Shaders/ParticleGPUSpriteVertexFactory.usf
--- a/Engine/Shaders/ParticleGPUSpriteVertexFactory.usf
+++ b/Engine/Shaders/ParticleGPUSpriteVertexFactory.usf
@@ -411,6 +411,7 @@ FMaterialVertexParameters GetMaterialVertexParameters(FVertexFactoryInput Input,
 {
 	FMaterialVertexParameters Result = (FMaterialVertexParameters)0;
 	Result.WorldPosition = WorldPosition - View.PreViewTranslation.xyz;
+	Result.PrevWorldPosition = WorldPosition - View.PrevPreViewTranslation.xyz;
 	Result.VertexColor = 1;
 	Result.TangentToWorld = mul(TangentToLocal, GetLocalToWorld3x3()); 
 
diff --git a/Engine/Shaders/ParticleSpriteVertexFactory.usf b/Engine/Shaders/ParticleSpriteVertexFactory.usf
--- a/Engine/Shaders/ParticleSpriteVertexFactory.usf
+++ b/Engine/Shaders/ParticleSpriteVertexFactory.usf
@@ -231,6 +231,7 @@ FMaterialVertexParameters GetMaterialVertexParameters(FVertexFactoryInput Input,
 {
 	FMaterialVertexParameters Result = (FMaterialVertexParameters)0;
 	Result.WorldPosition = WorldPosition - View.PreViewTranslation.xyz;
+	Result.PrevWorldPosition = WorldPosition - View.PrevPreViewTranslation.xyz;
 	Result.VertexColor = Input.Color;
 	Result.TangentToWorld = mul(TangentToLocal, GetLocalToWorld3x3()); 
 	Result.Particle.MacroUV = SpriteVF.MacroUVParameters;
diff --git a/Engine/Shaders/VectorFieldVisualizationVertexFactory.usf b/Engine/Shaders/VectorFieldVisualizationVertexFactory.usf
--- a/Engine/Shaders/VectorFieldVisualizationVertexFactory.usf
+++ b/Engine/Shaders/VectorFieldVisualizationVertexFactory.usf
@@ -110,6 +110,7 @@ FMaterialVertexParameters GetMaterialVertexParameters(FVertexFactoryInput Input,
 {
 	FMaterialVertexParameters Result = (FMaterialVertexParameters)0;
 	Result.WorldPosition = WorldPosition - View.PreViewTranslation.xyz;
+	Result.PrevWorldPosition = WorldPosition - View.PrevPreViewTranslation.xyz;
 	Result.VertexColor = float4(1,1,1,1);
 	Result.TangentToWorld = mul(TangentToLocal, GetLocalToWorld3x3()); 
 	return Result;
diff --git a/Engine/Source/Runtime/Engine/Private/Materials/HLSLMaterialTranslator.h b/Engine/Source/Runtime/Engine/Private/Materials/HLSLMaterialTranslator.h
--- a/Engine/Source/Runtime/Engine/Private/Materials/HLSLMaterialTranslator.h
+++ b/Engine/Source/Runtime/Engine/Private/Materials/HLSLMaterialTranslator.h
@@ -2093,8 +2093,16 @@ protected:
 		}
 		else
 		{
-			// If not in a pixel shader, only the normal WorldPosition is available
-			return AddInlinedCodeChunk(MCT_Float3, TEXT("Parameters.WorldPosition"));
+			if (bCompilingPreviousFrame && ShaderFrequency == SF_Vertex)
+			{
+				// Only if not in a pixel shader...
+				return AddInlinedCodeChunk(MCT_Float3, TEXT("Parameters.PrevWorldPosition"));
+			}
+			else
+			{
+				// If not in a pixel shader, only the normal WorldPosition is available
+				return AddInlinedCodeChunk(MCT_Float3, TEXT("Parameters.WorldPosition"));
+			}
 		}
 	}
 
