Fix module dependencies for NetworkReplayStreaming and clean up #include

diff --git a/Engine/Source/Runtime/Engine/Engine.Build.cs b/Engine/Source/Runtime/Engine/Engine.Build.cs
--- a/Engine/Source/Runtime/Engine/Engine.Build.cs
+++ b/Engine/Source/Runtime/Engine/Engine.Build.cs
@@ -28,7 +28,8 @@ public class Engine : ModuleRules
 				"TargetPlatform",
 				"ImageWrapper",
 				"HeadMountedDisplay",
-				"Advertising"
+				"Advertising",
+				"NetworkReplayStreaming"
 			}
 		);
 
@@ -173,6 +174,7 @@ public class Engine : ModuleRules
 				"AnalyticsET",
 				"NetworkReplayStreaming",
 				"NullNetworkReplayStreaming",
+				"HttpNetworkReplayStreaming",
 				"OnlineSubsystem", 
 				"OnlineSubsystemUtils",
 				"Advertising"
diff --git a/Engine/Source/Runtime/Engine/Private/DemoNetDriver.cpp b/Engine/Source/Runtime/Engine/Private/DemoNetDriver.cpp
--- a/Engine/Source/Runtime/Engine/Private/DemoNetDriver.cpp
+++ b/Engine/Source/Runtime/Engine/Private/DemoNetDriver.cpp
@@ -15,7 +15,7 @@
 #include "GameFramework/PlayerStart.h"
 #include "GameFramework/SpectatorPawnMovement.h"
 #include "Engine/GameInstance.h"
-#include "Runtime/NetworkReplayStreaming/NetworkReplayStreaming/Public/NetworkReplayStreaming.h"	// Why do I need the full path here?
+#include "NetworkReplayStreaming.h"
 
 DEFINE_LOG_CATEGORY_STATIC( LogDemo, Log, All );
 
