Merging Ms fixes to PostLoad spikes

diff --git a/Engine/Source/Runtime/CoreUObject/Private/Serialization/BulkData.cpp b/Engine/Source/Runtime/CoreUObject/Private/Serialization/BulkData.cpp
--- a/Engine/Source/Runtime/CoreUObject/Private/Serialization/BulkData.cpp
+++ b/Engine/Source/Runtime/CoreUObject/Private/Serialization/BulkData.cpp
@@ -691,15 +691,15 @@ void FUntypedBulkData::Serialize( FArchive& Ar, UObject* Owner, int32 Idx )
 			
 			Ar << BulkDataOffsetInFile;
 
-			// fix up the file offset 
-			if (Owner != NULL && Owner->GetLinker())
+			// determine whether the payload is stored inline or at the end of the file
+			bool bPayloadInline = !(BulkDataFlags&BULKDATA_PayloadAtEndOfFile);
+
+			// fix up the file offset, but only if not stored inline
+			if (Owner != NULL && Owner->GetLinker() && !bPayloadInline)
 			{
 				BulkDataOffsetInFile += Owner->GetLinker()->Summary.BulkDataStartOffset;
 			}
 
-			// determine whether the payload is stored inline or at the end of the file
-			bool bPayloadInline = !(BulkDataFlags&BULKDATA_PayloadAtEndOfFile);
-
 // 			check( (bPayloadInline && BulkDataOffsetInFile == Ar.Tell()) || 
 // 				   (!bPayloadInline && BulkDataOffsetInFile > Ar.Tell()));
 			
@@ -720,7 +720,15 @@ void FUntypedBulkData::Serialize( FArchive& Ar, UObject* Owner, int32 Idx )
 				// only skip over payload, if it's stored inline
 				if (bPayloadInline)
 				{
-					Ar.Seek( Ar.Tell() + BulkDataSizeOnDisk );
+					// Force non-lazy loading of inline bulk data to prevent PostLoad spikes.
+					// Memory for bulk data can come from preallocated GPU-accessible resource memory or default to system memory
+					BulkData = GetBulkDataResourceMemory(Owner, Idx);
+					if (!BulkData)
+					{
+						BulkData = FMemory::Realloc(BulkData, GetBulkDataSize());
+					}
+					// if the payload is stored inline, just serialize it
+					SerializeBulkData(Ar, BulkData);
 				}
 			}
 			// Serialize the bulk data right away.
@@ -1308,6 +1316,8 @@ void FFormatContainer::Serialize(FArchive& Ar, UObject* Owner, const TArray<FNam
 				FName Name = It.Key();
 				Ar << Name;
 				FByteBulkData* Bulk = It.Value();
+				// Force this kind of bulk data (physics, etc) to be stored inline for streaming
+				Bulk->SetBulkDataFlags(BULKDATA_ForceInlinePayload | BULKDATA_SingleUse);
 				check(Bulk);
 				Bulk->Serialize(Ar, Owner);
 			}
diff --git a/Engine/Source/Runtime/Engine/Classes/Engine/Texture.h b/Engine/Source/Runtime/Engine/Classes/Engine/Texture.h
--- a/Engine/Source/Runtime/Engine/Classes/Engine/Texture.h
+++ b/Engine/Source/Runtime/Engine/Classes/Engine/Texture.h
@@ -307,8 +307,14 @@ struct FTexturePlatformData
 	/** Serialization. */
 	void Serialize(FArchive& Ar, class UTexture* Owner);
 
-	/** Serialization for cooked builds. */
-	void SerializeCooked(FArchive& Ar, class UTexture* Owner);
+	/** 
+	 * Serialization for cooked builds.
+	 *
+	 * @param Ar Archive to serialize with
+	 * @param Owner Owner texture
+	 * @param bStreamable Store some mips inline, only used during cooking
+	 */
+	void SerializeCooked(FArchive& Ar, class UTexture* Owner, bool bStreamable);
 
 #if WITH_EDITOR
 	void Cache(
diff --git a/Engine/Source/Runtime/Engine/Private/TextureDerivedData.cpp b/Engine/Source/Runtime/Engine/Private/TextureDerivedData.cpp
--- a/Engine/Source/Runtime/Engine/Private/TextureDerivedData.cpp
+++ b/Engine/Source/Runtime/Engine/Private/TextureDerivedData.cpp
@@ -1102,6 +1102,18 @@ bool FTexturePlatformData::TryLoadMips(int32 FirstMipToLoad, void** OutMipData)
 	BeginLoadDerivedMips(Mips, FirstMipToLoad, AsyncHandles);
 #endif // #if WITH_EDITOR
 
+	// Handle the case where we inlined more mips than we intend to keep resident
+	// Discard unneeded mips
+	for (int32 MipIndex = 0; MipIndex < FirstMipToLoad && MipIndex < Mips.Num(); ++MipIndex)
+	{
+		FTexture2DMipMap& Mip = Mips[MipIndex];
+		if (Mip.BulkData.IsBulkDataLoaded())
+		{
+			Mip.BulkData.Lock(LOCK_READ_ONLY);
+			Mip.BulkData.Unlock();
+		}
+	}
+
 	// Load remaining mips (if any) from bulk data.
 	for (int32 MipIndex = FirstMipToLoad; MipIndex < Mips.Num(); ++MipIndex)
 	{
@@ -1181,9 +1193,10 @@ bool FTexturePlatformData::AreDerivedMipsAvailable() const
 
 static void SerializePlatformData(
 	FArchive& Ar,
-	FTexturePlatformData* PlatformData, 
+	FTexturePlatformData* PlatformData,
 	UTexture* Texture,
-	bool bCooked
+	bool bCooked,
+	bool bStreamable
 	)
 {
 	DECLARE_SCOPE_CYCLE_COUNTER( TEXT("SerializePlatformData"), STAT_Texture_SerializePlatformData, STATGROUP_LoadTime );
@@ -1236,6 +1249,17 @@ static void SerializePlatformData(
 		}
 	}
 
+	// Force resident mips inline
+	if (bCooked && Ar.IsSaving())
+	{
+		int32 MinMipToInline = bStreamable ? NumMips - UTexture2D::GetMinTextureResidentMipCount() : 0;
+		MinMipToInline = FMath::Max(MinMipToInline, 0);
+		for (int32 MipIndex = MinMipToInline; MipIndex < NumMips; ++MipIndex)
+		{
+			PlatformData->Mips[MipIndex + FirstMipToSerialize].BulkData.SetBulkDataFlags(BULKDATA_ForceInlinePayload | BULKDATA_SingleUse);
+		}
+	}
+
 	Ar << NumMips;
 	if (Ar.IsLoading())
 	{
@@ -1254,12 +1278,16 @@ static void SerializePlatformData(
 
 void FTexturePlatformData::Serialize(FArchive& Ar, UTexture* Owner)
 {
-	SerializePlatformData(Ar, this, Owner, false);
+	const bool bCooking = false;
+	const bool bStreamable = false;
+	SerializePlatformData(Ar, this, Owner, bCooking, bStreamable);
 }
 
-void FTexturePlatformData::SerializeCooked(FArchive& Ar, UTexture* Owner)
+// @BLACK_TUSK_CHANGE: BEGIN - [REDACTED_EMAIL] - Pass streamable flag for inlining mips
+void FTexturePlatformData::SerializeCooked(FArchive& Ar, UTexture* Owner, bool bStreamable)
 {
-	SerializePlatformData(Ar, this, Owner, true);
+	SerializePlatformData(Ar, this, Owner, true, bStreamable);
+	// @BLACK_TUSK_CHANGE: END - [REDACTED_EMAIL] 
 	if (Ar.IsLoading() && Mips.Num() > 0)
 	{
 		SizeX = Mips[0].SizeX;
@@ -1726,7 +1754,8 @@ void UTexture::SerializeCookedPlatformData(FArchive& Ar)
 				int32 SkipOffsetLoc = Ar.Tell();
 				int32 SkipOffset = 0;
 				Ar << SkipOffset;
-				PlatformDataToSave->SerializeCooked(Ar, this);
+				// Pass streamable flag for inlining mips
+				PlatformDataToSave->SerializeCooked(Ar, this, BuildSettings.bStreamable);
 				SkipOffset = Ar.Tell();
 				Ar.Seek(SkipOffsetLoc);
 				Ar << SkipOffset;
@@ -1761,7 +1790,9 @@ void UTexture::SerializeCookedPlatformData(FArchive& Ar)
 			bool bFormatSupported = GPixelFormats[PixelFormat].Supported;
 			if (RunningPlatformData->PixelFormat == PF_Unknown && bFormatSupported)
 			{
-				RunningPlatformData->SerializeCooked(Ar, this);
+				// Extra arg is unused here because we're loading
+				const bool bStreamable = false;
+				RunningPlatformData->SerializeCooked(Ar, this, bStreamable);
 			}
 			else
 			{
diff --git a/Engine/Source/Runtime/Engine/Private/VectorField.cpp b/Engine/Source/Runtime/Engine/Private/VectorField.cpp
--- a/Engine/Source/Runtime/Engine/Private/VectorField.cpp
+++ b/Engine/Source/Runtime/Engine/Private/VectorField.cpp
@@ -299,6 +299,13 @@ void UVectorFieldStatic::ReleaseResource()
 void UVectorFieldStatic::Serialize(FArchive& Ar)
 {
 	Super::Serialize(Ar);
+
+	// Store bulk data inline for streaming (to prevent PostLoad spikes)
+	if (Ar.IsCooking())
+	{
+		SourceData.SetBulkDataFlags(BULKDATA_ForceInlinePayload | BULKDATA_SingleUse);
+	}
+
 	SourceData.Serialize(Ar,this);
 }
 
