Integrate 2440376 change from Orion (just demonetdriver.cpp)

diff --git a/Engine/Source/Runtime/Engine/Private/DemoNetDriver.cpp b/Engine/Source/Runtime/Engine/Private/DemoNetDriver.cpp
--- a/Engine/Source/Runtime/Engine/Private/DemoNetDriver.cpp
+++ b/Engine/Source/Runtime/Engine/Private/DemoNetDriver.cpp
@@ -389,7 +389,8 @@ void UDemoNetDriver::TickFlush( float DeltaSeconds )
 			for ( int32 i = 0; i < World->StreamingLevels.Num(); ++i )
 			{
 				ULevelStreaming * StreamingLevel = World->StreamingLevels[i];
-				if ( StreamingLevel != NULL && ( !StreamingLevel->IsLevelLoaded() || !StreamingLevel->GetLoadedLevel()->GetOutermost()->IsFullyLoaded() || !StreamingLevel->IsLevelVisible() ) )
+
+				if ( StreamingLevel != NULL && StreamingLevel->ShouldBeLoaded() && (!StreamingLevel->IsLevelLoaded() || !StreamingLevel->GetLoadedLevel()->GetOutermost()->IsFullyLoaded() || !StreamingLevel->IsLevelVisible() ) )
 				{
 					// Abort, we have more streaming levels to load
 					return;
