UE4 - Expose Static Mesh's SetForcedLodModel() on Blueprints (StaticMesh->Rendering->LOD)

diff --git a/Engine/Source/Runtime/Engine/Classes/Components/StaticMeshComponent.h b/Engine/Source/Runtime/Engine/Classes/Components/StaticMeshComponent.h
--- a/Engine/Source/Runtime/Engine/Classes/Components/StaticMeshComponent.h
+++ b/Engine/Source/Runtime/Engine/Classes/Components/StaticMeshComponent.h
@@ -206,6 +206,9 @@ class ENGINE_API UStaticMeshComponent : public UMeshComponent
 	UFUNCTION(BlueprintCallable, Category="Components|StaticMesh")
 	virtual bool SetStaticMesh(class UStaticMesh* NewMesh);
 
+	UFUNCTION(BlueprintCallable, Category="Rendering|LOD")
+	void SetForcedLodModel(int32 NewForcedLodModel);
+
 	/** 
 	 * Get Local bounds
 	 */
diff --git a/Engine/Source/Runtime/Engine/Private/StaticMeshComponent.cpp b/Engine/Source/Runtime/Engine/Private/StaticMeshComponent.cpp
--- a/Engine/Source/Runtime/Engine/Private/StaticMeshComponent.cpp
+++ b/Engine/Source/Runtime/Engine/Private/StaticMeshComponent.cpp
@@ -1216,6 +1216,15 @@ bool UStaticMeshComponent::SetStaticMesh(UStaticMesh* NewMesh)
 	return true;
 }
 
+void UStaticMeshComponent::SetForcedLodModel(int32 NewForcedLodModel)
+{
+	if (ForcedLodModel != NewForcedLodModel)
+	{
+		ForcedLodModel = NewForcedLodModel;
+		MarkRenderStateDirty();
+	}
+}
+
 void UStaticMeshComponent::GetLocalBounds(FVector& Min, FVector& Max) const
 {
 	if (StaticMesh)
