Introduce contact modify flag for body instance. This is needed for destruction impact resistance.

UE-8777

diff --git a/Engine/Source/Runtime/Engine/Private/Collision/PhysicsFiltering.h b/Engine/Source/Runtime/Engine/Private/Collision/PhysicsFiltering.h
--- a/Engine/Source/Runtime/Engine/Private/Collision/PhysicsFiltering.h
+++ b/Engine/Source/Runtime/Engine/Private/Collision/PhysicsFiltering.h
@@ -16,7 +16,8 @@ enum EPhysXFilterDataFlags
 	EPDF_ComplexCollision	=	0x0002,
 	EPDF_CCD				=	0x0004,
 	EPDF_ContactNotify		=	0x0008,
-	EPDF_StaticShape		=	0x0010
+	EPDF_StaticShape		=	0x0010,
+	EPDF_ModifyContacts		=   0x0020
 };
 
 struct FPhysicsFilterBuilder
@@ -86,12 +87,14 @@ inline void CreateShapeFilterData(
 	PxFilterData& OutSimData,
 	bool bEnableCCD,
 	bool bEnableContactNotify,
-	bool bStaticShape)
+	bool bStaticShape,
+	bool bModifyContacts = false)
 {
 	FPhysicsFilterBuilder Builder((ECollisionChannel)MyChannel, ResponseToChannels);
 	Builder.ConditionalSetFlags(EPDF_CCD, bEnableCCD);
 	Builder.ConditionalSetFlags(EPDF_ContactNotify, bEnableContactNotify);
 	Builder.ConditionalSetFlags(EPDF_StaticShape, bStaticShape);
+	Builder.ConditionalSetFlags(EPDF_ModifyContacts, bModifyContacts);
 
 	OutQueryData.setToDefault();
 	OutSimData.setToDefault();
diff --git a/Engine/Source/Runtime/Engine/Private/DestructibleComponent.cpp b/Engine/Source/Runtime/Engine/Private/DestructibleComponent.cpp
--- a/Engine/Source/Runtime/Engine/Private/DestructibleComponent.cpp
+++ b/Engine/Source/Runtime/Engine/Private/DestructibleComponent.cpp
@@ -255,9 +255,11 @@ void UDestructibleComponent::CreatePhysicsState()
 		bEnableImpactDamage |= DepthParams.ImpactDamageOverride == EImpactDamageOverride::IDO_On;
 	}
 
+	bool bEnableContactModification = TheDestructibleMesh->DefaultDestructibleParameters.DamageParameters.bCustomImpactResistance && TheDestructibleMesh->DefaultDestructibleParameters.DamageParameters.ImpactResistance > 0.f;
+
 	// Passing AssetInstanceID = 0 so we'll have self-collision
 	AActor* Owner = GetOwner();
-	CreateShapeFilterData(MoveChannel, GetUniqueID(), CollResponse, 0, 0, PQueryFilterData, PSimFilterData, BodyInstance.bUseCCD, bEnableImpactDamage, false);
+	CreateShapeFilterData(MoveChannel, GetUniqueID(), CollResponse, 0, 0, PQueryFilterData, PSimFilterData, BodyInstance.bUseCCD, bEnableImpactDamage, false, bEnableContactModification);
 
 	// Build filterData variations for complex and simple
 	PSimFilterData.word3 |= EPDF_SimpleCollision | EPDF_ComplexCollision;
diff --git a/Engine/Source/Runtime/Engine/Private/PhysicsEngine/PhysXSupport.cpp b/Engine/Source/Runtime/Engine/Private/PhysicsEngine/PhysXSupport.cpp
--- a/Engine/Source/Runtime/Engine/Private/PhysicsEngine/PhysXSupport.cpp
+++ b/Engine/Source/Runtime/Engine/Private/PhysicsEngine/PhysXSupport.cpp
@@ -397,6 +397,11 @@ PxFilterFlags PhysXSimFilterShader(	PxFilterObjectAttributes attributes0, PxFilt
 		pairFlags |= (PxPairFlag::eNOTIFY_TOUCH_FOUND | PxPairFlag::eNOTIFY_TOUCH_PERSISTS | PxPairFlag::eNOTIFY_CONTACT_POINTS );
 	}
 
+	if ((FilterFlags0&EPDF_ModifyContacts) || (FilterFlags1&EPDF_ModifyContacts))
+	{
+		pairFlags |= (PxPairFlag::eMODIFY_CONTACTS);
+	}
+
 	return PxFilterFlags();
 }
 
