Fix IsFractured returning the wrong value for form extended structures. We really only care about this if the actor is initially dynamic so we now call IsFracturedOrInitiallyStatic

diff --git a/Engine/Source/Runtime/Engine/Classes/Components/DestructibleComponent.h b/Engine/Source/Runtime/Engine/Classes/Components/DestructibleComponent.h
--- a/Engine/Source/Runtime/Engine/Classes/Components/DestructibleComponent.h
+++ b/Engine/Source/Runtime/Engine/Classes/Components/DestructibleComponent.h
@@ -200,8 +200,6 @@ public:
 	void OnVisibilityEvent(const physx::apex::NxApexChunkStateEventData & InDamageEvent);
 #endif // WITH_APEX
 
-	bool IsFractured() const;
-
 	// End DestructibleComponent interface.
 
 	virtual bool DoCustomNavigableGeometryExport(struct FNavigableGeometryExport* GeomExport) const override;
@@ -221,6 +219,8 @@ private:
 
 	void SetChunksWorldTM(const TArray<FUpdateChunksInfo>& UpdateInfos);
 
+	bool IsFracturedOrInitiallyStatic() const;
+
 	/** Collision response used for chunks */
 	FCollisionResponse LargeChunkCollisionResponse;
 	FCollisionResponse SmallChunkCollisionResponse;
diff --git a/Engine/Source/Runtime/Engine/Private/DestructibleComponent.cpp b/Engine/Source/Runtime/Engine/Private/DestructibleComponent.cpp
--- a/Engine/Source/Runtime/Engine/Private/DestructibleComponent.cpp
+++ b/Engine/Source/Runtime/Engine/Private/DestructibleComponent.cpp
@@ -689,11 +689,19 @@ void UDestructibleComponent::OnVisibilityEvent(const NxApexChunkStateEventData &
 }
 #endif // WITH_APEX
 
-bool UDestructibleComponent::IsFractured() const
+bool UDestructibleComponent::IsFracturedOrInitiallyStatic() const
 {
+	if (ApexDestructibleActor == nullptr)
+	{
+		return false;
+	}
+
+	bool bInitiallyStatic = false;
 	bool bFractured = false;
 #if WITH_APEX
-	if (ApexDestructibleActor)
+	bInitiallyStatic = !ApexDestructibleActor->isInitiallyDynamic();
+
+	if (bInitiallyStatic == false)
 	{
 		//If we have only one chunk and its index is 0 we are NOT fractured. Otherwise we must have fractured
 		const physx::PxU32 VisibleChunkCount = ApexDestructibleActor->getNumVisibleChunks();
@@ -708,7 +716,7 @@ bool UDestructibleComponent::IsFractured() const
 		}
 	}
 #endif
-	return bFractured;
+	return bFractured || bInitiallyStatic;
 }
 
 void UDestructibleComponent::RefreshBoneTransforms(FActorComponentTickFunction* TickFunction)
@@ -934,7 +942,7 @@ void UDestructibleComponent::UpdateDestructibleChunkTM(const TArray<const PxRigi
 	{
 		UDestructibleComponent* DestructibleComponent = It.Key();
 		TArray<FUpdateChunksInfo>& UpdateInfos = It.Value();
-		if (DestructibleComponent->IsFractured())
+		if (DestructibleComponent->IsFracturedOrInitiallyStatic())
 		{
 			DestructibleComponent->SetChunksWorldTM(UpdateInfos);
 		}
