GitHub #796 : World composition: load all distance-dependent tiles in commmandlet
https://github.com/EpicGames/UnrealEngine/pull/796
UE-8974

diff --git a/Engine/Source/Runtime/Engine/Private/WorldComposition.cpp b/Engine/Source/Runtime/Engine/Private/WorldComposition.cpp
--- a/Engine/Source/Runtime/Engine/Private/WorldComposition.cpp
+++ b/Engine/Source/Runtime/Engine/Private/WorldComposition.cpp
@@ -509,6 +509,11 @@ void UWorldComposition::GetDistanceVisibleLevels(
 			// Dedicated server always loads all distance dependent tiles
 			bIsVisible = true;
 		}
+		else if (IsRunningCommandlet())
+		{
+			// Commandlets have no concept of viewer location, so always load all distance-dependent tiles.
+			bIsVisible = true;
+		}
 		else
 		{
 			//
@@ -575,8 +580,8 @@ void UWorldComposition::UpdateStreamingState()
 	UWorld* PlayWorld = GetWorld();
 	check(PlayWorld);
 
-	// Dedicated server does not use distance based streaming and just loads everything
-	if (PlayWorld->GetNetMode() == NM_DedicatedServer)
+	// Commandlets and Dedicated server does not use distance based streaming and just loads everything
+	if (IsRunningCommandlet() || PlayWorld->GetNetMode() == NM_DedicatedServer)
 	{
 		UpdateStreamingState(FVector::ZeroVector);
 		return;
