[INTEGRATE] Change 2433851 by [REDACTED_PERSON][REDACTED_HANDLE] on 2015/02/05 12:40:32

	UE-8940 Crash after adding a component in IWCE and then adding a new variable and changing its type in the BP editor

diff --git a/Engine/Source/Runtime/Engine/Private/ComponentInstanceDataCache.cpp b/Engine/Source/Runtime/Engine/Private/ComponentInstanceDataCache.cpp
--- a/Engine/Source/Runtime/Engine/Private/ComponentInstanceDataCache.cpp
+++ b/Engine/Source/Runtime/Engine/Private/ComponentInstanceDataCache.cpp
@@ -135,7 +135,11 @@ FComponentInstanceDataCache::FComponentInstanceDataCache(const AActor* Actor)
 				{
 					if (SceneComponent->AttachParent && SceneComponent->AttachParent->IsCreatedByConstructionScript())
 					{
-						InstanceComponentTransformToRootMap.Add(SceneComponent, SceneComponent->GetComponentTransform().GetRelativeTransform(Actor->GetRootComponent()->GetComponentTransform()));
+						auto RootComponent = Actor->GetRootComponent();
+						if (RootComponent)
+						{
+							InstanceComponentTransformToRootMap.Add(SceneComponent, SceneComponent->GetComponentTransform().GetRelativeTransform(RootComponent->GetComponentTransform()));
+						}
 					}
 				}
 			}
