fixed missing navlinks in loaded static navmesh
UE-9506 #ue4

diff --git a/Engine/Source/Runtime/Engine/Classes/AI/Navigation/RecastNavMesh.h b/Engine/Source/Runtime/Engine/Classes/AI/Navigation/RecastNavMesh.h
--- a/Engine/Source/Runtime/Engine/Classes/AI/Navigation/RecastNavMesh.h
+++ b/Engine/Source/Runtime/Engine/Classes/AI/Navigation/RecastNavMesh.h
@@ -16,8 +16,9 @@
 #define NAVMESHVER_DYNAMIC_LINKS		7
 #define NAVMESHVER_64BIT				9
 #define NAVMESHVER_CLUSTER_SIMPLIFIED	10
+#define NAVMESHVER_OFFMESH_HEIGHT_BUG	11
 
-#define NAVMESHVER_LATEST				NAVMESHVER_CLUSTER_SIMPLIFIED
+#define NAVMESHVER_LATEST				NAVMESHVER_OFFMESH_HEIGHT_BUG
 #define NAVMESHVER_MIN_COMPATIBLE		NAVMESHVER_CLUSTER_SIMPLIFIED
 
 #define RECAST_MAX_SEARCH_NODES		2048
@@ -677,7 +678,7 @@ public:
 
 protected:
 	/** Serialization helper. */
-	void SerializeRecastNavMesh(FArchive& Ar, FPImplRecastNavMesh*& NavMesh);
+	void SerializeRecastNavMesh(FArchive& Ar, FPImplRecastNavMesh*& NavMesh, int32 NavMeshVersion);
 
 	TArray<FIntPoint>& GetActiveTiles(); 
 	virtual void RestrictBuildingToActiveTiles(bool InRestrictBuildingToActiveTiles) override;
diff --git a/Engine/Source/Runtime/Engine/Classes/AI/Navigation/RecastNavMeshDataChunk.h b/Engine/Source/Runtime/Engine/Classes/AI/Navigation/RecastNavMeshDataChunk.h
--- a/Engine/Source/Runtime/Engine/Classes/AI/Navigation/RecastNavMeshDataChunk.h
+++ b/Engine/Source/Runtime/Engine/Classes/AI/Navigation/RecastNavMeshDataChunk.h
@@ -61,7 +61,7 @@ public:
 
 private:
 #if WITH_RECAST
-	void SerializeRecastData(FArchive& Ar);
+	void SerializeRecastData(FArchive& Ar, int32 NavMeshVersion);
 #endif//WITH_RECAST
 
 private:
diff --git a/Engine/Source/Runtime/Engine/Private/AI/Navigation/PImplRecastNavMesh.cpp b/Engine/Source/Runtime/Engine/Private/AI/Navigation/PImplRecastNavMesh.cpp
--- a/Engine/Source/Runtime/Engine/Private/AI/Navigation/PImplRecastNavMesh.cpp
+++ b/Engine/Source/Runtime/Engine/Private/AI/Navigation/PImplRecastNavMesh.cpp
@@ -271,7 +271,7 @@ void FPImplRecastNavMesh::ReleaseDetourNavMesh()
  * @param Ar - The archive with which to serialize.
  * @returns true if serialization was successful.
  */
-void FPImplRecastNavMesh::Serialize( FArchive& Ar )
+void FPImplRecastNavMesh::Serialize( FArchive& Ar, int32 NavMeshVersion )
 {
 	//@todo: How to handle loading nav meshes saved w/ recast when recast isn't present????
 
@@ -358,7 +358,7 @@ void FPImplRecastNavMesh::Serialize( FArchive& Ar )
 
 				unsigned char* TileData = NULL;
 				TileDataSize = 0;
-				SerializeRecastMeshTile(Ar, TileData, TileDataSize);
+				SerializeRecastMeshTile(Ar, NavMeshVersion, TileData, TileDataSize);
 				if (TileData != NULL)
 				{
 					dtMeshHeader* const TileHeader = (dtMeshHeader*)TileData;
@@ -370,7 +370,7 @@ void FPImplRecastNavMesh::Serialize( FArchive& Ar )
 					{
 						unsigned char* ComressedTileData = NULL;
 						int32 CompressedTileDataSize = 0;
-						SerializeCompressedTileCacheData(Ar, ComressedTileData, CompressedTileDataSize);
+						SerializeCompressedTileCacheData(Ar, NavMeshVersion, ComressedTileData, CompressedTileDataSize);
 						dtFree(ComressedTileData);
 					}
 				}
@@ -398,7 +398,7 @@ void FPImplRecastNavMesh::Serialize( FArchive& Ar )
 				
 				unsigned char* TileData = NULL;
 				TileDataSize = 0;
-				SerializeRecastMeshTile(Ar, TileData, TileDataSize);
+				SerializeRecastMeshTile(Ar, NavMeshVersion, TileData, TileDataSize);
 
 				if (TileData != NULL)
 				{
@@ -411,7 +411,7 @@ void FPImplRecastNavMesh::Serialize( FArchive& Ar )
 					{
 						uint8* ComressedTileData = nullptr;
 						int32 CompressedTileDataSize = 0;
-						SerializeCompressedTileCacheData(Ar, ComressedTileData, CompressedTileDataSize);
+						SerializeCompressedTileCacheData(Ar, NavMeshVersion, ComressedTileData, CompressedTileDataSize);
 						
 						if (CompressedTileDataSize > 0)
 						{
@@ -436,7 +436,7 @@ void FPImplRecastNavMesh::Serialize( FArchive& Ar )
 			Ar << TileRef << TileDataSize;
 
 			unsigned char* TileData = Tile->data;
-			SerializeRecastMeshTile(Ar, TileData, TileDataSize);
+			SerializeRecastMeshTile(Ar, NavMeshVersion, TileData, TileDataSize);
 
 			// Serialize compressed tile cache layer only if navmesh requires it
 			{
@@ -449,13 +449,13 @@ void FPImplRecastNavMesh::Serialize( FArchive& Ar )
 					CompressedDataSize = TileCacheLayer.DataSize;
 				}
 				
-				SerializeCompressedTileCacheData(Ar, CompressedData, CompressedDataSize);
+				SerializeCompressedTileCacheData(Ar, NavMeshVersion, CompressedData, CompressedDataSize);
 			}
 		}
 	}
 }
 
-void FPImplRecastNavMesh::SerializeRecastMeshTile(FArchive& Ar, unsigned char*& TileData, int32& TileDataSize)
+void FPImplRecastNavMesh::SerializeRecastMeshTile(FArchive& Ar, int32 NavMeshVersion, unsigned char*& TileData, int32& TileDataSize)
 {
 	// The strategy here is to serialize the data blob that is passed into addTile()
 	// @see dtCreateNavMeshData() for details on how this data is laid out
@@ -639,6 +639,15 @@ void FPImplRecastNavMesh::SerializeRecastMeshTile(FArchive& Ar, unsigned char*&
 			Ar << Conn.rad << Conn.poly << Conn.flags << Conn.side << Conn.userId;
 		}
 
+		if (NavMeshVersion >= NAVMESHVER_OFFMESH_HEIGHT_BUG)
+		{
+			for (int32 ConnIdx = 0; ConnIdx < offMeshConCount; ++ConnIdx)
+			{
+				dtOffMeshConnection& Conn = OffMeshCons[ConnIdx];
+				Ar << Conn.height;
+			}
+		}
+
 		for (int32 SegIdx=0; SegIdx < offMeshSegConCount; ++SegIdx)
 		{
 			dtOffMeshSegmentConnection& Seg = OffMeshSegs[SegIdx];
@@ -667,7 +676,7 @@ void FPImplRecastNavMesh::SerializeRecastMeshTile(FArchive& Ar, unsigned char*&
 	}
 }
 
-void FPImplRecastNavMesh::SerializeCompressedTileCacheData(FArchive& Ar, unsigned char*& CompressedData, int32& CompressedDataSize)
+void FPImplRecastNavMesh::SerializeCompressedTileCacheData(FArchive& Ar, int32 NavMeshVersion, unsigned char*& CompressedData, int32& CompressedDataSize)
 {
 	Ar << CompressedDataSize;
 
diff --git a/Engine/Source/Runtime/Engine/Private/AI/Navigation/RecastNavMesh.cpp b/Engine/Source/Runtime/Engine/Private/AI/Navigation/RecastNavMesh.cpp
--- a/Engine/Source/Runtime/Engine/Private/AI/Navigation/RecastNavMesh.cpp
+++ b/Engine/Source/Runtime/Engine/Private/AI/Navigation/RecastNavMesh.cpp
@@ -676,7 +676,7 @@ void ARecastNavMesh::RestrictBuildingToActiveTiles(bool InRestrictBuildingToActi
 	}
 }
 
-void ARecastNavMesh::SerializeRecastNavMesh(FArchive& Ar, FPImplRecastNavMesh*& NavMesh)
+void ARecastNavMesh::SerializeRecastNavMesh(FArchive& Ar, FPImplRecastNavMesh*& NavMesh, int32 NavMeshVersion)
 {
 	if (!Ar.IsLoading()	&& NavMesh == NULL)
 	{
@@ -694,7 +694,7 @@ void ARecastNavMesh::SerializeRecastNavMesh(FArchive& Ar, FPImplRecastNavMesh*&
 	
 	if (RecastNavMeshImpl)
 	{
-		RecastNavMeshImpl->Serialize(Ar);
+		RecastNavMeshImpl->Serialize(Ar, NavMeshVersion);
 	}	
 }
 
@@ -723,7 +723,7 @@ void ARecastNavMesh::Serialize( FArchive& Ar )
 		}
 		else if (RecastNavMeshSizeBytes > 4)
 		{
-			SerializeRecastNavMesh(Ar, RecastNavMeshImpl);
+			SerializeRecastNavMesh(Ar, RecastNavMeshImpl, NavMeshVersion);
 			bWantsUpdate = bForceRebuildOnLoad == true || HasValidNavmesh() == false;
 #if !(UE_BUILD_SHIPPING)
 			RequestDrawingUpdate();
@@ -739,7 +739,7 @@ void ARecastNavMesh::Serialize( FArchive& Ar )
 	}
 	else
 	{
-		SerializeRecastNavMesh(Ar, RecastNavMeshImpl);
+		SerializeRecastNavMesh(Ar, RecastNavMeshImpl, NavMeshVersion);
 
 		if (Ar.IsSaving())
 		{
diff --git a/Engine/Source/Runtime/Engine/Private/AI/Navigation/RecastNavMeshDataChunk.cpp b/Engine/Source/Runtime/Engine/Private/AI/Navigation/RecastNavMeshDataChunk.cpp
--- a/Engine/Source/Runtime/Engine/Private/AI/Navigation/RecastNavMeshDataChunk.cpp
+++ b/Engine/Source/Runtime/Engine/Private/AI/Navigation/RecastNavMeshDataChunk.cpp
@@ -68,7 +68,7 @@ void URecastNavMeshDataChunk::Serialize(FArchive& Ar)
 #if WITH_RECAST
 		else if (RecastNavMeshSizeBytes > 4)
 		{
-			SerializeRecastData(Ar);
+			SerializeRecastData(Ar, NavMeshVersion);
 		}
 #endif// WITH_RECAST
 		else
@@ -80,7 +80,7 @@ void URecastNavMeshDataChunk::Serialize(FArchive& Ar)
 	else
 	{
 #if WITH_RECAST
-		SerializeRecastData(Ar);
+		SerializeRecastData(Ar, NavMeshVersion);
 #endif// WITH_RECAST
 
 		if (Ar.IsSaving())
@@ -95,7 +95,7 @@ void URecastNavMeshDataChunk::Serialize(FArchive& Ar)
 }
 
 #if WITH_RECAST
-void URecastNavMeshDataChunk::SerializeRecastData(FArchive& Ar)
+void URecastNavMeshDataChunk::SerializeRecastData(FArchive& Ar, int32 NavMeshVersion)
 {
 	int32 TileNum = Tiles.Num();
 	Ar << TileNum;
@@ -110,7 +110,7 @@ void URecastNavMeshDataChunk::SerializeRecastData(FArchive& Ar)
 
 			// Load tile data 
 			uint8* TileRawData = nullptr;
-			FPImplRecastNavMesh::SerializeRecastMeshTile(Ar, TileRawData, TileDataSize);
+			FPImplRecastNavMesh::SerializeRecastMeshTile(Ar, NavMeshVersion, TileRawData, TileDataSize);
 			
 			if (TileRawData != nullptr)
 			{
@@ -120,7 +120,7 @@ void URecastNavMeshDataChunk::SerializeRecastData(FArchive& Ar)
 				if (Ar.UE4Ver() >= VER_UE4_ADD_MODIFIERS_RUNTIME_GENERATION && 
 					Ar.UE4Ver() != VER_UE4_MERGED_ADD_MODIFIERS_RUNTIME_GENERATION_TO_4_7) // Merged package from 4.7 branch
 				{
-					FPImplRecastNavMesh::SerializeCompressedTileCacheData(Ar, TileCacheRawData, TileCacheDataSize);
+					FPImplRecastNavMesh::SerializeCompressedTileCacheData(Ar, NavMeshVersion, TileCacheRawData, TileCacheDataSize);
 				}
 				
 				// We are owner of tile raw data
@@ -137,9 +137,9 @@ void URecastNavMeshDataChunk::SerializeRecastData(FArchive& Ar)
 			{
 				// Save tile itself
 				Ar << TileData.TileDataSize;
-				FPImplRecastNavMesh::SerializeRecastMeshTile(Ar, TileData.TileRawData->RawData, TileData.TileDataSize);
+				FPImplRecastNavMesh::SerializeRecastMeshTile(Ar, NavMeshVersion, TileData.TileRawData->RawData, TileData.TileDataSize);
 				// Save compressed tile cache layer
-				FPImplRecastNavMesh::SerializeCompressedTileCacheData(Ar, TileData.TileCacheRawData->RawData, TileData.TileCacheDataSize);
+				FPImplRecastNavMesh::SerializeCompressedTileCacheData(Ar, NavMeshVersion, TileData.TileCacheRawData->RawData, TileData.TileCacheDataSize);
 			}
 		}
 	}
diff --git a/Engine/Source/Runtime/Engine/Public/AI/Navigation/PImplRecastNavMesh.h b/Engine/Source/Runtime/Engine/Public/AI/Navigation/PImplRecastNavMesh.h
--- a/Engine/Source/Runtime/Engine/Public/AI/Navigation/PImplRecastNavMesh.h
+++ b/Engine/Source/Runtime/Engine/Public/AI/Navigation/PImplRecastNavMesh.h
@@ -72,7 +72,7 @@ public:
 	 * @param Ar - The archive with which to serialize.
 	 * @returns true if serialization was successful.
 	 */
-	void Serialize(FArchive& Ar);
+	void Serialize(FArchive& Ar, int32 NavMeshVersion);
 
 	/** Debug rendering. */
 	void GetDebugGeometry(FRecastDebugGeometry& OutGeometry, int32 TileIndex = INDEX_NONE) const;
@@ -227,10 +227,10 @@ public:
 	mutable dtNavMeshQuery SharedNavQuery;
 
 	/** Helper function to serialize a single Recast tile. */
-	static void SerializeRecastMeshTile(FArchive& Ar, unsigned char*& TileData, int32& TileDataSize);
+	static void SerializeRecastMeshTile(FArchive& Ar, int32 NavMeshVersion, unsigned char*& TileData, int32& TileDataSize);
 
 	/** Helper function to serialize a Recast tile compressed data. */
-	static void SerializeCompressedTileCacheData(FArchive& Ar, unsigned char*& CompressedData, int32& CompressedDataSize);
+	static void SerializeCompressedTileCacheData(FArchive& Ar, int32 NavMeshVersion, unsigned char*& CompressedData, int32& CompressedDataSize);
 
 	/** Initialize data for pathfinding */
 	bool InitPathfinding(const FVector& UnrealStart, const FVector& UnrealEnd, 
