Merging using UE4-GDC2015-To-UE4
Update navigation data only when shifting already registered components, also fixed navigation update for attached components

diff --git a/Engine/Source/Runtime/Engine/Private/Actor.cpp b/Engine/Source/Runtime/Engine/Private/Actor.cpp
--- a/Engine/Source/Runtime/Engine/Private/Actor.cpp
+++ b/Engine/Source/Runtime/Engine/Private/Actor.cpp
@@ -513,20 +513,26 @@ void AActor::ProcessEvent(UFunction* Function, void* Parameters)
 	}
 }
 
-
 void AActor::ApplyWorldOffset(const FVector& InOffset, bool bWorldShift)
 {
-	// Do not shift child components
-	if (RootComponent != NULL && RootComponent->AttachParent == NULL)
+	// Attached components will be shifted by parents
+	if (RootComponent != nullptr && RootComponent->AttachParent == nullptr)
 	{
 		RootComponent->ApplyWorldOffset(InOffset, bWorldShift);
+	}
 
-		UNavigationSystem::UpdateNavOctreeBounds(this);
-		UNavigationSystem::UpdateNavOctreeAll(this);
+	// Navigation receives update during component registration. World shift needs a separate path to shift all navigation data
+	// So this normally should happen only in the editor when user moves visible sub-levels
+	if (!bWorldShift && !InOffset.IsZero())
+	{
+		if (RootComponent != nullptr && RootComponent->IsRegistered())
+		{
+			UNavigationSystem::UpdateNavOctreeBounds(this);
+			UNavigationSystem::UpdateNavOctreeAll(this);
+		}
 	}
 }
 
-
 static AActor* GTestRegisterTickFunctions = NULL;
 
 void AActor::RegisterActorTickFunctions(bool bRegister)
