Fix assert in SpringBone controller when running with Parallel Anim Evaluation

diff --git a/Engine/Source/Runtime/Engine/Classes/Animation/BoneControllers/AnimNode_SpringBone.h b/Engine/Source/Runtime/Engine/Classes/Animation/BoneControllers/AnimNode_SpringBone.h
--- a/Engine/Source/Runtime/Engine/Classes/Animation/BoneControllers/AnimNode_SpringBone.h
+++ b/Engine/Source/Runtime/Engine/Classes/Animation/BoneControllers/AnimNode_SpringBone.h
@@ -45,6 +45,9 @@ struct ENGINE_API FAnimNode_SpringBone : public FAnimNode_SkeletalControlBase
 	/** Internal use - Amount of time we need to simulate. */
 	float RemainingTime;
 
+	/** Internal use - Current timestep */
+	float FixedTimeStep;
+
 	/** Did we have a non-zero ControlStrength last frame. */
 	bool bHadValidStrength;
 
diff --git a/Engine/Source/Runtime/Engine/Private/Animation/AnimNode_SpringBone.cpp b/Engine/Source/Runtime/Engine/Private/Animation/AnimNode_SpringBone.cpp
--- a/Engine/Source/Runtime/Engine/Private/Animation/AnimNode_SpringBone.cpp
+++ b/Engine/Source/Runtime/Engine/Private/Animation/AnimNode_SpringBone.cpp
@@ -37,6 +37,12 @@ void FAnimNode_SpringBone::Update(const FAnimationUpdateContext& Context)
 	FAnimNode_SkeletalControlBase::Update(Context);
 
 	RemainingTime += Context.GetDeltaTime();
+
+	const USkeletalMeshComponent* SkelComp = Context.AnimInstance->GetSkelMeshComponent();
+	const UWorld* World = SkelComp->GetWorld();
+	check(World->GetWorldSettings());
+	// Fixed step simulation at 120hz
+	FixedTimeStep = (1.f / 120.f) * World->GetWorldSettings()->GetEffectiveTimeDilation();
 }
 
 void FAnimNode_SpringBone::GatherDebugData(FNodeDebugData& DebugData)
@@ -74,10 +80,6 @@ void FAnimNode_SpringBone::EvaluateBoneTransforms(USkeletalMeshComponent* SkelCo
 		BoneVelocity = FVector::ZeroVector;
 	}
 		
-	UWorld* World = SkelComp->GetWorld();
-	check(World->GetWorldSettings());
-	// Fixed step simulation at 120hz
-	float const FixedTimeStep = (1.f/120.f) * World->GetWorldSettings()->GetEffectiveTimeDilation();
 	while (RemainingTime > FixedTimeStep)
 	{
 		// Update location of our base by how much our base moved this frame.
