Added InterpTrackEvent preview scrubbing

diff --git a/Engine/Source/Runtime/Engine/Classes/Matinee/InterpTrackEvent.h b/Engine/Source/Runtime/Engine/Classes/Matinee/InterpTrackEvent.h
--- a/Engine/Source/Runtime/Engine/Classes/Matinee/InterpTrackEvent.h
+++ b/Engine/Source/Runtime/Engine/Classes/Matinee/InterpTrackEvent.h
@@ -52,6 +52,9 @@ class UInterpTrackEvent : public UInterpTrack
 	UPROPERTY(EditAnywhere, Category=InterpTrackEvent)
 	uint32 bFireEventsWhenJumpingForwards:1;
 
+	/** If checked each key's event name is the exact name of the custom event function in level script that will be called */
+	UPROPERTY(EditAnywhere, Category=InterpTrackEvent)
+	uint32 bUseCustomEventName:1;
 
 	// Begin UInterpTrack Interface
 	virtual int32 GetNumKeyframes() const override;
@@ -64,6 +67,7 @@ class UInterpTrackEvent : public UInterpTrack
 	virtual void RemoveKeyframe( int32 KeyIndex ) override;
 	virtual int32 DuplicateKeyframe(int32 KeyIndex, float NewKeyTime, UInterpTrack* ToTrack = NULL) override;
 	virtual bool GetClosestSnapPosition( float InPosition, TArray<int32>& IgnoreKeys, float& OutPosition ) override;
+	virtual void PreviewUpdateTrack(float NewPosition, class UInterpTrackInst* TrInst) override;
 	virtual void UpdateTrack( float NewPosition, UInterpTrackInst* TrackInst, bool bJump ) override;
 	virtual const FString GetEdHelperClassName() const override;
 	virtual const FString GetSlateHelperClassName() const override;
diff --git a/Engine/Source/Runtime/Engine/Classes/Matinee/MatineeActor.h b/Engine/Source/Runtime/Engine/Classes/Matinee/MatineeActor.h
--- a/Engine/Source/Runtime/Engine/Classes/Matinee/MatineeActor.h
+++ b/Engine/Source/Runtime/Engine/Classes/Matinee/MatineeActor.h
@@ -464,10 +464,10 @@ public:
 	
 
 	/** Try to invoke the event with the given name in the level script */
-	virtual void NotifyEventTriggered(FName EventName, float EventTime);
+	virtual void NotifyEventTriggered(FName EventName, float EventTime, bool bUseCustomEventName=false);
 
 	/** Util to get the name of the function to find for the given event name */
-	ENGINE_API FName GetFunctionNameForEvent(FName EventName);
+	ENGINE_API FName GetFunctionNameForEvent(FName EventName,bool bUseCustomEventName=false);
 
 public:
 #if WITH_EDITORONLY_DATA
diff --git a/Engine/Source/Runtime/Engine/Private/Interpolation.cpp b/Engine/Source/Runtime/Engine/Private/Interpolation.cpp
--- a/Engine/Source/Runtime/Engine/Private/Interpolation.cpp
+++ b/Engine/Source/Runtime/Engine/Private/Interpolation.cpp
@@ -371,19 +371,28 @@ AMatineeActor::AMatineeActor(const FObjectInitializer& ObjectInitializer)
 	ReplicationForceIsPlaying = 0;
 }
 
-FName AMatineeActor::GetFunctionNameForEvent(FName EventName)
+FName AMatineeActor::GetFunctionNameForEvent(FName EventName,bool bUseCustomEventName)
 {
-	FString EventFuncName = FString::Printf(TEXT("%s_%s"), *MatineeControllerName.ToString(), *EventName.ToString());
-	return FName(*EventFuncName);
+	FName EventFuncName;
+	if( bUseCustomEventName )
+	{
+		EventFuncName = EventName;
+	}
+	else
+	{
+		EventFuncName = *(FString::Printf(TEXT("%s_%s"), *MatineeControllerName.ToString(), *EventName.ToString()));	
+	}
+
+	return EventFuncName;
 }
 
-void AMatineeActor::NotifyEventTriggered(FName EventName, float EventTime)
+void AMatineeActor::NotifyEventTriggered(FName EventName, float EventTime, bool bUseCustomEventName)
 {
 	ULevel* Level = GetLevel();
 	ALevelScriptActor* LevelScriptActor = Level->LevelScriptActor;
 	if(LevelScriptActor != NULL)
 	{
-		FName EventFuncName = GetFunctionNameForEvent(EventName);
+		FName EventFuncName = GetFunctionNameForEvent(EventName,bUseCustomEventName);
 		UFunction* EventFunction = LevelScriptActor->FindFunction(EventFuncName);
 		if(EventFunction != NULL)
 		{
@@ -6312,6 +6321,15 @@ void UInterpTrackEvent::RemoveKeyframe(int32 KeyIndex)
 	}
 }
 
+void UInterpTrackEvent::PreviewUpdateTrack(float NewPosition, class UInterpTrackInst* TrInst)
+{
+	UInterpGroupInst* GrInst = CastChecked<UInterpGroupInst>( TrInst->GetOuter() );
+	AMatineeActor* MatineeActor = CastChecked<AMatineeActor>( GrInst->GetOuter() );
+
+	bool bJump = !( MatineeActor->bIsPlaying );
+	UpdateTrack(NewPosition, TrInst, bJump);
+}
+
 void UInterpTrackEvent::UpdateTrack(float NewPosition, UInterpTrackInst* TrInst, bool bJump)
 {
 	UInterpTrackInstEvent* EventInst = CastChecked<UInterpTrackInstEvent>(TrInst);
@@ -6405,7 +6423,7 @@ void UInterpTrackEvent::UpdateTrack(float NewPosition, UInterpTrackInst* TrInst,
 
 			if( bFireThisEvent )
 			{
-				MatineeActor->NotifyEventTriggered(EventTrack[i].EventName, EventTime);
+				MatineeActor->NotifyEventTriggered(EventTrack[i].EventName, EventTime, bUseCustomEventName);
 			}
 		}
 	}
