[INTEGRATE] Change 2433956 by [REDACTED_PERSON][REDACTED_HANDLE] on 2015/02/05 14:12:10

	#ue4 - Trimmed down comments on CharacterMovementComponent so tooltip becomes smaller.

diff --git a/Engine/Source/Runtime/Engine/Classes/GameFramework/CharacterMovementComponent.h b/Engine/Source/Runtime/Engine/Classes/GameFramework/CharacterMovementComponent.h
--- a/Engine/Source/Runtime/Engine/Classes/GameFramework/CharacterMovementComponent.h
+++ b/Engine/Source/Runtime/Engine/Classes/GameFramework/CharacterMovementComponent.h
@@ -107,13 +107,7 @@ typedef TSharedPtr<class FSavedMove_Character> FSavedMovePtr;
 //=============================================================================
 /**
  * CharacterMovementComponent handles movement logic for the associated Character owner.
- *
- * It supports various movement modes including:
- *    - walking:  Walking on a surface, under the effects of friction, and able to "step up" barriers.
- *    - falling:  Falling under the effects of gravity. Usually this is the mode after jumping or walking off the edge of a surface.
- *    - flying:   Flying, ignoring the effects of gravity.
- *    - swimming: Swimming through a fluid volume, under the effects of gravity and buoyancy.
- *    - custom:   User-defined custom movement mode, including many possible sub-modes. See UCharacterMovementComponent::CustomMovementMode and SetMovementMode().
+ * It supports various movement modes including: walking, falling, swimming, flying, custom.
  *
  * Movement is affected primarily by current Velocity and Acceleration. Acceleration is updated each frame
  * based on the input vector accumulated thus far (see UPawnMovementComponent::GetPendingInputVector()).
@@ -173,8 +167,13 @@ public:
 	
 	/**
 	 * Actor's current movement mode (walking, falling, etc).
+	 *    - walking:  Walking on a surface, under the effects of friction, and able to "step up" barriers. Vertical velocity is zero.
+	 *    - falling:  Falling under the effects of gravity, after jumping or walking off the edge of a surface.
+	 *    - flying:   Flying, ignoring the effects of gravity.
+	 *    - swimming: Swimming through a fluid volume, under the effects of gravity and buoyancy.
+	 *    - custom:   User-defined custom movement mode, including many possible sub-modes.
 	 * This is automatically replicated through the Character owner and for client-server movement functions.
-	 * @see SetMovementMode()
+	 * @see SetMovementMode(), CustomMovementMode
 	 */
 	UPROPERTY(Category=MovementMode, BlueprintReadOnly)
 	TEnumAsByte<enum EMovementMode> MovementMode;
