[INTEGRATE] Change 2434010 by [REDACTED_PERSON][REDACTED_HANDLE] on 2015/02/05 14:40:25

	#ue4 - Warnings when trying to move components with non-movable mobility are now FMessageLog warnings show they show up at the end of a PIE session.

diff --git a/Engine/Source/Runtime/Engine/Private/PrimitiveComponent.cpp b/Engine/Source/Runtime/Engine/Private/PrimitiveComponent.cpp
--- a/Engine/Source/Runtime/Engine/Private/PrimitiveComponent.cpp
+++ b/Engine/Source/Runtime/Engine/Private/PrimitiveComponent.cpp
@@ -1290,10 +1290,12 @@ bool UPrimitiveComponent::MoveComponent( const FVector& Delta, const FRotator& N
 	AActor* const Actor = GetOwner();
 
 	// static things can move before they are registered (e.g. immediately after streaming), but not after.
+	// TODO: Static components without an owner can move, should they be able to?
 	if (Mobility == EComponentMobility::Static && Actor && Actor->bActorInitialized)
-	{
-		// TODO: Static components without an owner can move, should they be able to?
-		UE_LOG(LogPrimitiveComponent, Warning, TEXT("Trying to move static component '%s' after initialization"), *GetFullName());
+	{		
+		FMessageLog("PIE").Warning( FText::Format(LOCTEXT("InvalidMove", "Mobility of {0} : {1} has to be 'Movable' if you'd like to move. "),
+				FText::FromString(GetNameSafe(GetOwner())), FText::FromString(GetName()) ));
+
 		if (OutHit)
 		{
 			*OutHit = FHitResult();
diff --git a/Engine/Source/Runtime/Engine/Private/SceneComponent.cpp b/Engine/Source/Runtime/Engine/Private/SceneComponent.cpp
--- a/Engine/Source/Runtime/Engine/Private/SceneComponent.cpp
+++ b/Engine/Source/Runtime/Engine/Private/SceneComponent.cpp
@@ -1610,7 +1610,7 @@ bool USceneComponent::MoveComponent( const FVector& Delta, const FRotator& NewRo
 			// It's only a problem if we're in gameplay, and the owning level is visible
 			if (World->HasBegunPlay() && IsRegistered() && Level && Level->bIsVisible)
 			{
-				FMessageLog("Performance").Warning(FText::Format(LOCTEXT("InvalidMove", "Mobility of {0} : {1} has to be 'Movable' if you'd like to move. "),
+				FMessageLog("PIE").Warning(FText::Format(LOCTEXT("InvalidMove", "Mobility of {0} : {1} has to be 'Movable' if you'd like to move. "),
 					FText::FromString(GetNameSafe(GetOwner())), FText::FromString(GetName())));
 				return false;
 			}
