Added r.MobileReduceLoadedMips for platforms without texture streaming
#android
#ue4

diff --git a/Engine/Source/Runtime/Core/Private/HAL/ConsoleManager.cpp b/Engine/Source/Runtime/Core/Private/HAL/ConsoleManager.cpp
--- a/Engine/Source/Runtime/Core/Private/HAL/ConsoleManager.cpp
+++ b/Engine/Source/Runtime/Core/Private/HAL/ConsoleManager.cpp
@@ -1608,6 +1608,12 @@ static TAutoConsoleVariable<int32> CVarMobileHDR32bpp(
 	TEXT("1: Mobile HDR renders to an RGBA8 target."),
 	ECVF_RenderThreadSafe);
 
+static TAutoConsoleVariable<int32> CVarMobileReduceLoadedMips(
+	TEXT("r.MobileReduceLoadedMips"),
+	0,
+	TEXT("Reduce loaded texture mipmaps for nonstreaming mobile platforms.\n"),
+	ECVF_RenderThreadSafe);
+
 static TAutoConsoleVariable<int32> CVarSetClearSceneMethod(
 	TEXT("r.ClearSceneMethod"),
 	1,
diff --git a/Engine/Source/Runtime/Engine/Private/Texture2D.cpp b/Engine/Source/Runtime/Engine/Private/Texture2D.cpp
--- a/Engine/Source/Runtime/Engine/Private/Texture2D.cpp
+++ b/Engine/Source/Runtime/Engine/Private/Texture2D.cpp
@@ -760,7 +760,10 @@ FTextureResource* UTexture2D::CreateResource()
 	int32 NumMips = GetNumMips();
 
 	// Determine whether or not this texture can be streamed.
-	bIsStreamable = IStreamingManager::Get().IsTextureStreamingEnabled() &&
+	bIsStreamable = 
+#if !PLATFORM_ANDROID
+					IStreamingManager::Get().IsTextureStreamingEnabled() &&
+#endif
 					!NeverStream && 
 					(NumMips > 1) && 
 					(LODGroup != TEXTUREGROUP_UI) && 
@@ -820,9 +823,15 @@ FTextureResource* UTexture2D::CreateResource()
 		// Handle streaming textures.
 		if( bIsStreamable )
 		{
+#if PLATFORM_SUPPORTS_TEXTURE_STREAMING
 			// Only request lower miplevels and let texture streaming code load the rest.
 			NumNonStreamingMips = GetNumNonStreamingMips();
 			RequestedMips = NumNonStreamingMips;
+#else
+			static auto* MobileReduceLoadedMipsCvar = IConsoleManager::Get().FindTConsoleVariableDataInt(TEXT("r.MobileReduceLoadedMips"));
+			NumNonStreamingMips = GetNumNonStreamingMips();
+			RequestedMips = FMath::Min(NumMips, GMaxTextureMipCount) - MobileReduceLoadedMipsCvar->GetValueOnAnyThread();
+#endif
 		}
 		// Handle non- streaming textures.
 		else
