UE-9640 - Add virtual to UpdatePing and RecalculateAvgPing

diff --git a/Engine/Source/Runtime/Engine/Classes/GameFramework/PlayerController.h b/Engine/Source/Runtime/Engine/Classes/GameFramework/PlayerController.h
--- a/Engine/Source/Runtime/Engine/Classes/GameFramework/PlayerController.h
+++ b/Engine/Source/Runtime/Engine/Classes/GameFramework/PlayerController.h
@@ -1306,7 +1306,7 @@ public:
 	virtual void PawnLeavingGame();
 
 	/** Takes ping updates from the net driver (both clientside and serverside), and passes them on to PlayerState::UpdatePing */
-	void UpdatePing(float InPing);
+	virtual void UpdatePing(float InPing);
 
 	/**
 	 * Get next active viewable player in PlayerArray.
diff --git a/Engine/Source/Runtime/Engine/Classes/GameFramework/PlayerState.h b/Engine/Source/Runtime/Engine/Classes/GameFramework/PlayerState.h
--- a/Engine/Source/Runtime/Engine/Classes/GameFramework/PlayerState.h
+++ b/Engine/Source/Runtime/Engine/Classes/GameFramework/PlayerState.h
@@ -145,10 +145,10 @@ public:
 	 * Receives ping updates for the client (both clientside and serverside), from the net driver
 	 * NOTE: This updates much more frequently clientside, thus the clientside ping will often be different to what the server displays
 	 */
-	void UpdatePing(float InPing);
+	virtual void UpdatePing(float InPing);
 
 	/** Recalculates the replicated Ping value once per second (both clientside and serverside), based upon collected ping data */
-	void RecalculateAvgPing();
+	virtual void RecalculateAvgPing();
 
 	/**
 	 * Returns true if should broadcast player welcome/left messages.
