[INTEGRATE] Change 2438768 by bruce.nesbit@bnesbit_release_port_edit on 2015/02/09 16:57

	Moved code to fixup old UBlueprintGeneratedClass Actor Components in Serialize function instead of PostLoad
	(Fixes UE-9340)

diff --git a/Engine/Source/Runtime/Engine/Classes/Components/ActorComponent.h b/Engine/Source/Runtime/Engine/Classes/Components/ActorComponent.h
--- a/Engine/Source/Runtime/Engine/Classes/Components/ActorComponent.h
+++ b/Engine/Source/Runtime/Engine/Classes/Components/ActorComponent.h
@@ -565,6 +565,7 @@ public:
 	 */
 	virtual void ApplyWorldOffset(const FVector& InOffset, bool bWorldShift) {};
 
+	virtual void Serialize(FArchive& Ar) override;
 #if !(UE_BUILD_SHIPPING || UE_BUILD_TEST)
 
 private:
diff --git a/Engine/Source/Runtime/Engine/Private/ActorComponent.cpp b/Engine/Source/Runtime/Engine/Private/ActorComponent.cpp
--- a/Engine/Source/Runtime/Engine/Private/ActorComponent.cpp
+++ b/Engine/Source/Runtime/Engine/Private/ActorComponent.cpp
@@ -157,7 +157,18 @@ void UActorComponent::PostLoad()
 			}
 		}
 	}
+// 	if (!HasAllFlags(RF_Public) && GetOuter()->IsA<UBlueprintGeneratedClass>())
+// 	{
+// 		SetFlags(RF_Public);
+// 		ULinkerLoad::RefreshExportFlags(this);
+// 	}
+}
+
+void UActorComponent::Serialize(FArchive& Ar)
+{
+	Super::Serialize(Ar);
 
+	// Fixup older blueprints that were private but now need to be public
 	if (!HasAllFlags(RF_Public) && GetOuter()->IsA<UBlueprintGeneratedClass>())
 	{
 		SetFlags(RF_Public);
