#Morphtarget - Fix morphtarget UID issue
=========================
Merging

//depot/UE4-GDC2015/Engine/Source/...

to //depot/UE4/Engine/Source/...

diff --git a/Engine/Source/Editor/UnrealEd/Private/SkeletalMeshEdit.cpp b/Engine/Source/Editor/UnrealEd/Private/SkeletalMeshEdit.cpp
--- a/Engine/Source/Editor/UnrealEd/Private/SkeletalMeshEdit.cpp
+++ b/Engine/Source/Editor/UnrealEd/Private/SkeletalMeshEdit.cpp
@@ -939,10 +939,17 @@ bool UnFbx::FFbxImporter::ImportCurveToAnimSequence(class UAnimSequence * Target
 	if (TargetSequence && FbxCurve)
 	{
 		FName Name = *CurveName;
-		FSmartNameMapping* NameMapping = TargetSequence->GetSkeleton()->SmartNames.GetContainer(USkeleton::AnimCurveMappingName);
+		USkeleton* Skeleton = TargetSequence->GetSkeleton();
+		FSmartNameMapping* NameMapping = Skeleton->SmartNames.GetContainer(USkeleton::AnimCurveMappingName);
 
 		// Add or retrieve curve
 		USkeleton::AnimCurveUID Uid;
+		if (!NameMapping->Exists(Name))
+		{
+			// mark skeleton dirty
+			Skeleton->Modify();
+		}
+
 		NameMapping->AddOrFindName(Name, Uid);
 
 		FFloatCurve * CurveToImport = static_cast<FFloatCurve *>(TargetSequence->RawCurveData.GetCurveData(Uid, FRawCurveTracks::FloatType));
diff --git a/Engine/Source/Runtime/Engine/Classes/Animation/SmartName.h b/Engine/Source/Runtime/Engine/Classes/Animation/SmartName.h
--- a/Engine/Source/Runtime/Engine/Classes/Animation/SmartName.h
+++ b/Engine/Source/Runtime/Engine/Classes/Animation/SmartName.h
@@ -50,6 +50,11 @@ struct ENGINE_API FSmartNameMapping
 	// @return bool - true if the name was found and removed, false if the name wasn't present in the mapping
 	bool Remove(const UID& Uid);
 
+	// Return UID * if it finds it
+	//  @param NewName - New name to set 
+	// @return UID pointer - null if it doesn't find. pointer if it finds. 
+	const FSmartNameMapping::UID* FindUID(const FName& Name);
+
 	// Check whether a name already exists in the mapping
 	// @param Uid - the UID to check
 	// @return bool - whether the name was found
diff --git a/Engine/Source/Runtime/Engine/Private/Animation/AnimSequenceBase.cpp b/Engine/Source/Runtime/Engine/Private/Animation/AnimSequenceBase.cpp
--- a/Engine/Source/Runtime/Engine/Private/Animation/AnimSequenceBase.cpp
+++ b/Engine/Source/Runtime/Engine/Private/Animation/AnimSequenceBase.cpp
@@ -470,13 +470,21 @@ void UAnimSequenceBase::VerifyCurveNames(USkeleton* Skeleton, const FName& NameC
 	TArray<DataType*> UnlinkedCurves;
 	for(DataType& Curve : CurveList)
 	{
-		if(!NameMapping->Exists(Curve.LastObservedName))
+		const FSmartNameMapping::UID* UID = NameMapping->FindUID(Curve.LastObservedName);
+		if(!UID)
 		{
 			// The skeleton doesn't know our name. Use the last observed name that was saved with the
 			// curve to create a new name. This can happen if a user saves an animation but not a skeleton
 			// either when updating the assets or editing the curves within.
 			UnlinkedCurves.Add(&Curve);
 		}
+		else if (Curve.CurveUid != *UID)// we verify if UID is correct
+		{
+			// if UID doesn't match, this is suspicious
+			// because same name but different UID is not idea
+			// so we'll fix up UID
+			Curve.CurveUid = *UID;
+		}
 	}
 
 	for(DataType* Curve : UnlinkedCurves)
diff --git a/Engine/Source/Runtime/Engine/Private/Animation/SmartName.cpp b/Engine/Source/Runtime/Engine/Private/Animation/SmartName.cpp
--- a/Engine/Source/Runtime/Engine/Private/Animation/SmartName.cpp
+++ b/Engine/Source/Runtime/Engine/Private/Animation/SmartName.cpp
@@ -149,6 +149,11 @@ bool FSmartNameMapping::Exists(const FName& Name)
 	return UidMap.FindKey(Name) != nullptr;
 }
 
+const FSmartNameMapping::UID* FSmartNameMapping::FindUID(const FName& Name)
+{
+	return UidMap.FindKey(Name);
+}
+
 FArchive& operator<<(FArchive& Ar, FSmartNameMapping& Elem)
 {
 	Elem.Serialize(Ar);
