#Physics: PhAT: Apply component space transform when we calculate min size

Merging

//depot/UE4-GDC2015/Engine/Source/...

to //depot/UE4/Engine/Source/...

diff --git a/Engine/Source/Editor/UnrealEd/Private/PhysicsAssetUtils.cpp b/Engine/Source/Editor/UnrealEd/Private/PhysicsAssetUtils.cpp
--- a/Engine/Source/Editor/UnrealEd/Private/PhysicsAssetUtils.cpp
+++ b/Engine/Source/Editor/UnrealEd/Private/PhysicsAssetUtils.cpp
@@ -393,10 +393,16 @@ bool CreateCollisionFromBone( UBodySetup* bs, USkeletalMesh* skelMesh, int32 Bon
 
 	FVector BoxCenter(0,0,0), BoxExtent(0,0,0);
 
+	FBox TransformedBox = BoneBox;
 	if( BoneBox.IsValid )
+	{
+		// make sure to apply scale to the box size
+		FMatrix BoneMatrix = skelMesh->GetComposedRefPoseMatrix(BoneIndex);
+		TransformedBox = BoneBox.TransformBy(FTransform(BoneMatrix));
 		BoneBox.GetCenterAndExtents(BoxCenter, BoxExtent);
+	}
 
-	float MinRad = BoxExtent.GetMin();
+	float MinRad = TransformedBox.GetExtent().GetMin();
 	float MinAllowedSize = MinPrimSize;
 
 	// If the primitive is going to be too small - just use some default numbers and let the user tweak.
diff --git a/Engine/Source/Runtime/Engine/Classes/Engine/SkeletalMesh.h b/Engine/Source/Runtime/Engine/Classes/Engine/SkeletalMesh.h
--- a/Engine/Source/Runtime/Engine/Classes/Engine/SkeletalMesh.h
+++ b/Engine/Source/Runtime/Engine/Classes/Engine/SkeletalMesh.h
@@ -762,6 +762,7 @@ public:
 	 *	Get the component orientation of a bone or socket. Transforms by parent bones.
 	 */
 	ENGINE_API FMatrix GetComposedRefPoseMatrix( FName InBoneName ) const;
+	ENGINE_API FMatrix GetComposedRefPoseMatrix( int32 InBoneIndex ) const;
 
 	/** Allocate and initialise bone mirroring table for this skeletal mesh. Default is source = destination for each bone. */
 	void InitBoneMirrorInfo();
diff --git a/Engine/Source/Runtime/Engine/Private/SkeletalMesh.cpp b/Engine/Source/Runtime/Engine/Private/SkeletalMesh.cpp
--- a/Engine/Source/Runtime/Engine/Private/SkeletalMesh.cpp
+++ b/Engine/Source/Runtime/Engine/Private/SkeletalMesh.cpp
@@ -2893,34 +2893,22 @@ FMatrix USkeletalMesh::GetComposedRefPoseMatrix( FName InBoneName ) const
 	if ( InBoneName != NAME_None )
 	{
 		int32 BoneIndex = RefSkeleton.FindBoneIndex(InBoneName);
-
-		if ( BoneIndex != INDEX_NONE )
+		if (BoneIndex != INDEX_NONE)
 		{
-			while ( BoneIndex != INDEX_NONE )
-			{
-				LocalPose = LocalPose * GetRefPoseMatrix( BoneIndex );
-				BoneIndex = RefSkeleton.GetParentIndex( BoneIndex );
-			}
+			return GetComposedRefPoseMatrix(BoneIndex);
 		}
 		else
 		{
-			USkeletalMeshSocket const* Socket = FindSocket( InBoneName );
+			USkeletalMeshSocket const* Socket = FindSocket(InBoneName);
 
-			if ( Socket != NULL )
+			if(Socket != NULL)
 			{
-				BoneIndex = RefSkeleton.FindBoneIndex( Socket->BoneName );
+				BoneIndex = RefSkeleton.FindBoneIndex(Socket->BoneName);
 
-				if ( BoneIndex != INDEX_NONE )
+				if(BoneIndex != INDEX_NONE)
 				{
-					const FRotationTranslationMatrix SocketMatrix( Socket->RelativeRotation, Socket->RelativeLocation );
-					LocalPose = SocketMatrix * GetRefPoseMatrix( BoneIndex );
-					BoneIndex = RefSkeleton.GetParentIndex( BoneIndex );
-
-					while ( BoneIndex != INDEX_NONE )
-					{
-						LocalPose = LocalPose * GetRefPoseMatrix( BoneIndex );
-						BoneIndex = RefSkeleton.GetParentIndex( BoneIndex );
-					}
+					const FRotationTranslationMatrix SocketMatrix(Socket->RelativeRotation, Socket->RelativeLocation);
+					LocalPose = SocketMatrix * GetComposedRefPoseMatrix(BoneIndex);
 				}
 			}
 		}
@@ -2929,6 +2917,23 @@ FMatrix USkeletalMesh::GetComposedRefPoseMatrix( FName InBoneName ) const
 	return LocalPose;
 }
 
+FMatrix USkeletalMesh::GetComposedRefPoseMatrix(int32 InBoneIndex) const
+{
+	FMatrix LocalPose(FMatrix::Identity);
+	int32 BoneIndex = InBoneIndex;
+
+	if(BoneIndex != INDEX_NONE)
+	{
+		while(BoneIndex != INDEX_NONE)
+		{
+			LocalPose = LocalPose * GetRefPoseMatrix(BoneIndex);
+			BoneIndex = RefSkeleton.GetParentIndex(BoneIndex);
+		}
+	}
+
+	return LocalPose;
+}
+
 TArray<USkeletalMeshSocket*>& USkeletalMesh::GetMeshOnlySocketList()
 {
 	return Sockets;
