Multicast RPCs are now handled in network replays

diff --git a/Engine/Source/Runtime/Core/Public/UObject/UnrealNames.inl b/Engine/Source/Runtime/Core/Public/UObject/UnrealNames.inl
--- a/Engine/Source/Runtime/Core/Public/UObject/UnrealNames.inl
+++ b/Engine/Source/Runtime/Core/Public/UObject/UnrealNames.inl
@@ -128,6 +128,7 @@ REGISTER_NAME(282,GameNetDriver)
 REGISTER_NAME(283,PendingNetDriver)
 REGISTER_NAME(284,BeaconNetDriver)
 REGISTER_NAME(285,FlushNetDormancy)
+REGISTER_NAME(286,DemoNetDriver)
 
 // Texture settings.
 REGISTER_NAME(300,Linear)
diff --git a/Engine/Source/Runtime/Engine/Private/Actor.cpp b/Engine/Source/Runtime/Engine/Private/Actor.cpp
--- a/Engine/Source/Runtime/Engine/Private/Actor.cpp
+++ b/Engine/Source/Runtime/Engine/Private/Actor.cpp
@@ -3080,6 +3080,15 @@ bool AActor::CallRemoteFunction( UFunction* Function, void* Parameters, FOutParm
 	if (NetDriver)
 	{
 		NetDriver->ProcessRemoteFunction(this, Function, Parameters, OutParms, Stack, NULL);
+		if (NetDriver->NetDriverName == NAME_GameNetDriver)
+		{
+			// Replicate any RPCs to the replay net driver so that they can get saved in network replays
+			NetDriver = GEngine->FindNamedNetDriver(GetWorld(), NAME_DemoNetDriver);
+			if (NetDriver)
+			{
+				NetDriver->ProcessRemoteFunction(this, Function, Parameters, OutParms, Stack, NULL);
+			}
+		}
 		return true;
 	}
 
diff --git a/Engine/Source/Runtime/Engine/Private/DemoNetDriver.cpp b/Engine/Source/Runtime/Engine/Private/DemoNetDriver.cpp
--- a/Engine/Source/Runtime/Engine/Private/DemoNetDriver.cpp
+++ b/Engine/Source/Runtime/Engine/Private/DemoNetDriver.cpp
@@ -450,7 +450,13 @@ void UDemoNetDriver::TickFlush( float DeltaSeconds )
 
 void UDemoNetDriver::ProcessRemoteFunction( class AActor* Actor, class UFunction* Function, void* Parameters, struct FOutParmRec* OutParms, struct FFrame* Stack, class UObject* SubObject )
 {
-
+	if (ClientConnections.Num() > 0 && ClientConnections[0] != NULL)
+	{
+		if ((Function->FunctionFlags & FUNC_NetMulticast))
+		{
+			InternalProcessRemoteFunction(Actor, SubObject, ClientConnections[0], Function, Parameters, OutParms, Stack, IsServer());
+		}
+	}
 }
 
 bool UDemoNetDriver::Exec( UWorld* InWorld, const TCHAR* Cmd, FOutputDevice& Ar )
