Native classes mirrored as USTRUCTs in Object.h changed to structs.
Typedefs mirrored as USTRUCTs in Object.h wrapped as structs.
Hack in UStructProperty::GetCPPTypeForwardDeclaration removed which was mapping some of these types for correct forward declarations.

diff --git a/Engine/Source/Developer/ImageCore/Public/ImageCore.h b/Engine/Source/Developer/ImageCore/Public/ImageCore.h
--- a/Engine/Source/Developer/ImageCore/Public/ImageCore.h
+++ b/Engine/Source/Developer/ImageCore/Public/ImageCore.h
@@ -132,16 +132,16 @@ public:
 		return RawData.GetData();
 	}
 
-	class FColor* AsBGRA8()
+	struct FColor* AsBGRA8()
 	{
 		check(Format == ERawImageFormat::BGRA8);
-		return (class FColor*)RawData.GetData();
+		return (struct FColor*)RawData.GetData();
 	}
 
-	class FColor* AsBGRE8()
+	struct FColor* AsBGRE8()
 	{
 		check(Format == ERawImageFormat::BGRE8);
-		return (class FColor*)RawData.GetData();
+		return (struct FColor*)RawData.GetData();
 	}
 
 	uint16* AsRGBA16()
@@ -170,16 +170,16 @@ public:
 		return RawData.GetData();
 	}
 
-	const class FColor* AsBGRA8() const
+	const struct FColor* AsBGRA8() const
 	{
 		check(Format == ERawImageFormat::BGRA8);
-		return (const class FColor*)RawData.GetData();
+		return (const struct FColor*)RawData.GetData();
 	}
 
-	const class FColor* AsBGRE8() const
+	const struct FColor* AsBGRE8() const
 	{
 		check(Format == ERawImageFormat::BGRE8);
-		return (class FColor*)RawData.GetData();
+		return (struct FColor*)RawData.GetData();
 	}
 
 	const uint16* AsRGBA16() const
diff --git a/Engine/Source/Editor/MaterialEditor/Private/MaterialEditorModule.cpp b/Engine/Source/Editor/MaterialEditor/Private/MaterialEditorModule.cpp
--- a/Engine/Source/Editor/MaterialEditor/Private/MaterialEditorModule.cpp
+++ b/Engine/Source/Editor/MaterialEditor/Private/MaterialEditorModule.cpp
@@ -69,7 +69,7 @@ public:
 		return NewMaterialInstanceEditor;
 	}
 	
-	virtual void GetVisibleMaterialParameters(const class UMaterial* Material, class UMaterialInstance* MaterialInstance, TArray<class FGuid>& VisibleExpressions)
+	virtual void GetVisibleMaterialParameters(const class UMaterial* Material, class UMaterialInstance* MaterialInstance, TArray<struct FGuid>& VisibleExpressions)
 	{
 		FMaterialEditorUtilities::GetVisibleMaterialParameters(Material, MaterialInstance, VisibleExpressions);
 	}
diff --git a/Engine/Source/Editor/MaterialEditor/Public/MaterialEditorModule.h b/Engine/Source/Editor/MaterialEditor/Public/MaterialEditorModule.h
--- a/Engine/Source/Editor/MaterialEditor/Public/MaterialEditorModule.h
+++ b/Engine/Source/Editor/MaterialEditor/Public/MaterialEditorModule.h
@@ -35,7 +35,7 @@ public:
 	 * @param	MaterialInstance	The material instance that contains all parameter overrides.
 	 * @param	VisisbleExpressions	The array that will contain the id's of the visible parameter expressions.
 	 */
-	virtual void GetVisibleMaterialParameters(const class UMaterial* Material, class UMaterialInstance* MaterialInstance, TArray<class FGuid>& VisibleExpressions) = 0;
+	virtual void GetVisibleMaterialParameters(const class UMaterial* Material, class UMaterialInstance* MaterialInstance, TArray<struct FGuid>& VisibleExpressions) = 0;
 
 	/** Delegates to be called to extend the material menus */
 	DECLARE_DELEGATE_RetVal_OneParam( TSharedRef<FExtender>, FMaterialMenuExtender, const TSharedRef<FUICommandList>);
diff --git a/Engine/Source/Runtime/Core/Private/Misc/Parse.cpp b/Engine/Source/Runtime/Core/Private/Misc/Parse.cpp
--- a/Engine/Source/Runtime/Core/Private/Misc/Parse.cpp
+++ b/Engine/Source/Runtime/Core/Private/Misc/Parse.cpp
@@ -618,7 +618,7 @@ bool FParse::Bool( const TCHAR* Stream, const TCHAR* Match, bool& OnOff )
 //
 // Get a globally unique identifier.
 //
-bool FParse::Value( const TCHAR* Stream, const TCHAR* Match, class FGuid& Guid )
+bool FParse::Value( const TCHAR* Stream, const TCHAR* Match, struct FGuid& Guid )
 {
 	TCHAR Temp[256];
 	if( !FParse::Value( Stream, Match, Temp, ARRAY_COUNT(Temp) ) )
diff --git a/Engine/Source/Runtime/Core/Public/GenericPlatform/GenericPlatformMisc.h b/Engine/Source/Runtime/Core/Public/GenericPlatform/GenericPlatformMisc.h
--- a/Engine/Source/Runtime/Core/Public/GenericPlatform/GenericPlatformMisc.h
+++ b/Engine/Source/Runtime/Core/Public/GenericPlatform/GenericPlatformMisc.h
@@ -422,11 +422,11 @@ public:
 	/**
 	 * Platform specific function for adding a named event that can be viewed in PIX
 	 */
-	FORCEINLINE static void BeginNamedEvent(const class FColor& Color,const TCHAR* Text)
+	FORCEINLINE static void BeginNamedEvent(const struct FColor& Color,const TCHAR* Text)
 	{
 	}
 
-	FORCEINLINE static void BeginNamedEvent(const class FColor& Color,const ANSICHAR* Text)
+	FORCEINLINE static void BeginNamedEvent(const struct FColor& Color,const ANSICHAR* Text)
 	{
 	}
 
@@ -515,7 +515,7 @@ public:
 	static void ClipboardPaste(class FString& Dest);
 
 	/** Create a new globally unique identifier. **/
-	static void CreateGuid(class FGuid& Result);
+	static void CreateGuid(struct FGuid& Result);
 
 	/** 
 	 * Show a message box if possible, otherwise print a message and return the default
diff --git a/Engine/Source/Runtime/Core/Public/Internationalization/Text.h b/Engine/Source/Runtime/Core/Public/Internationalization/Text.h
--- a/Engine/Source/Runtime/Core/Public/Internationalization/Text.h
+++ b/Engine/Source/Runtime/Core/Public/Internationalization/Text.h
@@ -5,7 +5,7 @@
 #include "Templates/SharedPointer.h"
 #include "UObject/NameTypes.h"
 
-class FTimespan;
+struct FTimespan;
 struct FDateTime;
 
 #define ENABLE_TEXT_ERROR_CHECKING_RESULTS (UE_BUILD_DEBUG | UE_BUILD_DEVELOPMENT | UE_BUILD_TEST )
diff --git a/Engine/Source/Runtime/Core/Public/Mac/MacPlatformMisc.h b/Engine/Source/Runtime/Core/Public/Mac/MacPlatformMisc.h
--- a/Engine/Source/Runtime/Core/Public/Mac/MacPlatformMisc.h
+++ b/Engine/Source/Runtime/Core/Public/Mac/MacPlatformMisc.h
@@ -73,7 +73,7 @@ struct CORE_API FMacPlatformMisc : public FGenericPlatformMisc
 	static const TCHAR* GetSystemErrorMessage(TCHAR* OutBuffer, int32 BufferCount, int32 Error);
 	static void ClipboardCopy(const TCHAR* Str);
 	static void ClipboardPaste(class FString& Dest);
-	static void CreateGuid(class FGuid& Result);
+	static void CreateGuid(struct FGuid& Result);
 	static EAppReturnType::Type MessageBoxExt( EAppMsgType::Type MsgType, const TCHAR* Text, const TCHAR* Caption );
 	static bool ControlScreensaver(EScreenSaverAction Action);
 	static bool CommandLineCommands();
diff --git a/Engine/Source/Runtime/Core/Public/Math/Box.h b/Engine/Source/Runtime/Core/Public/Math/Box.h
--- a/Engine/Source/Runtime/Core/Public/Math/Box.h
+++ b/Engine/Source/Runtime/Core/Public/Math/Box.h
@@ -9,7 +9,7 @@
  * Boxes describe an axis-aligned extent in three dimensions. They are used for many different things in the
  * Engine and in games, such as bounding volumes, collision detection and visibility calculation.
  */
-class FBox
+struct FBox
 {
 public:
 
diff --git a/Engine/Source/Runtime/Core/Public/Math/Color.h b/Engine/Source/Runtime/Core/Public/Math/Color.h
--- a/Engine/Source/Runtime/Core/Public/Math/Color.h
+++ b/Engine/Source/Runtime/Core/Public/Math/Color.h
@@ -7,8 +7,8 @@
 #include "Serialization/ArchiveBase.h"
 
 
-class FColor;
-class FVector;
+struct FColor;
+struct FVector;
 class FFloat16Color;
 
 
@@ -332,7 +332,7 @@ FORCEINLINE FLinearColor operator*(float Scalar,const FLinearColor& Color)
 //	FColor
 //
 
-class FColor
+struct FColor
 {
 public:
 	// Variables.
diff --git a/Engine/Source/Runtime/Core/Public/Math/IntPoint.h b/Engine/Source/Runtime/Core/Public/Math/IntPoint.h
--- a/Engine/Source/Runtime/Core/Public/Math/IntPoint.h
+++ b/Engine/Source/Runtime/Core/Public/Math/IntPoint.h
@@ -277,7 +277,7 @@ public:
 };
 
 
-/* FInterval inline functions
+/* FIntPoint inline functions
  *****************************************************************************/
 
 FORCEINLINE FIntPoint::FIntPoint() { }
diff --git a/Engine/Source/Runtime/Core/Public/Math/InterpCurve.h b/Engine/Source/Runtime/Core/Public/Math/InterpCurve.h
--- a/Engine/Source/Runtime/Core/Public/Math/InterpCurve.h
+++ b/Engine/Source/Runtime/Core/Public/Math/InterpCurve.h
@@ -570,20 +570,39 @@ void FInterpCurve<T>::CalcBounds(T& OutMin, T& OutMax, const T& Default) const
 /* Common type definitions
  *****************************************************************************/
 
-/** Type definition for floating point interpolation curves. */
-typedef FInterpCurve<float> FInterpCurveFloat;
-
-/** Type definition for 2D vector interpolation curves. */
-typedef FInterpCurve<FVector2D> FInterpCurveVector2D;
-
-/** Type definition for 3D vector interpolation curves. */
-typedef FInterpCurve<FVector> FInterpCurveVector;
-
-/** Type definition for quaternion interpolation curves. */
-typedef FInterpCurve<FQuat> FInterpCurveQuat;
-
-/** Type definition for vector pair interpolation curves. */
-typedef FInterpCurve<FTwoVectors> FInterpCurveTwoVectors;
-
-/** Type definition for color interpolation curves. */
-typedef FInterpCurve<FLinearColor> FInterpCurveLinearColor;
+#define DEFINE_INTERPCURVE_WRAPPER_STRUCT(Name, ElementType) \
+	struct Name : FInterpCurve<ElementType> \
+	{ \
+	private: \
+		typedef FInterpCurve<ElementType> Super; \
+	 \
+	public: \
+		Name() \
+			: Super() \
+		{ \
+		} \
+		 \
+		Name(const Super& Other) \
+			: Super( Other ) \
+		{ \
+		} \
+	}; \
+	 \
+	template <> \
+	struct TIsBitwiseConstructible<Name, FInterpCurve<ElementType>> \
+	{ \
+		enum { Value = true }; \
+	}; \
+	 \
+	template <> \
+	struct TIsBitwiseConstructible<FInterpCurve<ElementType>, Name> \
+	{ \
+		enum { Value = true }; \
+	};
+
+DEFINE_INTERPCURVE_WRAPPER_STRUCT(FInterpCurveFloat,       float)
+DEFINE_INTERPCURVE_WRAPPER_STRUCT(FInterpCurveVector2D,    FVector2D)
+DEFINE_INTERPCURVE_WRAPPER_STRUCT(FInterpCurveVector,      FVector)
+DEFINE_INTERPCURVE_WRAPPER_STRUCT(FInterpCurveQuat,        FQuat)
+DEFINE_INTERPCURVE_WRAPPER_STRUCT(FInterpCurveTwoVectors,  FTwoVectors)
+DEFINE_INTERPCURVE_WRAPPER_STRUCT(FInterpCurveLinearColor, FLinearColor)
diff --git a/Engine/Source/Runtime/Core/Public/Math/Interval.h b/Engine/Source/Runtime/Core/Public/Math/Interval.h
--- a/Engine/Source/Runtime/Core/Public/Math/Interval.h
+++ b/Engine/Source/Runtime/Core/Public/Math/Interval.h
@@ -204,6 +204,48 @@ public:
 /* Default intervals for built-in types
  *****************************************************************************/
 
-typedef TInterval<float> FInterval;
-typedef TInterval<float> FFloatInterval;
-typedef TInterval<int32> FInt32Interval;
+#define DEFINE_INTERVAL_WRAPPER_STRUCT(Name, ElementType) \
+	struct Name : TInterval<ElementType> \
+	{ \
+	private: \
+		typedef TInterval<ElementType> Super; \
+		 \
+	public:  \
+		Name() \
+			: Super() \
+		{ \
+		} \
+		 \
+		Name( const Super& Other ) \
+			: Super( Other ) \
+		{ \
+		} \
+		 \
+		Name( ElementType InMin, ElementType InMax ) \
+			: Super( InMin, InMax ) \
+		{ \
+		} \
+		 \
+		friend Name Intersect( const Name& A, const Name& B ) \
+		{ \
+			return Intersect( static_cast<const Super&>( A ), static_cast<const Super&>( B ) ); \
+		} \
+	}; \
+	 \
+	template <> \
+	struct TIsBitwiseConstructible<Name, TInterval<ElementType>> \
+	{ \
+		enum { Value = true }; \
+	}; \
+	 \
+	template <> \
+	struct TIsBitwiseConstructible<TInterval<ElementType>, Name> \
+	{ \
+		enum { Value = true }; \
+	};
+
+DEPRECATED(4.8, "FInterval is deprecated, please use FFloatInterval instead.")
+DEFINE_INTERVAL_WRAPPER_STRUCT(FInterval, float)
+
+DEFINE_INTERVAL_WRAPPER_STRUCT(FFloatInterval, float)
+DEFINE_INTERVAL_WRAPPER_STRUCT(FInt32Interval, int32)
diff --git a/Engine/Source/Runtime/Core/Public/Math/Matrix.h b/Engine/Source/Runtime/Core/Public/Math/Matrix.h
--- a/Engine/Source/Runtime/Core/Public/Math/Matrix.h
+++ b/Engine/Source/Runtime/Core/Public/Math/Matrix.h
@@ -10,7 +10,7 @@
  * opposed to Res = Mat1 * Mat2.
  * Matrix elements are accessed with M[RowIndex][ColumnIndex].
  */
-class FMatrix
+struct FMatrix
 {
 public:
 	union
diff --git a/Engine/Source/Runtime/Core/Public/Math/OrientedBox.h b/Engine/Source/Runtime/Core/Public/Math/OrientedBox.h
--- a/Engine/Source/Runtime/Core/Public/Math/OrientedBox.h
+++ b/Engine/Source/Runtime/Core/Public/Math/OrientedBox.h
@@ -60,7 +60,7 @@ public:
 	 *
 	 * @param Axis The unit vector defining the axis to project the box onto.
 	 */
-	FORCEINLINE FInterval Project(const FVector& Axis) const;
+	FORCEINLINE FFloatInterval Project(const FVector& Axis) const;
 };
 
 
@@ -84,7 +84,7 @@ FORCEINLINE void FOrientedBox::CalcVertices( FVector* Verts ) const
 }
 
 
-FORCEINLINE FInterval FOrientedBox::Project( const FVector& Axis ) const
+FORCEINLINE FFloatInterval FOrientedBox::Project( const FVector& Axis ) const
 {
 	static const float Signs[] = {-1.0f, 1.0f};
 
@@ -94,7 +94,7 @@ FORCEINLINE FInterval FOrientedBox::Project( const FVector& Axis ) const
 	float ProjectedAxisY = Axis | (ExtentY * AxisY);
 	float ProjectedAxisZ = Axis | (ExtentZ * AxisZ);
 
-	FInterval ProjectionInterval;
+	FFloatInterval ProjectionInterval;
 
 	for (int32 i = 0; i < 2; i++)
 	{
diff --git a/Engine/Source/Runtime/Core/Public/Math/Plane.h b/Engine/Source/Runtime/Core/Public/Math/Plane.h
--- a/Engine/Source/Runtime/Core/Public/Math/Plane.h
+++ b/Engine/Source/Runtime/Core/Public/Math/Plane.h
@@ -9,7 +9,7 @@
  * Stores the coeffecients as Ax+By+Cz=D.
  * Note that this is different than many other Plane classes that use Ax+By+Cz+D=0.
  */
-MS_ALIGN(16) class FPlane
+MS_ALIGN(16) struct FPlane
 	: public FVector
 {
 public:
diff --git a/Engine/Source/Runtime/Core/Public/Math/Quat.h b/Engine/Source/Runtime/Core/Public/Math/Quat.h
--- a/Engine/Source/Runtime/Core/Public/Math/Quat.h
+++ b/Engine/Source/Runtime/Core/Public/Math/Quat.h
@@ -14,7 +14,7 @@
  * Example: LocalToWorld = (LocalToWorld * DeltaRotation) will change rotation in local space by DeltaRotation.
  * Example: LocalToWorld = (DeltaRotation * LocalToWorld) will change rotation in world space by DeltaRotation.
  */
-MS_ALIGN(16) class FQuat 
+MS_ALIGN(16) struct FQuat 
 {
 public:
 
diff --git a/Engine/Source/Runtime/Core/Public/Math/RandomStream.h b/Engine/Source/Runtime/Core/Public/Math/RandomStream.h
--- a/Engine/Source/Runtime/Core/Public/Math/RandomStream.h
+++ b/Engine/Source/Runtime/Core/Public/Math/RandomStream.h
@@ -8,7 +8,7 @@
  *
  * Very bad quality in the lower bits. Don't use the modulus (%) operator.
  */
-class FRandomStream
+struct FRandomStream
 {
 #if !defined(COREUOBJECT_API)
 	#define MAYBE_COREUOBJECT_API
diff --git a/Engine/Source/Runtime/Core/Public/Math/Range.h b/Engine/Source/Runtime/Core/Public/Math/Range.h
--- a/Engine/Source/Runtime/Core/Public/Math/Range.h
+++ b/Engine/Source/Runtime/Core/Public/Math/Range.h
@@ -6,7 +6,7 @@
 /**
  * Template for ranges.
  *
- * Note: This class is not intended for interval arithmetic (see FInterval for that).
+ * Note: This class is not intended for interval arithmetic (see TInterval for that).
  *
  * A range represents a contiguous set of elements that only stores the set's
  * lower and upper bound values (aka. endpoints) for storage efficiency. Bound
@@ -38,7 +38,7 @@
  * @param ElementType - The type of elements represented by the range
  *
  * @see RangeBound
- * @see FInterval
+ * @see TInterval
  */
 template<typename ElementType> class TRange
 {
@@ -80,7 +80,7 @@ public:
 	 * @param InLowerBound The range's lower bound.
 	 * @param InUpperBound The range's upper bound.
 	 */
-	TRange( const BoundsType& InLowerBound, const BoundsType& InUpperBound )
+	explicit TRange( const BoundsType& InLowerBound, const BoundsType& InUpperBound )
 		: LowerBound(InLowerBound)
 		, UpperBound(InUpperBound)
 	{ }
@@ -669,23 +669,120 @@ private:
 /* Default ranges for built-in types
  *****************************************************************************/
 
-/** Defines a discrete range of dates. */
-typedef TRange<FDateTime> FDateRange;
-
-/** Defines a discrete range of dates. */
-typedef TRange<double> FDoubleRange;
-
-/** Defines a discrete range of dates. */
-typedef TRange<float> FFloatRange;
-
-/** Defines a discrete range of 8-bit signed integers. */
-typedef TRange<int8> FInt8Range;
-
-/** Defines a discrete range of 16-bit signed integers. */
-typedef TRange<int16> FInt16Range;
-
-/** Defines a discrete range of 32-bit signed integers. */
-typedef TRange<int32> FInt32Range;
-
-/** Defines a discrete range of 64-bit signed integers. */
-typedef TRange<int64> FInt64Range;
+#define DEFINE_RANGE_WRAPPER_STRUCT(Name, ElementType) \
+	struct Name : TRange<ElementType> \
+	{ \
+	private: \
+		typedef TRange<ElementType> Super; \
+	 \
+	public: \
+		Name() \
+			: Super() \
+		{ \
+		} \
+		 \
+		Name( const Super& Rhs ) \
+			: Super( Rhs ) \
+		{ \
+		} \
+		 \
+		explicit Name( const ElementType& A ) \
+			: Super( A ) \
+		{ \
+		} \
+		 \
+		explicit Name( const ElementType& A, const ElementType& B ) \
+			: Super( A, B ) \
+		{ \
+		} \
+		 \
+		explicit Name( const TRangeBound<ElementType>& InLowerBound, const TRangeBound<ElementType>& InUpperBound ) \
+			: Super( InLowerBound, InUpperBound ) \
+		{ \
+		} \
+		 \
+		TArray<Name> Split( const ElementType& Element ) const \
+		{ \
+			return TArray<Name>( Super::Split( Element ) ); \
+		} \
+		 \
+		static FORCEINLINE TArray<Name> Difference( const Name& X, const Name& Y ) \
+		{ \
+			return TArray<Name>( Super::Difference( X, Y ) ); \
+		} \
+		 \
+		static FORCEINLINE Name Empty() \
+		{ \
+			return Super::Empty(); \
+		} \
+		 \
+		static FORCEINLINE Name Hull( const Name& X, const Name& Y ) \
+		{ \
+			return Super::Hull( X, Y ); \
+		} \
+		 \
+		static FORCEINLINE Name Hull( const TArray<Name>& Ranges ) \
+		{ \
+			return Super::Hull( reinterpret_cast<const TArray<Super>&>( Ranges ) ); \
+		} \
+		 \
+		static FORCEINLINE Name Intersection( const Name& X, const Name& Y ) \
+		{ \
+			return Super::Intersection( X, Y ); \
+		} \
+		 \
+		static FORCEINLINE Name Intersection( const TArray<Name>& Ranges ) \
+		{ \
+			return Super::Intersection( reinterpret_cast<const TArray<Super>&>( Ranges ) ); \
+		} \
+		 \
+		static FORCEINLINE TArray<Name> Union( const Name& X, const Name& Y ) \
+		{ \
+			return TArray<Name>( Super::Union( X, Y ) ); \
+		} \
+		 \
+		static FORCEINLINE Name All() \
+		{ \
+			return Super::All(); \
+		} \
+		 \
+		static FORCEINLINE Name AtLeast( const ElementType& Value ) \
+		{ \
+			return Super::AtLeast( Value ); \
+		} \
+		 \
+		static FORCEINLINE Name AtMost( const ElementType& Value ) \
+		{ \
+			return Super::AtMost( Value ); \
+		} \
+		 \
+		static FORCEINLINE TRange GreaterThan( const ElementType& Value ) \
+		{ \
+			return Super::GreaterThan( Value ); \
+		} \
+		 \
+		static FORCEINLINE TRange LessThan( const ElementType& Value ) \
+		{ \
+			return Super::LessThan( Value ); \
+		} \
+	}; \
+	 \
+	template <> \
+	struct TIsBitwiseConstructible<Name, TRange<ElementType>> \
+	{ \
+		enum { Value = true }; \
+	}; \
+	 \
+	template <> \
+	struct TIsBitwiseConstructible<TRange<ElementType>, Name> \
+	{ \
+		enum { Value = true }; \
+	};
+
+DEFINE_RANGE_WRAPPER_STRUCT(FDateRange,   FDateTime)
+DEFINE_RANGE_WRAPPER_STRUCT(FDoubleRange, double)
+DEFINE_RANGE_WRAPPER_STRUCT(FFloatRange,  float)
+DEFINE_RANGE_WRAPPER_STRUCT(FInt8Range,   int8)
+DEFINE_RANGE_WRAPPER_STRUCT(FInt16Range,  int16)
+DEFINE_RANGE_WRAPPER_STRUCT(FInt32Range,  int32)
+DEFINE_RANGE_WRAPPER_STRUCT(FInt64Range,  int64)
diff --git a/Engine/Source/Runtime/Core/Public/Math/RangeBound.h b/Engine/Source/Runtime/Core/Public/Math/RangeBound.h
--- a/Engine/Source/Runtime/Core/Public/Math/RangeBound.h
+++ b/Engine/Source/Runtime/Core/Public/Math/RangeBound.h
@@ -2,7 +2,6 @@
 
 #pragma once
 
-
 namespace ERangeBoundTypes
 {
 	/**
@@ -25,7 +24,8 @@ namespace ERangeBoundTypes
 /**
  * Template for range bounds.
  */
-template<typename ElementType> class TRangeBound
+template<typename ElementType>
+class TRangeBound
 {
 public:
 
@@ -313,23 +313,85 @@ private:
 /* Default range bounds for built-in types
  *****************************************************************************/
 
-/** Defines a discrete range bound for dates. */
-typedef TRangeBound<FDateTime> FDateRangeBound;
-
-/** Defines a discrete range bound for dates. */
-typedef TRangeBound<double> FDoubleRangeBound;
-
-/** Defines a discrete range bound for dates. */
-typedef TRangeBound<float> FFloatRangeBound;
-
-/** Defines a discrete range bound for 8-bit signed integers. */
-typedef TRangeBound<int8> FInt8RangeBound;
-
-/** Defines a discrete range bound for 16-bit signed integers. */
-typedef TRangeBound<int16> FInt16RangeBound;
-
-/** Defines a discrete range bound for 32-bit signed integers. */
-typedef TRangeBound<int32> FInt32RangeBound;
+#define DEFINE_RANGEBOUND_WRAPPER_STRUCT(Name, ElementType) \
+	struct Name : TRangeBound<ElementType> \
+	{ \
+	private: \
+		typedef TRangeBound<ElementType> Super; \
+	 \
+	public: \
+		Name() \
+			: Super() \
+		{ \
+		} \
+		 \
+		Name(const Super& Other) \
+			: Super( Other ) \
+		{ \
+		} \
+		 \
+		Name( const int64& InValue ) \
+			: Super( InValue ) \
+		{ \
+		} \
+		 \
+		static FORCEINLINE Name Exclusive( const ElementType& Value ) \
+		{ \
+			return static_cast<const Name&>( Super::Exclusive( Value ) ); \
+		} \
+		 \
+		static FORCEINLINE Name Inclusive( const ElementType& Value ) \
+		{ \
+			return static_cast<const Name&>( Super::Inclusive( Value ) ); \
+		} \
+		 \
+		static FORCEINLINE Name Open() \
+		{ \
+			return static_cast<const Name&>( Super::Open() ); \
+		} \
+		 \
+		static FORCEINLINE Name FlipInclusion( const Name& Bound ) \
+		{ \
+			return static_cast<const Name&>( Super::FlipInclusion( Bound ) ); \
+		} \
+		 \
+		static FORCEINLINE const Name& MaxLower( const Name& A, const Name& B ) \
+		{ \
+			return static_cast<const Name&>( Super::MaxLower( A, B ) ); \
+		} \
+		 \
+		static FORCEINLINE const Name& MaxUpper( const Name& A, const Name& B ) \
+		{ \
+			return static_cast<const Name&>( Super::MaxUpper( A, B ) ); \
+		} \
+		 \
+		static FORCEINLINE const Name& MinLower( const Name& A, const Name& B ) \
+		{ \
+			return static_cast<const Name&>( Super::MinLower( A, B ) ); \
+		} \
+		 \
+		static FORCEINLINE const Name& MinUpper( const Name& A, const Name& B ) \
+		{ \
+			return static_cast<const Name&>( Super::MinUpper( A, B ) ); \
+		} \
+	}; \
+	 \
+	template <> \
+	struct TIsBitwiseConstructible<Name, TRangeBound<ElementType>> \
+	{ \
+		enum { Value = true }; \
+	}; \
+	 \
+	template <> \
+	struct TIsBitwiseConstructible<TRangeBound<ElementType>, Name> \
+	{ \
+		enum { Value = true }; \
+	};
 
-/** Defines a discrete range bound for 64-bit signed integers. */
-typedef TRangeBound<int64> FInt64RangeBound;
+DEFINE_RANGEBOUND_WRAPPER_STRUCT(FDateRangeBound,   FDateTime)
+DEFINE_RANGEBOUND_WRAPPER_STRUCT(FDoubleRangeBound, double)
+DEFINE_RANGEBOUND_WRAPPER_STRUCT(FFloatRangeBound,  float)
+DEFINE_RANGEBOUND_WRAPPER_STRUCT(FInt8RangeBound,   int8)
+DEFINE_RANGEBOUND_WRAPPER_STRUCT(FInt16RangeBound,  int16)
+DEFINE_RANGEBOUND_WRAPPER_STRUCT(FInt32RangeBound,  int32)
+DEFINE_RANGEBOUND_WRAPPER_STRUCT(FInt64RangeBound,  int64)
diff --git a/Engine/Source/Runtime/Core/Public/Math/Rotator.h b/Engine/Source/Runtime/Core/Public/Math/Rotator.h
--- a/Engine/Source/Runtime/Core/Public/Math/Rotator.h
+++ b/Engine/Source/Runtime/Core/Public/Math/Rotator.h
@@ -8,7 +8,7 @@
  *
  * All rotation values are stored in degrees.
  */
-class FRotator
+struct FRotator
 {
 public:
 
diff --git a/Engine/Source/Runtime/Core/Public/Math/Transform.h b/Engine/Source/Runtime/Core/Public/Math/Transform.h
--- a/Engine/Source/Runtime/Core/Public/Math/Transform.h
+++ b/Engine/Source/Runtime/Core/Public/Math/Transform.h
@@ -19,7 +19,7 @@
  * Example: LocalToWorld = (LocalToWorld * DeltaRotation) will change rotation in world space by DeltaRotation.
  */
 
-class FTransform
+struct FTransform
 {
 #if !defined(COREUOBJECT_API)
 	#define MAYBE_COREUOBJECT_API
diff --git a/Engine/Source/Runtime/Core/Public/Math/TransformVectorized.h b/Engine/Source/Runtime/Core/Public/Math/TransformVectorized.h
--- a/Engine/Source/Runtime/Core/Public/Math/TransformVectorized.h
+++ b/Engine/Source/Runtime/Core/Public/Math/TransformVectorized.h
@@ -21,7 +21,7 @@
  * Example: LocalToWorld = (LocalToWorld * DeltaRotation) will change rotation in world space by DeltaRotation.
  */
 
-MS_ALIGN(16) class FTransform
+MS_ALIGN(16) struct FTransform
 {
 #if !defined(COREUOBJECT_API)
 	#define MAYBE_COREUOBJECT_API
diff --git a/Engine/Source/Runtime/Core/Public/Math/TwoVectors.h b/Engine/Source/Runtime/Core/Public/Math/TwoVectors.h
--- a/Engine/Source/Runtime/Core/Public/Math/TwoVectors.h
+++ b/Engine/Source/Runtime/Core/Public/Math/TwoVectors.h
@@ -6,7 +6,7 @@
 /**
  * A pair of 3D vectors.
  */
-class FTwoVectors
+struct FTwoVectors
 {
 public:
 
diff --git a/Engine/Source/Runtime/Core/Public/Math/UnrealMathUtility.h b/Engine/Source/Runtime/Core/Public/Math/UnrealMathUtility.h
--- a/Engine/Source/Runtime/Core/Public/Math/UnrealMathUtility.h
+++ b/Engine/Source/Runtime/Core/Public/Math/UnrealMathUtility.h
@@ -18,15 +18,15 @@
 -----------------------------------------------------------------------------*/
 
 // Forward declarations.
-class  FVector;
-class  FVector4;
-class  FPlane;
-class  FBox;
-class  FRotator;
-class  FMatrix;
-class  FQuat;
-class  FTwoVectors;
-class  FTransform;
+struct  FVector;
+struct  FVector4;
+struct  FPlane;
+struct  FBox;
+struct  FRotator;
+struct  FMatrix;
+struct  FQuat;
+struct  FTwoVectors;
+struct  FTransform;
 class  FSphere;
 struct FVector2D;
 struct FLinearColor;
diff --git a/Engine/Source/Runtime/Core/Public/Math/Vector.h b/Engine/Source/Runtime/Core/Public/Math/Vector.h
--- a/Engine/Source/Runtime/Core/Public/Math/Vector.h
+++ b/Engine/Source/Runtime/Core/Public/Math/Vector.h
@@ -6,7 +6,7 @@
 /**
  * A vector in 3-D space composed of components (X, Y, Z) with floating point precision.
  */
-class FVector 
+struct FVector 
 {
 public:
 
diff --git a/Engine/Source/Runtime/Core/Public/Math/Vector4.h b/Engine/Source/Runtime/Core/Public/Math/Vector4.h
--- a/Engine/Source/Runtime/Core/Public/Math/Vector4.h
+++ b/Engine/Source/Runtime/Core/Public/Math/Vector4.h
@@ -6,7 +6,7 @@
 /**
  * A 4D homogeneous vector, 4x1 FLOATs, 16-byte aligned.
  */
-MS_ALIGN(16) class FVector4
+MS_ALIGN(16) struct FVector4
 {
 public:
 
diff --git a/Engine/Source/Runtime/Core/Public/Misc/CoreMisc.h b/Engine/Source/Runtime/Core/Public/Misc/CoreMisc.h
--- a/Engine/Source/Runtime/Core/Public/Misc/CoreMisc.h
+++ b/Engine/Source/Runtime/Core/Public/Misc/CoreMisc.h
@@ -197,7 +197,7 @@ struct CORE_API FFileHelper
 	 *
 	 * @return true if success
 	 */
-	static bool CreateBitmap( const TCHAR* Pattern, int32 DataWidth, int32 DataHeight, const class FColor* Data, struct FIntRect* SubRectangle = NULL, IFileManager* FileManager = &IFileManager::Get(), FString* OutFilename = NULL, bool bInWriteAlpha = false );
+	static bool CreateBitmap( const TCHAR* Pattern, int32 DataWidth, int32 DataHeight, const struct FColor* Data, struct FIntRect* SubRectangle = NULL, IFileManager* FileManager = &IFileManager::Get(), FString* OutFilename = NULL, bool bInWriteAlpha = false );
 
 	/**
 	 * Generates the next unique bitmap filename with a specified extension
diff --git a/Engine/Source/Runtime/Core/Public/Misc/DateTime.h b/Engine/Source/Runtime/Core/Public/Misc/DateTime.h
--- a/Engine/Source/Runtime/Core/Public/Misc/DateTime.h
+++ b/Engine/Source/Runtime/Core/Public/Misc/DateTime.h
@@ -61,7 +61,7 @@ enum class EMonthOfYear
  * calculate the number of days in a given month and year, check for leap years and determine the
  * time of day, day of week and month of year of a given date and time.
  *
- * The companion class FTimespan is provided for enabling date and time based arithmetic, such as
+ * The companion struct FTimespan is provided for enabling date and time based arithmetic, such as
  * calculating the difference between two dates or adding a certain amount of time to a given date.
  *
  * Ranges of dates and times can be represented by the FDateRange class.
diff --git a/Engine/Source/Runtime/Core/Public/Misc/Guid.h b/Engine/Source/Runtime/Core/Public/Misc/Guid.h
--- a/Engine/Source/Runtime/Core/Public/Misc/Guid.h
+++ b/Engine/Source/Runtime/Core/Public/Misc/Guid.h
@@ -59,7 +59,7 @@ enum class EGuidFormats
 /**
  * Implements a globally unique identifier.
  */
-class FGuid
+struct FGuid
 {
 public:
 
diff --git a/Engine/Source/Runtime/Core/Public/Misc/Parse.h b/Engine/Source/Runtime/Core/Public/Misc/Parse.h
--- a/Engine/Source/Runtime/Core/Public/Misc/Parse.h
+++ b/Engine/Source/Runtime/Core/Public/Misc/Parse.h
@@ -23,7 +23,7 @@ struct CORE_API FParse
 	/** Parses a uint32. */
 	static bool Value( const TCHAR* Stream, const TCHAR* Match, uint32& Value );
 	/** Parses a globally unique identifier. */
-	static bool Value( const TCHAR* Stream, const TCHAR* Match, class FGuid& Guid );
+	static bool Value( const TCHAR* Stream, const TCHAR* Match, struct FGuid& Guid );
 	/** Parses a string from a text string. */
 	static bool Value( const TCHAR* Stream, const TCHAR* Match, TCHAR* Value, int32 MaxLen, bool bShouldStopOnComma=true );
 	/** Parses a byte. */
diff --git a/Engine/Source/Runtime/Core/Public/Misc/Timespan.h b/Engine/Source/Runtime/Core/Public/Misc/Timespan.h
--- a/Engine/Source/Runtime/Core/Public/Misc/Timespan.h
+++ b/Engine/Source/Runtime/Core/Public/Misc/Timespan.h
@@ -49,7 +49,7 @@ namespace ETimespan
  *
  * @see FDateTime
  */
-class FTimespan
+struct FTimespan
 {
 public:
 
diff --git a/Engine/Source/Runtime/Core/Public/Serialization/ArchiveBase.h b/Engine/Source/Runtime/Core/Public/Serialization/ArchiveBase.h
--- a/Engine/Source/Runtime/Core/Public/Serialization/ArchiveBase.h
+++ b/Engine/Source/Runtime/Core/Public/Serialization/ArchiveBase.h
@@ -553,7 +553,7 @@ public:
 	 *
 	 * @param Guid The guid of the custom version.  This must have previously been registered with FCustomVersionRegistration.
 	 */
-	void UsingCustomVersion(const class FGuid& Guid);
+	void UsingCustomVersion(const struct FGuid& Guid);
 
 	/**
 	 * Queries a custom version from the archive.  If the archive is being used to write, the custom version must have already been registered.
@@ -561,7 +561,7 @@ public:
 	 * @param Key The guid of the custom version to query.
 	 * @return The version number, or 0 if the custom tag isn't stored in the archive.
 	 */
-	int32 CustomVer(const class FGuid& Key) const;
+	int32 CustomVer(const struct FGuid& Key) const;
 
 	FORCEINLINE bool IsLoading() const
 	{
@@ -751,7 +751,7 @@ public:
 	 * @param Version - The version number to set key to
 	 * @param FriendlyName - Friendly name corresponding to the key
 	 */
-	void SetCustomVersion(const class FGuid& Key, int32 Version, FString FriendlyName);
+	void SetCustomVersion(const struct FGuid& Key, int32 Version, FString FriendlyName);
 
 	/**
 	 * Toggle saving as Unicode. This is needed when we need to make sure ANSI strings are saved as Unicode
diff --git a/Engine/Source/Runtime/Core/Public/WinRT/WinRTMisc.h b/Engine/Source/Runtime/Core/Public/WinRT/WinRTMisc.h
--- a/Engine/Source/Runtime/Core/Public/WinRT/WinRTMisc.h
+++ b/Engine/Source/Runtime/Core/Public/WinRT/WinRTMisc.h
@@ -43,7 +43,7 @@ struct CORE_API FWinRTMisc : public FGenericPlatformMisc
 	static void LowLevelOutputDebugString(const TCHAR *Message);
 	static void RequestExit(bool Force);
 	static const TCHAR* GetSystemErrorMessage(TCHAR* OutBuffer, int32 BufferCount, int32 Error);
-	static void CreateGuid(class FGuid& Result);
+	static void CreateGuid(struct FGuid& Result);
 	static int32 NumberOfCores();
 
 	/** Get the application root directory. */
diff --git a/Engine/Source/Runtime/Core/Public/Windows/WindowsPlatformMisc.h b/Engine/Source/Runtime/Core/Public/Windows/WindowsPlatformMisc.h
--- a/Engine/Source/Runtime/Core/Public/Windows/WindowsPlatformMisc.h
+++ b/Engine/Source/Runtime/Core/Public/Windows/WindowsPlatformMisc.h
@@ -70,7 +70,7 @@ struct CORE_API FWindowsPlatformMisc
 	static const TCHAR* GetSystemErrorMessage(TCHAR* OutBuffer, int32 BufferCount, int32 Error);
 	static void ClipboardCopy(const TCHAR* Str);
 	static void ClipboardPaste(class FString& Dest);
-	static void CreateGuid(class FGuid& Result);
+	static void CreateGuid(struct FGuid& Result);
 	static EAppReturnType::Type MessageBoxExt( EAppMsgType::Type MsgType, const TCHAR* Text, const TCHAR* Caption );
 	static void PreventScreenSaver();
 	static bool CommandLineCommands();
diff --git a/Engine/Source/Runtime/CoreUObject/Private/UObject/PropertyStruct.cpp b/Engine/Source/Runtime/CoreUObject/Private/UObject/PropertyStruct.cpp
--- a/Engine/Source/Runtime/CoreUObject/Private/UObject/PropertyStruct.cpp
+++ b/Engine/Source/Runtime/CoreUObject/Private/UObject/PropertyStruct.cpp
@@ -242,32 +242,6 @@ FString UStructProperty::GetCPPType( FString* ExtendedTypeText/*=NULL*/, uint32
 
 FString UStructProperty::GetCPPTypeForwardDeclaration() const
 {
-	// The following types are classes, but reflection system sees them as structs.
-	static FName ClassNames[] =
-	{
-		NAME_Box,
-		NAME_Color,
-		FName(TEXT("Guid")),
-		NAME_Plane,
-		NAME_Quat,
-		FName(TEXT("RandomStream")),
-		FName(TEXT("Rotator")),
-		FName(TEXT("Timespan")),
-		NAME_Transform,
-		NAME_Vector,
-		NAME_Matrix
-	};
-
-	auto Name = Struct->GetFName();
-
-	for (auto ClassName : ClassNames)
-	{
-		if (ClassName == Name)
-		{
-			return FString::Printf(TEXT("class F%s;"), *Struct->GetName());
-		}
-	}
-
 	return FString::Printf(TEXT("struct F%s;"), *Struct->GetName());
 }
 
diff --git a/Engine/Source/Runtime/Engine/Classes/Animation/AnimInstance.h b/Engine/Source/Runtime/Engine/Classes/Animation/AnimInstance.h
--- a/Engine/Source/Runtime/Engine/Classes/Animation/AnimInstance.h
+++ b/Engine/Source/Runtime/Engine/Classes/Animation/AnimInstance.h
@@ -19,7 +19,7 @@ class APawn;
 class UAnimationAsset;
 class UCanvas;
 class UWorld;
-class FTransform;
+struct FTransform;
 class FDebugDisplayInfo;
 
 DECLARE_DELEGATE_OneParam(FOnMontageStarted, UAnimMontage*)
diff --git a/Engine/Source/Runtime/Engine/Classes/Distributions/DistributionFloat.h b/Engine/Source/Runtime/Engine/Classes/Distributions/DistributionFloat.h
--- a/Engine/Source/Runtime/Engine/Classes/Distributions/DistributionFloat.h
+++ b/Engine/Source/Runtime/Engine/Classes/Distributions/DistributionFloat.h
@@ -57,7 +57,7 @@ public:
 		/**
 		 * Get the value at the specified F
 		 */
-		ENGINE_API float GetValue(float F=0.0f, UObject* Data=NULL, class FRandomStream* InRandomStream = NULL);
+		ENGINE_API float GetValue(float F=0.0f, UObject* Data=NULL, struct FRandomStream* InRandomStream = NULL);
 
 		/**
 		 * Get the min and max values
@@ -115,7 +115,7 @@ class UDistributionFloat : public UDistribution
 	virtual uint32 InitializeRawEntry(float Time, float* Values) const;
 
 	/** @todo document */
-	virtual float GetValue( float F = 0.f, UObject* Data = NULL, class FRandomStream* InRandomStream = NULL ) const;
+	virtual float GetValue( float F = 0.f, UObject* Data = NULL, struct FRandomStream* InRandomStream = NULL ) const;
 
 	// Begin FCurveEdInterface Interface
 	virtual void GetInRange(float& MinIn, float& MaxIn) const override;
diff --git a/Engine/Source/Runtime/Engine/Classes/Distributions/DistributionFloatConstant.h b/Engine/Source/Runtime/Engine/Classes/Distributions/DistributionFloatConstant.h
--- a/Engine/Source/Runtime/Engine/Classes/Distributions/DistributionFloatConstant.h
+++ b/Engine/Source/Runtime/Engine/Classes/Distributions/DistributionFloatConstant.h
@@ -20,7 +20,7 @@ class UDistributionFloatConstant : public UDistributionFloat
 	// End UObject Interface
 
 	// Begin UDistributionFloat Interface
-	virtual float GetValue( float F = 0.f, UObject* Data = NULL, class FRandomStream* InRandomStream = NULL ) const override;
+	virtual float GetValue( float F = 0.f, UObject* Data = NULL, struct FRandomStream* InRandomStream = NULL ) const override;
 	// End UDistributionFloat Interface
 
 
diff --git a/Engine/Source/Runtime/Engine/Classes/Distributions/DistributionFloatConstantCurve.h b/Engine/Source/Runtime/Engine/Classes/Distributions/DistributionFloatConstantCurve.h
--- a/Engine/Source/Runtime/Engine/Classes/Distributions/DistributionFloatConstantCurve.h
+++ b/Engine/Source/Runtime/Engine/Classes/Distributions/DistributionFloatConstantCurve.h
@@ -16,7 +16,7 @@ class UDistributionFloatConstantCurve : public UDistributionFloat
 
 
 	// Begin UDistributionFloat Interface
-	virtual float GetValue( float F = 0.f, UObject* Data = NULL, class FRandomStream* InRandomStream = NULL ) const override;
+	virtual float GetValue( float F = 0.f, UObject* Data = NULL, struct FRandomStream* InRandomStream = NULL ) const override;
 	// End UDistributionFloat Interface
 
 	// Begin FCurveEdInterface interface
diff --git a/Engine/Source/Runtime/Engine/Classes/Distributions/DistributionFloatParameterBase.h b/Engine/Source/Runtime/Engine/Classes/Distributions/DistributionFloatParameterBase.h
--- a/Engine/Source/Runtime/Engine/Classes/Distributions/DistributionFloatParameterBase.h
+++ b/Engine/Source/Runtime/Engine/Classes/Distributions/DistributionFloatParameterBase.h
@@ -36,7 +36,7 @@ class UDistributionFloatParameterBase : public UDistributionFloatConstant
 
 
 	// Begin UDistributionFloat Interface
-	virtual float GetValue( float F = 0.f, UObject* Data = NULL, class FRandomStream* InRandomStream = NULL ) const override;
+	virtual float GetValue( float F = 0.f, UObject* Data = NULL, struct FRandomStream* InRandomStream = NULL ) const override;
 	// End UDistributionFloat Interface
 	
 	
diff --git a/Engine/Source/Runtime/Engine/Classes/Distributions/DistributionFloatUniform.h b/Engine/Source/Runtime/Engine/Classes/Distributions/DistributionFloatUniform.h
--- a/Engine/Source/Runtime/Engine/Classes/Distributions/DistributionFloatUniform.h
+++ b/Engine/Source/Runtime/Engine/Classes/Distributions/DistributionFloatUniform.h
@@ -25,7 +25,7 @@ class UDistributionFloatUniform : public UDistributionFloat
 	// End UObject Interface
 
 	// Begin UDistributionFloat Interface
-	virtual float GetValue( float F = 0.f, UObject* Data = NULL, class FRandomStream* InRandomStream = NULL ) const override;
+	virtual float GetValue( float F = 0.f, UObject* Data = NULL, struct FRandomStream* InRandomStream = NULL ) const override;
 
 	[REDACTED_EMAIL]: Currently, consoles need this? At least until we have some sort of cooking/packaging step!
 	virtual ERawDistributionOperation GetOperation() const override;
diff --git a/Engine/Source/Runtime/Engine/Classes/Distributions/DistributionFloatUniformCurve.h b/Engine/Source/Runtime/Engine/Classes/Distributions/DistributionFloatUniformCurve.h
--- a/Engine/Source/Runtime/Engine/Classes/Distributions/DistributionFloatUniformCurve.h
+++ b/Engine/Source/Runtime/Engine/Classes/Distributions/DistributionFloatUniformCurve.h
@@ -17,7 +17,7 @@ class UDistributionFloatUniformCurve : public UDistributionFloat
 
 
 	// Begin UDistributionFloat Interface
-	virtual float GetValue( float F = 0.f, UObject* Data = NULL, class FRandomStream* InRandomStream = NULL ) const override;
+	virtual float GetValue( float F = 0.f, UObject* Data = NULL, struct FRandomStream* InRandomStream = NULL ) const override;
 	[REDACTED_EMAIL]: Currently, consoles need this? At least until we have some sort of cooking/packaging step!
 	virtual ERawDistributionOperation GetOperation() const override;
 	virtual uint32 InitializeRawEntry(float Time, float* Values) const override;
diff --git a/Engine/Source/Runtime/Engine/Classes/Distributions/DistributionVector.h b/Engine/Source/Runtime/Engine/Classes/Distributions/DistributionVector.h
--- a/Engine/Source/Runtime/Engine/Classes/Distributions/DistributionVector.h
+++ b/Engine/Source/Runtime/Engine/Classes/Distributions/DistributionVector.h
@@ -87,7 +87,7 @@ public:
 		/**
 		 * Get the value at the specified F
 		 */
-		ENGINE_API FVector GetValue(float F=0.0f, UObject* Data=NULL, int32 LastExtreme=0, class FRandomStream* InRandomStream = NULL);
+		ENGINE_API FVector GetValue(float F=0.0f, UObject* Data=NULL, int32 LastExtreme=0, struct FRandomStream* InRandomStream = NULL);
 
 		/**
 		 * Get the min and max values
@@ -145,7 +145,7 @@ class ENGINE_API UDistributionVector : public UDistribution
 	 */
 	virtual uint32 InitializeRawEntry(float Time, float* Values) const;
 
-	virtual FVector	GetValue( float F = 0.f, UObject* Data = NULL, int32 LastExtreme = 0, class FRandomStream* InRandomStream = NULL ) const;
+	virtual FVector	GetValue( float F = 0.f, UObject* Data = NULL, int32 LastExtreme = 0, struct FRandomStream* InRandomStream = NULL ) const;
 
 	// Begin FCurveEdInterface Interface
 	virtual void GetInRange(float& MinIn, float& MaxIn) const;
diff --git a/Engine/Source/Runtime/Engine/Classes/Distributions/DistributionVectorConstant.h b/Engine/Source/Runtime/Engine/Classes/Distributions/DistributionVectorConstant.h
--- a/Engine/Source/Runtime/Engine/Classes/Distributions/DistributionVectorConstant.h
+++ b/Engine/Source/Runtime/Engine/Classes/Distributions/DistributionVectorConstant.h
@@ -23,7 +23,7 @@ class ENGINE_API UDistributionVectorConstant : public UDistributionVector
 
 
 	//Begin UDistributionVector Interface
-	virtual FVector	GetValue( float F = 0.f, UObject* Data = NULL, int32 LastExtreme = 0, class FRandomStream* InRandomStream = NULL ) const override;
+	virtual FVector	GetValue( float F = 0.f, UObject* Data = NULL, int32 LastExtreme = 0, struct FRandomStream* InRandomStream = NULL ) const override;
 	virtual	void	GetRange(FVector& OutMin, FVector& OutMax) const override;
 	//End UDistributionVector Interface
 
diff --git a/Engine/Source/Runtime/Engine/Classes/Distributions/DistributionVectorConstantCurve.h b/Engine/Source/Runtime/Engine/Classes/Distributions/DistributionVectorConstantCurve.h
--- a/Engine/Source/Runtime/Engine/Classes/Distributions/DistributionVectorConstantCurve.h
+++ b/Engine/Source/Runtime/Engine/Classes/Distributions/DistributionVectorConstantCurve.h
@@ -23,7 +23,7 @@ class UDistributionVectorConstantCurve : public UDistributionVector
 
 
 	//Begin UDistributionVector Interface
-	virtual FVector	GetValue( float F = 0.f, UObject* Data = NULL, int32 LastExtreme = 0, class FRandomStream* InRandomStream = NULL ) const override;
+	virtual FVector	GetValue( float F = 0.f, UObject* Data = NULL, int32 LastExtreme = 0, struct FRandomStream* InRandomStream = NULL ) const override;
 	virtual	void	GetRange(FVector& OutMin, FVector& OutMax) const override;
 	//End UDistributionVector Interface
 
diff --git a/Engine/Source/Runtime/Engine/Classes/Distributions/DistributionVectorParameterBase.h b/Engine/Source/Runtime/Engine/Classes/Distributions/DistributionVectorParameterBase.h
--- a/Engine/Source/Runtime/Engine/Classes/Distributions/DistributionVectorParameterBase.h
+++ b/Engine/Source/Runtime/Engine/Classes/Distributions/DistributionVectorParameterBase.h
@@ -30,7 +30,7 @@ class UDistributionVectorParameterBase : public UDistributionVectorConstant
 
 
 	//Begin UDistributionVector Interface
-	virtual FVector GetValue(float F = 0.f, UObject* Data = NULL, int32 Extreme = 0, class FRandomStream* InRandomStream = NULL) const override;
+	virtual FVector GetValue(float F = 0.f, UObject* Data = NULL, int32 Extreme = 0, struct FRandomStream* InRandomStream = NULL) const override;
 	virtual bool CanBeBaked() const override { return false; }
 	//End UDistributionVector Interface
 	
diff --git a/Engine/Source/Runtime/Engine/Classes/Distributions/DistributionVectorUniform.h b/Engine/Source/Runtime/Engine/Classes/Distributions/DistributionVectorUniform.h
--- a/Engine/Source/Runtime/Engine/Classes/Distributions/DistributionVectorUniform.h
+++ b/Engine/Source/Runtime/Engine/Classes/Distributions/DistributionVectorUniform.h
@@ -37,7 +37,7 @@ class UDistributionVectorUniform : public UDistributionVector
 	virtual void PostLoad() override;
 	// End UObject Interface
 
-	virtual FVector	GetValue( float F = 0.f, UObject* Data = NULL, int32 LastExtreme = 0, class FRandomStream* InRandomStream = NULL ) const override;
+	virtual FVector	GetValue( float F = 0.f, UObject* Data = NULL, int32 LastExtreme = 0, struct FRandomStream* InRandomStream = NULL ) const override;
 
 	//Begin UDistributionVector Interface
 
diff --git a/Engine/Source/Runtime/Engine/Classes/Distributions/DistributionVectorUniformCurve.h b/Engine/Source/Runtime/Engine/Classes/Distributions/DistributionVectorUniformCurve.h
--- a/Engine/Source/Runtime/Engine/Classes/Distributions/DistributionVectorUniformCurve.h
+++ b/Engine/Source/Runtime/Engine/Classes/Distributions/DistributionVectorUniformCurve.h
@@ -30,7 +30,7 @@ class UDistributionVectorUniformCurve : public UDistributionVector
 	UPROPERTY(EditAnywhere, Category=DistributionVectorUniformCurve)
 	uint32 bUseExtremes:1;
 
-	virtual FVector	GetValue( float F = 0.f, UObject* Data = NULL, int32 LastExtreme = 0, class FRandomStream* InRandomStream = NULL ) const override;
+	virtual FVector	GetValue( float F = 0.f, UObject* Data = NULL, int32 LastExtreme = 0, struct FRandomStream* InRandomStream = NULL ) const override;
 
 	//Begin UDistributionVector Interface
 	[REDACTED_EMAIL]: Currently, consoles need this? At least until we have some sort of cooking/packaging step!
diff --git a/Engine/Source/Runtime/Engine/Classes/Particles/Color/ParticleModuleColor.h b/Engine/Source/Runtime/Engine/Classes/Particles/Color/ParticleModuleColor.h
--- a/Engine/Source/Runtime/Engine/Classes/Particles/Color/ParticleModuleColor.h
+++ b/Engine/Source/Runtime/Engine/Classes/Particles/Color/ParticleModuleColor.h
@@ -48,7 +48,7 @@ class UParticleModuleColor : public UParticleModuleColorBase
 	 *	@param	SpawnTime			The time of the spawn
 	 *	@param	InRandomStream		The random stream to use for retrieving random values
 	 */
-	void SpawnEx(FParticleEmitterInstance* Owner, int32 Offset, float SpawnTime, class FRandomStream* InRandomStream, FBaseParticle* ParticleBase);
+	void SpawnEx(FParticleEmitterInstance* Owner, int32 Offset, float SpawnTime, struct FRandomStream* InRandomStream, FBaseParticle* ParticleBase);
 
 };
 
diff --git a/Engine/Source/Runtime/Engine/Classes/Particles/Lifetime/ParticleModuleLifetime.h b/Engine/Source/Runtime/Engine/Classes/Particles/Lifetime/ParticleModuleLifetime.h
--- a/Engine/Source/Runtime/Engine/Classes/Particles/Lifetime/ParticleModuleLifetime.h
+++ b/Engine/Source/Runtime/Engine/Classes/Particles/Lifetime/ParticleModuleLifetime.h
@@ -43,7 +43,7 @@ class UParticleModuleLifetime : public UParticleModuleLifetimeBase
 	 *	@param	SpawnTime			The time of the spawn
 	 *	@param	InRandomStream		The random stream to use for retrieving random values
 	 */
-	void SpawnEx(FParticleEmitterInstance* Owner, int32 Offset, float SpawnTime, class FRandomStream* InRandomStream, FBaseParticle* ParticleBase);
+	void SpawnEx(FParticleEmitterInstance* Owner, int32 Offset, float SpawnTime, struct FRandomStream* InRandomStream, FBaseParticle* ParticleBase);
 };
 
 
diff --git a/Engine/Source/Runtime/Engine/Classes/Particles/Light/ParticleModuleLight.h b/Engine/Source/Runtime/Engine/Classes/Particles/Light/ParticleModuleLight.h
--- a/Engine/Source/Runtime/Engine/Classes/Particles/Light/ParticleModuleLight.h
+++ b/Engine/Source/Runtime/Engine/Classes/Particles/Light/ParticleModuleLight.h
@@ -69,7 +69,7 @@ class UParticleModuleLight : public UParticleModuleLightBase
 	virtual void Render3DPreview(FParticleEmitterInstance* Owner, const FSceneView* View,FPrimitiveDrawInterface* PDI) override;
 	//End UParticleModule Interface
 
-	void SpawnEx(FParticleEmitterInstance* Owner, int32 Offset, float SpawnTime, class FRandomStream* InRandomStream, FBaseParticle* ParticleBase);
+	void SpawnEx(FParticleEmitterInstance* Owner, int32 Offset, float SpawnTime, struct FRandomStream* InRandomStream, FBaseParticle* ParticleBase);
 };
 
 
diff --git a/Engine/Source/Runtime/Engine/Classes/Particles/Location/ParticleModuleLocation.h b/Engine/Source/Runtime/Engine/Classes/Particles/Location/ParticleModuleLocation.h
--- a/Engine/Source/Runtime/Engine/Classes/Particles/Location/ParticleModuleLocation.h
+++ b/Engine/Source/Runtime/Engine/Classes/Particles/Location/ParticleModuleLocation.h
@@ -46,7 +46,7 @@ protected:
 	 *	@param	SpawnTime			The time of the spawn
 	 *	@param	InRandomStream		The random stream to use for retrieving random values
 	 */
-	virtual void SpawnEx(FParticleEmitterInstance* Owner, int32 Offset, float SpawnTime, class FRandomStream* InRandomStream, FBaseParticle* ParticleBase);
+	virtual void SpawnEx(FParticleEmitterInstance* Owner, int32 Offset, float SpawnTime, struct FRandomStream* InRandomStream, FBaseParticle* ParticleBase);
 
 public:
 	//Begin UObject Interface
diff --git a/Engine/Source/Runtime/Engine/Classes/Particles/Location/ParticleModuleLocationPrimitiveBase.h b/Engine/Source/Runtime/Engine/Classes/Particles/Location/ParticleModuleLocationPrimitiveBase.h
--- a/Engine/Source/Runtime/Engine/Classes/Particles/Location/ParticleModuleLocationPrimitiveBase.h
+++ b/Engine/Source/Runtime/Engine/Classes/Particles/Location/ParticleModuleLocationPrimitiveBase.h
@@ -65,7 +65,7 @@ class UParticleModuleLocationPrimitiveBase : public UParticleModuleLocationBase
 	//End UObject Interface
 
 	//@todo document
-	virtual void	DetermineUnitDirection(FParticleEmitterInstance* Owner, FVector& vUnitDir, class FRandomStream* InRandomStream);
+	virtual void	DetermineUnitDirection(FParticleEmitterInstance* Owner, FVector& vUnitDir, struct FRandomStream* InRandomStream);
 };
 
 
diff --git a/Engine/Source/Runtime/Engine/Classes/Particles/Location/ParticleModuleLocationPrimitiveCylinder.h b/Engine/Source/Runtime/Engine/Classes/Particles/Location/ParticleModuleLocationPrimitiveCylinder.h
--- a/Engine/Source/Runtime/Engine/Classes/Particles/Location/ParticleModuleLocationPrimitiveCylinder.h
+++ b/Engine/Source/Runtime/Engine/Classes/Particles/Location/ParticleModuleLocationPrimitiveCylinder.h
@@ -68,7 +68,7 @@ class UParticleModuleLocationPrimitiveCylinder : public UParticleModuleLocationP
 	 *	@param	SpawnTime			The time of the spawn
 	 *	@param	InRandomStream		The random stream to use for retrieving random values
 	 */
-	void SpawnEx(FParticleEmitterInstance* Owner, int32 Offset, float SpawnTime, class FRandomStream* InRandomStream, FBaseParticle* ParticleBase);
+	void SpawnEx(FParticleEmitterInstance* Owner, int32 Offset, float SpawnTime, struct FRandomStream* InRandomStream, FBaseParticle* ParticleBase);
 };
 
 
diff --git a/Engine/Source/Runtime/Engine/Classes/Particles/Location/ParticleModuleLocationPrimitiveSphere.h b/Engine/Source/Runtime/Engine/Classes/Particles/Location/ParticleModuleLocationPrimitiveSphere.h
--- a/Engine/Source/Runtime/Engine/Classes/Particles/Location/ParticleModuleLocationPrimitiveSphere.h
+++ b/Engine/Source/Runtime/Engine/Classes/Particles/Location/ParticleModuleLocationPrimitiveSphere.h
@@ -41,7 +41,7 @@ class UParticleModuleLocationPrimitiveSphere : public UParticleModuleLocationPri
 	 *	@param	SpawnTime			The time of the spawn
 	 *	@param	InRandomStream		The random stream to use for retrieving random values
 	 */
-	void SpawnEx(FParticleEmitterInstance* Owner, int32 Offset, float SpawnTime, class FRandomStream* InRandomStream, FBaseParticle* ParticleBase);
+	void SpawnEx(FParticleEmitterInstance* Owner, int32 Offset, float SpawnTime, struct FRandomStream* InRandomStream, FBaseParticle* ParticleBase);
 };
 
 
diff --git a/Engine/Source/Runtime/Engine/Classes/Particles/Location/ParticleModuleLocationPrimitiveTriangle.h b/Engine/Source/Runtime/Engine/Classes/Particles/Location/ParticleModuleLocationPrimitiveTriangle.h
--- a/Engine/Source/Runtime/Engine/Classes/Particles/Location/ParticleModuleLocationPrimitiveTriangle.h
+++ b/Engine/Source/Runtime/Engine/Classes/Particles/Location/ParticleModuleLocationPrimitiveTriangle.h
@@ -50,7 +50,7 @@ class UParticleModuleLocationPrimitiveTriangle : public UParticleModuleLocationB
 	 *	@param	SpawnTime			The time of the spawn
 	 *	@param	InRandomStream		The random stream to use for retrieving random values
 	 */
-	void SpawnEx(FParticleEmitterInstance* Owner, int32 Offset, float SpawnTime, class FRandomStream* InRandomStream, FBaseParticle* ParticleBase);
+	void SpawnEx(FParticleEmitterInstance* Owner, int32 Offset, float SpawnTime, struct FRandomStream* InRandomStream, FBaseParticle* ParticleBase);
 };
 
 
diff --git a/Engine/Source/Runtime/Engine/Classes/Particles/Location/ParticleModuleLocationWorldOffset.h b/Engine/Source/Runtime/Engine/Classes/Particles/Location/ParticleModuleLocationWorldOffset.h
--- a/Engine/Source/Runtime/Engine/Classes/Particles/Location/ParticleModuleLocationWorldOffset.h
+++ b/Engine/Source/Runtime/Engine/Classes/Particles/Location/ParticleModuleLocationWorldOffset.h
@@ -13,7 +13,7 @@ class UParticleModuleLocationWorldOffset : public UParticleModuleLocation
 
 protected:
 	//Begin UParticleModuleLocation Interface
-	virtual void SpawnEx(FParticleEmitterInstance* Owner, int32 Offset, float SpawnTime, class FRandomStream* InRandomStream, FBaseParticle* ParticleBase) override;
+	virtual void SpawnEx(FParticleEmitterInstance* Owner, int32 Offset, float SpawnTime, struct FRandomStream* InRandomStream, FBaseParticle* ParticleBase) override;
 	//End UParticleModuleLocation Interface
 };
 
diff --git a/Engine/Source/Runtime/Engine/Classes/Particles/Parameter/ParticleModuleParameterDynamic.h b/Engine/Source/Runtime/Engine/Classes/Particles/Parameter/ParticleModuleParameterDynamic.h
--- a/Engine/Source/Runtime/Engine/Classes/Particles/Parameter/ParticleModuleParameterDynamic.h
+++ b/Engine/Source/Runtime/Engine/Classes/Particles/Parameter/ParticleModuleParameterDynamic.h
@@ -135,7 +135,7 @@ class UParticleModuleParameterDynamic : public UParticleModuleParameterBase
 	 *	@param	SpawnTime			The time of the spawn
 	 *	@param	InRandomStream		The random stream to use for retrieving random values
 	 */
-	void SpawnEx(FParticleEmitterInstance* Owner, int32 Offset, float SpawnTime, class FRandomStream* InRandomStream, FBaseParticle* ParticleBase);
+	void SpawnEx(FParticleEmitterInstance* Owner, int32 Offset, float SpawnTime, struct FRandomStream* InRandomStream, FBaseParticle* ParticleBase);
 	
 	/**
 	 *	Update the parameter names with the given material...
@@ -157,7 +157,7 @@ class UParticleModuleParameterDynamic : public UParticleModuleParameterBase
 	 *
 	 *	@return	float			The value for the parameter.
 	 */
-	FORCEINLINE float GetParameterValue(FEmitterDynamicParameter& InDynParams, FBaseParticle& Particle, FParticleEmitterInstance* Owner, class FRandomStream* InRandomStream)
+	FORCEINLINE float GetParameterValue(FEmitterDynamicParameter& InDynParams, FBaseParticle& Particle, FParticleEmitterInstance* Owner, struct FRandomStream* InRandomStream)
 	{
 		float ScaleValue = 1.0f;
 		float DistributionValue = 1.0f;
@@ -195,7 +195,7 @@ class UParticleModuleParameterDynamic : public UParticleModuleParameterBase
 	 *
 	 *	@return	float			The value for the parameter.
 	 */
-	FORCEINLINE float GetParameterValue_UserSet(FEmitterDynamicParameter& InDynParams, FBaseParticle& Particle, FParticleEmitterInstance* Owner, class FRandomStream* InRandomStream)
+	FORCEINLINE float GetParameterValue_UserSet(FEmitterDynamicParameter& InDynParams, FBaseParticle& Particle, FParticleEmitterInstance* Owner, struct FRandomStream* InRandomStream)
 	{
 		return InDynParams.ParamValue.GetValue(InDynParams.bUseEmitterTime ? Owner->EmitterTime : Particle.RelativeTime, Owner->Component, InRandomStream);
 	}
diff --git a/Engine/Source/Runtime/Engine/Classes/Particles/Rotation/ParticleModuleMeshRotation.h b/Engine/Source/Runtime/Engine/Classes/Particles/Rotation/ParticleModuleMeshRotation.h
--- a/Engine/Source/Runtime/Engine/Classes/Particles/Rotation/ParticleModuleMeshRotation.h
+++ b/Engine/Source/Runtime/Engine/Classes/Particles/Rotation/ParticleModuleMeshRotation.h
@@ -38,7 +38,7 @@ class UParticleModuleMeshRotation : public UParticleModuleRotationBase
 	 *	@param	SpawnTime			The time of the spawn
 	 *	@param	InRandomStream		The random stream to use for retrieving random values
 	 */
-	void SpawnEx(FParticleEmitterInstance* Owner, int32 Offset, float SpawnTime, class FRandomStream* InRandomStream, FBaseParticle* ParticleBase);
+	void SpawnEx(FParticleEmitterInstance* Owner, int32 Offset, float SpawnTime, struct FRandomStream* InRandomStream, FBaseParticle* ParticleBase);
 
 };
 
diff --git a/Engine/Source/Runtime/Engine/Classes/Particles/Rotation/ParticleModuleRotation.h b/Engine/Source/Runtime/Engine/Classes/Particles/Rotation/ParticleModuleRotation.h
--- a/Engine/Source/Runtime/Engine/Classes/Particles/Rotation/ParticleModuleRotation.h
+++ b/Engine/Source/Runtime/Engine/Classes/Particles/Rotation/ParticleModuleRotation.h
@@ -39,7 +39,7 @@ class UParticleModuleRotation : public UParticleModuleRotationBase
 	 *	@param	SpawnTime			The time of the spawn
 	 *	@param	InRandomStream		The random stream to use for retrieving random values
 	 */
-	void SpawnEx(FParticleEmitterInstance* Owner, int32 Offset, float SpawnTime, class FRandomStream* InRandomStream, FBaseParticle* ParticleBase);
+	void SpawnEx(FParticleEmitterInstance* Owner, int32 Offset, float SpawnTime, struct FRandomStream* InRandomStream, FBaseParticle* ParticleBase);
 };
 
 
diff --git a/Engine/Source/Runtime/Engine/Classes/Particles/RotationRate/ParticleModuleMeshRotationRate.h b/Engine/Source/Runtime/Engine/Classes/Particles/RotationRate/ParticleModuleMeshRotationRate.h
--- a/Engine/Source/Runtime/Engine/Classes/Particles/RotationRate/ParticleModuleMeshRotationRate.h
+++ b/Engine/Source/Runtime/Engine/Classes/Particles/RotationRate/ParticleModuleMeshRotationRate.h
@@ -35,7 +35,7 @@ class UParticleModuleMeshRotationRate : public UParticleModuleRotationRateBase
 	 *	@param	SpawnTime			The time of the spawn
 	 *	@param	InRandomStream		The random stream to use for retrieving random values
 	 */
-	void SpawnEx(FParticleEmitterInstance* Owner, int32 Offset, float SpawnTime, class FRandomStream* InRandomStream, FBaseParticle* ParticleBase);
+	void SpawnEx(FParticleEmitterInstance* Owner, int32 Offset, float SpawnTime, struct FRandomStream* InRandomStream, FBaseParticle* ParticleBase);
 
 };
 
diff --git a/Engine/Source/Runtime/Engine/Classes/Particles/RotationRate/ParticleModuleRotationRate.h b/Engine/Source/Runtime/Engine/Classes/Particles/RotationRate/ParticleModuleRotationRate.h
--- a/Engine/Source/Runtime/Engine/Classes/Particles/RotationRate/ParticleModuleRotationRate.h
+++ b/Engine/Source/Runtime/Engine/Classes/Particles/RotationRate/ParticleModuleRotationRate.h
@@ -41,7 +41,7 @@ class UParticleModuleRotationRate : public UParticleModuleRotationRateBase
 	 *	@param	SpawnTime			The time of the spawn
 	 *	@param	InRandomStream		The random stream to use for retrieving random values
 	 */
-	void SpawnEx(FParticleEmitterInstance* Owner, int32 Offset, float SpawnTime, class FRandomStream* InRandomStream, FBaseParticle* ParticleBase);
+	void SpawnEx(FParticleEmitterInstance* Owner, int32 Offset, float SpawnTime, struct FRandomStream* InRandomStream, FBaseParticle* ParticleBase);
 };
 
 
diff --git a/Engine/Source/Runtime/Engine/Classes/Particles/Size/ParticleModuleSize.h b/Engine/Source/Runtime/Engine/Classes/Particles/Size/ParticleModuleSize.h
--- a/Engine/Source/Runtime/Engine/Classes/Particles/Size/ParticleModuleSize.h
+++ b/Engine/Source/Runtime/Engine/Classes/Particles/Size/ParticleModuleSize.h
@@ -41,7 +41,7 @@ class UParticleModuleSize : public UParticleModuleSizeBase
 	 *	@param	SpawnTime			The time of the spawn
 	 *	@param	InRandomStream		The random stream to use for retrieving random values
 	 */
-	void SpawnEx(FParticleEmitterInstance* Owner, int32 Offset, float SpawnTime, class FRandomStream* InRandomStream, FBaseParticle* ParticleBase);
+	void SpawnEx(FParticleEmitterInstance* Owner, int32 Offset, float SpawnTime, struct FRandomStream* InRandomStream, FBaseParticle* ParticleBase);
 };
 
 
diff --git a/Engine/Source/Runtime/Engine/Classes/Particles/Velocity/ParticleModuleVelocity.h b/Engine/Source/Runtime/Engine/Classes/Particles/Velocity/ParticleModuleVelocity.h
--- a/Engine/Source/Runtime/Engine/Classes/Particles/Velocity/ParticleModuleVelocity.h
+++ b/Engine/Source/Runtime/Engine/Classes/Particles/Velocity/ParticleModuleVelocity.h
@@ -47,7 +47,7 @@ class UParticleModuleVelocity : public UParticleModuleVelocityBase
 	 *	@param	SpawnTime			The time of the spawn
 	 *	@param	InRandomStream		The random stream to use for retrieving random values
 	 */
-	void SpawnEx(FParticleEmitterInstance* Owner, int32 Offset, float SpawnTime, class FRandomStream* InRandomStream, FBaseParticle* ParticleBase);
+	void SpawnEx(FParticleEmitterInstance* Owner, int32 Offset, float SpawnTime, struct FRandomStream* InRandomStream, FBaseParticle* ParticleBase);
 };
 
 
diff --git a/Engine/Source/Runtime/Engine/Classes/Particles/Velocity/ParticleModuleVelocityCone.h b/Engine/Source/Runtime/Engine/Classes/Particles/Velocity/ParticleModuleVelocityCone.h
--- a/Engine/Source/Runtime/Engine/Classes/Particles/Velocity/ParticleModuleVelocityCone.h
+++ b/Engine/Source/Runtime/Engine/Classes/Particles/Velocity/ParticleModuleVelocityCone.h
@@ -45,7 +45,7 @@ class UParticleModuleVelocityCone : public UParticleModuleVelocityBase
 	 *	@param	SpawnTime			The time of the spawn
 	 *	@param	InRandomStream		The random stream to use for retrieving random values
 	 */
-	void			SpawnEx(FParticleEmitterInstance* Owner, int32 Offset, float SpawnTime, class FRandomStream* InRandomStream, FBaseParticle* ParticleBase);
+	void			SpawnEx(FParticleEmitterInstance* Owner, int32 Offset, float SpawnTime, struct FRandomStream* InRandomStream, FBaseParticle* ParticleBase);
 
 };
 
diff --git a/Engine/Source/Runtime/Engine/Private/Distributions.cpp b/Engine/Source/Runtime/Engine/Private/Distributions.cpp
--- a/Engine/Source/Runtime/Engine/Private/Distributions.cpp
+++ b/Engine/Source/Runtime/Engine/Private/Distributions.cpp
@@ -651,7 +651,7 @@ static void OptimizeLookupTable( FDistributionLookupTable* Table, float ErrorThr
 	*Table = OriginalTable;
 }
 
-void FRawDistribution::GetValue1(float Time, float* Value, int32 Extreme, class FRandomStream* InRandomStream) const
+void FRawDistribution::GetValue1(float Time, float* Value, int32 Extreme, struct FRandomStream* InRandomStream) const
 {
 	switch (LookupTable.Op)
 	{
@@ -671,7 +671,7 @@ void FRawDistribution::GetValue1(float Time, float* Value, int32 Extreme, class
 	}
 }
 
-void FRawDistribution::GetValue3(float Time, float* Value, int32 Extreme, class FRandomStream* InRandomStream) const
+void FRawDistribution::GetValue3(float Time, float* Value, int32 Extreme, struct FRandomStream* InRandomStream) const
 {
 	switch (LookupTable.Op)
 	{
@@ -687,7 +687,7 @@ void FRawDistribution::GetValue3(float Time, float* Value, int32 Extreme, class
 	}
 }
 
-void FRawDistribution::GetValue1Extreme(float Time, float* InValue, int32 Extreme, class FRandomStream* InRandomStream) const
+void FRawDistribution::GetValue1Extreme(float Time, float* InValue, int32 Extreme, struct FRandomStream* InRandomStream) const
 {
 	float* RESTRICT Value = InValue;
 	const float* Entry1;
@@ -701,7 +701,7 @@ void FRawDistribution::GetValue1Extreme(float Time, float* InValue, int32 Extrem
 	Value[0] = FMath::Lerp(NewEntry1[InitialElement + 0], NewEntry2[InitialElement + 0], LerpAlpha);
 }
 
-void FRawDistribution::GetValue3Extreme(float Time, float* InValue, int32 Extreme, class FRandomStream* InRandomStream) const
+void FRawDistribution::GetValue3Extreme(float Time, float* InValue, int32 Extreme, struct FRandomStream* InRandomStream) const
 {
 	float* RESTRICT Value = InValue;
 	const float* Entry1;
@@ -721,7 +721,7 @@ void FRawDistribution::GetValue3Extreme(float Time, float* InValue, int32 Extrem
 	Value[2] = T2;
 }
 
-void FRawDistribution::GetValue1Random(float Time, float* InValue, class FRandomStream* InRandomStream) const
+void FRawDistribution::GetValue1Random(float Time, float* InValue, struct FRandomStream* InRandomStream) const
 {
 	float* RESTRICT Value = InValue;
 	const float* Entry1;
@@ -737,7 +737,7 @@ void FRawDistribution::GetValue1Random(float Time, float* InValue, class FRandom
 	Value[0] = Value1 + (Value2 - Value1) * RandValue;
 }
 
-void FRawDistribution::GetValue3Random(float Time, float* InValue, class FRandomStream* InRandomStream) const
+void FRawDistribution::GetValue3Random(float Time, float* InValue, struct FRandomStream* InRandomStream) const
 {
 	float* RESTRICT Value = InValue;
 	const float* Entry1;
@@ -783,7 +783,7 @@ void FRawDistribution::GetValue3Random(float Time, float* InValue, class FRandom
 	Value[2] = Z0 + (Z1 - Z0) * RandValues[2];
 }
 
-void FRawDistribution::GetValue(float Time, float* Value, int32 NumCoords, int32 Extreme, class FRandomStream* InRandomStream) const
+void FRawDistribution::GetValue(float Time, float* Value, int32 NumCoords, int32 Extreme, struct FRandomStream* InRandomStream) const
 {
 	checkSlow(NumCoords == 3 || NumCoords == 1);
 	switch (LookupTable.Op)
@@ -894,7 +894,7 @@ void FRawDistributionFloat::Initialize()
 }
 #endif // WITH_EDITOR
 
-float FRawDistributionFloat::GetValue(float F, UObject* Data, class FRandomStream* InRandomStream)
+float FRawDistributionFloat::GetValue(float F, UObject* Data, struct FRandomStream* InRandomStream)
 {
 #if WITH_EDITOR
 	// make sure it's up to date
@@ -989,7 +989,7 @@ void UDistributionFloat::PostEditChangeProperty(FPropertyChangedEvent& PropertyC
 }
 #endif	// WITH_EDITOR
 
-float UDistributionFloat::GetValue( float F, UObject* Data, class FRandomStream* InRandomStream ) const
+float UDistributionFloat::GetValue( float F, UObject* Data, struct FRandomStream* InRandomStream ) const
 {
 	return 0.0;
 }
@@ -1072,7 +1072,7 @@ void FRawDistributionVector::Initialize()
 }
 #endif
 
-FVector FRawDistributionVector::GetValue(float F, UObject* Data, int32 Extreme, class FRandomStream* InRandomStream)
+FVector FRawDistributionVector::GetValue(float F, UObject* Data, int32 Extreme, struct FRandomStream* InRandomStream)
 {
 #if WITH_EDITOR
 	// make sure it's up to date
@@ -1165,7 +1165,7 @@ void UDistributionVector::PostEditChangeProperty(FPropertyChangedEvent& Property
 }
 #endif // WITH_EDITOR
 
-FVector UDistributionVector::GetValue(float F, UObject* Data, int32 Extreme, class FRandomStream* InRandomStream) const
+FVector UDistributionVector::GetValue(float F, UObject* Data, int32 Extreme, struct FRandomStream* InRandomStream) const
 {
 	return FVector::ZeroVector;
 }
@@ -1575,7 +1575,7 @@ void UDistributionFloatConstant::PostLoad()
 	}
 }
 
-float UDistributionFloatConstant::GetValue( float F, UObject* Data, class FRandomStream* InRandomStream ) const
+float UDistributionFloatConstant::GetValue( float F, UObject* Data, struct FRandomStream* InRandomStream ) const
 {
 	return Constant;
 }
@@ -1688,7 +1688,7 @@ UDistributionFloatConstantCurve::UDistributionFloatConstantCurve(const FObjectIn
 {
 }
 
-float UDistributionFloatConstantCurve::GetValue( float F, UObject* Data, class FRandomStream* InRandomStream ) const
+float UDistributionFloatConstantCurve::GetValue( float F, UObject* Data, struct FRandomStream* InRandomStream ) const
 {
 	return ConstantCurve.Eval(F, 0.f);
 }
@@ -1873,7 +1873,7 @@ void UDistributionFloatUniform::PostLoad()
 	}
 }
 
-float UDistributionFloatUniform::GetValue( float F, UObject* Data, class FRandomStream* InRandomStream ) const
+float UDistributionFloatUniform::GetValue( float F, UObject* Data, struct FRandomStream* InRandomStream ) const
 {
 	return Max + (Min - Max) * DIST_GET_RANDOM_VALUE(InRandomStream);
 }
@@ -2049,7 +2049,7 @@ UDistributionFloatUniformCurve::UDistributionFloatUniformCurve(const FObjectInit
 {
 }
 
-float UDistributionFloatUniformCurve::GetValue(float F, UObject* Data, class FRandomStream* InRandomStream) const
+float UDistributionFloatUniformCurve::GetValue(float F, UObject* Data, struct FRandomStream* InRandomStream) const
 {
 	FVector2D Val = ConstantCurve.Eval(F, FVector2D(0.f, 0.f));
 	return Val.X + (Val.Y - Val.X) * DIST_GET_RANDOM_VALUE(InRandomStream);
@@ -2337,7 +2337,7 @@ void UDistributionVectorConstant::PostLoad()
 	}
 }
 
-FVector UDistributionVectorConstant::GetValue(float F, UObject* Data, int32 Extreme, class FRandomStream* InRandomStream) const
+FVector UDistributionVectorConstant::GetValue(float F, UObject* Data, int32 Extreme, struct FRandomStream* InRandomStream) const
 {
 	switch (LockedAxes)
 	{
@@ -2576,7 +2576,7 @@ UDistributionVectorConstantCurve::UDistributionVectorConstantCurve(const FObject
 {
 }
 
-FVector UDistributionVectorConstantCurve::GetValue(float F, UObject* Data, int32 Extreme, class FRandomStream* InRandomStream) const
+FVector UDistributionVectorConstantCurve::GetValue(float F, UObject* Data, int32 Extreme, struct FRandomStream* InRandomStream) const
 {
 	FVector Val = ConstantCurve.Eval(F, FVector::ZeroVector);
 	switch (LockedAxes)
@@ -2961,7 +2961,7 @@ void UDistributionVectorUniform::PostLoad()
 	}
 }
 
-FVector UDistributionVectorUniform::GetValue(float F, UObject* Data, int32 Extreme, class FRandomStream* InRandomStream) const
+FVector UDistributionVectorUniform::GetValue(float F, UObject* Data, int32 Extreme, struct FRandomStream* InRandomStream) const
 {
 	FVector LocalMax = Max;
 	FVector LocalMin = Min;
@@ -3537,7 +3537,7 @@ UDistributionVectorUniformCurve::UDistributionVectorUniformCurve(const FObjectIn
 	bUseExtremes = false;
 }
 
-FVector UDistributionVectorUniformCurve::GetValue(float F, UObject* Data, int32 Extreme, class FRandomStream* InRandomStream) const
+FVector UDistributionVectorUniformCurve::GetValue(float F, UObject* Data, int32 Extreme, struct FRandomStream* InRandomStream) const
 {
 	FTwoVectors	Val = ConstantCurve.Eval(F, FTwoVectors());
 
@@ -4098,7 +4098,7 @@ UDistributionFloatParameterBase::UDistributionFloatParameterBase(const FObjectIn
 	MaxOutput = 1.0f;
 }
 
-float UDistributionFloatParameterBase::GetValue( float F, UObject* Data, class FRandomStream* InRandomStream ) const
+float UDistributionFloatParameterBase::GetValue( float F, UObject* Data, struct FRandomStream* InRandomStream ) const
 {
 	float ParamFloat = 0.f;
 	bool bFoundParam = GetParamValue(Data, ParameterName, ParamFloat);
@@ -4137,7 +4137,7 @@ UDistributionVectorParameterBase::UDistributionVectorParameterBase(const FObject
 
 }
 
-FVector UDistributionVectorParameterBase::GetValue( float F, UObject* Data, int32 Extreme, class FRandomStream* InRandomStream ) const
+FVector UDistributionVectorParameterBase::GetValue( float F, UObject* Data, int32 Extreme, struct FRandomStream* InRandomStream ) const
 {
 	FVector ParamVector(0.f);
 	bool bFoundParam = GetParamValue(Data, ParameterName, ParamVector);
diff --git a/Engine/Source/Runtime/Engine/Private/Particles/ParticleModules.cpp b/Engine/Source/Runtime/Engine/Private/Particles/ParticleModules.cpp
--- a/Engine/Source/Runtime/Engine/Private/Particles/ParticleModules.cpp
+++ b/Engine/Source/Runtime/Engine/Private/Particles/ParticleModules.cpp
@@ -1160,7 +1160,7 @@ void UParticleModuleMeshRotation::Spawn(FParticleEmitterInstance* Owner, int32 O
 }
 
 
-void UParticleModuleMeshRotation::SpawnEx(FParticleEmitterInstance* Owner, int32 Offset, float SpawnTime, class FRandomStream* InRandomStream, FBaseParticle* ParticleBase)
+void UParticleModuleMeshRotation::SpawnEx(FParticleEmitterInstance* Owner, int32 Offset, float SpawnTime, struct FRandomStream* InRandomStream, FBaseParticle* ParticleBase)
 {
 	SPAWN_INIT;
 	{
@@ -1257,7 +1257,7 @@ void UParticleModuleMeshRotationRate::Spawn(FParticleEmitterInstance* Owner, int
 	SpawnEx(Owner, Offset, SpawnTime, NULL, ParticleBase);
 }
 
-void UParticleModuleMeshRotationRate::SpawnEx(FParticleEmitterInstance* Owner, int32 Offset, float SpawnTime, class FRandomStream* InRandomStream, FBaseParticle* ParticleBase)
+void UParticleModuleMeshRotationRate::SpawnEx(FParticleEmitterInstance* Owner, int32 Offset, float SpawnTime, struct FRandomStream* InRandomStream, FBaseParticle* ParticleBase)
 {
 	SPAWN_INIT;
 	{
@@ -1517,7 +1517,7 @@ void UParticleModuleRotation::Spawn(FParticleEmitterInstance* Owner, int32 Offse
 	SpawnEx(Owner, Offset, SpawnTime, NULL, ParticleBase);
 }
 
-void UParticleModuleRotation::SpawnEx(FParticleEmitterInstance* Owner, int32 Offset, float SpawnTime, class FRandomStream* InRandomStream, FBaseParticle* ParticleBase)
+void UParticleModuleRotation::SpawnEx(FParticleEmitterInstance* Owner, int32 Offset, float SpawnTime, struct FRandomStream* InRandomStream, FBaseParticle* ParticleBase)
 {
 	SPAWN_INIT;
 	{
@@ -1612,7 +1612,7 @@ void UParticleModuleRotationRate::Spawn(FParticleEmitterInstance* Owner, int32 O
 	SpawnEx(Owner, Offset, SpawnTime, NULL, ParticleBase);
 }
 
-void UParticleModuleRotationRate::SpawnEx(FParticleEmitterInstance* Owner, int32 Offset, float SpawnTime, class FRandomStream* InRandomStream, FBaseParticle* ParticleBase)
+void UParticleModuleRotationRate::SpawnEx(FParticleEmitterInstance* Owner, int32 Offset, float SpawnTime, struct FRandomStream* InRandomStream, FBaseParticle* ParticleBase)
 {
 	SPAWN_INIT;
 	{
@@ -2707,7 +2707,7 @@ void UParticleModuleLight::PostEditChangeProperty(FPropertyChangedEvent& Propert
 }
 #endif // WITH_EDITOR
 
-void UParticleModuleLight::SpawnEx(FParticleEmitterInstance* Owner, int32 Offset, float SpawnTime, class FRandomStream* InRandomStream, FBaseParticle* ParticleBase)
+void UParticleModuleLight::SpawnEx(FParticleEmitterInstance* Owner, int32 Offset, float SpawnTime, struct FRandomStream* InRandomStream, FBaseParticle* ParticleBase)
 {
 	SPAWN_INIT;
 	PARTICLE_ELEMENT(FLightParticlePayload, LightData);
@@ -3285,7 +3285,7 @@ void UParticleModuleLifetime::Spawn(FParticleEmitterInstance* Owner, int32 Offse
 	SpawnEx(Owner, Offset, SpawnTime, NULL, ParticleBase);
 }
 
-void UParticleModuleLifetime::SpawnEx(FParticleEmitterInstance* Owner, int32 Offset, float SpawnTime, class FRandomStream* InRandomStream, FBaseParticle* ParticleBase)
+void UParticleModuleLifetime::SpawnEx(FParticleEmitterInstance* Owner, int32 Offset, float SpawnTime, struct FRandomStream* InRandomStream, FBaseParticle* ParticleBase)
 {
 	SPAWN_INIT;
 	{
diff --git a/Engine/Source/Runtime/Engine/Private/Particles/ParticleModules_Color.cpp b/Engine/Source/Runtime/Engine/Private/Particles/ParticleModules_Color.cpp
--- a/Engine/Source/Runtime/Engine/Private/Particles/ParticleModules_Color.cpp
+++ b/Engine/Source/Runtime/Engine/Private/Particles/ParticleModules_Color.cpp
@@ -150,7 +150,7 @@ void UParticleModuleColor::Spawn(FParticleEmitterInstance* Owner, int32 Offset,
 }
 
 
-void UParticleModuleColor::SpawnEx(FParticleEmitterInstance* Owner, int32 Offset, float SpawnTime, class FRandomStream* InRandomStream, FBaseParticle* ParticleBase)
+void UParticleModuleColor::SpawnEx(FParticleEmitterInstance* Owner, int32 Offset, float SpawnTime, struct FRandomStream* InRandomStream, FBaseParticle* ParticleBase)
 {
 	SPAWN_INIT;
 	{
diff --git a/Engine/Source/Runtime/Engine/Private/Particles/ParticleModules_Location.cpp b/Engine/Source/Runtime/Engine/Private/Particles/ParticleModules_Location.cpp
--- a/Engine/Source/Runtime/Engine/Private/Particles/ParticleModules_Location.cpp
+++ b/Engine/Source/Runtime/Engine/Private/Particles/ParticleModules_Location.cpp
@@ -83,7 +83,7 @@ void UParticleModuleLocation::Spawn(FParticleEmitterInstance* Owner, int32 Offse
 	SpawnEx(Owner, Offset, SpawnTime, NULL, ParticleBase);
 }
 
-void UParticleModuleLocation::SpawnEx(FParticleEmitterInstance* Owner, int32 Offset, float SpawnTime, class FRandomStream* InRandomStream, FBaseParticle* ParticleBase)
+void UParticleModuleLocation::SpawnEx(FParticleEmitterInstance* Owner, int32 Offset, float SpawnTime, struct FRandomStream* InRandomStream, FBaseParticle* ParticleBase)
 {
 	SPAWN_INIT;
 	UParticleLODLevel* LODLevel	= Owner->SpriteTemplate->GetCurrentLODLevel(Owner);
@@ -211,7 +211,7 @@ UParticleModuleLocationWorldOffset::UParticleModuleLocationWorldOffset(const FOb
 {
 }
 
-void UParticleModuleLocationWorldOffset::SpawnEx(FParticleEmitterInstance* Owner, int32 Offset, float SpawnTime, class FRandomStream* InRandomStream, FBaseParticle* ParticleBase)
+void UParticleModuleLocationWorldOffset::SpawnEx(FParticleEmitterInstance* Owner, int32 Offset, float SpawnTime, struct FRandomStream* InRandomStream, FBaseParticle* ParticleBase)
 {
 	SPAWN_INIT;
 	UParticleLODLevel* LODLevel = Owner->SpriteTemplate->GetCurrentLODLevel(Owner);
@@ -678,7 +678,7 @@ void UParticleModuleLocationPrimitiveBase::PostEditChangeProperty(FPropertyChang
 }
 #endif // WITH_EDITOR
 
-void UParticleModuleLocationPrimitiveBase::DetermineUnitDirection(FParticleEmitterInstance* Owner, FVector& vUnitDir, class FRandomStream* InRandomStream)
+void UParticleModuleLocationPrimitiveBase::DetermineUnitDirection(FParticleEmitterInstance* Owner, FVector& vUnitDir, struct FRandomStream* InRandomStream)
 {
 	FVector vRand;
 
@@ -813,7 +813,7 @@ void UParticleModuleLocationPrimitiveTriangle::Spawn(FParticleEmitterInstance* O
 	SpawnEx(Owner, Offset, SpawnTime, NULL, ParticleBase);
 }
 
-void UParticleModuleLocationPrimitiveTriangle::SpawnEx(FParticleEmitterInstance* Owner, int32 Offset, float SpawnTime, class FRandomStream* InRandomStream, FBaseParticle* ParticleBase)
+void UParticleModuleLocationPrimitiveTriangle::SpawnEx(FParticleEmitterInstance* Owner, int32 Offset, float SpawnTime, struct FRandomStream* InRandomStream, FBaseParticle* ParticleBase)
 {
 	SPAWN_INIT;
 	UParticleLODLevel* LODLevel	= Owner->SpriteTemplate->GetCurrentLODLevel(Owner);
@@ -966,7 +966,7 @@ void UParticleModuleLocationPrimitiveCylinder::Spawn(FParticleEmitterInstance* O
 	SpawnEx(Owner, Offset, SpawnTime, NULL, ParticleBase);
 }
 
-void UParticleModuleLocationPrimitiveCylinder::SpawnEx(FParticleEmitterInstance* Owner, int32 Offset, float SpawnTime, class FRandomStream* InRandomStream, FBaseParticle* ParticleBase)
+void UParticleModuleLocationPrimitiveCylinder::SpawnEx(FParticleEmitterInstance* Owner, int32 Offset, float SpawnTime, struct FRandomStream* InRandomStream, FBaseParticle* ParticleBase)
 {
 	SPAWN_INIT;
 
@@ -1265,7 +1265,7 @@ void UParticleModuleLocationPrimitiveSphere::Spawn(FParticleEmitterInstance* Own
 	SpawnEx(Owner, Offset, SpawnTime, NULL, ParticleBase);
 }
 
-void UParticleModuleLocationPrimitiveSphere::SpawnEx(FParticleEmitterInstance* Owner, int32 Offset, float SpawnTime, class FRandomStream* InRandomStream, FBaseParticle* ParticleBase)
+void UParticleModuleLocationPrimitiveSphere::SpawnEx(FParticleEmitterInstance* Owner, int32 Offset, float SpawnTime, struct FRandomStream* InRandomStream, FBaseParticle* ParticleBase)
 {
 	SPAWN_INIT;
 
diff --git a/Engine/Source/Runtime/Engine/Private/Particles/ParticleModules_Parameter.cpp b/Engine/Source/Runtime/Engine/Private/Particles/ParticleModules_Parameter.cpp
--- a/Engine/Source/Runtime/Engine/Private/Particles/ParticleModules_Parameter.cpp
+++ b/Engine/Source/Runtime/Engine/Private/Particles/ParticleModules_Parameter.cpp
@@ -118,7 +118,7 @@ void UParticleModuleParameterDynamic::Spawn(FParticleEmitterInstance* Owner, int
 	SpawnEx(Owner, Offset, SpawnTime, NULL, ParticleBase);
 }
 
-void UParticleModuleParameterDynamic::SpawnEx(FParticleEmitterInstance* Owner, int32 Offset, float SpawnTime, class FRandomStream* InRandomStream, FBaseParticle* ParticleBase)
+void UParticleModuleParameterDynamic::SpawnEx(FParticleEmitterInstance* Owner, int32 Offset, float SpawnTime, struct FRandomStream* InRandomStream, FBaseParticle* ParticleBase)
 {
 	SPAWN_INIT;
 	{
diff --git a/Engine/Source/Runtime/Engine/Private/Particles/ParticleModules_Size.cpp b/Engine/Source/Runtime/Engine/Private/Particles/ParticleModules_Size.cpp
--- a/Engine/Source/Runtime/Engine/Private/Particles/ParticleModules_Size.cpp
+++ b/Engine/Source/Runtime/Engine/Private/Particles/ParticleModules_Size.cpp
@@ -83,7 +83,7 @@ void UParticleModuleSize::Spawn(FParticleEmitterInstance* Owner, int32 Offset, f
 	SpawnEx(Owner, Offset, SpawnTime, NULL, ParticleBase);
 }
 
-void UParticleModuleSize::SpawnEx(FParticleEmitterInstance* Owner, int32 Offset, float SpawnTime, class FRandomStream* InRandomStream, FBaseParticle* ParticleBase)
+void UParticleModuleSize::SpawnEx(FParticleEmitterInstance* Owner, int32 Offset, float SpawnTime, struct FRandomStream* InRandomStream, FBaseParticle* ParticleBase)
 {
 	SPAWN_INIT;
 	FVector Size		 = StartSize.GetValue(Owner->EmitterTime, Owner->Component, 0, InRandomStream);
diff --git a/Engine/Source/Runtime/Engine/Private/Particles/ParticleModules_Velocity.cpp b/Engine/Source/Runtime/Engine/Private/Particles/ParticleModules_Velocity.cpp
--- a/Engine/Source/Runtime/Engine/Private/Particles/ParticleModules_Velocity.cpp
+++ b/Engine/Source/Runtime/Engine/Private/Particles/ParticleModules_Velocity.cpp
@@ -77,7 +77,7 @@ void UParticleModuleVelocity::Spawn(FParticleEmitterInstance* Owner, int32 Offse
 	SpawnEx(Owner, Offset, SpawnTime, NULL, ParticleBase);
 }
 
-void UParticleModuleVelocity::SpawnEx(FParticleEmitterInstance* Owner, int32 Offset, float SpawnTime, class FRandomStream* InRandomStream, FBaseParticle* ParticleBase)
+void UParticleModuleVelocity::SpawnEx(FParticleEmitterInstance* Owner, int32 Offset, float SpawnTime, struct FRandomStream* InRandomStream, FBaseParticle* ParticleBase)
 {
 	SPAWN_INIT;
 	{
@@ -418,7 +418,7 @@ void UParticleModuleVelocityCone::Spawn(FParticleEmitterInstance* Owner, int32 O
 	SpawnEx(Owner, Offset, SpawnTime, NULL, ParticleBase);
 }
 
-void UParticleModuleVelocityCone::SpawnEx(FParticleEmitterInstance* Owner, int32 Offset, float SpawnTime, class FRandomStream* InRandomStream, FBaseParticle* ParticleBase)
+void UParticleModuleVelocityCone::SpawnEx(FParticleEmitterInstance* Owner, int32 Offset, float SpawnTime, struct FRandomStream* InRandomStream, FBaseParticle* ParticleBase)
 {
 	static const float TwoPI = 2.0f * PI;
 	static const float ToRads = PI / 180.0f;
diff --git a/Engine/Source/Runtime/Engine/Public/AnimationUtils.h b/Engine/Source/Runtime/Engine/Public/AnimationUtils.h
--- a/Engine/Source/Runtime/Engine/Public/AnimationUtils.h
+++ b/Engine/Source/Runtime/Engine/Public/AnimationUtils.h
@@ -16,7 +16,7 @@ ENGINE_API DECLARE_LOG_CATEGORY_EXTERN(LogAnimNotify, Warning, All);
 // Forward declarations.
 class UAnimCompress;
 class UAnimSet;
-class FTransform;
+struct FTransform;
 
 /**
  * Encapsulates commonly useful data about bones.
diff --git a/Engine/Source/Runtime/Engine/Public/Distributions.h b/Engine/Source/Runtime/Engine/Public/Distributions.h
--- a/Engine/Source/Runtime/Engine/Public/Distributions.h
+++ b/Engine/Source/Runtime/Engine/Public/Distributions.h
@@ -178,11 +178,11 @@ struct ENGINE_API FRawDistribution
 	 * @param NumCoords The number of floats in the Value array
 	 * @param Extreme For distributions that use one of the extremes, this is which extreme to use
 	 */
-	void GetValue(float Time, float* Value, int32 NumCoords, int32 Extreme, class FRandomStream* InRandomStream) const;
+	void GetValue(float Time, float* Value, int32 NumCoords, int32 Extreme, struct FRandomStream* InRandomStream) const;
 
 	// prebaked versions of these
-	void GetValue1(float Time, float* Value, int32 Extreme, class FRandomStream* InRandomStream) const;
-	void GetValue3(float Time, float* Value, int32 Extreme, class FRandomStream* InRandomStream) const;
+	void GetValue1(float Time, float* Value, int32 Extreme, struct FRandomStream* InRandomStream) const;
+	void GetValue3(float Time, float* Value, int32 Extreme, struct FRandomStream* InRandomStream) const;
 	inline void GetValue1None(float Time, float* InValue) const
 	{
 		float* Value = InValue;
@@ -210,10 +210,10 @@ struct ENGINE_API FRawDistribution
 		Value[1] = T1;
 		Value[2] = T2;
 	}
-	void GetValue1Extreme(float Time, float* Value, int32 Extreme, class FRandomStream* InRandomStream) const;
-	void GetValue3Extreme(float Time, float* Value, int32 Extreme, class FRandomStream* InRandomStream) const;
-	void GetValue1Random(float Time, float* Value, class FRandomStream* InRandomStream) const;
-	void GetValue3Random(float Time, float* Value, class FRandomStream* InRandomStream) const;
+	void GetValue1Extreme(float Time, float* Value, int32 Extreme, struct FRandomStream* InRandomStream) const;
+	void GetValue3Extreme(float Time, float* Value, int32 Extreme, struct FRandomStream* InRandomStream) const;
+	void GetValue1Random(float Time, float* Value, struct FRandomStream* InRandomStream) const;
+	void GetValue3Random(float Time, float* Value, struct FRandomStream* InRandomStream) const;
 
 	FORCEINLINE bool IsSimple() 
 	{
