#UE4 - DebugCameraController is always local

diff --git a/Engine/Source/Runtime/Engine/Private/DebugCameraController.cpp b/Engine/Source/Runtime/Engine/Private/DebugCameraController.cpp
--- a/Engine/Source/Runtime/Engine/Private/DebugCameraController.cpp
+++ b/Engine/Source/Runtime/Engine/Private/DebugCameraController.cpp
@@ -50,6 +50,7 @@ ADebugCameraController::ADebugCameraController(const FObjectInitializer& ObjectI
 #endif // WITH_EDITORONLY_DATA
 	PrimaryActorTick.bTickEvenWhenPaused = true;
 	bShouldPerformFullTickWhenPaused = true;
+	SetAsLocalPlayerController();
 }
 
 void InitializeDebugCameraInputBindings()
