#ue4 - Workaround for OverlapTest() ambiguity, use OverlapMulti() instead to check if we have any blocking overlap.

UE-9878

diff --git a/Engine/Source/Runtime/Engine/Private/CharacterMovementComponent.cpp b/Engine/Source/Runtime/Engine/Private/CharacterMovementComponent.cpp
--- a/Engine/Source/Runtime/Engine/Private/CharacterMovementComponent.cpp
+++ b/Engine/Source/Runtime/Engine/Private/CharacterMovementComponent.cpp
@@ -1710,7 +1710,8 @@ void UCharacterMovementComponent::Crouch(bool bClientSimulation)
 			FCollisionQueryParams CapsuleParams(NAME_CrouchTrace, false, CharacterOwner);
 			FCollisionResponseParams ResponseParam;
 			InitCollisionParams(CapsuleParams, ResponseParam);
-			const bool bEncroached = GetWorld()->OverlapTest(CharacterOwner->GetActorLocation() - FVector(0.f,0.f,ScaledHalfHeightAdjust), FQuat::Identity,
+			TArray<FOverlapResult> Overlaps;
+			const bool bEncroached = GetWorld()->OverlapMulti(Overlaps, CharacterOwner->GetActorLocation() - FVector(0.f,0.f,ScaledHalfHeightAdjust), FQuat::Identity,
 				UpdatedComponent->GetCollisionObjectType(), GetPawnCapsuleCollisionShape(SHRINK_None), CapsuleParams, ResponseParam);
 
 			// If encroached, cancel
@@ -1786,7 +1787,8 @@ void UCharacterMovementComponent::UnCrouch(bool bClientSimulation)
 		if (!bCrouchMaintainsBaseLocation)
 		{
 			// Expand in place
-			bEncroached = GetWorld()->OverlapTest(PawnLocation, FQuat::Identity, CollisionChannel, StandingCapsuleShape, CapsuleParams, ResponseParam);
+			TArray<FOverlapResult> Overlaps;
+			bEncroached = GetWorld()->OverlapMulti(Overlaps, PawnLocation, FQuat::Identity, CollisionChannel, StandingCapsuleShape, CapsuleParams, ResponseParam);
 		
 			if (bEncroached)
 			{
@@ -1812,7 +1814,8 @@ void UCharacterMovementComponent::UnCrouch(bool bClientSimulation)
 						// Compute where the base of the sweep ended up, and see if we can stand there
 						const float DistanceToBase = (Hit.Time * TraceDist) + ShortCapsuleShape.Capsule.HalfHeight;
 						const FVector NewLoc = FVector(PawnLocation.X, PawnLocation.Y, PawnLocation.Z - DistanceToBase + PawnHalfHeight + SweepInflation + MIN_FLOOR_DIST / 2.f);
-						bEncroached = GetWorld()->OverlapTest(NewLoc, FQuat::Identity, CollisionChannel, StandingCapsuleShape, CapsuleParams);
+						Overlaps.Reset(); // OverlapMulti doesn't reset it, though it should.
+						bEncroached = GetWorld()->OverlapMulti(Overlaps, NewLoc, FQuat::Identity, CollisionChannel, StandingCapsuleShape, CapsuleParams, ResponseParam);
 						if (!bEncroached)
 						{
 							// Intentionally not using MoveUpdatedComponent, where a horizontal plane constraint would prevent the base of the capsule from staying at the same spot.
@@ -1826,7 +1829,8 @@ void UCharacterMovementComponent::UnCrouch(bool bClientSimulation)
 		{
 			// Expand while keeping base location the same.
 			FVector StandingLocation = PawnLocation + FVector(0.f, 0.f, StandingCapsuleShape.GetCapsuleHalfHeight() - CurrentCrouchedHalfHeight);
-			bEncroached = GetWorld()->OverlapTest(StandingLocation, FQuat::Identity, CollisionChannel, StandingCapsuleShape, CapsuleParams, ResponseParam);
+			TArray<FOverlapResult> Overlaps;
+			bEncroached = GetWorld()->OverlapMulti(Overlaps, StandingLocation, FQuat::Identity, CollisionChannel, StandingCapsuleShape, CapsuleParams, ResponseParam);
 
 			if (bEncroached)
 			{
@@ -1837,7 +1841,8 @@ void UCharacterMovementComponent::UnCrouch(bool bClientSimulation)
 					if (CurrentFloor.bBlockingHit && CurrentFloor.FloorDist > MinFloorDist)
 					{
 						StandingLocation.Z -= CurrentFloor.FloorDist - MinFloorDist;
-						bEncroached = GetWorld()->OverlapTest(StandingLocation, FQuat::Identity, CollisionChannel, StandingCapsuleShape, CapsuleParams, ResponseParam);
+						Overlaps.Reset(); // OverlapMulti doesn't reset it, though it should.
+						bEncroached = GetWorld()->OverlapMulti(Overlaps, StandingLocation, FQuat::Identity, CollisionChannel, StandingCapsuleShape, CapsuleParams, ResponseParam);
 					}
 				}				
 			}
diff --git a/Engine/Source/Runtime/Engine/Private/MovementComponent.cpp b/Engine/Source/Runtime/Engine/Private/MovementComponent.cpp
--- a/Engine/Source/Runtime/Engine/Private/MovementComponent.cpp
+++ b/Engine/Source/Runtime/Engine/Private/MovementComponent.cpp
@@ -314,7 +314,8 @@ bool UMovementComponent::OverlapTest(const FVector& Location, const FQuat& Rotat
 	FCollisionQueryParams QueryParams(NAME_TestOverlap, false, IgnoreActor);
 	FCollisionResponseParams ResponseParam;
 	InitCollisionParams(QueryParams, ResponseParam);
-	return GetWorld()->OverlapTest(Location, RotationQuat, CollisionChannel, CollisionShape, QueryParams, ResponseParam);
+	TArray<FOverlapResult> Overlaps;
+	return GetWorld()->OverlapMulti(Overlaps, Location, RotationQuat, CollisionChannel, CollisionShape, QueryParams, ResponseParam);
 }
 
 bool UMovementComponent::IsExceedingMaxSpeed(float MaxSpeed) const
