Cleanup some code (and get rid of a shadowed variable)

diff --git a/Engine/Source/Runtime/Engine/Private/Actor.cpp b/Engine/Source/Runtime/Engine/Private/Actor.cpp
--- a/Engine/Source/Runtime/Engine/Private/Actor.cpp
+++ b/Engine/Source/Runtime/Engine/Private/Actor.cpp
@@ -1224,11 +1224,7 @@ void AActor::DetachRootComponentFromParent(bool bMaintainWorldPosition)
 {
 	if(RootComponent)
 	{
-		USceneComponent * RootComponent = GetRootComponent();
-
-
-		GetRootComponent()->DetachFromParent(bMaintainWorldPosition);
-
+		RootComponent->DetachFromParent(bMaintainWorldPosition);
 		AttachmentReplication.AttachParent = NULL;
 	}
 }
