[INTEGRATE] Change 2434338 by [REDACTED_PERSON][REDACTED_HANDLE] on 2015/02/05 17:33:07

	Fix undo/redo of property changes on construction script created components

diff --git a/Engine/Source/Runtime/Engine/Classes/Components/ActorComponent.h b/Engine/Source/Runtime/Engine/Classes/Components/ActorComponent.h
--- a/Engine/Source/Runtime/Engine/Classes/Components/ActorComponent.h
+++ b/Engine/Source/Runtime/Engine/Classes/Components/ActorComponent.h
@@ -489,6 +489,7 @@ public:
 	virtual void PostLoad() override;
 	virtual void PostRename(UObject* OldOuter, const FName OldName) override;
 #if WITH_EDITOR
+	virtual bool Modify( bool bAlwaysMarkDirty = true ) override;
 	virtual void PreEditChange(UProperty* PropertyThatWillChange) override;
 	virtual void PostEditChangeProperty(FPropertyChangedEvent& PropertyChangedEvent) override;
 	virtual void PostEditUndo() override;
diff --git a/Engine/Source/Runtime/Engine/Private/ActorComponent.cpp b/Engine/Source/Runtime/Engine/Private/ActorComponent.cpp
--- a/Engine/Source/Runtime/Engine/Private/ActorComponent.cpp
+++ b/Engine/Source/Runtime/Engine/Private/ActorComponent.cpp
@@ -406,6 +406,18 @@ bool UActorComponent::CallRemoteFunction( UFunction* Function, void* Parameters,
 static TMap<UActorComponent*,FComponentReregisterContext*> EditReregisterContexts;
 
 #if WITH_EDITOR
+bool UActorComponent::Modify( bool bAlwaysMarkDirty/*=true*/ )
+{
+	// If this is a construction script component we don't store them in the transaction buffer.  Instead, mark
+	// the Actor as modified so that we store of the transaction annotation that has the component properties stashed
+	if (IsCreatedByConstructionScript())
+	{
+		return GetOwner()->Modify(bAlwaysMarkDirty);
+	}
+
+	return Super::Modify(bAlwaysMarkDirty);
+}
+
 void UActorComponent::PreEditChange(UProperty* PropertyThatWillChange)
 {
 	Super::PreEditChange(PropertyThatWillChange);
