[INTEGRATE] Change 2434687 by [REDACTED_PERSON][REDACTED_HANDLE] on 2015/02/05 21:44:10

	Engine: Fixing non-unity build

diff --git a/Engine/Source/Runtime/Engine/Private/BlueprintGeneratedClass.cpp b/Engine/Source/Runtime/Engine/Private/BlueprintGeneratedClass.cpp
--- a/Engine/Source/Runtime/Engine/Private/BlueprintGeneratedClass.cpp
+++ b/Engine/Source/Runtime/Engine/Private/BlueprintGeneratedClass.cpp
@@ -7,7 +7,7 @@
 #include "Engine/SimpleConstructionScript.h"
 #include "Engine/SCS_Node.h"
 #include "Engine/LevelScriptActor.h"
-#
+#include "Engine/InheritableComponentHandler.h"
 
 #if WITH_EDITOR
 #include "BlueprintEditorUtils.h"
