[INTEGRATE] Change 2435626 by [REDACTED_PERSON][REDACTED_HANDLE] on 2015/02/06 12:32:38

	Introduce PIE warning when performing invalid physics operations, like calling AddForce on a non simulating object

diff --git a/Engine/Source/Runtime/Engine/Classes/Components/PrimitiveComponent.h b/Engine/Source/Runtime/Engine/Classes/Components/PrimitiveComponent.h
--- a/Engine/Source/Runtime/Engine/Classes/Components/PrimitiveComponent.h
+++ b/Engine/Source/Runtime/Engine/Classes/Components/PrimitiveComponent.h
@@ -1477,6 +1477,12 @@ private:
 	 */
 	bool ApplyRigidBodyState(const FRigidBodyState& NewState, const FRigidBodyErrorCorrection& ErrorCorrection, FVector& OutDeltaPos, FName BoneName = NAME_None);
 
+	/** Check if mobility is set to non-static. If it's static we trigger a PIE warning and return true*/
+	bool CheckStaticMobilityAndWarn(const FText& ActionText) const;
+
+	/** Check if mobility is set to non-static. If BodyInstanceRequiresSimulation is non-null we check that it is simulated. Triggers a PIE warning if conditions fails */
+	void WarnInvalidPhysicsOperations(const FText& ActionText, const FBodyInstance* BodyInstanceRequiresSimulation = nullptr) const;
+
 public:
 
 	/**
diff --git a/Engine/Source/Runtime/Engine/Private/PrimitiveComponent.cpp b/Engine/Source/Runtime/Engine/Private/PrimitiveComponent.cpp
--- a/Engine/Source/Runtime/Engine/Private/PrimitiveComponent.cpp
+++ b/Engine/Source/Runtime/Engine/Private/PrimitiveComponent.cpp
@@ -1264,6 +1264,21 @@ FCollisionShape UPrimitiveComponent::GetCollisionShape(float Inflation) const
 	return FCollisionShape::MakeBox(Bounds.BoxExtent + Inflation);
 }
 
+bool UPrimitiveComponent::CheckStaticMobilityAndWarn(const FText& ActionText) const
+{
+	AActor* Actor = GetOwner();
+	// static things can move before they are registered (e.g. immediately after streaming), but not after.
+	if (Mobility == EComponentMobility::Static && Actor && Actor->bActorInitialized)
+	{
+		FMessageLog("PIE").Warning(FText::Format(LOCTEXT("InvalidStaticMove", "Mobility of {0} : {1} has to be 'Movable' if you'd like to {2}. "),
+			FText::FromString(GetNameSafe(GetOwner())), FText::FromString(GetName()), ActionText));
+
+		return true;
+	}
+
+	return false;
+}
+
 bool UPrimitiveComponent::MoveComponent( const FVector& Delta, const FRotator& NewRotation, bool bSweep, FHitResult* OutHit, EMoveComponentFlags MoveFlags)
 {
 	SCOPE_CYCLE_COUNTER(STAT_MoveComponentTime);
@@ -1291,11 +1306,8 @@ bool UPrimitiveComponent::MoveComponent( const FVector& Delta, const FRotator& N
 
 	// static things can move before they are registered (e.g. immediately after streaming), but not after.
 	// TODO: Static components without an owner can move, should they be able to?
-	if (Mobility == EComponentMobility::Static && Actor && Actor->bActorInitialized)
-	{		
-		FMessageLog("PIE").Warning( FText::Format(LOCTEXT("InvalidMove", "Mobility of {0} : {1} has to be 'Movable' if you'd like to move. "),
-				FText::FromString(GetNameSafe(GetOwner())), FText::FromString(GetName()) ));
-
+	if (CheckStaticMobilityAndWarn(LOCTEXT("InvalidMove", "move")))
+	{
 		if (OutHit)
 		{
 			*OutHit = FHitResult();
diff --git a/Engine/Source/Runtime/Engine/Private/PrimitiveComponentPhysics.cpp b/Engine/Source/Runtime/Engine/Private/PrimitiveComponentPhysics.cpp
--- a/Engine/Source/Runtime/Engine/Private/PrimitiveComponentPhysics.cpp
+++ b/Engine/Source/Runtime/Engine/Private/PrimitiveComponentPhysics.cpp
@@ -5,6 +5,8 @@
 
 //////////////// PRIMITIVECOMPONENT ///////////////
 
+#define LOCTEXT_NAMESPACE "PrimitiveComponent"
+
 bool UPrimitiveComponent::ApplyRigidBodyState(const FRigidBodyState& NewState, const FRigidBodyErrorCorrection& ErrorCorrection, FVector& OutDeltaPos, FName BoneName)
 {
 	bool bRestoredState = true;
@@ -149,6 +151,20 @@ void UPrimitiveComponent::SetWalkableSlopeOverride(const FWalkableSlopeOverride&
 	BodyInstance.SetWalkableSlopeOverride(NewOverride);
 }
 
+void UPrimitiveComponent::WarnInvalidPhysicsOperations(const FText& ActionText, const FBodyInstance* BI) const
+{
+#if !(UE_BUILD_SHIPPING || UE_BUILD_TEST)
+	if (!CheckStaticMobilityAndWarn(ActionText))	//all physics operations require non-static mobility
+	{
+		if (BI && BI->bSimulatePhysics == false)	//some require to be simulating too
+		{
+			FMessageLog("PIE").Warning(FText::Format(LOCTEXT("InvalidPhysicsOperation", "{0} has to have 'Simulate Physics' enabled if you'd like to {1}. "), FText::FromString(GetReadableName()), ActionText));
+		}
+	}
+#endif
+}
+
+
 void UPrimitiveComponent::SetSimulatePhysics(bool bSimulate)
 {
 	BodyInstance.SetInstanceSimulatePhysics(bSimulate);
@@ -168,9 +184,9 @@ void UPrimitiveComponent::SetLockedAxis(ELockedAxis::Type LockedAxis)
 
 void UPrimitiveComponent::AddImpulse(FVector Impulse, FName BoneName, bool bVelChange)
 {
-	FBodyInstance* BI = GetBodyInstance(BoneName);
-	if(BI)
+	if (FBodyInstance* BI = GetBodyInstance(BoneName))
 	{
+		WarnInvalidPhysicsOperations(LOCTEXT("AddImpulse", "AddImpulse"), BI);
 		BI->AddImpulse(Impulse, bVelChange);
 	}
 }
@@ -179,15 +195,16 @@ void UPrimitiveComponent::AddAngularImpulse(FVector Impulse, FName BoneName, boo
 {
 	if (FBodyInstance* BI = GetBodyInstance(BoneName))
 	{
+		WarnInvalidPhysicsOperations(LOCTEXT("AddAngularImpulse", "AddAngularImpulse"), BI);
 		BI->AddAngularImpulse(Impulse, bVelChange);
 	}
 }
 
 void UPrimitiveComponent::AddImpulseAtLocation(FVector Impulse, FVector Location, FName BoneName)
 {
-	FBodyInstance* BI = GetBodyInstance(BoneName);
-	if(BI)
+	if (FBodyInstance* BI = GetBodyInstance(BoneName))
 	{
+		WarnInvalidPhysicsOperations(LOCTEXT("AddImpulseAtLocation", "AddImpulseAtLocation"), BI);
 		BI->AddImpulseAtPosition(Impulse, Location);
 	}
 }
@@ -209,18 +226,18 @@ void UPrimitiveComponent::AddRadialImpulse(FVector Origin, float Radius, float S
 
 void UPrimitiveComponent::AddForce(FVector Force, FName BoneName)
 {
-	FBodyInstance* BI = GetBodyInstance(BoneName);
-	if(BI)
+	if (FBodyInstance* BI = GetBodyInstance(BoneName))
 	{
+		WarnInvalidPhysicsOperations(LOCTEXT("AddForce", "AddForce"), BI);
 		BI->AddForce(Force);
 	}
 }
 
 void UPrimitiveComponent::AddForceAtLocation(FVector Force, FVector Location, FName BoneName)
 {
-	FBodyInstance* BI = GetBodyInstance(BoneName);
-	if(BI)
+	if (FBodyInstance* BI = GetBodyInstance(BoneName))
 	{
+		WarnInvalidPhysicsOperations(LOCTEXT("AddForceAtLocation", "AddForceAtLocation"), BI);
 		BI->AddForceAtPosition(Force, Location);
 	}
 }
@@ -241,18 +258,18 @@ void UPrimitiveComponent::AddRadialForce(FVector Origin, float Radius, float Str
 
 void UPrimitiveComponent::AddTorque(FVector Torque, FName BoneName)
 {
-	FBodyInstance* BI = GetBodyInstance(BoneName);
-	if(BI)
+	if (FBodyInstance* BI = GetBodyInstance(BoneName))
 	{
+		WarnInvalidPhysicsOperations(LOCTEXT("AddTorque", "AddTorque"), BI);
 		BI->AddTorque(Torque);
 	}
 }
 
 void UPrimitiveComponent::SetPhysicsLinearVelocity(FVector NewVel, bool bAddToCurrent, FName BoneName)
 {
-	FBodyInstance* BI = GetBodyInstance(BoneName);
-	if(BI != NULL)
+	if (FBodyInstance* BI = GetBodyInstance(BoneName))
 	{
+		WarnInvalidPhysicsOperations(LOCTEXT("SetPhysicsLinearVelocity", "SetPhysicsLinearVelocity"));
 		BI->SetLinearVelocity(NewVel, bAddToCurrent);
 	}
 }
@@ -274,18 +291,18 @@ void UPrimitiveComponent::SetAllPhysicsLinearVelocity(FVector NewVel,bool bAddTo
 
 void UPrimitiveComponent::SetPhysicsAngularVelocity(FVector NewAngVel, bool bAddToCurrent, FName BoneName)
 {
-	FBodyInstance* const BI = GetBodyInstance(BoneName);
-	if(BI != NULL)
+	if (FBodyInstance* BI = GetBodyInstance(BoneName))
 	{
+		WarnInvalidPhysicsOperations(LOCTEXT("SetPhysicsAngularVelocity", "SetPhysicsAngularVelocity"));
 		BI->SetAngularVelocity(NewAngVel, bAddToCurrent);
 	}
 }
 
 void UPrimitiveComponent::SetPhysicsMaxAngularVelocity(float NewMaxAngVel, bool bAddToCurrent, FName BoneName)
 {
-	FBodyInstance* const BI = GetBodyInstance(BoneName);
-	if (BI != NULL)
+	if (FBodyInstance* BI = GetBodyInstance(BoneName))
 	{
+		WarnInvalidPhysicsOperations(LOCTEXT("SetPhysicsMaxAngularVelocity", "SetPhysicsMaxAngularVelocity"));
 		BI->SetMaxAngularVelocity(NewMaxAngVel, bAddToCurrent);
 	}
 }
@@ -943,3 +960,4 @@ void UPrimitiveComponent::UpdatePhysicsToRBChannels()
 	}
 }
 
+#undef LOCTEXT_NAMESPACE
\ No newline at end of file
