[INTEGRATE] Change 2435689 by [REDACTED_PERSON][REDACTED_HANDLE] on 2015/02/06 13:20:27

	Fix ReceiveComponentDamage from applying impulses to non simulating bodies

diff --git a/Engine/Source/Runtime/Engine/Private/PrimitiveComponent.cpp b/Engine/Source/Runtime/Engine/Private/PrimitiveComponent.cpp
--- a/Engine/Source/Runtime/Engine/Private/PrimitiveComponent.cpp
+++ b/Engine/Source/Runtime/Engine/Private/PrimitiveComponent.cpp
@@ -649,8 +649,11 @@ void UPrimitiveComponent::ReceiveComponentDamage(float DamageAmount, FDamageEven
 		FPointDamageEvent* const PointDamageEvent = (FPointDamageEvent*) &DamageEvent;
 		if((DamageTypeCDO->DamageImpulse > 0.f) && !PointDamageEvent->ShotDirection.IsNearlyZero())
 		{
-			FVector const ImpulseToApply = PointDamageEvent->ShotDirection.GetSafeNormal() * DamageTypeCDO->DamageImpulse;
-			AddImpulseAtLocation( ImpulseToApply, PointDamageEvent->HitInfo.ImpactPoint, PointDamageEvent->HitInfo.BoneName );
+			if (IsSimulatingPhysics(PointDamageEvent->HitInfo.BoneName))
+			{
+				FVector const ImpulseToApply = PointDamageEvent->ShotDirection.GetSafeNormal() * DamageTypeCDO->DamageImpulse;
+				AddImpulseAtLocation(ImpulseToApply, PointDamageEvent->HitInfo.ImpactPoint, PointDamageEvent->HitInfo.BoneName);
+			}
 		}
 	}
 	else if (DamageEvent.IsOfType(FRadialDamageEvent::ClassID))
