[INTEGRATE] Change 2435831 by [REDACTED_PERSON][REDACTED_HANDLE] on 2015/02/06 14:45:32

	Nonunity fix

diff --git a/Engine/Source/Runtime/Engine/Private/PrimitiveComponentPhysics.cpp b/Engine/Source/Runtime/Engine/Private/PrimitiveComponentPhysics.cpp
--- a/Engine/Source/Runtime/Engine/Private/PrimitiveComponentPhysics.cpp
+++ b/Engine/Source/Runtime/Engine/Private/PrimitiveComponentPhysics.cpp
@@ -2,6 +2,7 @@
 
 #include "EnginePrivate.h"
 #include "Components/LineBatchComponent.h"
+#include "MessageLog.h"
 
 //////////////// PRIMITIVECOMPONENT ///////////////
 
