[INTEGRATE] Change 2435942 by [REDACTED_PERSON][REDACTED_HANDLE] on 2015/02/06 15:40:42

	Blueprints: Exposed a PrintText node

diff --git a/Engine/Source/Runtime/Engine/Classes/Kismet/KismetSystemLibrary.h b/Engine/Source/Runtime/Engine/Classes/Kismet/KismetSystemLibrary.h
--- a/Engine/Source/Runtime/Engine/Classes/Kismet/KismetSystemLibrary.h
+++ b/Engine/Source/Runtime/Engine/Classes/Kismet/KismetSystemLibrary.h
@@ -198,6 +198,19 @@ class UKismetSystemLibrary : public UBlueprintFunctionLibrary
 	UFUNCTION(BlueprintCallable, meta=(WorldContext="WorldContextObject", CallableWithoutWorldContext, Keywords = "log", AdvancedDisplay = "2"), Category="Utilities|String")
 	static void PrintString(UObject* WorldContextObject, const FString& InString = FString(TEXT("Hello")), bool bPrintToScreen = true, bool bPrintToLog = true, FLinearColor TextColor = FLinearColor(0.0,0.66,1.0));
 
+	/**
+	 * Prints text to the log, and optionally, to the screen
+	 * If Print To Log is true, it will be visible in the Output Log window.  Otherwise it will be logged only as 'Verbose', so it generally won't show up.
+	 *
+	 * @param	InText			The text to log out
+	 * @param	bPrintToScreen	Whether or not to print the output to the screen
+	 * @param	bPrintToLog		Whether or not to print the output to the log
+	 * @param	bPrintToConsole	Whether or not to print the output to the console
+	 * @param	TextColor		Whether or not to print the output to the console
+	 */
+	UFUNCTION(BlueprintCallable, meta=(WorldContext="WorldContextObject", CallableWithoutWorldContext, Keywords = "log", AdvancedDisplay = "2"), Category="Utilities|Text")
+	static void PrintText(UObject* WorldContextObject, const FText& InText = FText::FromString(TEXT("Hello")), bool bPrintToScreen = true, bool bPrintToLog = true, FLinearColor TextColor = FLinearColor(0.0,0.66,1.0));
+
 	/**
 	 * Prints a warning string to the log and the screen. Meant to be used as a way to inform the user that they misused the node.
 	 *
diff --git a/Engine/Source/Runtime/Engine/Private/KismetSystemLibrary.cpp b/Engine/Source/Runtime/Engine/Private/KismetSystemLibrary.cpp
--- a/Engine/Source/Runtime/Engine/Private/KismetSystemLibrary.cpp
+++ b/Engine/Source/Runtime/Engine/Private/KismetSystemLibrary.cpp
@@ -220,6 +220,13 @@ void UKismetSystemLibrary::PrintString(UObject* WorldContextObject, const FStrin
 #endif
 }
 
+void UKismetSystemLibrary::PrintText(UObject* WorldContextObject, const FText& InText, bool bPrintToScreen, bool bPrintToLog, FLinearColor TextColor)
+{
+#if !(UE_BUILD_SHIPPING || UE_BUILD_TEST) // Do not Print in Shipping or Test
+	PrintString(WorldContextObject, InText.ToString(), bPrintToScreen, bPrintToLog, TextColor);
+#endif
+}
+
 void UKismetSystemLibrary::PrintWarning(const FString& InString)
 {
 	PrintString(NULL, InString, true, true);
