fixed unintialized param of FAreaNavModifier structure
#ue4

diff --git a/Engine/Source/Runtime/Engine/Private/AI/Navigation/NavigationModifier.cpp b/Engine/Source/Runtime/Engine/Private/AI/Navigation/NavigationModifier.cpp
--- a/Engine/Source/Runtime/Engine/Private/AI/Navigation/NavigationModifier.cpp
+++ b/Engine/Source/Runtime/Engine/Private/AI/Navigation/NavigationModifier.cpp
@@ -223,6 +223,7 @@ void FAreaNavModifier::GetConvex(FConvexNavAreaData& Data) const
 
 void FAreaNavModifier::Init(const TSubclassOf<UNavArea> InAreaClass)
 {
+	bIncludeAgentHeight = false;
 	Cost = 0.0f;
 	FixedCost = 0.0f;
 	ReplaceAreaClass = NULL;
diff --git a/Engine/Source/Runtime/Engine/Public/AI/NavigationModifier.h b/Engine/Source/Runtime/Engine/Public/AI/NavigationModifier.h
--- a/Engine/Source/Runtime/Engine/Public/AI/NavigationModifier.h
+++ b/Engine/Source/Runtime/Engine/Public/AI/NavigationModifier.h
@@ -68,7 +68,7 @@ struct ENGINE_API FAreaNavModifier : public FNavigationModifier
 	float Cost;
 	float FixedCost;
 
-	FAreaNavModifier() : Cost(0.0f), FixedCost(0.0f), AreaClass(NULL), ReplaceAreaClass(NULL), ShapeType(ENavigationShapeType::Unknown) {}
+	FAreaNavModifier() : Cost(0.0f), FixedCost(0.0f), AreaClass(NULL), ReplaceAreaClass(NULL), ShapeType(ENavigationShapeType::Unknown), bIncludeAgentHeight(false) {}
 	FAreaNavModifier(float Radius, float Height, const FTransform& LocalToWorld, const TSubclassOf<UNavArea> AreaClass);
 	FAreaNavModifier(const FVector& Extent, const FTransform& LocalToWorld, const TSubclassOf<UNavArea> AreaClass);
 	FAreaNavModifier(const FBox& Box, const FTransform& LocalToWorld, const TSubclassOf<UNavArea> AreaClass);
