Allow the Pause key during PIE trigger both its bound exec command in config, and its bound editor UI command.

diff --git a/Engine/Source/Runtime/Engine/Classes/GameFramework/PlayerInput.h b/Engine/Source/Runtime/Engine/Classes/GameFramework/PlayerInput.h
--- a/Engine/Source/Runtime/Engine/Classes/GameFramework/PlayerInput.h
+++ b/Engine/Source/Runtime/Engine/Classes/GameFramework/PlayerInput.h
@@ -485,7 +485,7 @@ public:
 	FKeyBind GetExecBind(FString const& ExecCommand);
 
 	/** Execute input commands within the legacy key binding system. */
-	void ExecInputCommands( UWorld* InWorld, const TCHAR* Cmd, class FOutputDevice& Ar);
+	bool ExecInputCommands( UWorld* InWorld, const TCHAR* Cmd, class FOutputDevice& Ar);
 #endif
 
 	/** Returns the list of keys mapped to the specified Action Name */
diff --git a/Engine/Source/Runtime/Engine/Private/LocalPlayer.cpp b/Engine/Source/Runtime/Engine/Private/LocalPlayer.cpp
--- a/Engine/Source/Runtime/Engine/Private/LocalPlayer.cpp
+++ b/Engine/Source/Runtime/Engine/Private/LocalPlayer.cpp
@@ -1133,7 +1133,7 @@ bool ULocalPlayer::HandlePauseCommand( const TCHAR* Cmd, FOutputDevice& Ar, UWor
 	}
 	
 
-	return true;
+	return !ViewportClient->bIsPlayInEditorViewport;
 }
 
 bool ULocalPlayer::HandleListMoveBodyCommand( const TCHAR* Cmd, FOutputDevice& Ar )
diff --git a/Engine/Source/Runtime/Engine/Private/UserInterface/PlayerInput.cpp b/Engine/Source/Runtime/Engine/Private/UserInterface/PlayerInput.cpp
--- a/Engine/Source/Runtime/Engine/Private/UserInterface/PlayerInput.cpp
+++ b/Engine/Source/Runtime/Engine/Private/UserInterface/PlayerInput.cpp
@@ -169,8 +169,7 @@ bool UPlayerInput::InputKey( FKey Key, EInputEvent Event, float AmountDepressed,
 	const FString Command = GetBind(Key);
 	if(Command.Len())
 	{
-		ExecInputCommands(World, *Command,*GLog);
-		return true;
+		return ExecInputCommands(World, *Command,*GLog);
 	}
 #endif
 
@@ -1530,8 +1529,10 @@ const TArray<FInputActionKeyMapping>& UPlayerInput::GetKeysForAction(const FName
 }
 
 #if !UE_BUILD_SHIPPING
-void UPlayerInput::ExecInputCommands( UWorld* InWorld, const TCHAR* Cmd, FOutputDevice& Ar )
+bool UPlayerInput::ExecInputCommands( UWorld* InWorld, const TCHAR* Cmd, FOutputDevice& Ar )
 {
+	bool bResult = false;
+
 	int32 CmdLen = FCString::Strlen(Cmd);
 	TCHAR* Line = (TCHAR*)FMemory::Malloc(sizeof(TCHAR)*(CmdLen+1));
 
@@ -1544,24 +1545,29 @@ void UPlayerInput::ExecInputCommands( UWorld* InWorld, const TCHAR* Cmd, FOutput
 			UPlayer* Player = Actor ? Actor->Player : NULL;
 			if(ProcessConsoleExec(Str,Ar,this))
 			{
+				bResult = true;
 				continue;
 			}
 			else if(Actor && Exec(Actor->GetWorld(), Str,Ar))
 			{
+				bResult = true;
 				continue;
 			}
 			else if(Player && Player->Exec( Actor->GetWorld(), Str,Ar))
 			{
+				bResult = true;
 				continue;
 			}
 		}
 		else
 		{
-			Exec(InWorld, Str,Ar);
+			bResult |= Exec(InWorld, Str,Ar);
 		}
 	}
 
 	FMemory::Free(Line);
+
+	return bResult;
 }
 
 bool UPlayerInput::Exec(UWorld* InWorld, const TCHAR* Str,FOutputDevice& Ar)
@@ -1595,9 +1601,9 @@ bool UPlayerInput::Exec(UWorld* InWorld, const TCHAR* Str,FOutputDevice& Ar)
 				if(DebugExecBindings[BindIndex].Key == Key)
 				{
 					bExecutingBindCommand = true;
-					ExecInputCommands(GetWorld(), *DebugExecBindings[BindIndex].Command,Ar);
+					bool bResult = ExecInputCommands(GetWorld(), *DebugExecBindings[BindIndex].Command,Ar);
 					bExecutingBindCommand = false;
-					return 1;
+					return bResult;
 				}
 			}
 		}
