[INTEGRATE] Change 2436120 by [REDACTED_PERSON][REDACTED_HANDLE] on 2015/02/06 17:13:19

	Engine - Droping delegate from an event name by giving it a cleaner display name.

diff --git a/Engine/Source/Runtime/Engine/Classes/Components/SceneComponent.h b/Engine/Source/Runtime/Engine/Classes/Components/SceneComponent.h
--- a/Engine/Source/Runtime/Engine/Classes/Components/SceneComponent.h
+++ b/Engine/Source/Runtime/Engine/Classes/Components/SceneComponent.h
@@ -511,7 +511,7 @@ public:
 
 public:
 	/** Delegate that will be called when PhysicsVolume has been changed **/
-	UPROPERTY(BlueprintAssignable, Category=PhysicsVolume)
+	UPROPERTY(BlueprintAssignable, Category=PhysicsVolume, meta=(DisplayName="Physics Volume Changed"))
 	FPhysicsVolumeChanged PhysicsVolumeChangedDelegate;
 
 	// Begin ActorComponent interface
