[INTEGRATE] Change 2436384 by [REDACTED_PERSON][REDACTED_HANDLE] on 2015/02/06 19:01:22

	Missing header include.

diff --git a/Engine/Source/Runtime/Engine/Classes/Components/SceneComponent.h b/Engine/Source/Runtime/Engine/Classes/Components/SceneComponent.h
--- a/Engine/Source/Runtime/Engine/Classes/Components/SceneComponent.h
+++ b/Engine/Source/Runtime/Engine/Classes/Components/SceneComponent.h
@@ -4,6 +4,7 @@
 #include "Components/ActorComponent.h"
 #include "Engine/EngineTypes.h"
 #include "RHIDefinitions.h"
+#include "ComponentInstanceDataCache.h" // for FComponentInstanceDataBase
 #include "SceneComponent.generated.h"
 
 /** Overlap info consisting of the primitive and the body that is overlapping */
