[INTEGRATE] Change 2436668 by [REDACTED_PERSON][REDACTED_HANDLE] on 2015/02/07 02:41:17

	Blueprints: Add USceneComponent to the Common group in addition to the Utility group, so it shows up closer to the top of the Add Components list

diff --git a/Engine/Source/Runtime/Engine/Classes/Components/SceneComponent.h b/Engine/Source/Runtime/Engine/Classes/Components/SceneComponent.h
--- a/Engine/Source/Runtime/Engine/Classes/Components/SceneComponent.h
+++ b/Engine/Source/Runtime/Engine/Classes/Components/SceneComponent.h
@@ -77,7 +77,7 @@ DECLARE_DYNAMIC_MULTICAST_DELEGATE_OneParam(FPhysicsVolumeChanged, class APhysic
  * Useful as a 'dummy' component in the hierarchy to offset others.
  * @see [Scene Components](https://docs.unrealengine.com/latest/INT/Programming/UnrealArchitecture/Actors/Components/index.html#scenecomponents)
  */
-UCLASS(ClassGroup=Utility, BlueprintType, hideCategories=(Trigger, PhysicsVolume), meta=(BlueprintSpawnableComponent, ShortTooltip="A Scene Component is a component that has a scene transform and can be attached to other scene components."))
+UCLASS(ClassGroup=(Utility, Common), BlueprintType, hideCategories=(Trigger, PhysicsVolume), meta=(BlueprintSpawnableComponent, ShortTooltip="A Scene Component is a component that has a scene transform and can be attached to other scene components."))
 class ENGINE_API USceneComponent : public UActorComponent
 {
 	GENERATED_BODY()
