merged fix for updating direct paths from dev branch
#ue4

diff --git a/Engine/Source/Runtime/Engine/Private/AI/Navigation/AbstractNavData.cpp b/Engine/Source/Runtime/Engine/Private/AI/Navigation/AbstractNavData.cpp
--- a/Engine/Source/Runtime/Engine/Private/AI/Navigation/AbstractNavData.cpp
+++ b/Engine/Source/Runtime/Engine/Private/AI/Navigation/AbstractNavData.cpp
@@ -42,6 +42,7 @@ FPathFindingResult AAbstractNavData::FindPathAbstract(const FNavAgentProperties&
 	FPathFindingResult Result;
 	Result.Path = Query.PathInstanceToFill.IsValid() ? Query.PathInstanceToFill : Self->CreatePathInstance<FAbstractNavigationPath>(Query.Owner.Get());
 
+	Result.Path->GetPathPoints().Reset();
 	Result.Path->GetPathPoints().Add(FNavPathPoint(Query.StartLocation));
 	Result.Path->GetPathPoints().Add(FNavPathPoint(Query.EndLocation));
 	Result.Path->MarkReady();
