Fix no lighting in any world using a PreviewScene

diff --git a/Engine/Source/Runtime/Engine/Private/PreviewScene.cpp b/Engine/Source/Runtime/Engine/Private/PreviewScene.cpp
--- a/Engine/Source/Runtime/Engine/Private/PreviewScene.cpp
+++ b/Engine/Source/Runtime/Engine/Private/PreviewScene.cpp
@@ -36,12 +36,12 @@ FPreviewScene::FPreviewScene(FPreviewScene::ConstructionValues CVS)
 
 	GetScene()->UpdateDynamicSkyLight(FLinearColor::White * CVS.SkyBrightness, FLinearColor::Black);
 
-	DirectionalLight = NewObject<UDirectionalLightComponent>(GetTransientPackage(), TEXT("DirectionalLight"));
+	DirectionalLight = NewObject<UDirectionalLightComponent>(GetTransientPackage());
 	DirectionalLight->Intensity = CVS.LightBrightness;
 	DirectionalLight->LightColor = FColor::White;
 	AddComponent(DirectionalLight, FTransform(CVS.LightRotation));
 
-	LineBatcher = NewObject<ULineBatchComponent>(GetTransientPackage(), TEXT("LineBatcher"));
+	LineBatcher = NewObject<ULineBatchComponent>(GetTransientPackage());
 	AddComponent(LineBatcher, FTransform::Identity);
 }
 
