[INTEGRATE] Change 2436913 by [REDACTED_PERSON][REDACTED_HANDLE] on 2015/02/07 15:44:49

	fixed wireframe shapes not showing any subdued selection color

diff --git a/Engine/Source/Editor/UnrealEd/Private/EditorViewportClient.cpp b/Engine/Source/Editor/UnrealEd/Private/EditorViewportClient.cpp
--- a/Engine/Source/Editor/UnrealEd/Private/EditorViewportClient.cpp
+++ b/Engine/Source/Editor/UnrealEd/Private/EditorViewportClient.cpp
@@ -708,7 +708,7 @@ FSceneView* FEditorViewportClient::CalcSceneView(FSceneViewFamily* ViewFamily)
 
 	FSceneView* View = new FSceneView(ViewInitOptions);
 
-	View->SubduedSelectionOutlineColor = GetDefault<UEditorStyleSettings>()->GetSubduedSelectionColor();
+	View->SubduedSelectionOutlineColor = GEngine->GetSubduedSelectionOutlineColor();
 
 	ViewFamily->Views.Add(View);
 
diff --git a/Engine/Source/Runtime/Engine/Classes/Engine/Engine.h b/Engine/Source/Runtime/Engine/Classes/Engine/Engine.h
--- a/Engine/Source/Runtime/Engine/Classes/Engine/Engine.h
+++ b/Engine/Source/Runtime/Engine/Classes/Engine/Engine.h
@@ -1321,6 +1321,10 @@ private:
 	UPROPERTY(transient)
 	FLinearColor SelectionOutlineColor;
 
+	/** Subdued version of the selection outline color. Used for indicating sub-selection of components vs actors */
+	UPROPERTY(transient)
+	FLinearColor SubduedSelectionOutlineColor;
+
 	/** An override to use in some cases instead of the selected material color */
 	UPROPERTY(transient)
 	FLinearColor SelectedMaterialColorOverride;
@@ -1474,6 +1478,8 @@ public:
 
 	const FLinearColor& GetSelectionOutlineColor() const { return SelectionOutlineColor; }
 
+	const FLinearColor& GetSubduedSelectionOutlineColor() const { return SubduedSelectionOutlineColor; }
+
 	const FLinearColor& GetHoveredMaterialColor() const { return GetSelectedMaterialColor(); }
 
 	/**
@@ -1487,6 +1493,7 @@ public:
 
 	void SetSelectionOutlineColor( const FLinearColor& InSelectionOutlineColor ) { SelectionOutlineColor = InSelectionOutlineColor; }
 
+	void SetSubduedSelectionOutlineColor( const FLinearColor& InSubduedSelectionOutlineColor ) { SubduedSelectionOutlineColor = InSubduedSelectionOutlineColor; }
 	/**
 	 * Sets an override color to use instead of the user setting
 	 *
diff --git a/Engine/Source/Runtime/Engine/Private/BillboardComponent.cpp b/Engine/Source/Runtime/Engine/Private/BillboardComponent.cpp
--- a/Engine/Source/Runtime/Engine/Private/BillboardComponent.cpp
+++ b/Engine/Source/Runtime/Engine/Private/BillboardComponent.cpp
@@ -147,8 +147,12 @@ public:
 					// Set the selection/hover color from the current engine setting.
 					// The color is multiplied by 10 because this value is normally expected to be blended
 					// additively, this is not how the sprites work and therefore need an extra boost
-					// to appear the same color as previously.
-					if (IsSelected())
+					// to appear the same color as previously
+					if( View->bHasSelectedComponents && !IsIndividuallySelected() )
+					{
+						ColorToUse = FLinearColor::White + (GEngine->GetSubduedSelectionOutlineColor() * GEngine->SelectionHighlightIntensityBillboards * 10);
+					}
+					else if (IsSelected())
 					{
 						ColorToUse = FLinearColor::White + (GEngine->GetSelectedMaterialColor() * GEngine->SelectionHighlightIntensityBillboards * 10);
 					}
diff --git a/Engine/Source/Runtime/Engine/Private/BoxComponent.cpp b/Engine/Source/Runtime/Engine/Private/BoxComponent.cpp
--- a/Engine/Source/Runtime/Engine/Private/BoxComponent.cpp
+++ b/Engine/Source/Runtime/Engine/Private/BoxComponent.cpp
@@ -101,12 +101,15 @@ FPrimitiveSceneProxy* UBoxComponent::CreateSceneProxy()
 			QUICK_SCOPE_CYCLE_COUNTER( STAT_BoxSceneProxy_GetDynamicMeshElements );
 
 			const FMatrix& LocalToWorld = GetLocalToWorld();
-			const FColor DrawColor = GetSelectionColor(BoxColor, IsSelected(), IsHovered(), /*bUseOverlayIntensity=*/false);
-
+			
 			for (int32 ViewIndex = 0; ViewIndex < Views.Num(); ViewIndex++)
 			{
 				if (VisibilityMap & (1 << ViewIndex))
 				{
+					const FSceneView* View = Views[ViewIndex];
+
+					const FLinearColor DrawColor = GetViewSelectionColor(BoxColor, *View, IsSelected(), IsHovered(), false, IsIndividuallySelected() );
+
 					FPrimitiveDrawInterface* PDI = Collector.GetPDI(ViewIndex);
 					DrawOrientedWireBox(PDI, LocalToWorld.GetOrigin(), LocalToWorld.GetScaledAxis( EAxis::X ), LocalToWorld.GetScaledAxis( EAxis::Y ), LocalToWorld.GetScaledAxis( EAxis::Z ), BoxExtents, DrawColor, SDPG_World);
 				}
diff --git a/Engine/Source/Runtime/Engine/Private/BrushComponent.cpp b/Engine/Source/Runtime/Engine/Private/BrushComponent.cpp
--- a/Engine/Source/Runtime/Engine/Private/BrushComponent.cpp
+++ b/Engine/Source/Runtime/Engine/Private/BrushComponent.cpp
@@ -294,7 +294,7 @@ public:
 						{
 							auto WireframeMaterial = new FColoredMaterialRenderProxy(
 								GEngine->LevelColorationUnlitMaterial->GetRenderProxy(IsSelected(), IsHovered()),
-								GetSelectionColor(DrawColor,!(GIsEditor && (View->Family->EngineShowFlags.Selection)) || IsSelected(), IsHovered(), /*bUseOverlayIntensity=*/false)
+								GetViewSelectionColor(DrawColor, *View, !(GIsEditor && (View->Family->EngineShowFlags.Selection)) || IsSelected(), IsHovered(), false, IsIndividuallySelected() )
 								);
 
 							Collector.RegisterOneFrameMaterialProxy(WireframeMaterial);
diff --git a/Engine/Source/Runtime/Engine/Private/CapsuleComponent.cpp b/Engine/Source/Runtime/Engine/Private/CapsuleComponent.cpp
--- a/Engine/Source/Runtime/Engine/Private/CapsuleComponent.cpp
+++ b/Engine/Source/Runtime/Engine/Private/CapsuleComponent.cpp
@@ -37,14 +37,18 @@ FPrimitiveSceneProxy* UCapsuleComponent::CreateSceneProxy()
 		{
 			QUICK_SCOPE_CYCLE_COUNTER( STAT_GetDynamicMeshElements_DrawDynamicElements );
 
-			const FLinearColor DrawCapsuleColor = GetSelectionColor(ShapeColor, IsSelected(), IsHovered(), /*bUseOverlayIntensity=*/false);
+		
 			const FMatrix& LocalToWorld = GetLocalToWorld();
 			const int32 CapsuleSides =  FMath::Clamp<int32>(CapsuleRadius/4.f, 16, 64);
 
 			for (int32 ViewIndex = 0; ViewIndex < Views.Num(); ViewIndex++)
 			{
+
 				if (VisibilityMap & (1 << ViewIndex))
 				{
+					const FSceneView* View = Views[ViewIndex];
+					const FLinearColor DrawCapsuleColor = GetViewSelectionColor(ShapeColor, *View, IsSelected(), IsHovered(), false, IsIndividuallySelected() );
+
 					FPrimitiveDrawInterface* PDI = Collector.GetPDI(ViewIndex);
 					DrawWireCapsule( PDI, LocalToWorld.GetOrigin(), LocalToWorld.GetScaledAxis( EAxis::X ), LocalToWorld.GetScaledAxis( EAxis::Y ), LocalToWorld.GetScaledAxis( EAxis::Z ), DrawCapsuleColor, CapsuleRadius, CapsuleHalfHeight, CapsuleSides, SDPG_World );
 				}
diff --git a/Engine/Source/Runtime/Engine/Private/PrimitiveDrawingUtils.cpp b/Engine/Source/Runtime/Engine/Private/PrimitiveDrawingUtils.cpp
--- a/Engine/Source/Runtime/Engine/Private/PrimitiveDrawingUtils.cpp
+++ b/Engine/Source/Runtime/Engine/Private/PrimitiveDrawingUtils.cpp
@@ -1429,22 +1429,16 @@ void DrawWireDiamond(FPrimitiveDrawInterface* PDI,const FMatrix& DiamondMatrix,
 	PDI->DrawLine(SquarePoints[3], SquarePoints[0], InColor, DepthPriority);
 }
 
-FLinearColor GetSelectionColor(const FLinearColor& BaseColor,bool bSelected,bool bHovered, bool bUseOverlayIntensity)
+static FLinearColor ApplySelectionIntensity(const FLinearColor& FinalColor, bool bSelected, bool bHovered, bool bUseOverlayIntensity)
 {
-	FLinearColor FinalColor = BaseColor;
-	if( bSelected )
-	{
-		FinalColor = GEngine->GetSelectedMaterialColor();
-	}
-
 	static const float BaseIntensity = 0.5f;
 	static const float SelectedIntensity = 0.5f;
 	static const float HoverIntensity = 0.15f;
 
 	const float OverlayIntensity = bUseOverlayIntensity ? GEngine->SelectionHighlightIntensity : 1.0f;
-	float ResultingIntensity = bSelected ? SelectedIntensity : ( bHovered ? HoverIntensity : 0.0f );
+	float ResultingIntensity = bSelected ? SelectedIntensity : (bHovered ? HoverIntensity : 0.0f);
 
-	ResultingIntensity = ( ResultingIntensity * OverlayIntensity ) + BaseIntensity;
+	ResultingIntensity = (ResultingIntensity * OverlayIntensity) + BaseIntensity;
 
 	FLinearColor ret = FinalColor * FMath::Pow(ResultingIntensity, 2.2f);
 	ret.A = FinalColor.A;
@@ -1452,6 +1446,33 @@ FLinearColor GetSelectionColor(const FLinearColor& BaseColor,bool bSelected,bool
 	return ret;
 }
 
+FLinearColor GetSelectionColor(const FLinearColor& BaseColor,bool bSelected,bool bHovered, bool bUseOverlayIntensity)
+{
+	FLinearColor FinalColor = BaseColor;
+	if( bSelected )
+	{
+		FinalColor = GEngine->GetSelectedMaterialColor();
+	}
+
+	return ApplySelectionIntensity(FinalColor, bSelected, bHovered, bUseOverlayIntensity );
+
+}
+
+FLinearColor GetViewSelectionColor(const FLinearColor& BaseColor, const FSceneView& View, bool bSelected, bool bHovered, bool bUseOverlayIntensity, bool bIndividuallySelected)
+{
+	FLinearColor FinalColor = BaseColor;
+	if (View.bHasSelectedComponents && !bIndividuallySelected)
+	{
+		FinalColor = GEngine->GetSubduedSelectionOutlineColor();
+	}
+	else
+	{
+		FinalColor = GEngine->GetSelectedMaterialColor();
+	}
+
+	return ApplySelectionIntensity(FinalColor, bSelected, bHovered, bUseOverlayIntensity );
+}
+
 bool IsRichView(const FSceneViewFamily& ViewFamily)
 {
 	// Flags which make the view rich when absent.
diff --git a/Engine/Source/Runtime/Engine/Private/SphereComponent.cpp b/Engine/Source/Runtime/Engine/Private/SphereComponent.cpp
--- a/Engine/Source/Runtime/Engine/Private/SphereComponent.cpp
+++ b/Engine/Source/Runtime/Engine/Private/SphereComponent.cpp
@@ -111,7 +111,7 @@ FPrimitiveSceneProxy* USphereComponent::CreateSceneProxy()
 					}
 					else
 					{
-						const FLinearColor DrawSphereColor = GetSelectionColor(SphereColor, IsSelected(), IsHovered(), /*bUseOverlayIntensity=*/false);
+						const FLinearColor DrawSphereColor = GetViewSelectionColor(SphereColor, *View, IsSelected(), IsHovered(), false, IsIndividuallySelected() );
 
 						float AbsScaleX = LocalToWorld.GetScaledAxis(EAxis::X).Size();
 						float AbsScaleY = LocalToWorld.GetScaledAxis(EAxis::Y).Size();
diff --git a/Engine/Source/Runtime/Engine/Public/SceneManagement.h b/Engine/Source/Runtime/Engine/Public/SceneManagement.h
--- a/Engine/Source/Runtime/Engine/Public/SceneManagement.h
+++ b/Engine/Source/Runtime/Engine/Public/SceneManagement.h
@@ -1625,6 +1625,8 @@ void BuildCylinderVerts(const FVector& Base, const FVector& XAxis, const FVector
  * @return The color to draw the object with, accounting for the selection state
  */
 extern ENGINE_API FLinearColor GetSelectionColor(const FLinearColor& BaseColor,bool bSelected,bool bHovered, bool bUseOverlayIntensity = true);
+extern ENGINE_API FLinearColor GetViewSelectionColor(const FLinearColor& BaseColor, const FSceneView& View, bool bSelected, bool bHovered, bool bUseOverlayIntensity, bool bIndividuallySelected);
+
 
 /** Vertex Color view modes */
 namespace EVertexColorViewMode
