[INTEGRATE] Change 2437075 by [REDACTED_PERSON][REDACTED_HANDLE] on 2015/02/07 20:45:18

	Fix non-unity, non-editor

diff --git a/Engine/Source/Runtime/Engine/Private/BillboardComponent.cpp b/Engine/Source/Runtime/Engine/Private/BillboardComponent.cpp
--- a/Engine/Source/Runtime/Engine/Private/BillboardComponent.cpp
+++ b/Engine/Source/Runtime/Engine/Private/BillboardComponent.cpp
@@ -148,11 +148,14 @@ public:
 					// The color is multiplied by 10 because this value is normally expected to be blended
 					// additively, this is not how the sprites work and therefore need an extra boost
 					// to appear the same color as previously
+#if WITH_EDITOR
 					if( View->bHasSelectedComponents && !IsIndividuallySelected() )
 					{
 						ColorToUse = FLinearColor::White + (GEngine->GetSubduedSelectionOutlineColor() * GEngine->SelectionHighlightIntensityBillboards * 10);
 					}
-					else if (IsSelected())
+					else 
+#endif
+					if (IsSelected())
 					{
 						ColorToUse = FLinearColor::White + (GEngine->GetSelectedMaterialColor() * GEngine->SelectionHighlightIntensityBillboards * 10);
 					}
diff --git a/Engine/Source/Runtime/Engine/Private/PrimitiveDrawingUtils.cpp b/Engine/Source/Runtime/Engine/Private/PrimitiveDrawingUtils.cpp
--- a/Engine/Source/Runtime/Engine/Private/PrimitiveDrawingUtils.cpp
+++ b/Engine/Source/Runtime/Engine/Private/PrimitiveDrawingUtils.cpp
@@ -1461,11 +1461,13 @@ FLinearColor GetSelectionColor(const FLinearColor& BaseColor,bool bSelected,bool
 FLinearColor GetViewSelectionColor(const FLinearColor& BaseColor, const FSceneView& View, bool bSelected, bool bHovered, bool bUseOverlayIntensity, bool bIndividuallySelected)
 {
 	FLinearColor FinalColor = BaseColor;
+#if WITH_EDITOR
 	if (View.bHasSelectedComponents && !bIndividuallySelected)
 	{
 		FinalColor = GEngine->GetSubduedSelectionOutlineColor();
 	}
 	else
+#endif
 	{
 		FinalColor = GEngine->GetSelectedMaterialColor();
 	}
