[INTEGRATE] Change 2437119 by [REDACTED_PERSON][REDACTED_HANDLE] on 2015/02/07 22:06:39

	Component tree: Fix crash when deleting a component that has no siblings but has children [UE-9224]

diff --git a/Engine/Source/Runtime/Engine/Private/SceneComponent.cpp b/Engine/Source/Runtime/Engine/Private/SceneComponent.cpp
--- a/Engine/Source/Runtime/Engine/Private/SceneComponent.cpp
+++ b/Engine/Source/Runtime/Engine/Private/SceneComponent.cpp
@@ -365,6 +365,8 @@ void USceneComponent::DestroyComponent(bool bPromoteChildren/*= false*/)
 					// Attach the child node that we're promoting to the parent and move it to the same position as the old node was in the array
 					ChildToPromote->AttachTo(CachedAttachParent, NAME_None, EAttachLocation::KeepWorldPosition);
 					CachedAttachParent->AttachChildren.Remove(ChildToPromote);
+
+					Index = FMath::Clamp<int32>(Index, 0, CachedAttachParent->AttachChildren.Num());
 					CachedAttachParent->AttachChildren.Insert(ChildToPromote, Index);
 				}
 			}
