[INTEGRATE] Change 2437230 by [REDACTED_PERSON][REDACTED_HANDLE] on 2015/02/08 03:51:49

	Improve overlap/hit comments some more

diff --git a/Engine/Source/Runtime/Engine/Classes/Components/PrimitiveComponent.h b/Engine/Source/Runtime/Engine/Classes/Components/PrimitiveComponent.h
--- a/Engine/Source/Runtime/Engine/Classes/Components/PrimitiveComponent.h
+++ b/Engine/Source/Runtime/Engine/Classes/Components/PrimitiveComponent.h
@@ -568,14 +568,16 @@ public:
 	virtual bool ComponentOverlapMulti(TArray<struct FOverlapResult>& OutOverlaps, const class UWorld* World, const FVector& Pos, const FRotator& Rot, ECollisionChannel TestChannel, const struct FComponentQueryParams& Params, const struct FCollisionObjectQueryParams& ObjectQueryParams = FCollisionObjectQueryParams::DefaultObjectQueryParam) const;
 
 	/** 
-	 *	Event called when a component collides with something (or is collided with). 
+	 *	Event called when a component hits (or is hit by) something solid. This could happen due to things like Character movement, using Set Location with 'sweep' enabled, or physics simulation.
+	 *	For events when objects overlap (e.g. walking into a trigger) see the 'Overlap' event.
 	 *	@note For collisions during physics simulation to generate hit events, 'Simulation Generates Hit Events' must be enabled for this component
 	 */
 	UPROPERTY(BlueprintAssignable, Category="Collision")
 	FComponentHitSignature OnComponentHit;
 
 	/** 
-	 *	Event called when something starts to overlaps this component. 
+	 *	Event called when something starts to overlaps this component, for example a player walking into a trigger.
+	 *	For events when objects have a blocking collision, for example a player hitting a wall, see 'Hit' events.
 	 *	@note Both this component and the other one must have bGenerateOverlapEvents set to true to generate overlap events.
 	 */
 	UPROPERTY(BlueprintAssignable, Category="Collision")
diff --git a/Engine/Source/Runtime/Engine/Classes/GameFramework/Actor.h b/Engine/Source/Runtime/Engine/Classes/GameFramework/Actor.h
--- a/Engine/Source/Runtime/Engine/Classes/GameFramework/Actor.h
+++ b/Engine/Source/Runtime/Engine/Classes/GameFramework/Actor.h
@@ -493,7 +493,8 @@ public:
 	FTakePointDamageSignature OnTakePointDamage;
 	
 	/** 
-	 *	Called when another actor begins to overlap this actor. 
+	 *	Called when another actor begins to overlap this actor, for example a player walking into a trigger.
+	 *	For events when objects have a blocking collision, for example a player hitting a wall, see 'Hit' events.
 	 *	@note Components on both this and the other Actor must have bGenerateOverlapEvents set to true to generate overlap events.
 	 */
 	UPROPERTY(BlueprintAssignable, Category="Collision")
@@ -539,7 +540,8 @@ public:
 	FActorEndTouchOverSignature OnInputTouchLeave;
 
 	/** 
-	 *	Called when this Actor hits (or is hit by) something solid. 
+	 *	Called when this Actor hits (or is hit by) something solid. This could happen due to things like Character movement, using Set Location with 'sweep' enabled, or physics simulation.
+	 *	For events when objects overlap (e.g. walking into a trigger) see the 'Overlap' event.
 	 *	@note For collisions during physics simulation to generate hit events, 'Simulation Generates Hit Events' must be enabled.
 	 */
 	UPROPERTY(BlueprintAssignable, Category="Collision")
@@ -1005,7 +1007,8 @@ public:
 	virtual void ReceiveTick(float DeltaSeconds);
 
 	/** 
-	 *	Event when this actor overlaps another actor. 
+	 *	Event when this actor overlaps another actor, for example a player walking into a trigger.
+	 *	For events when objects have a blocking collision, for example a player hitting a wall, see 'Hit' events.
 	 *	@note Components on both this and the other Actor must have bGenerateOverlapEvents set to true to generate overlap events.
 	 */
 	UFUNCTION(BlueprintImplementableEvent, meta=(FriendlyName = "ActorBeginOverlap"), Category="Collision")
@@ -1067,7 +1070,8 @@ public:
 	void GetOverlappingComponents(TArray<UPrimitiveComponent*>& OverlappingComponents) const;
 
 	/** 
-	 *	Event when this actor bumps into a blocking object, or blocks another actor that bumps into it. 
+	 *	Event when this actor bumps into a blocking object, or blocks another actor that bumps into it. This could happen due to things like Character movement, using Set Location with 'sweep' enabled, or physics simulation.
+	 *	For events when objects overlap (e.g. walking into a trigger) see the 'Overlap' event.
 	 *	@note For collisions during physics simulation to generate hit events, 'Simulation Generates Hit Events' must be enabled.
 	 */
 	UFUNCTION(BlueprintImplementableEvent, meta=(FriendlyName = "Hit"), Category="Collision")
