[INTEGRATE] Change 2438831 by [REDACTED_PERSON][REDACTED_HANDLE] on 2015/02/09 17:24:22

	Fix for node clean up occuring to late and missing the compile on load.
	This left the Property list out of date which resulted in issues in the blueprint editor.

diff --git a/Engine/Source/Editor/KismetCompiler/Private/KismetCompiler.cpp b/Engine/Source/Editor/KismetCompiler/Private/KismetCompiler.cpp
--- a/Engine/Source/Editor/KismetCompiler/Private/KismetCompiler.cpp
+++ b/Engine/Source/Editor/KismetCompiler/Private/KismetCompiler.cpp
@@ -581,6 +581,9 @@ void FKismetCompilerContext::CreateClassVariablesFromBlueprint()
 	// Create a class property for any simple-construction-script created components that should be exposed
 	if (Blueprint->SimpleConstructionScript != NULL)
 	{
+		// Ensure that nodes have valid templates (This will remove nodes that have had the classes the inherited from removed
+		Blueprint->SimpleConstructionScript->ValidateNodeTemplates(MessageLog);
+
 		// Ensure that variable names are valid and that there are no collisions with a parent class
 		Blueprint->SimpleConstructionScript->ValidateNodeVariableNames(MessageLog);
 
diff --git a/Engine/Source/Runtime/Engine/Classes/Engine/SimpleConstructionScript.h b/Engine/Source/Runtime/Engine/Classes/Engine/SimpleConstructionScript.h
--- a/Engine/Source/Runtime/Engine/Classes/Engine/SimpleConstructionScript.h
+++ b/Engine/Source/Runtime/Engine/Classes/Engine/SimpleConstructionScript.h
@@ -97,6 +97,9 @@ public:
 	/** Ensures that all nodes in the SCS have valid names for compilation/replication */
 	ENGINE_API void ValidateNodeVariableNames(class FCompilerResultsLog& MessageLog);
 
+	/** Ensures that all nodes in the SCS have valid templates */
+	ENGINE_API void ValidateNodeTemplates(class FCompilerResultsLog& MessageLog);
+
 	/** Called by the SCS editor to clear all SCS editor component references */
 	ENGINE_API void ClearEditorComponentReferences();
 
diff --git a/Engine/Source/Runtime/Engine/Private/SimpleConstructionScript.cpp b/Engine/Source/Runtime/Engine/Private/SimpleConstructionScript.cpp
--- a/Engine/Source/Runtime/Engine/Private/SimpleConstructionScript.cpp
+++ b/Engine/Source/Runtime/Engine/Private/SimpleConstructionScript.cpp
@@ -148,28 +148,9 @@ void USimpleConstructionScript::PreloadChain()
 void USimpleConstructionScript::PostLoad()
 {
 	Super::PostLoad();
-
+	
 	int32 NodeIndex;
 	TArray<USCS_Node*> Nodes = GetAllNodes();
-	
-	for (NodeIndex=0; NodeIndex < Nodes.Num(); ++NodeIndex)
-	{
-		USCS_Node* Node = Nodes[NodeIndex];
-		
-		if (GetLinkerUE4Version() < VER_UE4_REMOVE_INPUT_COMPONENTS_FROM_BLUEPRINTS)
-		{
-			if (!Node->bIsNative_DEPRECATED && Node->ComponentTemplate->IsA<UInputComponent>())
-			{
-				RemoveNodeAndPromoteChildren(Node);
-			}
-		}
-
-		// Couldn't find the template the Blueprint or C++ class must have been deleted out from under us
-		if (Node->ComponentTemplate == nullptr)
-		{
-			RemoveNodeAndPromoteChildren(Node);
-		}
-	}
 
 #if WITH_EDITOR
 	// Get the Blueprint that owns the SCS
@@ -869,6 +850,31 @@ void USimpleConstructionScript::ValidateNodeVariableNames(FCompilerResultsLog& M
 	}
 }
 
+void USimpleConstructionScript::ValidateNodeTemplates(FCompilerResultsLog& MessageLog)
+{
+	int32 NodeIndex;
+	TArray<USCS_Node*> Nodes = GetAllNodes();
+
+	for (NodeIndex = 0; NodeIndex < Nodes.Num(); ++NodeIndex)
+	{
+		USCS_Node* Node = Nodes[NodeIndex];
+
+		if (GetLinkerUE4Version() < VER_UE4_REMOVE_INPUT_COMPONENTS_FROM_BLUEPRINTS)
+		{
+			if (!Node->bIsNative_DEPRECATED && Node->ComponentTemplate->IsA<UInputComponent>())
+			{
+				RemoveNodeAndPromoteChildren(Node);
+			}
+		}
+
+		// Couldn't find the template the Blueprint or C++ class must have been deleted out from under us
+		if (Node->ComponentTemplate == nullptr)
+		{
+			RemoveNodeAndPromoteChildren(Node);
+		}
+	}
+}
+
 void USimpleConstructionScript::ClearEditorComponentReferences()
 {
 	TArray<USCS_Node*> Nodes = GetAllNodes();
