[INTEGRATE] Change 2439358 by [REDACTED_PERSON][REDACTED_HANDLE] on 2015/02/09 23:49:40

	UE-9463 - UE-9486 - UMG - Fixing custom cursor loading code so that it works in cooked projects, also adding checks to prevent it from crashing in the editor if the class can not be resolved after being renamed.

diff --git a/Engine/Source/Runtime/Engine/Classes/Engine/GameViewportClient.h b/Engine/Source/Runtime/Engine/Classes/Engine/GameViewportClient.h
--- a/Engine/Source/Runtime/Engine/Classes/Engine/GameViewportClient.h
+++ b/Engine/Source/Runtime/Engine/Classes/Engine/GameViewportClient.h
@@ -632,6 +632,9 @@ private:
 	/** Delegate handler for when an actor is spawned */
 	void ShowCollisionOnSpawnedActors(AActor* Actor);
 
+	/** Adds a cursor to the set based on the enum and the class reference to it. */
+	void AddCursor(EMouseCursor::Type Cursor, const FStringClassReference& CursorClass);
+
 private:
 	/** Slate window associated with this viewport client.  The same window may host more than one viewport client. */
 	TWeakPtr<SWindow> Window;
diff --git a/Engine/Source/Runtime/Engine/Classes/Engine/UserInterfaceSettings.h b/Engine/Source/Runtime/Engine/Classes/Engine/UserInterfaceSettings.h
--- a/Engine/Source/Runtime/Engine/Classes/Engine/UserInterfaceSettings.h
+++ b/Engine/Source/Runtime/Engine/Classes/Engine/UserInterfaceSettings.h
@@ -35,6 +35,8 @@ namespace EUIScalingRule
 	};
 }
 
+class UWidget;
+
 /**
  * Implements user interface related settings.
  */
@@ -80,6 +82,14 @@ public:
 	FRuntimeFloatCurve UIScaleCurve;
 
 public:
+	UPROPERTY(Transient)
+	TArray<UObject*> CursorClasses;
+
+public:
+
+	virtual void PostInitProperties() override;
+
+	void LoadCursors();
 
 	/** Gets the current scale of the UI based on the size */
 	float GetDPIScaleBasedOnSize(FIntPoint Size) const;
diff --git a/Engine/Source/Runtime/Engine/Private/GameViewportClient.cpp b/Engine/Source/Runtime/Engine/Private/GameViewportClient.cpp
--- a/Engine/Source/Runtime/Engine/Private/GameViewportClient.cpp
+++ b/Engine/Source/Runtime/Engine/Private/GameViewportClient.cpp
@@ -229,20 +229,12 @@ void UGameViewportClient::Init(struct FWorldContext& WorldContext, UGameInstance
 	// Create the cursor Widgets
 	UUserInterfaceSettings* UISettings = GetMutableDefault<UUserInterfaceSettings>(UUserInterfaceSettings::StaticClass());
 
-#define ADD_CURSOR(enumeration, variable) \
-	if (UISettings->variable.IsValid()) { \
-	UClass* Class = LoadObject<UClass>(NULL, *UISettings->variable.ToString()); \
-	UUserWidget * UserWidget = CreateWidget<UUserWidget>(GetWorld(), Class); \
-	CursorWidgets.Add(enumeration, UserWidget->TakeWidget()); }
-
-	ADD_CURSOR(EMouseCursor::Default, DefaultCursor);
-	ADD_CURSOR(EMouseCursor::TextEditBeam, TextEditBeamCursor);
-	ADD_CURSOR(EMouseCursor::Crosshairs, CrosshairsCursor);
-	ADD_CURSOR(EMouseCursor::GrabHand, GrabHandCursor);
-	ADD_CURSOR(EMouseCursor::GrabHandClosed, GrabHandClosedCursor);
-	ADD_CURSOR(EMouseCursor::SlashedCircle, SlashedCircleCursor);
-
-#undef ADD_CURSOR
+	AddCursor(EMouseCursor::Default, UISettings->DefaultCursor);
+	AddCursor(EMouseCursor::TextEditBeam, UISettings->TextEditBeamCursor);
+	AddCursor(EMouseCursor::Crosshairs, UISettings->CrosshairsCursor);
+	AddCursor(EMouseCursor::GrabHand, UISettings->GrabHandCursor);
+	AddCursor(EMouseCursor::GrabHandClosed, UISettings->GrabHandClosedCursor);
+	AddCursor(EMouseCursor::SlashedCircle, UISettings->SlashedCircleCursor);
 }
 
 UWorld* UGameViewportClient::GetWorld() const
@@ -556,6 +548,22 @@ EMouseCursor::Type UGameViewportClient::GetCursor(FViewport* InViewport, int32 X
 	return FViewportClient::GetCursor(InViewport, X, Y);
 }
 
+void UGameViewportClient::AddCursor(EMouseCursor::Type Cursor, const FStringClassReference& CursorClass)
+{
+	if ( CursorClass.IsValid() )
+	{
+		UClass* Class = CursorClass.TryLoadClass<UUserWidget>();
+		if ( ensure(Class) )
+		{
+			UUserWidget* UserWidget = CreateWidget<UUserWidget>(GetGameInstance(), Class);
+			if ( ensure(UserWidget) )
+			{
+				CursorWidgets.Add(Cursor, UserWidget->TakeWidget());
+			}
+		}
+	}
+}
+
 TOptional<TSharedRef<SWidget>> UGameViewportClient::MapCursor(FViewport* Viewport, const FCursorReply& CursorReply)
 {
 	const TSharedRef<SWidget>* CursorWidgetPtr = CursorWidgets.Find(CursorReply.GetCursorType());
diff --git a/Engine/Source/Runtime/Engine/Private/UnrealEngine.cpp b/Engine/Source/Runtime/Engine/Private/UnrealEngine.cpp
--- a/Engine/Source/Runtime/Engine/Private/UnrealEngine.cpp
+++ b/Engine/Source/Runtime/Engine/Private/UnrealEngine.cpp
@@ -99,6 +99,7 @@
 #include "PhysicalMaterials/PhysicalMaterial.h"
 #include "NotificationManager.h"
 #include "SNotificationList.h"
+#include "Engine/UserInterfaceSettings.h"
 
 DEFINE_LOG_CATEGORY_STATIC(LogEngine, Log, All);
 
@@ -1457,6 +1458,9 @@ void UEngine::InitializeObjectReferences()
 
 		checkf(DefaultPreviewPawnClass != NULL, TEXT("Engine config value DefaultPreviewPawnClass is not a valid class name."));
 	}
+
+	UUserInterfaceSettings* UISettings = GetMutableDefault<UUserInterfaceSettings>(UUserInterfaceSettings::StaticClass());
+	UISettings->LoadCursors();
 }
 
 //
diff --git a/Engine/Source/Runtime/Engine/Private/UserInterfaceSettings.cpp b/Engine/Source/Runtime/Engine/Private/UserInterfaceSettings.cpp
--- a/Engine/Source/Runtime/Engine/Private/UserInterfaceSettings.cpp
+++ b/Engine/Source/Runtime/Engine/Private/UserInterfaceSettings.cpp
@@ -15,6 +15,16 @@ UUserInterfaceSettings::UUserInterfaceSettings(const FObjectInitializer& ObjectI
 {
 }
 
+void UUserInterfaceSettings::PostInitProperties()
+{
+	Super::PostInitProperties();
+
+	if ( HasAnyFlags(RF_ClassDefaultObject) == false )
+	{
+		LoadCursors();
+	}
+}
+
 float UUserInterfaceSettings::GetDPIScaleBasedOnSize(FIntPoint Size) const
 {
 	float Scale = 1;
@@ -49,6 +59,25 @@ float UUserInterfaceSettings::GetDPIScaleBasedOnSize(FIntPoint Size) const
 	return FMath::Max(Scale, 0.01f);
 }
 
+void UUserInterfaceSettings::LoadCursors()
+{
+	TArray<UObject*> LoadedClasses;
+	LoadedClasses.Add(DefaultCursor.TryLoad());
+	LoadedClasses.Add(TextEditBeamCursor.TryLoad());
+	LoadedClasses.Add(CrosshairsCursor.TryLoad());
+	LoadedClasses.Add(GrabHandCursor.TryLoad());
+	LoadedClasses.Add(GrabHandClosedCursor.TryLoad());
+	LoadedClasses.Add(SlashedCircleCursor.TryLoad());
+
+	for ( UObject* Cursor : LoadedClasses )
+	{
+		if ( Cursor )
+		{
+			CursorClasses.Add(Cursor);
+		}
+	}
+}
+
 #if WITH_EDITOR
 void UUserInterfaceSettings::PostEditChangeProperty(FPropertyChangedEvent& PropertyChangedEvent)
 {
