Fix potential undo crashes from painted vertex colors

diff --git a/Engine/Source/Runtime/Engine/Private/StaticMeshComponent.cpp b/Engine/Source/Runtime/Engine/Private/StaticMeshComponent.cpp
--- a/Engine/Source/Runtime/Engine/Private/StaticMeshComponent.cpp
+++ b/Engine/Source/Runtime/Engine/Private/StaticMeshComponent.cpp
@@ -1609,8 +1609,10 @@ void UStaticMeshComponent::ApplyComponentInstanceData(FStaticMeshComponentInstan
 		bHasCachedStaticLighting = true;
 	}
 
-	StaticMeshInstanceData->ApplyVertexColorData(this);
-	MarkRenderStateDirty();	
+	{
+		FComponentReregisterContext ReregisterStaticMesh(this);
+		StaticMeshInstanceData->ApplyVertexColorData(this);
+	}
 }
 
 #include "AI/Navigation/RecastHelpers.h"
