#ue4 - Integrate 2440500 from 4.7:

Increase cooked triangle small area threshold (TriangleMeshTriangleMinAreaThreshold) from 0.25cm -> 5.0cm
The previous default was too conservative to catch long degenerate triangles which could break physics collision.

Related to UE-9547 (UT BSP colission issues)

diff --git a/Engine/Source/Runtime/Engine/Private/PhysicsEngine/PhysDerivedData.h b/Engine/Source/Runtime/Engine/Private/PhysicsEngine/PhysDerivedData.h
--- a/Engine/Source/Runtime/Engine/Private/PhysicsEngine/PhysDerivedData.h
+++ b/Engine/Source/Runtime/Engine/Private/PhysicsEngine/PhysDerivedData.h
@@ -45,7 +45,8 @@ public:
 	{
 		//  1 - base version
 		//  2 - cook out small area trimesh triangles from BSP (see UPhysicsSettings::TriangleMeshTriangleMinAreaThreshold)
-		enum { UE_PHYSX_DERIVEDDATA_VER = 2 };
+		//  3 - increase default small area threshold and force recook.
+		enum { UE_PHYSX_DERIVEDDATA_VER = 3 };
 
 		const uint16 PhysXVersion = ((PX_PHYSICS_VERSION_MAJOR  & 0xF) << 12) |
 				((PX_PHYSICS_VERSION_MINOR  & 0xF) << 8) |
diff --git a/Engine/Source/Runtime/Engine/Private/PhysicsEngine/PhysicsSettings.cpp b/Engine/Source/Runtime/Engine/Private/PhysicsEngine/PhysicsSettings.cpp
--- a/Engine/Source/Runtime/Engine/Private/PhysicsEngine/PhysicsSettings.cpp
+++ b/Engine/Source/Runtime/Engine/Private/PhysicsEngine/PhysicsSettings.cpp
@@ -11,7 +11,7 @@ UPhysicsSettings::UPhysicsSettings(const FObjectInitializer& ObjectInitializer)
 	, DefaultTerminalVelocity(4000.f)
 	, DefaultFluidFriction(0.3f)
 	, RagdollAggregateThreshold(4)
-	, TriangleMeshTriangleMinAreaThreshold(0.25f)
+	, TriangleMeshTriangleMinAreaThreshold(5.0f)
 	, bEnableAsyncScene(false)
 	, bEnable2DPhysics(false)
 	, BounceThresholdVelocity(200.f)
