Allow matinee movie recording to capture UMG/Slate hud if bHideHud is disabled.

diff --git a/Engine/Source/Editor/Matinee/Private/MatineeRecordMovie.cpp b/Engine/Source/Editor/Matinee/Private/MatineeRecordMovie.cpp
--- a/Engine/Source/Editor/Matinee/Private/MatineeRecordMovie.cpp
+++ b/Engine/Source/Editor/Matinee/Private/MatineeRecordMovie.cpp
@@ -214,17 +214,17 @@ FReply SMatineeRecordMovie::OnOK()
 	if ( Mode != NULL && Mode->InterpEd != NULL )
 	{
 		// Store the options for the capture of the Matinee
-		GEngine->MatineeCaptureName = Mode->InterpEd->GetMatineeActor()->GetName();
-		GEngine->MatineePackageCaptureName = FPackageName::GetShortName(Mode->InterpEd->GetMatineeActor()->GetOutermost()->GetName());
+		GEngine->MatineeScreenshotOptions.MatineeCaptureName = Mode->InterpEd->GetMatineeActor()->GetName();
+		GEngine->MatineeScreenshotOptions.MatineePackageCaptureName = FPackageName::GetShortName(Mode->InterpEd->GetMatineeActor()->GetOutermost()->GetName());
 
-		GUnrealEd->bNoTextureStreaming = Options.DisableTextureStreaming;
-		GUnrealEd->MatineeCaptureFPS = Options.CaptureResolutionFPS;
-		GUnrealEd->bCompressMatineeCapture = Options.Compress;
+		GEngine->MatineeScreenshotOptions.bNoTextureStreaming = Options.DisableTextureStreaming;
+		GEngine->MatineeScreenshotOptions.MatineeCaptureFPS = Options.CaptureResolutionFPS;
+		GEngine->MatineeScreenshotOptions.bCompressMatineeCapture = Options.Compress;
 
 		GUnrealEd->MatineeCaptureResolutionX = CaptureWidth;
 		GUnrealEd->MatineeCaptureResolutionY = CaptureHeight;
 
-		GUnrealEd->MatineeCaptureType = (EMatineeCaptureType::Type)Options.CaptureTypeIndex;
+		GEngine->MatineeScreenshotOptions.MatineeCaptureType = (EMatineeCaptureType::Type)Options.CaptureTypeIndex;
 		Mode->InterpEd->StartRecordingMovie();
 		
 		//if Options.CloseEditor == true, Editor will request a close action in UEditorEngine::PlayForMovieCapture
diff --git a/Engine/Source/Editor/UnrealEd/Private/PlayLevel.cpp b/Engine/Source/Editor/UnrealEd/Private/PlayLevel.cpp
--- a/Engine/Source/Editor/UnrealEd/Private/PlayLevel.cpp
+++ b/Engine/Source/Editor/UnrealEd/Private/PlayLevel.cpp
@@ -1666,9 +1666,6 @@ void UEditorEngine::PlayForMovieCapture()
 	// (we add to EditorCommandLine because the URL is ignored by WindowsTools)
 	EditorCommandLine += TEXT(" -PIEVIACONSOLE");
 	
-	// Disable splash screen
-	EditorCommandLine += TEXT(" -NoSplash");
-	
 	// if we want to start movie capturing right away, then append the argument for that
 	if (bStartMovieCapture)
 	{
@@ -1679,19 +1676,19 @@ void UEditorEngine::PlayForMovieCapture()
 		EditorCommandLine += FString::Printf(TEXT(" -ResX=%d"), GEditor->MatineeCaptureResolutionX);
 		EditorCommandLine += FString::Printf(TEXT(" -ResY=%d"), GEditor->MatineeCaptureResolutionY);
 					
-		if( GUnrealEd->bNoTextureStreaming )
+		if( GUnrealEd->MatineeScreenshotOptions.bNoTextureStreaming )
 		{
 			EditorCommandLine += TEXT(" -NoTextureStreaming");
 		}
 
 		//set fps
-		EditorCommandLine += FString::Printf(TEXT(" -BENCHMARK -FPS=%d"), GEditor->MatineeCaptureFPS);
+		EditorCommandLine += FString::Printf(TEXT(" -BENCHMARK -FPS=%d"), GEditor->MatineeScreenshotOptions.MatineeCaptureFPS);
 	
-		if (GEditor->MatineeCaptureType.GetValue() != EMatineeCaptureType::AVI)
+		if (GEditor->MatineeScreenshotOptions.MatineeCaptureType.GetValue() != EMatineeCaptureType::AVI)
 		{
-			EditorCommandLine += FString::Printf(TEXT(" -MATINEESSCAPTURE=%s"), *GEngine->MatineeCaptureName);//*GEditor->MatineeNameForRecording);
+			EditorCommandLine += FString::Printf(TEXT(" -MATINEESSCAPTURE=%s"), *GEngine->MatineeScreenshotOptions.MatineeCaptureName);//*GEditor->MatineeNameForRecording);
 
-			switch(GEditor->MatineeCaptureType.GetValue())
+			switch(GEditor->MatineeScreenshotOptions.MatineeCaptureType.GetValue())
 			{
 			case EMatineeCaptureType::BMP:
 				EditorCommandLine += TEXT(" -MATINEESSFORMAT=BMP");
@@ -1714,12 +1711,12 @@ void UEditorEngine::PlayForMovieCapture()
 		}
 		else
 		{
-			EditorCommandLine += FString::Printf(TEXT(" -MATINEEAVICAPTURE=%s"), *GEngine->MatineeCaptureName);//*GEditor->MatineeNameForRecording);
+			EditorCommandLine += FString::Printf(TEXT(" -MATINEEAVICAPTURE=%s"), *GEngine->MatineeScreenshotOptions.MatineeCaptureName);//*GEditor->MatineeNameForRecording);
 		}
 					
-		EditorCommandLine += FString::Printf(TEXT(" -MATINEEPACKAGE=%s"), *GEngine->MatineePackageCaptureName);//*GEditor->MatineePackageNameForRecording);
+		EditorCommandLine += FString::Printf(TEXT(" -MATINEEPACKAGE=%s"), *GEngine->MatineeScreenshotOptions.MatineePackageCaptureName);//*GEditor->MatineePackageNameForRecording);
 	
-		if (GEditor->bCompressMatineeCapture == 1)
+		if (GEditor->MatineeScreenshotOptions.bCompressMatineeCapture == 1)
 		{
 			EditorCommandLine += TEXT(" -CompressCapture");
 		}
diff --git a/Engine/Source/Runtime/Engine/Classes/Engine/Engine.h b/Engine/Source/Runtime/Engine/Classes/Engine/Engine.h
--- a/Engine/Source/Runtime/Engine/Classes/Engine/Engine.h
+++ b/Engine/Source/Runtime/Engine/Classes/Engine/Engine.h
@@ -539,6 +539,42 @@
 DECLARE_DELEGATE_OneParam(FBeginStreamingPauseDelegate, FViewport*);
 DECLARE_DELEGATE(FEndStreamingPauseDelegate);
 
+USTRUCT()
+struct FMatineeScreenshotOptions
+{
+	GENERATED_USTRUCT_BODY()
 
+	/** determines if we should start the matinee capture as soon as the game loads */
+	UPROPERTY(transient)
+	uint32 bStartWithMatineeCapture:1;
+
+	/** should we compress the capture */
+	UPROPERTY(transient)
+	uint32 bCompressMatineeCapture:1;
+
+	/** the name of the matine that we want to record */
+	UPROPERTY(transient)
+	FString MatineeCaptureName;
+
+	/** The package name where the matinee belongs to */
+	UPROPERTY(transient)
+	FString MatineePackageCaptureName;
+
+	/** the fps of the matine that we want to record */
+	UPROPERTY(transient)
+	int32 MatineeCaptureFPS;
+
+	/** The capture type, e.g. AVI or Screen Shots */
+	UPROPERTY(transient)
+	TEnumAsByte<EMatineeCaptureType::Type> MatineeCaptureType;
+
+	/** Whether or not to disable texture streaming during matinee movie capture */
+	UPROPERTY(transient)
+	bool bNoTextureStreaming;
+
+	UPROPERTY(transient)
+	bool bHideHud;
+};
+
 /**
  * Abstract base class of all Engine classes, responsible for management of systems critical to editor or game systems.
@@ -1385,34 +1421,8 @@
 	int32 ScreenSaverInhibitorSemaphore;
 
 public:
+	UPROPERTY(transient)
+	FMatineeScreenshotOptions MatineeScreenshotOptions;
 
-	/** determines if we should start the matinee capture as soon as the game loads */
-	UPROPERTY(transient)
-	uint32 bStartWithMatineeCapture:1;
-
-	/** should we compress the capture */
-	UPROPERTY(transient)
-	uint32 bCompressMatineeCapture:1;
-
-	/** the name of the matine that we want to record */
-	UPROPERTY(transient)
-	FString MatineeCaptureName;
-
-	/** The package name where the matinee belongs to */
-	UPROPERTY(transient)
-	FString MatineePackageCaptureName;
-
-	/** the fps of the matine that we want to record */
-	UPROPERTY(transient)
-	int32 MatineeCaptureFPS;
-
-	/** The capture type, e.g. AVI or Screen Shots */
-	UPROPERTY(transient)
-	TEnumAsByte<EMatineeCaptureType::Type> MatineeCaptureType;
-
-	/** Whether or not to disable texture streaming during matinee movie capture */
-	UPROPERTY(transient)
-	bool bNoTextureStreaming;
-
 	/** true if the the user cannot modify levels that are read only. */
 	UPROPERTY(transient)
diff --git a/Engine/Source/Runtime/Engine/Private/AVIWriter.cpp b/Engine/Source/Runtime/Engine/Private/AVIWriter.cpp
--- a/Engine/Source/Runtime/Engine/Private/AVIWriter.cpp
+++ b/Engine/Source/Runtime/Engine/Private/AVIWriter.cpp
@@ -190,7 +190,7 @@ public:
 			}
 
 			IAMVideoCompression* VideoCompression = NULL;
-			if (GEngine->bCompressMatineeCapture)
+			if (GEngine->MatineeScreenshotOptions.bCompressMatineeCapture)
 			{
 				// Create the filter graph manager and query for interfaces.
 				hr = CoCreateInstance (CLSID_MJPGEnc, NULL, CLSCTX_INPROC, IID_IBaseFilter, ( void ** ) &VideoCompression);
@@ -232,20 +232,20 @@ public:
 				FString OutputDir = FPaths::VideoCaptureDir();
 				IFileManager::Get().MakeDirectory(*OutputDir, true);
 
-				if (GEngine->bStartWithMatineeCapture)
+				if (GEngine->MatineeScreenshotOptions.bCompressMatineeCapture)
 				{
 					bool bFoundUnusedName = false;
 					while (!bFoundUnusedName)
 					{
 						const FString FileNameInfo = FString::Printf( TEXT("_%dfps_%dx%d"), 
-							GEngine->MatineeCaptureFPS, CaptureViewport->GetSizeXY().X, CaptureViewport->GetSizeXY().Y );
-						if (GEngine->bCompressMatineeCapture)
+							GEngine->MatineeScreenshotOptions.MatineeCaptureFPS, CaptureViewport->GetSizeXY().X, CaptureViewport->GetSizeXY().Y );
+						if (GEngine->MatineeScreenshotOptions.bCompressMatineeCapture)
 						{
-							FCString::Sprintf( File, TEXT("%s_%d.avi"), *(OutputDir + TEXT("/") + GEngine->MatineePackageCaptureName + "_" + GEngine->MatineeCaptureName + FileNameInfo + "_compressed"), MovieCaptureIndex );
+							FCString::Sprintf( File, TEXT("%s_%d.avi"), *(OutputDir + TEXT("/") + GEngine->MatineeScreenshotOptions.MatineePackageCaptureName + "_" + GEngine->MatineeScreenshotOptions.MatineeCaptureName + FileNameInfo + "_compressed"), MovieCaptureIndex );
 						}
 						else
 						{
-							FCString::Sprintf( File, TEXT("%s_%d.avi"), *(OutputDir + TEXT("/") + GEngine->MatineePackageCaptureName + "_" + GEngine->MatineeCaptureName + FileNameInfo), MovieCaptureIndex );
+							FCString::Sprintf( File, TEXT("%s_%d.avi"), *(OutputDir + TEXT("/") + GEngine->MatineeScreenshotOptions.MatineePackageCaptureName + "_" + GEngine->MatineeScreenshotOptions.MatineeCaptureName + FileNameInfo), MovieCaptureIndex );
 						}
 						if (IFileManager::Get().FileSize(File) != -1)
 						{
@@ -360,7 +360,7 @@ public:
 	void Update(float DeltaSeconds)
 	{
 		bool bShouldUpdate = true;
-		const float CaptureFrequency = 1.0f/(float)GEngine->MatineeCaptureFPS;
+		const float CaptureFrequency = 1.0f/(float)GEngine->MatineeScreenshotOptions.MatineeCaptureFPS;
 		if (CaptureFrequency > 0.f)
 		{
 			CurrentAccumSeconds += DeltaSeconds;
@@ -385,7 +385,7 @@ public:
 
 		if (bMatineeFinished)
 		{
-			if (GEngine->MatineeCaptureType == 0)
+			if (GEngine->MatineeScreenshotOptions.MatineeCaptureType == 0)
 			{
 				if (CaptureViewport)
 				{
diff --git a/Engine/Source/Runtime/Engine/Private/CapturePin.cpp b/Engine/Source/Runtime/Engine/Private/CapturePin.cpp
--- a/Engine/Source/Runtime/Engine/Private/CapturePin.cpp
+++ b/Engine/Source/Runtime/Engine/Private/CapturePin.cpp
@@ -18,7 +18,7 @@ FCapturePin::FCapturePin(HRESULT *phr, CSource *pFilter)
         : CSourceStream(NAME("Push Source"), phr, pFilter, L"Out"),
 		FramesWritten(0),
 		bZeroMemory(0),
-		FrameLength(UNITS/GEngine->MatineeCaptureFPS),
+		FrameLength(UNITS/GEngine->MatineeScreenshotOptions.MatineeCaptureFPS),
 		CurrentBitDepth(32)
 {
 	// Get the dimensions of the window
diff --git a/Engine/Source/Runtime/Engine/Private/GameEngine.cpp b/Engine/Source/Runtime/Engine/Private/GameEngine.cpp
--- a/Engine/Source/Runtime/Engine/Private/GameEngine.cpp
+++ b/Engine/Source/Runtime/Engine/Private/GameEngine.cpp
@@ -95,7 +95,7 @@ void UGameEngine::CreateGameViewportWidget( UGameViewportClient* GameViewportCli
 		SNew( SViewport )
 			// Render directly to the window backbuffer unless capturing a movie or getting screenshots
 			// @todo TEMP
-			.RenderDirectlyToWindow( !GEngine->bStartWithMatineeCapture && GIsDumpingMovie == 0 )
+			.RenderDirectlyToWindow( !GEngine->MatineeScreenshotOptions.bStartWithMatineeCapture && GIsDumpingMovie == 0 )
 			.EnableStereoRendering(true)
 			[
 				SNew(SDPIScaler)
@@ -972,7 +972,7 @@ void UGameEngine::Tick( float DeltaSeconds, bool bIdleMode )
 		}
 
 		// Start the movie capture if needed
-		if (bCheckForMovieCapture && GEngine->bStartWithMatineeCapture && GEngine->MatineeCaptureType == EMatineeCaptureType::AVI && GameViewport->Viewport->GetSizeXY() != FIntPoint::ZeroValue )
+		if (bCheckForMovieCapture && GEngine->MatineeScreenshotOptions.bStartWithMatineeCapture && GEngine->MatineeScreenshotOptions.MatineeCaptureType == EMatineeCaptureType::AVI && GameViewport->Viewport->GetSizeXY() != FIntPoint::ZeroValue )
 		{
 			if (AVIWriter)
 			{
diff --git a/Engine/Source/Runtime/Engine/Private/GameMode.cpp b/Engine/Source/Runtime/Engine/Private/GameMode.cpp
--- a/Engine/Source/Runtime/Engine/Private/GameMode.cpp
+++ b/Engine/Source/Runtime/Engine/Private/GameMode.cpp
@@ -285,7 +285,7 @@ void AGameMode::PostLogin( APlayerController* NewPlayer )
 bool AGameMode::ShouldStartInCinematicMode(bool& OutHidePlayer,bool& OutHideHUD,bool& OutDisableMovement,bool& OutDisableTurning)
 {
 	bool StartInCinematicMode = false;
-	if(GEngine->bStartWithMatineeCapture)
+	if(GEngine->MatineeScreenshotOptions.bStartWithMatineeCapture)
 	{
 		GConfig->GetBool( TEXT("MatineeCreateMovieOptions"), TEXT("CinematicMode"), StartInCinematicMode, GEditorUserSettingsIni );
 		if(StartInCinematicMode)
diff --git a/Engine/Source/Runtime/Engine/Private/GameViewportClient.cpp b/Engine/Source/Runtime/Engine/Private/GameViewportClient.cpp
--- a/Engine/Source/Runtime/Engine/Private/GameViewportClient.cpp
+++ b/Engine/Source/Runtime/Engine/Private/GameViewportClient.cpp
@@ -1193,14 +1193,30 @@ void UGameViewportClient::ProcessScreenShots(FViewport* InViewport)
 	{
 		TArray<FColor> Bitmap;
 
+		bool bShowUI = false;
 		TSharedPtr<SWindow> WindowPtr = GetWindow();
 		if (!GIsDumpingMovie && (FScreenshotRequest::ShouldShowUI() && WindowPtr.IsValid()))
 		{
-			TSharedRef<SWindow> WindowRef = WindowPtr.ToSharedRef();
-			FSlateApplication::Get().ForceRedrawWindow(WindowRef);
+			bShowUI = true;
+		}
+		else if( GIsDumpingMovie && WindowPtr.IsValid() && !GEngine->MatineeScreenshotOptions.bHideHud )
+		{
+			bShowUI = true;
+		}
+
+		bool bScreenshotSuccessful = false;
+		if( bShowUI && FSlateApplication::IsInitialized() )
+		{
+			FIntVector Size;
+			TSharedRef<SWidget> WindowRef = WindowPtr.ToSharedRef();
+			bScreenshotSuccessful = FSlateApplication::Get().TakeScreenshot( WindowRef, Bitmap, Size);
+		}
+		else
+		{
+			bScreenshotSuccessful = GetViewportScreenShot(InViewport, Bitmap);
 		}
 
-		if (GetViewportScreenShot(InViewport, Bitmap))
+		if (bScreenshotSuccessful)
 		{
 			if (ScreenshotCapturedDelegate.IsBound())
 			{
diff --git a/Engine/Source/Runtime/Engine/Private/UnrealEngine.cpp b/Engine/Source/Runtime/Engine/Private/UnrealEngine.cpp
--- a/Engine/Source/Runtime/Engine/Private/UnrealEngine.cpp
+++ b/Engine/Source/Runtime/Engine/Private/UnrealEngine.cpp
@@ -1146,36 +1146,41 @@ void UEngine::ParseCommandline()
 		bDisableAILogging = false;
 	}
 
-	bStartWithMatineeCapture = false;
-	bCompressMatineeCapture = false;
+	MatineeScreenshotOptions.bStartWithMatineeCapture = false;
+	MatineeScreenshotOptions.bCompressMatineeCapture = false;
 #if WITH_EDITOR
-	if (!GIsEditor && FParse::Value(FCommandLine::Get(), TEXT("-MATINEEAVICAPTURE="), MatineeCaptureName))
+	if (!GIsEditor && FParse::Value(FCommandLine::Get(), TEXT("-MATINEEAVICAPTURE="), MatineeScreenshotOptions.MatineeCaptureName))
 	{
-		MatineeCaptureType = EMatineeCaptureType::AVI;
-		bStartWithMatineeCapture = true;
+		MatineeScreenshotOptions.MatineeCaptureType = EMatineeCaptureType::AVI;
+		MatineeScreenshotOptions.bStartWithMatineeCapture = true;
 	}
-	else if (!GIsEditor && FParse::Value(FCommandLine::Get(), TEXT("-MATINEESSCAPTURE="), MatineeCaptureName))
+	else if (!GIsEditor && FParse::Value(FCommandLine::Get(), TEXT("-MATINEESSCAPTURE="), MatineeScreenshotOptions.MatineeCaptureName))
 	{
-		MatineeCaptureType = EMatineeCaptureType::BMP;
+		MatineeScreenshotOptions.MatineeCaptureType = EMatineeCaptureType::BMP;
 
 		FString MatineeCaptureFormat;
 		if(FParse::Value(FCommandLine::Get(), TEXT("-MATINEESSFORMAT="), MatineeCaptureFormat))
 		{
 			if(MatineeCaptureFormat == TEXT("BMP"))
 			{
-				MatineeCaptureType = EMatineeCaptureType::BMP;
+				MatineeScreenshotOptions.MatineeCaptureType = EMatineeCaptureType::BMP;
 			}
 			else if(MatineeCaptureFormat == TEXT("PNG"))
 			{
-				MatineeCaptureType = EMatineeCaptureType::PNG;
+				MatineeScreenshotOptions.MatineeCaptureType = EMatineeCaptureType::PNG;
 			}
 			else if(MatineeCaptureFormat == TEXT("JPEG"))
 			{
-				MatineeCaptureType = EMatineeCaptureType::JPEG;
+				MatineeScreenshotOptions.MatineeCaptureType = EMatineeCaptureType::JPEG;
 			}
 		}
 
-		bStartWithMatineeCapture = true;
+		MatineeScreenshotOptions.bStartWithMatineeCapture = true;
+	}
+
+	if( !GIsEditor && MatineeScreenshotOptions.bStartWithMatineeCapture )
+	{
+		GConfig->GetBool( TEXT("MatineeCreateMovieOptions"), TEXT("HideHUD"), MatineeScreenshotOptions.bHideHud, GEditorUserSettingsIni );
 	}
 
 	// If we are capturing a matinee movie and we want to dump the buffer visualization shots too, for on all required functionality
@@ -1189,17 +1194,17 @@ void UEngine::ParseCommandline()
 		}
 	}
 
-	if (bStartWithMatineeCapture)
+	if (MatineeScreenshotOptions.bStartWithMatineeCapture)
 	{
-		FParse::Value(FCommandLine::Get(), TEXT("-MATINEEPACKAGE="), MatineePackageCaptureName);
+		FParse::Value(FCommandLine::Get(), TEXT("-MATINEEPACKAGE="), MatineeScreenshotOptions.MatineePackageCaptureName);
 	}
 
 	if ( !GIsEditor && FParse::Param(FCommandLine::Get(), TEXT("COMPRESSCAPTURE")) )
 	{
-		bCompressMatineeCapture = true;
+		MatineeScreenshotOptions.bCompressMatineeCapture = true;
 	}
 #endif
-	MatineeCaptureFPS = 30;
+	MatineeScreenshotOptions.MatineeCaptureFPS = 30;
 }
 
 
@@ -2175,11 +2180,11 @@ bool UEngine::Exec( UWorld* InWorld, const TCHAR* Cmd, FOutputDevice& Ar )
 	{
 		return HandleStatCommand(InWorld, GStatProcessingViewportClient, Cmd, Ar);
 	}
-	else if( FParse::Command(&Cmd,TEXT("STARTMOVIECAPTURE")) && (GEngine->bStartWithMatineeCapture == true || GIsEditor) )
+	else if( FParse::Command(&Cmd,TEXT("STARTMOVIECAPTURE")) && (GEngine->MatineeScreenshotOptions.bStartWithMatineeCapture == true || GIsEditor) )
 	{
 		return HandleStartMovieCaptureCommand( Cmd, Ar );
 	}
-	else if( FParse::Command(&Cmd,TEXT("STOPMOVIECAPTURE")) && (GEngine->bStartWithMatineeCapture == true || GIsEditor) )
+	else if( FParse::Command(&Cmd,TEXT("STOPMOVIECAPTURE")) && (GEngine->MatineeScreenshotOptions.bStartWithMatineeCapture == true || GIsEditor) )
 	{
 		return HandleStopMovieCaptureCommand( Cmd, Ar );
 	}
@@ -10652,7 +10657,7 @@ void UEngine::HandleScreenshotCaptured(int32 Width, int32 Height, const TArray<F
 
 		IImageWrapperModule& ImageWrapperModule = FModuleManager::LoadModuleChecked<IImageWrapperModule>( FName("ImageWrapper") );
 
-		switch(MatineeCaptureType.GetValue())
+		switch(MatineeScreenshotOptions.MatineeCaptureType.GetValue())
 		{
 		default:
 		case EMatineeCaptureType::BMP:
diff --git a/Engine/Source/Runtime/Launch/Private/LaunchEngineLoop.cpp b/Engine/Source/Runtime/Launch/Private/LaunchEngineLoop.cpp
--- a/Engine/Source/Runtime/Launch/Private/LaunchEngineLoop.cpp
+++ b/Engine/Source/Runtime/Launch/Private/LaunchEngineLoop.cpp
@@ -1914,7 +1914,7 @@ void FEngineLoop::InitTime()
 	FParse::Value(FCommandLine::Get(),TEXT("FPS="),FixedFPS);
 	if( FixedFPS > 0 )
 	{
-		GEngine->MatineeCaptureFPS = (int32)FixedFPS;
+		GEngine->MatineeScreenshotOptions.MatineeCaptureFPS = (int32)FixedFPS;
 		FApp::SetFixedDeltaTime(1 / FixedFPS);
 	}
 
