#ue4 - Fix Projectile orientation when bRotationFollowsVelocity is set but there is no velocity.

diff --git a/Engine/Source/Runtime/Engine/Private/ProjectileMovementComponent.cpp b/Engine/Source/Runtime/Engine/Private/ProjectileMovementComponent.cpp
--- a/Engine/Source/Runtime/Engine/Private/ProjectileMovementComponent.cpp
+++ b/Engine/Source/Runtime/Engine/Private/ProjectileMovementComponent.cpp
@@ -113,7 +113,7 @@ void UProjectileMovementComponent::TickComponent(float DeltaTime, enum ELevelTic
 	int32 NumBounces = 0;
 	int32 Iterations = 0;
 	FHitResult Hit(1.f);
-
+	
 	while( RemainingTime >= MIN_TICK_TIME && (Iterations < MaxSimulationIterations) && !ActorOwner->IsPendingKill() && UpdatedComponent )
 	{
 		Iterations++;
@@ -126,7 +126,7 @@ void UProjectileMovementComponent::TickComponent(float DeltaTime, enum ELevelTic
 		const FVector OldVelocity = Velocity;
 		FVector MoveDelta = ComputeMoveDelta(OldVelocity, TimeTick);
 
-		const FRotator NewRotation = bRotationFollowsVelocity ? OldVelocity.Rotation() : ActorOwner->GetActorRotation();
+		const FRotator NewRotation = (bRotationFollowsVelocity && !OldVelocity.IsNearlyZero()) ? OldVelocity.Rotation() : ActorOwner->GetActorRotation();
 
 		// Move the component
 		if (bShouldBounce)
