Fix initialization order

diff --git a/Engine/Source/Runtime/Engine/Private/PhysicsEngine/BodyInstance.cpp b/Engine/Source/Runtime/Engine/Private/PhysicsEngine/BodyInstance.cpp
--- a/Engine/Source/Runtime/Engine/Private/PhysicsEngine/BodyInstance.cpp
+++ b/Engine/Source/Runtime/Engine/Private/PhysicsEngine/BodyInstance.cpp
@@ -209,23 +209,28 @@ FBodyInstance::FBodyInstance()
 , bUseCCD(false)
 , bNotifyRigidBodyCollision(false)
 , bSimulatePhysics(false)
+, bOverrideMass(false)
+, MassInKg(100.f)
+, LinearDamping(0.01)
+, AngularDamping(0.0)
+, bEnableGravity(true)
 , bAutoWeld(false)
 , bWelded(false)
 , bStartAwake(true)
-, bEnableGravity(true)
 , bUpdateMassWhenScaleChanges(false)
-, bOverrideMass(false)
-, MassInKg(100.f)
-, DOFMode(0)
-, CustomDOFPlaneNormal(FVector::ZeroVector)
+, bLockTranslation(true)
+, bLockRotation(true)
 , bLockXTranslation(false)
 , bLockYTranslation(false)
 , bLockZTranslation(false)
 , bLockXRotation(false)
 , bLockYRotation(false)
 , bLockZRotation(false)
-, bLockTranslation(true)
-, bLockRotation(true)
+, DOFMode(0)
+, CustomDOFPlaneNormal(FVector::ZeroVector)
+, COMNudge(ForceInit)
+, MassScale(1.f)
+, MaxAngularVelocity(400.f)
 , DOFConstraint(NULL)
 , WeldParent(NULL)
 , bUseAsyncScene(false)
@@ -233,12 +238,7 @@ FBodyInstance::FBodyInstance()
 , bOverrideMaxDepenetrationVelocity(false)
 , MaxDepenetrationVelocity(0.f)
 , PhysMaterialOverride(NULL)
-, COMNudge(ForceInit)
 , SleepFamily(SF_Normal)
-, MassScale(1.f)
-, AngularDamping(0.0)
-, LinearDamping(0.01)
-, MaxAngularVelocity(400.f)
 , PhysicsBlendWeight(0.f)
 , PositionSolverIterationCount(8)
 , VelocitySolverIterationCount(1)
