Fixed Mac compile

diff --git a/Engine/Source/Runtime/Engine/Private/AVIWriter.cpp b/Engine/Source/Runtime/Engine/Private/AVIWriter.cpp
--- a/Engine/Source/Runtime/Engine/Private/AVIWriter.cpp
+++ b/Engine/Source/Runtime/Engine/Private/AVIWriter.cpp
@@ -494,20 +494,20 @@ public:
 				FString OutputDir = FPaths::VideoCaptureDir();
 				IFileManager::Get().MakeDirectory(*OutputDir, true);
 				
-				if (GEngine->bStartWithMatineeCapture)
+				if (GEngine->MatineeScreenshotOptions.bStartWithMatineeCapture)
 				{
 					bool bFoundUnusedName = false;
 					while (!bFoundUnusedName)
 					{
 						const FString FileNameInfo = FString::Printf( TEXT("_%dfps_%dx%d"),
-																	 GEngine->MatineeCaptureFPS, CaptureViewport->GetSizeXY().X, CaptureViewport->GetSizeXY().Y );
-						if (GEngine->bCompressMatineeCapture)
+																	 GEngine->MatineeScreenshotOptions.MatineeCaptureFPS, CaptureViewport->GetSizeXY().X, CaptureViewport->GetSizeXY().Y );
+						if (GEngine->MatineeScreenshotOptions.bCompressMatineeCapture)
 						{
-							FCString::Sprintf( File, TEXT("%s_%d.mov"), *(OutputDir + TEXT("/") + GEngine->MatineePackageCaptureName + "_" + GEngine->MatineeCaptureName + FileNameInfo + "_compressed"), MovieCaptureIndex );
+							FCString::Sprintf( File, TEXT("%s_%d.mov"), *(OutputDir + TEXT("/") + GEngine->MatineeScreenshotOptions.MatineePackageCaptureName + "_" + GEngine->MatineeScreenshotOptions.MatineeCaptureName + FileNameInfo + "_compressed"), MovieCaptureIndex );
 						}
 						else
 						{
-							FCString::Sprintf( File, TEXT("%s_%d.mov"), *(OutputDir + TEXT("/") + GEngine->MatineePackageCaptureName + "_" + GEngine->MatineeCaptureName + FileNameInfo), MovieCaptureIndex );
+							FCString::Sprintf( File, TEXT("%s_%d.mov"), *(OutputDir + TEXT("/") + GEngine->MatineeScreenshotOptions.MatineePackageCaptureName + "_" + GEngine->MatineeScreenshotOptions.MatineeCaptureName + FileNameInfo), MovieCaptureIndex );
 						}
 						if (IFileManager::Get().FileSize(File) != -1)
 						{
@@ -558,7 +558,7 @@ public:
 			}
 			
 			NSDictionary* VideoSettings = nil;
-			if (GEngine->bCompressMatineeCapture && GEngine->bStartWithMatineeCapture)
+			if (GEngine->MatineeScreenshotOptions.bCompressMatineeCapture && GEngine->MatineeScreenshotOptions.bStartWithMatineeCapture)
 			{
 				VideoSettings = [NSDictionary dictionaryWithObjectsAndKeys:
 								 AVVideoCodecH264, AVVideoCodecKey,
@@ -645,7 +645,7 @@ public:
 	void Update(float DeltaSeconds)
 	{
 		bool bShouldUpdate = true;
-		const float CaptureFrequency = 1.0f/(float)GEngine->MatineeCaptureFPS;
+		const float CaptureFrequency = 1.0f/(float)GEngine->MatineeScreenshotOptions.MatineeCaptureFPS;
 		if (CaptureFrequency > 0.f)
 		{
 			CurrentAccumSeconds += DeltaSeconds;
@@ -670,7 +670,7 @@ public:
 		
 		if (bMatineeFinished)
 		{
-			if (GEngine->MatineeCaptureType == 0)
+			if (GEngine->MatineeScreenshotOptions.MatineeCaptureType == 0)
 			{
 				if (CaptureViewport)
 				{
@@ -705,7 +705,7 @@ public:
 			FMemory::Memcpy(Data, ViewportColorBuffer.GetData(), ViewportColorBuffer.Num()*sizeof(FColor));
 			CVPixelBufferUnlockBaseAddress(PixeBuffer, 0);
 			
-            CMTime PresentTime = FrameNumber > 0 ? CMTimeMake(FrameNumber, GEngine->MatineeCaptureFPS) : kCMTimeZero;
+            CMTime PresentTime = FrameNumber > 0 ? CMTimeMake(FrameNumber, GEngine->MatineeScreenshotOptions.MatineeCaptureFPS) : kCMTimeZero;
 			BOOL OK = [AVFPixelBufferAdaptorRef appendPixelBuffer:PixeBuffer withPresentationTime:PresentTime];
 			check(OK);
 			
