Expose moment of inertia

diff --git a/Engine/Source/Runtime/Engine/Classes/Components/PrimitiveComponent.h b/Engine/Source/Runtime/Engine/Classes/Components/PrimitiveComponent.h
--- a/Engine/Source/Runtime/Engine/Classes/Components/PrimitiveComponent.h
+++ b/Engine/Source/Runtime/Engine/Classes/Components/PrimitiveComponent.h
@@ -1429,6 +1429,10 @@ public:
 	/** Returns the mass of this component in kg. */
 	UFUNCTION(BlueprintCallable, Category="Physics")
 	virtual float GetMass() const;
+
+	/** Returns the moment of inertia of this component in kg cm^2. */
+	UFUNCTION(BlueprintCallable, Category = "Physics", meta =(Keywords = "physics moment of inertia tensor MOI"))
+	virtual FVector GetMomentOfInertia(FName BoneName = NAME_None) const;
 	
 	/** Returns the calculated mass in kg. This is not 100% exactly the mass physx will calculate, but it is very close ( difference < 0.1kg ). */
 	virtual float CalculateMass(FName BoneName = NAME_None);
diff --git a/Engine/Source/Runtime/Engine/Classes/PhysicsEngine/BodyInstance.h b/Engine/Source/Runtime/Engine/Classes/PhysicsEngine/BodyInstance.h
--- a/Engine/Source/Runtime/Engine/Classes/PhysicsEngine/BodyInstance.h
+++ b/Engine/Source/Runtime/Engine/Classes/PhysicsEngine/BodyInstance.h
@@ -512,6 +512,8 @@ public:
 	float GetBodyMass() const;
 	/** Return bounds of physics representation */
 	FBox GetBodyBounds() const;
+	/** Return the body's moment of inertia */
+	FVector GetBodyMOI() const;
 
 
 	/** Set this body to be fixed (kinematic) or not. */
diff --git a/Engine/Source/Runtime/Engine/Private/PhysicsEngine/BodyInstance.cpp b/Engine/Source/Runtime/Engine/Private/PhysicsEngine/BodyInstance.cpp
--- a/Engine/Source/Runtime/Engine/Private/PhysicsEngine/BodyInstance.cpp
+++ b/Engine/Source/Runtime/Engine/Private/PhysicsEngine/BodyInstance.cpp
@@ -3058,6 +3058,29 @@ float FBodyInstance::GetBodyMass() const
 	return Retval;
 }
 
+
+FVector FBodyInstance::GetBodyMOI() const
+{
+	FVector Retval = FVector::ZeroVector;
+
+#if WITH_PHYSX
+	if (PxRigidBody* PRigidBody = GetPxRigidBody())
+	{
+		SCOPED_SCENE_READ_LOCK(PRigidBody->getScene());
+		Retval = P2UVector(PRigidBody->getMassSpaceInertiaTensor());
+	}
+#endif // WITH_PHYSX
+
+#if WITH_BOX2D
+	if (BodyInstancePtr != nullptr)
+	{
+		Retval = FVector(BodyInstancePtr->GetInertia(), 0.f, 0.f);
+	}
+#endif
+
+	return Retval;
+}
+
 FBox FBodyInstance::GetBodyBounds() const
 {
 	FBox Bounds;
diff --git a/Engine/Source/Runtime/Engine/Private/PrimitiveComponentPhysics.cpp b/Engine/Source/Runtime/Engine/Private/PrimitiveComponentPhysics.cpp
--- a/Engine/Source/Runtime/Engine/Private/PrimitiveComponentPhysics.cpp
+++ b/Engine/Source/Runtime/Engine/Private/PrimitiveComponentPhysics.cpp
@@ -453,6 +453,16 @@ float UPrimitiveComponent::GetMass() const
 	return 0.0f;
 }
 
+FVector UPrimitiveComponent::GetMomentOfInertia(FName BoneName /* = NAME_None */) const 
+{
+	if(FBodyInstance* BI = GetBodyInstance(BoneName))
+	{
+		return BI->GetBodyMOI();
+	}
+
+	return FVector::ZeroVector;
+}
+
 float UPrimitiveComponent::CalculateMass(FName)
 {
 	if (BodyInstance.bOverrideMass)
