Added helper function for world space scaling by moment of inertia. Rename moment of inertia to inertia tensor to better indicate that this is the diagonal.

diff --git a/Engine/Source/Runtime/Engine/Classes/Components/PrimitiveComponent.h b/Engine/Source/Runtime/Engine/Classes/Components/PrimitiveComponent.h
--- a/Engine/Source/Runtime/Engine/Classes/Components/PrimitiveComponent.h
+++ b/Engine/Source/Runtime/Engine/Classes/Components/PrimitiveComponent.h
@@ -1430,10 +1430,14 @@ public:
 	UFUNCTION(BlueprintCallable, Category="Physics")
 	virtual float GetMass() const;
 
-	/** Returns the moment of inertia of this component in kg cm^2. */
+	/** Returns the inertia tensor of this component in kg cm^2. The inertia tensor is in local component space.*/
 	UFUNCTION(BlueprintCallable, Category = "Physics", meta =(Keywords = "physics moment of inertia tensor MOI"))
-	virtual FVector GetMomentOfInertia(FName BoneName = NAME_None) const;
-	
+	virtual FVector GetInertiaTensor(FName BoneName = NAME_None) const;
+
+	/** Scales the given vector by the world space moment of inertia. Useful for computing the torque needed to rotate an object.*/
+	UFUNCTION(BlueprintCallable, Category = "Physics", meta = (Keywords = "physics moment of inertia tensor MOI"))
+	virtual FVector ScaleByMomentOfInertia(FVector InputVector, FName BoneName = NAME_None) const;
+
 	/** Returns the calculated mass in kg. This is not 100% exactly the mass physx will calculate, but it is very close ( difference < 0.1kg ). */
 	virtual float CalculateMass(FName BoneName = NAME_None);
 
diff --git a/Engine/Source/Runtime/Engine/Classes/PhysicsEngine/BodyInstance.h b/Engine/Source/Runtime/Engine/Classes/PhysicsEngine/BodyInstance.h
--- a/Engine/Source/Runtime/Engine/Classes/PhysicsEngine/BodyInstance.h
+++ b/Engine/Source/Runtime/Engine/Classes/PhysicsEngine/BodyInstance.h
@@ -512,8 +512,8 @@ public:
 	float GetBodyMass() const;
 	/** Return bounds of physics representation */
 	FBox GetBodyBounds() const;
-	/** Return the body's moment of inertia */
-	FVector GetBodyMOI() const;
+	/** Return the body's inertia tensor. This is returned in local mass space */
+	FVector GetBodyInertiaTensor() const;
 
 
 	/** Set this body to be fixed (kinematic) or not. */
diff --git a/Engine/Source/Runtime/Engine/Private/PhysicsEngine/BodyInstance.cpp b/Engine/Source/Runtime/Engine/Private/PhysicsEngine/BodyInstance.cpp
--- a/Engine/Source/Runtime/Engine/Private/PhysicsEngine/BodyInstance.cpp
+++ b/Engine/Source/Runtime/Engine/Private/PhysicsEngine/BodyInstance.cpp
@@ -3059,7 +3059,7 @@ float FBodyInstance::GetBodyMass() const
 }
 
 
-FVector FBodyInstance::GetBodyMOI() const
+FVector FBodyInstance::GetBodyInertiaTensor() const
 {
 	FVector Retval = FVector::ZeroVector;
 
diff --git a/Engine/Source/Runtime/Engine/Private/PrimitiveComponentPhysics.cpp b/Engine/Source/Runtime/Engine/Private/PrimitiveComponentPhysics.cpp
--- a/Engine/Source/Runtime/Engine/Private/PrimitiveComponentPhysics.cpp
+++ b/Engine/Source/Runtime/Engine/Private/PrimitiveComponentPhysics.cpp
@@ -453,16 +453,25 @@ float UPrimitiveComponent::GetMass() const
 	return 0.0f;
 }
 
-FVector UPrimitiveComponent::GetMomentOfInertia(FName BoneName /* = NAME_None */) const 
+FVector UPrimitiveComponent::GetInertiaTensor(FName BoneName /* = NAME_None */) const 
 {
 	if(FBodyInstance* BI = GetBodyInstance(BoneName))
 	{
-		return BI->GetBodyMOI();
+		return BI->GetBodyInertiaTensor();
 	}
 
 	return FVector::ZeroVector;
 }
 
+FVector UPrimitiveComponent::ScaleByMomentOfInertia(FVector InputVector, FName BoneName /* = NAME_None */) const
+{
+	const FVector LocalInertiaTensor = GetInertiaTensor(BoneName);
+	const FVector InputVectorLocal = ComponentToWorld.InverseTransformVectorNoScale(InputVector);
+	const FVector LocalScaled = InputVectorLocal * LocalInertiaTensor;
+	const FVector WorldScaled = ComponentToWorld.TransformFVector4NoScale(LocalScaled);
+	return WorldScaled;
+}
+
 float UPrimitiveComponent::CalculateMass(FName)
 {
 	if (BodyInstance.bOverrideMass)
