[INTEGRATE] Change 2439603 by bruce.nesbit@bnesbit_Releases on 2015/02/10 05:28:20

	Minor code cleanup in ActorComponent::PostLoad

diff --git a/Engine/Source/Runtime/Engine/Private/ActorComponent.cpp b/Engine/Source/Runtime/Engine/Private/ActorComponent.cpp
--- a/Engine/Source/Runtime/Engine/Private/ActorComponent.cpp
+++ b/Engine/Source/Runtime/Engine/Private/ActorComponent.cpp
@@ -157,11 +157,6 @@ void UActorComponent::PostLoad()
 			}
 		}
 	}
-// 	if (!HasAllFlags(RF_Public) && GetOuter()->IsA<UBlueprintGeneratedClass>())
-// 	{
-// 		SetFlags(RF_Public);
-// 		ULinkerLoad::RefreshExportFlags(this);
-// 	}
 }
 
 void UActorComponent::PostRename(UObject* OldOuter, const FName OldName)
