[INTEGRATE] Change 2440053 by [REDACTED_PERSON][REDACTED_HANDLE] on 2015/02/10 13:35:34

	Remove PhysicsMobility workaround in favor of a "safer" solution to preserving Mobility for physics initialization (and other events during registration of UCS-added components).

	change summary:
	- restricted the temp mobility swap to happen during UCS only
	- modified AActor::AddComponent() to defer the temp mobility swap that we do during UCS until after new scene components are registered
	- modified USceneComponent::AttachTo() to skip the mobility check when the owner Actor instance is running the UCS
	- removed previous PhysicsMobility workaround in UPrimitiveComponent (i.e. it's no longer necessary to shadow the Mobility setting)
	- going to table privatization of the Mobility setting until after v4.7 release (see JIRA comments for more details)

diff --git a/Engine/Source/Runtime/Engine/Classes/Components/PrimitiveComponent.h b/Engine/Source/Runtime/Engine/Classes/Components/PrimitiveComponent.h
--- a/Engine/Source/Runtime/Engine/Classes/Components/PrimitiveComponent.h
+++ b/Engine/Source/Runtime/Engine/Classes/Components/PrimitiveComponent.h
@@ -400,12 +400,6 @@ public:
 	/** Cached navigation relevancy flag for collision updates */
 	uint32 bNavigationRelevant : 1;
 
-	/** 
-	 * Shadowed mobility setting used e.g. to assist with physics scene initialization at runtime.
-	 * It will usually match the value of the 'Mobility' property, except e.g. during UCS execution.
-	 */
-	TEnumAsByte<EComponentMobility::Type> PhysicsMobility;
-
 protected:
 
 	virtual void UpdateNavigationData() override;
@@ -960,7 +954,6 @@ public:
 #endif
 
 	// Begin UActorComponent Interface
-	virtual void OnComponentCreated() override;
 	virtual void InvalidateLightingCacheDetailed(bool bInvalidateBuildEnqueuedLighting, bool bTranslationOnly) override;
 	virtual bool IsEditorOnly() const override;
 	virtual bool ShouldCreatePhysicsState() const override;
@@ -1314,7 +1307,6 @@ public:
 	virtual ECollisionChannel GetCollisionObjectType() const override;
 	virtual const FCollisionResponseContainer& GetCollisionResponseToChannels() const override;
 	virtual FVector GetComponentVelocity() const override;
-	virtual void SetMobility(EComponentMobility::Type NewMobility) override;
 	//End USceneComponent Interface
 
 	/**
diff --git a/Engine/Source/Runtime/Engine/Private/ActorConstruction.cpp b/Engine/Source/Runtime/Engine/Private/ActorConstruction.cpp
--- a/Engine/Source/Runtime/Engine/Private/ActorConstruction.cpp
+++ b/Engine/Source/Runtime/Engine/Private/ActorConstruction.cpp
@@ -607,12 +607,6 @@ UActorComponent* AActor::AddComponent(FName TemplateName, bool bManualAttachment
 		// Call function to notify component it has been created
 		NewActorComp->OnComponentCreated();
 
-		// Keep track of the new component during UCS execution. This also does a temporary mobility swap during UCS execution in order to allow attachment and other things to succeed within that context.
-		if(bRunningUserConstructionScript)
-		{
-		FUCSComponentManager::Get().AddComponent(this, NewActorComp);
-		}
-
 		// The user has the option of doing attachment manually where they have complete control or via the automatic rule
 		// that the first component added becomes the root component, with subsequent components attached to the root.
 		USceneComponent* NewSceneComp = Cast<USceneComponent>(NewActorComp);
@@ -635,6 +629,13 @@ UActorComponent* AActor::AddComponent(FName TemplateName, bool bManualAttachment
 
 		// Register component, which will create physics/rendering state, now component is in correct position
 		NewActorComp->RegisterComponent();
+
+		// Keep track of the new component during UCS execution. This also does a temporary mobility swap during UCS execution in order to allow dynamic data to be changed within that context.
+		// Note: This should only be done AFTER registration.
+		if(bRunningUserConstructionScript)
+		{
+			FUCSComponentManager::Get().AddComponent(this, NewActorComp);
+		}
 	}
 
 	return NewActorComp;
diff --git a/Engine/Source/Runtime/Engine/Private/PhysicsEngine/BodyInstance.cpp b/Engine/Source/Runtime/Engine/Private/PhysicsEngine/BodyInstance.cpp
--- a/Engine/Source/Runtime/Engine/Private/PhysicsEngine/BodyInstance.cpp
+++ b/Engine/Source/Runtime/Engine/Private/PhysicsEngine/BodyInstance.cpp
@@ -1038,8 +1038,8 @@ void FBodyInstance::InitBody_Legacy(UBodySetup* Setup, const FTransform& Transfo
 		}
 	}
 
-	// See if we are 'static' - note that we use the PhysicsMobility setting here which "shadows" the actual Mobility setting in most cases but also will remain constant e.g. during UCS execution
-	const bool bPhysicsStatic = (OwnerComponent == NULL) || (OwnerComponent->PhysicsMobility != EComponentMobility::Movable);
+	// See if we are 'static'
+	const bool bPhysicsStatic = (OwnerComponent == NULL) || (OwnerComponent->Mobility != EComponentMobility::Movable);
 
 	// In skeletal case, we need both our bone and skelcomponent flag to be true.
 	// This might be 'and'ing us with ourself, but thats fine.
diff --git a/Engine/Source/Runtime/Engine/Private/PrimitiveComponent.cpp b/Engine/Source/Runtime/Engine/Private/PrimitiveComponent.cpp
--- a/Engine/Source/Runtime/Engine/Private/PrimitiveComponent.cpp
+++ b/Engine/Source/Runtime/Engine/Private/PrimitiveComponent.cpp
@@ -114,8 +114,6 @@ UPrimitiveComponent::UPrimitiveComponent(const FObjectInitializer& ObjectInitial
 
 	bCachedAllCollideableDescendantsRelative = false;
 	LastCheckedAllCollideableDescendantsTime = 0.f;
-
-	PhysicsMobility = Mobility;
 }
 
 bool UPrimitiveComponent::UsesOnlyUnlitMaterials() const
@@ -139,14 +137,6 @@ void UPrimitiveComponent::GetLightAndShadowMapMemoryUsage( int32& LightMapMemory
 	return;
 }
 
-void UPrimitiveComponent::OnComponentCreated()
-{
-	Super::OnComponentCreated();
-
-	// Shadow the current mobility setting for physics scene initialization (since e.g. during UCS execution we temporarily override the base mobility setting)
-	PhysicsMobility = Mobility;
-}
-
 void UPrimitiveComponent::InvalidateLightingCacheDetailed(bool bInvalidateBuildEnqueuedLighting, bool bTranslationOnly) 
 {
 	bHasCachedStaticLighting = false;
@@ -948,14 +938,6 @@ bool UPrimitiveComponent::IsWorldGeometry() const
 	return Mobility != EComponentMobility::Movable && GetCollisionObjectType()==ECC_WorldStatic;
 }
 
-void UPrimitiveComponent::SetMobility(EComponentMobility::Type NewMobility)
-{
-	Super::SetMobility(NewMobility);
-
-	// Update the shadowed mobility setting for physics
-	PhysicsMobility = Mobility;
-}
-
 ECollisionChannel UPrimitiveComponent::GetCollisionObjectType() const
 {
 	return ECollisionChannel(BodyInstance.GetObjectType());
diff --git a/Engine/Source/Runtime/Engine/Private/SceneComponent.cpp b/Engine/Source/Runtime/Engine/Private/SceneComponent.cpp
--- a/Engine/Source/Runtime/Engine/Private/SceneComponent.cpp
+++ b/Engine/Source/Runtime/Engine/Private/SceneComponent.cpp
@@ -889,7 +889,8 @@ void USceneComponent::AttachTo(class USceneComponent* Parent, FName InSocketName
 			return;
 		}
 
-		if(Mobility == EComponentMobility::Static && Parent->Mobility != EComponentMobility::Static )
+		// Don't allow components with static mobility to be attached to non-static parents (except during UCS)
+		if(!IsRunningUserConstructionScript() && Mobility == EComponentMobility::Static && Parent->Mobility != EComponentMobility::Static)
 		{
 			FString ExtraBlueprintInfo;
 #if WITH_EDITORONLY_DATA
