[INTEGRATE] Change 2440263 by [REDACTED_PERSON][REDACTED_HANDLE] on 2015/02/10 15:11:31

	Properties of BP created non-Scene components save across reconstruction
	Primitive components will mark render state dirty if registered after instance data applied
	Instanced static mesh will save properties and attachment across reconstruction even if there is no static lighting or selected instances

diff --git a/Engine/Source/Runtime/Engine/Classes/Components/ActorComponent.h b/Engine/Source/Runtime/Engine/Classes/Components/ActorComponent.h
--- a/Engine/Source/Runtime/Engine/Classes/Components/ActorComponent.h
+++ b/Engine/Source/Runtime/Engine/Classes/Components/ActorComponent.h
@@ -474,10 +474,10 @@ public:
 	virtual void PostNetReceive() { }
 
 	/** Called before we throw away components during RerunConstructionScripts, to cache any data we wish to persist across that operation */
-	virtual class FComponentInstanceDataBase* GetComponentInstanceData() const { return NULL; }
+	virtual class FComponentInstanceDataBase* GetComponentInstanceData() const;
 
 	/** The type of the component instance data that this component is interested in */
-	virtual FName GetComponentInstanceDataType() const { return NAME_None; }
+	virtual FName GetComponentInstanceDataType() const;
 
 	// Begin UObject interface.
 	virtual void BeginDestroy() override;
diff --git a/Engine/Source/Runtime/Engine/Classes/Components/PrimitiveComponent.h b/Engine/Source/Runtime/Engine/Classes/Components/PrimitiveComponent.h
--- a/Engine/Source/Runtime/Engine/Classes/Components/PrimitiveComponent.h
+++ b/Engine/Source/Runtime/Engine/Classes/Components/PrimitiveComponent.h
@@ -958,6 +958,8 @@ public:
 	virtual bool IsEditorOnly() const override;
 	virtual bool ShouldCreatePhysicsState() const override;
 	virtual bool HasValidPhysicsState() const override;
+	virtual class FComponentInstanceDataBase* GetComponentInstanceData() const override;
+	virtual FName GetComponentInstanceDataType() const override;
 	// End UActorComponent Interface
 
 	/** @return true if the owner is selected and this component is selectable */
@@ -1648,3 +1650,20 @@ public:
 	void DispatchOnInputTouchBegin(const ETouchIndex::Type Key);
 	void DispatchOnInputTouchEnd(const ETouchIndex::Type Key);
 };
+
+/** 
+ *  Component instance cached data base class for primitive components. 
+ *  Stores a list of instance components attached to the 
+ */
+class ENGINE_API FPrimitiveComponentInstanceData : public FSceneComponentInstanceData
+{
+public:
+	FPrimitiveComponentInstanceData(const UPrimitiveComponent* SourceComponent);
+			
+	virtual ~FPrimitiveComponentInstanceData()
+	{}
+
+	virtual void ApplyToComponent(UActorComponent* Component) override;
+
+	bool ContainsData() const;
+};
diff --git a/Engine/Source/Runtime/Engine/Private/ActorComponent.cpp b/Engine/Source/Runtime/Engine/Private/ActorComponent.cpp
--- a/Engine/Source/Runtime/Engine/Private/ActorComponent.cpp
+++ b/Engine/Source/Runtime/Engine/Private/ActorComponent.cpp
@@ -534,6 +534,25 @@ void UActorComponent::UninitializeComponent()
 	bHasBeenInitialized = false;
 }
 
+FComponentInstanceDataBase* UActorComponent::GetComponentInstanceData() const
+{
+	FComponentInstanceDataBase* InstanceData = new FComponentInstanceDataBase(this);
+
+	if (!InstanceData->ContainsSavedProperties())
+	{
+		delete InstanceData;
+		InstanceData = nullptr;
+	}
+
+	return InstanceData;
+}
+
+FName UActorComponent::GetComponentInstanceDataType() const
+{
+	static const FName ActorComponentInstanceDataTypeName(TEXT("ActorComponentInstanceData"));
+	return ActorComponentInstanceDataTypeName;
+}
+
 void FActorComponentTickFunction::ExecuteTick(float DeltaTime, enum ELevelTick TickType, ENamedThreads::Type CurrentThread, const FGraphEventRef& MyCompletionGraphEvent)
 {
 	if (Target && !Target->HasAnyFlags(RF_PendingKill | RF_Unreachable))
diff --git a/Engine/Source/Runtime/Engine/Private/InstancedStaticMesh.cpp b/Engine/Source/Runtime/Engine/Private/InstancedStaticMesh.cpp
--- a/Engine/Source/Runtime/Engine/Private/InstancedStaticMesh.cpp
+++ b/Engine/Source/Runtime/Engine/Private/InstancedStaticMesh.cpp
@@ -803,6 +803,12 @@ FComponentInstanceDataBase* UInstancedStaticMeshComponent::GetComponentInstanceD
 		StaticMeshInstanceData->SelectedInstances = SelectedInstances;
 	}
 
+	if (InstanceData == nullptr)
+	{
+		// Skip up the hierarchy
+		InstanceData = UPrimitiveComponent::GetComponentInstanceData();
+	}
+
 	return InstanceData;
 #else
 	return nullptr;
diff --git a/Engine/Source/Runtime/Engine/Private/PrimitiveComponent.cpp b/Engine/Source/Runtime/Engine/Private/PrimitiveComponent.cpp
--- a/Engine/Source/Runtime/Engine/Private/PrimitiveComponent.cpp
+++ b/Engine/Source/Runtime/Engine/Private/PrimitiveComponent.cpp
@@ -327,6 +327,45 @@ void UPrimitiveComponent::OnUnregister()
 	}
 }
 
+FPrimitiveComponentInstanceData::FPrimitiveComponentInstanceData(const UPrimitiveComponent* SourceComponent)
+	: FSceneComponentInstanceData(SourceComponent)
+{
+}
+
+void FPrimitiveComponentInstanceData::ApplyToComponent(UActorComponent* Component)
+{
+	FSceneComponentInstanceData::ApplyToComponent(Component);
+
+	if (ContainsSavedProperties() && Component->IsRegistered())
+	{
+		Component->MarkRenderStateDirty();
+	}
+}
+
+bool FPrimitiveComponentInstanceData::ContainsData() const
+{
+	return (ContainsSavedProperties() || AttachedInstanceComponents.Num() > 0);
+}
+
+FComponentInstanceDataBase* UPrimitiveComponent::GetComponentInstanceData() const
+{
+	FPrimitiveComponentInstanceData* InstanceData = new FPrimitiveComponentInstanceData(this);
+
+	if (!InstanceData->ContainsData())
+	{
+		delete InstanceData;
+		InstanceData = nullptr;
+	}
+
+	return InstanceData;
+}
+
+FName UPrimitiveComponent::GetComponentInstanceDataType() const
+{
+	static const FName PrimitiveComponentInstanceDataTypeName(TEXT("PrimitiveComponentInstanceData"));
+	return PrimitiveComponentInstanceDataTypeName;
+}
+
 void UPrimitiveComponent::OnAttachmentChanged()
 {
 	if (World && World->Scene)
diff --git a/Engine/Source/Runtime/Engine/Private/SceneComponent.cpp b/Engine/Source/Runtime/Engine/Private/SceneComponent.cpp
--- a/Engine/Source/Runtime/Engine/Private/SceneComponent.cpp
+++ b/Engine/Source/Runtime/Engine/Private/SceneComponent.cpp
@@ -1213,6 +1213,11 @@ FComponentInstanceDataBase* USceneComponent::GetComponentInstanceData() const
 		}
 	}
 
+	if (InstanceData == nullptr)
+	{
+		InstanceData = Super::GetComponentInstanceData();
+	}
+
 	return InstanceData;
 }
 
diff --git a/Engine/Source/Runtime/Engine/Public/ComponentInstanceDataCache.h b/Engine/Source/Runtime/Engine/Public/ComponentInstanceDataCache.h
--- a/Engine/Source/Runtime/Engine/Public/ComponentInstanceDataCache.h
+++ b/Engine/Source/Runtime/Engine/Public/ComponentInstanceDataCache.h
@@ -28,6 +28,8 @@ public:
 	/** Replaces any references to old instances during Actor reinstancing */
 	virtual void FindAndReplaceInstances(const TMap<UObject*, UObject*>& OldToNewInstanceMap) { };
 
+	bool ContainsSavedProperties() const { return SavedProperties.Num() > 0; }
+
 protected:
 	/** The name of the source component */
 	FName SourceComponentName;
