[INTEGRATE] Change 2440361 by [REDACTED_PERSON][REDACTED_HANDLE] on 2015/02/10 16:04:45

	UE-9486 - We no longer ensure when the class is not found for a custom cursor class.  We log to the error list for PIE that the class could not be found.

diff --git a/Engine/Source/Runtime/Engine/Private/GameViewportClient.cpp b/Engine/Source/Runtime/Engine/Private/GameViewportClient.cpp
--- a/Engine/Source/Runtime/Engine/Private/GameViewportClient.cpp
+++ b/Engine/Source/Runtime/Engine/Private/GameViewportClient.cpp
@@ -32,6 +32,8 @@
 #include "GameFramework/GameUserSettings.h"
 #include "Runtime/Engine/Classes/Engine/UserInterfaceSettings.h"
 
+#define LOCTEXT_NAMESPACE "GameViewport"
+
 /** This variable allows forcing full screen of the first player controller viewport, even if there are multiple controllers plugged in and no cinematic playing. */
 bool GForceFullscreen = false;
 
@@ -553,7 +555,7 @@ void UGameViewportClient::AddCursor(EMouseCursor::Type Cursor, const FStringClas
 	if ( CursorClass.IsValid() )
 	{
 		UClass* Class = CursorClass.TryLoadClass<UUserWidget>();
-		if ( ensure(Class) )
+		if ( Class )
 		{
 			UUserWidget* UserWidget = CreateWidget<UUserWidget>(GetGameInstance(), Class);
 			if ( ensure(UserWidget) )
@@ -561,6 +563,10 @@ void UGameViewportClient::AddCursor(EMouseCursor::Type Cursor, const FStringClas
 				CursorWidgets.Add(Cursor, UserWidget->TakeWidget());
 			}
 		}
+		else
+		{
+			FMessageLog("PIE").Error(FText::Format(LOCTEXT("CursorClassNotFoundFormat", "The cursor class '{0}' was not found, check your custom cursor settings."), FText::FromString(CursorClass.ToString())));
+		}
 	}
 }
 
@@ -3019,3 +3025,4 @@ void UGameViewportClient::HandleViewportStatDisableAll(const bool bInAnyViewport
 	}
 }
 
+#undef LOCTEXT_NAMESPACE
