[INTEGRATE] Change 2440411 by [REDACTED_PERSON][REDACTED_HANDLE] on 2015/02/10 16:34:10

	[UE-9332] - Editing BP components in world no longer causes the selection to be reset

diff --git a/Engine/Source/Editor/UnrealEd/Private/Editor.cpp b/Engine/Source/Editor/UnrealEd/Private/Editor.cpp
--- a/Engine/Source/Editor/UnrealEd/Private/Editor.cpp
+++ b/Engine/Source/Editor/UnrealEd/Private/Editor.cpp
@@ -6040,13 +6040,32 @@ void UEditorEngine::DoConvertActors( const TArray<AActor*>& ActorsToConvert, UCl
 void UEditorEngine::NotifyToolsOfObjectReplacement(const TMap<UObject*, UObject*>& OldToNewInstanceMap)
 {
 	// This can be called early on during startup if blueprints need to be compiled.  
-	// If the property module isn't loaded then there arent any property windows to update
+	// If the property module isn't loaded then there aren't any property windows to update
 	if( FModuleManager::Get().IsModuleLoaded( "PropertyEditor" ) )
 	{
 		FPropertyEditorModule& PropertyEditorModule = FModuleManager::GetModuleChecked<FPropertyEditorModule>( "PropertyEditor" );
 		PropertyEditorModule.ReplaceViewedObjects( OldToNewInstanceMap );
 	}
 
+	// Check to see if any selected components were reinstanced
+	USelection* ComponentSelection = GetSelectedComponents();
+	TArray<TWeakObjectPtr<UObject> > SelectedComponents;
+	ComponentSelection->GetSelectedObjects(SelectedComponents);
+
+	ComponentSelection->BeginBatchSelectOperation();
+	for (int32 i = 0; i < SelectedComponents.Num(); ++i)
+	{
+		UObject* Component = SelectedComponents[i].GetEvenIfUnreachable();
+
+		// If the component corresponds to a new instance in the map, update the selection accordingly
+		if (OldToNewInstanceMap.Contains(Component))
+		{
+			ComponentSelection->Deselect(Component);
+			SelectComponent(CastChecked<UActorComponent>(OldToNewInstanceMap[Component]), true, false);
+		}
+	}
+	ComponentSelection->EndBatchSelectOperation();
+
 	BroadcastObjectsReplaced(OldToNewInstanceMap);
 }
 
diff --git a/Engine/Source/Editor/UnrealEd/Private/EditorSelectUtils.cpp b/Engine/Source/Editor/UnrealEd/Private/EditorSelectUtils.cpp
--- a/Engine/Source/Editor/UnrealEd/Private/EditorSelectUtils.cpp
+++ b/Engine/Source/Editor/UnrealEd/Private/EditorSelectUtils.cpp
@@ -367,7 +367,7 @@ void UUnrealEdEngine::NoteSelectionChange()
 		ActiveModes[ModeIndex]->ActorSelectionChangeNotify();
 	}
 
-	bool bComponentSelectionChanged = GetSelectedComponentCount() > 0;
+	const bool bComponentSelectionChanged = GetSelectedComponentCount() > 0;
 	USelection* Selection = bComponentSelectionChanged ? GetSelectedComponents() : GetSelectedActors();
 	USelection::SelectionChangedEvent.Broadcast(Selection);
 	
diff --git a/Engine/Source/Runtime/Engine/Private/ActorConstruction.cpp b/Engine/Source/Runtime/Engine/Private/ActorConstruction.cpp
--- a/Engine/Source/Runtime/Engine/Private/ActorConstruction.cpp
+++ b/Engine/Source/Runtime/Engine/Private/ActorConstruction.cpp
@@ -7,6 +7,7 @@
 #include "Engine/LevelScriptActor.h"
 
 #if WITH_EDITOR
+#include "Editor.h"
 #include "Editor/UnrealEd/Public/Kismet2/BlueprintEditorUtils.h"
 #include "Editor/UnrealEd/Public/Kismet2/KismetEditorUtilities.h"
 #endif
@@ -369,6 +370,37 @@ void AActor::RerunConstructionScripts()
 			OldTransform.SetTranslation(RootComponent->GetComponentLocation()); // take into account any custom location
 		}
 
+#if WITH_EDITOR
+		// Save the current construction script-created components by name
+		TMap<const FName, UObject*> DestroyedComponentsByName;
+		TInlineComponentArray<UActorComponent*> PreviouslyAttachedComponents;
+		GetComponents(PreviouslyAttachedComponents);
+		for (auto Component : PreviouslyAttachedComponents)
+		{
+			if (Component)
+			{
+				if (Component->IsCreatedByConstructionScript())
+				{
+
+					DestroyedComponentsByName.Add(Component->GetFName(), Component);
+				}
+				else
+				{
+					UActorComponent* OuterComponent = Component->GetTypedOuter<UActorComponent>();
+					while (OuterComponent)
+					{
+						if (OuterComponent->IsCreatedByConstructionScript())
+						{
+							DestroyedComponentsByName.Add(Component->GetFName(), Component);
+							break;
+						}
+						OuterComponent = OuterComponent->GetTypedOuter<UActorComponent>();
+					}
+				}
+			}
+		}
+#endif
+
 		// Destroy existing components
 		DestroyConstructedComponents();
 
@@ -421,6 +453,25 @@ void AActor::RerunConstructionScripts()
 		GUndo = CurrentTransaction;
 
 #if WITH_EDITOR
+		// Create the mapping of old->new components and notify the editor of the replacements
+		TMap<UObject*, UObject*> OldToNewComponentMapping;
+
+		TInlineComponentArray<UActorComponent*> NewComponents;
+		GetComponents(NewComponents);
+		for (auto NewComp : NewComponents)
+		{
+			const FName NewCompName = NewComp->GetFName();
+			if (DestroyedComponentsByName.Contains(NewCompName))
+			{
+				OldToNewComponentMapping.Add(DestroyedComponentsByName[NewCompName], NewComp);
+			}
+		}
+
+		if (OldToNewComponentMapping.Num() > 0)
+		{
+			GEditor->NotifyToolsOfObjectReplacement(OldToNewComponentMapping);
+		}
+
 		if (ActorTransactionAnnotation)
 		{
 			CurrentTransactionAnnotation = NULL;
