[INTEGRATE] Change 2440694 by [REDACTED_PERSON][REDACTED_HANDLE] on 2015/02/10 18:52:55

	Fixed crash when copying/pasting components
	- Copied/pasted components have the transient package as an outer, so I removed an assert and added an if check

diff --git a/Engine/Source/Runtime/Engine/Private/ActorComponent.cpp b/Engine/Source/Runtime/Engine/Private/ActorComponent.cpp
--- a/Engine/Source/Runtime/Engine/Private/ActorComponent.cpp
+++ b/Engine/Source/Runtime/Engine/Private/ActorComponent.cpp
@@ -139,7 +139,7 @@ void UActorComponent::PostLoad()
 
 	if (CreationMethod == EComponentCreationMethod::SimpleConstructionScript)
 	{
-		UBlueprintGeneratedClass* Class = CastChecked<UBlueprintGeneratedClass>(GetOuter()->GetClass());
+		UBlueprintGeneratedClass* Class = Cast/*Checked*/<UBlueprintGeneratedClass>(GetOuter()->GetClass());
 		while (Class)
 		{
 			USimpleConstructionScript* SCS = Class->SimpleConstructionScript;
@@ -405,7 +405,7 @@ bool UActorComponent::Modify( bool bAlwaysMarkDirty/*=true*/ )
 {
 	// If this is a construction script component we don't store them in the transaction buffer.  Instead, mark
 	// the Actor as modified so that we store of the transaction annotation that has the component properties stashed
-	if (IsCreatedByConstructionScript())
+	if (IsCreatedByConstructionScript() && GetOwner())
 	{
 		return GetOwner()->Modify(bAlwaysMarkDirty);
 	}
