[INTEGRATE] Change 2440751 by [REDACTED_PERSON][REDACTED_HANDLE] on 2015/02/10 19:44:32

	#ue4 - Fix BSP collision issues with initial overlaps (distance = 0 or less) when MTD fails.

	- The previous fallback of FindOverlappedTriangleNormal() could return no triangles when exactly against the point of contact, so if that test fails then use a slightly inflated query geometry.

	UE-9547

diff --git a/Engine/Source/Runtime/Engine/Private/Collision/CollisionConversions.cpp b/Engine/Source/Runtime/Engine/Private/Collision/CollisionConversions.cpp
--- a/Engine/Source/Runtime/Engine/Private/Collision/CollisionConversions.cpp
+++ b/Engine/Source/Runtime/Engine/Private/Collision/CollisionConversions.cpp
@@ -638,36 +638,100 @@ static bool ComputeInflatedMTD(const float MtdInflation, const PxLocationHit& PH
 }
 
 
-static bool FindOverlappedTriangleNormal(const PxGeometry& Geom, const PxTransform& QueryTM, const PxShape* PShape, const PxTransform& PShapeWorldPose, FVector& OutNormal, bool bCanDrawOverlaps = false)
+
+static bool CanFindOverlappedTriangle(const PxShape* PShape)
 {
-	PxTriangleMeshGeometry PTriMeshGeom;
-	PxHeightFieldGeometry PHeightfieldGeom;
+	return (PShape && (PShape->getGeometryType() == PxGeometryType::eTRIANGLEMESH || PShape->getGeometryType() == PxGeometryType::eHEIGHTFIELD));
+}
 
-	if (PShape->getTriangleMeshGeometry(PTriMeshGeom) || PShape->getHeightFieldGeometry(PHeightfieldGeom))
+
+static bool FindOverlappedTriangleNormal_Internal(const PxGeometry& Geom, const PxTransform& QueryTM, const PxShape* PShape, const PxTransform& PShapeWorldPose, FVector& OutNormal, bool bCanDrawOverlaps = false)
+{
+	if (CanFindOverlappedTriangle(PShape))
 	{
-		PxGeometryType::Enum GeometryType = PShape->getGeometryType();
-		const bool bIsTriMesh = (GeometryType == PxGeometryType::eTRIANGLEMESH);
-		PxU32 HitTris[64];
-		bool bOverflow = false;
+		PxTriangleMeshGeometry PTriMeshGeom;
+		PxHeightFieldGeometry PHeightfieldGeom;
 
-		const int32 NumTrisHit = bIsTriMesh ? PxMeshQuery::findOverlapTriangleMesh(Geom, QueryTM, PTriMeshGeom, PShapeWorldPose, HitTris, ARRAY_COUNT(HitTris), 0, bOverflow) :
-											  PxMeshQuery::findOverlapHeightField(Geom, QueryTM, PHeightfieldGeom, PShapeWorldPose, HitTris, ARRAY_COUNT(HitTris), 0, bOverflow);
-		if (NumTrisHit > 0)
+		if (PShape->getTriangleMeshGeometry(PTriMeshGeom) || PShape->getHeightFieldGeometry(PHeightfieldGeom))
 		{
-			if (bIsTriMesh)
+			PxGeometryType::Enum GeometryType = PShape->getGeometryType();
+			const bool bIsTriMesh = (GeometryType == PxGeometryType::eTRIANGLEMESH);
+			PxU32 HitTris[64];
+			bool bOverflow = false;
+
+			const int32 NumTrisHit = bIsTriMesh ?
+				PxMeshQuery::findOverlapTriangleMesh(Geom, QueryTM, PTriMeshGeom, PShapeWorldPose, HitTris, ARRAY_COUNT(HitTris), 0, bOverflow) :
+				PxMeshQuery::findOverlapHeightField(Geom, QueryTM, PHeightfieldGeom, PShapeWorldPose, HitTris, ARRAY_COUNT(HitTris), 0, bOverflow);
+
+			if (NumTrisHit > 0)
 			{
-				OutNormal = FindBestOverlappingNormal(Geom, QueryTM, PTriMeshGeom, PShapeWorldPose, HitTris, NumTrisHit, bCanDrawOverlaps);
+				if (bIsTriMesh)
+				{
+					OutNormal = FindBestOverlappingNormal(Geom, QueryTM, PTriMeshGeom, PShapeWorldPose, HitTris, NumTrisHit, bCanDrawOverlaps);
+				}
+				else
+				{
+					OutNormal = FindBestOverlappingNormal(Geom, QueryTM, PHeightfieldGeom, PShapeWorldPose, HitTris, NumTrisHit, bCanDrawOverlaps);
+				}
+
+				return true;
 			}
-			else
+		}
+	}
+
+	return false;
+}
+
+
+static bool FindOverlappedTriangleNormal(const PxGeometry& Geom, const PxTransform& QueryTM, const PxShape* PShape, const PxTransform& PShapeWorldPose, FVector& OutNormal, float Inflation, bool bCanDrawOverlaps = false)
+{
+	bool bSuccess = false;
+
+	if (CanFindOverlappedTriangle(PShape))
+	{
+		if (Inflation <= 0.f)
+		{
+			bSuccess = FindOverlappedTriangleNormal_Internal(Geom, QueryTM, PShape, PShapeWorldPose, OutNormal, bCanDrawOverlaps);
+		}
+		else
+		{
+			// Try a slightly inflated test if possible.
+			switch (Geom.getType())
 			{
-				OutNormal = FindBestOverlappingNormal(Geom, QueryTM, PHeightfieldGeom, PShapeWorldPose, HitTris, NumTrisHit, bCanDrawOverlaps);
-			}
+				case PxGeometryType::eCAPSULE:
+				{
+					const PxCapsuleGeometry* InCapsule = static_cast<const PxCapsuleGeometry*>(&Geom);
+					PxCapsuleGeometry InflatedCapsule(InCapsule->radius + Inflation, InCapsule->halfHeight); // don't inflate halfHeight, radius is added all around.
+					bSuccess = FindOverlappedTriangleNormal_Internal(InflatedCapsule, QueryTM, PShape, PShapeWorldPose, OutNormal, bCanDrawOverlaps);
+					break;
+				}
 
-			return true;
+				case PxGeometryType::eBOX:
+				{
+					const PxBoxGeometry* InBox = static_cast<const PxBoxGeometry*>(&Geom);
+					PxBoxGeometry InflatedBox(InBox->halfExtents + PxVec3(Inflation));
+					bSuccess = FindOverlappedTriangleNormal_Internal(InflatedBox, QueryTM, PShape, PShapeWorldPose, OutNormal, bCanDrawOverlaps);
+					break;
+				}
+
+				case PxGeometryType::eSPHERE:
+				{
+					const PxSphereGeometry* InSphere = static_cast<const PxSphereGeometry*>(&Geom);
+					PxSphereGeometry InflatedSphere(InSphere->radius + Inflation);
+					bSuccess = FindOverlappedTriangleNormal_Internal(InflatedSphere, QueryTM, PShape, PShapeWorldPose, OutNormal, bCanDrawOverlaps);
+					break;
+				}
+
+				default:
+				{
+					// No inflation possible
+					break;
+				}
+			}
 		}
 	}
 
-	return false;
+	return bSuccess;
 }
 
 
@@ -723,11 +787,11 @@ static bool ConvertOverlappedShapeToImpactHit(const PxLocationHit& PHit, const F
 	{
 		FVector DummyNormal(0.f);
 		const PxTransform PShapeWorldPose = PxShapeExt::getGlobalPose(*PShape, *PActor);
-		FindOverlappedTriangleNormal(Geom, QueryTM, PShape, PShapeWorldPose, DummyNormal, true);
+		FindOverlappedTriangleNormal(Geom, QueryTM, PShape, PShapeWorldPose, DummyNormal, 0.f, true);
 	}
 #endif
 
-	// Zero-distance hits are often valid hits and we can extract the hit normal from a larger shape with MTD.
+	// Zero-distance hits are often valid hits and we can extract the hit normal.
 	// For invalid normals we can try other methods as well (get overlapping triangles).
 
 	if (PHit.distance == 0.f || !bValidNormal)
@@ -736,7 +800,7 @@ static bool ConvertOverlappedShapeToImpactHit(const PxLocationHit& PHit, const F
 
 		// Try MTD with a small inflation for better accuracy, then a larger one in case the first one fails due to precision issues.
 		static const float SmallMtdInflation = 0.250f;
-		static const float LargeMtdInflation = 2.500f;
+		static const float LargeMtdInflation = 1.750f;
 
 		if (ComputeInflatedMTD(SmallMtdInflation, PHit, OutResult, QueryTM, Geom, PShapeWorldPose) ||
 			ComputeInflatedMTD(LargeMtdInflation, PHit, OutResult, QueryTM, Geom, PShapeWorldPose))
@@ -745,14 +809,14 @@ static bool ConvertOverlappedShapeToImpactHit(const PxLocationHit& PHit, const F
 		}
 		else
 		{
-			if (FindOverlappedTriangleNormal(Geom, QueryTM, PShape, PShapeWorldPose, OutResult.ImpactNormal, false))
+			static const float SmallOverlapInflation = 0.250f;
+			if (FindOverlappedTriangleNormal(Geom, QueryTM, PShape, PShapeWorldPose, OutResult.ImpactNormal, 0.f, false) ||
+				FindOverlappedTriangleNormal(Geom, QueryTM, PShape, PShapeWorldPose, OutResult.ImpactNormal, SmallOverlapInflation, false))
 			{
 				// Success
 			}
 			else
 			{
-				// Not a tri-mesh or heightfield
-
 				// MTD failed, use point distance. This is not ideal.
 				// Note: faceIndex seems to be unreliable for convex meshes in these cases, so not using FindGeomOpposingNormal() for them here.
 				PxGeometry& PGeom = PShape->getGeometry().any();
