[INTEGRATE] Change 2440882 by [REDACTED_PERSON][REDACTED_HANDLE] on 2015/02/10 21:19:11

	Fix crash on exit when destroying ChildActorComponents

diff --git a/Engine/Source/Runtime/Engine/Private/ChildActorComponent.cpp b/Engine/Source/Runtime/Engine/Private/ChildActorComponent.cpp
--- a/Engine/Source/Runtime/Engine/Private/ChildActorComponent.cpp
+++ b/Engine/Source/Runtime/Engine/Private/ChildActorComponent.cpp
@@ -216,7 +216,7 @@ void UChildActorComponent::CreateChildActor()
 void UChildActorComponent::DestroyChildActor()
 {
 	// If we own an Actor, kill it now
-	if(ChildActor != nullptr)
+	if(ChildActor != nullptr && !GExitPurge)
 	{
 		// if still alive, destroy, otherwise just clear the pointer
 		if(!ChildActor->IsPendingKill())
