[INTEGRATE] Change 2440926 by [REDACTED_PERSON][REDACTED_HANDLE] on 2015/02/10 21:48:19

	Correctly recompute the world transform after applying a cached relative transform change

diff --git a/Engine/Source/Runtime/Engine/Private/SceneComponent.cpp b/Engine/Source/Runtime/Engine/Private/SceneComponent.cpp
--- a/Engine/Source/Runtime/Engine/Private/SceneComponent.cpp
+++ b/Engine/Source/Runtime/Engine/Private/SceneComponent.cpp
@@ -1180,6 +1180,12 @@ void FSceneComponentInstanceData::ApplyToComponent(UActorComponent* Component)
 	FComponentInstanceDataBase::ApplyToComponent(Component);
 
 	USceneComponent* SceneComponent = CastChecked<USceneComponent>(Component);
+
+	if (ContainsSavedProperties())
+	{
+		SceneComponent->UpdateComponentToWorld();
+	}
+
 	for (USceneComponent* ChildComponent : AttachedInstanceComponents)
 	{
 		if (ChildComponent)
