[INTEGRATE] Change 2440968 by [REDACTED_PERSON][REDACTED_HANDLE] on 2015/02/10 22:35:52

	Fix for bad assumption of GEditor

diff --git a/Engine/Source/Runtime/Engine/Private/ActorConstruction.cpp b/Engine/Source/Runtime/Engine/Private/ActorConstruction.cpp
--- a/Engine/Source/Runtime/Engine/Private/ActorConstruction.cpp
+++ b/Engine/Source/Runtime/Engine/Private/ActorConstruction.cpp
@@ -467,7 +467,7 @@ void AActor::RerunConstructionScripts()
 			}
 		}
 
-		if (OldToNewComponentMapping.Num() > 0)
+		if (GEditor && (OldToNewComponentMapping.Num() > 0))
 		{
 			GEditor->NotifyToolsOfObjectReplacement(OldToNewComponentMapping);
 		}
