Fix crash reported by UDN Thread:  https://udn.unrealengine.com/questions/236387/editor-crashed-when-ta-test-animation.html

diff --git a/Engine/Source/Runtime/Engine/Private/SkeletalMeshComponent.cpp b/Engine/Source/Runtime/Engine/Private/SkeletalMeshComponent.cpp
--- a/Engine/Source/Runtime/Engine/Private/SkeletalMeshComponent.cpp
+++ b/Engine/Source/Runtime/Engine/Private/SkeletalMeshComponent.cpp
@@ -1928,6 +1928,20 @@ void USkeletalMeshComponent::UnregisterOnSkeletalMeshPropertyChanged( FDelegateH
 
 void USkeletalMeshComponent::ValidateAnimation()
 {
+	if (SkeletalMesh->Skeleton == NULL)
+	{
+		UE_LOG(LogAnimation, Warning, TEXT("SkeletalMesh %s has no skeleton. This needs to fixed before an animation can be set"), *SkeletalMesh->GetName());
+		if (AnimationMode == EAnimationMode::AnimationSingleNode)
+		{
+			AnimationData.AnimToPlay = NULL;
+		}
+		else
+		{
+			AnimBlueprintGeneratedClass = NULL;
+		}
+		return;
+	}
+
 	if(AnimationMode == EAnimationMode::AnimationSingleNode)
 	{
 		if(AnimationData.AnimToPlay && SkeletalMesh && AnimationData.AnimToPlay->GetSkeleton() != SkeletalMesh->Skeleton)
