[INTEGRATE] Change 2442173 by [REDACTED_PERSON][REDACTED_HANDLE] on 2015/02/11 15:17:53

	Keep Instance and BlueprintCreated components arrays cleaned up when removing from OwnedComponents array

diff --git a/Engine/Source/Runtime/Engine/Private/Actor.cpp b/Engine/Source/Runtime/Engine/Private/Actor.cpp
--- a/Engine/Source/Runtime/Engine/Private/Actor.cpp
+++ b/Engine/Source/Runtime/Engine/Private/Actor.cpp
@@ -2005,18 +2005,47 @@ void AActor::AddOwnedComponent(UActorComponent* Component)
 	{
 		ReplicatedComponents.AddUnique(Component);
 	}
+
+	if (Component->IsCreatedByConstructionScript())
+	{
+		BlueprintCreatedComponents.AddUnique(Component);
+	}
+	else if (Component->CreationMethod == EComponentCreationMethod::Instance)
+	{
+		InstanceComponents.AddUnique(Component);
+	}
 }
 
 void AActor::RemoveOwnedComponent(UActorComponent* Component)
 {
-	if (OwnedComponents.RemoveSwap(Component) == 0)
+	if (OwnedComponents.RemoveSwap(Component) > 0)
+	{
+		ReplicatedComponents.RemoveSwap(Component);
+		if (Component->IsCreatedByConstructionScript())
+		{
+			BlueprintCreatedComponents.RemoveSwap(Component);
+		}
+		else if (Component->CreationMethod == EComponentCreationMethod::Instance)
+		{
+			InstanceComponents.RemoveSwap(Component);
+		}
+	}
+	else
 	{
 		// If we didn't remove something we expected to then it probably got NULL through the
 		// property system so take the time to pull them out now
-		OwnedComponents.RemoveSwap(NULL);
-	}
+		OwnedComponents.RemoveSwap(nullptr);
 
-	ReplicatedComponents.RemoveSwap(Component);
+		ReplicatedComponents.RemoveSwap(nullptr);
+		if (Component->IsCreatedByConstructionScript())
+		{
+			BlueprintCreatedComponents.RemoveSwap(nullptr);
+		}
+		else if (Component->CreationMethod == EComponentCreationMethod::Instance)
+		{
+			InstanceComponents.RemoveSwap(nullptr);
+		}
+	}
 }
 
 #if DO_CHECK
@@ -2065,7 +2094,7 @@ const TArray<UActorComponent*>& AActor::GetInstanceComponents() const
 void AActor::AddInstanceComponent(UActorComponent* Component)
 {
 	Component->CreationMethod = EComponentCreationMethod::Instance;
-	InstanceComponents.Add(Component);
+	InstanceComponents.AddUnique(Component);
 }
 
 void AActor::RemoveInstanceComponent(UActorComponent* Component)
