[INTEGRATE] Change 2442176 by [REDACTED_PERSON][REDACTED_HANDLE] on 2015/02/11 15:18:59

	Don't do 4.7 backwards compat code in cooked builds (since we can't use an object version to only do it on out of date packages)

diff --git a/Engine/Source/Runtime/Engine/Private/ActorComponent.cpp b/Engine/Source/Runtime/Engine/Private/ActorComponent.cpp
--- a/Engine/Source/Runtime/Engine/Private/ActorComponent.cpp
+++ b/Engine/Source/Runtime/Engine/Private/ActorComponent.cpp
@@ -127,32 +127,35 @@ void UActorComponent::PostLoad()
 {
 	Super::PostLoad();
 
-	// TODO: Wrap all this up with an engine version
-	if (bCreatedByConstructionScript_DEPRECATED)
+	// TODO: Wrap all this up with an engine version, for now just don't do it in cooked builds since we know they are good
+	if (!FPlatformProperties::RequiresCookedData())
 	{
-		CreationMethod = EComponentCreationMethod::SimpleConstructionScript;
-	}
-	else if (bInstanceComponent_DEPRECATED)
-	{
-		CreationMethod = EComponentCreationMethod::Instance;
-	}
+		if (bCreatedByConstructionScript_DEPRECATED)
+		{
+			CreationMethod = EComponentCreationMethod::SimpleConstructionScript;
+		}
+		else if (bInstanceComponent_DEPRECATED)
+		{
+			CreationMethod = EComponentCreationMethod::Instance;
+		}
 
-	if (CreationMethod == EComponentCreationMethod::SimpleConstructionScript)
-	{
-		UBlueprintGeneratedClass* Class = Cast/*Checked*/<UBlueprintGeneratedClass>(GetOuter()->GetClass());
-		while (Class)
+		if (CreationMethod == EComponentCreationMethod::SimpleConstructionScript)
 		{
-			USimpleConstructionScript* SCS = Class->SimpleConstructionScript;
-			if (SCS != nullptr && SCS->FindSCSNode(GetFName()))
-			{
-				break;
-			}
-			else
+			UBlueprintGeneratedClass* Class = Cast/*Checked*/<UBlueprintGeneratedClass>(GetOuter()->GetClass());
+			while (Class)
 			{
-				Class = Cast<UBlueprintGeneratedClass>(Class->GetSuperClass());
-				if (Class == nullptr)
+				USimpleConstructionScript* SCS = Class->SimpleConstructionScript;
+				if (SCS != nullptr && SCS->FindSCSNode(GetFName()))
+				{
+					break;
+				}
+				else
 				{
-					CreationMethod = EComponentCreationMethod::UserConstructionScript;
+					Class = Cast<UBlueprintGeneratedClass>(Class->GetSuperClass());
+					if (Class == nullptr)
+					{
+						CreationMethod = EComponentCreationMethod::UserConstructionScript;
+					}
 				}
 			}
 		}
