[INTEGRATE] Change 2442273 by [REDACTED_PERSON][REDACTED_HANDLE] on 2015/02/11 16:17:11

	[UE-9489] Since we use the SCS for each class along the super chain (to find a component's archetype), now ensuring that each is fully loaded.

diff --git a/Engine/Source/Runtime/Engine/Private/BlueprintGeneratedClass.cpp b/Engine/Source/Runtime/Engine/Private/BlueprintGeneratedClass.cpp
--- a/Engine/Source/Runtime/Engine/Private/BlueprintGeneratedClass.cpp
+++ b/Engine/Source/Runtime/Engine/Private/BlueprintGeneratedClass.cpp
@@ -289,15 +289,16 @@ UObject* UBlueprintGeneratedClass::FindArchetype(UClass* ArchetypeClass, const F
 	// and since preloading the SCS of a script in a world package is bad news, we need to filter them out
 	if (SimpleConstructionScript && !IsChildOf<ALevelScriptActor>())
 	{
-		if (SimpleConstructionScript->HasAllFlags(RF_NeedLoad))
-		{
-			SimpleConstructionScript->PreloadChain();
-		}
-
 		UBlueprintGeneratedClass* Class = const_cast<UBlueprintGeneratedClass*>(this);
 		while (Class)
 		{
-			USCS_Node* SCSNode = Class->SimpleConstructionScript->FindSCSNode(ArchetypeName);
+			USimpleConstructionScript* ClassSCS = Class->SimpleConstructionScript;
+			if (ClassSCS->HasAnyFlags(RF_NeedLoad))
+			{
+				ClassSCS->PreloadChain();
+			}
+
+			USCS_Node* SCSNode = ClassSCS->FindSCSNode(ArchetypeName);
 			if (SCSNode)
 			{
 				Archetype = SCSNode->ComponentTemplate;
