[INTEGRATE] Change 2442813 by [REDACTED_EMAIL]_Z2721_4_7 on 2015/02/11 21:40:31

	UE-9697
	VR: Incorrect canvas size is being used in certain window resizing events

	Fixed.

diff --git a/Engine/Source/Runtime/Engine/Private/UserInterface/Canvas.cpp b/Engine/Source/Runtime/Engine/Private/UserInterface/Canvas.cpp
--- a/Engine/Source/Runtime/Engine/Private/UserInterface/Canvas.cpp
+++ b/Engine/Source/Runtime/Engine/Private/UserInterface/Canvas.cpp
@@ -839,20 +839,18 @@ void FCanvas::Flush_GameThread(bool bForce)
 	{
 		ViewRect = FIntRect(FIntPoint::ZeroValue, RenderTargetSize);
 	}
-	if (IsScaledToRenderTarget())
-	{
-		ViewRect = FIntRect(0, 0, RenderTargetSize.X, RenderTargetSize.Y);
-	}
 
 	struct FCanvasFlushParameters
 	{
 		FIntRect ViewRect;
 		const FRenderTarget* CanvasRenderTarget;
+		bool bIsScaledToRenderTarget;
 	};
 	FCanvasFlushParameters FlushParameters =
 	{
 		ViewRect,
-		RenderTarget
+		RenderTarget,
+		IsScaledToRenderTarget()
 	};
 	bool bEmitCanvasDrawEvents = GEmitDrawEvents;
 
@@ -867,6 +865,11 @@ void FCanvas::Flush_GameThread(bool bForce)
 		// disable depth test & writes
 		RHICmdList.SetDepthStencilState(TStaticDepthStencilState<false,CF_Always>::GetRHI());
 
+		if (Parameters.bIsScaledToRenderTarget)
+		{
+			Parameters.ViewRect = FIntRect(0, 0, Parameters.CanvasRenderTarget->GetRenderTargetTexture()->GetSizeX(), Parameters.CanvasRenderTarget->GetRenderTargetTexture()->GetSizeY());
+		}
+
 		const FIntRect& ViewportRect = Parameters.ViewRect;
 		// set viewport to RT size
 		RHICmdList.SetViewport(ViewportRect.Min.X, ViewportRect.Min.Y, 0.0f, ViewportRect.Max.X, ViewportRect.Max.Y, 1.0f);
