[INTEGRATE] Change 2442903 by [REDACTED_PERSON][REDACTED_HANDLE] on 2015/02/11 23:38:44

	Fix non-editor build

diff --git a/Engine/Source/Runtime/Engine/Private/EdGraph/EdGraph.cpp b/Engine/Source/Runtime/Engine/Private/EdGraph/EdGraph.cpp
--- a/Engine/Source/Runtime/Engine/Private/EdGraph/EdGraph.cpp
+++ b/Engine/Source/Runtime/Engine/Private/EdGraph/EdGraph.cpp
@@ -194,13 +194,16 @@ void UEdGraph::MoveNodesToAnotherGraph(UEdGraph* DestinationGraph, bool bIsLoadi
 	{
 		if (UEdGraphNode* Node = Nodes.Pop())
 		{
+#if WITH_EDITOR
 			// During compilation, do not move ghost nodes, they are not used during compilation.
-			if(bInIsCompiling && !Node->bIsNodeEnabled)
+			if (bInIsCompiling && !Node->bIsNodeEnabled)
 			{
 				// Break all node links, if any exist, do not move the node
 				Node->BreakAllNodeLinks();
 				continue;
 			}
+#endif
+
 			// Let the name be autogenerated, to automatically avoid naming conflicts
 			// Since this graph is always going to come from a cloned source graph, user readable names can come from the remap stored in a MessageLog.
 	
