Deprecate OverlapSingle and introduce:
OverlapBlockingTest, OverlapAnyTest

diff --git a/Engine/Source/Runtime/Engine/Classes/Engine/World.h b/Engine/Source/Runtime/Engine/Classes/Engine/World.h
--- a/Engine/Source/Runtime/Engine/Classes/Engine/World.h
+++ b/Engine/Source/Runtime/Engine/Classes/Engine/World.h
@@ -1233,7 +1233,7 @@ public:
 	bool SweepMultiByProfile(TArray<FHitResult>& OutHits, const FVector& Start, const FVector& End, const FQuat& Rot, FName ProfileName, const struct FCollisionShape& CollisionShape, const struct FCollisionQueryParams& Params) const;
 
 	/**
-	 *  Test the collision of an AABB at the supplied location using a specific channel, and return if any blocking overlap is found
+	 *  Test the collision of a shape at the supplied location using a specific channel, and return if any blocking overlap is found
 	 *  @param  Pos             Location of center of box to test against the world
 	 *  @param  TraceChannel    The 'channel' that this query is in, used to determine which components to hit
 	 *  @param	CollisionShape	CollisionShape - supports Box, Sphere, Capsule
@@ -1241,64 +1241,122 @@ public:
 	 * 	@param 	ResponseParam	ResponseContainer to be used for this trace	 
 	 *  @return TRUE if any blocking results are found
 	 */
+	DEPRECATED(4.8, "Use OverlapBlockingTest instead.")
 	bool OverlapTest(const FVector& Pos, const FQuat& Rot, ECollisionChannel TraceChannel, const struct FCollisionShape& CollisionShape, const struct FCollisionQueryParams& Params, const struct FCollisionResponseParams& ResponseParam = FCollisionResponseParams::DefaultResponseParam) const;
 
 	/**
-	 *  Test the collision of an AABB at the supplied location using object types, and return if any blocking overlap is found
+	*  Test the collision of a shape at the supplied location using a specific channel, and return if any blocking overlap is found
+	*  @param  Pos             Location of center of box to test against the world
+	*  @param  TraceChannel    The 'channel' that this query is in, used to determine which components to hit
+	*  @param  CollisionShape	CollisionShape - supports Box, Sphere, Capsule
+	*  @param  Params          Additional parameters used for the trace
+	*  @param  ResponseParam	ResponseContainer to be used for this trace
+	*  @return TRUE if any blocking results are found
+	*/
+	bool OverlapBlockingTest(const FVector& Pos, const FQuat& Rot, ECollisionChannel TraceChannel, const struct FCollisionShape& CollisionShape, const struct FCollisionQueryParams& Params, const struct FCollisionResponseParams& ResponseParam = FCollisionResponseParams::DefaultResponseParam) const;
+
+	/**
+	*  Test the collision of a shape at the supplied location using a specific channel, and return if any blocking or overlapping shape is found
+	*  @param  Pos             Location of center of box to test against the world
+	*  @param  TraceChannel    The 'channel' that this query is in, used to determine which components to hit
+	*  @param  CollisionShape	CollisionShape - supports Box, Sphere, Capsule
+	*  @param  Params          Additional parameters used for the trace
+	*  @param  ResponseParam	ResponseContainer to be used for this trace
+	*  @return TRUE if any blocking or overlapping results are found
+	*/
+	bool OverlapAnyTest(const FVector& Pos, const FQuat& Rot, ECollisionChannel TraceChannel, const struct FCollisionShape& CollisionShape, const struct FCollisionQueryParams& Params, const struct FCollisionResponseParams& ResponseParam = FCollisionResponseParams::DefaultResponseParam) const;
+
+	/**
+	 *  Test the collision of a shape at the supplied location using object types, and return if any blocking overlap is found
 	 *  @param  Pos             Location of center of box to test against the world
 	 *  @param	CollisionShape	CollisionShape - supports Box, Sphere, Capsule
 	 *  @param  Params          Additional parameters used for the trace
 	 *	@param	ObjectQueryParams	List of object types it's looking for
-	 *  @return TRUE if any hit is found
+	 *  @return TRUE if any blocking results are found
 	 */
+	DEPRECATED(4.8, "Use OverlapAnyTest instead.")
 	bool OverlapTest(const FVector& Pos, const FQuat& Rot, const struct FCollisionShape& CollisionShape, const struct FCollisionQueryParams& Params, const struct FCollisionObjectQueryParams& ObjectQueryParams) const;
 
 	/**
-	 *  Test the collision of an AABB at the supplied location using a specific profile, and return if any blocking overlap is found
+	*  Test the collision of a shape at the supplied location using object types, and return if any overlap is found
+	*  @param  Pos             Location of center of box to test against the world
+	*  @param  CollisionShape	CollisionShape - supports Box, Sphere, Capsule
+	*  @param  Params          Additional parameters used for the trace
+	*  @param  ObjectQueryParams	List of object types it's looking for
+	*  @return TRUE if any blocking results are found
+	*/
+	bool OverlapAnyTest(const FVector& Pos, const FQuat& Rot, const struct FCollisionShape& CollisionShape, const struct FCollisionQueryParams& Params, const struct FCollisionObjectQueryParams& ObjectQueryParams) const;
+
+	/**
+	 *  Test the collision of a shape at the supplied location using a specific profile, and return if any blocking overlap is found
 	 *  @param  Pos             Location of center of box to test against the world
 	 *  @param  ProfileName     The 'profile' used to determine which components to hit
 	 *  @param	CollisionShape	CollisionShape - supports Box, Sphere, Capsule
 	 *  @param  Params          Additional parameters used for the trace
 	 *  @return TRUE if any blocking results are found
 	 */
+	DEPRECATED(4.8, "Use OverlapBlockingTestByProfile instead.")
 	bool OverlapTestByProfile(const FVector& Pos, const FQuat& Rot, FName ProfileName, const struct FCollisionShape& CollisionShape, const struct FCollisionQueryParams& Params) const;
 
 	/**
-	 *  Test the collision of a sphere at the supplied location using a specific channel, and determine the set of components that it overlaps
-	 *  @param  OutOverlaps     Array of components found to overlap supplied box
+	*  Test the collision of a shape at the supplied location using a specific profile, and return if any blocking overlap is found
+	*  @param  Pos             Location of center of box to test against the world
+	*  @param  ProfileName     The 'profile' used to determine which components to hit
+	*  @param	CollisionShape	CollisionShape - supports Box, Sphere, Capsule
+	*  @param  Params          Additional parameters used for the trace
+	*  @return TRUE if any blocking results are found
+	*/
+	bool OverlapBlockingTestByProfile(const FVector& Pos, const FQuat& Rot, FName ProfileName, const struct FCollisionShape& CollisionShape, const struct FCollisionQueryParams& Params) const;
+
+	/**
+	*  Test the collision of a shape at the supplied location using a specific profile, and return if any blocking or overlap is found
+	*  @param  Pos             Location of center of box to test against the world
+	*  @param  ProfileName     The 'profile' used to determine which components to hit
+	*  @param	CollisionShape	CollisionShape - supports Box, Sphere, Capsule
+	*  @param  Params          Additional parameters used for the trace
+	*  @return TRUE if any blocking or overlapping results are found
+	*/
+	bool OverlapAnyTestByProfile(const FVector& Pos, const FQuat& Rot, FName ProfileName, const struct FCollisionShape& CollisionShape, const struct FCollisionQueryParams& Params) const;
+
+	/**
+	 *  Test the collision of a shape at the supplied location using a specific channel, and determine the set of components that it overlaps
+	 *  @param  OutOverlap      Component found to overlap supplied shape
 	 *  @param  Pos             Location of center of sphere to test against the world
 	 *  @param  TraceChannel    The 'channel' that this query is in, used to determine which components to hit
 	 *  @param	CollisionShape	CollisionShape - supports Box, Sphere, Capsule
 	 *  @param  Params          Additional parameters used for the trace
 	 * 	@param 	ResponseParam	ResponseContainer to be used for this trace
-	 *  @return TRUE if OutOverlaps contains any blocking results
+	 *  @return TRUE if OutOverlap is a blocking result
 	 */
+	DEPRECATED(4.8, "Use OverlapMulti instead. Deciding which overlap to use is gameplay specific.")
 	bool OverlapSingle(struct FOverlapResult& OutOverlap, const FVector& Pos, const FQuat& Rot, ECollisionChannel TraceChannel, const struct FCollisionShape& CollisionShape, const struct FCollisionQueryParams& Params, const struct FCollisionResponseParams& ResponseParam = FCollisionResponseParams::DefaultResponseParam) const;
 
 	/**
-	 *  Test the collision of a sphere at the supplied location using object types, and determine the set of components that it overlaps
-	 *  @param  OutOverlaps     Array of components found to overlap supplied box
+	 *  Test the collision of a shape at the supplied location using object types, and determine the set of components that it overlaps
+	 *  @param  OutOverlap      Component found to overlap supplied shape
 	 *  @param  Pos             Location of center of sphere to test against the world
 	 *  @param	CollisionShape	CollisionShape - supports Box, Sphere, Capsule
 	 *  @param  Params          Additional parameters used for the trace
 	 *	@param	ObjectQueryParams	List of object types it's looking for
-	 *  @return TRUE if any hit is found
+	 *  @return TRUE if any overlap is found
 	 */
+	DEPRECATED(4.8, "Use OverlapMulti instead. Deciding which overlap to use is gameplay specific.")
 	bool OverlapSingle(struct FOverlapResult& OutOverlap, const FVector& Pos, const FQuat& Rot, const struct FCollisionShape& CollisionShape, const struct FCollisionQueryParams& Params, const struct FCollisionObjectQueryParams& ObjectQueryParams) const;
 
 	/**
-	 *  Test the collision of a sphere at the supplied location using a specific profile, and determine the set of components that it overlaps
-	 *  @param  OutOverlaps     Array of components found to overlap supplied box
+	 *  Test the collision of a shape at the supplied location using a specific profile, and determine the set of components that it overlaps
+	 *  @param  OutOverlap      Component found to overlap supplied shape
 	 *  @param  Pos             Location of center of sphere to test against the world
 	 *  @param  ProfileName     The 'profile' used to determine which components to hit
 	 *  @param	CollisionShape	CollisionShape - supports Box, Sphere, Capsule
 	 *  @param  Params          Additional parameters used for the trace
-	 *  @return TRUE if OutOverlaps contains any blocking results
+	 *  @return TRUE if OutOverlap is blocking
 	 */
+	DEPRECATED(4.8, "Use OverlapMultiByProfile instead. Deciding which overlap to use is gameplay specific.")
 	bool OverlapSingleByProfile(struct FOverlapResult& OutOverlap, const FVector& Pos, const FQuat& Rot, FName ProfileName, const struct FCollisionShape& CollisionShape, const struct FCollisionQueryParams& Params) const;
 
 	/**
-	 *  Test the collision of a sphere at the supplied location using a specific channel, and determine the set of components that it overlaps
+	 *  Test the collision of a shape at the supplied location using a specific channel, and determine the set of components that it overlaps
 	 *  @param  OutOverlaps     Array of components found to overlap supplied box
 	 *  @param  Pos             Location of center of sphere to test against the world
 	 *  @param  TraceChannel    The 'channel' that this query is in, used to determine which components to hit
@@ -1310,18 +1368,18 @@ public:
 	bool OverlapMulti(TArray<struct FOverlapResult>& OutOverlaps, const FVector& Pos, const FQuat& Rot, ECollisionChannel TraceChannel, const struct FCollisionShape& CollisionShape, const struct FCollisionQueryParams& Params, const struct FCollisionResponseParams& ResponseParam = FCollisionResponseParams::DefaultResponseParam) const;
 
 	/**
-	 *  Test the collision of a sphere at the supplied location using object types, and determine the set of components that it overlaps
+	 *  Test the collision of a shape at the supplied location using object types, and determine the set of components that it overlaps
 	 *  @param  OutOverlaps     Array of components found to overlap supplied box
 	 *  @param  Pos             Location of center of sphere to test against the world
 	 *  @param	CollisionShape	CollisionShape - supports Box, Sphere, Capsule
 	 *  @param  Params          Additional parameters used for the trace
 	 *	@param	ObjectQueryParams	List of object types it's looking for
-	 *  @return TRUE if any hit is found
+	 *  @return TRUE if any overlap is found
 	 */
 	bool OverlapMulti(TArray<struct FOverlapResult>& OutOverlaps, const FVector& Pos, const FQuat& Rot, const struct FCollisionShape& CollisionShape, const struct FCollisionQueryParams& Params, const struct FCollisionObjectQueryParams& ObjectQueryParams) const;
 	
 	/**
-	 *  Test the collision of a sphere at the supplied location using a specific profile, and determine the set of components that it overlaps
+	 *  Test the collision of a shape at the supplied location using a specific profile, and determine the set of components that it overlaps
 	 *  @param  OutOverlaps     Array of components found to overlap supplied box
 	 *  @param  Pos             Location of center of sphere to test against the world
 	 *  @param  ProfileName     The 'profile' used to determine which components to hit
@@ -1480,9 +1538,8 @@ public:
 	 *							TraceDelegate.BindRaw(this, &MyActor::TraceDone);
 	 * 
 	 *	@param UserData			UserData
-	 *	@param bMultiTrace		true if you'd like to get continuous result from the trace, false if you want single
 	 */ 
-	FTraceHandle	AsyncOverlap(const FVector& Pos, const FQuat& Rot, ECollisionChannel TraceChannel, const struct FCollisionShape& CollisionShape, const struct FCollisionQueryParams& Params, const struct FCollisionResponseParams& ResponseParam = FCollisionResponseParams::DefaultResponseParam, FOverlapDelegate * InDelegate = NULL, uint32 UserData = 0, bool bMultiTrace = false);
+	FTraceHandle	AsyncOverlap(const FVector& Pos, const FQuat& Rot, ECollisionChannel TraceChannel, const struct FCollisionShape& CollisionShape, const struct FCollisionQueryParams& Params, const struct FCollisionResponseParams& ResponseParam = FCollisionResponseParams::DefaultResponseParam, FOverlapDelegate * InDelegate = NULL, uint32 UserData = 0);
 
 	/**
 	 * Interface for Async trace
@@ -1503,9 +1560,8 @@ public:
 	 *							TraceDelegate.BindRaw(this, &MyActor::TraceDone);
 	 * 
 	 *	@param UserData			UserData
-	 *	@param bMultiTrace		true if you'd like to get continuous result from the trace, false if you want single
 	 */ 
-	FTraceHandle	AsyncOverlap(const FVector& Pos, const FQuat& Rot, const struct FCollisionShape& CollisionShape, const struct FCollisionQueryParams& Params, const struct FCollisionObjectQueryParams& ObjectQueryParams, FOverlapDelegate * InDelegate = NULL, uint32 UserData = 0, bool bMultiTrace = false);
+	FTraceHandle	AsyncOverlap(const FVector& Pos, const FQuat& Rot, const struct FCollisionShape& CollisionShape, const struct FCollisionQueryParams& Params, const struct FCollisionObjectQueryParams& ObjectQueryParams, FOverlapDelegate * InDelegate = NULL, uint32 UserData = 0);
 
 	/**
 	 * Query function 
diff --git a/Engine/Source/Runtime/Engine/Private/Collision/CollisionConversions.cpp b/Engine/Source/Runtime/Engine/Private/Collision/CollisionConversions.cpp
--- a/Engine/Source/Runtime/Engine/Private/Collision/CollisionConversions.cpp
+++ b/Engine/Source/Runtime/Engine/Private/Collision/CollisionConversions.cpp
@@ -841,11 +841,7 @@ static bool ConvertOverlappedShapeToImpactHit(const PxLocationHit& PHit, const F
 
 void ConvertQueryOverlap(const PxShape* PShape, const PxRigidActor* PActor, FOverlapResult& OutOverlap, const PxFilterData& QueryFilter)
 {
-	// See if this is a 'blocking' hit
-	PxFilterData PShapeFilter = PShape->getQueryFilterData();
-	// word1 is block, word2 is overlap
-	PxSceneQueryHitType::Enum HitType = FPxQueryFilterCallback::CalcQueryHitType(QueryFilter, PShapeFilter);
-	bool bBlock = (HitType == PxSceneQueryHitType::eBLOCK); 
+	const bool bBlock = IsBlocking(PShape, QueryFilter);
 
 	FBodyInstance* BodyInst = FPhysxUserData::Get<FBodyInstance>(PActor->userData);
 	FDestructibleChunkInfo* ChunkInfo = FPhysxUserData::Get<FDestructibleChunkInfo>(PActor->userData);
diff --git a/Engine/Source/Runtime/Engine/Private/Collision/CollisionConversions.h b/Engine/Source/Runtime/Engine/Private/Collision/CollisionConversions.h
--- a/Engine/Source/Runtime/Engine/Private/Collision/CollisionConversions.h
+++ b/Engine/Source/Runtime/Engine/Private/Collision/CollisionConversions.h
@@ -60,6 +60,23 @@ bool AddSweepResults(int32 NumHits, const PxSweepHit* Hits, float CheckLength, c
  */
 void ConvertQueryOverlap(const PxShape* PShape, const PxRigidActor* PActor, FOverlapResult& OutOverlap, const PxFilterData& QueryFilter);
 
+/**
+ * Util to determine if a shape is deemed blocking based on the query filter
+ *
+ * @param PShape Shape that overlaps
+ * @QueryFilter Query Filter
+ * @return true if the query filter and shape filter resolve to be blocking
+ */
+inline bool IsBlocking(const PxShape* PShape, const PxFilterData& QueryFilter)
+{
+	// See if this is a 'blocking' hit
+	const PxFilterData PShapeFilter = PShape->getQueryFilterData();
+	// word1 is block, word2 is overlap
+	const PxSceneQueryHitType::Enum HitType = FPxQueryFilterCallback::CalcQueryHitType(QueryFilter, PShapeFilter);
+	const bool bBlock = (HitType == PxSceneQueryHitType::eBLOCK);
+	return bBlock;
+}
+
 
 /** 
  * Util to convert a list of overlap hits into FOverlapResult and add them to OutOverlaps, if not already there
diff --git a/Engine/Source/Runtime/Engine/Private/Collision/PhysXCollision.cpp b/Engine/Source/Runtime/Engine/Private/Collision/PhysXCollision.cpp
--- a/Engine/Source/Runtime/Engine/Private/Collision/PhysXCollision.cpp
+++ b/Engine/Source/Runtime/Engine/Private/Collision/PhysXCollision.cpp
@@ -1549,157 +1549,21 @@ bool GeomSweepMulti(const UWorld* World, const struct FCollisionShape& Collision
 //////////////////////////////////////////////////////////////////////////
 // GEOM OVERLAP
 
-bool GeomOverlapTest(const UWorld* World, const struct FCollisionShape& CollisionShape, const FVector& Pos, const FQuat& Rot, ECollisionChannel TraceChannel, const struct FCollisionQueryParams& Params, const struct FCollisionResponseParams& ResponseParams, const struct FCollisionObjectQueryParams& ObjectParams)
-{
-	if ((World == NULL) || (World->GetPhysicsScene() == NULL))
-	{
-		return false;
-	}
-	SCOPE_CYCLE_COUNTER(STAT_Collision_GeomOverlapAny);
-
-	// Track if we get any 'blocking' hits
-	bool bHaveBlockingHit = false;
-
 #if WITH_PHYSX
-	FPhysXShapeAdaptor ShapeAdaptor(Rot, CollisionShape);
-	const PxGeometry& PGeom = ShapeAdaptor.GetGeometry();
-	const PxTransform& PGeomPose = ShapeAdaptor.GetGeomPose(Pos);
-
-	// Create filter data used to filter collisions
-	PxFilterData PFilter = CreateQueryFilterData(TraceChannel, Params.bTraceComplex, ResponseParams.CollisionResponse, ObjectParams, false);
-	PxSceneQueryFilterData PQueryFilterData(PFilter, PxSceneQueryFilterFlag::eSTATIC | PxSceneQueryFilterFlag::eDYNAMIC | PxSceneQueryFilterFlag::ePREFILTER);
-	FPxQueryFilterCallback PQueryCallback(Params.IgnoreComponents);
 
-	PxOverlapHit POverlapResult;
-	bool bHit = false;
-
-	FPhysScene* PhysScene = World->GetPhysicsScene();
-	{
-		// Enable scene locks, in case they are required
-		PxScene* SyncScene = PhysScene->GetPhysXScene(PST_Sync);
-		SCOPED_SCENE_READ_LOCK(SyncScene);
-
-		bHit = SyncScene->overlapAny(PGeom, PGeomPose, POverlapResult, PQueryFilterData, &PQueryCallback);
-
-		// if we got a hit, need to see if it was a touch or a block
-		// @todo UE4 james remove this once overlapAny only returns blocking hits
-		if (bHit)
-		{
-			FOverlapResult NewOverlap;
-			ConvertQueryOverlap(POverlapResult.shape, POverlapResult.actor, NewOverlap, PFilter);
-			bHaveBlockingHit = NewOverlap.bBlockingHit;
-		}
-	}
-
-	// Test async scene if async tests are requested and there was no blocking hit was found in the sync scene (since no hit info other than a boolean yes/no is recorded)
-	if( !bHaveBlockingHit && Params.bTraceAsyncScene && PhysScene->HasAsyncScene() )
-	{
-		PxScene* AsyncScene = PhysScene->GetPhysXScene(PST_Async);
-		SCOPED_SCENE_READ_LOCK(AsyncScene);
-
-		PxOverlapHit PAnotherOverlapResult;
-		bHit = AsyncScene->overlapAny(PGeom, PGeomPose,  PAnotherOverlapResult, PQueryFilterData, &PQueryCallback);
-
-		if(bHit)
-		{
-			FOverlapResult NewOverlap;
-			ConvertQueryOverlap( PAnotherOverlapResult.shape, PAnotherOverlapResult.actor, NewOverlap, PFilter);
-			bHaveBlockingHit = NewOverlap.bBlockingHit;
-		}
-	}
-
-#if !(UE_BUILD_SHIPPING || UE_BUILD_TEST)
-	if((World->DebugDrawTraceTag != NAME_None) && (World->DebugDrawTraceTag == Params.TraceTag))
-	{
-		TArray<FOverlapResult> Overlaps;
-		DrawGeomOverlaps(World, PGeom, PGeomPose, Overlaps, DebugLineLifetime);
-	}
-#endif //!(UE_BUILD_SHIPPING || UE_BUILD_TEST)
-#endif // WITH_PHYSX
-
-	//@TODO: BOX2D: Implement GeomOverlapTest
-
-	return bHaveBlockingHit;
-}
-
-bool GeomOverlapSingle(const UWorld* World, const struct FCollisionShape& CollisionShape, const FVector& Pos, const FQuat& Rot, FOverlapResult& OutOverlap, ECollisionChannel TraceChannel, const struct FCollisionQueryParams& Params, const struct FCollisionResponseParams& ResponseParams, const struct FCollisionObjectQueryParams& ObjectParams)
+namespace EQueryInfo
 {
-	if ((World == NULL) || (World->GetPhysicsScene() == NULL))
-	{
-		return false;
-	}
-	SCOPE_CYCLE_COUNTER(STAT_Collision_GeomOverlapSingle);
-
-	// Track if we get any 'blocking' hits
-	bool bHaveBlockingHit = false;
-	bool bHit = false;
-
-#if WITH_PHYSX
-	FPhysXShapeAdaptor ShapeAdaptor(Rot, CollisionShape);
-	const PxGeometry& PGeom = ShapeAdaptor.GetGeometry();
-	const PxTransform& PGeomPose = ShapeAdaptor.GetGeomPose(Pos);
-
-	// Create filter data used to filter collisions
-	PxFilterData PFilter = CreateQueryFilterData(TraceChannel, Params.bTraceComplex, ResponseParams.CollisionResponse, ObjectParams, false);
-	PxSceneQueryFilterData PQueryFilterData(PFilter, PxSceneQueryFilterFlag::eSTATIC | PxSceneQueryFilterFlag::eDYNAMIC | PxSceneQueryFilterFlag::ePREFILTER);
-	FPxQueryFilterCallback PQueryCallback(Params.IgnoreComponents);
-
-	// Enable scene locks, in case they are required
-	FPhysScene* PhysScene = World->GetPhysicsScene();
-	{
-		PxScene* SyncScene = PhysScene->GetPhysXScene(PST_Sync);
-		SCOPED_SCENE_READ_LOCK(SyncScene);
-
-		PxOverlapHit POverlapResult;
-		bHit = SyncScene->overlapAny(PGeom, PGeomPose, POverlapResult, PQueryFilterData, &PQueryCallback);
-
-		// if we got a hit, need to see if it was a touch or a block
-		// @todo UE4 james remove this once overlapAny only returns blocking hits
-		if (bHit)
-		{
-			ConvertQueryOverlap( POverlapResult.shape, POverlapResult.actor, OutOverlap, PFilter);
-			bHaveBlockingHit = OutOverlap.bBlockingHit;
-		}
-	}
-
-	// Test async scene if async tests are requested and there was no blocking hit was found in the sync scene (since no hit info other than a boolean yes/no is recorded)
-	if (!bHaveBlockingHit && Params.bTraceAsyncScene && PhysScene->HasAsyncScene())
-	{
-		PxScene* AsyncScene = PhysScene->GetPhysXScene(PST_Async);
-		SCOPED_SCENE_READ_LOCK(AsyncScene);
-
-		PxOverlapHit PAnotherOverlapResult;
-		bHit = AsyncScene->overlapAny(PGeom, PGeomPose, PAnotherOverlapResult, PQueryFilterData, &PQueryCallback);
-
-		if (bHit)
-		{
-			FOverlapResult NewOverlap;
-			ConvertQueryOverlap(PAnotherOverlapResult.shape, PAnotherOverlapResult.actor, NewOverlap, PFilter);
-			bHaveBlockingHit = NewOverlap.bBlockingHit;
-		}
-	}
-
-#if !(UE_BUILD_SHIPPING || UE_BUILD_TEST)
-	if ((World->DebugDrawTraceTag != NAME_None) && (World->DebugDrawTraceTag == Params.TraceTag))
+	//This is used for templatizing code based on the info we're trying to get out.
+	enum Type
 	{
-		TArray<FOverlapResult> Overlaps;
-		if (bHit)
-		{
-			Overlaps.Add(OutOverlap);
-		}
-
-		DrawGeomOverlaps(World, PGeom, PGeomPose, Overlaps, DebugLineLifetime);
-	}
-#endif //!(UE_BUILD_SHIPPING || UE_BUILD_TEST)
-#endif // WITH_PHYSX
-
-	//@TODO: BOX2D: Implement GeomOverlapSingle
-
-	return bHaveBlockingHit;
+		GatherAll,		//get all data and actually return it
+		IsBlocking,		//is any of the data blocking? only return a bool so don't bother collecting
+		IsAnything		//is any of the data blocking or touching? only return a bool so don't bother collecting
+	};
 }
 
-#if WITH_PHYSX
-bool GeomOverlapMulti_PhysX(const UWorld* World, const PxGeometry& PGeom, const PxTransform& PGeomPose, TArray<FOverlapResult>& OutOverlaps, ECollisionChannel TraceChannel, const struct FCollisionQueryParams& Params, const struct FCollisionResponseParams& ResponseParams, const struct FCollisionObjectQueryParams& ObjectParams)
+template <EQueryInfo::Type InfoType>
+bool GeomOverlapMultiImp_PhysX(const UWorld* World, const PxGeometry& PGeom, const PxTransform& PGeomPose, TArray<FOverlapResult>& OutOverlaps, ECollisionChannel TraceChannel, const struct FCollisionQueryParams& Params, const struct FCollisionResponseParams& ResponseParams, const struct FCollisionObjectQueryParams& ObjectParams)
 {
 	SCOPE_CYCLE_COUNTER(STAT_Collision_GeomOverlapMultiple);
 	bool bHaveBlockingHit = false;
@@ -1708,7 +1572,7 @@ bool GeomOverlapMulti_PhysX(const UWorld* World, const PxGeometry& PGeom, const
 	if (PGeom.getType()==PxGeometryType::eSPHERE || PGeom.getType()==PxGeometryType::eCAPSULE || PGeom.getType()==PxGeometryType::eBOX || PGeom.getType()==PxGeometryType::eCONVEXMESH )
 	{
 		// Create filter data used to filter collisions
-		PxFilterData PFilter = CreateQueryFilterData(TraceChannel, Params.bTraceComplex, ResponseParams.CollisionResponse, ObjectParams, true);
+		PxFilterData PFilter = CreateQueryFilterData(TraceChannel, Params.bTraceComplex, ResponseParams.CollisionResponse, ObjectParams, InfoType != EQueryInfo::IsAnything);
 		PxSceneQueryFilterData PQueryFilterData(PFilter, PxSceneQueryFilterFlag::eSTATIC | PxSceneQueryFilterFlag::eDYNAMIC | PxSceneQueryFilterFlag::ePREFILTER);
 		FPxQueryFilterCallback PQueryCallback(Params.IgnoreComponents);
 
@@ -1723,22 +1587,46 @@ bool GeomOverlapMulti_PhysX(const UWorld* World, const PxGeometry& PGeom, const
 		// Create buffer for hits. Note: memory is not initialized (for perf reasons), since API does not require it.
 		TTypeCompatibleBytes<PxOverlapHit> RawPOverlapArray[OVERLAP_BUFFER_SIZE];
 		PxOverlapHit* POverlapArray = (PxOverlapHit*)RawPOverlapArray;
+		
+		PxOverlapHit POverlapResult;	//IsAnything code path only needs one so we use this instead of array. Compiler will optimize out the buffer since code won't use it
 
-		PxI32 NumHits = SyncScene->overlapMultiple(PGeom, PGeomPose, POverlapArray,  OVERLAP_BUFFER_SIZE_MAX_SYNC_QUERIES, PQueryFilterData, &PQueryCallback);
+		PxI32 NumHits = 0;
 		
-		if (NumHits == -1)
+		if (InfoType == EQueryInfo::IsAnything)
 		{
-			UE_LOG(LogCollision, Warning, TEXT("GeomOverlapMulti : Buffer overflow from synchronous scene query"));
-			UE_LOG(LogCollision, Warning, TEXT("--------TraceChannel : %d"), (int32)TraceChannel);
-			UE_LOG(LogCollision, Warning, TEXT("--------%s"), *Params.ToString());
-			// overlapMultiple still fills in the buffer with an arbitrary set of overlapping objects, so we have a full buffer of results
-			// we can use, though that is not the entire set
-			NumHits = OVERLAP_BUFFER_SIZE_MAX_SYNC_QUERIES;
+			if (SyncScene->overlapAny(PGeom, PGeomPose, POverlapResult, PQueryFilterData, &PQueryCallback))
+			{
+				return true;
+			}
 		}
-		else if (NumHits == 0)
+		else
 		{
-			// Not using anything from this scene, so unlock it.
-			SceneLocks.UnlockRead(SyncScene, PST_Sync);
+			NumHits = SyncScene->overlapMultiple(PGeom, PGeomPose, POverlapArray, OVERLAP_BUFFER_SIZE_MAX_SYNC_QUERIES, PQueryFilterData, &PQueryCallback);
+			if (NumHits == -1)
+			{
+				UE_LOG(LogCollision, Warning, TEXT("GeomOverlapMulti : Buffer overflow from synchronous scene query"));
+				UE_LOG(LogCollision, Warning, TEXT("--------TraceChannel : %d"), (int32)TraceChannel);
+				UE_LOG(LogCollision, Warning, TEXT("--------%s"), *Params.ToString());
+				// overlapMultiple still fills in the buffer with an arbitrary set of overlapping objects, so we have a full buffer of results
+				// we can use, though that is not the entire set
+				NumHits = OVERLAP_BUFFER_SIZE_MAX_SYNC_QUERIES;
+			}
+			else if (NumHits == 0)
+			{
+				// Not using anything from this scene, so unlock it.
+				SceneLocks.UnlockRead(SyncScene, PST_Sync);
+			}
+
+			if (InfoType == EQueryInfo::IsBlocking)	//in the case where we just want to know if we're blocking, quickly check if anything blocks and return
+			{
+				for (PxI32 HitIdx = 0; HitIdx < NumHits; ++HitIdx)
+				{
+					if (IsBlocking(POverlapArray[HitIdx].shape, PFilter))
+					{
+						return true;
+					}
+				}
+			}
 		}
 
 		// Test async scene if async tests are requested and there was no overflow
@@ -1753,37 +1641,63 @@ bool GeomOverlapMulti_PhysX(const UWorld* World, const PxGeometry& PGeom, const
 			PxOverlapHit* PAsyncOverlapArray = POverlapArray + NumHits;
 			const PxI32 AsyncBufferSize = (ARRAY_COUNT(RawPOverlapArray) - NumHits);
 			check(AsyncBufferSize > 0);
-			PxI32 NumAsyncHits = AsyncScene->overlapMultiple(PGeom, PGeomPose, PAsyncOverlapArray, AsyncBufferSize, PQueryFilterData, &PQueryCallback);
-
-			if (NumAsyncHits == -1)
+			PxI32 NumAsyncHits = 0;
+			
+			if (InfoType == EQueryInfo::IsAnything)
 			{
-				UE_LOG(LogCollision, Log, TEXT("GeomOverlapMulti : Buffer overflow from asynchronous scene query"));
-				UE_LOG(LogCollision, Warning, TEXT("--------TraceChannel : %d"), (int32)TraceChannel);
-				UE_LOG(LogCollision, Warning, TEXT("--------%s"), *Params.ToString());
-				// overlapMultiple still fills in the buffer with an arbitrary set of overlapping objects, so we have a full buffer of results
-				// we can use, though that is not the entire set
-				NumAsyncHits = AsyncBufferSize;
+				if (AsyncScene->overlapAny(PGeom, PGeomPose, POverlapResult, PQueryFilterData, &PQueryCallback))
+				{
+					return true;
+				}
 			}
-			else if (NumAsyncHits == 0)
+			else
 			{
-				// Not using anything from this scene, so unlock it.
-				SceneLocks.UnlockRead(AsyncScene, PST_Async);
-			}
+				NumAsyncHits = AsyncScene->overlapMultiple(PGeom, PGeomPose, PAsyncOverlapArray, AsyncBufferSize, PQueryFilterData, &PQueryCallback);
 
-			NumHits += NumAsyncHits;
+				if (NumAsyncHits == -1)
+				{
+					UE_LOG(LogCollision, Log, TEXT("GeomOverlapMulti : Buffer overflow from asynchronous scene query"));
+					UE_LOG(LogCollision, Warning, TEXT("--------TraceChannel : %d"), (int32)TraceChannel);
+					UE_LOG(LogCollision, Warning, TEXT("--------%s"), *Params.ToString());
+					// overlapMultiple still fills in the buffer with an arbitrary set of overlapping objects, so we have a full buffer of results
+					// we can use, though that is not the entire set
+					NumAsyncHits = AsyncBufferSize;
+				}
+				else if (NumAsyncHits == 0)
+				{
+					// Not using anything from this scene, so unlock it.
+					SceneLocks.UnlockRead(AsyncScene, PST_Async);
+				}
+
+				if (InfoType == EQueryInfo::IsBlocking)	//in the case where we just want to know if we're blocking, quickly check if anything blocks and return
+				{
+					for (PxI32 HitIdx = 0; HitIdx < NumAsyncHits; ++HitIdx)
+					{
+						if (IsBlocking(PAsyncOverlapArray[HitIdx].shape, PFilter))
+						{
+							return true;
+						}
+					}
+				}
+
+				NumHits += NumAsyncHits;
+			}
 		}
 
-		if (NumHits > 0)
+		if (InfoType == EQueryInfo::GatherAll)	//if we are gathering all we need to actually convert to UE format
 		{
-			bHaveBlockingHit = ConvertOverlapResults(NumHits, POverlapArray, PFilter, OutOverlaps);
-		}		
-		
+			if (NumHits > 0)
+			{
+				bHaveBlockingHit = ConvertOverlapResults(NumHits, POverlapArray, PFilter, OutOverlaps);
+			}
+
 #if !(UE_BUILD_SHIPPING || UE_BUILD_TEST)
-		if ((World->DebugDrawTraceTag != NAME_None) && (World->DebugDrawTraceTag == Params.TraceTag))
-		{
-			DrawGeomOverlaps(World, PGeom, PGeomPose, OutOverlaps, DebugLineLifetime);
-		}
+			if ((World->DebugDrawTraceTag != NAME_None) && (World->DebugDrawTraceTag == Params.TraceTag))
+			{
+				DrawGeomOverlaps(World, PGeom, PGeomPose, OutOverlaps, DebugLineLifetime);
+			}
 #endif //!(UE_BUILD_SHIPPING || UE_BUILD_TEST)
+		}
 	}
 	else
 	{
@@ -1792,10 +1706,16 @@ bool GeomOverlapMulti_PhysX(const UWorld* World, const PxGeometry& PGeom, const
 
 	return bHaveBlockingHit;
 }
-#endif
 
+bool GeomOverlapMulti_PhysX(const UWorld* World, const PxGeometry& PGeom, const PxTransform& PGeomPose, TArray<FOverlapResult>& OutOverlaps, ECollisionChannel TraceChannel, const struct FCollisionQueryParams& Params, const struct FCollisionResponseParams& ResponseParams, const struct FCollisionObjectQueryParams& ObjectParams)
+{
+	return GeomOverlapMultiImp_PhysX<EQueryInfo::GatherAll>(World, PGeom, PGeomPose, OutOverlaps, TraceChannel, Params, ResponseParams, ObjectParams);
+}
 
-bool GeomOverlapMulti(const UWorld* World, const struct FCollisionShape& CollisionShape, const FVector& Pos, const FQuat& Rot, TArray<FOverlapResult>& OutOverlaps, ECollisionChannel TraceChannel, const struct FCollisionQueryParams& Params, const struct FCollisionResponseParams& ResponseParams, const struct FCollisionObjectQueryParams& ObjectParams)
+#endif
+
+template <EQueryInfo::Type InfoType>
+bool GeomOverlapMultiImp(const UWorld* World, const struct FCollisionShape& CollisionShape, const FVector& Pos, const FQuat& Rot, TArray<FOverlapResult>& OutOverlaps, ECollisionChannel TraceChannel, const struct FCollisionQueryParams& Params, const struct FCollisionResponseParams& ResponseParams, const struct FCollisionObjectQueryParams& ObjectParams)
 {
 	if ((World == NULL) || (World->GetPhysicsScene() == NULL))
 	{
@@ -1810,7 +1730,7 @@ bool GeomOverlapMulti(const UWorld* World, const struct FCollisionShape& Collisi
 	FPhysXShapeAdaptor ShapeAdaptor(Rot, CollisionShape);
 	const PxGeometry& PGeom = ShapeAdaptor.GetGeometry();
 	const PxTransform& PGeomPose = ShapeAdaptor.GetGeomPose(Pos);
-	bHaveBlockingHit = GeomOverlapMulti_PhysX(World, PGeom, PGeomPose, OutOverlaps, TraceChannel, Params, ResponseParams, ObjectParams);
+	bHaveBlockingHit = GeomOverlapMultiImp_PhysX<InfoType>(World, PGeom, PGeomPose, OutOverlaps, TraceChannel, Params, ResponseParams, ObjectParams);
 
 #endif // WITH_PHYSX
 
@@ -1818,6 +1738,24 @@ bool GeomOverlapMulti(const UWorld* World, const struct FCollisionShape& Collisi
 
 	return bHaveBlockingHit;
 }
+
+bool GeomOverlapMulti(const UWorld* World, const struct FCollisionShape& CollisionShape, const FVector& Pos, const FQuat& Rot, TArray<FOverlapResult>& OutOverlaps, ECollisionChannel TraceChannel, const struct FCollisionQueryParams& Params, const struct FCollisionResponseParams& ResponseParams, const struct FCollisionObjectQueryParams& ObjectParams)
+{
+	OutOverlaps.Empty();
+	return GeomOverlapMultiImp<EQueryInfo::GatherAll>(World, CollisionShape, Pos, Rot, OutOverlaps, TraceChannel, Params, ResponseParams, ObjectParams);
+}
+
+bool GeomOverlapBlockingTest(const UWorld* World, const struct FCollisionShape& CollisionShape, const FVector& Pos, const FQuat& Rot, ECollisionChannel TraceChannel, const struct FCollisionQueryParams& Params, const struct FCollisionResponseParams& ResponseParams, const struct FCollisionObjectQueryParams& ObjectParams)
+{
+	TArray<FOverlapResult> Overlaps;	//needed only for template shared code
+	return GeomOverlapMultiImp<EQueryInfo::IsBlocking>(World, CollisionShape, Pos, Rot, Overlaps, TraceChannel, Params, ResponseParams, ObjectParams);
+}
+
+bool GeomOverlapAnyTest(const UWorld* World, const struct FCollisionShape& CollisionShape, const FVector& Pos, const FQuat& Rot, ECollisionChannel TraceChannel, const struct FCollisionQueryParams& Params, const struct FCollisionResponseParams& ResponseParams, const struct FCollisionObjectQueryParams& ObjectParams)
+{
+	TArray<FOverlapResult> Overlaps;	//needed only for template shared code
+	return GeomOverlapMultiImp<EQueryInfo::IsAnything>(World, CollisionShape, Pos, Rot, Overlaps, TraceChannel, Params, ResponseParams, ObjectParams);
+}
 #endif //UE_WITH_PHYSICS
 
 //////////////////////////////////////////////////////////////////////////
diff --git a/Engine/Source/Runtime/Engine/Private/Collision/PhysXCollision.h b/Engine/Source/Runtime/Engine/Private/Collision/PhysXCollision.h
--- a/Engine/Source/Runtime/Engine/Private/Collision/PhysXCollision.h
+++ b/Engine/Source/Runtime/Engine/Private/Collision/PhysXCollision.h
@@ -178,15 +178,14 @@ bool RaycastMulti(const UWorld* World, TArray<struct FHitResult>& OutHits, const
 
 // GEOM OVERLAP
 
-/** Function for testing overlap between a supplied PxGeometry and the world. */
-bool GeomOverlapTest(const UWorld* World, const struct FCollisionShape& CollisionShape, const FVector& Pos, const FQuat& Rot, ECollisionChannel TraceChannel, const struct FCollisionQueryParams& Params, const struct FCollisionResponseParams& ResponseParams, const struct FCollisionObjectQueryParams& ObjectParams = FCollisionObjectQueryParams::DefaultObjectQueryParam);
+/** Function for testing overlaps between a supplied PxGeometry and the world. Returns true if at least one overlapping shape is blocking*/
+bool GeomOverlapBlockingTest(const UWorld* World, const struct FCollisionShape& CollisionShape, const FVector& Pos, const FQuat& Rot, ECollisionChannel TraceChannel, const struct FCollisionQueryParams& Params, const struct FCollisionResponseParams& ResponseParams, const struct FCollisionObjectQueryParams& ObjectParams = FCollisionObjectQueryParams::DefaultObjectQueryParam);
 
-/** Function for testing overlap between a supplied PxGeometry and the world. */
-bool GeomOverlapSingle(const UWorld* World, const struct FCollisionShape& CollisionShape, const FVector& Pos, const FQuat& Rot, FOverlapResult& OutOverlaps, ECollisionChannel TraceChannel, const struct FCollisionQueryParams& Params, const struct FCollisionResponseParams& ResponseParams, const struct FCollisionObjectQueryParams& ObjectParams = FCollisionObjectQueryParams::DefaultObjectQueryParam);
+/** Function for testing overlaps between a supplied PxGeometry and the world. Returns true if anything is overlapping (blocking or touching)*/
+bool GeomOverlapAnyTest(const UWorld* World, const struct FCollisionShape& CollisionShape, const FVector& Pos, const FQuat& Rot, ECollisionChannel TraceChannel, const struct FCollisionQueryParams& Params, const struct FCollisionResponseParams& ResponseParams, const struct FCollisionObjectQueryParams& ObjectParams = FCollisionObjectQueryParams::DefaultObjectQueryParam);
 
 /**
-* Function for overlapping a supplied PxGeometry against the world
-* Note that this does not clear OutOverlaps, but adds unique results to it
+* Function for overlapping a supplied PxGeometry against the world.
 */
 bool GeomOverlapMulti(const UWorld* World, const struct FCollisionShape& CollisionShape, const FVector& Pos, const FQuat& Rot, TArray<FOverlapResult>& OutOverlaps, ECollisionChannel TraceChannel, const struct FCollisionQueryParams& Params, const struct FCollisionResponseParams& ResponseParams, const struct FCollisionObjectQueryParams& ObjectParams = FCollisionObjectQueryParams::DefaultObjectQueryParam);
 
diff --git a/Engine/Source/Runtime/Engine/Private/Collision/WorldCollision.cpp b/Engine/Source/Runtime/Engine/Private/Collision/WorldCollision.cpp
--- a/Engine/Source/Runtime/Engine/Private/Collision/WorldCollision.cpp
+++ b/Engine/Source/Runtime/Engine/Private/Collision/WorldCollision.cpp
@@ -54,7 +54,7 @@ ECollisionChannel DefaultCollisionChannel = (ECollisionChannel) 0;
 
 /* Set functions for each Shape type */
 void FBaseTraceDatum::Set(UWorld * World, const FCollisionShape& InCollisionShape, const FCollisionQueryParams& Param, const struct FCollisionResponseParams &InResponseParam, const struct FCollisionObjectQueryParams& InObjectQueryParam,
-	ECollisionChannel Channel, uint32 InUserData, bool bInIsMultiTrace, int32 FrameCounter)
+	ECollisionChannel Channel, uint32 InUserData, int32 FrameCounter)
 {
 	ensure(World);
 	CollisionParams.CollisionShape = InCollisionShape;
@@ -63,7 +63,6 @@ void FBaseTraceDatum::Set(UWorld * World, const FCollisionShape& InCollisionShap
 	CollisionParams.ObjectQueryParam = InObjectQueryParam;
 	TraceChannel = Channel;
 	UserData = InUserData;
-	bIsMultiTrace = bInIsMultiTrace;
 	FrameNumber = FrameCounter;
 	PhysWorld = World;
 }
@@ -153,20 +152,41 @@ bool UWorld::SweepMulti(TArray<FHitResult>& OutHits, const FVector& Start, const
 
 bool UWorld::OverlapTest(const FVector& Pos, const FQuat& Rot, ECollisionChannel TraceChannel, const struct FCollisionShape& CollisionShape, const struct FCollisionQueryParams& Params, const struct FCollisionResponseParams& ResponseParam) const
 {
+	return OverlapBlockingTest(Pos, Rot, TraceChannel, CollisionShape, Params, ResponseParam);
+}
+
+bool UWorld::OverlapBlockingTest(const FVector& Pos, const FQuat& Rot, ECollisionChannel TraceChannel, const struct FCollisionShape& CollisionShape, const struct FCollisionQueryParams& Params, const struct FCollisionResponseParams& ResponseParam) const
+{
+	bool bBlocking = false;
 #if UE_WITH_PHYSICS
-	return GeomOverlapTest(this, CollisionShape, Pos, Rot, TraceChannel, Params, ResponseParam, FCollisionObjectQueryParams::DefaultObjectQueryParam);
-#else
-	return false;
+	bBlocking = GeomOverlapBlockingTest(this, CollisionShape, Pos, Rot, TraceChannel, Params, ResponseParam, FCollisionObjectQueryParams::DefaultObjectQueryParam);
+#endif
+	return bBlocking;
+}
+
+bool UWorld::OverlapAnyTest(const FVector& Pos, const FQuat& Rot, ECollisionChannel TraceChannel, const struct FCollisionShape& CollisionShape, const struct FCollisionQueryParams& Params, const struct FCollisionResponseParams& ResponseParam) const
+{
+	bool bBlocking = false;
+#if UE_WITH_PHYSICS
+	bBlocking = GeomOverlapAnyTest(this, CollisionShape, Pos, Rot, TraceChannel, Params, ResponseParam, FCollisionObjectQueryParams::DefaultObjectQueryParam);
 #endif
+	return bBlocking;
 }
 
 bool UWorld::OverlapSingle(struct FOverlapResult& OutOverlap, const FVector& Pos, const FQuat& Rot, ECollisionChannel TraceChannel, const struct FCollisionShape& CollisionShape, const struct FCollisionQueryParams& Params, const struct FCollisionResponseParams& ResponseParam) const
 {
+	bool bBlocking = false;
 #if UE_WITH_PHYSICS
-	return GeomOverlapSingle(this, CollisionShape, Pos, Rot, OutOverlap, TraceChannel, Params, ResponseParam, FCollisionObjectQueryParams::DefaultObjectQueryParam);
-#else
-	return false;
+	TArray<FOverlapResult> Overlaps;
+	return GeomOverlapMulti(this, CollisionShape, Pos, Rot, Overlaps, TraceChannel, Params, ResponseParam, FCollisionObjectQueryParams::DefaultObjectQueryParam);
+	if (Overlaps.Num())
+	{
+		OutOverlap = Overlaps[0];
+		bBlocking = Overlaps[0].bBlockingHit;
+	}
 #endif
+
+	return bBlocking;
 }
 
 bool UWorld::OverlapMulti(TArray<struct FOverlapResult>& OutOverlaps, const FVector& Pos, const FQuat& Rot, ECollisionChannel TraceChannel, const struct FCollisionShape& CollisionShape, const struct FCollisionQueryParams& Params, const struct FCollisionResponseParams& ResponseParam) const
@@ -280,20 +300,32 @@ bool UWorld::SweepMulti(TArray<struct FHitResult>& OutHits, const FVector& Start
 
 bool UWorld::OverlapTest(const FVector& Pos, const FQuat& Rot, const struct FCollisionShape& CollisionShape, const struct FCollisionQueryParams& Params, const struct FCollisionObjectQueryParams& ObjectQueryParams) const
 {
+	return OverlapAnyTest(Pos, Rot, CollisionShape, Params, ObjectQueryParams);
+}
+
+bool UWorld::OverlapAnyTest(const FVector& Pos, const FQuat& Rot, const struct FCollisionShape& CollisionShape, const struct FCollisionQueryParams& Params, const struct FCollisionObjectQueryParams& ObjectQueryParams) const
+{
+	bool bBlocking = false;
 #if UE_WITH_PHYSICS
-	return GeomOverlapTest(this, CollisionShape, Pos, Rot, DefaultCollisionChannel, Params, FCollisionResponseParams::DefaultResponseParam, ObjectQueryParams);
-#else
-	return false;
+	bBlocking = GeomOverlapAnyTest(this, CollisionShape, Pos, Rot, DefaultCollisionChannel, Params, FCollisionResponseParams::DefaultResponseParam, ObjectQueryParams);
 #endif
+	return bBlocking;
 }
 
 bool UWorld::OverlapSingle(struct FOverlapResult& OutOverlap, const FVector& Pos, const FQuat& Rot, const struct FCollisionShape& CollisionShape, const struct FCollisionQueryParams& Params, const struct FCollisionObjectQueryParams& ObjectQueryParams) const
 {
+	bool bBlocking = false;
 #if UE_WITH_PHYSICS
-	return GeomOverlapSingle(this, CollisionShape, Pos, Rot, OutOverlap, DefaultCollisionChannel, Params, FCollisionResponseParams::DefaultResponseParam, ObjectQueryParams);
-#else
-	return false;
+	TArray<FOverlapResult> Overlaps;
+	GeomOverlapMulti(this, CollisionShape, Pos, Rot, Overlaps, DefaultCollisionChannel, Params, FCollisionResponseParams::DefaultResponseParam, ObjectQueryParams);
+	if (Overlaps.Num() > 0)
+	{
+		OutOverlap = Overlaps[0];
+		bBlocking = true;
+	}
 #endif
+
+	return bBlocking;
 }
 
 // profile interfaces
@@ -366,21 +398,41 @@ bool UWorld::SweepMultiByProfile(TArray<FHitResult>& OutHits, const FVector& Sta
 }
 
 bool UWorld::OverlapTestByProfile(const FVector& Pos, const FQuat& Rot, FName ProfileName, const struct FCollisionShape& CollisionShape, const struct FCollisionQueryParams& Params) const
+{
+	return OverlapBlockingTestByProfile(Pos, Rot, ProfileName, CollisionShape, Params);
+}
+
+bool UWorld::OverlapBlockingTestByProfile(const FVector& Pos, const FQuat& Rot, FName ProfileName, const struct FCollisionShape& CollisionShape, const struct FCollisionQueryParams& Params) const
 {
 	ECollisionChannel TraceChannel;
 	FCollisionResponseParams ResponseParam;
 	GetCollisionProfileChannelAndResponseParams(ProfileName, TraceChannel, ResponseParam);
 
-	return OverlapTest(Pos, Rot, TraceChannel, CollisionShape, Params, ResponseParam);
+	return OverlapBlockingTest(Pos, Rot, TraceChannel, CollisionShape, Params, ResponseParam);
 }
 
-bool UWorld::OverlapSingleByProfile(struct FOverlapResult& OutOverlap, const FVector& Pos, const FQuat& Rot, FName ProfileName, const struct FCollisionShape& CollisionShape, const struct FCollisionQueryParams& Params) const
+bool UWorld::OverlapAnyTestByProfile(const FVector& Pos, const FQuat& Rot, FName ProfileName, const struct FCollisionShape& CollisionShape, const struct FCollisionQueryParams& Params) const
 {
 	ECollisionChannel TraceChannel;
 	FCollisionResponseParams ResponseParam;
 	GetCollisionProfileChannelAndResponseParams(ProfileName, TraceChannel, ResponseParam);
 
-	return OverlapSingle(OutOverlap, Pos, Rot, TraceChannel, CollisionShape, Params, ResponseParam);
+	return OverlapAnyTest(Pos, Rot, TraceChannel, CollisionShape, Params, ResponseParam);
+}
+
+bool UWorld::OverlapSingleByProfile(struct FOverlapResult& OutOverlap, const FVector& Pos, const FQuat& Rot, FName ProfileName, const struct FCollisionShape& CollisionShape, const struct FCollisionQueryParams& Params) const
+{
+	TArray<FOverlapResult> Overlaps;
+	OverlapMultiByProfile(Overlaps, Pos, Rot, ProfileName, CollisionShape, Params);
+
+	bool bBlocking = false;
+	if (Overlaps.Num())
+	{
+		bBlocking = Overlaps[0].bBlockingHit;
+		OutOverlap = Overlaps[0];
+	}
+
+	return bBlocking;
 }
 
 bool UWorld::OverlapMultiByProfile(TArray<struct FOverlapResult>& OutOverlaps, const FVector& Pos, const FQuat& Rot, FName ProfileName, const struct FCollisionShape& CollisionShape, const struct FCollisionQueryParams& Params) const
diff --git a/Engine/Source/Runtime/Engine/Private/Collision/WorldCollisionAsync.cpp b/Engine/Source/Runtime/Engine/Private/Collision/WorldCollisionAsync.cpp
--- a/Engine/Source/Runtime/Engine/Private/Collision/WorldCollisionAsync.cpp
+++ b/Engine/Source/Runtime/Engine/Private/Collision/WorldCollisionAsync.cpp
@@ -145,34 +145,16 @@ namespace
 				continue;
 			}
 
-			if (OverlapData.bIsMultiTrace)
-			{
-				GeomOverlapMulti(
-					OverlapData.PhysWorld.Get(),
-					OverlapData.CollisionParams.CollisionShape,
-					OverlapData.Pos,
-					OverlapData.Rot,
-					OverlapData.OutOverlaps,
-					OverlapData.TraceChannel,
-					OverlapData.CollisionParams.CollisionQueryParam,
-					OverlapData.CollisionParams.ResponseParam,
-					OverlapData.CollisionParams.ObjectQueryParam);
-			}
-			else
-			{
-				OverlapData.OutOverlaps.AddZeroed(1);
-
-				GeomOverlapSingle(
-					OverlapData.PhysWorld.Get(),
-					OverlapData.CollisionParams.CollisionShape,
-					OverlapData.Pos,
-					OverlapData.Rot,
-					OverlapData.OutOverlaps[0],
-					OverlapData.TraceChannel,
-					OverlapData.CollisionParams.CollisionQueryParam,
-					OverlapData.CollisionParams.ResponseParam,
-					OverlapData.CollisionParams.ObjectQueryParam);
-			}
+			GeomOverlapMulti(
+				OverlapData.PhysWorld.Get(),
+				OverlapData.CollisionParams.CollisionShape,
+				OverlapData.Pos,
+				OverlapData.Rot,
+				OverlapData.OutOverlaps,
+				OverlapData.TraceChannel,
+				OverlapData.CollisionParams.CollisionQueryParam,
+				OverlapData.CollisionParams.ResponseParam,
+				OverlapData.CollisionParams.ObjectQueryParam);
 		}
 	#endif //UE_WITH_PHYSICS
 	}
@@ -333,14 +315,14 @@ FTraceHandle UWorld::AsyncSweep(const FVector& Start,const FVector& End,const st
 }
 
 // overlap functions
-FTraceHandle UWorld::AsyncOverlap(const FVector& Pos, const FQuat& Rot, ECollisionChannel TraceChannel,const struct FCollisionShape & CollisionShape, const struct FCollisionQueryParams& Params, const struct FCollisionResponseParams& ResponseParam, FOverlapDelegate * InDelegate, uint32 UserData, bool bMultiTrace)
+FTraceHandle UWorld::AsyncOverlap(const FVector& Pos, const FQuat& Rot, ECollisionChannel TraceChannel,const struct FCollisionShape & CollisionShape, const struct FCollisionQueryParams& Params, const struct FCollisionResponseParams& ResponseParam, FOverlapDelegate * InDelegate, uint32 UserData)
 {
-	return StartNewTrace(AsyncTraceState, FOverlapDatum(this, CollisionShape, Params, ResponseParam, FCollisionObjectQueryParams::DefaultObjectQueryParam, TraceChannel, UserData, bMultiTrace, Pos, Rot, InDelegate, AsyncTraceState.CurrentFrame));
+	return StartNewTrace(AsyncTraceState, FOverlapDatum(this, CollisionShape, Params, ResponseParam, FCollisionObjectQueryParams::DefaultObjectQueryParam, TraceChannel, UserData, Pos, Rot, InDelegate, AsyncTraceState.CurrentFrame));
 }
 
-FTraceHandle UWorld::AsyncOverlap(const FVector& Pos, const FQuat& Rot, const struct FCollisionShape & CollisionShape, const struct FCollisionQueryParams& Params, const struct FCollisionObjectQueryParams& ObjectQueryParams, FOverlapDelegate * InDelegate, uint32 UserData, bool bMultiTrace)
+FTraceHandle UWorld::AsyncOverlap(const FVector& Pos, const FQuat& Rot, const struct FCollisionShape & CollisionShape, const struct FCollisionQueryParams& Params, const struct FCollisionObjectQueryParams& ObjectQueryParams, FOverlapDelegate * InDelegate, uint32 UserData)
 {
-	return StartNewTrace(AsyncTraceState, FOverlapDatum(this, CollisionShape, Params, FCollisionResponseParams::DefaultResponseParam, ObjectQueryParams, DefaultCollisionChannel, UserData, bMultiTrace, Pos, Rot, InDelegate, AsyncTraceState.CurrentFrame));
+	return StartNewTrace(AsyncTraceState, FOverlapDatum(this, CollisionShape, Params, FCollisionResponseParams::DefaultResponseParam, ObjectQueryParams, DefaultCollisionChannel, UserData, Pos, Rot, InDelegate, AsyncTraceState.CurrentFrame));
 }
 
 bool UWorld::IsTraceHandleValid(const FTraceHandle& Handle, bool bOverlapTrace)
diff --git a/Engine/Source/Runtime/Engine/Public/WorldCollision.h b/Engine/Source/Runtime/Engine/Public/WorldCollision.h
--- a/Engine/Source/Runtime/Engine/Public/WorldCollision.h
+++ b/Engine/Source/Runtime/Engine/Public/WorldCollision.h
@@ -262,14 +262,11 @@ struct FBaseTraceDatum
 	/** User data */
 	uint32	UserData;
 
-	/** single or multi trace. */
-	bool 	bIsMultiTrace;
-
 	FBaseTraceDatum() {};
 
 	/** Set functions for each Shape type */
 	void Set(UWorld * World, const FCollisionShape& InCollisionShape, const FCollisionQueryParams& Param, const struct FCollisionResponseParams &InResponseParam, const struct FCollisionObjectQueryParams& InObjectQueryParam, 
-		ECollisionChannel Channel, uint32 InUserData, bool bInIsMultiTrace, int32 FrameCounter);
+		ECollisionChannel Channel, uint32 InUserData, int32 FrameCounter);
 };
 
 struct FTraceDatum;
@@ -314,13 +311,16 @@ struct FTraceDatum : public FBaseTraceDatum
 	/** Output of the overlap request. Filled up by worker thread */
 	TArray<struct FHitResult> OutHits;
 
+	/** single or multi trace. */
+	bool 	bIsMultiTrace;
+
 	FTraceDatum() {}
 
 	/** Set Trace Datum for each shape type **/
 	FTraceDatum(UWorld * World, const FCollisionShape& CollisionShape, const FCollisionQueryParams& Param, const struct FCollisionResponseParams &InResponseParam, const struct FCollisionObjectQueryParams& InObjectQueryParam, 
 		ECollisionChannel Channel, uint32 InUserData, bool bInIsMultiTrace,	const FVector& InStart, const FVector& InEnd, FTraceDelegate * InDelegate, int32 FrameCounter)
 	{
-		Set(World, CollisionShape, Param, InResponseParam, InObjectQueryParam, Channel, InUserData, bInIsMultiTrace, FrameCounter);
+		Set(World, CollisionShape, Param, InResponseParam, InObjectQueryParam, Channel, InUserData, FrameCounter);
 		Start = InStart;
 		End = InEnd;
 		if (InDelegate)
@@ -332,6 +332,7 @@ struct FTraceDatum : public FBaseTraceDatum
 			Delegate.Unbind();
 		}
 		OutHits.Reset();
+		bIsMultiTrace = bInIsMultiTrace;
 	}
 };
 
@@ -359,10 +360,10 @@ struct FOverlapDatum : FBaseTraceDatum
 	FOverlapDatum() {}
 
 	FOverlapDatum(UWorld * World, const FCollisionShape& CollisionShape, const FCollisionQueryParams& Param, const struct FCollisionResponseParams &InResponseParam, const struct FCollisionObjectQueryParams& InObjectQueryParam, 
-		ECollisionChannel Channel, uint32 InUserData, bool bInIsMultiTrace,
+		ECollisionChannel Channel, uint32 InUserData,
 		const FVector& InPos, const FQuat& InRot, FOverlapDelegate * InDelegate, int32 FrameCounter)
 	{
-		Set(World, CollisionShape, Param, InResponseParam, InObjectQueryParam, Channel, InUserData, bInIsMultiTrace, FrameCounter);
+		Set(World, CollisionShape, Param, InResponseParam, InObjectQueryParam, Channel, InUserData, FrameCounter);
 		Pos = InPos;
 		Rot = InRot;
 		if (InDelegate)
diff --git a/Engine/Source/Runtime/Foliage/Private/InstancedFoliage.cpp b/Engine/Source/Runtime/Foliage/Private/InstancedFoliage.cpp
--- a/Engine/Source/Runtime/Foliage/Private/InstancedFoliage.cpp
+++ b/Engine/Source/Runtime/Foliage/Private/InstancedFoliage.cpp
@@ -2128,7 +2128,7 @@ bool AInstancedFoliageActor::CheckCollisionWithWorld(const UWorld* InWorld, cons
 			FBoxSphereBounds WorldBound = ShrinkBound.TransformBy(InstTransform);
 			//::DrawDebugBox(World, WorldBound.Origin, WorldBound.BoxExtent, FColor::Red, true, 10.f);
 			static FName NAME_FoliageCollisionWithWorld = FName(TEXT("FoliageCollisionWithWorld"));
-			if (InWorld->OverlapTest(WorldBound.Origin, FQuat(Inst.Rotation), FCollisionShape::MakeBox(ShrinkBound.BoxExtent * Inst.DrawScale3D * Settings->CollisionScale), FCollisionQueryParams(NAME_FoliageCollisionWithWorld, false), FCollisionObjectQueryParams(ECC_WorldStatic)))
+			if (InWorld->OverlapAnyTest(WorldBound.Origin, FQuat(Inst.Rotation), FCollisionShape::MakeBox(ShrinkBound.BoxExtent * Inst.DrawScale3D * Settings->CollisionScale), FCollisionQueryParams(NAME_FoliageCollisionWithWorld, false), FCollisionObjectQueryParams(ECC_WorldStatic)))
 			{
 				return false;
 			}
