UMG - Deprecating PlaySound on UserWidget.  There's a new PlaySound2D available everywhere, it no longer needs to be associated with a specific widget.

diff --git a/Engine/Source/Runtime/Engine/Classes/Kismet/GameplayStatics.h b/Engine/Source/Runtime/Engine/Classes/Kismet/GameplayStatics.h
--- a/Engine/Source/Runtime/Engine/Classes/Kismet/GameplayStatics.h
+++ b/Engine/Source/Runtime/Engine/Classes/Kismet/GameplayStatics.h
@@ -1,4 +1,4 @@
-// Copyright 1998-2015 Epic Games, Inc. All Rights Reserved.
+﻿// Copyright 1998-2015 Epic Games, Inc. All Rights Reserved.
 
 #pragma once
 #include "Engine/LatentActionManager.h"
@@ -274,10 +274,22 @@ class ENGINE_API UGameplayStatics : public UBlueprintFunctionLibrary
 	UFUNCTION(BlueprintCallable, Category="Audio")
 	static bool AreAnyListenersWithinRange(FVector Location, float MaximumRange);
 
-	/** Plays a sound at the given location. This is a fire and forget sound and does not travel with any actor. Replication is also not handled at this point.
+	/**
+	 * Plays a sound directly with no attenuation, perfect for UI sounds.
+	 *
+	 * ● Fire and Forget.
+	 * ● Not Replicated.
+	 * @param Sound - Sound to play.
+	 * @param VolumeMultiplier - Multiplied with the volume to make the sound louder or softer.
+	 * @param PitchMultiplier - Multiplies the pitch.
+	 */
+	UFUNCTION(BlueprintCallable, BlueprintCosmetic, Category="Audio", meta=( WorldContext="WorldContextObject", AdvancedDisplay = "2", UnsafeDuringActorConstruction = "true" ))
+	static void PlaySound2D(UObject* WorldContextObject, class USoundBase* Sound, float VolumeMultiplier = 1.f, float PitchMultiplier = 1.f, float StartTime = 0.f);
+
+	/**
+	 * Plays a sound at the given location. This is a fire and forget sound and does not travel with any actor. Replication is also not handled at this point.
 	 * @param Sound - sound to play
 	 * @param Location - World position to play sound at
-	 * @param World - The World in which the sound is to be played
 	 * @param VolumeMultiplier - Volume multiplier 
 	 * @param PitchMultiplier - PitchMultiplier
 	 * @param AttenuationSettings - Override attenuation settings package to play sound with
diff --git a/Engine/Source/Runtime/Engine/Private/GameplayStatics.cpp b/Engine/Source/Runtime/Engine/Private/GameplayStatics.cpp
--- a/Engine/Source/Runtime/Engine/Private/GameplayStatics.cpp
+++ b/Engine/Source/Runtime/Engine/Private/GameplayStatics.cpp
@@ -623,9 +623,42 @@ bool UGameplayStatics::AreAnyListenersWithinRange(FVector Location, float Maximu
 	return false;
 }
 
+void UGameplayStatics::PlaySound2D(UObject* WorldContextObject, class USoundBase* Sound, float VolumeMultiplier, float PitchMultiplier, float StartTime)
+{
+	if ( Sound == nullptr )
+	{
+		return;
+	}
+
+	UWorld* ThisWorld = GEngine->GetWorldFromContextObject(WorldContextObject);
+	if ( ThisWorld == nullptr )
+	{
+		return;
+	}
+
+	if ( GEngine && GEngine->UseSound() && ThisWorld->bAllowAudioPlayback && ThisWorld->GetNetMode() != NM_DedicatedServer )
+	{
+		// TODO - Audio Threading. This call would be a task call to dispatch to the audio thread
+		if ( FAudioDevice* AudioDevice = GEngine->GetAudioDevice() )
+		{
+			FActiveSound NewActiveSound;
+			NewActiveSound.Sound = Sound;
+
+			NewActiveSound.VolumeMultiplier = VolumeMultiplier;
+			NewActiveSound.PitchMultiplier = PitchMultiplier;
+
+			NewActiveSound.RequestedStartTime = FMath::Max(0.f, StartTime);
+
+			NewActiveSound.bIsUISound = true;
+
+			AudioDevice->AddNewActiveSound(NewActiveSound);
+		}
+	}
+}
+
 void UGameplayStatics::PlaySoundAtLocation(UObject* WorldContextObject, class USoundBase* Sound, FVector Location, float VolumeMultiplier, float PitchMultiplier, float StartTime, class USoundAttenuation* AttenuationSettings)
 {
-	if ( Sound == NULL )
+	if ( Sound == nullptr )
 	{
 		return;
 	}
diff --git a/Engine/Source/Runtime/UMG/Public/Blueprint/UserWidget.h b/Engine/Source/Runtime/UMG/Public/Blueprint/UserWidget.h
--- a/Engine/Source/Runtime/UMG/Public/Blueprint/UserWidget.h
+++ b/Engine/Source/Runtime/UMG/Public/Blueprint/UserWidget.h
@@ -578,7 +578,7 @@ public:
 	 *
 	 * @param The sound to play
 	 */
-	UFUNCTION(BlueprintCallable, BlueprintCosmetic, Category="Sound")
+	UFUNCTION(BlueprintCallable, BlueprintCosmetic, Category="Sound", meta=( DeprecatedFunction, DeprecationMessage="Use the global function PlaySound2D instead." ))
 	void PlaySound(class USoundBase* SoundToPlay);
 
 	/** @returns The UObject wrapper for a given SWidget */
