#ANIM: import - Make sure it marks as dirty when bones are merged

Merging

//depot/UE4-GDC2015/Engine/Source/Runtime/Engine/Private/Animation/Skeleton.cpp

to //depot/UE4-GDC2015/Engine/Source/Runtime/Engine/Private/Animation/Skeleton.cpp

diff --git a/Engine/Source/Runtime/Engine/Private/Animation/Skeleton.cpp b/Engine/Source/Runtime/Engine/Private/Animation/Skeleton.cpp
--- a/Engine/Source/Runtime/Engine/Private/Animation/Skeleton.cpp
+++ b/Engine/Source/Runtime/Engine/Private/Animation/Skeleton.cpp
@@ -558,7 +558,6 @@ bool USkeleton::MergeBonesToBoneTree(USkeletalMesh * InSkeletalMesh, const TArra
 
 					ReferenceSkeleton.Add(NewMeshBoneInfo, InSkeletalMesh->RefSkeleton.GetRefBonePose()[MeshBoneIndex]);
 					BoneTree.AddZeroed(1);
-					MarkPackageDirty();
 				}
 			}
 
@@ -569,8 +568,13 @@ bool USkeleton::MergeBonesToBoneTree(USkeletalMesh * InSkeletalMesh, const TArra
 	// if succeed
 	if (bSuccess)
 	{
-		InSkeletalMesh->Skeleton = this;
-		InSkeletalMesh->MarkPackageDirty();
+		if (InSkeletalMesh->Skeleton != this)
+		{
+			InSkeletalMesh->Skeleton = this;
+			InSkeletalMesh->MarkPackageDirty();
+		}
+
+		MarkPackageDirty();
 		// make sure to refresh all base poses
 		// so that they have same number of bones of ref pose
 #if WITH_EDITORONLY_DATA
