Update incorrect comments

diff --git a/Engine/Source/Runtime/Engine/Classes/Engine/World.h b/Engine/Source/Runtime/Engine/Classes/Engine/World.h
--- a/Engine/Source/Runtime/Engine/Classes/Engine/World.h
+++ b/Engine/Source/Runtime/Engine/Classes/Engine/World.h
@@ -1248,7 +1248,7 @@ public:
 	*  Test the collision of a shape at the supplied location using a specific channel, and return if any blocking overlap is found
 	*  @param  Pos             Location of center of box to test against the world
 	*  @param  TraceChannel    The 'channel' that this query is in, used to determine which components to hit
-	*  @param  CollisionShape	CollisionShape - supports Box, Sphere, Capsule
+	*  @param  CollisionShape	CollisionShape - supports Box, Sphere, Capsule, Convex
 	*  @param  Params          Additional parameters used for the trace
 	*  @param  ResponseParam	ResponseContainer to be used for this trace
 	*  @return TRUE if any blocking results are found
@@ -1259,7 +1259,7 @@ public:
 	*  Test the collision of a shape at the supplied location using a specific channel, and return if any blocking or overlapping shape is found
 	*  @param  Pos             Location of center of box to test against the world
 	*  @param  TraceChannel    The 'channel' that this query is in, used to determine which components to hit
-	*  @param  CollisionShape	CollisionShape - supports Box, Sphere, Capsule
+	*  @param  CollisionShape	CollisionShape - supports Box, Sphere, Capsule, Convex
 	*  @param  Params          Additional parameters used for the trace
 	*  @param  ResponseParam	ResponseContainer to be used for this trace
 	*  @return TRUE if any blocking or overlapping results are found
@@ -1280,7 +1280,7 @@ public:
 	/**
 	*  Test the collision of a shape at the supplied location using object types, and return if any overlap is found
 	*  @param  Pos             Location of center of box to test against the world
-	*  @param  CollisionShape	CollisionShape - supports Box, Sphere, Capsule
+	*  @param  CollisionShape	CollisionShape - supports Box, Sphere, Capsule, Convex
 	*  @param  Params          Additional parameters used for the trace
 	*  @param  ObjectQueryParams	List of object types it's looking for
 	*  @return TRUE if any blocking results are found
diff --git a/Engine/Source/Runtime/Engine/Private/Collision/CollisionConversions.h b/Engine/Source/Runtime/Engine/Private/Collision/CollisionConversions.h
--- a/Engine/Source/Runtime/Engine/Private/Collision/CollisionConversions.h
+++ b/Engine/Source/Runtime/Engine/Private/Collision/CollisionConversions.h
@@ -71,7 +71,6 @@ inline bool IsBlocking(const PxShape* PShape, const PxFilterData& QueryFilter)
 {
 	// See if this is a 'blocking' hit
 	const PxFilterData PShapeFilter = PShape->getQueryFilterData();
-	// word1 is block, word2 is overlap
 	const PxSceneQueryHitType::Enum HitType = FPxQueryFilterCallback::CalcQueryHitType(QueryFilter, PShapeFilter);
 	const bool bBlock = (HitType == PxSceneQueryHitType::eBLOCK);
 	return bBlock;
