#ue4 - Partial revert of CL 2373230, which broke sweeps/traces returning physical material for hits against Box, Sphere, and Capsule shapes.

PhysX internally handles the invalid face index for the trimesh/heightfield case (ignores and prints a warning). Originally we were trying to reduce the log spam in these cases.

diff --git a/Engine/Source/Runtime/Engine/Private/Collision/CollisionConversions.cpp b/Engine/Source/Runtime/Engine/Private/Collision/CollisionConversions.cpp
--- a/Engine/Source/Runtime/Engine/Private/Collision/CollisionConversions.cpp
+++ b/Engine/Source/Runtime/Engine/Private/Collision/CollisionConversions.cpp
@@ -293,7 +293,7 @@ static void SetHitResultFromShapeAndFaceIndex(const PxShape* PShape,  const PxRi
 		OutResult.Actor = OwningComponent->GetOwner();
 		OutResult.Component = OwningComponent;
 
-		if (bReturnPhysMat && FaceIndex != InvalidQueryHit.faceIndex)
+		if (bReturnPhysMat)
 		{
 			// @fixme: only do this for InGameThread, otherwise, this will be done in AsyncTrace
 			if ( IsInGameThread() )
