#ue4 - Use new OverlapBlockingTest() in MovementComponent code.

diff --git a/Engine/Source/Runtime/Engine/Private/CharacterMovementComponent.cpp b/Engine/Source/Runtime/Engine/Private/CharacterMovementComponent.cpp
--- a/Engine/Source/Runtime/Engine/Private/CharacterMovementComponent.cpp
+++ b/Engine/Source/Runtime/Engine/Private/CharacterMovementComponent.cpp
@@ -1710,8 +1710,7 @@ void UCharacterMovementComponent::Crouch(bool bClientSimulation)
 			FCollisionQueryParams CapsuleParams(NAME_CrouchTrace, false, CharacterOwner);
 			FCollisionResponseParams ResponseParam;
 			InitCollisionParams(CapsuleParams, ResponseParam);
-			TArray<FOverlapResult> Overlaps;
-			const bool bEncroached = GetWorld()->OverlapMulti(Overlaps, CharacterOwner->GetActorLocation() - FVector(0.f,0.f,ScaledHalfHeightAdjust), FQuat::Identity,
+			const bool bEncroached = GetWorld()->OverlapBlockingTest(CharacterOwner->GetActorLocation() - FVector(0.f,0.f,ScaledHalfHeightAdjust), FQuat::Identity,
 				UpdatedComponent->GetCollisionObjectType(), GetPawnCapsuleCollisionShape(SHRINK_None), CapsuleParams, ResponseParam);
 
 			// If encroached, cancel
@@ -1787,8 +1786,7 @@ void UCharacterMovementComponent::UnCrouch(bool bClientSimulation)
 		if (!bCrouchMaintainsBaseLocation)
 		{
 			// Expand in place
-			TArray<FOverlapResult> Overlaps;
-			bEncroached = GetWorld()->OverlapMulti(Overlaps, PawnLocation, FQuat::Identity, CollisionChannel, StandingCapsuleShape, CapsuleParams, ResponseParam);
+			bEncroached = GetWorld()->OverlapBlockingTest(PawnLocation, FQuat::Identity, CollisionChannel, StandingCapsuleShape, CapsuleParams, ResponseParam);
 		
 			if (bEncroached)
 			{
@@ -1814,8 +1812,7 @@ void UCharacterMovementComponent::UnCrouch(bool bClientSimulation)
 						// Compute where the base of the sweep ended up, and see if we can stand there
 						const float DistanceToBase = (Hit.Time * TraceDist) + ShortCapsuleShape.Capsule.HalfHeight;
 						const FVector NewLoc = FVector(PawnLocation.X, PawnLocation.Y, PawnLocation.Z - DistanceToBase + PawnHalfHeight + SweepInflation + MIN_FLOOR_DIST / 2.f);
-						Overlaps.Reset(); // OverlapMulti doesn't reset it, though it should.
-						bEncroached = GetWorld()->OverlapMulti(Overlaps, NewLoc, FQuat::Identity, CollisionChannel, StandingCapsuleShape, CapsuleParams, ResponseParam);
+						bEncroached = GetWorld()->OverlapBlockingTest(NewLoc, FQuat::Identity, CollisionChannel, StandingCapsuleShape, CapsuleParams, ResponseParam);
 						if (!bEncroached)
 						{
 							// Intentionally not using MoveUpdatedComponent, where a horizontal plane constraint would prevent the base of the capsule from staying at the same spot.
@@ -1829,8 +1826,7 @@ void UCharacterMovementComponent::UnCrouch(bool bClientSimulation)
 		{
 			// Expand while keeping base location the same.
 			FVector StandingLocation = PawnLocation + FVector(0.f, 0.f, StandingCapsuleShape.GetCapsuleHalfHeight() - CurrentCrouchedHalfHeight);
-			TArray<FOverlapResult> Overlaps;
-			bEncroached = GetWorld()->OverlapMulti(Overlaps, StandingLocation, FQuat::Identity, CollisionChannel, StandingCapsuleShape, CapsuleParams, ResponseParam);
+			bEncroached = GetWorld()->OverlapBlockingTest(StandingLocation, FQuat::Identity, CollisionChannel, StandingCapsuleShape, CapsuleParams, ResponseParam);
 
 			if (bEncroached)
 			{
@@ -1841,8 +1837,7 @@ void UCharacterMovementComponent::UnCrouch(bool bClientSimulation)
 					if (CurrentFloor.bBlockingHit && CurrentFloor.FloorDist > MinFloorDist)
 					{
 						StandingLocation.Z -= CurrentFloor.FloorDist - MinFloorDist;
-						Overlaps.Reset(); // OverlapMulti doesn't reset it, though it should.
-						bEncroached = GetWorld()->OverlapMulti(Overlaps, StandingLocation, FQuat::Identity, CollisionChannel, StandingCapsuleShape, CapsuleParams, ResponseParam);
+						bEncroached = GetWorld()->OverlapBlockingTest(StandingLocation, FQuat::Identity, CollisionChannel, StandingCapsuleShape, CapsuleParams, ResponseParam);
 					}
 				}				
 			}
diff --git a/Engine/Source/Runtime/Engine/Private/MovementComponent.cpp b/Engine/Source/Runtime/Engine/Private/MovementComponent.cpp
--- a/Engine/Source/Runtime/Engine/Private/MovementComponent.cpp
+++ b/Engine/Source/Runtime/Engine/Private/MovementComponent.cpp
@@ -314,8 +314,7 @@ bool UMovementComponent::OverlapTest(const FVector& Location, const FQuat& Rotat
 	FCollisionQueryParams QueryParams(NAME_TestOverlap, false, IgnoreActor);
 	FCollisionResponseParams ResponseParam;
 	InitCollisionParams(QueryParams, ResponseParam);
-	TArray<FOverlapResult> Overlaps;
-	return GetWorld()->OverlapMulti(Overlaps, Location, RotationQuat, CollisionChannel, CollisionShape, QueryParams, ResponseParam);
+	return GetWorld()->OverlapBlockingTest(Location, RotationQuat, CollisionChannel, CollisionShape, QueryParams, ResponseParam);
 }
 
 bool UMovementComponent::IsExceedingMaxSpeed(float MaxSpeed) const
