Fix checks for packages saved with the 4.7 editor which aren't compatible with main. The 4.7 editor also writes packages with a later object version (which overlaps with a version from main) so we have to test against the engine version instead.

diff --git a/Engine/Source/Runtime/Engine/Private/AI/Navigation/PImplRecastNavMesh.cpp b/Engine/Source/Runtime/Engine/Private/AI/Navigation/PImplRecastNavMesh.cpp
--- a/Engine/Source/Runtime/Engine/Private/AI/Navigation/PImplRecastNavMesh.cpp
+++ b/Engine/Source/Runtime/Engine/Private/AI/Navigation/PImplRecastNavMesh.cpp
@@ -366,7 +366,7 @@ void FPImplRecastNavMesh::Serialize( FArchive& Ar, int32 NavMeshVersion )
 
 					//
 					if (Ar.UE4Ver() >= VER_UE4_ADD_MODIFIERS_RUNTIME_GENERATION && 
-						Ar.UE4Ver() != VER_UE4_MERGED_ADD_MODIFIERS_RUNTIME_GENERATION_TO_4_7) // Merged package from 4.7 branch
+						(Ar.EngineVer().GetMajor() != 4 || Ar.EngineVer().GetMinor() != 7)) // Merged package from 4.7 branch
 					{
 						unsigned char* ComressedTileData = NULL;
 						int32 CompressedTileDataSize = 0;
@@ -407,7 +407,7 @@ void FPImplRecastNavMesh::Serialize( FArchive& Ar, int32 NavMeshVersion )
 
 					// Serialize compressed tile cache layer
 					if (Ar.UE4Ver() >= VER_UE4_ADD_MODIFIERS_RUNTIME_GENERATION &&
-						Ar.UE4Ver() != VER_UE4_MERGED_ADD_MODIFIERS_RUNTIME_GENERATION_TO_4_7) // Merged package from 4.7 branch
+						(Ar.EngineVer().GetMajor() != 4 || Ar.EngineVer().GetMinor() != 7)) // Merged package from 4.7 branch
 					{
 						uint8* ComressedTileData = nullptr;
 						int32 CompressedTileDataSize = 0;
diff --git a/Engine/Source/Runtime/Engine/Private/AI/Navigation/RecastNavMeshDataChunk.cpp b/Engine/Source/Runtime/Engine/Private/AI/Navigation/RecastNavMeshDataChunk.cpp
--- a/Engine/Source/Runtime/Engine/Private/AI/Navigation/RecastNavMeshDataChunk.cpp
+++ b/Engine/Source/Runtime/Engine/Private/AI/Navigation/RecastNavMeshDataChunk.cpp
@@ -118,7 +118,7 @@ void URecastNavMeshDataChunk::SerializeRecastData(FArchive& Ar, int32 NavMeshVer
 				int32 TileCacheDataSize = 0;
 				uint8* TileCacheRawData = nullptr;
 				if (Ar.UE4Ver() >= VER_UE4_ADD_MODIFIERS_RUNTIME_GENERATION && 
-					Ar.UE4Ver() != VER_UE4_MERGED_ADD_MODIFIERS_RUNTIME_GENERATION_TO_4_7) // Merged package from 4.7 branch
+					(Ar.EngineVer().GetMajor() != 4 || Ar.EngineVer().GetMinor() != 7)) // Merged package from 4.7 branch
 				{
 					FPImplRecastNavMesh::SerializeCompressedTileCacheData(Ar, NavMeshVersion, TileCacheRawData, TileCacheDataSize);
 				}
