Fix non-unity build

diff --git a/Engine/Source/Runtime/Engine/Private/Collision/CollisionConversions.cpp b/Engine/Source/Runtime/Engine/Private/Collision/CollisionConversions.cpp
--- a/Engine/Source/Runtime/Engine/Private/Collision/CollisionConversions.cpp
+++ b/Engine/Source/Runtime/Engine/Private/Collision/CollisionConversions.cpp
@@ -891,6 +891,15 @@ static void AddUniqueOverlap(TArray<FOverlapResult>& OutOverlaps, const FOverlap
 	OutOverlaps.Add(NewOverlap);
 }
 
+bool IsBlocking(const PxShape* PShape, const PxFilterData& QueryFilter)
+{
+	// See if this is a 'blocking' hit
+	const PxFilterData PShapeFilter = PShape->getQueryFilterData();
+	const PxSceneQueryHitType::Enum HitType = FPxQueryFilterCallback::CalcQueryHitType(QueryFilter, PShapeFilter);
+	const bool bBlock = (HitType == PxSceneQueryHitType::eBLOCK);
+	return bBlock;
+}
+
 bool ConvertOverlapResults(int32 NumOverlaps, PxOverlapHit* POverlapResults, const PxFilterData& QueryFilter, TArray<FOverlapResult>& OutOverlaps)
 {
 	SCOPE_CYCLE_COUNTER(STAT_CollisionConvertOverlap);
diff --git a/Engine/Source/Runtime/Engine/Private/Collision/CollisionConversions.h b/Engine/Source/Runtime/Engine/Private/Collision/CollisionConversions.h
--- a/Engine/Source/Runtime/Engine/Private/Collision/CollisionConversions.h
+++ b/Engine/Source/Runtime/Engine/Private/Collision/CollisionConversions.h
@@ -67,14 +67,7 @@ void ConvertQueryOverlap(const PxShape* PShape, const PxRigidActor* PActor, FOve
  * @QueryFilter Query Filter
  * @return true if the query filter and shape filter resolve to be blocking
  */
-inline bool IsBlocking(const PxShape* PShape, const PxFilterData& QueryFilter)
-{
-	// See if this is a 'blocking' hit
-	const PxFilterData PShapeFilter = PShape->getQueryFilterData();
-	const PxSceneQueryHitType::Enum HitType = FPxQueryFilterCallback::CalcQueryHitType(QueryFilter, PShapeFilter);
-	const bool bBlock = (HitType == PxSceneQueryHitType::eBLOCK);
-	return bBlock;
-}
+bool IsBlocking(const PxShape* PShape, const PxFilterData& QueryFilter);
 
 
 /** 
