Move ComputePenetrationDepth into FBodyInstance. Actually used for the overlap test now with the option to output MTD

diff --git a/Engine/Source/Runtime/Engine/Classes/PhysicsEngine/BodyInstance.h b/Engine/Source/Runtime/Engine/Classes/PhysicsEngine/BodyInstance.h
--- a/Engine/Source/Runtime/Engine/Classes/PhysicsEngine/BodyInstance.h
+++ b/Engine/Source/Runtime/Engine/Classes/PhysicsEngine/BodyInstance.h
@@ -683,9 +683,10 @@ public:
 	 *
 	 *	@param	PGeom			Geometry it would like to test
 	 *  @param  ShapePose       Transform information in world. Use U2PTransform to convert from FTransform
+	 *  @param  OutMTD			The minimum translation direction needed to push the shape out of this BodyInstance. (Optional)
 	 *  @return true if PrimComp overlaps this component at the specified location/rotation
 	 */
-	bool OverlapPhysX(const physx::PxGeometry& Geom, const physx::PxTransform&  ShapePose) const;
+	bool OverlapPhysX(const physx::PxGeometry& Geom, const physx::PxTransform&  ShapePose, FMTDResult* OutMTD = nullptr) const;
 #endif	//WITH_PHYSX
 
 	/**
@@ -694,9 +695,10 @@ public:
 	 *  @param  Position		Position to place the shape at before testing
 	 *  @param  Rotation		Rotation to apply to the shape before testing
 	 *	@param	CollisionShape	Shape to test against
+	 *  @param  OutMTD			The minimum translation direction needed to push the shape out of this BodyInstance. (Optional)
 	 *  @return true if the geometry associated with this body instance overlaps the query shape at the specified location/rotation
 	 */
-	bool OverlapTest(const FVector& Position, const FQuat& Rotation, const struct FCollisionShape& CollisionShape) const;
+	bool OverlapTest(const FVector& Position, const FQuat& Rotation, const struct FCollisionShape& CollisionShape, FMTDResult* OutMTD = nullptr) const;
 
 	/**
 	 *  Determines the set of components that this body instance would overlap with at the supplied location/rotation
diff --git a/Engine/Source/Runtime/Engine/Private/PhysicsEngine/BodyInstance.cpp b/Engine/Source/Runtime/Engine/Private/PhysicsEngine/BodyInstance.cpp
--- a/Engine/Source/Runtime/Engine/Private/PhysicsEngine/BodyInstance.cpp
+++ b/Engine/Source/Runtime/Engine/Private/PhysicsEngine/BodyInstance.cpp
@@ -4193,13 +4193,13 @@ float FBodyInstance::GetDistanceToBody(const FVector& Point, FVector& OutPointOn
 	//@TODO: BOX2D: Implement DistanceToBody
 }
 
-bool FBodyInstance::OverlapTest(const FVector& Position, const FQuat& Rotation, const struct FCollisionShape& CollisionShape) const
+bool FBodyInstance::OverlapTest(const FVector& Position, const FQuat& Rotation, const struct FCollisionShape& CollisionShape, FMTDResult* OutMTD) const
 {
 	bool bHasOverlap = false;
 
 #if WITH_PHYSX
 	FPhysXShapeAdaptor ShapeAdaptor(Rotation, CollisionShape);
-	bHasOverlap = OverlapPhysX(ShapeAdaptor.GetGeometry(), ShapeAdaptor.GetGeomPose(Position));
+	bHasOverlap = OverlapPhysX(ShapeAdaptor.GetGeometry(), ShapeAdaptor.GetGeomPose(Position), OutMTD);
 #endif
 
 #if WITH_BOX2D
@@ -4316,7 +4316,7 @@ bool FBodyInstance::OverlapMulti(TArray<struct FOverlapResult>& InOutOverlaps, c
 
 
 #if WITH_PHYSX
-bool FBodyInstance::OverlapPhysX(const PxGeometry& PGeom, const PxTransform& ShapePose) const
+bool FBodyInstance::OverlapPhysX(const PxGeometry& PGeom, const PxTransform& ShapePose, FMTDResult* OutMTD) const
 {
 	const PxRigidActor* RigidBody = WeldParent ? WeldParent->GetPxRigidActor() : GetPxRigidActor();
 
@@ -4338,10 +4338,18 @@ bool FBodyInstance::OverlapPhysX(const PxGeometry& PGeom, const PxTransform& Sha
 
 		if (ShapeBoundToBody(PShape, this) == true)
 		{
-			if (PxGeometryQuery::overlap(PShape->getGeometry().any(), PxShapeExt::getGlobalPose(*PShape, *RigidBody), PGeom, ShapePose))
+			PxVec3 POutDirection;
+			float OutDistance;
+			if (PxGeometryQuery::computePenetration(POutDirection, OutDistance, PGeom, ShapePose, PShape->getGeometry().any(), PxShapeExt::getGlobalPose(*PShape, *RigidBody)))
 			{
+				if(OutMTD)
+				{
+					OutMTD->Direction = P2UVector(POutDirection);
+					OutMTD->Distance = OutDistance;
+				}
+					
 				return true;
-			}
+			}			
 		}
 	}
 	return false;
diff --git a/Engine/Source/Runtime/Engine/Private/PrimitiveComponent.cpp b/Engine/Source/Runtime/Engine/Private/PrimitiveComponent.cpp
--- a/Engine/Source/Runtime/Engine/Private/PrimitiveComponent.cpp
+++ b/Engine/Source/Runtime/Engine/Private/PrimitiveComponent.cpp
@@ -1802,59 +1802,12 @@ bool UPrimitiveComponent::OverlapComponent(const FVector& Pos, const FQuat& Rot,
 	return BodyInstance.OverlapTest(Pos, Rot, CollisionShape);
 }
 
-bool UPrimitiveComponent::ComputePenetration(FMTDResult & OutMTD, const FCollisionShape & CollisionShape, const FVector & Pos, const FQuat & Rot)
+bool UPrimitiveComponent::ComputePenetration(FMTDResult& OutMTD, const FCollisionShape & CollisionShape, const FVector& Pos, const FQuat& Rot)
 {
-#if WITH_PHYSX
-	UCollision2PGeom GeomStorage0(CollisionShape);
-	PxTransform PGeomPose0 = ConvertToPhysXCapsulePose(FTransform(Rot, Pos));
-	const PxGeometry * PGeom0 = GeomStorage0.GetGeometry();
-
-	check(PGeom0);	//couldn't convert FCollisionShape to PxGeometry - something is wrong
-
-	const PxRigidActor* PRigidBody = BodyInstance.GetPxRigidActor();
-	if (PRigidBody == NULL || PRigidBody->getNbShapes() == 0)
-	{
-		return false;
-	}
-
-	const PxTransform PGlobalPose1 = PRigidBody->getGlobalPose();
-
-	// Get all the shapes from the actor
-	// TODO: we should really pass the shape from the overlap info since doing it this way we do an overlap test twice
-	TArray<PxShape*> PShapes;
-	PShapes.AddZeroed(PRigidBody->getNbShapes());
-	int32 NumTargetShapes = PRigidBody->getShapes(PShapes.GetData(), PShapes.Num());
-
-	for (int32 PShapeIdx = 0; PShapeIdx < PShapes.Num(); ++PShapeIdx)
+	if(FBodyInstance* BodyInstance = GetBodyInstance())
 	{
-		const PxShape * PShape = PShapes[PShapeIdx];
-		check(PShape);
-
-		// Calc shape global pose
-		PxTransform PGeomPose1 = PGlobalPose1.transform(PShape->getLocalPose());
-		GeometryFromShapeStorage GeomStorage1;
-		PxGeometry * PGeom1 = GetGeometryFromShape(GeomStorage1, PShape, true);
-
-		if (PGeom1)
-		{
-			PxVec3 POutDirection;
-			bool bSuccess = PxGeometryQuery::computePenetration(POutDirection, OutMTD.Distance, *PGeom0, PGeomPose0, *PGeom1, PGeomPose1);
-			if (bSuccess)
-			{
-				if (POutDirection.isFinite())
-				{
-					OutMTD.Direction = P2UVector(POutDirection);
-					return true;
-				}
-				else
-				{
-					UE_LOG(LogPhysics, Warning, TEXT("UPrimitiveComponent::ComputePenetration: MTD returned NaN"));
-					return false;
-				}
-			}
-		}
+		return BodyInstance->OverlapTest(Pos, Rot, CollisionShape, &OutMTD);
 	}
-#endif
 
 	return false;
 }
