Fix skeletal mesh AddRadialForce and AddRadialImpulse to average out the strength based on body mass. PR - 821

diff --git a/Engine/Source/Runtime/Engine/Private/SkeletalMeshComponentPhysics.cpp b/Engine/Source/Runtime/Engine/Private/SkeletalMeshComponentPhysics.cpp
--- a/Engine/Source/Runtime/Engine/Private/SkeletalMeshComponentPhysics.cpp
+++ b/Engine/Source/Runtime/Engine/Private/SkeletalMeshComponentPhysics.cpp
@@ -474,9 +474,11 @@ void USkeletalMeshComponent::AddRadialImpulse(FVector Origin, float Radius, floa
 		return;
 	}
 
+	const float StrengthPerMass = Strength / FMath::Max(GetMass(), KINDA_SMALL_NUMBER);
 	for(int32 i=0; i<Bodies.Num(); i++)
 	{
-		Bodies[i]->AddRadialImpulseToBody(Origin, Radius, Strength, Falloff, bVelChange);
+		const float StrengthPerBody = bVelChange ? Strength : StrengthPerMass * Bodies[i]->GetBodyMass();
+		Bodies[i]->AddRadialImpulseToBody(Origin, Radius, StrengthPerBody, Falloff, bVelChange);
 	}
 }
 
@@ -489,9 +491,11 @@ void USkeletalMeshComponent::AddRadialForce(FVector Origin, float Radius, float
 		return;
 	}
 
+	const float StrengthPerMass = Strength / FMath::Max(GetMass(), KINDA_SMALL_NUMBER);
 	for(int32 i=0; i<Bodies.Num(); i++)
 	{
-		Bodies[i]->AddRadialForceToBody(Origin, Radius, Strength, Falloff);
+		const float StrengthPerBody = StrengthPerMass * Bodies[i]->GetBodyMass();
+		Bodies[i]->AddRadialForceToBody(Origin, Radius, StrengthPerBody, Falloff);
 	}
 
 }
