[INTEGRATE] Change 2443306 by [REDACTED_PERSON][REDACTED_HANDLE] on 2015/02/12 08:36:41

	Fixes for compiling with -stresstestunity. Will probably be enabled for Rocket builds; reduces lib size by ~450mb.

diff --git a/Engine/Source/Runtime/Engine/Classes/GameFramework/Volume.h b/Engine/Source/Runtime/Engine/Classes/GameFramework/Volume.h
--- a/Engine/Source/Runtime/Engine/Classes/GameFramework/Volume.h
+++ b/Engine/Source/Runtime/Engine/Classes/GameFramework/Volume.h
@@ -4,6 +4,8 @@
 #include "Engine/Brush.h"
 #include "Volume.generated.h"
 
+DECLARE_LOG_CATEGORY_EXTERN(LogVolume, Log, All);
+
 /** 
  *	An editable 3D volume placed in a level. Different types of volumes perform different functions
  *	@see https://docs.unrealengine.com/latest/INT/Engine/Actors/Volumes
diff --git a/Engine/Source/Runtime/Engine/Private/AI/Navigation/RecastNavMesh.cpp b/Engine/Source/Runtime/Engine/Private/AI/Navigation/RecastNavMesh.cpp
--- a/Engine/Source/Runtime/Engine/Private/AI/Navigation/RecastNavMesh.cpp
+++ b/Engine/Source/Runtime/Engine/Private/AI/Navigation/RecastNavMesh.cpp
@@ -135,8 +135,6 @@ void ARecastNavMesh::Serialize( FArchive& Ar )
 
 #define DO_NAVMESH_DEBUG_DRAWING_PER_TILE 0
 
-static const FColor NavMeshRenderColor_RecastMesh(72,255,64);
-
 //----------------------------------------------------------------------//
 // FRecastDebugGeometry
 //----------------------------------------------------------------------//
diff --git a/Engine/Source/Runtime/Engine/Private/Actor.cpp b/Engine/Source/Runtime/Engine/Private/Actor.cpp
--- a/Engine/Source/Runtime/Engine/Private/Actor.cpp
+++ b/Engine/Source/Runtime/Engine/Private/Actor.cpp
@@ -1743,13 +1743,10 @@ float AActor::InternalTakePointDamage(float Damage, FPointDamageEvent const& Poi
 }
 
 /** Util to check if prim comp pointer is valid and still alive */
-static bool IsPrimCompValidAndAlive(UPrimitiveComponent* PrimComp)
-{
-	return (PrimComp != NULL) && !PrimComp->IsPendingKill();
-}
+extern bool IsPrimCompValidAndAlive(UPrimitiveComponent* PrimComp);
 
 /** Used to determine if it is ok to call a notification on this object */
-static bool IsActorValidToNotify(AActor* Actor)
+bool IsActorValidToNotify(AActor* Actor)
 {
 	return (Actor != NULL) && !Actor->IsPendingKill() && !Actor->GetClass()->HasAnyClassFlags(CLASS_NewerVersionExists);
 }
diff --git a/Engine/Source/Runtime/Engine/Private/Collision/PhysXCollision.cpp b/Engine/Source/Runtime/Engine/Private/Collision/PhysXCollision.cpp
--- a/Engine/Source/Runtime/Engine/Private/Collision/PhysXCollision.cpp
+++ b/Engine/Source/Runtime/Engine/Private/Collision/PhysXCollision.cpp
@@ -476,7 +476,7 @@ private:
 #define VERIFY_GEOMSWEEPSINGLE 0
 #define VERIFY_GEOMSWEEPMULTI 0
 
-static float DebugLineLifetime = 2.f;
+float DebugLineLifetime = 2.f;
 
 /**
  * Helper to lock/unlock multiple scenes that also makes sure to unlock everything when it goes out of scope.
diff --git a/Engine/Source/Runtime/Engine/Private/DataReplication.cpp b/Engine/Source/Runtime/Engine/Private/DataReplication.cpp
--- a/Engine/Source/Runtime/Engine/Private/DataReplication.cpp
+++ b/Engine/Source/Runtime/Engine/Private/DataReplication.cpp
@@ -47,7 +47,7 @@ public:
 	UNetDriver * Driver;
 };
 
-static bool IsCustomDeltaProperty( UProperty * Property )
+bool IsCustomDeltaProperty( UProperty * Property )
 {
 	UStructProperty * StructProperty = Cast< UStructProperty >( Property );
 
diff --git a/Engine/Source/Runtime/Engine/Private/DefaultPhysicsVolume.cpp b/Engine/Source/Runtime/Engine/Private/DefaultPhysicsVolume.cpp
--- a/Engine/Source/Runtime/Engine/Private/DefaultPhysicsVolume.cpp
+++ b/Engine/Source/Runtime/Engine/Private/DefaultPhysicsVolume.cpp
@@ -3,8 +3,6 @@
 #include "EnginePrivate.h"
 #include "GameFramework/DefaultPhysicsVolume.h"
 
-DEFINE_LOG_CATEGORY_STATIC(LogVolume, Log, All);
-
 ADefaultPhysicsVolume::ADefaultPhysicsVolume(const FObjectInitializer& ObjectInitializer)
 	: Super(ObjectInitializer)
 {
diff --git a/Engine/Source/Runtime/Engine/Private/GameEngine.cpp b/Engine/Source/Runtime/Engine/Private/GameEngine.cpp
--- a/Engine/Source/Runtime/Engine/Private/GameEngine.cpp
+++ b/Engine/Source/Runtime/Engine/Private/GameEngine.cpp
@@ -38,8 +38,6 @@
 #include "GameFramework/GameUserSettings.h"
 #include "GameFramework/GameMode.h"
 
-DEFINE_LOG_CATEGORY_STATIC(LogEngine, Log, All);
-
 ENGINE_API bool GDisallowNetworkTravel = false;
 
 /** Benchmark results to the log */
diff --git a/Engine/Source/Runtime/Engine/Private/LocalPlayer.cpp b/Engine/Source/Runtime/Engine/Private/LocalPlayer.cpp
--- a/Engine/Source/Runtime/Engine/Private/LocalPlayer.cpp
+++ b/Engine/Source/Runtime/Engine/Private/LocalPlayer.cpp
@@ -24,7 +24,6 @@
 #include "GameFramework/PlayerState.h"
 
 DEFINE_LOG_CATEGORY(LogPlayerManagement);
-DEFINE_LOG_CATEGORY_STATIC(LogEngine, Log, All);
 
 #if !(UE_BUILD_SHIPPING || UE_BUILD_TEST)
 
diff --git a/Engine/Source/Runtime/Engine/Private/NetworkSettings.cpp b/Engine/Source/Runtime/Engine/Private/NetworkSettings.cpp
--- a/Engine/Source/Runtime/Engine/Private/NetworkSettings.cpp
+++ b/Engine/Source/Runtime/Engine/Private/NetworkSettings.cpp
@@ -8,8 +8,7 @@ UNetworkSettings::UNetworkSettings(const FObjectInitializer& ObjectInitializer)
 {
 }
 
-static FName ConsoleVariableFName(TEXT("ConsoleVariable"));
-static FName ToolTipFName(TEXT("ConsoleVariable"));
+static FName NetworkConsoleVariableFName(TEXT("ConsoleVariable"));
 
 void UNetworkSettings::PostInitProperties()
 {
@@ -24,7 +23,7 @@ void UNetworkSettings::PostInitProperties()
 				continue;
 			}
 
-			FString CVarName = Property->GetMetaData(ConsoleVariableFName);
+			FString CVarName = Property->GetMetaData(NetworkConsoleVariableFName);
 			if (!CVarName.IsEmpty())
 			{
 				IConsoleVariable* CVar = IConsoleManager::Get().FindConsoleVariable(*CVarName);
@@ -51,7 +50,7 @@ void UNetworkSettings::PostEditChangeProperty(FPropertyChangedEvent& PropertyCha
 	Super::PostEditChangeProperty(PropertyChangedEvent);
 	if (PropertyChangedEvent.Property)
 	{
-		FString CVarName = PropertyChangedEvent.Property->GetMetaData(ConsoleVariableFName);
+		FString CVarName = PropertyChangedEvent.Property->GetMetaData(NetworkConsoleVariableFName);
 		if (!CVarName.IsEmpty())
 		{
 			IConsoleVariable* CVar = IConsoleManager::Get().FindConsoleVariable(*CVarName);
diff --git a/Engine/Source/Runtime/Engine/Private/PrimitiveComponent.cpp b/Engine/Source/Runtime/Engine/Private/PrimitiveComponent.cpp
--- a/Engine/Source/Runtime/Engine/Private/PrimitiveComponent.cpp
+++ b/Engine/Source/Runtime/Engine/Private/PrimitiveComponent.cpp
@@ -1697,7 +1697,7 @@ void UPrimitiveComponent::SetCanEverAffectNavigation(bool bRelevant)
 //////////////////////////////////////////////////////////////////////////
 // COLLISION
 
-float DebugLineLifetime = 2.f;
+extern float DebugLineLifetime;
 
 
 bool UPrimitiveComponent::LineTraceComponent(struct FHitResult& OutHit, const FVector Start, const FVector End, const struct FCollisionQueryParams& Params)
@@ -1899,12 +1899,10 @@ bool UPrimitiveComponent::ComponentIsTouchingSelectionFrustum(const FConvexVolum
 
 
 /** Used to determine if it is ok to call a notification on this object */
-static bool IsActorValidToNotify(AActor* Actor)
-{
-	return (Actor != NULL) && !Actor->IsPendingKill() && !Actor->GetClass()->HasAnyClassFlags(CLASS_NewerVersionExists);
-}
+extern bool IsActorValidToNotify(AActor* Actor);
+
 // @fixme, duplicated, make an inline member?
-static bool IsPrimCompValidAndAlive(UPrimitiveComponent* PrimComp)
+bool IsPrimCompValidAndAlive(UPrimitiveComponent* PrimComp)
 {
 	return (PrimComp != NULL) && !PrimComp->IsPendingKill();
 }
diff --git a/Engine/Source/Runtime/Engine/Private/RenderingSettings.cpp b/Engine/Source/Runtime/Engine/Private/RenderingSettings.cpp
--- a/Engine/Source/Runtime/Engine/Private/RenderingSettings.cpp
+++ b/Engine/Source/Runtime/Engine/Private/RenderingSettings.cpp
@@ -9,8 +9,7 @@ URendererSettings::URendererSettings(const FObjectInitializer& ObjectInitializer
 	TranslucentSortAxis = FVector(0.0f, -1.0f, 0.0f);
 }
 
-static FName ConsoleVariableFName(TEXT("ConsoleVariable"));
-static FName ToolTipFName(TEXT("ConsoleVariable"));
+static FName RenderingConsoleVariableFName(TEXT("ConsoleVariable"));
 
 void URendererSettings::PostInitProperties()
 {
@@ -43,7 +42,7 @@ void URendererSettings::PostInitProperties()
 				continue;
 			}
 
-			FString CVarName = Property->GetMetaData(ConsoleVariableFName);
+			FString CVarName = Property->GetMetaData(RenderingConsoleVariableFName);
 			if (!CVarName.IsEmpty())
 			{
 				IConsoleVariable* CVar = IConsoleManager::Get().FindConsoleVariable(*CVarName);
@@ -71,7 +70,7 @@ void URendererSettings::PostEditChangeProperty(FPropertyChangedEvent& PropertyCh
 
 	if (PropertyChangedEvent.Property)
 	{
-		FString CVarName = PropertyChangedEvent.Property->GetMetaData(ConsoleVariableFName);
+		FString CVarName = PropertyChangedEvent.Property->GetMetaData(RenderingConsoleVariableFName);
 		if (!CVarName.IsEmpty())
 		{
 			IConsoleVariable* CVar = IConsoleManager::Get().FindConsoleVariable(*CVarName);
diff --git a/Engine/Source/Runtime/Engine/Private/RepLayout.cpp b/Engine/Source/Runtime/Engine/Private/RepLayout.cpp
--- a/Engine/Source/Runtime/Engine/Private/RepLayout.cpp
+++ b/Engine/Source/Runtime/Engine/Private/RepLayout.cpp
@@ -2178,17 +2178,7 @@ uint16 FRepLayout::AddParentProperty( UProperty * Property, int32 ArrayIndex )
 	return Parents.Add( FRepParentCmd( Property, ArrayIndex ) );
 }
 
-static bool IsCustomDeltaProperty( UProperty * Property )
-{
-	UStructProperty * StructProperty = Cast< UStructProperty >( Property );
-
-	if ( StructProperty != NULL && StructProperty->Struct->StructFlags & STRUCT_NetDeltaSerializeNative )
-	{
-		return true;
-	}
-
-	return false;
-}
+extern bool IsCustomDeltaProperty( UProperty * Property );
 
 void FRepLayout::InitFromObjectClass( UClass * InObjectClass )
 {
diff --git a/Engine/Source/Runtime/Engine/Private/UnrealEngine.cpp b/Engine/Source/Runtime/Engine/Private/UnrealEngine.cpp
--- a/Engine/Source/Runtime/Engine/Private/UnrealEngine.cpp
+++ b/Engine/Source/Runtime/Engine/Private/UnrealEngine.cpp
@@ -101,7 +101,7 @@
 #include "SNotificationList.h"
 #include "Engine/UserInterfaceSettings.h"
 
-DEFINE_LOG_CATEGORY_STATIC(LogEngine, Log, All);
+DEFINE_LOG_CATEGORY(LogEngine);
 
 IMPLEMENT_MODULE( FEngineModule, Engine );
 
diff --git a/Engine/Source/Runtime/Engine/Private/Volume.cpp b/Engine/Source/Runtime/Engine/Private/Volume.cpp
--- a/Engine/Source/Runtime/Engine/Private/Volume.cpp
+++ b/Engine/Source/Runtime/Engine/Private/Volume.cpp
@@ -12,7 +12,7 @@
 AVolume::FOnVolumeShapeChanged AVolume::OnVolumeShapeChanged;
 #endif
 
-DEFINE_LOG_CATEGORY_STATIC(LogVolume, Log, All);
+DEFINE_LOG_CATEGORY(LogVolume);
 
 AVolume::AVolume(const FObjectInitializer& ObjectInitializer)
 	: Super(ObjectInitializer)
diff --git a/Engine/Source/Runtime/Engine/Public/UnrealEngine.h b/Engine/Source/Runtime/Engine/Public/UnrealEngine.h
--- a/Engine/Source/Runtime/Engine/Public/UnrealEngine.h
+++ b/Engine/Source/Runtime/Engine/Public/UnrealEngine.h
@@ -8,6 +8,8 @@
 
 #include "Engine/Engine.h"
 
+DECLARE_LOG_CATEGORY_EXTERN(LogEngine, Log, All);
+
 //
 //	FLocalPlayerIterator - Iterates over local players in the game.
 //	There are no advantages to using this over GEngine->GetLocalPlayerIterator(GetWorld());
