[INTEGRATE] Change 2443447 by [REDACTED_PERSON][REDACTED_HANDLE] on 2015/02/12 10:24:17

	Fix issues with Child Actor component when used as an instance component
	Fix cases where EditReregisterContext was not getting processed due to PostEditChangeChainProperty being called instead of PostEditChangeProperty
	OnComponentCreated is now called for components after copy/pasting an Actor
	OnComponentCreated is now called when creating instance components via duplication

diff --git a/Engine/Source/Editor/UnrealEd/Private/Kismet2/ComponentEditorUtils.cpp b/Engine/Source/Editor/UnrealEd/Private/Kismet2/ComponentEditorUtils.cpp
--- a/Engine/Source/Editor/UnrealEd/Private/Kismet2/ComponentEditorUtils.cpp
+++ b/Engine/Source/Editor/UnrealEd/Private/Kismet2/ComponentEditorUtils.cpp
@@ -489,9 +489,18 @@ UActorComponent* FComponentEditorUtils::DuplicateComponent(UActorComponent* Temp
 		FName NewComponentName = *FComponentEditorUtils::GenerateValidVariableName(ComponentClass, Actor);
 
 		bool bKeepWorldLocationOnAttach = false;
-		NewCloneComponent = NewObject<UActorComponent>(Actor, ComponentClass, NewComponentName, RF_Transactional, TemplateComponent);
+
+		const bool bTemplateTransactional = TemplateComponent->HasAllFlags(RF_Transactional);
+		TemplateComponent->SetFlags(RF_Transactional);
+
+		NewCloneComponent = DuplicateObject<UActorComponent>(TemplateComponent, Actor, *NewComponentName.ToString() );
 		
-		auto NewSceneComponent = Cast<USceneComponent>(NewCloneComponent);
+		if (!bTemplateTransactional)
+		{
+			TemplateComponent->ClearFlags(RF_Transactional);
+		}
+			
+		USceneComponent* NewSceneComponent = Cast<USceneComponent>(NewCloneComponent);
 		if (NewSceneComponent)
 		{
 			// Ensure the clone doesn't think it has children
@@ -510,6 +519,8 @@ UActorComponent* FComponentEditorUtils::DuplicateComponent(UActorComponent* Temp
 			}
 		}
 
+		NewCloneComponent->OnComponentCreated();
+
 		// Add to SerializedComponents array so it gets saved
 		Actor->AddInstanceComponent(NewCloneComponent);
 		
diff --git a/Engine/Source/Runtime/Engine/Classes/Components/ActorComponent.h b/Engine/Source/Runtime/Engine/Classes/Components/ActorComponent.h
--- a/Engine/Source/Runtime/Engine/Classes/Components/ActorComponent.h
+++ b/Engine/Source/Runtime/Engine/Classes/Components/ActorComponent.h
@@ -231,6 +231,8 @@ private:
 	/** Calls OnRegister, CreateRenderState_Concurrent and CreatePhysicsState. */
 	void ExecuteRegisterEvents();
 
+	/* Utility function for each of the PostEditChange variations to call for the same behavior */
+	void ConsolidatedPostEditChange();
 protected:
 
 	friend class FComponentReregisterContextBase;
@@ -492,6 +494,7 @@ public:
 	virtual bool Modify( bool bAlwaysMarkDirty = true ) override;
 	virtual void PreEditChange(UProperty* PropertyThatWillChange) override;
 	virtual void PostEditChangeProperty(FPropertyChangedEvent& PropertyChangedEvent) override;
+	virtual void PostEditChangeChainProperty( FPropertyChangedChainEvent& PropertyChangedEvent ) override;
 	virtual void PostEditUndo() override;
 #endif // WITH_EDITOR
 	// End UObject interface.
diff --git a/Engine/Source/Runtime/Engine/Classes/Components/ChildActorComponent.h b/Engine/Source/Runtime/Engine/Classes/Components/ChildActorComponent.h
--- a/Engine/Source/Runtime/Engine/Classes/Components/ChildActorComponent.h
+++ b/Engine/Source/Runtime/Engine/Classes/Components/ChildActorComponent.h
@@ -18,7 +18,7 @@ class UChildActorComponent : public USceneComponent
 	TSubclassOf<AActor>	ChildActorClass;
 
 	/** The actor that we spawned and own */
-	UPROPERTY(BlueprintReadOnly, Category=ChildActorComponent, TextExportTransient)
+	UPROPERTY(BlueprintReadOnly, Category=ChildActorComponent, TextExportTransient, NonPIEDuplicateTransient)
 	AActor*	ChildActor;
 
 	/** We try to keep the child actor's name as best we can, so we store it off here when destroying */
@@ -27,6 +27,12 @@ class UChildActorComponent : public USceneComponent
 	/** Cached copy of the instance data when the ChildActor is destroyed to be available when needed */
 	mutable FChildActorComponentInstanceData* CachedInstanceData;
 
+	// Begin Object interface.
+#if WITH_EDITOR
+	virtual void PostEditChangeProperty(FPropertyChangedEvent& PropertyChangedEvent) override;
+#endif
+	// End Object interface.
+
 	// Begin ActorComponent interface.
 	virtual void OnComponentCreated() override;
 	virtual void OnComponentDestroyed() override;
diff --git a/Engine/Source/Runtime/Engine/Classes/GameFramework/Actor.h b/Engine/Source/Runtime/Engine/Classes/GameFramework/Actor.h
--- a/Engine/Source/Runtime/Engine/Classes/GameFramework/Actor.h
+++ b/Engine/Source/Runtime/Engine/Classes/GameFramework/Actor.h
@@ -1137,6 +1137,7 @@ public:
 	virtual void PostEditChangeProperty(FPropertyChangedEvent& PropertyChangedEvent) override;
 	virtual void PreEditUndo() override;
 	virtual void PostEditUndo() override;
+	virtual void PostEditImport() override;
 
 	struct FActorRootComponentReconstructionData
 	{
diff --git a/Engine/Source/Runtime/Engine/Private/Actor.cpp b/Engine/Source/Runtime/Engine/Private/Actor.cpp
--- a/Engine/Source/Runtime/Engine/Private/Actor.cpp
+++ b/Engine/Source/Runtime/Engine/Private/Actor.cpp
@@ -2215,16 +2215,22 @@ static void DispatchOnComponentsCreated(AActor* NewActor)
 	TInlineComponentArray<UActorComponent*> Components;
 	NewActor->GetComponents(Components);
 
-	for (int32 Idx = 0; Idx < Components.Num(); Idx++)
+	for (UActorComponent* ActorComp : Components)
 	{
-		UActorComponent* ActorComp = Components[Idx];
-		if (ActorComp != NULL && !ActorComp->bHasBeenCreated)
+		if (ActorComp && !ActorComp->bHasBeenCreated)
 		{
 			ActorComp->OnComponentCreated();
 		}
 	}
 }
 
+void AActor::PostEditImport()
+{
+	Super::PostEditImport();
+
+	DispatchOnComponentsCreated(this);
+}
+
 void AActor::PostSpawnInitialize(FVector const& SpawnLocation, FRotator const& SpawnRotation, AActor* InOwner, APawn* InInstigator, bool bRemoteOwned, bool bNoFail, bool bDeferConstruction)
 {
 	// General flow here is like so
diff --git a/Engine/Source/Runtime/Engine/Private/ActorComponent.cpp b/Engine/Source/Runtime/Engine/Private/ActorComponent.cpp
--- a/Engine/Source/Runtime/Engine/Private/ActorComponent.cpp
+++ b/Engine/Source/Runtime/Engine/Private/ActorComponent.cpp
@@ -466,7 +466,7 @@ void UActorComponent::PostEditUndo()
 	Super::PostEditUndo();
 }
 
-void UActorComponent::PostEditChangeProperty(FPropertyChangedEvent& PropertyChangedEvent)
+void UActorComponent::ConsolidatedPostEditChange()
 {
 	FComponentReregisterContext* ReregisterContext = EditReregisterContexts.FindRef(this);
 	if(ReregisterContext)
@@ -479,13 +479,27 @@ void UActorComponent::PostEditChangeProperty(FPropertyChangedEvent& PropertyChan
 	// Don't do do a full recreate in this situation, and instead simply detach.
 	if( IsPendingKill() )
 	{
-		// @todo UE4 james should this call UnregsiterComponent instead to remove itself from the RegisteteredComponents array on the owner?
+		// @todo UE4 james should this call UnregisterComponent instead to remove itself from the RegisteredComponents array on the owner?
 		ExecuteUnregisterEvents();
 		World = NULL;
 	}
+}
+
+void UActorComponent::PostEditChangeProperty(FPropertyChangedEvent& PropertyChangedEvent)
+{
+	ConsolidatedPostEditChange();
 
 	Super::PostEditChangeProperty(PropertyChangedEvent);
 }
+
+void UActorComponent::PostEditChangeChainProperty(FPropertyChangedChainEvent& PropertyChangedEvent)
+{
+	ConsolidatedPostEditChange();
+
+	Super::PostEditChangeChainProperty(PropertyChangedEvent);
+}
+
+
 #endif // WITH_EDITOR
 
 void UActorComponent::OnRegister()
diff --git a/Engine/Source/Runtime/Engine/Private/ActorEditor.cpp b/Engine/Source/Runtime/Engine/Private/ActorEditor.cpp
--- a/Engine/Source/Runtime/Engine/Private/ActorEditor.cpp
+++ b/Engine/Source/Runtime/Engine/Private/ActorEditor.cpp
@@ -263,7 +263,10 @@ void AActor::PreEditUndo()
 
 	for (UChildActorComponent* ChildActorComponent : ChildActorComponents)
 	{
-		ChildActorComponent->DestroyChildActor();
+		if (ChildActorComponent->IsCreatedByConstructionScript())
+		{
+			ChildActorComponent->DestroyChildActor();
+		}
 	}
 
 	Super::PreEditUndo();
diff --git a/Engine/Source/Runtime/Engine/Private/ChildActorComponent.cpp b/Engine/Source/Runtime/Engine/Private/ChildActorComponent.cpp
--- a/Engine/Source/Runtime/Engine/Private/ChildActorComponent.cpp
+++ b/Engine/Source/Runtime/Engine/Private/ChildActorComponent.cpp
@@ -26,6 +26,21 @@ void UChildActorComponent::OnRegister()
 	}
 }
 
+#if WITH_EDITOR
+void UChildActorComponent::PostEditChangeProperty(FPropertyChangedEvent& PropertyChangedEvent)
+{
+	Super::PostEditChangeProperty(PropertyChangedEvent);
+
+	static const FName NAME_ChildActorClass = GET_MEMBER_NAME_CHECKED(UChildActorComponent, ChildActorClass);
+
+	if (PropertyChangedEvent.Property && PropertyChangedEvent.Property->GetFName() == NAME_ChildActorClass)
+	{
+		DestroyChildActor();
+		CreateChildActor();
+	}
+}
+#endif
+
 void UChildActorComponent::OnComponentCreated()
 {
 	Super::OnComponentCreated();
@@ -190,7 +205,10 @@ void UChildActorComponent::CreateChildActor()
 				Params.bAllowDuringConstructionScript = true;
 				Params.OverrideLevel = GetOwner()->GetLevel();
 				Params.Name = ChildActorName;
-				Params.ObjectFlags &= ~RF_Transactional;
+				if (!HasAllFlags(RF_Transactional))
+				{
+					Params.ObjectFlags &= ~RF_Transactional;
+				}
 
 				// Spawn actor of desired class
 				FVector Location = GetComponentLocation();
@@ -205,6 +223,8 @@ void UChildActorComponent::CreateChildActor()
 					// Remember which actor spawned it (for selection in editor etc)
 					ChildActor->ParentComponentActor = GetOwner();
 
+					ChildActor->AttachRootComponentTo(this);
+
 					// Parts that we deferred from SpawnActor
 					ChildActor->FinishSpawning(ComponentToWorld);
 				}
