[INTEGRATE] Change 2443512 by [REDACTED_PERSON][REDACTED_HANDLE] on 2015/02/12 11:04:38

	WITH_EDITOR fix for monolithic builds.

diff --git a/Engine/Source/Runtime/Engine/Private/Actor.cpp b/Engine/Source/Runtime/Engine/Private/Actor.cpp
--- a/Engine/Source/Runtime/Engine/Private/Actor.cpp
+++ b/Engine/Source/Runtime/Engine/Private/Actor.cpp
@@ -2224,12 +2224,14 @@ static void DispatchOnComponentsCreated(AActor* NewActor)
 	}
 }
 
+#if WITH_EDITOR
 void AActor::PostEditImport()
 {
 	Super::PostEditImport();
 
 	DispatchOnComponentsCreated(this);
 }
+#endif
 
 void AActor::PostSpawnInitialize(FVector const& SpawnLocation, FRotator const& SpawnRotation, AActor* InOwner, APawn* InInstigator, bool bRemoteOwned, bool bNoFail, bool bDeferConstruction)
 {
