[INTEGRATE] Change 2443748 by [REDACTED_PERSON][REDACTED_HANDLE] on 2015/02/12 17:26:44

	Fix component properties modified by construction script getting stuck
	Rename FComponentInstanceDataBase to FActorComponentInstanceData
	Add parameter to FComponentInstanceDataCache::ApplyToActor and FActorComponentInstanceData::ApplyToComponent to differentiate at which point in the actor rerun construction script process cache data is being applied.

diff --git a/Engine/Source/Runtime/Engine/Classes/Atmosphere/AtmosphericFogComponent.h b/Engine/Source/Runtime/Engine/Classes/Atmosphere/AtmosphericFogComponent.h
--- a/Engine/Source/Runtime/Engine/Classes/Atmosphere/AtmosphericFogComponent.h
+++ b/Engine/Source/Runtime/Engine/Classes/Atmosphere/AtmosphericFogComponent.h
@@ -258,7 +258,7 @@ public:
 	ENGINE_API void ReleaseResource();
 
 	// Begin UActorComponent interface.
-	virtual FComponentInstanceDataBase* GetComponentInstanceData() const override;
+	virtual FActorComponentInstanceData* GetComponentInstanceData() const override;
 	virtual FName GetComponentInstanceDataType() const override;
 	// End UActorComponent interface.
 
diff --git a/Engine/Source/Runtime/Engine/Classes/Components/ActorComponent.h b/Engine/Source/Runtime/Engine/Classes/Components/ActorComponent.h
--- a/Engine/Source/Runtime/Engine/Classes/Components/ActorComponent.h
+++ b/Engine/Source/Runtime/Engine/Classes/Components/ActorComponent.h
@@ -476,7 +476,7 @@ public:
 	virtual void PostNetReceive() { }
 
 	/** Called before we throw away components during RerunConstructionScripts, to cache any data we wish to persist across that operation */
-	virtual class FComponentInstanceDataBase* GetComponentInstanceData() const;
+	virtual class FActorComponentInstanceData* GetComponentInstanceData() const;
 
 	/** The type of the component instance data that this component is interested in */
 	virtual FName GetComponentInstanceDataType() const;
diff --git a/Engine/Source/Runtime/Engine/Classes/Components/ChildActorComponent.h b/Engine/Source/Runtime/Engine/Classes/Components/ChildActorComponent.h
--- a/Engine/Source/Runtime/Engine/Classes/Components/ChildActorComponent.h
+++ b/Engine/Source/Runtime/Engine/Classes/Components/ChildActorComponent.h
@@ -37,7 +37,7 @@ class UChildActorComponent : public USceneComponent
 	virtual void OnComponentCreated() override;
 	virtual void OnComponentDestroyed() override;
 	virtual void OnRegister() override;
-	virtual FComponentInstanceDataBase* GetComponentInstanceData() const override;
+	virtual FActorComponentInstanceData* GetComponentInstanceData() const override;
 	virtual FName GetComponentInstanceDataType() const override;
 	// End ActorComponent interface.
 
diff --git a/Engine/Source/Runtime/Engine/Classes/Components/InstancedStaticMeshComponent.h b/Engine/Source/Runtime/Engine/Classes/Components/InstancedStaticMeshComponent.h
--- a/Engine/Source/Runtime/Engine/Classes/Components/InstancedStaticMeshComponent.h
+++ b/Engine/Source/Runtime/Engine/Classes/Components/InstancedStaticMeshComponent.h
@@ -150,7 +150,7 @@ public:
 	TArray<FBodyInstance*> InstanceBodies;
 
 	// Begin UActorComponent interface 
-	virtual FComponentInstanceDataBase* GetComponentInstanceData() const override;
+	virtual FActorComponentInstanceData* GetComponentInstanceData() const override;
 	virtual FName GetComponentInstanceDataType() const override;
 	// End UActorComponent interface 
 
diff --git a/Engine/Source/Runtime/Engine/Classes/Components/LightComponent.h b/Engine/Source/Runtime/Engine/Classes/Components/LightComponent.h
--- a/Engine/Source/Runtime/Engine/Classes/Components/LightComponent.h
+++ b/Engine/Source/Runtime/Engine/Classes/Components/LightComponent.h
@@ -277,7 +277,7 @@ public:
 	virtual bool IsReadyForFinishDestroy() override;
 	// End UObject interface.
 
-	virtual FComponentInstanceDataBase* GetComponentInstanceData() const override;
+	virtual FActorComponentInstanceData* GetComponentInstanceData() const override;
 	virtual FName GetComponentInstanceDataType() const override;
 	 void ApplyComponentInstanceData(class FPrecomputedLightInstanceData* ComponentInstanceData);
 
diff --git a/Engine/Source/Runtime/Engine/Classes/Components/PrimitiveComponent.h b/Engine/Source/Runtime/Engine/Classes/Components/PrimitiveComponent.h
--- a/Engine/Source/Runtime/Engine/Classes/Components/PrimitiveComponent.h
+++ b/Engine/Source/Runtime/Engine/Classes/Components/PrimitiveComponent.h
@@ -958,7 +958,7 @@ public:
 	virtual bool IsEditorOnly() const override;
 	virtual bool ShouldCreatePhysicsState() const override;
 	virtual bool HasValidPhysicsState() const override;
-	virtual class FComponentInstanceDataBase* GetComponentInstanceData() const override;
+	virtual class FActorComponentInstanceData* GetComponentInstanceData() const override;
 	virtual FName GetComponentInstanceDataType() const override;
 	// End UActorComponent Interface
 
@@ -1663,7 +1663,7 @@ public:
 	virtual ~FPrimitiveComponentInstanceData()
 	{}
 
-	virtual void ApplyToComponent(UActorComponent* Component) override;
+	virtual void ApplyToComponent(UActorComponent* Component, const ECacheApplyPhase CacheApplyPhase) override;
 
 	bool ContainsData() const;
 };
diff --git a/Engine/Source/Runtime/Engine/Classes/Components/SceneComponent.h b/Engine/Source/Runtime/Engine/Classes/Components/SceneComponent.h
--- a/Engine/Source/Runtime/Engine/Classes/Components/SceneComponent.h
+++ b/Engine/Source/Runtime/Engine/Classes/Components/SceneComponent.h
@@ -4,7 +4,7 @@
 #include "Components/ActorComponent.h"
 #include "Engine/EngineTypes.h"
 #include "RHIDefinitions.h"
-#include "ComponentInstanceDataCache.h" // for FComponentInstanceDataBase
+#include "ComponentInstanceDataCache.h" // for FActorComponentInstanceData
 #include "SceneComponent.generated.h"
 
 /** Overlap info consisting of the primitive and the body that is overlapping */
@@ -521,7 +521,7 @@ public:
 	virtual void DestroyComponent(bool bPromoteChildren = false) override;
 	virtual void OnComponentDestroyed() override;
 	virtual void ApplyWorldOffset(const FVector& InOffset, bool bWorldShift) override;
-	virtual class FComponentInstanceDataBase* GetComponentInstanceData() const override;
+	virtual class FActorComponentInstanceData* GetComponentInstanceData() const override;
 	virtual FName GetComponentInstanceDataType() const override;
 	// End ActorComponent interface
 
@@ -812,7 +812,7 @@ protected:
  *  Component instance cached data base class for scene components. 
  *  Stores a list of instance components attached to the 
  */
-class ENGINE_API FSceneComponentInstanceData : public FComponentInstanceDataBase
+class ENGINE_API FSceneComponentInstanceData : public FActorComponentInstanceData
 {
 public:
 	FSceneComponentInstanceData(const USceneComponent* SourceComponent);
@@ -820,7 +820,7 @@ public:
 	virtual ~FSceneComponentInstanceData()
 	{}
 
-	virtual void ApplyToComponent(UActorComponent* Component) override;
+	virtual void ApplyToComponent(UActorComponent* Component, const ECacheApplyPhase CacheApplyPhase) override;
 	virtual void FindAndReplaceInstances(const TMap<UObject*, UObject*>& OldToNewInstanceMap) override;
 
 	TArray<USceneComponent*> AttachedInstanceComponents;
diff --git a/Engine/Source/Runtime/Engine/Classes/Components/SkyLightComponent.h b/Engine/Source/Runtime/Engine/Classes/Components/SkyLightComponent.h
--- a/Engine/Source/Runtime/Engine/Classes/Components/SkyLightComponent.h
+++ b/Engine/Source/Runtime/Engine/Classes/Components/SkyLightComponent.h
@@ -148,7 +148,7 @@ class ENGINE_API USkyLightComponent : public ULightComponentBase
 	virtual bool IsReadyForFinishDestroy() override;
 	// End UObject Interface
 
-	virtual FComponentInstanceDataBase* GetComponentInstanceData() const override;
+	virtual FActorComponentInstanceData* GetComponentInstanceData() const override;
 	virtual FName GetComponentInstanceDataType() const override;
 	void ApplyComponentInstanceData(class FPrecomputedSkyLightInstanceData* ComponentInstanceData);
 
diff --git a/Engine/Source/Runtime/Engine/Classes/Components/SplineComponent.h b/Engine/Source/Runtime/Engine/Classes/Components/SplineComponent.h
--- a/Engine/Source/Runtime/Engine/Classes/Components/SplineComponent.h
+++ b/Engine/Source/Runtime/Engine/Classes/Components/SplineComponent.h
@@ -73,7 +73,7 @@ public:
 #endif
 
 	// Begin UActorComponent interface.
-	virtual FComponentInstanceDataBase* GetComponentInstanceData() const override;
+	virtual FActorComponentInstanceData* GetComponentInstanceData() const override;
 	virtual FName GetComponentInstanceDataType() const override;
 	// End UActorComponent interface.
 
diff --git a/Engine/Source/Runtime/Engine/Classes/Components/StaticMeshComponent.h b/Engine/Source/Runtime/Engine/Classes/Components/StaticMeshComponent.h
--- a/Engine/Source/Runtime/Engine/Classes/Components/StaticMeshComponent.h
+++ b/Engine/Source/Runtime/Engine/Classes/Components/StaticMeshComponent.h
@@ -259,7 +259,7 @@ public:
 #if WITH_EDITOR
 	virtual void CheckForErrors() override;
 #endif
-	virtual FComponentInstanceDataBase* GetComponentInstanceData() const override;
+	virtual FActorComponentInstanceData* GetComponentInstanceData() const override;
 	virtual FName GetComponentInstanceDataType() const override;
 	// End UActorComponent interface.
 
diff --git a/Engine/Source/Runtime/Engine/Private/ActorComponent.cpp b/Engine/Source/Runtime/Engine/Private/ActorComponent.cpp
--- a/Engine/Source/Runtime/Engine/Private/ActorComponent.cpp
+++ b/Engine/Source/Runtime/Engine/Private/ActorComponent.cpp
@@ -551,9 +551,9 @@ void UActorComponent::UninitializeComponent()
 	bHasBeenInitialized = false;
 }
 
-FComponentInstanceDataBase* UActorComponent::GetComponentInstanceData() const
+FActorComponentInstanceData* UActorComponent::GetComponentInstanceData() const
 {
-	FComponentInstanceDataBase* InstanceData = new FComponentInstanceDataBase(this);
+	FActorComponentInstanceData* InstanceData = new FActorComponentInstanceData(this);
 
 	if (!InstanceData->ContainsSavedProperties())
 	{
diff --git a/Engine/Source/Runtime/Engine/Private/ActorConstruction.cpp b/Engine/Source/Runtime/Engine/Private/ActorConstruction.cpp
--- a/Engine/Source/Runtime/Engine/Private/ActorConstruction.cpp
+++ b/Engine/Source/Runtime/Engine/Private/ActorConstruction.cpp
@@ -523,7 +523,7 @@ void AActor::ExecuteConstruction(const FTransform& Transform, const FComponentIn
 			// If we passed in cached data, we apply it now, so that the UserConstructionScript can use the updated values
 			if(InstanceDataCache)
 			{
-				InstanceDataCache->ApplyToActor(this);
+				InstanceDataCache->ApplyToActor(this, ECacheApplyPhase::PostSimpleConstructionScript);
 			}
 
 #if WITH_EDITOR
@@ -543,7 +543,7 @@ void AActor::ExecuteConstruction(const FTransform& Transform, const FComponentIn
 			// @TODO Don't re-apply to components we already applied to above
 			if (InstanceDataCache)
 			{
-				InstanceDataCache->ApplyToActor(this);
+				InstanceDataCache->ApplyToActor(this, ECacheApplyPhase::PostUserConstructionScript);
 			}
 		}
 		else
diff --git a/Engine/Source/Runtime/Engine/Private/Atmosphere/Atmosphere.cpp b/Engine/Source/Runtime/Engine/Private/Atmosphere/Atmosphere.cpp
--- a/Engine/Source/Runtime/Engine/Private/Atmosphere/Atmosphere.cpp
+++ b/Engine/Source/Runtime/Engine/Private/Atmosphere/Atmosphere.cpp
@@ -711,9 +711,9 @@ public:
 	virtual ~FAtmospherePrecomputeInstanceData()
 	{}
 
-	virtual void ApplyToComponent(UActorComponent* Component) override
+	virtual void ApplyToComponent(UActorComponent* Component, const ECacheApplyPhase CacheApplyPhase) override
 	{
-		FSceneComponentInstanceData::ApplyToComponent(Component);
+		FSceneComponentInstanceData::ApplyToComponent(Component, CacheApplyPhase);
 		CastChecked<UAtmosphericFogComponent>(Component)->ApplyComponentInstanceData(this);
 	}
 
@@ -731,9 +731,9 @@ FName UAtmosphericFogComponent::GetComponentInstanceDataType() const
 }
 
 // Backup the precomputed data before re-running Blueprint construction script
-FComponentInstanceDataBase* UAtmosphericFogComponent::GetComponentInstanceData() const
+FActorComponentInstanceData* UAtmosphericFogComponent::GetComponentInstanceData() const
 {
-	FComponentInstanceDataBase* InstanceData = nullptr;
+	FActorComponentInstanceData* InstanceData = nullptr;
 
 	if (TransmittanceData.GetElementCount() && IrradianceData.GetElementCount() && InscatterData.GetElementCount() && PrecomputeCounter.GetValue() == EValid)
 	{
diff --git a/Engine/Source/Runtime/Engine/Private/ChildActorComponent.cpp b/Engine/Source/Runtime/Engine/Private/ChildActorComponent.cpp
--- a/Engine/Source/Runtime/Engine/Private/ChildActorComponent.cpp
+++ b/Engine/Source/Runtime/Engine/Private/ChildActorComponent.cpp
@@ -86,9 +86,9 @@ public:
 		}
 	}
 
-	virtual void ApplyToComponent(UActorComponent* Component) override
+	virtual void ApplyToComponent(UActorComponent* Component, const ECacheApplyPhase CacheApplyPhase) override
 	{
-		FSceneComponentInstanceData::ApplyToComponent(Component);
+		FSceneComponentInstanceData::ApplyToComponent(Component, CacheApplyPhase);
 		CastChecked<UChildActorComponent>(Component)->ApplyComponentInstanceData(this);
 	}
 
@@ -110,7 +110,7 @@ FName UChildActorComponent::GetComponentInstanceDataType() const
 	return ChildActorComponentInstanceDataName;
 }
 
-FComponentInstanceDataBase* UChildActorComponent::GetComponentInstanceData() const
+FActorComponentInstanceData* UChildActorComponent::GetComponentInstanceData() const
 {
 	FChildActorComponentInstanceData* InstanceData = CachedInstanceData;
 
diff --git a/Engine/Source/Runtime/Engine/Private/ComponentInstanceDataCache.cpp b/Engine/Source/Runtime/Engine/Private/ComponentInstanceDataCache.cpp
--- a/Engine/Source/Runtime/Engine/Private/ComponentInstanceDataCache.cpp
+++ b/Engine/Source/Runtime/Engine/Private/ComponentInstanceDataCache.cpp
@@ -3,7 +3,7 @@
 #include "EnginePrivate.h"
 #include "ComponentInstanceDataCache.h"
 
-FComponentInstanceDataBase::FComponentInstanceDataBase(const UActorComponent* SourceComponent)
+FActorComponentInstanceData::FActorComponentInstanceData(const UActorComponent* SourceComponent)
 {
 	check(SourceComponent);
 	SourceComponentName = SourceComponent->GetFName();
@@ -41,9 +41,11 @@ FComponentInstanceDataBase::FComponentInstanceDataBase(const UActorComponent* So
 		class FComponentPropertyWriter : public FObjectWriter
 		{
 		public:
-			FComponentPropertyWriter(TArray<uint8>& InBytes)
+			FComponentPropertyWriter(const UActorComponent* Component, TArray<uint8>& InBytes)
 				: FObjectWriter(InBytes)
 			{
+				UClass* ComponentClass = Component->GetClass();
+				ComponentClass->SerializeTaggedProperties(*this, (uint8*)Component, ComponentClass, (uint8*)Component->GetArchetype());
 			}
 
 			virtual bool ShouldSkipProperty(const UProperty* InProperty) const override
@@ -52,13 +54,11 @@ FComponentInstanceDataBase::FComponentInstanceDataBase(const UActorComponent* So
 						|| !InProperty->HasAnyPropertyFlags(CPF_Edit | CPF_Interp));
 			}
 
-		} ComponentPropertyWriter(SavedProperties);
-
-		SourceComponentClass->SerializeTaggedProperties(ComponentPropertyWriter, (uint8*)SourceComponent, SourceComponentClass, (uint8*)SourceComponent->GetArchetype());
+		} ComponentPropertyWriter(SourceComponent, SavedProperties);
 	}
 }
 
-bool FComponentInstanceDataBase::MatchesComponent(const UActorComponent* Component) const
+bool FActorComponentInstanceData::MatchesComponent(const UActorComponent* Component) const
 {
 	bool bMatches = false;
 	if (Component && Component->GetClass() == SourceComponentClass)
@@ -89,23 +89,20 @@ bool FComponentInstanceDataBase::MatchesComponent(const UActorComponent* Compone
 	return bMatches;
 }
 
-void FComponentInstanceDataBase::ApplyToComponent(UActorComponent* Component)
+void FActorComponentInstanceData::ApplyToComponent(UActorComponent* Component, const ECacheApplyPhase CacheApplyPhase)
 {
-	if (SavedProperties.Num() > 0)
+	if (CacheApplyPhase == ECacheApplyPhase::PostSimpleConstructionScript && SavedProperties.Num() > 0)
 	{
 		class FComponentPropertyReader : public FObjectReader
 		{
 		public:
-			FComponentPropertyReader(TArray<uint8>& InBytes)
+			FComponentPropertyReader(UActorComponent* Component, TArray<uint8>& InBytes)
 				: FObjectReader(InBytes)
 			{
+				UClass* Class = Component->GetClass();
+				Class->SerializeTaggedProperties(*this, (uint8*)Component, Class, nullptr);
 			}
-		} ComponentPropertyReader(SavedProperties);
-
-		UObject* ArchetypeToSearch = Component->GetOuter()->GetArchetype();
-		UClass* Class = Component->GetClass();
-
-		Class->SerializeTaggedProperties(ComponentPropertyReader, (uint8*)Component, Class, nullptr);
+		} ComponentPropertyReader(Component, SavedProperties);	
 	}
 }
 
@@ -121,7 +118,7 @@ FComponentInstanceDataCache::FComponentInstanceDataCache(const AActor* Actor)
 		{
 			if (Component->IsCreatedByConstructionScript()) // Only cache data from 'created by construction script' components
 			{
-				FComponentInstanceDataBase* ComponentInstanceData = Component->GetComponentInstanceData();
+				FActorComponentInstanceData* ComponentInstanceData = Component->GetComponentInstanceData();
 				if (ComponentInstanceData)
 				{
 					check(!Component->GetComponentInstanceDataType().IsNone());
@@ -155,7 +152,7 @@ FComponentInstanceDataCache::~FComponentInstanceDataCache()
 	}
 }
 
-void FComponentInstanceDataCache::ApplyToActor(AActor* Actor) const
+void FComponentInstanceDataCache::ApplyToActor(AActor* Actor, const ECacheApplyPhase CacheApplyPhase) const
 {
 	if(Actor != NULL)
 	{
@@ -171,14 +168,14 @@ void FComponentInstanceDataCache::ApplyToActor(AActor* Actor) const
 
 				if (!ComponentInstanceDataType.IsNone())
 				{
-					TArray< FComponentInstanceDataBase* > CachedData;
+					TArray< FActorComponentInstanceData* > CachedData;
 					TypeToDataMap.MultiFind(ComponentInstanceDataType, CachedData);
 
-					for (FComponentInstanceDataBase* ComponentInstanceData : CachedData)
+					for (FActorComponentInstanceData* ComponentInstanceData : CachedData)
 					{
 						if (ComponentInstanceData && ComponentInstanceData->MatchesComponent(Component))
 						{
-							ComponentInstanceData->ApplyToComponent(Component);
+							ComponentInstanceData->ApplyToComponent(Component, CacheApplyPhase);
 							break;
 						}
 					}
diff --git a/Engine/Source/Runtime/Engine/Private/InstancedStaticMesh.cpp b/Engine/Source/Runtime/Engine/Private/InstancedStaticMesh.cpp
--- a/Engine/Source/Runtime/Engine/Private/InstancedStaticMesh.cpp
+++ b/Engine/Source/Runtime/Engine/Private/InstancedStaticMesh.cpp
@@ -728,9 +728,9 @@ public:
 	{
 	}
 
-	virtual void ApplyToComponent(UActorComponent* Component) override
+	virtual void ApplyToComponent(UActorComponent* Component, const ECacheApplyPhase CacheApplyPhase) override
 	{
-		FSceneComponentInstanceData::ApplyToComponent(Component);
+		FSceneComponentInstanceData::ApplyToComponent(Component, CacheApplyPhase);
 		CastChecked<UInstancedStaticMeshComponent>(Component)->ApplyComponentInstanceData(this);
 	}
 
@@ -767,10 +767,10 @@ FName UInstancedStaticMeshComponent::GetComponentInstanceDataType() const
 	return InstancedStaticMeshComponentInstanceDataName;
 }
 
-FComponentInstanceDataBase* UInstancedStaticMeshComponent::GetComponentInstanceData() const
+FActorComponentInstanceData* UInstancedStaticMeshComponent::GetComponentInstanceData() const
 {
 #if WITH_EDITOR
-	FComponentInstanceDataBase* InstanceData = nullptr;
+	FActorComponentInstanceData* InstanceData = nullptr;
 	FInstancedStaticMeshComponentInstanceData* StaticMeshInstanceData = nullptr;
 
 	// Don't back up static lighting if there isn't any
diff --git a/Engine/Source/Runtime/Engine/Private/LightComponent.cpp b/Engine/Source/Runtime/Engine/Private/LightComponent.cpp
--- a/Engine/Source/Runtime/Engine/Private/LightComponent.cpp
+++ b/Engine/Source/Runtime/Engine/Private/LightComponent.cpp
@@ -878,9 +878,9 @@ public:
 		, bPrecomputedLightingIsValid(SourceComponent->bPrecomputedLightingIsValid)
 	{}
 
-	virtual void ApplyToComponent(UActorComponent* Component) override
+	virtual void ApplyToComponent(UActorComponent* Component, const ECacheApplyPhase CacheApplyPhase) override
 	{
-		FSceneComponentInstanceData::ApplyToComponent(Component);
+		FSceneComponentInstanceData::ApplyToComponent(Component, CacheApplyPhase);
 		CastChecked<ULightComponent>(Component)->ApplyComponentInstanceData(this);
 	}
 
@@ -897,7 +897,7 @@ FName ULightComponent::GetComponentInstanceDataType() const
 	return PrecomputedLightInstanceDataTypeName;
 }
 
-FComponentInstanceDataBase* ULightComponent::GetComponentInstanceData() const
+FActorComponentInstanceData* ULightComponent::GetComponentInstanceData() const
 {
 	// Allocate new struct for holding light map data
 	return new FPrecomputedLightInstanceData(this);
diff --git a/Engine/Source/Runtime/Engine/Private/PrimitiveComponent.cpp b/Engine/Source/Runtime/Engine/Private/PrimitiveComponent.cpp
--- a/Engine/Source/Runtime/Engine/Private/PrimitiveComponent.cpp
+++ b/Engine/Source/Runtime/Engine/Private/PrimitiveComponent.cpp
@@ -332,9 +332,9 @@ FPrimitiveComponentInstanceData::FPrimitiveComponentInstanceData(const UPrimitiv
 {
 }
 
-void FPrimitiveComponentInstanceData::ApplyToComponent(UActorComponent* Component)
+void FPrimitiveComponentInstanceData::ApplyToComponent(UActorComponent* Component, const ECacheApplyPhase CacheApplyPhase)
 {
-	FSceneComponentInstanceData::ApplyToComponent(Component);
+	FSceneComponentInstanceData::ApplyToComponent(Component, CacheApplyPhase);
 
 	if (ContainsSavedProperties() && Component->IsRegistered())
 	{
@@ -347,7 +347,7 @@ bool FPrimitiveComponentInstanceData::ContainsData() const
 	return (ContainsSavedProperties() || AttachedInstanceComponents.Num() > 0);
 }
 
-FComponentInstanceDataBase* UPrimitiveComponent::GetComponentInstanceData() const
+FActorComponentInstanceData* UPrimitiveComponent::GetComponentInstanceData() const
 {
 	FPrimitiveComponentInstanceData* InstanceData = new FPrimitiveComponentInstanceData(this);
 
diff --git a/Engine/Source/Runtime/Engine/Private/SceneComponent.cpp b/Engine/Source/Runtime/Engine/Private/SceneComponent.cpp
--- a/Engine/Source/Runtime/Engine/Private/SceneComponent.cpp
+++ b/Engine/Source/Runtime/Engine/Private/SceneComponent.cpp
@@ -1163,7 +1163,7 @@ bool USceneComponent::IsAttachedTo(class USceneComponent* TestComp) const
 }
 
 FSceneComponentInstanceData::FSceneComponentInstanceData(const USceneComponent* SourceComponent)
-	: FComponentInstanceDataBase(SourceComponent)
+	: FActorComponentInstanceData(SourceComponent)
 {
 	for (int32 i = SourceComponent->AttachChildren.Num()-1; i >= 0; --i)
 	{
@@ -1175,9 +1175,9 @@ FSceneComponentInstanceData::FSceneComponentInstanceData(const USceneComponent*
 	}
 }
 
-void FSceneComponentInstanceData::ApplyToComponent(UActorComponent* Component)
+void FSceneComponentInstanceData::ApplyToComponent(UActorComponent* Component, const ECacheApplyPhase CacheApplyPhase)
 {
-	FComponentInstanceDataBase::ApplyToComponent(Component);
+	FActorComponentInstanceData::ApplyToComponent(Component, CacheApplyPhase);
 
 	USceneComponent* SceneComponent = CastChecked<USceneComponent>(Component);
 
@@ -1206,9 +1206,9 @@ void FSceneComponentInstanceData::FindAndReplaceInstances(const TMap<UObject*, U
 	}
 }
 
-FComponentInstanceDataBase* USceneComponent::GetComponentInstanceData() const
+FActorComponentInstanceData* USceneComponent::GetComponentInstanceData() const
 {
-	FComponentInstanceDataBase* InstanceData = nullptr;
+	FActorComponentInstanceData* InstanceData = nullptr;
 
 	for (USceneComponent* Child : AttachChildren)
 	{
diff --git a/Engine/Source/Runtime/Engine/Private/SkyLightComponent.cpp b/Engine/Source/Runtime/Engine/Private/SkyLightComponent.cpp
--- a/Engine/Source/Runtime/Engine/Private/SkyLightComponent.cpp
+++ b/Engine/Source/Runtime/Engine/Private/SkyLightComponent.cpp
@@ -333,9 +333,9 @@ public:
 		: FSceneComponentInstanceData(SourceComponent)
 	{}
 
-	virtual void ApplyToComponent(UActorComponent* Component) override
+	virtual void ApplyToComponent(UActorComponent* Component, const ECacheApplyPhase CacheApplyPhase) override
 	{
-		FSceneComponentInstanceData::ApplyToComponent(Component);
+		FSceneComponentInstanceData::ApplyToComponent(Component, CacheApplyPhase);
 		CastChecked<USkyLightComponent>(Component)->ApplyComponentInstanceData(this);
 	}
 
@@ -353,7 +353,7 @@ FName USkyLightComponent::GetComponentInstanceDataType() const
 	return PrecomputedSkyLightInstanceDataTypeName;
 }
 
-FComponentInstanceDataBase* USkyLightComponent::GetComponentInstanceData() const
+FActorComponentInstanceData* USkyLightComponent::GetComponentInstanceData() const
 {
 	FPrecomputedSkyLightInstanceData* InstanceData = new FPrecomputedSkyLightInstanceData(this);
 	InstanceData->LightGuid = LightGuid;
diff --git a/Engine/Source/Runtime/Engine/Private/SplineComponent.cpp b/Engine/Source/Runtime/Engine/Private/SplineComponent.cpp
--- a/Engine/Source/Runtime/Engine/Private/SplineComponent.cpp
+++ b/Engine/Source/Runtime/Engine/Private/SplineComponent.cpp
@@ -594,9 +594,9 @@ public:
 	{
 	}
 
-	virtual void ApplyToComponent(UActorComponent* Component) override
+	virtual void ApplyToComponent(UActorComponent* Component, const ECacheApplyPhase CacheApplyPhase) override
 	{
-		FSceneComponentInstanceData::ApplyToComponent(Component);
+		FSceneComponentInstanceData::ApplyToComponent(Component, CacheApplyPhase);
 		CastChecked<USplineComponent>(Component)->ApplyComponentInstanceData(this);
 	}
 
@@ -610,9 +610,9 @@ FName USplineComponent::GetComponentInstanceDataType() const
 	return SplineInstanceDataTypeName;
 }
 
-FComponentInstanceDataBase* USplineComponent::GetComponentInstanceData() const
+FActorComponentInstanceData* USplineComponent::GetComponentInstanceData() const
 {
-	FComponentInstanceDataBase* InstanceData = nullptr;
+	FActorComponentInstanceData* InstanceData = nullptr;
 	if (bAllowSplineEditingPerInstance)
 	{
 		FSplineInstanceData* SplineInstanceData = new FSplineInstanceData(this);
diff --git a/Engine/Source/Runtime/Engine/Private/StaticMeshComponent.cpp b/Engine/Source/Runtime/Engine/Private/StaticMeshComponent.cpp
--- a/Engine/Source/Runtime/Engine/Private/StaticMeshComponent.cpp
+++ b/Engine/Source/Runtime/Engine/Private/StaticMeshComponent.cpp
@@ -33,9 +33,9 @@ public:
 	{
 	}
 
-	virtual void ApplyToComponent(UActorComponent* Component) override
+	virtual void ApplyToComponent(UActorComponent* Component, const ECacheApplyPhase CacheApplyPhase) override
 	{
-		FSceneComponentInstanceData::ApplyToComponent(Component);
+		FSceneComponentInstanceData::ApplyToComponent(Component, CacheApplyPhase);
 		CastChecked<UStaticMeshComponent>(Component)->ApplyComponentInstanceData(this);
 	}
 
@@ -1541,7 +1541,7 @@ FName UStaticMeshComponent::GetComponentInstanceDataType() const
 	return StaticMeshComponentInstanceDataName;
 }
 
-FComponentInstanceDataBase* UStaticMeshComponent::GetComponentInstanceData() const
+FActorComponentInstanceData* UStaticMeshComponent::GetComponentInstanceData() const
 {
 	FStaticMeshComponentInstanceData* StaticMeshInstanceData = nullptr;
 
diff --git a/Engine/Source/Runtime/Engine/Public/ComponentInstanceDataCache.h b/Engine/Source/Runtime/Engine/Public/ComponentInstanceDataCache.h
--- a/Engine/Source/Runtime/Engine/Public/ComponentInstanceDataCache.h
+++ b/Engine/Source/Runtime/Engine/Public/ComponentInstanceDataCache.h
@@ -5,25 +5,32 @@
 class UActorComponent;
 class AActor;
 
+/** At what point in the rerun construction script process is ApplyToActor being called for */
+enum class ECacheApplyPhase
+{
+	PostSimpleConstructionScript,	// After the simple construction script has been run
+	PostUserConstructionScript,		// After the user construction script has been run
+};
+
 /** Base class for component instance cached data of a particular type. */
-class ENGINE_API FComponentInstanceDataBase
+class ENGINE_API FActorComponentInstanceData
 {
 public:
-	FComponentInstanceDataBase()
+	FActorComponentInstanceData()
 		: SourceComponentClass(nullptr)
 		, SourceComponentTypeSerializedIndex(-1)
 	{}
 
-	FComponentInstanceDataBase(const UActorComponent* SourceComponent);
+	FActorComponentInstanceData(const UActorComponent* SourceComponent);
 
-	virtual ~FComponentInstanceDataBase()
+	virtual ~FActorComponentInstanceData()
 	{}
 
 	/** Determines whether this component instance data matches the component */
 	bool MatchesComponent(const UActorComponent* Component) const;
 
 	/** Applies this component instance data to the supplied component */
-	virtual void ApplyToComponent(UActorComponent* Component);
+	virtual void ApplyToComponent(UActorComponent* Component, const ECacheApplyPhase CacheApplyPhase);
 
 	/** Replaces any references to old instances during Actor reinstancing */
 	virtual void FindAndReplaceInstances(const TMap<UObject*, UObject*>& OldToNewInstanceMap) { };
@@ -51,6 +58,7 @@ protected:
 class ENGINE_API FComponentInstanceDataCache
 {
 public:
+
 	FComponentInstanceDataCache() {}
 
 	/** Constructor that also populates cache from Actor */
@@ -59,7 +67,7 @@ public:
 	~FComponentInstanceDataCache();
 
 	/** Iterates over an Actor's components and applies the stored component instance data to each */
-	void ApplyToActor(AActor* Actor) const;
+	void ApplyToActor(AActor* Actor, const ECacheApplyPhase CacheApplyPhase) const;
 
 	/** Iterates over components and replaces any object references with the reinstanced information */
 	void FindAndReplaceInstances(const TMap<UObject*, UObject*>& OldToNewInstanceMap);
@@ -68,7 +76,7 @@ public:
 
 private:
 	/** Map of data type name to data of that type */
-	TMultiMap< FName, FComponentInstanceDataBase* >	TypeToDataMap;
+	TMultiMap< FName, FActorComponentInstanceData* >	TypeToDataMap;
 
 	TMap< USceneComponent*, FTransform > InstanceComponentTransformToRootMap;
 };
\ No newline at end of file
diff --git a/Engine/Source/Runtime/UMG/Private/Components/WidgetComponent.cpp b/Engine/Source/Runtime/UMG/Private/Components/WidgetComponent.cpp
--- a/Engine/Source/Runtime/UMG/Private/Components/WidgetComponent.cpp
+++ b/Engine/Source/Runtime/UMG/Private/Components/WidgetComponent.cpp
@@ -603,9 +603,9 @@ public:
 		, RenderTarget( SourceComponent->GetRenderTarget() )
 	{}
 
-	virtual void ApplyToComponent(UActorComponent* Component) override
+	virtual void ApplyToComponent(UActorComponent* Component, const ECacheApplyPhase CacheApplyPhase) override
 	{
-		FSceneComponentInstanceData::ApplyToComponent(Component);
+		FSceneComponentInstanceData::ApplyToComponent(Component, CacheApplyPhase);
 		CastChecked<UWidgetComponent>(Component)->ApplyComponentInstanceData(this);
 	}
 
@@ -620,7 +620,7 @@ FName UWidgetComponent::GetComponentInstanceDataType() const
 	return InstanceDataName;
 }
 
-FComponentInstanceDataBase* UWidgetComponent::GetComponentInstanceData() const
+FActorComponentInstanceData* UWidgetComponent::GetComponentInstanceData() const
 {
 	return new FWidgetComponentInstanceData( this );
 }
diff --git a/Engine/Source/Runtime/UMG/Public/Components/WidgetComponent.h b/Engine/Source/Runtime/UMG/Public/Components/WidgetComponent.h
--- a/Engine/Source/Runtime/UMG/Public/Components/WidgetComponent.h
+++ b/Engine/Source/Runtime/UMG/Public/Components/WidgetComponent.h
@@ -33,7 +33,7 @@ public:
 	virtual void DestroyComponent(bool bPromoteChildren = false) override;
 	virtual void TickComponent(float DeltaTime, enum ELevelTick TickType, FActorComponentTickFunction *ThisTickFunction) override;
 
-	virtual FComponentInstanceDataBase* GetComponentInstanceData() const override;
+	virtual FActorComponentInstanceData* GetComponentInstanceData() const override;
 	virtual FName GetComponentInstanceDataType() const override;
 
 	void ApplyComponentInstanceData(class FWidgetComponentInstanceData* ComponentInstanceData);
