[INTEGRATE] Change 2443824 by [REDACTED_PERSON][REDACTED_HANDLE] on 2015/02/12 17:56:32

	Add additional logging for UE-9598 to try and track down cause of check

diff --git a/Engine/Source/Runtime/Engine/Classes/Components/ActorComponent.h b/Engine/Source/Runtime/Engine/Classes/Components/ActorComponent.h
--- a/Engine/Source/Runtime/Engine/Classes/Components/ActorComponent.h
+++ b/Engine/Source/Runtime/Engine/Classes/Components/ActorComponent.h
@@ -82,6 +82,9 @@ private:
 	/** Is this component's dynamic data in need of sending to the renderer? */
 	uint32 bRenderDynamicDataDirty:1;
 
+	/** Used to ensure that any subclass of UActorComponent that overrides PostRename calls up to the Super to make OwnedComponents arrays get updated correctly */
+	uint32 bRoutedPostRename:1;
+
 public:
 
 	/** Does this component automatically register with its owner */
@@ -489,6 +492,7 @@ public:
 	virtual bool CallRemoteFunction( UFunction* Function, void* Parameters, FOutParmRec* OutParms, FFrame* Stack ) override;
 	virtual void PostInitProperties() override;
 	virtual void PostLoad() override;
+	virtual bool Rename( const TCHAR* NewName=NULL, UObject* NewOuter=NULL, ERenameFlags Flags=REN_None ) override;
 	virtual void PostRename(UObject* OldOuter, const FName OldName) override;
 #if WITH_EDITOR
 	virtual bool Modify( bool bAlwaysMarkDirty = true ) override;
diff --git a/Engine/Source/Runtime/Engine/Private/Actor.cpp b/Engine/Source/Runtime/Engine/Private/Actor.cpp
--- a/Engine/Source/Runtime/Engine/Private/Actor.cpp
+++ b/Engine/Source/Runtime/Engine/Private/Actor.cpp
@@ -2059,7 +2059,7 @@ const TArray<UActorComponent*>& AActor::GetReplicatedComponents() const
 
 void AActor::UpdateReplicatedComponent(UActorComponent* Component)
 {
-	check(Component->GetOwner() == this);
+	checkf(Component->GetOwner() == this, TEXT("UE-9568: Component %s being updated for Actor %s"), *Component->GetPathName(), *GetPathName() );
 	if (Component->GetIsReplicated())
 	{
 		ReplicatedComponents.AddUnique(Component);
diff --git a/Engine/Source/Runtime/Engine/Private/ActorComponent.cpp b/Engine/Source/Runtime/Engine/Private/ActorComponent.cpp
--- a/Engine/Source/Runtime/Engine/Private/ActorComponent.cpp
+++ b/Engine/Source/Runtime/Engine/Private/ActorComponent.cpp
@@ -162,8 +162,24 @@ void UActorComponent::PostLoad()
 	}
 }
 
+bool UActorComponent::Rename( const TCHAR* InName, UObject* NewOuter, ERenameFlags Flags )
+{
+	bRoutedPostRename = false;
+	
+	const bool bRenameSuccessful = Super::Rename(InName, NewOuter, Flags);
+	
+	if (!bRoutedPostRename)
+	{
+		UE_LOG(LogActorComponent, Fatal, TEXT("%s failed to route PostRename.  Please call Super::PostRename() in your <className>::PostRename() function. "), *GetFullName() );
+	}
+
+	return bRenameSuccessful;
+}
+
 void UActorComponent::PostRename(UObject* OldOuter, const FName OldName)
 {
+	Super::PostRename(OldOuter, OldName);
+
 	if (OldOuter != GetOuter())
 	{
 		AActor* Owner = GetOwner();
@@ -181,6 +197,8 @@ void UActorComponent::PostRename(UObject* OldOuter, const FName OldName)
 			}
 		}
 	}
+
+	bRoutedPostRename = true;
 }
 
 bool UActorComponent::IsCreatedByConstructionScript() const
