[INTEGRATE] Change 2444207 by [REDACTED_PERSON][REDACTED_HANDLE] on 2015/02/12 21:40:39

	Prevent a crash when a newly added actor component is not renamed (incomplete routing check) [UE-9802]

diff --git a/Engine/Source/Runtime/Engine/Private/ActorComponent.cpp b/Engine/Source/Runtime/Engine/Private/ActorComponent.cpp
--- a/Engine/Source/Runtime/Engine/Private/ActorComponent.cpp
+++ b/Engine/Source/Runtime/Engine/Private/ActorComponent.cpp
@@ -165,10 +165,14 @@ void UActorComponent::PostLoad()
 bool UActorComponent::Rename( const TCHAR* InName, UObject* NewOuter, ERenameFlags Flags )
 {
 	bRoutedPostRename = false;
+
+	const FName OldName = GetFName();
+	const UObject* OldOuter = GetOuter();
 	
 	const bool bRenameSuccessful = Super::Rename(InName, NewOuter, Flags);
 	
-	if (!bRoutedPostRename)
+	const bool bMoved = (OldName != GetFName()) || (OldOuter != GetOuter());
+	if (!bRoutedPostRename && ((Flags & REN_Test) == 0) && bMoved)
 	{
 		UE_LOG(LogActorComponent, Fatal, TEXT("%s failed to route PostRename.  Please call Super::PostRename() in your <className>::PostRename() function. "), *GetFullName() );
 	}
