[INTEGRATE] Change 2445223 by [REDACTED_PERSON][REDACTED_HANDLE] on 2015/02/13 15:02:36

	Disable editing of properties on blueprint created components

diff --git a/Engine/Source/Editor/LevelEditor/Private/SActorDetails.cpp b/Engine/Source/Editor/LevelEditor/Private/SActorDetails.cpp
--- a/Engine/Source/Editor/LevelEditor/Private/SActorDetails.cpp
+++ b/Engine/Source/Editor/LevelEditor/Private/SActorDetails.cpp
@@ -462,7 +462,7 @@ bool SActorDetails::IsPropertyEditingEnabled() const
 		UActorComponent* ActorComp = Cast<UActorComponent>(Object.Get());
 		if(ActorComp)
 		{
-			bIsEditable = ActorComp->CreationMethod != EComponentCreationMethod::UserConstructionScript;
+			bIsEditable = !ActorComp->IsCreatedByConstructionScript();
 			if(!bIsEditable)
 			{
 				break;
diff --git a/Engine/Source/Runtime/Engine/Classes/Engine/Selection.h b/Engine/Source/Runtime/Engine/Classes/Engine/Selection.h
--- a/Engine/Source/Runtime/Engine/Classes/Engine/Selection.h
+++ b/Engine/Source/Runtime/Engine/Classes/Engine/Selection.h
@@ -467,7 +467,7 @@ public:
 	bool IsObjectValid(const UObject* Object) const
 	{
 		const UActorComponent* Comp = Cast<UActorComponent>( Object );
-		return Comp && Comp->CreationMethod != EComponentCreationMethod::UserConstructionScript;
+		return Comp && !Comp->IsCreatedByConstructionScript();
 	}
 };
 
diff --git a/Engine/Source/Runtime/Engine/Private/ComponentInstanceDataCache.cpp b/Engine/Source/Runtime/Engine/Private/ComponentInstanceDataCache.cpp
--- a/Engine/Source/Runtime/Engine/Private/ComponentInstanceDataCache.cpp
+++ b/Engine/Source/Runtime/Engine/Private/ComponentInstanceDataCache.cpp
@@ -36,7 +36,7 @@ FActorComponentInstanceData::FActorComponentInstanceData(const UActorComponent*
 		}
 	}
 
-	if (SourceComponent->CreationMethod == EComponentCreationMethod::SimpleConstructionScript)
+	/*if (SourceComponent->CreationMethod == EComponentCreationMethod::SimpleConstructionScript)
 	{
 		class FComponentPropertyWriter : public FObjectWriter
 		{
@@ -55,7 +55,7 @@ FActorComponentInstanceData::FActorComponentInstanceData(const UActorComponent*
 			}
 
 		} ComponentPropertyWriter(SourceComponent, SavedProperties);
-	}
+	}*/
 }
 
 bool FActorComponentInstanceData::MatchesComponent(const UActorComponent* Component) const
