[INTEGRATE] Change 2445414 by [REDACTED_PERSON][REDACTED_HANDLE] on 2015/02/13 16:26:24

	Make renames of construction script created components when destroying them non transactional

diff --git a/Engine/Source/Runtime/Engine/Private/ActorConstruction.cpp b/Engine/Source/Runtime/Engine/Private/ActorConstruction.cpp
--- a/Engine/Source/Runtime/Engine/Private/ActorConstruction.cpp
+++ b/Engine/Source/Runtime/Engine/Private/ActorConstruction.cpp
@@ -220,7 +220,7 @@ void AActor::DestroyConstructedComponents()
 				// Rename component to avoid naming conflicts in the case where we rerun the SCS and name the new components the same way.
 				FName const NewBaseName( *(FString::Printf(TEXT("TRASH_%s"), *Component->GetClass()->GetName())) );
 				FName const NewObjectName = MakeUniqueObjectName(this, GetClass(), NewBaseName);
-				Component->Rename(*NewObjectName.ToString(), this, REN_ForceNoResetLoaders|REN_DontCreateRedirectors);
+				Component->Rename(*NewObjectName.ToString(), this, REN_ForceNoResetLoaders|REN_DontCreateRedirectors|REN_NonTransactional);
 			}
 		}
 	}
diff --git a/Engine/Source/Runtime/Engine/Private/ComponentInstanceDataCache.cpp b/Engine/Source/Runtime/Engine/Private/ComponentInstanceDataCache.cpp
--- a/Engine/Source/Runtime/Engine/Private/ComponentInstanceDataCache.cpp
+++ b/Engine/Source/Runtime/Engine/Private/ComponentInstanceDataCache.cpp
@@ -132,11 +132,7 @@ FComponentInstanceDataCache::FComponentInstanceDataCache(const AActor* Actor)
 				{
 					if (SceneComponent->AttachParent && SceneComponent->AttachParent->IsCreatedByConstructionScript())
 					{
-						auto RootComponent = Actor->GetRootComponent();
-						if (RootComponent)
-						{
-							InstanceComponentTransformToRootMap.Add(SceneComponent, SceneComponent->GetComponentTransform().GetRelativeTransform(RootComponent->GetComponentTransform()));
-						}
+						InstanceComponentTransformToRootMap.Add(SceneComponent, SceneComponent->GetComponentTransform().GetRelativeTransform(Actor->GetRootComponent()->GetComponentTransform()));
 					}
 				}
 			}
