[INTEGRATE] Change 2445502 by [REDACTED_PERSON][REDACTED_HANDLE] on 2015/02/13 17:12:00

	Fix GC PIE crash with gdc example demo map

	https://jira.ol.epicgames.net/browse/UE-9739

diff --git a/Engine/Source/Runtime/Engine/Private/GameViewportClient.cpp b/Engine/Source/Runtime/Engine/Private/GameViewportClient.cpp
--- a/Engine/Source/Runtime/Engine/Private/GameViewportClient.cpp
+++ b/Engine/Source/Runtime/Engine/Private/GameViewportClient.cpp
@@ -1872,6 +1872,8 @@ void UGameViewportClient::RemoveViewportWidgetContent( TSharedRef<SWidget> Viewp
 
 void UGameViewportClient::RemoveAllViewportWidgets()
 {
+	CursorWidgets.Empty();
+
 	TSharedPtr< SOverlay > PinnedViewportOverlayWidget( ViewportOverlayWidget.Pin() );
 	if( PinnedViewportOverlayWidget.IsValid() )
 	{
