Fix SplineMeshComponent compilation.

diff --git a/Engine/Source/Runtime/Engine/Classes/Components/SplineMeshComponent.h b/Engine/Source/Runtime/Engine/Classes/Components/SplineMeshComponent.h
--- a/Engine/Source/Runtime/Engine/Classes/Components/SplineMeshComponent.h
+++ b/Engine/Source/Runtime/Engine/Classes/Components/SplineMeshComponent.h
@@ -132,7 +132,7 @@ class ENGINE_API USplineMeshComponent : public UStaticMeshComponent, public IInt
 	//End UObject Interface
 
 	// Begin UActorComponent interface.
-	virtual FComponentInstanceDataBase* GetComponentInstanceData() const override;
+	virtual FActorComponentInstanceData* GetComponentInstanceData() const override;
 	virtual FName GetComponentInstanceDataType() const override;
 	// End UActorComponent interface.
 
diff --git a/Engine/Source/Runtime/Engine/Private/SplineMeshComponent.cpp b/Engine/Source/Runtime/Engine/Private/SplineMeshComponent.cpp
--- a/Engine/Source/Runtime/Engine/Private/SplineMeshComponent.cpp
+++ b/Engine/Source/Runtime/Engine/Private/SplineMeshComponent.cpp
@@ -826,9 +826,9 @@ public:
 	{
 	}
 
-	virtual void ApplyToComponent(UActorComponent* Component) override
+	virtual void ApplyToComponent(UActorComponent* Component, const ECacheApplyPhase CacheApplyPhase) override
 	{
-		FSceneComponentInstanceData::ApplyToComponent(Component);
+		FSceneComponentInstanceData::ApplyToComponent(Component, CacheApplyPhase);
 		CastChecked<USplineMeshComponent>(Component)->ApplyComponentInstanceData(this);
 	}
 
@@ -844,9 +844,9 @@ FName USplineMeshComponent::GetComponentInstanceDataType() const
 	return SplineMeshInstanceDataTypeName;
 }
 
-FComponentInstanceDataBase* USplineMeshComponent::GetComponentInstanceData() const
+FActorComponentInstanceData* USplineMeshComponent::GetComponentInstanceData() const
 {
-	FComponentInstanceDataBase* InstanceData = nullptr;
+	FActorComponentInstanceData* InstanceData = nullptr;
 	if (bAllowSplineEditingPerInstance)
 	{
 		FSplineMeshInstanceData *SplineMeshInstanceData = new FSplineMeshInstanceData(this);
