Fix undo/redo of changes that result in a null lightmap/shadowmap reference
(minimal version for 4.7)

diff --git a/Engine/Source/Runtime/Engine/Private/LightMap.cpp b/Engine/Source/Runtime/Engine/Private/LightMap.cpp
--- a/Engine/Source/Runtime/Engine/Private/LightMap.cpp
+++ b/Engine/Source/Runtime/Engine/Private/LightMap.cpp
@@ -1822,6 +1822,10 @@ FArchive& operator<<(FArchive& Ar,FLightMap*& R)
 		{
 			R = new FLightMap2D();
 		}
+		else
+		{
+			R = NULL;
+		}
 	}
 
 	if(R != NULL)
diff --git a/Engine/Source/Runtime/Engine/Private/ShadowMap.cpp b/Engine/Source/Runtime/Engine/Private/ShadowMap.cpp
--- a/Engine/Source/Runtime/Engine/Private/ShadowMap.cpp
+++ b/Engine/Source/Runtime/Engine/Private/ShadowMap.cpp
@@ -1038,9 +1038,16 @@ FArchive& operator<<(FArchive& Ar,FShadowMap*& R)
 
 	Ar << ShadowMapType;
 
-	if (Ar.IsLoading() && ShadowMapType == FShadowMap::SMT_2D)
+	if (Ar.IsLoading())
 	{
-		R = new FShadowMap2D();
+		if (ShadowMapType == FShadowMap::SMT_2D)
+		{
+			R = new FShadowMap2D();
+		}
+		else
+		{
+			R = NULL;
+		}
 	}
 
 	if (R != NULL)
