Fix undo/redo of changes that result in a null lightmap/shadowmap reference
(full version)

diff --git a/Engine/Source/Runtime/Engine/Private/LightMap.cpp b/Engine/Source/Runtime/Engine/Private/LightMap.cpp
--- a/Engine/Source/Runtime/Engine/Private/LightMap.cpp
+++ b/Engine/Source/Runtime/Engine/Private/LightMap.cpp
@@ -1797,53 +1797,62 @@ void TQuantizedLightSampleBulkData<QuantizedLightSampleType>::SerializeElement(
 	} 
 };
 
-FArchive& operator<<(FArchive& Ar,FLightMap*& R)
+FArchive& operator<<(FArchive& Ar, FLightMap*& R)
 {
 	uint32 LightMapType = FLightMap::LMT_None;
-	if(Ar.IsSaving())
+
+	if (Ar.IsSaving())
 	{
-		if(R != NULL)
+		if (R != nullptr)
 		{
-			if(R->GetLightMap2D())
+			if (R->GetLightMap2D())
 			{
 				LightMapType = FLightMap::LMT_2D;
 			}
 		}
 	}
+
 	Ar << LightMapType;
 
-	if(Ar.IsLoading())
+	if (Ar.IsLoading())
 	{
-		if(LightMapType == FLightMap::LMT_1D)
+		// explicitly don't call "delete R;",
+		// we expect the calling code to handle that
+		switch (LightMapType)
 		{
+		case FLightMap::LMT_None:
+			R = nullptr;
+			break;
+		case FLightMap::LMT_1D:
 			R = new FLegacyLightMap1D();
-		}
-		else if(LightMapType == FLightMap::LMT_2D)
-		{
+			break;
+		case FLightMap::LMT_2D:
 			R = new FLightMap2D();
-		}
-		else
-		{
-			R = NULL;
+			break;
+		default:
+			check(0);
 		}
 	}
 
-	if(R != NULL)
+	if (R != nullptr)
 	{
 		R->Serialize(Ar);
 
-		// Toss legacy vertex lightmaps
-		if (LightMapType == FLightMap::LMT_1D)
+		if (Ar.IsLoading())
 		{
-			delete R;
-			R = NULL;
-		}
+			// Toss legacy vertex lightmaps
+			if (LightMapType == FLightMap::LMT_1D)
+			{
+				delete R;
+				R = nullptr;
+			}
 
-		// Dump old lightmaps
-		if( Ar.IsLoading() && Ar.UE4Ver() < VER_UE4_COMBINED_LIGHTMAP_TEXTURES )
-		{
-			delete R;
-			R = NULL;
+			// Dump old lightmaps
+			if (Ar.UE4Ver() < VER_UE4_COMBINED_LIGHTMAP_TEXTURES)
+			{
+				delete R; // safe because if we're loading we new'd this above
+				R = nullptr;
+			}
 		}
 	}
 
diff --git a/Engine/Source/Runtime/Engine/Private/ShadowMap.cpp b/Engine/Source/Runtime/Engine/Private/ShadowMap.cpp
--- a/Engine/Source/Runtime/Engine/Private/ShadowMap.cpp
+++ b/Engine/Source/Runtime/Engine/Private/ShadowMap.cpp
@@ -1031,34 +1031,48 @@ FArchive& operator<<(FArchive& Ar,FShadowMap*& R)
 {
 	uint32 ShadowMapType = FShadowMap::SMT_None;
 
-	if (Ar.IsSaving() && R != NULL && R->GetShadowMap2D())
+	if (Ar.IsSaving())
 	{
-		ShadowMapType = FShadowMap::SMT_2D;
+		if (R != nullptr)
+		{
+			if (R->GetShadowMap2D())
+			{
+				ShadowMapType = FShadowMap::SMT_2D;
+			}
+		}
 	}
 
 	Ar << ShadowMapType;
 
 	if (Ar.IsLoading())
 	{
-		if (ShadowMapType == FShadowMap::SMT_2D)
+		// explicitly don't call "delete R;",
+		// we expect the calling code to handle that
+		switch (ShadowMapType)
 		{
+		case FShadowMap::SMT_None:
+			R = nullptr;
+			break;
+		case FShadowMap::SMT_2D:
 			R = new FShadowMap2D();
-		}
-		else
-		{
-			R = NULL;
+			break;
+		default:
+			check(0);
 		}
 	}
 
-	if (R != NULL)
+	if (R != nullptr)
 	{
 		R->Serialize(Ar);
-		
-		// Dump old lightmaps
-		if( Ar.IsLoading() && Ar.UE4Ver() < VER_UE4_COMBINED_LIGHTMAP_TEXTURES )
+
+		if (Ar.IsLoading())
 		{
-			delete R;
-			R = NULL;
+			// Dump old Shadowmaps
+			if (Ar.UE4Ver() < VER_UE4_COMBINED_LIGHTMAP_TEXTURES)
+			{
+				delete R; // safe because if we're loading we new'd this above
+				R = nullptr;
+			}
 		}
 	}
 
diff --git a/Engine/Source/Runtime/Engine/Public/LightMap.h b/Engine/Source/Runtime/Engine/Public/LightMap.h
--- a/Engine/Source/Runtime/Engine/Public/LightMap.h
+++ b/Engine/Source/Runtime/Engine/Public/LightMap.h
@@ -124,7 +124,10 @@ private:
 
 
 
-/** Lightmap reference serializer */
+/**
+ * Lightmap reference serializer
+ * Intended to be used by TRefCountPtr's serializer, not called directly
+ */
 extern ENGINE_API FArchive& operator<<(FArchive& Ar, FLightMap*& R);
 
 /** 
diff --git a/Engine/Source/Runtime/Engine/Public/ShadowMap.h b/Engine/Source/Runtime/Engine/Public/ShadowMap.h
--- a/Engine/Source/Runtime/Engine/Public/ShadowMap.h
+++ b/Engine/Source/Runtime/Engine/Public/ShadowMap.h
@@ -360,6 +360,8 @@ private:
 
 
 
-
-/** Shadowmap reference serializer */
+/**
+ * Shadowmap reference serializer
+ * Intended to be used by TRefCountPtr's serializer, not called directly
+ */
 extern ENGINE_API FArchive& operator<<(FArchive& Ar, FShadowMap*& R);
