[INTEGRATE] Change 2446374 by [REDACTED_PERSON][REDACTED_HANDLE] on 2015/02/15 13:21:38

	Prevent SCS Editor from picking as name for the component the name of an object that already exists in the CDO of the class

diff --git a/Engine/Source/Runtime/Engine/Private/SCS_Node.cpp b/Engine/Source/Runtime/Engine/Private/SCS_Node.cpp
--- a/Engine/Source/Runtime/Engine/Private/SCS_Node.cpp
+++ b/Engine/Source/Runtime/Engine/Private/SCS_Node.cpp
@@ -406,6 +406,15 @@ void USCS_Node::GenerateListOfExistingNames( TArray<FName>& CurrentNames ) const
 	// <<< End Backwards Compatibility
 	check(Blueprint);
 
+	TArray<UObject*> NativeCDOChildren;
+	UClass* FirstNativeClass = FBlueprintEditorUtils::FindFirstNativeClass(Blueprint->ParentClass);
+	GetObjectsWithOuter(FirstNativeClass->GetDefaultObject(), NativeCDOChildren, false);
+
+	for (UObject* NativeCDOChild : NativeCDOChildren)
+	{
+		CurrentNames.Add(NativeCDOChild->GetFName());
+	}
+
 	if(Blueprint->SkeletonGeneratedClass)
 	{
 		// First add the class variables.
