[INTEGRATE] Change 2446430 by [REDACTED_PERSON][REDACTED_HANDLE] on 2015/02/15 17:33:44

	#ANIM: Fix for common aninmation crash due to index mismatch.

diff --git a/Engine/Source/Runtime/Engine/Private/Animation/AnimSequence.cpp b/Engine/Source/Runtime/Engine/Private/Animation/AnimSequence.cpp
--- a/Engine/Source/Runtime/Engine/Private/Animation/AnimSequence.cpp
+++ b/Engine/Source/Runtime/Engine/Private/Animation/AnimSequence.cpp
@@ -938,16 +938,18 @@ void UAnimSequence::GetBonePose(FTransformArrayA2& OutAtoms, const FBoneContaine
 		{
 			const int32 SkeletonBoneIndex = GetSkeletonIndexFromTrackIndex(TrackIndex);
 			// not sure it's safe to assume that SkeletonBoneIndex can never be INDEX_NONE
-			check( (SkeletonBoneIndex != INDEX_NONE) && (SkeletonBoneIndex < MAX_BONES) );
-			const int32 PoseBoneIndex = SkeletonToPoseBoneIndexArray[SkeletonBoneIndex];
-			if( PoseBoneIndex != INDEX_NONE )
+			if( (SkeletonBoneIndex != INDEX_NONE) && (SkeletonBoneIndex < MAX_BONES) )
 			{
-				// extract animation
-				ExtractBoneTransform(*AnimationData, OutAtoms[PoseBoneIndex], TrackIndex, ExtractionContext.CurrentTime);
-
-				if (!bDisableRetargeting)
+				const int32 PoseBoneIndex = SkeletonToPoseBoneIndexArray[SkeletonBoneIndex];
+				if(PoseBoneIndex != INDEX_NONE)
 				{
-					RetargetBoneTransform(OutAtoms[PoseBoneIndex], SkeletonBoneIndex, PoseBoneIndex, RequiredBones);
+					// extract animation
+					ExtractBoneTransform(*AnimationData, OutAtoms[PoseBoneIndex], TrackIndex, ExtractionContext.CurrentTime);
+
+					if(!bDisableRetargeting)
+					{
+						RetargetBoneTransform(OutAtoms[PoseBoneIndex], SkeletonBoneIndex, PoseBoneIndex, RequiredBones);
+					}
 				}
 			}
 		}
@@ -985,27 +987,29 @@ void UAnimSequence::GetBonePose(FTransformArrayA2& OutAtoms, const FBoneContaine
 	{
 		const int32 SkeletonBoneIndex = GetSkeletonIndexFromTrackIndex(TrackIndex);
 		// not sure it's safe to assume that SkeletonBoneIndex can never be INDEX_NONE
-		checkSlow( SkeletonBoneIndex != INDEX_NONE );
-		const int32 PoseBoneIndex = SkeletonToPoseBoneIndexArray[SkeletonBoneIndex];
-		if( PoseBoneIndex != INDEX_NONE )
+		if( SkeletonBoneIndex != INDEX_NONE )
 		{
-			checkSlow( PoseBoneIndex < RequiredBones.GetNumBones() );
-			RotationScalePairs.Add(BoneTrackPair(PoseBoneIndex, TrackIndex));
-
-			// Skip extracting translation component for EBoneTranslationRetargetingMode::Skeleton.
-			switch( BoneTree[SkeletonBoneIndex].TranslationRetargetingMode )
+			const int32 PoseBoneIndex = SkeletonToPoseBoneIndexArray[SkeletonBoneIndex];
+			if(PoseBoneIndex != INDEX_NONE)
 			{
-				case EBoneTranslationRetargetingMode::Animation : 
-					TranslationPairs.Add(BoneTrackPair(PoseBoneIndex, TrackIndex));
-					break;
-				case EBoneTranslationRetargetingMode::AnimationScaled :
-					TranslationPairs.Add(BoneTrackPair(PoseBoneIndex, TrackIndex));
-					AnimScaleRetargetingPairs.Add(BoneTrackPair(PoseBoneIndex, SkeletonBoneIndex));
-					break;
+				checkSlow(PoseBoneIndex < RequiredBones.GetNumBones());
+				RotationScalePairs.Add(BoneTrackPair(PoseBoneIndex, TrackIndex));
+
+				// Skip extracting translation component for EBoneTranslationRetargetingMode::Skeleton.
+				switch(BoneTree[SkeletonBoneIndex].TranslationRetargetingMode)
+				{
+				case EBoneTranslationRetargetingMode::Animation:
+				TranslationPairs.Add(BoneTrackPair(PoseBoneIndex, TrackIndex));
+				break;
+				case EBoneTranslationRetargetingMode::AnimationScaled:
+				TranslationPairs.Add(BoneTrackPair(PoseBoneIndex, TrackIndex));
+				AnimScaleRetargetingPairs.Add(BoneTrackPair(PoseBoneIndex, SkeletonBoneIndex));
+				break;
 				case EBoneTranslationRetargetingMode::AnimationRelative:
-					TranslationPairs.Add(BoneTrackPair(PoseBoneIndex, TrackIndex));
-					AnimRelativeRetargetingPairs.Add(BoneTrackPair(PoseBoneIndex, SkeletonBoneIndex));
-					break;
+				TranslationPairs.Add(BoneTrackPair(PoseBoneIndex, TrackIndex));
+				AnimRelativeRetargetingPairs.Add(BoneTrackPair(PoseBoneIndex, SkeletonBoneIndex));
+				break;
+				}
 			}
 		}
 	}
