[INTEGRATE] Change 2447233 by [REDACTED_PERSON][REDACTED_HANDLE] on 2015/02/16 11:23:39

	Fix various undo/redo issues with ChildActorComponent when it is used as an Instance component

diff --git a/Engine/Source/Editor/UnrealEd/Private/EditorTransaction.cpp b/Engine/Source/Editor/UnrealEd/Private/EditorTransaction.cpp
--- a/Engine/Source/Editor/UnrealEd/Private/EditorTransaction.cpp
+++ b/Engine/Source/Editor/UnrealEd/Private/EditorTransaction.cpp
@@ -265,6 +265,13 @@ void FTransaction::Apply()
 		Records[i].Restore( this );
 	}
 
+	// An Actor's components must always get its PostEditUndo before the owning Actor so do a quick sort
+	ChangedObjects.KeySort([](UObject& A, UObject& B)
+	{
+		UActorComponent* BAsComponent = Cast<UActorComponent>(&B);
+		return (BAsComponent ? (BAsComponent->GetOwner() != &A) : true);
+	});
+
 	NumModelsModified = 0;		// Count the number of UModels that were changed.
 	for (auto ChangedObjectIt : ChangedObjects)
 	{
diff --git a/Engine/Source/Runtime/Engine/Classes/Components/ChildActorComponent.h b/Engine/Source/Runtime/Engine/Classes/Components/ChildActorComponent.h
--- a/Engine/Source/Runtime/Engine/Classes/Components/ChildActorComponent.h
+++ b/Engine/Source/Runtime/Engine/Classes/Components/ChildActorComponent.h
@@ -30,6 +30,7 @@ class UChildActorComponent : public USceneComponent
 	// Begin Object interface.
 #if WITH_EDITOR
 	virtual void PostEditChangeProperty(FPropertyChangedEvent& PropertyChangedEvent) override;
+	virtual void PostEditUndo() override;
 #endif
 	// End Object interface.
 
diff --git a/Engine/Source/Runtime/Engine/Classes/GameFramework/Actor.h b/Engine/Source/Runtime/Engine/Classes/GameFramework/Actor.h
--- a/Engine/Source/Runtime/Engine/Classes/GameFramework/Actor.h
+++ b/Engine/Source/Runtime/Engine/Classes/GameFramework/Actor.h
@@ -1145,6 +1145,8 @@ public:
 		struct FAttachedActorInfo
 		{
 			TWeakObjectPtr<AActor> Actor;
+			TWeakObjectPtr<USceneComponent> AttachParent;
+			FName AttachParentName;
 			FName SocketName;
 			FTransform RelativeTransform;
 		};
diff --git a/Engine/Source/Runtime/Engine/Private/ActorConstruction.cpp b/Engine/Source/Runtime/Engine/Private/ActorConstruction.cpp
--- a/Engine/Source/Runtime/Engine/Private/ActorConstruction.cpp
+++ b/Engine/Source/Runtime/Engine/Private/ActorConstruction.cpp
@@ -296,8 +296,10 @@ void AActor::RerunConstructionScripts()
 				Parent = ActorTransactionAnnotation->RootComponentData.AttachedParentInfo.Actor.Get();
 				if (Parent)
 				{
-					DetachRootComponentFromParent();
+					USceneComponent* AttachParent = ActorTransactionAnnotation->RootComponentData.AttachedParentInfo.AttachParent.Get();
+					ParentComponent = (AttachParent ? AttachParent : FindObjectFast<USceneComponent>(Parent, ActorTransactionAnnotation->RootComponentData.AttachedParentInfo.AttachParentName));
 					SocketName = ActorTransactionAnnotation->RootComponentData.AttachedParentInfo.SocketName;
+					DetachRootComponentFromParent();
 				}
 
 				for (const auto& CachedAttachInfo : ActorTransactionAnnotation->RootComponentData.AttachedToInfo)
diff --git a/Engine/Source/Runtime/Engine/Private/ActorEditor.cpp b/Engine/Source/Runtime/Engine/Private/ActorEditor.cpp
--- a/Engine/Source/Runtime/Engine/Private/ActorEditor.cpp
+++ b/Engine/Source/Runtime/Engine/Private/ActorEditor.cpp
@@ -213,6 +213,8 @@ AActor::FActorTransactionAnnotation::FActorTransactionAnnotation(const AActor* A
 		if (RootComponent->AttachParent)
 		{
 			RootComponentData.AttachedParentInfo.Actor = RootComponent->AttachParent->GetOwner();
+			RootComponentData.AttachedParentInfo.AttachParent = RootComponent->AttachParent;
+			RootComponentData.AttachedParentInfo.AttachParentName = RootComponent->AttachParent->GetFName();
 			RootComponentData.AttachedParentInfo.SocketName = RootComponent->AttachSocketName;
 			RootComponentData.AttachedParentInfo.RelativeTransform = RootComponent->GetRelativeTransform();
 		}
diff --git a/Engine/Source/Runtime/Engine/Private/ChildActorComponent.cpp b/Engine/Source/Runtime/Engine/Private/ChildActorComponent.cpp
--- a/Engine/Source/Runtime/Engine/Private/ChildActorComponent.cpp
+++ b/Engine/Source/Runtime/Engine/Private/ChildActorComponent.cpp
@@ -39,6 +39,23 @@ void UChildActorComponent::PostEditChangeProperty(FPropertyChangedEvent& Propert
 		CreateChildActor();
 	}
 }
+
+void UChildActorComponent::PostEditUndo()
+{
+	Super::PostEditUndo();
+
+	// This hack exists to fix up known cases where the AttachChildren array is broken in very problematic ways.
+	// The correct fix will be to use a Transaction Annotation at the SceneComponent level, however, it is too risky
+	// to do right now, so this will go away when that is done.
+	for (USceneComponent*& Component : AttachChildren)
+	{
+		if (Component->IsPendingKill() && Component->GetOwner() == ChildActor)
+		{
+			Component = ChildActor->GetRootComponent();
+		}
+	}
+	
+}
 #endif
 
 void UChildActorComponent::OnComponentCreated()
