Optimise FRichCurve::Eval to use binary search to find curve segment interpolation nodes

diff --git a/Engine/Source/Runtime/Engine/Private/CurveBase.cpp b/Engine/Source/Runtime/Engine/Private/CurveBase.cpp
--- a/Engine/Source/Runtime/Engine/Private/CurveBase.cpp
+++ b/Engine/Source/Runtime/Engine/Private/CurveBase.cpp
@@ -688,69 +688,86 @@ static float BezierInterp2(float P0, float Y1, float Y2, float P3, float mu)
 	return Result;
 }
 
-float FRichCurve::Eval( const float InTime, float DefaultValue ) const
+float FRichCurve::Eval(const float InTime, float DefaultValue) const
 {
 	const int32 NumKeys = Keys.Num();
+	float InterpVal = DefaultValue;
 
-	// If no keys in curve, return the Default value we passed in.
-	if( NumKeys == 0 )
+	if (NumKeys == 0)
 	{
-		return DefaultValue;
-	}
-
-	// If only one point, or before the first point in the curve, return the first points value.
-	if( NumKeys < 2 || (InTime <= Keys[0].Time) )
+		// If no keys in curve, return the Default value we passed in.
+		InterpVal = DefaultValue;
+	} 
+	else if (NumKeys < 2 || (InTime <= Keys[0].Time))
 	{
-		return Keys[0].Value;
+		// If only one point, or before the first point in the curve, return the first points value.
+		InterpVal = Keys[0].Value;
 	}
-
-	// If beyond the last point in the curve, return its value.
-	if( InTime >= Keys[NumKeys-1].Time )
+	else
 	{
-		return Keys[NumKeys-1].Value;
-	}
+		// perform a lower bound to get the second of the interpolation nodes
+		int32 first = 1;
+		int32 last = NumKeys - 1;
+		int32 count = last - first;
+
+		while (count > 0)
+		{
+			int32 step = count / 2;
+			int32 middle = first + step;
+
+			if (InTime > Keys[middle].Time)
+			{
+				first = middle + 1;
+				count -= step + 1;
+			}
+			else
+			{
+				count = step;
+			}
+		}
+
+		int32 InterpNode = first;
 
-	// Somewhere with curve range - linear search to find value.
-	for( int32 i=1; i<NumKeys; i++ )
-	{	
-		if( InTime < Keys[i].Time )
+		if (InTime < Keys[InterpNode].Time)
 		{
-			const float Diff = Keys[i].Time - Keys[i-1].Time;
+			const float Diff = Keys[InterpNode].Time - Keys[InterpNode - 1].Time;
 
-			if( Diff > 0.f && Keys[i-1].InterpMode != RCIM_Constant )
+			if (Diff > 0.f && Keys[InterpNode - 1].InterpMode != RCIM_Constant)
 			{
-				const float Alpha = (InTime - Keys[i-1].Time) / Diff;
-				const float P0 = Keys[i-1].Value;
-				const float P3 = Keys[i].Value;
+				const float Alpha = (InTime - Keys[InterpNode - 1].Time) / Diff;
+				const float P0 = Keys[InterpNode - 1].Value;
+				const float P3 = Keys[InterpNode].Value;
 
-				if( Keys[i-1].InterpMode == RCIM_Linear )
+				if (Keys[InterpNode - 1].InterpMode == RCIM_Linear)
 				{
-					return FMath::Lerp( P0, P3, Alpha );
+					InterpVal = FMath::Lerp(P0, P3, Alpha);
 				}
 				else
 				{
-					const float OneThird = 1.0f/3.0f;
-					const float P1 = P0 + (Keys[i-1].LeaveTangent * Diff*OneThird);
-					const float P2 = P3 - (Keys[i].ArriveTangent * Diff*OneThird);
-
-					//1st and 2nd derivatives of 0.0
-					//	float  QuinticHermite = 6.0f*FMath::Pow(Alpha,5.0f) - 15.0f*FMath::Pow(Alpha,4.0f) + 10.0f*FMath::Pow(Alpha,3.0f);
-					return BezierInterp( P0, P1, P2, P3, Alpha );
-					//return FMath::CubicInterp( Keys(i-1).Value, Keys(i-1).LeaveTangent * Diff, Keys(i).Value, Keys(i).ArriveTangent * Diff, Alpha );
+					const float OneThird = 1.0f / 3.0f;
+					const float P1 = P0 + (Keys[InterpNode - 1].LeaveTangent * Diff*OneThird);
+					const float P2 = P3 - (Keys[InterpNode].ArriveTangent * Diff*OneThird);
+
+					InterpVal = BezierInterp(P0, P1, P2, P3, Alpha);
 				}
 			}
 			else
 			{
-				return Keys[i-1].Value;
+				InterpVal = Keys[InterpNode - 1].Value;
 			}
 		}
+		else
+		{
+			// If beyond the last point in the curve, return its value.
+			InterpVal = Keys[NumKeys - 1].Value;
+		}
 	}
 
-	// Shouldn't really reach here.
-	return Keys[NumKeys-1].Value;
+	return InterpVal;
 }
 
-bool FRichCurve::operator==( const FRichCurve& Curve ) const
+
+bool FRichCurve::operator==(const FRichCurve& Curve) const
 {
 	if(Keys.Num() != Curve.Keys.Num())
 	{
