Changed editor box selection criteria for BrushComponents.

diff --git a/Engine/Source/Runtime/Engine/Private/BrushComponent.cpp b/Engine/Source/Runtime/Engine/Private/BrushComponent.cpp
--- a/Engine/Source/Runtime/Engine/Private/BrushComponent.cpp
+++ b/Engine/Source/Runtime/Engine/Private/BrushComponent.cpp
@@ -621,20 +621,48 @@ bool UBrushComponent::ComponentIsTouchingSelectionBox(const FBox& InSelBBox, con
 
 		if (!bMustEncompassEntireComponent)
 		{
-			for (const auto& Poly : Brush->Polys->Element)
+			const bool bMustContainVertex = true;
+			if (bMustContainVertex)
 			{
-				// Just an intersection will do...
-				Vertices.Empty(Poly.Vertices.Num());
-				for (const auto& Vertex : Poly.Vertices)
+				// Test if any poly vertex intersects the box
+				for (const auto& Poly : Brush->Polys->Element)
 				{
-					Vertices.Add(ComponentToWorld.TransformPosition(Vertex));
-				}
+					for (const auto& Vertex : Poly.Vertices)
+					{
+						const FVector Location = ComponentToWorld.TransformPosition(Vertex);
+						const bool bLocationIntersected = FMath::PointBoxIntersection(Location, InSelBBox);
 
-				FSeparatingAxisPointCheck PointCheck(Vertices, InSelBBox.GetCenter(), InSelBBox.GetExtent(), false);
-				if (PointCheck.bHit)
+						// If the selection box doesn't have to encompass the entire component and any poly vertex intersects with the selection
+						// box, this component qualifies
+						if (bLocationIntersected)
+						{
+							return true;
+						}
+					}
+				}
+			}
+			else
+			{
+				// Alternative method, which can be enabled by setting bMustContainVertex = false:
+				// Test if any poly edge intersects the box
+				for (const auto& Poly : Brush->Polys->Element)
 				{
-					// If any poly intersected with the bounding box, this component is considered to be touching
-					return true;
+					const int32 NumVerts = Poly.Vertices.Num();
+					if (NumVerts > 0)
+					{
+						FVector StartVert = ComponentToWorld.TransformPosition(Poly.Vertices[NumVerts - 1]);
+						for (int32 Index = 0; Index < NumVerts; ++Index)
+						{
+							const FVector EndVert = ComponentToWorld.TransformPosition(Poly.Vertices[Index]);
+
+							if (FMath::LineBoxIntersection(InSelBBox, StartVert, EndVert, EndVert - StartVert))
+							{
+								return true;
+							}
+
+							StartVert = EndVert;
+						}
+					}
 				}
 			}
 
