Fixed a crash accessing an invalid viewport on GameViewportClient.

diff --git a/Engine/Source/Runtime/Engine/Private/GameViewportClient.cpp b/Engine/Source/Runtime/Engine/Private/GameViewportClient.cpp
--- a/Engine/Source/Runtime/Engine/Private/GameViewportClient.cpp
+++ b/Engine/Source/Runtime/Engine/Private/GameViewportClient.cpp
@@ -448,10 +448,10 @@ void UGameViewportClient::MouseLeave(FViewport* InViewport)
 {
 	Super::MouseLeave(InViewport);
 
-	if (GetDefault<UInputSettings>()->bUseMouseForTouch)
+	if (InViewport && GetDefault<UInputSettings>()->bUseMouseForTouch)
 	{
 		FIntPoint LastViewportCursorPos;
-		Viewport->GetMousePos(LastViewportCursorPos, false);
+		InViewport->GetMousePos(LastViewportCursorPos, false);
 		FVector2D CursorPos(LastViewportCursorPos.X, LastViewportCursorPos.Y);
 		FSlateApplication::Get().SetGameIsFakingTouchEvents(false, &CursorPos);
 	}
