Fix shadow variables

diff --git a/Engine/Source/Editor/UnrealEd/Private/AutoReimport/AutoReimportUtilities.h b/Engine/Source/Editor/UnrealEd/Private/AutoReimport/AutoReimportUtilities.h
--- a/Engine/Source/Editor/UnrealEd/Private/AutoReimport/AutoReimportUtilities.h
+++ b/Engine/Source/Editor/UnrealEd/Private/AutoReimport/AutoReimportUtilities.h
@@ -39,15 +39,15 @@ struct FImmutableString
 	const FString& Get() const { return String; }
 
 	/** Serialise this string */
-	friend FArchive& operator<<(FArchive& Ar, FImmutableString& String)
+	friend FArchive& operator<<(FArchive& Ar, FImmutableString& InString)
 	{
-		Ar << String.String;
+		Ar << InString.String;
 		if (Ar.IsSaving())
 		{
-			GetTypeHash(String);	
+			GetTypeHash(InString);	
 		}
 		
-		Ar << String.CachedHash;
+		Ar << InString.CachedHash;
 
 		return Ar;
 	}
diff --git a/Engine/Source/Runtime/Engine/Private/AI/Navigation/NavLinkRenderingComponent.cpp b/Engine/Source/Runtime/Engine/Private/AI/Navigation/NavLinkRenderingComponent.cpp
--- a/Engine/Source/Runtime/Engine/Private/AI/Navigation/NavLinkRenderingComponent.cpp
+++ b/Engine/Source/Runtime/Engine/Private/AI/Navigation/NavLinkRenderingComponent.cpp
@@ -29,7 +29,7 @@ UNavLinkRenderingComponent::UNavLinkRenderingComponent(const FObjectInitializer&
 	bGenerateOverlapEvents = false;
 }
 
-FBoxSphereBounds UNavLinkRenderingComponent::CalcBounds(const FTransform& LocalToWorld) const
+FBoxSphereBounds UNavLinkRenderingComponent::CalcBounds(const FTransform& InLocalToWorld) const
 {
 	AActor* LinkOwnerActor = Cast<AActor>(GetOwner());
 	INavLinkHostInterface* LinkOwnerHost = Cast<INavLinkHostInterface>(GetOwner());
diff --git a/Engine/Source/Runtime/Engine/Private/AI/Navigation/NavigationSystem.cpp b/Engine/Source/Runtime/Engine/Private/AI/Navigation/NavigationSystem.cpp
--- a/Engine/Source/Runtime/Engine/Private/AI/Navigation/NavigationSystem.cpp
+++ b/Engine/Source/Runtime/Engine/Private/AI/Navigation/NavigationSystem.cpp
@@ -2190,12 +2190,12 @@ void UNavigationSystem::AddElementToNavOctree(const FNavigationDirtyElement& Dir
 	FNavigationOctreeElement GeneratedData;
 	const FBox ElementBounds = DirtyElement.NavInterface->GetNavigationBounds();
 
-	UObject* ParentNode = DirtyElement.NavInterface->GetNavigationParent();
-	if (ParentNode)
+	UObject* NavigationParent = DirtyElement.NavInterface->GetNavigationParent();
+	if (NavigationParent)
 	{
 		// check if parent node is waiting in queue
-		const FSetElementId ParentRequestId = PendingOctreeUpdates.FindId(FNavigationDirtyElement(ParentNode));
-		const FOctreeElementId* ParentId = GetObjectsNavOctreeId(ParentNode);
+		const FSetElementId ParentRequestId = PendingOctreeUpdates.FindId(FNavigationDirtyElement(NavigationParent));
+		const FOctreeElementId* ParentId = GetObjectsNavOctreeId(NavigationParent);
 		if (ParentRequestId.IsValidId() && ParentId == NULL)
 		{
 			FNavigationDirtyElement& ParentNode = PendingOctreeUpdates[ParentRequestId];
@@ -2205,15 +2205,15 @@ void UNavigationSystem::AddElementToNavOctree(const FNavigationDirtyElement& Dir
 			ParentNode.bInvalidRequest = true;
 		}
 
-		const FOctreeElementId* UseParentId = ParentId ? ParentId : GetObjectsNavOctreeId(ParentNode);
+		const FOctreeElementId* UseParentId = ParentId ? ParentId : GetObjectsNavOctreeId(NavigationParent);
 		if (UseParentId)
 		{
-			UE_LOG(LogNavOctree, Log, TEXT("ADD %s to %s"), *GetNameSafe(ElementOwner), *GetNameSafe(ParentNode));
+			UE_LOG(LogNavOctree, Log, TEXT("ADD %s to %s"), *GetNameSafe(ElementOwner), *GetNameSafe(NavigationParent));
 			NavOctree->AppendToNode(*UseParentId, DirtyElement.NavInterface, ElementBounds, GeneratedData);
 		}
 		else 
 		{
-			UE_LOG(LogNavOctree, Warning, TEXT("Can't add node [%s] - parent [%s] not found in octree!"), *GetNameSafe(ElementOwner), *GetNameSafe(ParentNode));
+			UE_LOG(LogNavOctree, Warning, TEXT("Can't add node [%s] - parent [%s] not found in octree!"), *GetNameSafe(ElementOwner), *GetNameSafe(NavigationParent));
 		}
 	}
 	else
diff --git a/Engine/Source/Runtime/Engine/Private/Animation/AnimSingleNodeInstance.cpp b/Engine/Source/Runtime/Engine/Private/Animation/AnimSingleNodeInstance.cpp
--- a/Engine/Source/Runtime/Engine/Private/Animation/AnimSingleNodeInstance.cpp
+++ b/Engine/Source/Runtime/Engine/Private/Animation/AnimSingleNodeInstance.cpp
@@ -173,8 +173,8 @@ void UAnimSingleNodeInstance::UpdateBlendspaceSamples(FVector InBlendInput)
 {
 	if (UBlendSpaceBase* BlendSpace = Cast<UBlendSpaceBase>(CurrentAsset))
 	{
-		float CurrentTime = 0.f;
-		BlendSpaceAdvanceImmediate(BlendSpace, InBlendInput, BlendSampleData, BlendFilter, false, 1.f, 0.f, CurrentTime);
+		float OutCurrentTime = 0.f;
+		BlendSpaceAdvanceImmediate(BlendSpace, InBlendInput, BlendSampleData, BlendFilter, false, 1.f, 0.f, OutCurrentTime);
 	}
 }
 
diff --git a/Engine/Source/Runtime/Engine/Private/Debug/ReporterGraph.cpp b/Engine/Source/Runtime/Engine/Private/Debug/ReporterGraph.cpp
--- a/Engine/Source/Runtime/Engine/Private/Debug/ReporterGraph.cpp
+++ b/Engine/Source/Runtime/Engine/Private/Debug/ReporterGraph.cpp
@@ -163,7 +163,6 @@ void UReporterGraph::DrawLegend(UCanvas* Canvas)
 		FVector2D ScreenPos = ToScreenSpace(CurrentTextPos, Canvas);
 		if (LegendPosition == ELegendPosition::Outside)
 		{
-			int32 DummyY, CurrentX;
 			StringSize(Font, CurrentX, DummyY, *CurrentData[i].LineName);
 			LegendWidth = CurrentX + 10;
 		}
@@ -227,7 +226,6 @@ void UReporterGraph::DrawAxis(UCanvas* Canvas, FVector2D Start, FVector2D End, f
 		case EGraphAxisStyle::Lines:
 		{
 			NumNotches = 1;
-			FVector2D NotchDataDelta = (GraphMinMaxData.Max - GraphMinMaxData.Min) / NumNotches;
 			NotchDelta = (Start - End).Size();
 			if(bIsVerticalAxis)
 			{
diff --git a/Engine/Source/Runtime/Engine/Private/Interpolation.cpp b/Engine/Source/Runtime/Engine/Private/Interpolation.cpp
--- a/Engine/Source/Runtime/Engine/Private/Interpolation.cpp
+++ b/Engine/Source/Runtime/Engine/Private/Interpolation.cpp
@@ -909,8 +909,6 @@ void AMatineeActor::UpdateInterpForParentMovementTracks( float Time, UInterpGrou
 		UInterpGroupInst* ParentInst = FindGroupInst(Parent);
 		if(ParentInst)
 		{
-			TArray<UInterpTrack*> FoundTracks;
-
 			UInterpTrackInst* ParentTrackInst = nullptr;
 			for(UInterpTrackInst* Inst : ParentInst->TrackInst)
 			{
diff --git a/Engine/Source/Runtime/Engine/Private/NetworkDriver.cpp b/Engine/Source/Runtime/Engine/Private/NetworkDriver.cpp
--- a/Engine/Source/Runtime/Engine/Private/NetworkDriver.cpp
+++ b/Engine/Source/Runtime/Engine/Private/NetworkDriver.cpp
@@ -2244,9 +2244,8 @@ int32 UNetDriver::ServerReplicateActors(float DeltaSeconds)
 				AGameMode const* const GameMode = World->GetAuthGameMode();
 				bool bLowNetBandwidth = !bCPUSaturated && (Connection->CurrentNetSpeed / float(GameMode->NumPlayers + GameMode->NumBots) < 500.f );
 
-				for( j=0; j<ConsiderList.Num(); j++ )
+				for( AActor* Actor : ConsiderList )
 				{
-					AActor* Actor = ConsiderList[j];
 					UActorChannel* Channel = Connection->ActorChannels.FindRef(Actor);
 
 					// Skip Actor if dormant
@@ -2276,10 +2275,10 @@ int32 UNetDriver::ServerReplicateActors(float DeltaSeconds)
 							bool ShouldGoDormant = true;
 							if (Actor->NetDormancy == DORM_DormantPartial)
 							{
-								float Time  = Channel ? (Connection->Driver->Time - Channel->LastUpdateTime) : Connection->Driver->SpawnPrioritySeconds;
+								float LastReplicationTime  = Channel ? (Connection->Driver->Time - Channel->LastUpdateTime) : Connection->Driver->SpawnPrioritySeconds;
 								for (int32 viewerIdx = 0; viewerIdx < ConnectionViewers.Num(); viewerIdx++)
 								{
-									if (!Actor->GetNetDormancy(ConnectionViewers[viewerIdx].ViewLocation, ConnectionViewers[viewerIdx].ViewDir, ConnectionViewers[viewerIdx].InViewer, Channel, Time, bLowNetBandwidth))
+									if (!Actor->GetNetDormancy(ConnectionViewers[viewerIdx].ViewLocation, ConnectionViewers[viewerIdx].ViewDir, ConnectionViewers[viewerIdx].InViewer, Channel, LastReplicationTime, bLowNetBandwidth))
 									{
 										ShouldGoDormant = false;
 										break;
@@ -2351,9 +2350,8 @@ int32 UNetDriver::ServerReplicateActors(float DeltaSeconds)
 				int32 ChildIndex = 0;
 				while (NextConnection != NULL)
 				{
-					for (int32 j = 0; j < NextConnection->OwnedConsiderList.Num(); j++)
+					for (AActor* Actor : NextConnection->OwnedConsiderList)
 					{
-						AActor* Actor = NextConnection->OwnedConsiderList[j];
 						UE_LOG(LogNetTraffic, Log, TEXT("Consider owned %s always relevant %d frequency %f  "),*Actor->GetName(), Actor->bAlwaysRelevant,Actor->NetUpdateFrequency);
 						if (Actor->NetTag != NetTag)
 						{
diff --git a/Engine/Source/Runtime/Engine/Private/OnlineReplStructs.cpp b/Engine/Source/Runtime/Engine/Private/OnlineReplStructs.cpp
--- a/Engine/Source/Runtime/Engine/Private/OnlineReplStructs.cpp
+++ b/Engine/Source/Runtime/Engine/Private/OnlineReplStructs.cpp
@@ -102,9 +102,9 @@ void TestUniqueIdRepl(UWorld* InWorld)
 				else
 				{
 					// Net serialize
-					bool bSuccess = false;
-					EmptyIdIn.NetSerialize(TestWriteUniqueId, NULL, bSuccess);
-					ValidIdIn.NetSerialize(TestWriteUniqueId, NULL, bSuccess);
+					bool bOutSuccess = false;
+					EmptyIdIn.NetSerialize(TestWriteUniqueId, NULL, bOutSuccess);
+					ValidIdIn.NetSerialize(TestWriteUniqueId, NULL, bOutSuccess);
 				}
 
 				FMemoryReader TestReadUniqueId(Buffer);
diff --git a/Engine/Source/Runtime/Engine/Private/Slate/SceneViewport.cpp b/Engine/Source/Runtime/Engine/Private/Slate/SceneViewport.cpp
--- a/Engine/Source/Runtime/Engine/Private/Slate/SceneViewport.cpp
+++ b/Engine/Source/Runtime/Engine/Private/Slate/SceneViewport.cpp
@@ -426,7 +426,7 @@ FReply FSceneViewport::OnMouseButtonUp( const FGeometry& InGeometry, const FPoin
 
 	// Switch to the viewport clients world before processing input
 	FScopedConditionalWorldSwitcher WorldSwitcher( ViewportClient );
-	bool bIsCursorVisible = true;
+	bool bCursorVisible = true;
 	bool bReleaseMouse = true;
 	if( ViewportClient && GetSizeXY() != FIntPoint::ZeroValue  )
 	{
@@ -434,9 +434,9 @@ FReply FSceneViewport::OnMouseButtonUp( const FGeometry& InGeometry, const FPoin
 		{
 			CurrentReplyState = FReply::Unhandled(); 
 		}
-		bIsCursorVisible = ViewportClient->GetCursor(this, GetMouseX(), GetMouseY()) != EMouseCursor::None;
+		bCursorVisible = ViewportClient->GetCursor(this, GetMouseX(), GetMouseY()) != EMouseCursor::None;
 		bReleaseMouse = 
-			bIsCursorVisible || 
+			bCursorVisible || 
 			ViewportClient->CaptureMouseOnClick() == EMouseCaptureMode::CaptureDuringMouseDown ||
 			( ViewportClient->CaptureMouseOnClick() == EMouseCaptureMode::CaptureDuringRightMouseDown && InMouseEvent.GetEffectingButton() == EKeys::RightMouseButton );
 	}
@@ -454,7 +454,7 @@ FReply FSceneViewport::OnMouseButtonUp( const FGeometry& InGeometry, const FPoin
 			}
 
 			CurrentReplyState.ReleaseMouseCapture();
-			if (bIsCursorVisible)
+			if (bCursorVisible)
 			{
 				CurrentReplyState.ReleaseMouseLock();
 			}
diff --git a/Engine/Source/Runtime/Engine/Private/Slate/SlateTextures.cpp b/Engine/Source/Runtime/Engine/Private/Slate/SlateTextures.cpp
--- a/Engine/Source/Runtime/Engine/Private/Slate/SlateTextures.cpp
+++ b/Engine/Source/Runtime/Engine/Private/Slate/SlateTextures.cpp
@@ -118,17 +118,17 @@ void FSlateTexture2DRHIRef::ClearTextureData()
 	TextureData.Reset();
 }
 
-void FSlateTexture2DRHIRef::ResizeTexture(uint32 Width, uint32 Height)
+void FSlateTexture2DRHIRef::ResizeTexture(uint32 InWidth, uint32 InHeight)
 {
-	if (GetWidth() != Width || GetHeight() != Height)
+	if (GetWidth() != InWidth || GetHeight() != InHeight)
 	{
 		if (IsInRenderingThread())
 		{
-			Resize(Width, Height);
+			Resize(InWidth, InHeight);
 		}
 		else
 		{
-			FIntPoint Dimensions(Width, Height);
+			FIntPoint Dimensions(InWidth, InHeight);
 			ENQUEUE_UNIQUE_RENDER_COMMAND_TWOPARAMETER(ResizeSlateTexture,
 			FSlateTexture2DRHIRef*, TextureRHIRef, this,
 			FIntPoint, InDimensions, Dimensions,
diff --git a/Engine/Source/Runtime/Engine/Public/ShadowMap.h b/Engine/Source/Runtime/Engine/Public/ShadowMap.h
--- a/Engine/Source/Runtime/Engine/Public/ShadowMap.h
+++ b/Engine/Source/Runtime/Engine/Public/ShadowMap.h
@@ -215,8 +215,8 @@ public:
 	{
 	}
 
-	FShadowMap(TArray<FGuid> LightGuids)
-		: LightGuids(MoveTemp(LightGuids))
+	FShadowMap(TArray<FGuid> InLightGuids)
+		: LightGuids(MoveTemp(InLightGuids))
 		, NumRefs(0)
 	{
 	}
diff --git a/Engine/Source/Runtime/UMG/Public/Binding/DynamicPropertyPath.h b/Engine/Source/Runtime/UMG/Public/Binding/DynamicPropertyPath.h
--- a/Engine/Source/Runtime/UMG/Public/Binding/DynamicPropertyPath.h
+++ b/Engine/Source/Runtime/UMG/Public/Binding/DynamicPropertyPath.h
@@ -236,17 +236,17 @@ private:
 					else if ( UArrayProperty* ArrayProp = Cast<UArrayProperty>(Property) )
 					{
 						// It is an array, now check to see if this is an array of structures
-						if ( UStructProperty* StructProp = Cast<UStructProperty>(ArrayProp->Inner) )
+						if ( UStructProperty* ArrayOfStructsProp = Cast<UStructProperty>(ArrayProp->Inner) )
 						{
 							FScriptArrayHelper_InContainer ArrayHelper(ArrayProp, InContainer);
 							if ( ArrayHelper.IsValidIndex(Segment.ArrayIndex) )
 							{
 								// Recursively call back into this function with the array element and container value
-								return GetValueRecursive(StructProp->Struct, ArrayHelper.GetRawPtr(Segment.ArrayIndex), ArrayIndex, SegmentIndex + 1, OutValue, OutProperty);
+								return GetValueRecursive(ArryayOfStructsProp->Struct, ArrayHelper.GetRawPtr(Segment.ArrayIndex), ArrayIndex, SegmentIndex + 1, OutValue, OutProperty);
 							}
 						}
 						// if it's not an array of structs, maybe it's an array of classes
-						else if ( UObjectProperty* ObjectProperty = Cast<UObjectProperty>(ArrayProp->Inner) )
+						//else if ( UObjectProperty* ObjectProperty = Cast<UObjectProperty>(ArrayProp->Inner) )
 						{
 							//TODO Add support for arrays of objects.
 						}
