Fixed volumes always drawing with the current selection color even when unselected

diff --git a/Engine/Source/Runtime/Engine/Private/PrimitiveDrawingUtils.cpp b/Engine/Source/Runtime/Engine/Private/PrimitiveDrawingUtils.cpp
--- a/Engine/Source/Runtime/Engine/Private/PrimitiveDrawingUtils.cpp
+++ b/Engine/Source/Runtime/Engine/Private/PrimitiveDrawingUtils.cpp
@@ -1466,7 +1466,7 @@ FLinearColor GetViewSelectionColor(const FLinearColor& BaseColor, const FSceneVi
 	{
 		FinalColor = GEngine->GetSubduedSelectionOutlineColor();
 	}
-	else
+	else if( bSelected )
 #endif
 	{
 		FinalColor = GEngine->GetSelectedMaterialColor();
