Deprecated properties should be available in PostLoad for CDO.

https://udn.unrealengine.com/questions/236375/45-bug-deprecated-uproperties-and-out-of-date-blue.html

diff --git a/Engine/Source/Editor/KismetCompiler/Private/KismetCompiler.cpp b/Engine/Source/Editor/KismetCompiler/Private/KismetCompiler.cpp
--- a/Engine/Source/Editor/KismetCompiler/Private/KismetCompiler.cpp
+++ b/Engine/Source/Editor/KismetCompiler/Private/KismetCompiler.cpp
@@ -3548,7 +3548,9 @@ void FKismetCompilerContext::Compile()
 						}
 					}
 
-					UEditorEngine::CopyPropertiesForUnrelatedObjects(OldCDO, NewCDO);
+					UEditorEngine::FCopyPropertiesForUnrelatedObjectsParams CopyDetails;
+					CopyDetails.bCopyDeprecatedProperties = Blueprint->bIsRegeneratingOnLoad;
+					UEditorEngine::CopyPropertiesForUnrelatedObjects(OldCDO, NewCDO, CopyDetails);
 					FBlueprintEditorUtils::PatchCDOSubobjectsIntoExport(OldCDO, NewCDO);
 				}
 
diff --git a/Engine/Source/Editor/UnrealEd/Private/Kismet2/BlueprintEditorUtils.cpp b/Engine/Source/Editor/UnrealEd/Private/Kismet2/BlueprintEditorUtils.cpp
--- a/Engine/Source/Editor/UnrealEd/Private/Kismet2/BlueprintEditorUtils.cpp
+++ b/Engine/Source/Editor/UnrealEd/Private/Kismet2/BlueprintEditorUtils.cpp
@@ -1316,6 +1316,7 @@ UClass* FBlueprintEditorUtils::RegenerateBlueprintClass(UBlueprint* Blueprint, U
 						UEditorEngine::FCopyPropertiesForUnrelatedObjectsParams CopyDetails;
 						CopyDetails.bAggressiveDefaultSubobjectReplacement = true;
 						CopyDetails.bDoDelta = false;
+						CopyDetails.bCopyDeprecatedProperties = true;
 						UEditorEngine::CopyPropertiesForUnrelatedObjects(PreviousCDO, NewCDO, CopyDetails);
 					}
 
diff --git a/Engine/Source/Runtime/CoreUObject/Private/UObject/Property.cpp b/Engine/Source/Runtime/CoreUObject/Private/UObject/Property.cpp
--- a/Engine/Source/Runtime/CoreUObject/Private/UObject/Property.cpp
+++ b/Engine/Source/Runtime/CoreUObject/Private/UObject/Property.cpp
@@ -604,7 +604,7 @@ bool UProperty::ShouldSerializeValue( FArchive& Ar ) const
 		||	((PropertyFlags & CPF_NonPIEDuplicateTransient) && !(Ar.GetPortFlags() & PPF_DuplicateForPIE) && (Ar.GetPortFlags() & PPF_Duplicate))
 		||  (Ar.IsFilterEditorOnly() && IsEditorOnlyProperty())
 		||	((PropertyFlags & CPF_NonTransactional) && Ar.IsTransacting())
-		||	((PropertyFlags & CPF_Deprecated) && (Ar.IsSaving() || Ar.IsTransacting() || Ar.WantBinaryPropertySerialization()));
+		||	((PropertyFlags & CPF_Deprecated) && !Ar.HasAllPortFlags(PPF_UseDeprecatedProperties) && (Ar.IsSaving() || Ar.IsTransacting() || Ar.WantBinaryPropertySerialization()));
 
 	return !Skip;
 }
diff --git a/Engine/Source/Runtime/CoreUObject/Public/Serialization/ArchiveUObject.h b/Engine/Source/Runtime/CoreUObject/Public/Serialization/ArchiveUObject.h
--- a/Engine/Source/Runtime/CoreUObject/Public/Serialization/ArchiveUObject.h
+++ b/Engine/Source/Runtime/CoreUObject/Public/Serialization/ArchiveUObject.h
@@ -87,12 +87,13 @@ class FObjectWriter : public FMemoryWriter
 {
 
 public:
-	FObjectWriter(UObject* Obj, TArray<uint8>& InBytes, bool bIgnoreClassRef = false, bool bIgnoreArchetypeRef = false, bool bDoDelta = true)
+	FObjectWriter(UObject* Obj, TArray<uint8>& InBytes, bool bIgnoreClassRef = false, bool bIgnoreArchetypeRef = false, bool bDoDelta = true, uint32 AdditionalPortFlags = 0)
 		: FMemoryWriter(InBytes)
 	{
 		ArIgnoreClassRef = bIgnoreClassRef;
 		ArIgnoreArchetypeRef = bIgnoreArchetypeRef;
 		ArNoDelta = !bDoDelta;
+		ArPortFlags |= AdditionalPortFlags;
 		Obj->Serialize(*this);
 	}
 
diff --git a/Engine/Source/Runtime/CoreUObject/Public/UObject/UnrealType.h b/Engine/Source/Runtime/CoreUObject/Public/UObject/UnrealType.h
--- a/Engine/Source/Runtime/CoreUObject/Public/UObject/UnrealType.h
+++ b/Engine/Source/Runtime/CoreUObject/Public/UObject/UnrealType.h
@@ -98,6 +98,9 @@ enum EPropertyPortFlags
 
 	/** Exporting properties for console variables. */
 	PPF_ConsoleVariable				= 0x04000000,
+
+	/** Ignores CPF_Deprecated flag */
+	PPF_UseDeprecatedProperties		= 0x08000000,
 };
 
 enum EPropertyExportCPPFlags
diff --git a/Engine/Source/Runtime/Engine/Classes/Engine/Engine.h b/Engine/Source/Runtime/Engine/Classes/Engine/Engine.h
--- a/Engine/Source/Runtime/Engine/Classes/Engine/Engine.h
+++ b/Engine/Source/Runtime/Engine/Classes/Engine/Engine.h
@@ -2166,11 +2166,13 @@ public:
 		bool bAggressiveDefaultSubobjectReplacement;
 		bool bDoDelta;
 		bool bReplaceObjectClassReferences;
+		bool bCopyDeprecatedProperties;
 
 		FCopyPropertiesForUnrelatedObjectsParams()
 			: bAggressiveDefaultSubobjectReplacement(false)
 			, bDoDelta(true)
 			, bReplaceObjectClassReferences(true)
+			, bCopyDeprecatedProperties(false)
 		{}
 	};
 	static void CopyPropertiesForUnrelatedObjects(UObject* OldObject, UObject* NewObject, FCopyPropertiesForUnrelatedObjectsParams Params = FCopyPropertiesForUnrelatedObjectsParams());//bool bAggressiveDefaultSubobjectReplacement = false, bool bDoDelta = true);
diff --git a/Engine/Source/Runtime/Engine/Private/UnrealEngine.cpp b/Engine/Source/Runtime/Engine/Private/UnrealEngine.cpp
--- a/Engine/Source/Runtime/Engine/Private/UnrealEngine.cpp
+++ b/Engine/Source/Runtime/Engine/Private/UnrealEngine.cpp
@@ -10338,9 +10338,10 @@ void UEngine::CopyPropertiesForUnrelatedObjects(UObject* OldObject, UObject* New
 	TIndirectArray<FInstancedObjectRecord> SavedInstances;
 	TMap<FName, int32> OldInstanceMap;
 
+	const uint32 AdditinalPortFlags = Params.bCopyDeprecatedProperties ? PPF_UseDeprecatedProperties : PPF_None;
 	// Save the modified properties of the old CDO
 	{
-		FObjectWriter Writer(OldObject, SavedProperties, true, true, Params.bDoDelta);
+		FObjectWriter Writer(OldObject, SavedProperties, true, true, Params.bDoDelta, AdditinalPortFlags);
 	}
 
 	{
@@ -10354,7 +10355,7 @@ void UEngine::CopyPropertiesForUnrelatedObjects(UObject* OldObject, UObject* New
 			UObject* OldInstance = Components[Index];
 			pRecord->OldInstance = OldInstance;
 			OldInstanceMap.Add(OldInstance->GetFName(), SavedInstances.Num() - 1);
-			FObjectWriter Writer(OldInstance, pRecord->SavedProperties, true, true);
+			FObjectWriter Writer(OldInstance, pRecord->SavedProperties, true, true, true, AdditinalPortFlags);
 		}
 	}
 
