Add the ability to ignore mass for all force/torque type functions

diff --git a/Engine/Source/Runtime/Engine/Classes/Components/DestructibleComponent.h b/Engine/Source/Runtime/Engine/Classes/Components/DestructibleComponent.h
--- a/Engine/Source/Runtime/Engine/Classes/Components/DestructibleComponent.h
+++ b/Engine/Source/Runtime/Engine/Classes/Components/DestructibleComponent.h
@@ -129,10 +129,10 @@ public:
 
 	virtual void AddImpulse(FVector Impulse, FName BoneName = NAME_None, bool bVelChange = false) override;
 	virtual void AddImpulseAtLocation(FVector Impulse, FVector Position, FName BoneName = NAME_None) override;
-	virtual void AddForce(FVector Force, FName BoneName = NAME_None) override;
+	virtual void AddForce(FVector Force, FName BoneName = NAME_None, bool bAccelChange = false) override;
 	virtual void AddForceAtLocation(FVector Force, FVector Location, FName BoneName = NAME_None) override;
 	virtual void AddRadialImpulse(FVector Origin, float Radius, float Strength, ERadialImpulseFalloff Falloff, bool bVelChange=false) override;
-	virtual void AddRadialForce(FVector Origin, float Radius, float Strength, ERadialImpulseFalloff Falloff) override;
+	virtual void AddRadialForce(FVector Origin, float Radius, float Strength, ERadialImpulseFalloff Falloff, bool bAccelChange = false) override;
 	virtual void ReceiveComponentDamage(float DamageAmount, FDamageEvent const& DamageEvent, AController* EventInstigator, AActor* DamageCauser) override;
 
 	virtual bool LineTraceComponent( FHitResult& OutHit, const FVector Start, const FVector End, const FCollisionQueryParams& Params ) override;
diff --git a/Engine/Source/Runtime/Engine/Classes/Components/PrimitiveComponent.h b/Engine/Source/Runtime/Engine/Classes/Components/PrimitiveComponent.h
--- a/Engine/Source/Runtime/Engine/Classes/Components/PrimitiveComponent.h
+++ b/Engine/Source/Runtime/Engine/Classes/Components/PrimitiveComponent.h
@@ -736,11 +736,12 @@ public:
 	 *	Add a force to a single rigid body.
 	 *  This is like a 'thruster'. Good for adding a burst over some (non zero) time. Should be called every frame for the duration of the force.
 	 *
-	 *	@param	Force		Force vector to apply. Magnitude indicates strength of force.
-	 *	@param	BoneName	If a SkeletalMeshComponent, name of body to apply force to. 'None' indicates root body.
+	 *	@param	Force		 Force vector to apply. Magnitude indicates strength of force.
+	 *	@param	BoneName	 If a SkeletalMeshComponent, name of body to apply force to. 'None' indicates root body.
+	 *  @param  bAccelChange If true, Force is taken as a change in acceleration instead of a physical force (i.e. mass will have no affect).
 	 */
 	UFUNCTION(BlueprintCallable, Category="Physics")
-	virtual void AddForce(FVector Force, FName BoneName = NAME_None);
+	virtual void AddForce(FVector Force, FName BoneName = NAME_None, bool bAccelChange = false);
 
 	/**
 	 *	Add a force to a single rigid body at a particular location.
@@ -760,17 +761,19 @@ public:
 	 *	@param Radius		Radius within which to apply the force.
 	 *	@param Strength		Strength of force to apply.
 	 *  @param Falloff		Allows you to control the strength of the force as a function of distance from Origin.
+	 *  @param bAccelChange If true, Strength is taken as a change in acceleration instead of a physical force (i.e. mass will have no affect).
 	 */
 	UFUNCTION(BlueprintCallable, Category="Physics")
-	virtual void AddRadialForce(FVector Origin, float Radius, float Strength, enum ERadialImpulseFalloff Falloff);
+	virtual void AddRadialForce(FVector Origin, float Radius, float Strength, enum ERadialImpulseFalloff Falloff, bool bAccelChange = false);
 
 	/**
 	 *	Add a torque to a single rigid body.
 	 *	@param Torque		Torque to apply. Direction is axis of rotation and magnitude is strength of torque.
 	 *	@param BoneName		If a SkeletalMeshComponent, name of body to apply torque to. 'None' indicates root body.
+	 *  @param bAccelChange If true, Torque is taken as a change in angular acceleration instead of a physical torque (i.e. mass will have no affect).
 	 */
 	UFUNCTION(BlueprintCallable, Category="Physics")
-	void AddTorque(FVector Torque, FName BoneName = NAME_None);
+	void AddTorque(FVector Torque, FName BoneName = NAME_None, bool bAccelChange = false);
 
 	/**
 	 *	Set the linear velocity of a single body.
diff --git a/Engine/Source/Runtime/Engine/Classes/Components/SkeletalMeshComponent.h b/Engine/Source/Runtime/Engine/Classes/Components/SkeletalMeshComponent.h
--- a/Engine/Source/Runtime/Engine/Classes/Components/SkeletalMeshComponent.h
+++ b/Engine/Source/Runtime/Engine/Classes/Components/SkeletalMeshComponent.h
@@ -838,7 +838,7 @@ public:
 	virtual bool OverlapComponent(const FVector& Pos, const FQuat& Rot, const FCollisionShape& CollisionShape) override;
 	virtual void SetSimulatePhysics(bool bEnabled) override;
 	virtual void AddRadialImpulse(FVector Origin, float Radius, float Strength, ERadialImpulseFalloff Falloff, bool bVelChange=false) override;
-	virtual void AddRadialForce(FVector Origin, float Radius, float Strength, ERadialImpulseFalloff Falloff) override;
+	virtual void AddRadialForce(FVector Origin, float Radius, float Strength, ERadialImpulseFalloff Falloff, bool bAccelChange=false) override;
 	virtual void SetAllPhysicsLinearVelocity(FVector NewVel,bool bAddToCurrent = false) override;
 	virtual void SetAllMassScale(float InMassScale = 1.f) override;
 	virtual float GetMass() const override;
diff --git a/Engine/Source/Runtime/Engine/Classes/PhysicsEngine/BodyInstance.h b/Engine/Source/Runtime/Engine/Classes/PhysicsEngine/BodyInstance.h
--- a/Engine/Source/Runtime/Engine/Classes/PhysicsEngine/BodyInstance.h
+++ b/Engine/Source/Runtime/Engine/Classes/PhysicsEngine/BodyInstance.h
@@ -533,11 +533,11 @@ public:
 	/** Add custom forces and torques on the body. The callback will be called more than once, if substepping enabled, for every substep.  */
 	void AddCustomPhysics(FCalculateCustomPhysics& CalculateCustomPhysics);
 	/** Add a force to this body */
-	void AddForce(const FVector& Force, bool bAllowSubstepping = true);
+	void AddForce(const FVector& Force, bool bAllowSubstepping = true, bool bAccelChange = false);
 	/** Add a force at a particular world position to this body */
 	void AddForceAtPosition(const FVector& Force, const FVector& Position, bool bAllowSubstepping = true);
 	/** Add a torque to this body */
-	void AddTorque(const FVector& Torque, bool bAllowSubstepping = true);
+	void AddTorque(const FVector& Torque, bool bAllowSubstepping = true, bool bAccelChange = false);
 	/** Add a rotational impulse to this body */
 	void AddAngularImpulse(const FVector& Impulse, bool bVelChange);
 	/** Add an impulse to this body */
@@ -724,7 +724,7 @@ public:
 	 * @param Falloff		Allows you to control the strength of the impulse as a function of distance from Origin.
 	 * @param bVelChange	If true, the Strength is taken as a change in velocity instead of an impulse (ie. mass will have no affect).
 	 */
-	void AddRadialImpulseToBody(const FVector& Origin, float Radius, float Strength, uint8 Falloff, bool bVelChange);
+	void AddRadialImpulseToBody(const FVector& Origin, float Radius, float Strength, uint8 Falloff, bool bVelChange = false);
 
 	/**
 	 *	Add a force to this bodyinstance, originating from the supplied world-space location.
@@ -733,8 +733,9 @@ public:
 	 *	@param Radius		Radius within which to apply the force.
 	 *	@param Strength		Strength of force to apply.
 	 *  @param Falloff		Allows you to control the strength of the force as a function of distance from Origin.
+	 *  @param bAccelChange If true, Strength is taken as a change in acceleration instead of a physical force (i.e. mass will have no affect).
 	 */
-	void AddRadialForceToBody(const FVector& Origin, float Radius, float Strength, uint8 Falloff);
+	void AddRadialForceToBody(const FVector& Origin, float Radius, float Strength, uint8 Falloff, bool bAccelChange = false);
 
 	/**
 	 * Get distance to the body surface if available
diff --git a/Engine/Source/Runtime/Engine/Private/DestructibleComponent.cpp b/Engine/Source/Runtime/Engine/Private/DestructibleComponent.cpp
--- a/Engine/Source/Runtime/Engine/Private/DestructibleComponent.cpp
+++ b/Engine/Source/Runtime/Engine/Private/DestructibleComponent.cpp
@@ -456,7 +456,7 @@ void UDestructibleComponent::AddImpulseAtLocation( FVector Impulse, FVector Posi
 #endif
 }
 
-void UDestructibleComponent::AddForce( FVector Force, FName BoneName /*= NAME_None*/ )
+void UDestructibleComponent::AddForce( FVector Force, FName BoneName /*= NAME_None*/, bool bAccelChange /* = false */ )
 {
 #if WITH_APEX
 	int32 ChunkIdx = NxModuleDestructibleConst::INVALID_CHUNK_INDEX;
@@ -470,7 +470,7 @@ void UDestructibleComponent::AddForce( FVector Force, FName BoneName /*= NAME_No
 	{
 		if (ChunkInfos[i].ChunkIndex == ChunkIdx && ChunkInfos[i].Actor)
 		{
-			ChunkInfos[i].Actor->addForce(U2PVector(Force), PxForceMode::eFORCE);
+			ChunkInfos[i].Actor->addForce(U2PVector(Force), bAccelChange ? PxForceMode::eACCELERATION : PxForceMode::eFORCE);
 		}
 	}
 #endif
@@ -545,7 +545,7 @@ void UDestructibleComponent::AddRadialImpulse(FVector Origin, float Radius, floa
 #endif	// #if WITH_APEX
 }
 
-void UDestructibleComponent::AddRadialForce(FVector Origin, float Radius, float Strength, ERadialImpulseFalloff Falloff)
+void UDestructibleComponent::AddRadialForce(FVector Origin, float Radius, float Strength, ERadialImpulseFalloff Falloff, bool bAccelChange /* = false */)
 {
 #if WITH_APEX
 	if(bIgnoreRadialForce)
@@ -572,7 +572,7 @@ void UDestructibleComponent::AddRadialForce(FVector Origin, float Radius, float
 					LockedScene->lockRead();
 				}
 
-				AddRadialForceToPxRigidBody(*PActor, Origin, Radius, Strength, Falloff);
+				AddRadialForceToPxRigidBody(*PActor, Origin, Radius, Strength, Falloff, bAccelChange);
 			}
 
 			if (LockedScene)
diff --git a/Engine/Source/Runtime/Engine/Private/PhysicsEngine/BodyInstance.cpp b/Engine/Source/Runtime/Engine/Private/PhysicsEngine/BodyInstance.cpp
--- a/Engine/Source/Runtime/Engine/Private/PhysicsEngine/BodyInstance.cpp
+++ b/Engine/Source/Runtime/Engine/Private/PhysicsEngine/BodyInstance.cpp
@@ -3657,7 +3657,7 @@ void FBodyInstance::AddCustomPhysics(FCalculateCustomPhysics& CalculateCustomPhy
 #endif
 }
 
-void FBodyInstance::AddForce(const FVector& Force, bool bAllowSubstepping)
+void FBodyInstance::AddForce(const FVector& Force, bool bAllowSubstepping, bool bAccelChange)
 {
 #if WITH_PHYSX
 	PxRigidBody* PRigidBody = GetPxRigidBody();
@@ -3665,7 +3665,7 @@ void FBodyInstance::AddForce(const FVector& Force, bool bAllowSubstepping)
 	{
 		const PxScene* PScene = PRigidBody->getScene();
 		FPhysScene* PhysScene = FPhysxUserData::Get<FPhysScene>(PScene->userData);
-		PhysScene->AddForce(this, Force, bAllowSubstepping);
+		PhysScene->AddForce(this, Force, bAllowSubstepping, bAccelChange);
 		
 	}
 #endif // WITH_PHYSX
@@ -3676,7 +3676,7 @@ void FBodyInstance::AddForce(const FVector& Force, bool bAllowSubstepping)
 		if (!Force.IsNearlyZero())
 		{
 			const b2Vec2 Force2D = FPhysicsIntegration2D::ConvertUnrealVectorToBox(Force);
-			BodyInstancePtr->ApplyForceToCenter(Force2D, /*wake=*/ true);
+			BodyInstancePtr->ApplyForceToCenter(bAccelChange ? (BodyInstancePtr->GetMass() * Force2D) : Force2D, /*wake=*/ true);
 		}
 	}
 #endif
@@ -3707,7 +3707,7 @@ void FBodyInstance::AddForceAtPosition(const FVector& Force, const FVector& Posi
 #endif
 }
 
-void FBodyInstance::AddTorque(const FVector& Torque, bool bAllowSubstepping)
+void FBodyInstance::AddTorque(const FVector& Torque, bool bAllowSubstepping, bool bAccelChange)
 {
 #if WITH_PHYSX
 	PxRigidBody* PRigidBody = GetPxRigidBody();
@@ -3715,14 +3715,14 @@ void FBodyInstance::AddTorque(const FVector& Torque, bool bAllowSubstepping)
 	{
 		const PxScene* PScene = PRigidBody->getScene();
 		FPhysScene* PhysScene = FPhysxUserData::Get<FPhysScene>(PScene->userData);
-		PhysScene->AddTorque(this, Torque, bAllowSubstepping);
+		PhysScene->AddTorque(this, Torque, bAllowSubstepping, bAccelChange);
 	}
 #endif // WITH_PHYSX
 
 #if WITH_BOX2D
 	if (BodyInstancePtr)
 	{
-		const float Torque1D = FPhysicsIntegration2D::ConvertUnrealTorqueToBox(Torque);
+		const float Torque1D = FPhysicsIntegration2D::ConvertUnrealTorqueToBox(Torque) * (bAccelChange ? BodyInstancePtr->GetInertia() : 1.f);
 		if (!FMath::IsNearlyZero(Torque1D))
 		{
 			BodyInstancePtr->ApplyTorque(Torque1D, /*wake=*/ true);
@@ -3895,20 +3895,20 @@ void FBodyInstance::AddRadialImpulseToBody(const FVector& Origin, float Radius,
 #endif
 }
 
-void FBodyInstance::AddRadialForceToBody(const FVector& Origin, float Radius, float Strength, uint8 Falloff)
+void FBodyInstance::AddRadialForceToBody(const FVector& Origin, float Radius, float Strength, uint8 Falloff, bool bAccelChange)
 {
 #if WITH_PHYSX
 	PxRigidBody* PRigidBody = GetPxRigidBody();
 	if (IsRigidBodyNonKinematic(PRigidBody))
 	{
-		AddRadialForceToPxRigidBody(*PRigidBody, Origin, Radius, Strength, Falloff);
+		AddRadialForceToPxRigidBody(*PRigidBody, Origin, Radius, Strength, Falloff, bAccelChange);
 	}
 #endif // WITH_PHYSX
 
 #if WITH_BOX2D
 	if (BodyInstancePtr)
 	{
-		AddRadialImpulseToBox2DBodyInstance(BodyInstancePtr, Origin, Radius, Strength, static_cast<ERadialImpulseFalloff>(Falloff), /*bImpulse=*/ false, /*bMassIndependent=*/ false);
+		AddRadialImpulseToBox2DBodyInstance(BodyInstancePtr, Origin, Radius, Strength, static_cast<ERadialImpulseFalloff>(Falloff), /*bImpulse=*/ false, /*bMassIndependent=*/ bAccelChange);
 	}
 #endif
 }
diff --git a/Engine/Source/Runtime/Engine/Private/PhysicsEngine/PhysScene.cpp b/Engine/Source/Runtime/Engine/Private/PhysicsEngine/PhysScene.cpp
--- a/Engine/Source/Runtime/Engine/Private/PhysicsEngine/PhysScene.cpp
+++ b/Engine/Source/Runtime/Engine/Private/PhysicsEngine/PhysScene.cpp
@@ -241,7 +241,7 @@ void FPhysScene::AddCustomPhysics(FBodyInstance* BodyInstance, FCalculateCustomP
 #endif
 }
 
-void FPhysScene::AddForce(FBodyInstance* BodyInstance, const FVector& Force, bool bAllowSubstepping)
+void FPhysScene::AddForce(FBodyInstance* BodyInstance, const FVector& Force, bool bAllowSubstepping, bool bAccelChange)
 {
 #if WITH_PHYSX
 
@@ -252,13 +252,13 @@ void FPhysScene::AddForce(FBodyInstance* BodyInstance, const FVector& Force, boo
 		if (bAllowSubstepping && IsSubstepping(BodySceneType))
 		{
 			FPhysSubstepTask * PhysSubStepper = PhysSubSteppers[BodySceneType];
-			PhysSubStepper->AddForce(BodyInstance, Force);
+			PhysSubStepper->AddForce(BodyInstance, Force, bAccelChange);
 		}
 		else
 #endif
 		{
 			SCOPED_SCENE_WRITE_LOCK(PRigidBody->getScene());
-			PRigidBody->addForce(U2PVector(Force), PxForceMode::eFORCE, true);
+			PRigidBody->addForce(U2PVector(Force), bAccelChange ? PxForceMode::eACCELERATION : PxForceMode::eFORCE, true);
 		}
 	}
 #endif
@@ -287,7 +287,7 @@ void FPhysScene::AddForceAtPosition(FBodyInstance* BodyInstance, const FVector&
 #endif
 }
 
-void FPhysScene::AddTorque(FBodyInstance* BodyInstance, const FVector& Torque, bool bAllowSubstepping)
+void FPhysScene::AddTorque(FBodyInstance* BodyInstance, const FVector& Torque, bool bAllowSubstepping, bool bAccelChange)
 {
 #if WITH_PHYSX
 
@@ -298,13 +298,13 @@ void FPhysScene::AddTorque(FBodyInstance* BodyInstance, const FVector& Torque, b
 		if (bAllowSubstepping && IsSubstepping(BodySceneType))
 		{
 			FPhysSubstepTask * PhysSubStepper = PhysSubSteppers[BodySceneType];
-			PhysSubStepper->AddTorque(BodyInstance, Torque);
+			PhysSubStepper->AddTorque(BodyInstance, Torque, bAccelChange);
 		}
 		else
 #endif
 		{
 			SCOPED_SCENE_WRITE_LOCK(PRigidBody->getScene());
-			PRigidBody->addTorque(U2PVector(Torque), PxForceMode::eFORCE, true);
+			PRigidBody->addTorque(U2PVector(Torque), bAccelChange ? PxForceMode::eACCELERATION : PxForceMode::eFORCE, true);
 		}
 	}
 #endif
diff --git a/Engine/Source/Runtime/Engine/Private/PhysicsEngine/PhysSubstepTasks.cpp b/Engine/Source/Runtime/Engine/Private/PhysicsEngine/PhysSubstepTasks.cpp
--- a/Engine/Source/Runtime/Engine/Private/PhysicsEngine/PhysSubstepTasks.cpp
+++ b/Engine/Source/Runtime/Engine/Private/PhysicsEngine/PhysSubstepTasks.cpp
@@ -102,7 +102,7 @@ void FPhysSubstepTask::AddCustomPhysics(FBodyInstance* Body, const FCalculateCus
 #endif
 }
 
-void FPhysSubstepTask::AddForce(FBodyInstance* Body, const FVector& Force)
+void FPhysSubstepTask::AddForce(FBodyInstance* Body, const FVector& Force, bool bAccelChange)
 {
 #if WITH_PHYSX
 	check(Body);
@@ -115,6 +115,7 @@ void FPhysSubstepTask::AddForce(FBodyInstance* Body, const FVector& Force)
 		FForceTarget ForceTarget;
 		ForceTarget.bPosition = false;
 		ForceTarget.Force = Force;
+		ForceTarget.bAccelChange = bAccelChange;
 
 		FPhysTarget & TargetState = PhysTargetBuffers[External].FindOrAdd(Body);
 		TargetState.Forces.Add(ForceTarget);
@@ -142,7 +143,7 @@ void FPhysSubstepTask::AddForceAtPosition(FBodyInstance* Body, const FVector& Fo
 	}
 #endif
 }
-void FPhysSubstepTask::AddTorque(FBodyInstance* Body, const FVector& Torque)
+void FPhysSubstepTask::AddTorque(FBodyInstance* Body, const FVector& Torque, bool bAccelChange)
 {
 #if WITH_PHYSX
 	check(Body);
@@ -154,6 +155,7 @@ void FPhysSubstepTask::AddTorque(FBodyInstance* Body, const FVector& Torque)
 	{
 		FTorqueTarget TorqueTarget;
 		TorqueTarget.Torque = Torque;
+		TorqueTarget.bAccelChange = bAccelChange;
 
 		FPhysTarget & TargetState = PhysTargetBuffers[External].FindOrAdd(Body);
 		TargetState.Torques.Add(TorqueTarget);
@@ -191,7 +193,7 @@ void FPhysSubstepTask::ApplyForces(const FPhysTarget& PhysTarget, FBodyInstance*
 		}
 		else
 		{
-			PRigidBody->addForce(U2PVector(ForceTarget.Force), PxForceMode::eFORCE, true);
+			PRigidBody->addForce(U2PVector(ForceTarget.Force), ForceTarget.bAccelChange ? PxForceMode::eACCELERATION : PxForceMode::eFORCE, true);
 		}
 	}
 #endif
@@ -207,7 +209,7 @@ void FPhysSubstepTask::ApplyTorques(const FPhysTarget& PhysTarget, FBodyInstance
 	for (int32 i = 0; i < PhysTarget.Torques.Num(); ++i)
 	{
 		const FTorqueTarget& TorqueTarget = PhysTarget.Torques[i];
-		PRigidBody->addTorque(U2PVector(TorqueTarget.Torque), PxForceMode::eFORCE, true);
+		PRigidBody->addTorque(U2PVector(TorqueTarget.Torque), TorqueTarget.bAccelChange ? PxForceMode::eACCELERATION : PxForceMode::eFORCE, true);
 	}
 #endif
 }
diff --git a/Engine/Source/Runtime/Engine/Private/PhysicsEngine/PhysSubstepTasks.h b/Engine/Source/Runtime/Engine/Private/PhysicsEngine/PhysSubstepTasks.h
--- a/Engine/Source/Runtime/Engine/Private/PhysicsEngine/PhysSubstepTasks.h
+++ b/Engine/Source/Runtime/Engine/Private/PhysicsEngine/PhysSubstepTasks.h
@@ -64,6 +64,7 @@ struct FForceTarget
 	FVector Force;
 	FVector Position;
 	bool bPosition;
+	bool bAccelChange;
 };
 
 struct FTorqueTarget
@@ -71,6 +72,7 @@ struct FTorqueTarget
 	FTorqueTarget(){}
 	FTorqueTarget(const FVector& GivenTorque) : Torque(GivenTorque){}
 	FVector Torque;
+	bool bAccelChange;
 };
 
 struct FCustomTarget
@@ -115,9 +117,9 @@ public:
 	void SetKinematicTarget(FBodyInstance* Body, const FTransform& TM);
 	bool GetKinematicTarget(const FBodyInstance* Body, FTransform& OutTM) const;
 	void AddCustomPhysics(FBodyInstance* Body, const FCalculateCustomPhysics& CalculateCustomPhysics);
-	void AddForce(FBodyInstance* Body, const FVector& Force);
+	void AddForce(FBodyInstance* Body, const FVector& Force, bool bAccelChange);
 	void AddForceAtPosition(FBodyInstance* Body, const FVector& Force, const FVector& Position);
-	void AddTorque(FBodyInstance* Body, const FVector& Torque);
+	void AddTorque(FBodyInstance* Body, const FVector& Torque, bool bAccelChange);
 
 	/** Removes a BodyInstance from doing substep work - should only be called when the FBodyInstance is getting destroyed */
 	void RemoveBodyInstance(FBodyInstance* Body);
diff --git a/Engine/Source/Runtime/Engine/Private/PhysicsEngine/PhysXSupport.cpp b/Engine/Source/Runtime/Engine/Private/PhysicsEngine/PhysXSupport.cpp
--- a/Engine/Source/Runtime/Engine/Private/PhysicsEngine/PhysXSupport.cpp
+++ b/Engine/Source/Runtime/Engine/Private/PhysicsEngine/PhysXSupport.cpp
@@ -249,7 +249,7 @@ void AddRadialImpulseToPxRigidBody(PxRigidBody& PRigidBody, const FVector& Origi
 #endif // WITH_PHYSX
 }
 
-void AddRadialForceToPxRigidBody(PxRigidBody& PRigidBody, const FVector& Origin, float Radius, float Strength, uint8 Falloff)
+void AddRadialForceToPxRigidBody(PxRigidBody& PRigidBody, const FVector& Origin, float Radius, float Strength, uint8 Falloff, bool bAccelChange)
 {
 #if WITH_PHYSX
 	if (!(PRigidBody.getRigidBodyFlags() & PxRigidBodyFlag::eKINEMATIC))
@@ -279,7 +279,7 @@ void AddRadialForceToPxRigidBody(PxRigidBody& PRigidBody, const FVector& Origin,
 
 		// Apply force
 		PxVec3 PImpulse = PDelta * ForceMag;
-		PRigidBody.addForce(PImpulse, PxForceMode::eFORCE);
+		PRigidBody.addForce(PImpulse, bAccelChange ? PxForceMode::eACCELERATION : PxForceMode::eFORCE);
 	}
 #endif // WITH_PHYSX
 }
diff --git a/Engine/Source/Runtime/Engine/Private/PhysicsEngine/PhysXSupport.h b/Engine/Source/Runtime/Engine/Private/PhysicsEngine/PhysXSupport.h
--- a/Engine/Source/Runtime/Engine/Private/PhysicsEngine/PhysXSupport.h
+++ b/Engine/Source/Runtime/Engine/Private/PhysicsEngine/PhysXSupport.h
@@ -167,7 +167,7 @@ void DeferredPhysResourceCleanup();
 ENGINE_API void AddRadialImpulseToPxRigidBody(PxRigidBody& PRigidBody, const FVector& Origin, float Radius, float Strength, uint8 Falloff, bool bVelChange);
 
 /** Calculates correct force at the body's center of mass and adds force to the body. */
-ENGINE_API void AddRadialForceToPxRigidBody(PxRigidBody& PRigidBody, const FVector& Origin, float Radius, float Strength, uint8 Falloff);
+ENGINE_API void AddRadialForceToPxRigidBody(PxRigidBody& PRigidBody, const FVector& Origin, float Radius, float Strength, uint8 Falloff, bool bAccelChange);
 
 /** Util to see if a PxRigidActor is non-kinematic */
 bool IsRigidBodyNonKinematic(PxRigidBody* PRigidBody);
diff --git a/Engine/Source/Runtime/Engine/Private/PrimitiveComponentPhysics.cpp b/Engine/Source/Runtime/Engine/Private/PrimitiveComponentPhysics.cpp
--- a/Engine/Source/Runtime/Engine/Private/PrimitiveComponentPhysics.cpp
+++ b/Engine/Source/Runtime/Engine/Private/PrimitiveComponentPhysics.cpp
@@ -230,12 +230,12 @@ void UPrimitiveComponent::AddRadialImpulse(FVector Origin, float Radius, float S
 }
 
 
-void UPrimitiveComponent::AddForce(FVector Force, FName BoneName)
+void UPrimitiveComponent::AddForce(FVector Force, FName BoneName, bool bAccelChange)
 {
 	if (FBodyInstance* BI = GetBodyInstance(BoneName))
 	{
 		WarnInvalidPhysicsOperations(LOCTEXT("AddForce", "AddForce"), BI);
-		BI->AddForce(Force);
+		BI->AddForce(Force, true, bAccelChange);
 	}
 }
 
@@ -248,7 +248,7 @@ void UPrimitiveComponent::AddForceAtLocation(FVector Force, FVector Location, FN
 	}
 }
 
-void UPrimitiveComponent::AddRadialForce(FVector Origin, float Radius, float Strength, ERadialImpulseFalloff Falloff)
+void UPrimitiveComponent::AddRadialForce(FVector Origin, float Radius, float Strength, ERadialImpulseFalloff Falloff, bool bAccelChange)
 {
 	if(bIgnoreRadialForce)
 	{
@@ -258,16 +258,16 @@ void UPrimitiveComponent::AddRadialForce(FVector Origin, float Radius, float Str
 	FBodyInstance* BI = GetBodyInstance();
 	if (BI)
 	{
-		BI->AddRadialForceToBody(Origin, Radius, Strength, Falloff);
+		BI->AddRadialForceToBody(Origin, Radius, Strength, Falloff, bAccelChange);
 	}
 }
 
-void UPrimitiveComponent::AddTorque(FVector Torque, FName BoneName)
+void UPrimitiveComponent::AddTorque(FVector Torque, FName BoneName, bool bAccelChange)
 {
 	if (FBodyInstance* BI = GetBodyInstance(BoneName))
 	{
 		WarnInvalidPhysicsOperations(LOCTEXT("AddTorque", "AddTorque"), BI);
-		BI->AddTorque(Torque);
+		BI->AddTorque(Torque, true, bAccelChange);
 	}
 }
 
diff --git a/Engine/Source/Runtime/Engine/Private/SkeletalMeshComponentPhysics.cpp b/Engine/Source/Runtime/Engine/Private/SkeletalMeshComponentPhysics.cpp
--- a/Engine/Source/Runtime/Engine/Private/SkeletalMeshComponentPhysics.cpp
+++ b/Engine/Source/Runtime/Engine/Private/SkeletalMeshComponentPhysics.cpp
@@ -477,14 +477,14 @@ void USkeletalMeshComponent::AddRadialImpulse(FVector Origin, float Radius, floa
 	const float StrengthPerMass = Strength / FMath::Max(GetMass(), KINDA_SMALL_NUMBER);
 	for(int32 i=0; i<Bodies.Num(); i++)
 	{
-		const float StrengthPerBody = bVelChange ? Strength : StrengthPerMass * Bodies[i]->GetBodyMass();
+		const float StrengthPerBody = bVelChange ? Strength : (StrengthPerMass * Bodies[i]->GetBodyMass());
 		Bodies[i]->AddRadialImpulseToBody(Origin, Radius, StrengthPerBody, Falloff, bVelChange);
 	}
 }
 
 
 
-void USkeletalMeshComponent::AddRadialForce(FVector Origin, float Radius, float Strength, ERadialImpulseFalloff Falloff)
+void USkeletalMeshComponent::AddRadialForce(FVector Origin, float Radius, float Strength, ERadialImpulseFalloff Falloff, bool bAccelChange)
 {
 	if(bIgnoreRadialForce)
 	{
@@ -494,8 +494,8 @@ void USkeletalMeshComponent::AddRadialForce(FVector Origin, float Radius, float
 	const float StrengthPerMass = Strength / FMath::Max(GetMass(), KINDA_SMALL_NUMBER);
 	for(int32 i=0; i<Bodies.Num(); i++)
 	{
-		const float StrengthPerBody = StrengthPerMass * Bodies[i]->GetBodyMass();
-		Bodies[i]->AddRadialForceToBody(Origin, Radius, StrengthPerBody, Falloff);
+		const float StrengthPerBody = bAccelChange ? Strength : (StrengthPerMass * Bodies[i]->GetBodyMass());
+		Bodies[i]->AddRadialForceToBody(Origin, Radius, StrengthPerBody, Falloff, bAccelChange);
 	}
 
 }
diff --git a/Engine/Source/Runtime/Engine/Public/PhysicsPublic.h b/Engine/Source/Runtime/Engine/Public/PhysicsPublic.h
--- a/Engine/Source/Runtime/Engine/Public/PhysicsPublic.h
+++ b/Engine/Source/Runtime/Engine/Public/PhysicsPublic.h
@@ -442,13 +442,13 @@ public:
 	void AddCustomPhysics(FBodyInstance* BodyInstance, FCalculateCustomPhysics& CalculateCustomPhysics);
 
 	/** Adds a force to a body - We need to go through scene to support substepping */
-	void AddForce(FBodyInstance* BodyInstance, const FVector& Force, bool bAllowSubstepping);
+	void AddForce(FBodyInstance* BodyInstance, const FVector& Force, bool bAllowSubstepping, bool bAccelChange);
 
 	/** Adds a force to a body at a specific position - We need to go through scene to support substepping */
 	void AddForceAtPosition(FBodyInstance* BodyInstance, const FVector& Force, const FVector& Position, bool bAllowSubstepping);
 
 	/** Adds torque to a body - We need to go through scene to support substepping */
-	void AddTorque(FBodyInstance* BodyInstance, const FVector& Torque, bool bAllowSubstepping);
+	void AddTorque(FBodyInstance* BodyInstance, const FVector& Torque, bool bAllowSubstepping, bool bAccelChange);
 
 	/** Sets a Kinematic actor's target position - We need to do this here to support substepping*/
 	void SetKinematicTarget(FBodyInstance* BodyInstance, const FTransform& TargetTM, bool bAllowSubstepping);
