Add substepping support for radial forces

diff --git a/Engine/Source/Runtime/Engine/Private/PhysicsEngine/BodyInstance.cpp b/Engine/Source/Runtime/Engine/Private/PhysicsEngine/BodyInstance.cpp
--- a/Engine/Source/Runtime/Engine/Private/PhysicsEngine/BodyInstance.cpp
+++ b/Engine/Source/Runtime/Engine/Private/PhysicsEngine/BodyInstance.cpp
@@ -3895,13 +3895,16 @@ void FBodyInstance::AddRadialImpulseToBody(const FVector& Origin, float Radius,
 #endif
 }
 
-void FBodyInstance::AddRadialForceToBody(const FVector& Origin, float Radius, float Strength, uint8 Falloff, bool bAccelChange)
+void FBodyInstance::AddRadialForceToBody(const FVector& Origin, float Radius, float Strength, uint8 Falloff, bool bAccelChange, bool bAllowSubstepping)
 {
 #if WITH_PHYSX
 	PxRigidBody* PRigidBody = GetPxRigidBody();
 	if (IsRigidBodyNonKinematic(PRigidBody))
 	{
-		AddRadialForceToPxRigidBody(*PRigidBody, Origin, Radius, Strength, Falloff, bAccelChange);
+		const PxScene* PScene = PRigidBody->getScene();
+		FPhysScene* PhysScene = FPhysxUserData::Get<FPhysScene>(PScene->userData);
+		PhysScene->AddRadialForceToBody(this, Origin, Radius, Strength, Falloff, bAccelChange, bAllowSubstepping);
+
 	}
 #endif // WITH_PHYSX
 
diff --git a/Engine/Source/Runtime/Engine/Private/PhysicsEngine/PhysScene.cpp b/Engine/Source/Runtime/Engine/Private/PhysicsEngine/PhysScene.cpp
--- a/Engine/Source/Runtime/Engine/Private/PhysicsEngine/PhysScene.cpp
+++ b/Engine/Source/Runtime/Engine/Private/PhysicsEngine/PhysScene.cpp
@@ -287,6 +287,29 @@ void FPhysScene::AddForceAtPosition(FBodyInstance* BodyInstance, const FVector&
 #endif
 }
 
+void FPhysScene::AddRadialForceToBody(FBodyInstance* BodyInstance, const FVector& Origin, const float Radius, const float Strength, const uint8 Falloff, bool bAccelChange, bool bAllowSubstepping)
+{
+#if WITH_PHYSX
+
+	if (PxRigidBody * PRigidBody = BodyInstance->GetPxRigidBody())
+	{
+#if WITH_SUBSTEPPING
+		uint32 BodySceneType = SceneType(BodyInstance);
+		if (bAllowSubstepping && IsSubstepping(BodySceneType))
+		{
+			FPhysSubstepTask * PhysSubStepper = PhysSubSteppers[BodySceneType];
+			PhysSubStepper->AddRadialForceToBody(BodyInstance, Origin, Radius, Strength, Falloff, bAccelChange);
+		}
+		else
+#endif
+		{
+			SCOPED_SCENE_WRITE_LOCK(PRigidBody->getScene());
+			AddRadialForceToPxRigidBody(*PRigidBody, Origin, Radius, Strength, Falloff, bAccelChange);
+		}
+	}
+#endif
+}
+
 void FPhysScene::AddTorque(FBodyInstance* BodyInstance, const FVector& Torque, bool bAllowSubstepping, bool bAccelChange)
 {
 #if WITH_PHYSX
diff --git a/Engine/Source/Runtime/Engine/Private/PhysicsEngine/PhysSubstepTasks.cpp b/Engine/Source/Runtime/Engine/Private/PhysicsEngine/PhysSubstepTasks.cpp
--- a/Engine/Source/Runtime/Engine/Private/PhysicsEngine/PhysSubstepTasks.cpp
+++ b/Engine/Source/Runtime/Engine/Private/PhysicsEngine/PhysSubstepTasks.cpp
@@ -163,6 +163,29 @@ void FPhysSubstepTask::AddTorque(FBodyInstance* Body, const FVector& Torque, boo
 #endif
 }
 
+void FPhysSubstepTask::AddRadialForceToBody(FBodyInstance* Body, const FVector& Origin, const float Radius, const float Strength, const uint8 Falloff, const bool bAccelChange)
+{
+#if WITH_PHYSX
+	check(Body);
+
+	PxRigidBody* PRigidBody = Body->GetPxRigidBody();
+	SCOPED_SCENE_READ_LOCK(PRigidBody->getScene());
+	//We should only apply torque on non kinematic actors
+	if (IsRigidBodyNonKinematic(PRigidBody))
+	{
+		FRadialForceTarget RadialForceTarget;
+		RadialForceTarget.Origin = Origin;
+		RadialForceTarget.Radius = Radius;
+		RadialForceTarget.Strength = Strength;
+		RadialForceTarget.Falloff = Falloff;
+		RadialForceTarget.bAccelChange = bAccelChange;
+
+		FPhysTarget & TargetState = PhysTargetBuffers[External].FindOrAdd(Body);
+		TargetState.RadialForces.Add(RadialForceTarget);
+	}
+#endif
+}
+
 /** Applies custom physics - Assumes caller has obtained writer lock */
 void FPhysSubstepTask::ApplyCustomPhysics(const FPhysTarget& PhysTarget, FBodyInstance* BodyInstance, float DeltaTime)
 {
@@ -214,6 +237,22 @@ void FPhysSubstepTask::ApplyTorques(const FPhysTarget& PhysTarget, FBodyInstance
 #endif
 }
 
+/** Applies radial forces - Assumes caller has obtained writer lock */
+void FPhysSubstepTask::ApplyRadialForces(const FPhysTarget& PhysTarget, FBodyInstance* BodyInstance)
+{
+#if WITH_PHYSX
+	/** Apply Torques */
+	PxRigidBody* PRigidBody = BodyInstance->GetPxRigidBody();
+
+	for (int32 i = 0; i < PhysTarget.RadialForces.Num(); ++i)
+	{
+		const FRadialForceTarget& RadialForceTArget= PhysTarget.RadialForces[i];
+		AddRadialForceToPxRigidBody(*PRigidBody, RadialForceTArget.Origin, RadialForceTArget.Radius, RadialForceTArget.Strength, RadialForceTArget.Falloff, RadialForceTArget.bAccelChange);
+	}
+#endif
+}
+
+
 /** Interpolates kinematic actor transform - Assumes caller has obtained writer lock */
 void FPhysSubstepTask::InterpolateKinematicActor(const FPhysTarget& PhysTarget, FBodyInstance* BodyInstance, float InAlpha)
 {
@@ -274,6 +313,7 @@ void FPhysSubstepTask::SubstepInterpolation(float InAlpha, float DeltaTime)
 		ApplyCustomPhysics(PhysTarget, BodyInstance, DeltaTime);
 		ApplyForces(PhysTarget, BodyInstance);
 		ApplyTorques(PhysTarget, BodyInstance);
+		ApplyRadialForces(PhysTarget, BodyInstance);
 		InterpolateKinematicActor(PhysTarget, BodyInstance, InAlpha);
 	}
 
diff --git a/Engine/Source/Runtime/Engine/Private/PhysicsEngine/PhysSubstepTasks.h b/Engine/Source/Runtime/Engine/Private/PhysicsEngine/PhysSubstepTasks.h
--- a/Engine/Source/Runtime/Engine/Private/PhysicsEngine/PhysSubstepTasks.h
+++ b/Engine/Source/Runtime/Engine/Private/PhysicsEngine/PhysSubstepTasks.h
@@ -75,6 +75,16 @@ struct FTorqueTarget
 	bool bAccelChange;
 };
 
+struct FRadialForceTarget
+{
+	FRadialForceTarget(){}
+	FVector Origin;
+	float Radius;
+	float Strength;
+	uint8 Falloff;
+	bool bAccelChange;
+};
+
 struct FCustomTarget
 {
 	FCustomTarget(){ }
@@ -89,6 +99,7 @@ struct FPhysTarget
 
 	TArray<FForceTarget>  Forces;	//we can apply force at multiple places
 	TArray<FTorqueTarget> Torques;
+	TArray<FRadialForceTarget> RadialForces;
 	TArray<FCustomTarget> CustomPhysics; //for calculating custom physics forces
 	FKinematicTarget KinematicTarget;
 
@@ -120,6 +131,7 @@ public:
 	void AddForce(FBodyInstance* Body, const FVector& Force, bool bAccelChange);
 	void AddForceAtPosition(FBodyInstance* Body, const FVector& Force, const FVector& Position);
 	void AddTorque(FBodyInstance* Body, const FVector& Torque, bool bAccelChange);
+	void AddRadialForceToBody(FBodyInstance* Body, const FVector& Origin, const float Radius, const float Strength, const uint8 Falloff, const bool bAccelChange);
 
 	/** Removes a BodyInstance from doing substep work - should only be called when the FBodyInstance is getting destroyed */
 	void RemoveBodyInstance(FBodyInstance* Body);
@@ -142,6 +154,7 @@ private:
 	void ApplyCustomPhysics(const FPhysTarget& PhysTarget, FBodyInstance* BodyInstance, float DeltaTime);
 	void ApplyForces(const FPhysTarget& PhysTarget, FBodyInstance* BodyInstance);
 	void ApplyTorques(const FPhysTarget& PhysTarget, FBodyInstance* BodyInstance);
+	void ApplyRadialForces(const FPhysTarget& PhysTarget, FBodyInstance* BodyInstance);
 	void InterpolateKinematicActor(const FPhysTarget& PhysTarget, FBodyInstance* BodyInstance, float Alpha);
 
 	typedef TMap<FBodyInstance*, FPhysTarget> PhysTargetMap;
diff --git a/Engine/Source/Runtime/Engine/Public/PhysicsPublic.h b/Engine/Source/Runtime/Engine/Public/PhysicsPublic.h
--- a/Engine/Source/Runtime/Engine/Public/PhysicsPublic.h
+++ b/Engine/Source/Runtime/Engine/Public/PhysicsPublic.h
@@ -447,6 +447,9 @@ public:
 	/** Adds a force to a body at a specific position - We need to go through scene to support substepping */
 	void AddForceAtPosition(FBodyInstance* BodyInstance, const FVector& Force, const FVector& Position, bool bAllowSubstepping);
 
+	/** Adds a radial force to a body - We need to go through scene to support substepping */
+	void AddRadialForceToBody(FBodyInstance* BodyInstance, const FVector& Origin, const float Radius, const float Strength, const uint8 Falloff, bool bAccelChange, bool bAllowSubstepping);
+
 	/** Adds torque to a body - We need to go through scene to support substepping */
 	void AddTorque(FBodyInstance* BodyInstance, const FVector& Torque, bool bAllowSubstepping, bool bAccelChange);
 
