Fix build, forgot to submit file

diff --git a/Engine/Source/Runtime/Engine/Classes/PhysicsEngine/BodyInstance.h b/Engine/Source/Runtime/Engine/Classes/PhysicsEngine/BodyInstance.h
--- a/Engine/Source/Runtime/Engine/Classes/PhysicsEngine/BodyInstance.h
+++ b/Engine/Source/Runtime/Engine/Classes/PhysicsEngine/BodyInstance.h
@@ -734,8 +734,9 @@ public:
 	 *	@param Strength		Strength of force to apply.
 	 *  @param Falloff		Allows you to control the strength of the force as a function of distance from Origin.
 	 *  @param bAccelChange If true, Strength is taken as a change in acceleration instead of a physical force (i.e. mass will have no affect).
+	 *  @param bAllowSubstepping Whether we should sub-step this radial force. You should only turn this off if you're calling it from a sub-step callback, otherwise there will be energy loss
 	 */
-	void AddRadialForceToBody(const FVector& Origin, float Radius, float Strength, uint8 Falloff, bool bAccelChange = false);
+	void AddRadialForceToBody(const FVector& Origin, float Radius, float Strength, uint8 Falloff, bool bAccelChange = false, bool bAllowSubstepping = true);
 
 	/**
 	 * Get distance to the body surface if available
