Fixed UE-9905 "Free camera doesn't work in LogVisualilzer tool in StrategyGame."
Improved ToggleDebugCamera cheat to work correctly with projects like StrategyGame - we can fly free around level now.

diff --git a/Engine/Source/Runtime/Engine/Classes/Engine/DebugCameraController.h b/Engine/Source/Runtime/Engine/Classes/Engine/DebugCameraController.h
--- a/Engine/Source/Runtime/Engine/Classes/Engine/DebugCameraController.h
+++ b/Engine/Source/Runtime/Engine/Classes/Engine/DebugCameraController.h
@@ -117,6 +117,8 @@ public:
 	virtual FString ConsoleCommand(const FString& Command, bool bWriteToLog = true) override;
 	virtual void AddCheats(bool bForce) override;
 	virtual void EndSpectatingState() override;
+	/** Custom spawn to spawn a default SpectatorPawn, to use as a spectator and initialize it. By default it is spawned at the PC's current location and rotation. */
+	virtual ASpectatorPawn* SpawnSpectatorPawn() override;
 
 protected:
 
diff --git a/Engine/Source/Runtime/Engine/Private/DebugCameraController.cpp b/Engine/Source/Runtime/Engine/Private/DebugCameraController.cpp
--- a/Engine/Source/Runtime/Engine/Private/DebugCameraController.cpp
+++ b/Engine/Source/Runtime/Engine/Private/DebugCameraController.cpp
@@ -14,7 +14,7 @@
 #include "Engine/DebugCameraController.h"
 #include "Components/DrawFrustumComponent.h"
 #include "GameFramework/PlayerInput.h"
-
+#include "GameFramework/GameState.h"
 
 /** The currently selected actor. */
 AActor* GDebugSelectedActor = NULL;
@@ -246,6 +246,43 @@ void ADebugCameraController::UpdateHiddenComponents(const FVector& ViewLocation,
 	}
 }
 
+ASpectatorPawn* ADebugCameraController::SpawnSpectatorPawn()
+{
+	ASpectatorPawn* SpawnedSpectator = NULL;
+
+	// Only spawned for the local player
+	if (GetSpectatorPawn() == NULL && IsLocalController())
+	{
+		AGameState const* const GameState = GetWorld()->GameState;
+		if (GameState)
+		{
+			FActorSpawnParameters SpawnParams;
+			SpawnParams.Owner = this;
+			SpawnParams.bNoCollisionFail = true;
+			SpawnParams.ObjectFlags |= RF_Transient;	// We never want to save spectator pawns into a map
+			SpawnedSpectator = GetWorld()->SpawnActor<ASpectatorPawn>(ASpectatorPawn::StaticClass(), GetSpawnLocation(), GetControlRotation(), SpawnParams);
+			if (SpawnedSpectator)
+			{
+				SpawnedSpectator->PossessedBy(this);
+				SpawnedSpectator->PawnClientRestart();
+				SpawnedSpectator->SetActorTickEnabled(true);
+
+				UE_LOG(LogPlayerController, Verbose, TEXT("Spawned spectator %s [server:%d]"), *GetNameSafe(SpawnedSpectator), GetNetMode() < NM_Client);
+			}
+			else
+			{
+				UE_LOG(LogPlayerController, Warning, TEXT("Failed to spawn spectator with class %s"), GameState->SpectatorClass ? *GameState->SpectatorClass->GetName() : TEXT("NULL"));
+			}
+		}
+		else
+		{
+			// This normally happens on clients if the Player is replicated but the GameState has not yet.
+			UE_LOG(LogPlayerController, Verbose, TEXT("NULL GameState when trying to spawn spectator!"));
+		}
+	}
+
+	return SpawnedSpectator != NULL ? SpawnedSpectator : Super::SpawnSpectatorPawn();
+}
 
 void ADebugCameraController::SetSpectatorPawn(ASpectatorPawn* NewSpectatorPawn)
 {
