Track bActorHasBegunPlay on an Actor
Encapsulate bActorInitialized and expose via Getter

diff --git a/Engine/Source/Runtime/Engine/Classes/GameFramework/Actor.h b/Engine/Source/Runtime/Engine/Classes/GameFramework/Actor.h
--- a/Engine/Source/Runtime/Engine/Classes/GameFramework/Actor.h
+++ b/Engine/Source/Runtime/Engine/Classes/GameFramework/Actor.h
@@ -453,14 +453,21 @@ public:
 	UPROPERTY()
 	TWeakObjectPtr<AActor> ParentComponentActor;	
 
-public:
+private:
 
 	/** 
 	 *	Indicates that PreInitializeComponents/PostInitializeComponents have been called on this Actor 
 	 *	Prevents re-initializing of actors spawned during level startup
 	 */
 	uint32 bActorInitialized:1;
+	
+	/** 
+	 *	Indicates that BeginPlay has been called for this Actor.
+	 *  Set back to false once EndPlay has been called.
+	 */
+	uint32 bActorHasBegunPlay:1;
 
+public:
 	/** Indicates the actor was pulled through a seamless travel.  */
 	UPROPERTY()
 	uint32 bActorSeamlessTraveled:1;
@@ -987,6 +994,12 @@ public:
 	/** Event when play begins for this actor. */
 	virtual void BeginPlay();
 
+	/** Returns whether an actor has been initialized */
+	bool IsActorInitialized() const { return bActorInitialized; }
+
+	/** Returns whether an actor has had BeginPlay called on it (and not subsequently had EndPlay called) */
+	bool HasActorBegunPlay() const { return bActorHasBegunPlay; }
+
 	/** Event when this actor takes ANY damage */
 	UFUNCTION(BlueprintImplementableEvent, BlueprintAuthorityOnly, meta=(FriendlyName = "AnyDamage"), Category="Game|Damage")
 	virtual void ReceiveAnyDamage(float Damage, const class UDamageType* DamageType, class AController* InstigatedBy, AActor* DamageCauser);
diff --git a/Engine/Source/Runtime/Engine/Private/Actor.cpp b/Engine/Source/Runtime/Engine/Private/Actor.cpp
--- a/Engine/Source/Runtime/Engine/Private/Actor.cpp
+++ b/Engine/Source/Runtime/Engine/Private/Actor.cpp
@@ -1523,9 +1523,14 @@ void AActor::RouteEndPlay(const EEndPlayReason::Type EndPlayReason)
 
 void AActor::EndPlay(const EEndPlayReason::Type EndPlayReason)
 {
-	// Dispatch the blueprint events
-	ReceiveEndPlay(EndPlayReason);
-	OnEndPlay.Broadcast(EndPlayReason);
+	if (bActorHasBegunPlay)
+	{
+		bActorHasBegunPlay = false;
+
+		// Dispatch the blueprint events
+		ReceiveEndPlay(EndPlayReason);
+		OnEndPlay.Broadcast(EndPlayReason);
+	}
 
 	// Behaviors specific to an actor being unloaded due to a streaming level removal
 	if (EndPlayReason == EEndPlayReason::RemovedFromWorld)
@@ -2435,9 +2440,12 @@ void AActor::ExchangeNetRoles(bool bRemoteOwned)
 
 void AActor::BeginPlay()
 {
+	ensure(!bActorHasBegunPlay);
 	SetLifeSpan( InitialLifeSpan );
 
 	ReceiveBeginPlay();
+
+	bActorHasBegunPlay = true;
 }
 
 void AActor::EnableInput(APlayerController* PlayerController)
diff --git a/Engine/Source/Runtime/Engine/Private/ActorComponent.cpp b/Engine/Source/Runtime/Engine/Private/ActorComponent.cpp
--- a/Engine/Source/Runtime/Engine/Private/ActorComponent.cpp
+++ b/Engine/Source/Runtime/Engine/Private/ActorComponent.cpp
@@ -748,7 +748,7 @@ void UActorComponent::RegisterComponentWithWorld(UWorld* InWorld)
 	ExecuteRegisterEvents();
 	RegisterAllComponentTickFunctions(true);
 
-	if (Owner == nullptr || Owner->bActorInitialized)
+	if (Owner == nullptr || Owner->IsActorInitialized())
 	{
 		if (!bHasBeenInitialized && bWantsInitializeComponent)
 		{
diff --git a/Engine/Source/Runtime/Engine/Private/Level.cpp b/Engine/Source/Runtime/Engine/Private/Level.cpp
--- a/Engine/Source/Runtime/Engine/Private/Level.cpp
+++ b/Engine/Source/Runtime/Engine/Private/Level.cpp
@@ -1468,7 +1468,7 @@ void ULevel::InitializeNetworkActors()
 		if( Actor )
 		{
 			// Kill off actors that aren't interesting to the client.
-			if( !Actor->bActorInitialized && !Actor->bActorSeamlessTraveled )
+			if( !Actor->IsActorInitialized() && !Actor->bActorSeamlessTraveled )
 			{
 				// Add to startup list
 				if (Actor->bNetLoadOnClient)
@@ -1522,15 +1522,11 @@ void ULevel::ReleaseRenderingResources()
 void ULevel::RouteActorInitialize()
 {
 	// Send PreInitializeComponents and collect volumes.
-	for( int32 ActorIndex=0; ActorIndex<Actors.Num(); ActorIndex++ )
+	for( AActor* const Actor : Actors )
 	{
-		AActor* const Actor = Actors[ActorIndex];
-		if( Actor )
+		if( Actor && !Actor->IsActorInitialized() )
 		{
-			if( !Actor->bActorInitialized )
-			{
-				Actor->PreInitializeComponents();
-			}
+			Actor->PreInitializeComponents();
 		}
 	}
 
@@ -1538,18 +1534,17 @@ void ULevel::RouteActorInitialize()
 	TArray<AActor *> ActorsToBeginPlay;
 
 	// Send InitializeComponents on components and PostInitializeComponents.
-	for( int32 ActorIndex=0; ActorIndex<Actors.Num(); ActorIndex++ )
+	for( AActor* const Actor : Actors )
 	{
-		AActor* Actor = Actors[ActorIndex];
 		if( Actor )
 		{
-			if( !Actor->bActorInitialized )
+			if( !Actor->IsActorInitialized() )
 			{
 				// Call Initialize on Components.
 				Actor->InitializeComponents();
 
 				Actor->PostInitializeComponents(); // should set Actor->bActorInitialized = true
-				if (!Actor->bActorInitialized && !Actor->IsPendingKill())
+				if (!Actor->IsActorInitialized() && !Actor->IsPendingKill())
 				{
 					UE_LOG(LogActor, Fatal, TEXT("%s failed to route PostInitializeComponents.  Please call Super::PostInitializeComponents() in your <className>::PostInitializeComponents() function. "), *Actor->GetFullName() );
 				}
diff --git a/Engine/Source/Runtime/Engine/Private/PrimitiveComponent.cpp b/Engine/Source/Runtime/Engine/Private/PrimitiveComponent.cpp
--- a/Engine/Source/Runtime/Engine/Private/PrimitiveComponent.cpp
+++ b/Engine/Source/Runtime/Engine/Private/PrimitiveComponent.cpp
@@ -1256,7 +1256,7 @@ static bool ShouldIgnoreOverlapResult(const UWorld* World, const AActor* ThisAct
 		return true;
 	}
 
-	if (!World || OtherActor == World->GetWorldSettings() || !OtherActor->bActorInitialized)
+	if (!World || OtherActor == World->GetWorldSettings() || !OtherActor->IsActorInitialized())
 	{
 		return true;
 	}
@@ -1292,7 +1292,7 @@ bool UPrimitiveComponent::CheckStaticMobilityAndWarn(const FText& ActionText) co
 {
 	AActor* Actor = GetOwner();
 	// static things can move before they are registered (e.g. immediately after streaming), but not after.
-	if (Mobility == EComponentMobility::Static && Actor && Actor->bActorInitialized)
+	if (Mobility == EComponentMobility::Static && Actor && Actor->IsActorInitialized())
 	{
 		FMessageLog("PIE").Warning(FText::Format(LOCTEXT("InvalidStaticMove", "Mobility of {0} : {1} has to be 'Movable' if you'd like to {2}. "),
 			FText::FromString(GetNameSafe(GetOwner())), FText::FromString(GetName()), ActionText));
@@ -2179,7 +2179,7 @@ void UPrimitiveComponent::UpdateOverlaps(TArray<FOverlapInfo> const* PendingOver
 		// if we haven't begun play, we're still setting things up (e.g. we might be inside one of the construction scripts)
 		// so we don't want to generate overlaps yet.
 		AActor* const MyActor = GetOwner();
-		if ( MyActor && MyActor->bActorInitialized )
+		if ( MyActor && MyActor->IsActorInitialized() )
 		{
 			if (PendingOverlaps)
 			{
diff --git a/Engine/Source/Runtime/Engine/Private/World.cpp b/Engine/Source/Runtime/Engine/Private/World.cpp
--- a/Engine/Source/Runtime/Engine/Private/World.cpp
+++ b/Engine/Source/Runtime/Engine/Private/World.cpp
@@ -2867,7 +2867,7 @@ void UWorld::InitializeActorsForPlay(const FURL& InURL, bool bResetTime)
 		}
 
 		// Init the game mode.
-		if (AuthorityGameMode && !AuthorityGameMode->bActorInitialized)
+		if (AuthorityGameMode && !AuthorityGameMode->IsActorInitialized())
 		{
 			AuthorityGameMode->InitGame( FPaths::GetBaseFilename(InURL.Map), Options, Error );
 		}
