Don't allow writing to CachedMaxDrawDistance (since it doesn't do what you want anyways)

diff --git a/Engine/Source/Runtime/Engine/Classes/Components/PrimitiveComponent.h b/Engine/Source/Runtime/Engine/Classes/Components/PrimitiveComponent.h
--- a/Engine/Source/Runtime/Engine/Classes/Components/PrimitiveComponent.h
+++ b/Engine/Source/Runtime/Engine/Classes/Components/PrimitiveComponent.h
@@ -128,7 +128,7 @@ public:
 	 * The distance to cull this primitive at.  
 	 * A CachedMaxDrawDistance of 0 indicates that the primitive should not be culled by distance.
 	 */
-	UPROPERTY(Category=LOD, AdvancedDisplay, VisibleAnywhere, BlueprintReadWrite, meta=(DisplayName="Current Max Draw Distance") )
+	UPROPERTY(Category=LOD, AdvancedDisplay, VisibleAnywhere, BlueprintReadOnly, meta=(DisplayName="Current Max Draw Distance") )
 	float CachedMaxDrawDistance;
 
 	/** The scene depth priority group to draw the primitive in. */
