Clean up constraint instance init code.
Fix SetAngularParams and SetLinearParams ignoring certain flags

diff --git a/Engine/Source/Runtime/Engine/Classes/PhysicsEngine/ConstraintInstance.h b/Engine/Source/Runtime/Engine/Classes/PhysicsEngine/ConstraintInstance.h
--- a/Engine/Source/Runtime/Engine/Classes/PhysicsEngine/ConstraintInstance.h
+++ b/Engine/Source/Runtime/Engine/Classes/PhysicsEngine/ConstraintInstance.h
@@ -7,6 +7,8 @@
 namespace physx
 {
 	class PxD6Joint;
+	class PxRigidActor;
+	class PxScene;
 }
 #endif // WITH_PHYSX
 
@@ -489,6 +491,17 @@ public:
 	//Hacks to easily get zeroed memory for special case when we don't use GC
 	static void Free(FConstraintInstance * Ptr);
 	static FConstraintInstance * Alloc();
+
+
+private:
+#if WITH_PHYSX 
+	bool CreatePxJoint(physx::PxRigidActor* PActor1, physx::PxRigidActor* PActor2, physx::PxScene* PScene, const float Scale);
+	void UpdateConstraintFlags();
+	void UpdateAverageMass(const physx::PxRigidActor* PActor1, const physx::PxRigidActor* PActor2);
+#endif
+
+	void UpdateBreakable();
+	void UpdateDriveTarget();
 };
 
 template<>
diff --git a/Engine/Source/Runtime/Engine/Private/PhysicsEngine/ConstraintInstance.cpp b/Engine/Source/Runtime/Engine/Private/PhysicsEngine/ConstraintInstance.cpp
--- a/Engine/Source/Runtime/Engine/Private/PhysicsEngine/ConstraintInstance.cpp
+++ b/Engine/Source/Runtime/Engine/Private/PhysicsEngine/ConstraintInstance.cpp
@@ -7,6 +7,8 @@
 
 #define LOCTEXT_NAMESPACE "ConstraintInstance"
 
+const bool bIsAccelerationDrive = true;
+
 /** Handy macro for setting BIT of VAR based on the bool CONDITION */
 #define SET_DRIVE_PARAM(VAR, CONDITION, BIT)   (VAR) = (CONDITION) ? ((VAR) | (BIT)) : ((VAR) & ~(BIT))
 
@@ -141,6 +143,29 @@ void FConstraintInstance::UpdateAngularLimit()
 	
 }
 
+void FConstraintInstance::UpdateBreakable()
+{
+	const float LinearBreakForce = bLinearBreakable ? LinearBreakThreshold : PX_MAX_REAL;
+	const float AngularBreakForce = bAngularBreakable ? AngularBreakThreshold : PX_MAX_REAL;
+
+#if WITH_PHYSX
+	ConstraintData->setBreakForce(LinearBreakForce, AngularBreakForce);
+#endif
+}
+
+void FConstraintInstance::UpdateDriveTarget()
+{
+#if WITH_PHYSX
+	if (ConstraintData && !(ConstraintData->getConstraintFlags()&PxConstraintFlag::eBROKEN))
+	{
+		FQuat OrientationTargetQuat(AngularOrientationTarget);
+
+		ConstraintData->setDrivePosition(PxTransform(U2PVector(LinearPositionTarget), U2PQuat(OrientationTargetQuat)));
+		ConstraintData->setDriveVelocity(U2PVector(LinearVelocityTarget), U2PVector(AngularVelocityTarget * 2 * PI));
+	}
+#endif
+}
+
 /** Constructor **/
 FConstraintInstance::FConstraintInstance()
 	: ConstraintIndex(0)
@@ -220,22 +245,36 @@ void FConstraintInstance::SetDisableCollision(bool InDisableCollision)
 #endif
 }
 
-/** 
- *	Create physics engine constraint.
- */
-void FConstraintInstance::InitConstraint(USceneComponent* Owner, FBodyInstance* Body1, FBodyInstance* Body2, float Scale)
+#if WITH_PHYSX
+float ComputeAverageMass(const PxRigidActor* PActor1, const PxRigidActor* PActor2)
 {
-	OwnerComponent = Owner;
+	float AverageMass = 0;
+	{
+		float TotalMass = 0;
+		int NumDynamic = 0;
 
-#if WITH_PHYSX
-	PhysxUserData = FPhysxUserData(this);
+		if (PActor1 && PActor1->isRigidBody())
+		{
+			TotalMass += PActor1->isRigidBody()->getMass();
+			++NumDynamic;
+		}
 
-	// if there's already a constraint, get rid of it first
-	if (ConstraintData != NULL)
-	{
-		TermConstraint();
+		if (PActor2 && PActor2->isRigidBody())
+		{
+			TotalMass += PActor2->isRigidBody()->getMass();
+			++NumDynamic;
+		}
+
+		check(NumDynamic);
+		AverageMass = TotalMass / NumDynamic;
 	}
 
+	return AverageMass;
+}
+
+/*various logical checks to find the correct physx actor. Returns true if found valid actors that can be constrained*/
+bool GetPActors(const FBodyInstance* Body1, const FBodyInstance* Body2, PxRigidActor** PActor1Out, PxRigidActor** PActor2Out)
+{
 	PxRigidActor* PActor1 = Body1 ? Body1->GetPxRigidActor() : NULL;
 	PxRigidActor* PActor2 = Body2 ? Body2->GetPxRigidActor() : NULL;
 
@@ -243,11 +282,11 @@ void FConstraintInstance::InitConstraint(USceneComponent* Owner, FBodyInstance*
 	// Do not create joint unless one of the actors is dynamic
 	if ((!PActor1 || !PActor1->isRigidBody()) && (!PActor2 || !PActor2->isRigidBody()))
 	{
-		return;
+		return false;
 	}
 
 	// Need to worry about the case where one is static and one is dynamic, and make sure the static scene is used which matches the dynamic scene
-	if(PActor1 != NULL && PActor2 != NULL)
+	if (PActor1 != NULL && PActor2 != NULL)
 	{
 		if (PActor1->isRigidStatic() && PActor2->isRigidBody())
 		{
@@ -262,47 +301,33 @@ void FConstraintInstance::InitConstraint(USceneComponent* Owner, FBodyInstance*
 		}
 	}
 
-	AverageMass = 0;
-	{
-		float TotalMass = 0;
-		int NumDynamic = 0;
-
-		if (PActor1 && PActor1->isRigidBody())
-		{
-			TotalMass += PActor1->isRigidBody()->getMass();
-			++NumDynamic;
-		}
-
-		if (PActor2 && PActor2->isRigidBody())
-		{
-			TotalMass += PActor2->isRigidBody()->getMass();
-			++NumDynamic;
-		}
-
-		check(NumDynamic);
-		AverageMass = TotalMass / NumDynamic;
-	}
-
-
-	// record if actors are asleep before creating joint, so we can sleep them afterwards if so (creating joint wakes them)
-	const bool bActor1Asleep = (PActor1 == NULL || !PActor1->isRigidDynamic() || PActor1->isRigidDynamic()->isSleeping());
-	const bool bActor2Asleep = (PActor2 == NULL || !PActor2->isRigidDynamic() || PActor2->isRigidDynamic()->isSleeping());
+	*PActor1Out = PActor1;
+	*PActor2Out = PActor2;
+	return true;
+}
 
+bool GetPScene(const PxRigidActor* PActor1, const PxRigidActor* PActor2, PxScene** PSceneOut)
+{
 	// make sure actors are in same scene
-	PxScene* PScene1 = (PActor1 != NULL) ? PActor1->getScene() : NULL;
-	PxScene* PScene2 = (PActor2 != NULL) ? PActor2->getScene() : NULL;
+	PxScene* PScene1 = PActor1 ? PActor1->getScene() : nullptr;
+	PxScene* PScene2 = PActor2 ? PActor2->getScene() : nullptr;
 
 	// make sure actors are in same scene
-	if(PScene1 && PScene2 && PScene1 != PScene2)
+	if (PScene1 && PScene2 && PScene1 != PScene2)
 	{
-		UE_LOG(LogPhysics, Log,  TEXT("Attempting to create a joint between actors in two different scenes.  No joint created.") );
-		return;
+		UE_LOG(LogPhysics, Log, TEXT("Attempting to create a joint between actors in two different scenes.  No joint created."));
+		return false;
 	}
 
-	PxScene* PScene = PScene1 ? PScene1 : PScene2;
-	check(PScene);
+	*PSceneOut = PScene1 ? PScene1 : PScene2;
+	check(PSceneOut);
+	
+	return true;
+}
 
-	ConstraintData = NULL;
+bool FConstraintInstance::CreatePxJoint(physx::PxRigidActor* PActor1, physx::PxRigidActor* PActor2, physx::PxScene* PScene, const float Scale)
+{
+	ConstraintData = nullptr;
 
 	FTransform Local1 = GetRefFrame(EConstraintFrame::Frame1);
 	Local1.ScaleTranslation(FVector(Scale));
@@ -317,10 +342,10 @@ void FConstraintInstance::InitConstraint(USceneComponent* Owner, FBodyInstance*
 	// Because PhysX keeps limits/axes locked in the first body reference frame, whereas Unreal keeps them in the second body reference frame, we have to flip the bodies here.
 	PxD6Joint* PD6Joint = PxD6JointCreate(*GPhysXSDK, PActor2, U2PTransform(Local2), PActor1, U2PTransform(Local1));
 
-	if(PD6Joint == NULL)
+	if (PD6Joint == nullptr)
 	{
 		UE_LOG(LogPhysics, Log, TEXT("URB_ConstraintInstance::InitConstraint - Invalid 6DOF joint (%s)"), *JointName.ToString());
-		return;
+		return false;
 	}
 
 	///////// POINTERS
@@ -328,134 +353,123 @@ void FConstraintInstance::InitConstraint(USceneComponent* Owner, FBodyInstance*
 
 	// Remember reference to scene index.
 	FPhysScene* RBScene = FPhysxUserData::Get<FPhysScene>(PScene->userData);
-	if(RBScene->GetPhysXScene(PST_Sync) == PScene)
+	if (RBScene->GetPhysXScene(PST_Sync) == PScene)
 	{
 		SceneIndex = RBScene->PhysXSceneIndex[PST_Sync];
 	}
 	else
-	if(RBScene->GetPhysXScene(PST_Async) == PScene)
+	if (RBScene->GetPhysXScene(PST_Async) == PScene)
 	{
 		SceneIndex = RBScene->PhysXSceneIndex[PST_Async];
 	}
 	else
 	{
-		UE_LOG(LogPhysics, Log,  TEXT("URB_ConstraintInstance::InitConstraint: PxScene has inconsistent FPhysScene userData.  No joint created.") );
-		return;
+		UE_LOG(LogPhysics, Log, TEXT("URB_ConstraintInstance::InitConstraint: PxScene has inconsistent FPhysScene userData.  No joint created."));
+		return false;
 	}
 
 	ConstraintData = PD6Joint;
+	return true;
+}
 
-	///////// FLAGS/PROJECTION
-
+void FConstraintInstance::UpdateConstraintFlags()
+{
 	PxConstraintFlags Flags = PxConstraintFlags();
 
 #if !(UE_BUILD_SHIPPING || UE_BUILD_TEST)
 	Flags |= PxConstraintFlag::eVISUALIZATION;
 #endif
 
-	if(!bDisableCollision)
+	if (!bDisableCollision)
 	{
 		Flags |= PxConstraintFlag::eCOLLISION_ENABLED;
 	}
 
-	if( bEnableProjection )
+	if (bEnableProjection)
 	{
 		Flags |= PxConstraintFlag::ePROJECTION;
 
-		PD6Joint->setProjectionLinearTolerance(ProjectionLinearTolerance); 
-		PD6Joint->setProjectionAngularTolerance(ProjectionAngularTolerance * ((float)PI/180.0f)); 
+		ConstraintData->setProjectionLinearTolerance(ProjectionLinearTolerance);
+		ConstraintData->setProjectionAngularTolerance(ProjectionAngularTolerance * ((float)PI / 180.0f));
 	}
 
-	PD6Joint->setConstraintFlags(Flags);
-
-	/////////////// ANUGULAR LIMIT
-	UpdateAngularLimit();
-
-	/////////////// LINEAR LIMIT
-	UpdateLinearLimit();
+	ConstraintData->setConstraintFlags(Flags);
+}
 
-	///////// BREAKABLE
-	float LinearBreakForce = PX_MAX_REAL;
-	float AngularBreakForce = PX_MAX_REAL;
-	if (bLinearBreakable)
-	{
-		LinearBreakForce = LinearBreakThreshold;
-	}
 
-	if (bAngularBreakable)
-	{
-		AngularBreakForce = AngularBreakThreshold;
-	}
-
-	PD6Joint->setBreakForce(LinearBreakForce, AngularBreakForce);
+void FConstraintInstance::UpdateAverageMass(const PxRigidActor* PActor1, const PxRigidActor* PActor2)
+{
+	AverageMass = ComputeAverageMass(PActor1, PActor2);
+}
 
-	///////// DRIVE
-	const PxReal LinearForceLimit = LinearDriveForceLimit > 0.0f ? LinearDriveForceLimit : PX_MAX_F32;
-	const PxReal AngularForceLimit = AngularDriveForceLimit > 0.0f ? AngularDriveForceLimit : PX_MAX_F32;
-	const bool bAccelerationDrive = true;
+void EnsureSleepingActorsStaySleeping(PxRigidActor* PActor1, PxRigidActor* PActor2)
+{
+	// record if actors are asleep before creating joint, so we can sleep them afterwards if so (creating joint wakes them)
+	const bool bActor1Asleep = (PActor1 == nullptr || !PActor1->isRigidDynamic() || PActor1->isRigidDynamic()->isSleeping());
+	const bool bActor2Asleep = (PActor2 == nullptr || !PActor2->isRigidDynamic() || PActor2->isRigidDynamic()->isSleeping());
 
-	// Set up the linear drives
-	if ( bLinearPositionDrive || bLinearVelocityDrive )
+	// creation of joints wakes up rigid bodies, so we put them to sleep again if both were initially asleep
+	if (bActor1Asleep && bActor2Asleep)
 	{
-		// X-Axis linear drive
-		float DriveSpring = bLinearXPositionDrive ? LinearDriveSpring : 0.0f;
-		float DriveDamping = (bLinearVelocityDrive && FMath::Abs(LinearVelocityTarget.X) > 0.0f) ? LinearDriveDamping : 0.0f;
-
-		PD6Joint->setDrive(PxD6Drive::eX, PxD6JointDrive(DriveSpring, DriveDamping, LinearForceLimit, bAccelerationDrive));
+		if (PActor1 != NULL && IsRigidBodyNonKinematic(PActor1->isRigidDynamic()))
+		{
+			PActor1->isRigidDynamic()->putToSleep();
+		}
 
-		// Y-Axis linear drive
-		DriveSpring = bLinearYPositionDrive ? LinearDriveSpring : 0.0f;
-		DriveDamping = (bLinearVelocityDrive && FMath::Abs(LinearVelocityTarget.Y) > 0.0f) ? LinearDriveDamping : 0.0f;
+		if (PActor2 != NULL && IsRigidBodyNonKinematic(PActor2->isRigidDynamic()))
+		{
+			PActor2->isRigidDynamic()->putToSleep();
+		}
+	}
+}
 
-		PD6Joint->setDrive(PxD6Drive::eY, PxD6JointDrive(DriveSpring, DriveDamping, LinearForceLimit, bAccelerationDrive));
+#endif
 
-		// Z-Axis linear drive
-		DriveSpring = bLinearZPositionDrive ? LinearDriveSpring : 0.0f;
-		DriveDamping = (bLinearVelocityDrive && FMath::Abs(LinearVelocityTarget.Z) > 0.0f) ? LinearDriveDamping : 0.0f;
+/** 
+ *	Create physics engine constraint.
+ */
+void FConstraintInstance::InitConstraint(USceneComponent* Owner, FBodyInstance* Body1, FBodyInstance* Body2, float Scale)
+{
+	OwnerComponent = Owner;
 
-		PD6Joint->setDrive(PxD6Drive::eZ, PxD6JointDrive(DriveSpring, DriveDamping, LinearForceLimit, bAccelerationDrive));
+#if WITH_PHYSX
+	PhysxUserData = FPhysxUserData(this);
 
-		// Warn the user if he specified a drive for an axis that has no DOF
-		if (LinearXMotion == LCM_Locked )
-		{
-			UE_LOG(LogPhysics, Warning,  TEXT("Attempting to create a linear joint drive for an locked axis.") );
-		}
+	// if there's already a constraint, get rid of it first
+	if (ConstraintData)
+	{
+		TermConstraint();
 	}
 
-	if ( bAngularOrientationDrive || bAngularVelocityDrive )
-	{
-		float DriveSpring = bAngularOrientationDrive ? AngularDriveSpring : 0.0f;
-		float DriveDamping = bAngularVelocityDrive  ? AngularDriveDamping : 0.0f;
+	PxRigidActor* PActor1;
+	PxRigidActor* PActor2;
+	PxScene* PScene;
 
-		if (AngularDriveMode == EAngularDriveMode::SLERP)
-		{
-			PD6Joint->setDrive(PxD6Drive::eSLERP, PxD6JointDrive(DriveSpring, DriveDamping, AngularForceLimit, bAccelerationDrive));
-		}
-		else
-		{
-			PD6Joint->setDrive(PxD6Drive::eTWIST, PxD6JointDrive(DriveSpring, DriveDamping, AngularForceLimit, bAccelerationDrive));
-			PD6Joint->setDrive(PxD6Drive::eSWING, PxD6JointDrive(DriveSpring, DriveDamping, AngularForceLimit, bAccelerationDrive));
-		}
+	const bool bValidConstraintSetup = GetPActors(Body1, Body2, &PActor1, &PActor2) && GetPScene(PActor1, PActor2, &PScene) && CreatePxJoint(PActor1, PActor2, PScene, Scale);
+	if (!bValidConstraintSetup)
+	{
+		return;
 	}
 
-	FQuat OrientationTargetQuat(AngularOrientationTarget);
-
-	PD6Joint->setDrivePosition(PxTransform(U2PVector(LinearPositionTarget), U2PQuat(OrientationTargetQuat)));
-	PD6Joint->setDriveVelocity(U2PVector(LinearVelocityTarget), U2PVector(AngularVelocityTarget * 2 * PI));
+	// update mass
+	UpdateAverageMass(PActor1, PActor2);
 	
-	// creation of joints wakes up rigid bodies, so we put them to sleep again if both were initially asleep
-	if(bActor1Asleep && bActor2Asleep)
-	{
-		if ( PActor1 != NULL && IsRigidBodyNonKinematic(PActor1->isRigidDynamic()))
-		{
-			PActor1->isRigidDynamic()->putToSleep();
-		}
+	//flags and projection settings
+	UpdateConstraintFlags();
+	
+	//limits
+	UpdateAngularLimit();
+	UpdateLinearLimit();
 
-		if (PActor2 != NULL && IsRigidBodyNonKinematic(PActor2->isRigidDynamic()))
-		{
-			PActor2->isRigidDynamic()->putToSleep();
-		}
-	}
+	//breakable
+	UpdateBreakable();
+	
+	//motors
+	SetLinearDriveParams(LinearDriveSpring, LinearDriveDamping, LinearDriveForceLimit);
+	SetAngularDriveParams(AngularDriveSpring, AngularDriveDamping, AngularDriveForceLimit);
+	UpdateDriveTarget();
+	
+	EnsureSleepingActorsStaySleeping(PActor1, PActor2);
 #endif // WITH_PHYSX
 
 }
@@ -765,11 +779,11 @@ void FConstraintInstance::SetAngularPositionDrive(bool bEnableSwingDrive, bool b
 		PxD6JointDrive CurrentDriveSwing = Joint->getDrive(PxD6Drive::eSWING);
 		PxD6JointDrive CurrentDriveTwist = Joint->getDrive(PxD6Drive::eTWIST);
 		PxD6JointDrive CurrentDriveSlerp = Joint->getDrive(PxD6Drive::eSLERP);
+		const bool bSlerp = AngularDriveMode == EAngularDriveMode::SLERP;
 
-
-		CurrentDriveSwing.stiffness = bEnableSwingDrive ? AngularDriveSpring : 0.0f;
-		CurrentDriveTwist.stiffness = bEnableTwistDrive ? AngularDriveSpring : 0.0f;
-		CurrentDriveSlerp.stiffness = (bEnableSwingDrive && bEnableTwistDrive) ? AngularDriveSpring : 0.0f;
+		CurrentDriveSwing.stiffness = !bSlerp && bEnableSwingDrive ? AngularDriveSpring : 0.0f;
+		CurrentDriveTwist.stiffness = !bSlerp && bEnableTwistDrive ? AngularDriveSpring : 0.0f;
+		CurrentDriveSlerp.stiffness = bSlerp  && (bEnableSwingDrive || bEnableTwistDrive) ? AngularDriveSpring : 0.0f;
 
 		Joint->setDrive(PxD6Drive::eSWING, CurrentDriveSwing);
 		Joint->setDrive(PxD6Drive::eTWIST, CurrentDriveTwist);
@@ -791,10 +805,11 @@ void FConstraintInstance::SetAngularVelocityDrive(bool bEnableSwingDrive, bool b
 		PxD6JointDrive CurrentDriveSwing = Joint->getDrive(PxD6Drive::eSWING);
 		PxD6JointDrive CurrentDriveTwist = Joint->getDrive(PxD6Drive::eTWIST);
 		PxD6JointDrive CurrentDriveSlerp = Joint->getDrive(PxD6Drive::eSLERP);
+		const bool bSlerp = AngularDriveMode == EAngularDriveMode::SLERP;
 
-		CurrentDriveSwing.damping = bEnableSwingDrive ? AngularDriveDamping : 0.0f;
-		CurrentDriveTwist.damping = bEnableTwistDrive ? AngularDriveDamping : 0.0f;
-		CurrentDriveSlerp.damping = (bEnableSwingDrive && bEnableTwistDrive) ? AngularDriveDamping : 0.0f;
+		CurrentDriveSwing.damping = !bSlerp && bEnableSwingDrive ? AngularDriveDamping : 0.0f;
+		CurrentDriveTwist.damping = !bSlerp && bEnableTwistDrive ? AngularDriveDamping : 0.0f;
+		CurrentDriveSlerp.damping = bSlerp  && (bEnableSwingDrive && bEnableTwistDrive) ? AngularDriveDamping : 0.0f;
 
 		Joint->setDrive(PxD6Drive::eSWING, CurrentDriveSwing);
 		Joint->setDrive(PxD6Drive::eTWIST, CurrentDriveTwist);
@@ -849,23 +864,33 @@ void FConstraintInstance::SetLinearVelocityTarget(const FVector& InVelTarget)
 	LinearVelocityTarget = InVelTarget;
 }
 
-/** Function for setting angular motor parameters. */
+/** Function for setting linear motor parameters. */
 void FConstraintInstance::SetLinearDriveParams(float InSpring, float InDamping, float InForceLimit)
 {
 #if WITH_PHYSX
-	PxD6Joint* Joint = (PxD6Joint*)ConstraintData;
-	if (Joint &&  !(Joint->getConstraintFlags()&PxConstraintFlag::eBROKEN))
+	if (bLinearPositionDrive || bLinearVelocityDrive)
 	{
-		float LinearForceLimit = InForceLimit;
+		if (ConstraintData &&  !(ConstraintData->getConstraintFlags()&PxConstraintFlag::eBROKEN))
+		{
+			// X-Axis linear drive
+			const float DriveSpringX = bLinearPositionDrive && bLinearXPositionDrive ? InSpring : 0.0f;
+			const float DriveDampingX = (bLinearVelocityDrive && LinearVelocityTarget.X != 0.f) ? InDamping : 0.0f;
+			const float LinearForceLimit = InForceLimit > 0.f ? InForceLimit : PX_MAX_F32;
+
+			ConstraintData->setDrive(PxD6Drive::eX, PxD6JointDrive(DriveSpringX, DriveDampingX, LinearForceLimit, bIsAccelerationDrive));
 
-		PxD6JointDrive NewJointDrive(InSpring, InDamping, LinearForceLimit > 0.f ? LinearForceLimit : PX_MAX_F32, Joint->getDrive(PxD6Drive::eX).flags);
-		Joint->setDrive(PxD6Drive::eX, NewJointDrive);
+			// Y-Axis linear drive
+			const float DriveSpringY = bLinearPositionDrive && bLinearYPositionDrive ? InSpring : 0.0f;
+			const float DriveDampingY = (bLinearVelocityDrive && LinearVelocityTarget.Y != 0.f) ? InDamping : 0.0f;
 
-		NewJointDrive.flags = Joint->getDrive(PxD6Drive::eY).flags;
-		Joint->setDrive(PxD6Drive::eY, NewJointDrive);
+			ConstraintData->setDrive(PxD6Drive::eY, PxD6JointDrive(DriveSpringY, DriveDampingY, LinearForceLimit, bIsAccelerationDrive));
 
-		NewJointDrive.flags = Joint->getDrive(PxD6Drive::eZ).flags;
-		Joint->setDrive(PxD6Drive::eZ, NewJointDrive);
+			// Z-Axis linear drive
+			const float DriveSpringZ = bLinearPositionDrive && bLinearZPositionDrive ? InSpring : 0.0f;
+			const float DriveDampingZ = (bLinearVelocityDrive && LinearVelocityTarget.Z != 0.f) ? InDamping : 0.0f;
+
+			ConstraintData->setDrive(PxD6Drive::eZ, PxD6JointDrive(DriveSpringZ, DriveDampingZ, LinearForceLimit, bIsAccelerationDrive));
+		}
 	}
 #endif
 
@@ -968,20 +993,27 @@ void FConstraintInstance::SetAngularVelocityTarget(const FVector& InVelTarget)
 void FConstraintInstance::SetAngularDriveParams(float InSpring, float InDamping, float InForceLimit)
 {
 #if WITH_PHYSX
-	
-	PxD6Joint* Joint = (PxD6Joint*)ConstraintData;
-	if (Joint &&  !(Joint->getConstraintFlags()&PxConstraintFlag::eBROKEN))
+	if (bAngularOrientationDrive || bAngularVelocityDrive)
 	{
-		float AngularForceLimit = InForceLimit;
-
-		PxD6JointDrive NewJointDrive(InSpring, InDamping, AngularForceLimit > 0.f ? AngularForceLimit : PX_MAX_F32, Joint->getDrive(PxD6Drive::eSWING).flags);
-		Joint->setDrive(PxD6Drive::eSWING, NewJointDrive);
-
-		NewJointDrive.flags = Joint->getDrive(PxD6Drive::eTWIST).flags;
-		Joint->setDrive(PxD6Drive::eTWIST, NewJointDrive);
+		if (ConstraintData &&  !(ConstraintData->getConstraintFlags()&PxConstraintFlag::eBROKEN))
+		{
+			const float AngularForceLimit = InForceLimit > 0.0f ? InForceLimit : PX_MAX_F32;
+			const float DriveSpring = bAngularOrientationDrive ? InSpring : 0.0f;
+			const float DriveDamping = bAngularVelocityDrive ? InDamping : 0.0f;
 
-		NewJointDrive.flags = Joint->getDrive(PxD6Drive::eSLERP).flags;
-		Joint->setDrive(PxD6Drive::eSLERP, NewJointDrive);
+			if (AngularDriveMode == EAngularDriveMode::SLERP)
+			{
+				ConstraintData->setDrive(PxD6Drive::eTWIST, PxD6JointDrive());
+				ConstraintData->setDrive(PxD6Drive::eSWING, PxD6JointDrive());
+				ConstraintData->setDrive(PxD6Drive::eSLERP, PxD6JointDrive(DriveSpring, DriveDamping, AngularForceLimit, bIsAccelerationDrive));
+			}
+			else
+			{
+				ConstraintData->setDrive(PxD6Drive::eTWIST, PxD6JointDrive(DriveSpring, DriveDamping, AngularForceLimit, bIsAccelerationDrive));
+				ConstraintData->setDrive(PxD6Drive::eSWING, PxD6JointDrive(DriveSpring, DriveDamping, AngularForceLimit, bIsAccelerationDrive));
+				ConstraintData->setDrive(PxD6Drive::eSLERP, PxD6JointDrive());
+			}
+		}
 	}
 #endif
 
@@ -990,11 +1022,6 @@ void FConstraintInstance::SetAngularDriveParams(float InSpring, float InDamping,
 	AngularDriveForceLimit = InForceLimit;
 }
 
-
-
-
-
-
 /** Scale Angular Limit Constraints (as defined in RB_ConstraintSetup) */
 void FConstraintInstance::SetAngularDOFLimitScale(float InSwing1LimitScale, float InSwing2LimitScale, float InTwistLimitScale)
 {
