Add the requested extra "fog" depth based blur to CircleDOF.

diff --git a/Engine/Shaders/PostProcessDOF.usf b/Engine/Shaders/PostProcessDOF.usf
--- a/Engine/Shaders/PostProcessDOF.usf
+++ b/Engine/Shaders/PostProcessDOF.usf
@@ -198,14 +198,22 @@ void MainRecombinePS(
 // PROTOTYPE CIRCLE DOF : WORK IN PROGRESS
 //
 
-float CircleDofParams;
+// {radius, depth blur amp, depth blur radius}
+float3 CircleDofParams;
 
 // CIRCLE DOF: Compute circle of confusion size in pixels.
 float DepthToCoc(float Depth) 
 {
 	float Focus = View.DepthOfFieldFocalDistance;
-	float Radius = CircleDofParams;
-	return ((Depth - Focus) / Depth) * Radius;
+	float Radius = CircleDofParams.x;
+	float CocRadius = ((Depth - Focus) / Depth) * Radius;
+	float DepthBlurRadius = (1.0 - exp2(-Depth * CircleDofParams.y)) * CircleDofParams.z;
+	float ReturnCoc = max(abs(CocRadius), DepthBlurRadius);
+	if(CocRadius < 0.0) 
+	{
+		ReturnCoc = -ReturnCoc;
+	}
+	return ReturnCoc;
 }
 
 // pixel shader entry point
diff --git a/Engine/Source/Runtime/Engine/Classes/Engine/Scene.h b/Engine/Source/Runtime/Engine/Classes/Engine/Scene.h
--- a/Engine/Source/Runtime/Engine/Classes/Engine/Scene.h
+++ b/Engine/Source/Runtime/Engine/Classes/Engine/Scene.h
@@ -320,6 +320,12 @@ struct FPostProcessSettings
 	UPROPERTY(BlueprintReadWrite, Category=Overrides, meta=(PinHiddenByDefault))
 	uint32 bOverride_DepthOfFieldFstop:1;
 
+	UPROPERTY(BlueprintReadWrite, Category=Overrides, meta=(PinHiddenByDefault))
+	uint32 bOverride_DepthOfFieldDepthBlurRadius:1;
+
+	UPROPERTY(BlueprintReadWrite, Category=Overrides, meta=(PinHiddenByDefault))
+	uint32 bOverride_DepthOfFieldDepthBlurAmount:1;
+
 	UPROPERTY(BlueprintReadWrite, Category=Overrides, meta=(PinHiddenByDefault))
 	uint32 bOverride_DepthOfFieldFocalRegion:1;
 
@@ -769,6 +775,14 @@ struct FPostProcessSettings
 	UPROPERTY(EditAnywhere, BlueprintReadWrite, Category=DepthOfField, meta=(editcondition = "bOverride_DepthOfFieldMethod", DisplayName = "Method"))
 	TEnumAsByte<enum EDepthOfFieldMethod> DepthOfFieldMethod;
 
+	/** CircleDOF only: Depth blur km for 50% */
+	UPROPERTY(interp, BlueprintReadWrite, Category=DepthOfField, meta=(ClampMin = "0.000001", ClampMax = "100.0", editcondition = "bOverride_DepthOfFieldDepthBlurAmount", DisplayName = "Depth Blur km for 50%"))
+	float DepthOfFieldDepthBlurAmount;
+
+	/** CircleDOF only: Depth blur radius in pixels at 1920x */
+	UPROPERTY(interp, BlueprintReadWrite, Category=DepthOfField, meta=(ClampMin = "0.0", ClampMax = "4.0", editcondition = "bOverride_DepthOfFieldDepthBlurRadius", DisplayName = "Depth Blur Radius"))
+	float DepthOfFieldDepthBlurRadius;
+	
 	/** CircleDOF only: F-stop number (1/fstop) */
 	UPROPERTY(interp, BlueprintReadWrite, Category=DepthOfField, meta=(ClampMin = "1.0", ClampMax = "32.0", editcondition = "bOverride_DepthOfFieldFstop", DisplayName = "1/Fstop"))
 	float DepthOfFieldFstop;
@@ -978,6 +992,8 @@ struct FPostProcessSettings
 		ColorGradingIntensity = 1.0f;
 		DepthOfFieldFocalDistance = 1000.0f;
 		DepthOfFieldFstop = 4.0f;
+		DepthOfFieldDepthBlurAmount = 1.0f;
+		DepthOfFieldDepthBlurRadius = 0.0f;
 		DepthOfFieldFocalRegion = 0.0f;
 		DepthOfFieldNearTransitionRegion = 300.0f;
 		DepthOfFieldFarTransitionRegion = 500.0f;
diff --git a/Engine/Source/Runtime/Engine/Private/SceneView.cpp b/Engine/Source/Runtime/Engine/Private/SceneView.cpp
--- a/Engine/Source/Runtime/Engine/Private/SceneView.cpp
+++ b/Engine/Source/Runtime/Engine/Private/SceneView.cpp
@@ -813,6 +813,8 @@ void FSceneView::OverridePostProcessSettings(const FPostProcessSettings& Src, fl
 	LERP_PP(IndirectLightingIntensity);
 	LERP_PP(DepthOfFieldFocalDistance);
 	LERP_PP(DepthOfFieldFstop);
+	LERP_PP(DepthOfFieldDepthBlurRadius);
+	LERP_PP(DepthOfFieldDepthBlurAmount);
 	LERP_PP(DepthOfFieldFocalRegion);
 	LERP_PP(DepthOfFieldNearTransitionRegion);
 	LERP_PP(DepthOfFieldFarTransitionRegion);
diff --git a/Engine/Source/Runtime/Renderer/Private/PostProcess/PostProcessDOF.cpp b/Engine/Source/Runtime/Renderer/Private/PostProcess/PostProcessDOF.cpp
--- a/Engine/Source/Runtime/Renderer/Private/PostProcess/PostProcessDOF.cpp
+++ b/Engine/Source/Runtime/Renderer/Private/PostProcess/PostProcessDOF.cpp
@@ -354,7 +354,8 @@ FPooledRenderTargetDesc FRCPassPostProcessDOFRecombine::ComputeOutputDesc(EPassO
 
 
 // Convert f-stop and focal distance into projected size in half resolution pixels.
-static float CircleDofCoc(const FRenderingCompositePassContext& Context)
+// Setup depth based blur.
+static FVector CircleDofCoc(const FRenderingCompositePassContext& Context)
 {
     // Convert FOV to focal length,
 	// 
@@ -398,7 +399,13 @@ static float CircleDofCoc(const FRenderingCompositePassContext& Context)
 			Radius = 6.0f; 
 		}
 	#endif
-	return Radius;
+
+	// The DepthOfFieldDepthBlurAmount = km at which depth blur is 50%.
+	// Need to convert to cm here.
+	return FVector(
+		Radius, 
+		1.0f/(Context.View.FinalPostProcessSettings.DepthOfFieldDepthBlurAmount * 100000.0),
+		Context.View.FinalPostProcessSettings.DepthOfFieldDepthBlurRadius * Width / 1920.0f);
 }
 
 /** Encapsulates the Circle DOF setup pixel shader. */
