Fix subtitle regressions

diff --git a/Engine/Source/Runtime/Engine/Private/GameViewportClient.cpp b/Engine/Source/Runtime/Engine/Private/GameViewportClient.cpp
--- a/Engine/Source/Runtime/Engine/Private/GameViewportClient.cpp
+++ b/Engine/Source/Runtime/Engine/Private/GameViewportClient.cpp
@@ -1093,11 +1093,11 @@ void UGameViewportClient::Draw(FViewport* InViewport, FCanvas* SceneCanvas)
 							FVector2D MaxPos(1.f, 1.f);
 							GetSubtitleRegion(MinPos, MaxPos);
 
-							uint32 SizeX = SceneCanvas->GetRenderTarget()->GetSizeXY().X;
-							uint32 SizeY = SceneCanvas->GetRenderTarget()->GetSizeXY().Y;
+							const uint32 SizeX = SceneCanvas->GetRenderTarget()->GetSizeXY().X;
+							const uint32 SizeY = SceneCanvas->GetRenderTarget()->GetSizeXY().Y;
 							FIntRect SubtitleRegion(FMath::TruncToInt(SizeX * MinPos.X), FMath::TruncToInt(SizeY * MinPos.Y), FMath::TruncToInt(SizeX * MaxPos.X), FMath::TruncToInt(SizeY * MaxPos.Y));
-							// We need a world to do this
 							FSubtitleManager::GetSubtitleManager()->DisplaySubtitles( SceneCanvas, SubtitleRegion, GetWorld()->GetAudioTimeSeconds() );
+							bDisplayedSubtitles = true;
 						}
 					}
 				}
diff --git a/Engine/Source/Runtime/Engine/Private/SubtitleManager.cpp b/Engine/Source/Runtime/Engine/Private/SubtitleManager.cpp
--- a/Engine/Source/Runtime/Engine/Private/SubtitleManager.cpp
+++ b/Engine/Source/Runtime/Engine/Private/SubtitleManager.cpp
@@ -120,13 +120,14 @@ void FSubtitleManager::DisplaySubtitle( FCanvas* Canvas, FActiveSubtitle* Subtit
 	// Needed to add a drop shadow and doing all 4 sides was the only way to make them look correct.  If this shows up as a framerate hit we'll have to think of a different way of dealing with this.
 	if( Subtitle->bSingleLine )
 	{
-		if ( !( Subtitle->Subtitles[Subtitle->Index].Text.IsEmpty() ) )
+		const FText& SubtitleText = Subtitle->Subtitles[Subtitle->Index].Text;
+		if ( !SubtitleText.IsEmpty() )
 		{
 			
 			// Display lines up from the bottom of the region
 			Parms.Max.Y -= SUBTITLE_CHAR_HEIGHT;
 			
-			FCanvasTextItem TextItem( FVector2D( Parms.Min.X + ( Parms.Width() / 2 ), Parms.Max.Y ), Subtitle->Subtitles[ Subtitle->Index ].Text , GEngine->GetSubtitleFont(), Color );
+			FCanvasTextItem TextItem( FVector2D( Parms.Min.X + ( Parms.Width() / 2 ), Parms.Max.Y ), SubtitleText , GEngine->GetSubtitleFont(), Color );
 			TextItem.Depth = SUBTITLE_SCREEN_DEPTH_FOR_3D;
 			TextItem.bOutlined = true;
 			TextItem.bCentreX = true;
@@ -140,10 +141,11 @@ void FSubtitleManager::DisplaySubtitle( FCanvas* Canvas, FActiveSubtitle* Subtit
 	{
 		for( int32 Idx = Subtitle->Subtitles.Num() - 1; Idx >= 0; Idx-- )
 		{
+			const FText& SubtitleText = Subtitle->Subtitles[Idx].Text;
 			// Display lines up from the bottom of the region
-			if ( !( Subtitle->Subtitles[Idx].Text.IsEmpty() ) )
+			if ( !SubtitleText.IsEmpty() )
 			{
-				FCanvasTextItem TextItem( FVector2D( Parms.Min.X + ( Parms.Width() / 2 ), Parms.Max.Y ), Subtitle->Subtitles[ Subtitle->Index ].Text, GEngine->GetSubtitleFont(), Color );
+				FCanvasTextItem TextItem( FVector2D( Parms.Min.X + ( Parms.Width() / 2 ), Parms.Max.Y ), SubtitleText, GEngine->GetSubtitleFont(), Color );
 				TextItem.Depth = SUBTITLE_SCREEN_DEPTH_FOR_3D;
 				TextItem.bOutlined = true;
 				TextItem.bCentreX = true;
