Integrate - Fixed instanced meshes not supporting speedtree wind

diff --git a/Engine/Shaders/ShaderVersion.usf b/Engine/Shaders/ShaderVersion.usf
--- a/Engine/Shaders/ShaderVersion.usf
+++ b/Engine/Shaders/ShaderVersion.usf
@@ -2,4 +2,4 @@
 // Each time the console command is executed it generates a new GUID. As this file is included
 // in common.usf (which should be included in any shader) it allows to invalidate the shader DDC.
 // 
-// GUID = {2490AEE3-37E5-4684-93D8-F170F910138C}
+// GUID = {2490AEE3-37E5-4684-93D8-F170F910238C}
diff --git a/Engine/Source/Runtime/Engine/Private/InstancedStaticMesh.h b/Engine/Source/Runtime/Engine/Private/InstancedStaticMesh.h
--- a/Engine/Source/Runtime/Engine/Private/InstancedStaticMesh.h
+++ b/Engine/Source/Runtime/Engine/Private/InstancedStaticMesh.h
@@ -192,6 +192,7 @@ public:
 		const bool bInstanced = RHISupportsInstancing(Platform);
 		OutEnvironment.SetDefine(TEXT("USE_INSTANCING_EMULATED"), bInstanced ? TEXT("0") : TEXT("1"));
 		OutEnvironment.SetDefine(TEXT("USE_DITHERED_LOD_TRANSITION"), ALLOW_DITHERED_LOD_FOR_INSTANCED_STATIC_MESHES ? TEXT("1") : TEXT("0"));
+		FLocalVertexFactory::ModifyCompilationEnvironment(Platform, Material, OutEnvironment);
 	}
 
 	/**
