Integrate - Rendering category is now shown for skylight actors

diff --git a/Engine/Source/Runtime/Engine/Classes/Engine/SkyLight.h b/Engine/Source/Runtime/Engine/Classes/Engine/SkyLight.h
--- a/Engine/Source/Runtime/Engine/Classes/Engine/SkyLight.h
+++ b/Engine/Source/Runtime/Engine/Classes/Engine/SkyLight.h
@@ -4,7 +4,7 @@
 #include "GameFramework/Info.h"
 #include "SkyLight.generated.h"
 
-UCLASS(ClassGroup=Lights, hidecategories=(Input,Collision,Replication,Info), showcategories=("Input|MouseInput", "Input|TouchInput"), ComponentWrapperClass, ConversionRoot, Blueprintable)
+UCLASS(ClassGroup=Lights, hidecategories=(Input,Collision,Replication,Info), showcategories=("Rendering", "Input|MouseInput", "Input|TouchInput"), ComponentWrapperClass, ConversionRoot, Blueprintable)
 class ENGINE_API ASkyLight : public AInfo
 {
 	GENERATED_UCLASS_BODY()
