Integrate - Skylights only capture on being enabled if they have never captured before
Added fast path for changing skylight color or intensity that doesn't recreate rendering state (full static draw list recreation)

diff --git a/Engine/Source/Runtime/Engine/Classes/Components/SkyLightComponent.h b/Engine/Source/Runtime/Engine/Classes/Components/SkyLightComponent.h
--- a/Engine/Source/Runtime/Engine/Classes/Components/SkyLightComponent.h
+++ b/Engine/Source/Runtime/Engine/Classes/Components/SkyLightComponent.h
@@ -201,6 +201,8 @@ protected:
 	/** Indicates whether the cached data stored in GetComponentInstanceData is valid to be applied in ApplyComponentInstanceData. */
 	bool bSavedConstructionScriptValuesValid;
 
+	bool bHasEverCaptured;
+
 	TRefCountPtr<FSkyTextureCubeResource> ProcessedSkyTexture;
 
 	FSHVectorRGB3 IrradianceEnvironmentMap;
@@ -222,6 +224,8 @@ protected:
 	virtual void DestroyRenderState_Concurrent() override;
 	// Begin UActorComponent Interface
 
+	void UpdateLimitedRenderingStateFast();
+
 	friend class FSkyLightSceneProxy;
 };
 
diff --git a/Engine/Source/Runtime/Engine/Private/SkyLightComponent.cpp b/Engine/Source/Runtime/Engine/Private/SkyLightComponent.cpp
--- a/Engine/Source/Runtime/Engine/Private/SkyLightComponent.cpp
+++ b/Engine/Source/Runtime/Engine/Private/SkyLightComponent.cpp
@@ -110,6 +110,7 @@ USkyLightComponent::USkyLightComponent(const FObjectInitializer& ObjectInitializ
 	bCaptureDirty = false;
 	bLowerHemisphereIsBlack = true;
 	bSavedConstructionScriptValuesValid = true;
+	bHasEverCaptured = false;
 	OcclusionMaxDistance = 1000;
 	MinOcclusion = 0;
 	OcclusionTint = FColor::Black;
@@ -195,6 +196,26 @@ void USkyLightComponent::PostLoad()
 	}
 }
 
+/** 
+ * Fast path for updating light properties that doesn't require a re-register,
+ * Which would otherwise cause the scene's static draw lists to be recreated.
+ */
+void USkyLightComponent::UpdateLimitedRenderingStateFast()
+{
+	if (SceneProxy)
+	{
+		ENQUEUE_UNIQUE_RENDER_COMMAND_THREEPARAMETER(
+			FFastUpdateSkyLightCommand,
+			FSkyLightSceneProxy*,LightSceneProxy,SceneProxy,
+			FLinearColor,LightColor,FLinearColor(LightColor) * Intensity,
+			float,IndirectLightingIntensity,IndirectLightingIntensity,
+		{
+			LightSceneProxy->LightColor = LightColor;
+			LightSceneProxy->IndirectLightingIntensity = IndirectLightingIntensity;
+		});
+	}
+}
+
 /** 
 * This is called when property is modified by InterpPropertyTracks
 *
@@ -211,7 +232,7 @@ void USkyLightComponent::PostInterpChange(UProperty* PropertyThatChanged)
 		|| PropertyName == IntensityName
 		|| PropertyName == IndirectLightingIntensityName)
 	{
-		MarkRenderStateDirty();
+		UpdateLimitedRenderingStateFast();
 	}
 	else
 	{
@@ -420,6 +441,7 @@ void USkyLightComponent::UpdateSkyCaptureContents(UWorld* WorldToUpdate)
 
 					WorldToUpdate->Scene->UpdateSkyCaptureContents(CaptureComponent, false, CaptureComponent->ProcessedSkyTexture, CaptureComponent->IrradianceEnvironmentMap);
 
+					CaptureComponent->bHasEverCaptured = true;
 					CaptureComponent->MarkRenderStateDirty();
 				}
 
@@ -450,7 +472,7 @@ void USkyLightComponent::SetIntensity(float NewIntensity)
 		&& Intensity != NewIntensity)
 	{
 		Intensity = NewIntensity;
-		MarkRenderStateDirty();
+		UpdateLimitedRenderingStateFast();
 	}
 }
 
@@ -464,7 +486,7 @@ void USkyLightComponent::SetLightColor(FLinearColor NewLightColor)
 		&& LightColor != NewColor)
 	{
 		LightColor = NewColor;
-		MarkRenderStateDirty();
+		UpdateLimitedRenderingStateFast();
 	}
 }
 
@@ -508,8 +530,9 @@ void USkyLightComponent::SetVisibility(bool bNewVisibility, bool bPropagateToChi
 
 	Super::SetVisibility(bNewVisibility, bPropagateToChildren);
 
-	if (bVisible && !bOldWasVisible)
+	if (bVisible && !bOldWasVisible && !bHasEverCaptured)
 	{
+		// Capture if we are being enabled for the first time
 		SetCaptureIsDirty();
 	}
 }
