Fix FRichCurve::Eval when the time value exactly coincides with a node time stamp

diff --git a/Engine/Source/Runtime/Engine/Private/CurveBase.cpp b/Engine/Source/Runtime/Engine/Private/CurveBase.cpp
--- a/Engine/Source/Runtime/Engine/Private/CurveBase.cpp
+++ b/Engine/Source/Runtime/Engine/Private/CurveBase.cpp
@@ -703,7 +703,7 @@ float FRichCurve::Eval(const float InTime, float DefaultValue) const
 		// If only one point, or before the first point in the curve, return the first points value.
 		InterpVal = Keys[0].Value;
 	}
-	else
+	else if (InTime < Keys[NumKeys - 1].Time)
 	{
 		// perform a lower bound to get the second of the interpolation nodes
 		int32 first = 1;
@@ -727,41 +727,37 @@ float FRichCurve::Eval(const float InTime, float DefaultValue) const
 		}
 
 		int32 InterpNode = first;
+		const float Diff = Keys[InterpNode].Time - Keys[InterpNode - 1].Time;
 
-		if (InTime < Keys[InterpNode].Time)
+		if (Diff > 0.f && Keys[InterpNode - 1].InterpMode != RCIM_Constant)
 		{
-			const float Diff = Keys[InterpNode].Time - Keys[InterpNode - 1].Time;
+			const float Alpha = (InTime - Keys[InterpNode - 1].Time) / Diff;
+			const float P0 = Keys[InterpNode - 1].Value;
+			const float P3 = Keys[InterpNode].Value;
 
-			if (Diff > 0.f && Keys[InterpNode - 1].InterpMode != RCIM_Constant)
+			if (Keys[InterpNode - 1].InterpMode == RCIM_Linear)
 			{
-				const float Alpha = (InTime - Keys[InterpNode - 1].Time) / Diff;
-				const float P0 = Keys[InterpNode - 1].Value;
-				const float P3 = Keys[InterpNode].Value;
-
-				if (Keys[InterpNode - 1].InterpMode == RCIM_Linear)
-				{
-					InterpVal = FMath::Lerp(P0, P3, Alpha);
-				}
-				else
-				{
-					const float OneThird = 1.0f / 3.0f;
-					const float P1 = P0 + (Keys[InterpNode - 1].LeaveTangent * Diff*OneThird);
-					const float P2 = P3 - (Keys[InterpNode].ArriveTangent * Diff*OneThird);
-
-					InterpVal = BezierInterp(P0, P1, P2, P3, Alpha);
-				}
+				InterpVal = FMath::Lerp(P0, P3, Alpha);
 			}
 			else
 			{
-				InterpVal = Keys[InterpNode - 1].Value;
+				const float OneThird = 1.0f / 3.0f;
+				const float P1 = P0 + (Keys[InterpNode - 1].LeaveTangent * Diff*OneThird);
+				const float P2 = P3 - (Keys[InterpNode].ArriveTangent * Diff*OneThird);
+
+				InterpVal = BezierInterp(P0, P1, P2, P3, Alpha);
 			}
 		}
 		else
 		{
-			// If beyond the last point in the curve, return its value.
-			InterpVal = Keys[NumKeys - 1].Value;
+			InterpVal = Keys[InterpNode - 1].Value;
 		}
 	}
+	else
+	{
+		// If beyond the last point in the curve, return its value.
+		InterpVal = Keys[NumKeys - 1].Value;
+	}
 
 	return InterpVal;
 }
