Finished fix for 'file requires check out' spam when compiling a blueprint

diff --git a/Engine/Source/Editor/BlueprintGraph/Private/EdGraphSchema_K2.cpp b/Engine/Source/Editor/BlueprintGraph/Private/EdGraphSchema_K2.cpp
--- a/Engine/Source/Editor/BlueprintGraph/Private/EdGraphSchema_K2.cpp
+++ b/Engine/Source/Editor/BlueprintGraph/Private/EdGraphSchema_K2.cpp
@@ -1925,6 +1925,7 @@ bool UEdGraphSchema_K2::TryCreateConnection(UEdGraphPin* PinA, UEdGraphPin* PinB
 
 	if (bModified && !PinA->HasAnyFlags(RF_Transient))
 	{
+		check(!PinB->HasAnyFlags(RF_Transient));
 		FBlueprintEditorUtils::MarkBlueprintAsModified(Blueprint);
 	}
 
diff --git a/Engine/Source/Editor/KismetCompiler/Public/KismetCompiler.h b/Engine/Source/Editor/KismetCompiler/Public/KismetCompiler.h
--- a/Engine/Source/Editor/KismetCompiler/Public/KismetCompiler.h
+++ b/Engine/Source/Editor/KismetCompiler/Public/KismetCompiler.h
@@ -103,7 +103,7 @@ public:
 		NodeType* Result = ParentGraph->CreateBlankNode<NodeType>();
 		MessageLog.NotifyIntermediateObjectCreation(Result, SourceNode);
 		Result->CreateNewGuid();
-		Result->SetFlags(RF_Transient);
+
 		AutoAssignNodePosition(Result);
 
 		return Result;
diff --git a/Engine/Source/Runtime/Engine/Private/EdGraph/EdGraph.cpp b/Engine/Source/Runtime/Engine/Private/EdGraph/EdGraph.cpp
--- a/Engine/Source/Runtime/Engine/Private/EdGraph/EdGraph.cpp
+++ b/Engine/Source/Runtime/Engine/Private/EdGraph/EdGraph.cpp
@@ -124,6 +124,12 @@ void UEdGraph::RemoveOnGraphChangedHandler( FDelegateHandle Handle )
 UEdGraphNode* UEdGraph::CreateNode( TSubclassOf<UEdGraphNode> NewNodeClass, bool bSelectNewNode/* = true*/ )
 {
 	UEdGraphNode* NewNode = NewObject<UEdGraphNode>(this, NewNodeClass, NAME_None, RF_Transactional);
+
+	if (HasAnyFlags(RF_Transient))
+	{
+		NewNode->SetFlags(RF_Transient);
+	}
+
 	AddNode(NewNode, false, bSelectNewNode );
 	return NewNode;
 }
