Flush network profiler bunches in the demo net connection, so that we can get useful profiles while recording replays.

diff --git a/Engine/Source/Runtime/Engine/Private/DemoNetDriver.cpp b/Engine/Source/Runtime/Engine/Private/DemoNetDriver.cpp
--- a/Engine/Source/Runtime/Engine/Private/DemoNetDriver.cpp
+++ b/Engine/Source/Runtime/Engine/Private/DemoNetDriver.cpp
@@ -974,6 +974,8 @@ void UDemoNetConnection::LowLevelSend( void* Data, int32 Count )
 		*FileAr << Count;
 		FileAr->Serialize( Data, Count );
 		
+		NETWORK_PROFILER(GNetworkProfiler.FlushOutgoingBunches(this));
+
 #if DEMO_CHECKSUMS == 1
 		uint32 Checksum = FCrc::MemCrc32( Data, Count, 0 );
 		*FileAr << Checksum;
