Fix non-unity compile

diff --git a/Engine/Source/Runtime/Engine/Private/DemoNetDriver.cpp b/Engine/Source/Runtime/Engine/Private/DemoNetDriver.cpp
--- a/Engine/Source/Runtime/Engine/Private/DemoNetDriver.cpp
+++ b/Engine/Source/Runtime/Engine/Private/DemoNetDriver.cpp
@@ -17,6 +17,7 @@
 #include "Engine/GameInstance.h"
 #include "NetworkReplayStreaming.h"
 #include "Net/UnrealNetwork.h"
+#include "Net/NetworkProfiler.h"
 
 DEFINE_LOG_CATEGORY_STATIC( LogDemo, Log, All );
 
