Fully DLL exposed UCurveFloat so you can derive from it

UDN - Linker error when inheriting from UCurveFloat

diff --git a/Engine/Source/Runtime/Engine/Classes/Curves/CurveFloat.h b/Engine/Source/Runtime/Engine/Classes/Curves/CurveFloat.h
--- a/Engine/Source/Runtime/Engine/Classes/Curves/CurveFloat.h
+++ b/Engine/Source/Runtime/Engine/Classes/Curves/CurveFloat.h
@@ -23,8 +23,8 @@ struct ENGINE_API FRuntimeFloatCurve
 	const FRichCurve* GetRichCurveConst() const;
 };
 
-UCLASS(BlueprintType, MinimalAPI)
-class UCurveFloat : public UCurveBase
+UCLASS(BlueprintType)
+class ENGINE_API UCurveFloat : public UCurveBase
 {
 	GENERATED_UCLASS_BODY()
 
@@ -38,7 +38,7 @@ class UCurveFloat : public UCurveBase
 
 	/** Evaluate this float curve at the specified time */
 	UFUNCTION(BlueprintCallable, Category="Math|Curves")
-	ENGINE_API float GetFloatValue(float InTime) const;
+	float GetFloatValue(float InTime) const;
 
 	// Begin FCurveOwnerInterface
 	virtual TArray<FRichCurveEditInfoConst> GetCurves() const override;
@@ -46,6 +46,6 @@ class UCurveFloat : public UCurveBase
 	virtual bool IsValidCurve( FRichCurveEditInfo CurveInfo ) override;
 
 	/** Determine if Curve is the same */
-	ENGINE_API bool operator == (const UCurveFloat& Curve) const;
+	bool operator == (const UCurveFloat& Curve) const;
 };
 
