Move handling of Alt-Enter out of DebugExecBindings and in to UGameViewportClient
Added bool to UInputSettings to set whether Alt-Enter should be tied to toggle fullscreen or not

diff --git a/Engine/Config/BaseInput.ini b/Engine/Config/BaseInput.ini
--- a/Engine/Config/BaseInput.ini
+++ b/Engine/Config/BaseInput.ini
@@ -3,13 +3,12 @@ bEnableMouseSmoothing=true
 bEnableFOVScaling=true
 FOVScale=0.01111
 DoubleClickTime=0.2f
-
 DefaultTouchInterface=/Engine/MobileResources/HUD/DefaultVirtualJoysticks.DefaultVirtualJoysticks
 bAlwaysShowTouchInterface=false
 bShowConsoleOnFourFingerTap=true
-
-+ConsoleKeys=Tilde
+bAltEnterTogglesFullscreen=true
 bRequireCtrlToNavigateAutoComplete=False
++ConsoleKeys=Tilde
 
 ;-----------------------------------------------------------------------------------------
 ; Axis properties
@@ -38,7 +37,6 @@ bRequireCtrlToNavigateAutoComplete=False
 +DebugExecBindings=(Key=Period,Command="RECOMPILESHADERS CHANGED",Control=True,Shift=True)
 +DebugExecBindings=(Key=Comma,Command="PROFILEGPU",Control=True,Shift=True)
 +DebugExecBindings=(Key=Escape,Command="CloseEditorViewport")
-+DebugExecBindings=(Key=Enter,Command="TOGGLE_FULLSCREEN",Alt=True)
 +DebugExecBindings=(Key=Tab,Command="FocusNextPIEWindow",Control=True)
 +DebugExecBindings=(Key=Tab,Command="FocusLastPIEWindow",Control=True,Shift=True)
 +DebugExecBindings=(Key=Apostrophe,Command="EnableGDT")
diff --git a/Engine/Source/Runtime/Engine/Classes/Engine/GameViewportClient.h b/Engine/Source/Runtime/Engine/Classes/Engine/GameViewportClient.h
--- a/Engine/Source/Runtime/Engine/Classes/Engine/GameViewportClient.h
+++ b/Engine/Source/Runtime/Engine/Classes/Engine/GameViewportClient.h
@@ -174,7 +174,7 @@ public:
 	bool HandleShowMouseCursorCommand( const TCHAR* Cmd, FOutputDevice& Ar );
 #endif // WITH_EDITOR
 	bool HandlePreCacheCommand( const TCHAR* Cmd, FOutputDevice& Ar );
-	bool HandleToggleFullscreenCommand( const TCHAR* Cmd, FOutputDevice& Ar );
+	bool HandleToggleFullscreenCommand();
 	bool HandleSetResCommand( const TCHAR* Cmd, FOutputDevice& Ar );
 	bool HandleHighresScreenshotCommand( const TCHAR* Cmd, FOutputDevice& Ar );
 	bool HandleHighresScreenshotUICommand( const TCHAR* Cmd, FOutputDevice& Ar );
diff --git a/Engine/Source/Runtime/Engine/Classes/GameFramework/InputSettings.h b/Engine/Source/Runtime/Engine/Classes/GameFramework/InputSettings.h
--- a/Engine/Source/Runtime/Engine/Classes/GameFramework/InputSettings.h
+++ b/Engine/Source/Runtime/Engine/Classes/GameFramework/InputSettings.h
@@ -20,6 +20,9 @@ class ENGINE_API UInputSettings
 	UPROPERTY(config, EditAnywhere, EditFixedSize, Category="Bindings", meta=(ToolTip="List of Axis Properties"), AdvancedDisplay)
 	TArray<struct FInputAxisConfigEntry> AxisConfig;
 
+	UPROPERTY(config, EditAnywhere, Category="Bindings", AdvancedDisplay)
+	uint32 bAltEnterTogglesFullscreen:1;
+
 	// Allow mouse to be used for touch
 	UPROPERTY(config, EditAnywhere, Category="MouseProperties")
 	uint32 bUseMouseForTouch:1;
diff --git a/Engine/Source/Runtime/Engine/Private/GameViewportClient.cpp b/Engine/Source/Runtime/Engine/Private/GameViewportClient.cpp
--- a/Engine/Source/Runtime/Engine/Private/GameViewportClient.cpp
+++ b/Engine/Source/Runtime/Engine/Private/GameViewportClient.cpp
@@ -256,6 +256,12 @@ bool UGameViewportClient::InputKey(FViewport* InViewport, int32 ControllerId, FK
 		return false;
 	}
 
+	if (Key == EKeys::Enter && EventType == EInputEvent::IE_Pressed && FSlateApplication::Get().GetModifierKeys().IsAltDown() && GetDefault<UInputSettings>()->bAltEnterTogglesFullscreen)
+	{
+		HandleToggleFullscreenCommand();
+		return true;
+	}
+
 	if (InViewport->IsPlayInEditorViewport() && Key.IsGamepadKey())
 	{
 		GEngine->RemapGamepadControllerIdForPIE(this, ControllerId);
@@ -1951,7 +1957,7 @@ bool UGameViewportClient::Exec( UWorld* InWorld, const TCHAR* Cmd,FOutputDevice&
 	}
 	else if( FParse::Command(&Cmd,TEXT("TOGGLE_FULLSCREEN")) || FParse::Command(&Cmd,TEXT("FULLSCREEN")) )
 	{
-		return HandleToggleFullscreenCommand( Cmd, Ar );
+		return HandleToggleFullscreenCommand();
 	}	
 	else if( FParse::Command(&Cmd,TEXT("SETRES")) )
 	{
@@ -2580,7 +2586,7 @@ bool UGameViewportClient::SetDisplayConfiguration(const FIntPoint* Dimensions, E
 	return true;
 }
 
-bool UGameViewportClient::HandleToggleFullscreenCommand(const TCHAR* Cmd, FOutputDevice& Ar)
+bool UGameViewportClient::HandleToggleFullscreenCommand()
 {
 	auto CVar = IConsoleManager::Get().FindTConsoleVariableDataInt(TEXT("r.FullScreenMode"));
 	check(CVar);
