Updated the analytics automation test to check the OS Id and to temporarily set a new account ID.

diff --git a/Engine/Source/Runtime/Engine/Private/Tests/AnalyticsTest.cpp b/Engine/Source/Runtime/Engine/Private/Tests/AnalyticsTest.cpp
--- a/Engine/Source/Runtime/Engine/Private/Tests/AnalyticsTest.cpp
+++ b/Engine/Source/Runtime/Engine/Private/Tests/AnalyticsTest.cpp
@@ -16,45 +16,60 @@ bool FAnalyticStartUpSimTest::RunTest(const FString& Parameters)
 {
 	if (FEngineAnalytics::IsAvailable())
 	{
+		//Setup a temp AccountID.
+		FGuid TempAccountID(FGuid::NewGuid());
+		FString OldEpicAccountID = FPlatformMisc::GetEpicAccountId();
+		FString NewEpicAccountID = TempAccountID.ToString().ToLower();
+		FPlatformMisc::SetEpicAccountId(NewEpicAccountID);
+
 		//Setup the 'Event Attributes'
 		TArray<FAnalyticsEventAttribute> EventAttributes;
-		EventAttributes.Add(FAnalyticsEventAttribute(TEXT("MachineID"), FPlatformMisc::GetMachineId().ToString(EGuidFormats::Digits).ToLower()));
-		EventAttributes.Add(FAnalyticsEventAttribute(TEXT("AccountID"), FPlatformMisc::GetEpicAccountId()));
-		EventAttributes.Add(FAnalyticsEventAttribute(TEXT("GameName"), FApp::GetGameName()));
-		EventAttributes.Add(FAnalyticsEventAttribute(TEXT("CommandLine"), FCommandLine::Get()));
-
-		//@todo - Set the user ID and check that
-		//@todo - Record OS ID as well
+		EventAttributes.Add(FAnalyticsEventAttribute(TEXT("MachineID"),		FPlatformMisc::GetMachineId().ToString(EGuidFormats::Digits).ToLower()));
+		EventAttributes.Add(FAnalyticsEventAttribute(TEXT("AccountID"),		FPlatformMisc::GetEpicAccountId()));
+		EventAttributes.Add(FAnalyticsEventAttribute(TEXT("OSID"),			FPlatformMisc::GetOperatingSystemId()));
+		EventAttributes.Add(FAnalyticsEventAttribute(TEXT("GameName"),		FApp::GetGameName()));
+		EventAttributes.Add(FAnalyticsEventAttribute(TEXT("CommandLine"),	FCommandLine::Get()));
 		
 		//Record the event with the 'Engine.AutomationTest.Analytics.ProgramStartedEvent' title
 		FEngineAnalytics::GetProvider().RecordEvent(TEXT("Engine.AutomationTest.Analytics.ProgramStartedEvent"), EventAttributes);
 
 		//Get the event strings used
-		FString MachineIDTest		=	FPlatformMisc::GetMachineId().ToString(EGuidFormats::Digits).ToLower();
-		FString AccountIDTest		=	FPlatformMisc::GetEpicAccountId();
-		FString GameNameTest		=	FApp::GetGameName();
-		FString CommandLineArgs		=	FCommandLine::Get();
+		FString MachineIDTest	=	FPlatformMisc::GetMachineId().ToString(EGuidFormats::Digits).ToLower();
+		FString AccountIDTest	=	FPlatformMisc::GetEpicAccountId();
+		FString OSID			=	FPlatformMisc::GetOperatingSystemId();
+		FString GameNameTest	=	FApp::GetGameName();
+		FString CommandLineArgs	=	FCommandLine::Get();
 
 		//Test the strings to verify they have data.
 		TestFalse(TEXT("'MachineID' is not expected to be empty!"), MachineIDTest.IsEmpty());
 		TestFalse(TEXT("'AccountID' is not expected to be empty!"), AccountIDTest.IsEmpty());
+		TestFalse(TEXT("'OperatingSystemID' is not expected to be empty!"), OSID.IsEmpty());
 		TestFalse(TEXT("'GameName' is expected."), GameNameTest.IsEmpty());
 
+
 		//Verify record event is holding the actual data.  This only triggers if the command line argument of 'AnalyticsDisableCaching' was used.
 		if (CommandLineArgs.Contains(TEXT("AnalyticsDisableCaching")))
 		{
+			FString FullMachineIDTestEventName = FString::Printf(TEXT("MachineID\":\"%s"), *MachineIDTest);
+			FString FullAccountIDTestEventName = FString::Printf(TEXT("AccountID\":\"%s"), *AccountIDTest);
+			FString FullOSIDTestEventName = FString::Printf(TEXT("OSID\":\"%s"), *OSID);
+
 			for (int32 i = 0; i < ExecutionInfo.LogItems.Num(); i++)
 			{
 				if (ExecutionInfo.LogItems[i].Contains(TEXT("Engine.AutomationTest.Analytics.ProgramStartedEvent")))
 				{
 					TestTrue(TEXT("Recorded event name is expected to be in the sent event."), ExecutionInfo.LogItems[i].Contains(TEXT("Engine.AutomationTest.Analytics.ProgramStartedEvent")));
-					TestTrue(TEXT("'MachineID' is expected to be in the sent event."), ExecutionInfo.LogItems[i].Contains(*MachineIDTest));
-					TestTrue(TEXT("'AccountID' is expected to be in the sent event."), ExecutionInfo.LogItems[i].Contains(*AccountIDTest));
+					TestTrue(TEXT("'MachineID' is expected to be in the sent event."), ExecutionInfo.LogItems[i].Contains(*FullMachineIDTestEventName));
+					TestTrue(TEXT("'AccountID' is expected to be in the sent event."), ExecutionInfo.LogItems[i].Contains(*FullAccountIDTestEventName));
+					TestTrue(TEXT("'OperatingSystemID' is expected to be in the sent event."), ExecutionInfo.LogItems[i].Contains(*FullOSIDTestEventName));
 					TestTrue(TEXT("'GameName' is expected to be in the sent event."), ExecutionInfo.LogItems[i].Contains(*GameNameTest));
 					TestTrue(TEXT("'CommandLine arguments' are expected to be in the sent event."), ExecutionInfo.LogItems[i].Contains(TEXT("AnalyticsDisableCaching")));
 				}
 			}
 		}
+
+		//Set the AccountID to it's original state.
+		FPlatformMisc::SetEpicAccountId(OldEpicAccountID);
 		return true;
 	}
 
@@ -120,6 +135,6 @@ bool FAnalyticsEventAttributeUnitTest::RunTest(const FString& Parameters)
 		return true;
 	}
 
-	ExecutionInfo.Warnings.Add(TEXT("Skipped 'FAnalyticsEventAttributeUnitTest' test.  EngineAnalytics are not currently available."));
+	ExecutionInfo.LogItems.Add(TEXT("SKIPPED 'FAnalyticsEventAttributeUnitTest' test.  EngineAnalytics are not currently available."));
 	return true;
 }
\ No newline at end of file
