Don't send netspeed until after we're fully connected. This fixes spam when our connection attempt gets rejected

diff --git a/Engine/Source/Runtime/Engine/Private/PendingNetGame.cpp b/Engine/Source/Runtime/Engine/Private/PendingNetGame.cpp
--- a/Engine/Source/Runtime/Engine/Private/PendingNetGame.cpp
+++ b/Engine/Source/Runtime/Engine/Private/PendingNetGame.cpp
@@ -190,7 +190,6 @@ void UPendingNetGame::NotifyControlMessage(UNetConnection* Connection, uint8 Mes
 			Connection->ClientResponse = TEXT("0");
 			FString URLString(PartialURL.ToString());
 			FNetControlMessage<NMT_Login>::Send(Connection, Connection->ClientResponse, URLString, UniqueIdRepl);
-			FNetControlMessage<NMT_Netspeed>::Send(Connection, Connection->CurrentNetSpeed);
 			NetDriver->ServerConnection->FlushNet();
 			break;
 		}
@@ -220,6 +219,9 @@ void UPendingNetGame::NotifyControlMessage(UNetConnection* Connection, uint8 Mes
 				URL.AddOption(*FString::Printf(TEXT("game=%s"), *GameName));
 			}
 
+			// Send out netspeed now that we're connected
+			FNetControlMessage<NMT_Netspeed>::Send(Connection, Connection->CurrentNetSpeed);
+
 			// We have successfully connected.
 			bSuccessfullyConnected = true;
 			break;
