PR #780: Add to Blueprints Truncate and add as default to convert from Float to Int (Contributed by [REDACTED_PERSON])

https://github.com/EpicGames/UnrealEngine/pull/780

Existing floor nodes will still be present in graphs

diff --git a/Engine/Source/Editor/BlueprintGraph/Private/EdGraphSchema_K2.cpp b/Engine/Source/Editor/BlueprintGraph/Private/EdGraphSchema_K2.cpp
--- a/Engine/Source/Editor/BlueprintGraph/Private/EdGraphSchema_K2.cpp
+++ b/Engine/Source/Editor/BlueprintGraph/Private/EdGraphSchema_K2.cpp
@@ -1973,7 +1973,7 @@ bool UEdGraphSchema_K2::SearchForAutocastFunction(const UEdGraphPin* OutputPin,
 	{
 		if (InputPin->PinType.PinCategory == PC_Int)
 		{
-			TargetFunction = TEXT("FFloor");
+			TargetFunction = TEXT("FTrunc");
 		}
 		else if ((InputPin->PinType.PinCategory == PC_Struct) && (InputStructType == VectorStruct))
 		{
diff --git a/Engine/Source/Runtime/Engine/Classes/Kismet/KismetMathLibrary.h b/Engine/Source/Runtime/Engine/Classes/Kismet/KismetMathLibrary.h
--- a/Engine/Source/Runtime/Engine/Classes/Kismet/KismetMathLibrary.h
+++ b/Engine/Source/Runtime/Engine/Classes/Kismet/KismetMathLibrary.h
@@ -453,7 +453,11 @@ class UKismetMathLibrary : public UBlueprintFunctionLibrary
 	/* Rounds A to the largest previous integer */
 	UFUNCTION(BlueprintPure, meta=(FriendlyName = "Floor"), Category="Math|Float")
 	static int32 FFloor(float A);
-
+	
+	/* Rounds A to an integer with truncation towards zero.  (e.g. -1.7 truncated to -1, 2.8 truncated to 2) */
+	UFUNCTION(BlueprintPure, meta=(FriendlyName = "Truncate"), Category="Math|Float")
+	static int32 FTrunc(float A);
+	
 	/* Rounds A to the smallest following integer */
 	UFUNCTION(BlueprintPure, Category="Math|Float")
 	static int32 FCeil(float A);
diff --git a/Engine/Source/Runtime/Engine/Private/KismetMathLibrary.cpp b/Engine/Source/Runtime/Engine/Private/KismetMathLibrary.cpp
--- a/Engine/Source/Runtime/Engine/Private/KismetMathLibrary.cpp
+++ b/Engine/Source/Runtime/Engine/Private/KismetMathLibrary.cpp
@@ -483,6 +483,11 @@ int32 UKismetMathLibrary::FFloor(float A)
 	return FMath::FloorToInt(A);
 }	
 
+int32 UKismetMathLibrary::FTrunc(float A)
+{
+	return FMath::TruncToInt(A);
+}	
+
 int32 UKismetMathLibrary::FCeil(float A)
 {
 	return FMath::CeilToInt(A);
