Integrate - Fixed shadow depth shaders not being cached with PixelDepthOffset

diff --git a/Engine/Source/Runtime/Engine/Private/Materials/MaterialShared.cpp b/Engine/Source/Runtime/Engine/Private/Materials/MaterialShared.cpp
--- a/Engine/Source/Runtime/Engine/Private/Materials/MaterialShared.cpp
+++ b/Engine/Source/Runtime/Engine/Private/Materials/MaterialShared.cpp
@@ -519,7 +519,7 @@ bool FMaterial::MaterialMayModifyMeshPosition() const
 {
 	// Conservative estimate when called before material translation has occurred. 
 	// This function is only intended for use in deciding whether or not shader permutations are required.
-	return HasVertexPositionOffsetConnected() || HasMaterialAttributesConnected() || GetTessellationMode() != MTM_NoTessellation;
+	return HasVertexPositionOffsetConnected() || HasPixelDepthOffsetConnected() || HasMaterialAttributesConnected() || GetTessellationMode() != MTM_NoTessellation;
 }
 
 FMaterialShaderMap* FMaterial::GetRenderingThreadShaderMap() const 
@@ -703,6 +703,7 @@ bool FMaterialResource::IsUsedWithEditorCompositing() const { return Material->b
 bool FMaterialResource::IsUsedWithDeferredDecal() const { return Material->MaterialDomain == MD_DeferredDecal; }
 bool FMaterialResource::IsSpecialEngineMaterial() const { return Material->bUsedAsSpecialEngineMaterial; }
 bool FMaterialResource::HasVertexPositionOffsetConnected() const { return !Material->bUseMaterialAttributes && Material->WorldPositionOffset.Expression != NULL; }
+bool FMaterialResource::HasPixelDepthOffsetConnected() const { return !Material->bUseMaterialAttributes && Material->PixelDepthOffset.Expression != NULL; }
 bool FMaterialResource::HasMaterialAttributesConnected() const { return Material->bUseMaterialAttributes && Material->MaterialAttributes.Expression != NULL; }
 FString FMaterialResource::GetBaseMaterialPathName() const { return Material->GetPathName(); }
 
diff --git a/Engine/Source/Runtime/Engine/Public/MaterialShared.h b/Engine/Source/Runtime/Engine/Public/MaterialShared.h
--- a/Engine/Source/Runtime/Engine/Public/MaterialShared.h
+++ b/Engine/Source/Runtime/Engine/Public/MaterialShared.h
@@ -944,6 +944,7 @@ public:
 	virtual bool UseTranslucencyVertexFog() const { return false; }
 	virtual FString GetFriendlyName() const = 0;
 	virtual bool HasVertexPositionOffsetConnected() const { return false; }
+	virtual bool HasPixelDepthOffsetConnected() const { return false; }
 	virtual bool HasMaterialAttributesConnected() const { return false; }
 	virtual uint32 GetDecalBlendMode() const { return 0; }
 	virtual uint32 GetMaterialDecalResponse() const { return 0; }
@@ -1564,6 +1565,7 @@ protected:
 	/* Gives the material a chance to compile any custom output nodes it has added */
 	ENGINE_API virtual void GatherCustomOutputExpressions(TArray<class UMaterialExpressionCustomOutput*>& OutCustomOutputs) const override;
 	ENGINE_API virtual bool HasVertexPositionOffsetConnected() const;
+	ENGINE_API virtual bool HasPixelDepthOffsetConnected() const;
 	ENGINE_API virtual bool HasMaterialAttributesConnected() const;
 	/** Useful for debugging. */
 	ENGINE_API virtual FString GetBaseMaterialPathName() const;
