[UEBP-32] Remove temp mobility swap in favor of a "proper" method for checking whether or not dynamic data changes are allowed during UCS execution.

change summary:
- renamed UActorComponent::IsRunningUserConstructionScript() to IsOwnerRunningUserConstructionScript()
- added USceneComponent::AreDynamicDataChangesAllowed() to both replace existing checks for whether or not the component is both registered and dynamic and to check for whether or not the UCS is running
- removed FUCSComponentManager and usage from ActorConstruction.cpp; retained post-UCS mobility fixup code

diff --git a/Engine/Plugins/2D/Paper2D/Source/Paper2D/Private/PaperFlipbookComponent.cpp b/Engine/Plugins/2D/Paper2D/Source/Paper2D/Private/PaperFlipbookComponent.cpp
--- a/Engine/Plugins/2D/Paper2D/Source/Paper2D/Private/PaperFlipbookComponent.cpp
+++ b/Engine/Plugins/2D/Paper2D/Source/Paper2D/Private/PaperFlipbookComponent.cpp
@@ -399,7 +399,7 @@ UMaterialInterface* UPaperFlipbookComponent::GetSpriteMaterial() const
 void UPaperFlipbookComponent::SetSpriteColor(FLinearColor NewColor)
 {
 	// Can't set color on a static component
-	if (!(IsRegistered() && (Mobility == EComponentMobility::Static)) && (SpriteColor != NewColor))
+	if (AreDynamicDataChangesAllowed() && (SpriteColor != NewColor))
 	{
 		SpriteColor = NewColor;
 
diff --git a/Engine/Plugins/2D/Paper2D/Source/Paper2D/Private/PaperSpriteComponent.cpp b/Engine/Plugins/2D/Paper2D/Source/Paper2D/Private/PaperSpriteComponent.cpp
--- a/Engine/Plugins/2D/Paper2D/Source/Paper2D/Private/PaperSpriteComponent.cpp
+++ b/Engine/Plugins/2D/Paper2D/Source/Paper2D/Private/PaperSpriteComponent.cpp
@@ -283,7 +283,7 @@ UPaperSprite* UPaperSpriteComponent::GetSprite()
 void UPaperSpriteComponent::SetSpriteColor(FLinearColor NewColor)
 {
 	// Can't set color on a static component
-	if (!(IsRegistered() && (Mobility == EComponentMobility::Static)) && (SpriteColor != NewColor))
+	if (AreDynamicDataChangesAllowed() && (SpriteColor != NewColor))
 	{
 		SpriteColor = NewColor;
 
diff --git a/Engine/Plugins/2D/Paper2D/Source/Paper2D/Private/Terrain/PaperTerrainComponent.cpp b/Engine/Plugins/2D/Paper2D/Source/Paper2D/Private/Terrain/PaperTerrainComponent.cpp
--- a/Engine/Plugins/2D/Paper2D/Source/Paper2D/Private/Terrain/PaperTerrainComponent.cpp
+++ b/Engine/Plugins/2D/Paper2D/Source/Paper2D/Private/Terrain/PaperTerrainComponent.cpp
@@ -938,7 +938,7 @@ void UPaperTerrainComponent::InsertConvexCollisionDataFromPolygon(const TArray<F
 void UPaperTerrainComponent::SetTerrainColor(FLinearColor NewColor)
 {
 	// Can't set color on a static component
-	if (!(IsRegistered() && (Mobility == EComponentMobility::Static)) && (TerrainColor != NewColor))
+	if (AreDynamicDataChangesAllowed() && (TerrainColor != NewColor))
 	{
 		TerrainColor = NewColor;
 
diff --git a/Engine/Source/Runtime/Engine/Classes/Components/ActorComponent.h b/Engine/Source/Runtime/Engine/Classes/Components/ActorComponent.h
--- a/Engine/Source/Runtime/Engine/Classes/Components/ActorComponent.h
+++ b/Engine/Source/Runtime/Engine/Classes/Components/ActorComponent.h
@@ -510,7 +510,7 @@ public:
 	// End IInterface_AssetUserData Interface
 
 	/** See if the owning Actor is currently running the UCS */
-	bool IsRunningUserConstructionScript() const;
+	bool IsOwnerRunningUserConstructionScript() const;
 
 	/** See if this component is currently registered */
 	FORCEINLINE bool IsRegistered() const
diff --git a/Engine/Source/Runtime/Engine/Classes/Components/SceneComponent.h b/Engine/Source/Runtime/Engine/Classes/Components/SceneComponent.h
--- a/Engine/Source/Runtime/Engine/Classes/Components/SceneComponent.h
+++ b/Engine/Source/Runtime/Engine/Classes/Components/SceneComponent.h
@@ -806,6 +806,17 @@ protected:
 
 	/** Called after changing transform, tries to update navigation octree */
 	virtual void UpdateNavigationData();
+
+	/**
+	 * Determine if dynamic data is allowed to be changed.
+	 * 
+	 * @param bIgnoreStationary Whether or not to ignore stationary mobility when checking. Default is true (i.e. - check for static mobility only).
+	 * @return Whether or not dynamic data is allowed to be changed.
+	 */
+	FORCEINLINE bool AreDynamicDataChangesAllowed(bool bIgnoreStationary = true) const
+	{
+		return (IsOwnerRunningUserConstructionScript()) || !(IsRegistered() && (Mobility == EComponentMobility::Static || (!bIgnoreStationary && Mobility == EComponentMobility::Stationary)));
+	}
 };
 
 /** 
diff --git a/Engine/Source/Runtime/Engine/Private/ActorComponent.cpp b/Engine/Source/Runtime/Engine/Private/ActorComponent.cpp
--- a/Engine/Source/Runtime/Engine/Private/ActorComponent.cpp
+++ b/Engine/Source/Runtime/Engine/Private/ActorComponent.cpp
@@ -1252,7 +1252,7 @@ void UActorComponent::SetTickableWhenPaused(bool bTickableWhenPaused)
 	PrimaryComponentTick.bTickEvenWhenPaused = bTickableWhenPaused;
 }
 
-bool UActorComponent::IsRunningUserConstructionScript() const
+bool UActorComponent::IsOwnerRunningUserConstructionScript() const
 {
 	const AActor* Owner = GetOwner();
 	return Owner != nullptr && Owner->bRunningUserConstructionScript;
diff --git a/Engine/Source/Runtime/Engine/Private/ActorConstruction.cpp b/Engine/Source/Runtime/Engine/Private/ActorConstruction.cpp
--- a/Engine/Source/Runtime/Engine/Private/ActorConstruction.cpp
+++ b/Engine/Source/Runtime/Engine/Private/ActorConstruction.cpp
@@ -15,103 +15,6 @@
 
 DEFINE_LOG_CATEGORY(LogBlueprintUserMessages);
 
-//////////////////////////////////////////////////////////////////////////
-// Local Types
-
-namespace
-{
-	/** Tracks info for components instanced during UCS execution */
-	struct FUCSComponentInfo
-	{
-		EComponentMobility::Type Mobility;
-		TWeakObjectPtr<UActorComponent> ComponentPtr;
-
-		FUCSComponentInfo(UActorComponent* InComponent)
-			: ComponentPtr(InComponent)
-		{
-			ensure(!InComponent || !InComponent->IsPendingKill());
-			USceneComponent* SceneComponent = Cast<USceneComponent>(InComponent);
-			if(SceneComponent != nullptr)
-			{
-				// Save original mobility
-				Mobility = SceneComponent->Mobility;
-
-				// Temporarily set to 'movable' to allow changes during UCS
-				SceneComponent->Mobility = EComponentMobility::Movable;
-			}
-		}
-	};
-
-	/** Helper class to manage components instanced during UCS execution */
-	class FUCSComponentManager
-	{
-		/** Map of actors and any components created during UCS */
-		TMap<const AActor*, TArray<FUCSComponentInfo> > UCSComponentsMap;
-
-	public:
-		/** Called before UCS execution has started for the given Actor */
-		void PreProcessComponents(const AActor* InActor)
-		{
-			TInlineComponentArray<UActorComponent*> ActorComponents;
-			InActor->GetComponents(ActorComponents);
-			for (auto CompIt = ActorComponents.CreateConstIterator(); CompIt; ++CompIt)
-			{
-				AddComponent(InActor, *CompIt);
-			}
-		}
-
-		/** Add a component instance for the given Actor */
-		void AddComponent(const AActor* InActor, UActorComponent* InComponent)
-		{
-			TArray<FUCSComponentInfo>& UCSComponentsList = UCSComponentsMap.FindOrAdd(InActor);
-			UCSComponentsList.Add(FUCSComponentInfo(InComponent));
-		}
-
-		/** Called after UCS execution has finished for the given Actor */
-		void PostProcessComponents(const AActor* InActor)
-		{
-			TArray<FUCSComponentInfo> UCSComponentsList;
-			const bool bFound = UCSComponentsMap.RemoveAndCopyValue(InActor, UCSComponentsList);
-			if (bFound)
-			{
-				for (int32 ComponentIndex = 0; ComponentIndex < UCSComponentsList.Num(); ++ComponentIndex)
-				{
-					const FUCSComponentInfo& UCSComponentInfo = UCSComponentsList[ComponentIndex];
-
-					USceneComponent* SceneComponent = Cast<USceneComponent>(UCSComponentInfo.ComponentPtr.Get());
-					if(SceneComponent != nullptr)
-					{
-						// Restore original mobility after UCS execution
-						SceneComponent->Mobility = UCSComponentInfo.Mobility;
-
-						// A parent component can't be more mobile than its children, so we check for that here and adjust as needed.
-						if(SceneComponent->AttachParent != nullptr && SceneComponent->AttachParent->Mobility > SceneComponent->Mobility)
-						{
-							if(SceneComponent->IsA<UStaticMeshComponent>())
-							{
-								// SMCs can't be stationary, so always set them (and any children) to be movable
-								SceneComponent->SetMobility(EComponentMobility::Movable);
-							}
-							else
-							{
-								// Set the new component (and any children) to be at least as mobile as its parent
-								SceneComponent->SetMobility(SceneComponent->AttachParent->Mobility);
-							}
-						}
-					}
-				}
-			}
-		}
-
-		/** Gets the FUCSComponentManager singleton. */
-		static FUCSComponentManager& Get()
-		{
-			static FUCSComponentManager Singleton;
-			return Singleton;
-		}
-	};
-}
-
 //////////////////////////////////////////////////////////////////////////
 // AActor Blueprint Stuff
 
@@ -574,17 +477,31 @@ void AActor::ExecuteConstruction(const FTransform& Transform, const FComponentIn
 
 void AActor::ProcessUserConstructionScript()
 {
-	// Process components that may have already been instanced before UCS execution.
-	FUCSComponentManager& UCSComponentManager = FUCSComponentManager::Get();
-	UCSComponentManager.PreProcessComponents(this);
-
 	// Set a flag that this actor is currently running UserConstructionScript.
 	bRunningUserConstructionScript = true;
 	UserConstructionScript();
 	bRunningUserConstructionScript = false;
 
-	// Perform any post processing on this Actor's components after UCS execution.
-	UCSComponentManager.PostProcessComponents(this);
+	// Validate component mobility after UCS execution
+	TInlineComponentArray<USceneComponent*> SceneComponents;
+	GetComponents(SceneComponents);
+	for (auto SceneComponent : SceneComponents)
+	{
+		// A parent component can't be more mobile than its children, so we check for that here and adjust as needed.
+		if(SceneComponent != RootComponent && SceneComponent->AttachParent != nullptr && SceneComponent->AttachParent->Mobility > SceneComponent->Mobility)
+		{
+			if(SceneComponent->IsA<UStaticMeshComponent>())
+			{
+				// SMCs can't be stationary, so always set them (and any children) to be movable
+				SceneComponent->SetMobility(EComponentMobility::Movable);
+			}
+			else
+			{
+				// Set the new component (and any children) to be at least as mobile as its parent
+				SceneComponent->SetMobility(SceneComponent->AttachParent->Mobility);
+			}
+		}
+	}
 }
 
 void AActor::FinishAndRegisterComponent(UActorComponent* Component)
@@ -659,7 +576,7 @@ UActorComponent* AActor::AddComponent(FName TemplateName, bool bManualAttachment
 	{
 		// Call function to notify component it has been created
 		NewActorComp->OnComponentCreated();
-
+		
 		// The user has the option of doing attachment manually where they have complete control or via the automatic rule
 		// that the first component added becomes the root component, with subsequent components attached to the root.
 		USceneComponent* NewSceneComp = Cast<USceneComponent>(NewActorComp);
@@ -682,13 +599,6 @@ UActorComponent* AActor::AddComponent(FName TemplateName, bool bManualAttachment
 
 		// Register component, which will create physics/rendering state, now component is in correct position
 		NewActorComp->RegisterComponent();
-
-		// Keep track of the new component during UCS execution. This also does a temporary mobility swap during UCS execution in order to allow dynamic data to be changed within that context.
-		// Note: This should only be done AFTER registration.
-		if(bRunningUserConstructionScript)
-		{
-			FUCSComponentManager::Get().AddComponent(this, NewActorComp);
-		}
 	}
 
 	return NewActorComp;
diff --git a/Engine/Source/Runtime/Engine/Private/DirectionalLightComponent.cpp b/Engine/Source/Runtime/Engine/Private/DirectionalLightComponent.cpp
--- a/Engine/Source/Runtime/Engine/Private/DirectionalLightComponent.cpp
+++ b/Engine/Source/Runtime/Engine/Private/DirectionalLightComponent.cpp
@@ -837,7 +837,7 @@ void UDirectionalLightComponent::SetShadowDistanceFadeoutFraction(float NewValue
 
 void UDirectionalLightComponent::SetEnableLightShaftOcclusion(bool bNewValue)
 {
-	if ((IsRunningUserConstructionScript() || !(IsRegistered() && Mobility == EComponentMobility::Static))
+	if (AreDynamicDataChangesAllowed()
 		&& bEnableLightShaftOcclusion != bNewValue)
 	{
 		bEnableLightShaftOcclusion = bNewValue;
@@ -847,7 +847,7 @@ void UDirectionalLightComponent::SetEnableLightShaftOcclusion(bool bNewValue)
 
 void UDirectionalLightComponent::SetOcclusionMaskDarkness(float NewValue)
 {
-	if ((IsRunningUserConstructionScript() || !(IsRegistered() && Mobility == EComponentMobility::Static))
+	if (AreDynamicDataChangesAllowed()
 		&& OcclusionMaskDarkness != NewValue)
 	{
 		OcclusionMaskDarkness = NewValue;
@@ -857,7 +857,7 @@ void UDirectionalLightComponent::SetOcclusionMaskDarkness(float NewValue)
 
 void UDirectionalLightComponent::SetLightShaftOverrideDirection(FVector NewValue)
 {
-	if ((IsRunningUserConstructionScript() || !(IsRegistered() && Mobility == EComponentMobility::Static))
+	if (AreDynamicDataChangesAllowed()
 		&& LightShaftOverrideDirection != NewValue)
 	{
 		LightShaftOverrideDirection = NewValue;
diff --git a/Engine/Source/Runtime/Engine/Private/LightComponent.cpp b/Engine/Source/Runtime/Engine/Private/LightComponent.cpp
--- a/Engine/Source/Runtime/Engine/Private/LightComponent.cpp
+++ b/Engine/Source/Runtime/Engine/Private/LightComponent.cpp
@@ -70,7 +70,7 @@ FArchive& operator<<(FArchive& Ar, FStaticShadowDepthMap& ShadowMap)
 
 void ULightComponentBase::SetCastShadows(bool bNewValue)
 {
-	if ((IsRunningUserConstructionScript() || !(IsRegistered() && Mobility == EComponentMobility::Static))
+	if (AreDynamicDataChangesAllowed()
 		&& CastShadows != bNewValue)
 	{
 		CastShadows = bNewValue;
@@ -666,7 +666,7 @@ void ULightComponent::DestroyRenderState_Concurrent()
 void ULightComponent::SetIntensity(float NewIntensity)
 {
 	// Can't set brightness on a static light
-	if ((IsRunningUserConstructionScript() || !(IsRegistered() && Mobility == EComponentMobility::Static))
+	if (AreDynamicDataChangesAllowed()
 		&& Intensity != NewIntensity)
 	{
 		Intensity = NewIntensity;
@@ -686,7 +686,7 @@ void ULightComponent::SetLightColor(FLinearColor NewLightColor)
 	FColor NewColor(NewLightColor);
 
 	// Can't set color on a static light
-	if ((IsRunningUserConstructionScript() || !(IsRegistered() && Mobility == EComponentMobility::Static))
+	if (AreDynamicDataChangesAllowed()
 		&& LightColor != NewColor)
 	{
 		LightColor	= NewColor;
@@ -703,7 +703,7 @@ void ULightComponent::SetLightColor(FLinearColor NewLightColor)
 void ULightComponent::SetLightFunctionMaterial(UMaterialInterface* NewLightFunctionMaterial)
 {
 	// Can't set light function on a static light
-	if ((IsRunningUserConstructionScript() || !(IsRegistered() && Mobility == EComponentMobility::Static))
+	if (AreDynamicDataChangesAllowed()
 		&& NewLightFunctionMaterial != LightFunctionMaterial)
 	{
 		LightFunctionMaterial = NewLightFunctionMaterial;
@@ -713,7 +713,7 @@ void ULightComponent::SetLightFunctionMaterial(UMaterialInterface* NewLightFunct
 
 void ULightComponent::SetLightFunctionScale(FVector NewLightFunctionScale)
 {
-	if ((IsRunningUserConstructionScript() || !(IsRegistered() && Mobility == EComponentMobility::Static))
+	if (AreDynamicDataChangesAllowed()
 		&& NewLightFunctionScale != LightFunctionScale)
 	{
 		LightFunctionScale = NewLightFunctionScale;
@@ -723,7 +723,7 @@ void ULightComponent::SetLightFunctionScale(FVector NewLightFunctionScale)
 
 void ULightComponent::SetLightFunctionFadeDistance(float NewLightFunctionFadeDistance)
 {
-	if ((IsRunningUserConstructionScript() || !(IsRegistered() && Mobility == EComponentMobility::Static))
+	if (AreDynamicDataChangesAllowed()
 		&& NewLightFunctionFadeDistance != LightFunctionFadeDistance)
 	{
 		LightFunctionFadeDistance = NewLightFunctionFadeDistance;
@@ -733,7 +733,7 @@ void ULightComponent::SetLightFunctionFadeDistance(float NewLightFunctionFadeDis
 
 void ULightComponent::SetAffectDynamicIndirectLighting(bool bNewValue)
 {
-	if ((IsRunningUserConstructionScript() || !(IsRegistered() && Mobility == EComponentMobility::Static))
+	if (AreDynamicDataChangesAllowed()
 		&& bAffectDynamicIndirectLighting != bNewValue)
 	{
 		bAffectDynamicIndirectLighting = bNewValue;
@@ -743,7 +743,7 @@ void ULightComponent::SetAffectDynamicIndirectLighting(bool bNewValue)
 
 void ULightComponent::SetAffectTranslucentLighting(bool bNewValue)
 {
-	if ((IsRunningUserConstructionScript() || !(IsRegistered() && Mobility == EComponentMobility::Static))
+	if (AreDynamicDataChangesAllowed()
 		&& bAffectTranslucentLighting != bNewValue)
 	{
 		bAffectTranslucentLighting = bNewValue;
@@ -753,7 +753,7 @@ void ULightComponent::SetAffectTranslucentLighting(bool bNewValue)
 
 void ULightComponent::SetEnableLightShaftBloom(bool bNewValue)
 {
-	if ((IsRunningUserConstructionScript() || !(IsRegistered() && Mobility == EComponentMobility::Static))
+	if (AreDynamicDataChangesAllowed()
 		&& bEnableLightShaftBloom != bNewValue)
 	{
 		bEnableLightShaftBloom = bNewValue;
@@ -763,7 +763,7 @@ void ULightComponent::SetEnableLightShaftBloom(bool bNewValue)
 
 void ULightComponent::SetBloomScale(float NewValue)
 {
-	if ((IsRunningUserConstructionScript() || !(IsRegistered() && Mobility == EComponentMobility::Static))
+	if (AreDynamicDataChangesAllowed()
 		&& BloomScale != NewValue)
 	{
 		BloomScale = NewValue;
@@ -773,7 +773,7 @@ void ULightComponent::SetBloomScale(float NewValue)
 
 void ULightComponent::SetBloomThreshold(float NewValue)
 {
-	if ((IsRunningUserConstructionScript() || !(IsRegistered() && Mobility == EComponentMobility::Static))
+	if (AreDynamicDataChangesAllowed()
 		&& BloomThreshold != NewValue)
 	{
 		BloomThreshold = NewValue;
@@ -783,7 +783,7 @@ void ULightComponent::SetBloomThreshold(float NewValue)
 
 void ULightComponent::SetBloomTint(FColor NewValue)
 {
-	if ((IsRunningUserConstructionScript() || !(IsRegistered() && Mobility == EComponentMobility::Static))
+	if (AreDynamicDataChangesAllowed()
 		&& BloomTint != NewValue)
 	{
 		BloomTint = NewValue;
@@ -793,7 +793,7 @@ void ULightComponent::SetBloomTint(FColor NewValue)
 
 void ULightComponent::SetIESTexture(UTextureLightProfile* NewValue)
 {
-	if ((IsRunningUserConstructionScript() || !(IsRegistered() && Mobility == EComponentMobility::Static))
+	if (AreDynamicDataChangesAllowed()
 		&& IESTexture != NewValue)
 	{
 		IESTexture = NewValue;
diff --git a/Engine/Source/Runtime/Engine/Private/PointLightComponent.cpp b/Engine/Source/Runtime/Engine/Private/PointLightComponent.cpp
--- a/Engine/Source/Runtime/Engine/Private/PointLightComponent.cpp
+++ b/Engine/Source/Runtime/Engine/Private/PointLightComponent.cpp
@@ -127,7 +127,7 @@ FLightSceneProxy* UPointLightComponent::CreateSceneProxy() const
 void UPointLightComponent::SetAttenuationRadius(float NewRadius)
 {
 	// Only movable lights can change their radius at runtime
-	if ((IsRunningUserConstructionScript() || !(IsRegistered() && (Mobility == EComponentMobility::Static || Mobility == EComponentMobility::Stationary)))
+	if (AreDynamicDataChangesAllowed(false)
 		&& NewRadius != AttenuationRadius)
 	{
 		AttenuationRadius = NewRadius;
@@ -137,7 +137,7 @@ void UPointLightComponent::SetAttenuationRadius(float NewRadius)
 
 void UPointLightComponent::SetLightFalloffExponent(float NewLightFalloffExponent)
 {
-	if ((IsRunningUserConstructionScript() || !(IsRegistered() && Mobility == EComponentMobility::Static))
+	if (AreDynamicDataChangesAllowed()
 		&& NewLightFalloffExponent != LightFalloffExponent)
 	{
 		LightFalloffExponent = NewLightFalloffExponent;
@@ -147,7 +147,7 @@ void UPointLightComponent::SetLightFalloffExponent(float NewLightFalloffExponent
 
 void UPointLightComponent::SetSourceRadius(float bNewValue)
 {
-	if ((IsRunningUserConstructionScript() || !(IsRegistered() && Mobility == EComponentMobility::Static))
+	if (AreDynamicDataChangesAllowed()
 		&& SourceRadius != bNewValue)
 	{
 		SourceRadius = bNewValue;
diff --git a/Engine/Source/Runtime/Engine/Private/SceneComponent.cpp b/Engine/Source/Runtime/Engine/Private/SceneComponent.cpp
--- a/Engine/Source/Runtime/Engine/Private/SceneComponent.cpp
+++ b/Engine/Source/Runtime/Engine/Private/SceneComponent.cpp
@@ -890,7 +890,7 @@ void USceneComponent::AttachTo(class USceneComponent* Parent, FName InSocketName
 		}
 
 		// Don't allow components with static mobility to be attached to non-static parents (except during UCS)
-		if(!IsRunningUserConstructionScript() && Mobility == EComponentMobility::Static && Parent->Mobility != EComponentMobility::Static)
+		if(!IsOwnerRunningUserConstructionScript() && Mobility == EComponentMobility::Static && Parent->Mobility != EComponentMobility::Static)
 		{
 			FString ExtraBlueprintInfo;
 #if WITH_EDITORONLY_DATA
diff --git a/Engine/Source/Runtime/Engine/Private/SkyLightComponent.cpp b/Engine/Source/Runtime/Engine/Private/SkyLightComponent.cpp
--- a/Engine/Source/Runtime/Engine/Private/SkyLightComponent.cpp
+++ b/Engine/Source/Runtime/Engine/Private/SkyLightComponent.cpp
@@ -468,7 +468,7 @@ void USkyLightComponent::CaptureEmissiveIrradianceEnvironmentMap(FSHVectorRGB3&
 void USkyLightComponent::SetIntensity(float NewIntensity)
 {
 	// Can't set brightness on a static light
-	if ((IsRunningUserConstructionScript() || !(IsRegistered() && Mobility == EComponentMobility::Static))
+	if (AreDynamicDataChangesAllowed()
 		&& Intensity != NewIntensity)
 	{
 		Intensity = NewIntensity;
@@ -482,7 +482,7 @@ void USkyLightComponent::SetLightColor(FLinearColor NewLightColor)
 	FColor NewColor(NewLightColor);
 
 	// Can't set color on a static light
-	if ((IsRunningUserConstructionScript() || !(IsRegistered() && Mobility == EComponentMobility::Static))
+	if (AreDynamicDataChangesAllowed()
 		&& LightColor != NewColor)
 	{
 		LightColor = NewColor;
@@ -493,7 +493,7 @@ void USkyLightComponent::SetLightColor(FLinearColor NewLightColor)
 void USkyLightComponent::SetCubemap(UTextureCube* NewCubemap)
 {
 	// Can't set color on a static light
-	if ((IsRunningUserConstructionScript() || !(IsRegistered() && Mobility == EComponentMobility::Static))
+	if (AreDynamicDataChangesAllowed()
 		&& Cubemap != NewCubemap)
 	{
 		Cubemap = NewCubemap;
@@ -505,7 +505,7 @@ void USkyLightComponent::SetCubemap(UTextureCube* NewCubemap)
 void USkyLightComponent::SetOcclusionTint(const FColor& InTint)
 {
 	// Can't set on a static light
-	if ((IsRunningUserConstructionScript() || !(IsRegistered() && Mobility == EComponentMobility::Static))
+	if (AreDynamicDataChangesAllowed()
 		&& OcclusionTint != InTint)
 	{
 		OcclusionTint = InTint;
@@ -516,7 +516,7 @@ void USkyLightComponent::SetOcclusionTint(const FColor& InTint)
 void USkyLightComponent::SetMinOcclusion(float InMinOcclusion)
 {
 	// Can't set on a static light
-	if ((IsRunningUserConstructionScript() || !(IsRegistered() && Mobility == EComponentMobility::Static))
+	if (AreDynamicDataChangesAllowed()
 		&& MinOcclusion != InMinOcclusion)
 	{
 		MinOcclusion = InMinOcclusion;
diff --git a/Engine/Source/Runtime/Engine/Private/SpotLightComponent.cpp b/Engine/Source/Runtime/Engine/Private/SpotLightComponent.cpp
--- a/Engine/Source/Runtime/Engine/Private/SpotLightComponent.cpp
+++ b/Engine/Source/Runtime/Engine/Private/SpotLightComponent.cpp
@@ -172,7 +172,7 @@ USpotLightComponent::USpotLightComponent(const FObjectInitializer& ObjectInitial
 
 void USpotLightComponent::SetInnerConeAngle(float NewInnerConeAngle)
 {
-	if (!(IsRegistered() && (Mobility == EComponentMobility::Static || Mobility == EComponentMobility::Stationary))
+	if (AreDynamicDataChangesAllowed(false)
 		&& NewInnerConeAngle != InnerConeAngle)
 	{
 		InnerConeAngle = NewInnerConeAngle;
@@ -182,7 +182,7 @@ void USpotLightComponent::SetInnerConeAngle(float NewInnerConeAngle)
 
 void USpotLightComponent::SetOuterConeAngle(float NewOuterConeAngle)
 {
-	if (!(IsRegistered() && (Mobility == EComponentMobility::Static || Mobility == EComponentMobility::Stationary))
+	if (AreDynamicDataChangesAllowed(false)
 		&& NewOuterConeAngle != OuterConeAngle)
 	{
 		OuterConeAngle = NewOuterConeAngle;
diff --git a/Engine/Source/Runtime/Engine/Private/StaticMeshComponent.cpp b/Engine/Source/Runtime/Engine/Private/StaticMeshComponent.cpp
--- a/Engine/Source/Runtime/Engine/Private/StaticMeshComponent.cpp
+++ b/Engine/Source/Runtime/Engine/Private/StaticMeshComponent.cpp
@@ -1191,7 +1191,7 @@ bool UStaticMeshComponent::SetStaticMesh(UStaticMesh* NewMesh)
 
 	// Don't allow changing static meshes if "static" and registered
 	AActor* Owner = GetOwner();
-	if(Mobility == EComponentMobility::Static && IsRegistered() && Owner != NULL)
+	if(!AreDynamicDataChangesAllowed() && Owner != NULL)
 	{
 		FMessageLog("PIE").Warning(FText::Format(LOCTEXT("SetMeshOnStatic", "Calling SetStaticMesh on '{0}' but Mobility is Static."), 
 			FText::FromString(GetPathName(this))));
