Integrate 2454658 from UE4-Orion.

don't replicate the vision manager to network replays (we don't need it and it is ~80% of the data in the replays)

diff --git a/Engine/Source/Runtime/Engine/Classes/GameFramework/Actor.h b/Engine/Source/Runtime/Engine/Classes/GameFramework/Actor.h
--- a/Engine/Source/Runtime/Engine/Classes/GameFramework/Actor.h
+++ b/Engine/Source/Runtime/Engine/Classes/GameFramework/Actor.h
@@ -331,6 +331,10 @@ public:
 	/** If true, this actor should search for an owned camera component to view through when used as a view target. */
 	UPROPERTY(EditAnywhere, BlueprintReadWrite, Category=Actor, AdvancedDisplay)
 	uint32 bFindCameraComponentWhenViewTarget:1;
+	
+	/** If true, this actor will be replicated to network replys (default is true) */
+	UPROPERTY()
+	uint32 bRelevantForNetworkReplays:1;
 
 	/** Pawn responsible for damage caused by this actor. */
 	UPROPERTY(BlueprintReadWrite, replicatedUsing=OnRep_Instigator, meta=(ExposeOnSpawn=true), Category=Actor)
@@ -362,7 +366,7 @@ protected:
 
 	/** How long this Actor lives before dying, 0=forever. Note this is the INITIAL value and should not be modified once play has begun. */
 	UPROPERTY(EditAnywhere, BlueprintReadOnly, Category=Actor)
-	float InitialLifeSpan; 
+	float InitialLifeSpan;
 
 	/** Handle for efficient management of LifeSpanExpired timer */
 	FTimerHandle TimerHandle_LifeSpanExpired;
diff --git a/Engine/Source/Runtime/Engine/Private/Actor.cpp b/Engine/Source/Runtime/Engine/Private/Actor.cpp
--- a/Engine/Source/Runtime/Engine/Private/Actor.cpp
+++ b/Engine/Source/Runtime/Engine/Private/Actor.cpp
@@ -94,6 +94,7 @@ void AActor::InitializeDefaults()
 	bPendingKillPending = false;
 	bFindCameraComponentWhenViewTarget = true;
 	bAllowReceiveTickEventOnDedicatedServer = true;
+	bRelevantForNetworkReplays = true;
 }
 
 void FActorTickFunction::ExecuteTick(float DeltaTime, enum ELevelTick TickType, ENamedThreads::Type CurrentThread, const FGraphEventRef& MyCompletionGraphEvent)
diff --git a/Engine/Source/Runtime/Engine/Private/DemoNetDriver.cpp b/Engine/Source/Runtime/Engine/Private/DemoNetDriver.cpp
--- a/Engine/Source/Runtime/Engine/Private/DemoNetDriver.cpp
+++ b/Engine/Source/Runtime/Engine/Private/DemoNetDriver.cpp
@@ -572,28 +572,31 @@ static void DemoReplicateActor(AActor* Actor, UNetConnection* Connection, bool I
 	*/
 	if ( Actor != NULL && Actor->GetRemoteRole() != ROLE_None && ( Actor == Connection->PlayerController || Cast< APlayerController >( Actor ) == NULL ) )
 	{
-		// Create a new channel for this actor.
-		const bool StartupActor = Actor->IsNetStartupActor();
-		UActorChannel* Channel = Connection->ActorChannels.FindRef( Actor );
-
-		if ( !Channel && ( !StartupActor || Connection->ClientHasInitializedLevelFor( Actor ) ) )
+		if (Actor->bRelevantForNetworkReplays)
 		{
-			// create a channel if possible
-			Channel = (UActorChannel*)Connection->CreateChannel( CHTYPE_Actor, 1 );
-			if ( Channel != NULL )
+			// Create a new channel for this actor.
+			const bool StartupActor = Actor->IsNetStartupActor();
+			UActorChannel* Channel = Connection->ActorChannels.FindRef(Actor);
+
+			if (!Channel && (!StartupActor || Connection->ClientHasInitializedLevelFor(Actor)))
 			{
-				Channel->SetChannelActor( Actor );
+				// create a channel if possible
+				Channel = (UActorChannel*)Connection->CreateChannel(CHTYPE_Actor, 1);
+				if (Channel != NULL)
+				{
+					Channel->SetChannelActor(Actor);
+				}
 			}
-		}
 
-		if ( Channel )
-		{
-			// Send it out!
-			check( !Channel->Closing );
-
-			if ( Channel->IsNetReady( 0 ) )
+			if (Channel)
 			{
-				Channel->ReplicateActor();
+				// Send it out!
+				check(!Channel->Closing);
+
+				if (Channel->IsNetReady(0))
+				{
+					Channel->ReplicateActor();
+				}
 			}
 		}
 	}
