PR #837: Fixed a copy paste error in the documentation. (Contributed by [REDACTED_PERSON])

https://github.com/EpicGames/UnrealEngine/pull/837

diff --git a/Engine/Source/Runtime/Engine/Classes/Engine/StaticMesh.h b/Engine/Source/Runtime/Engine/Classes/Engine/StaticMesh.h
--- a/Engine/Source/Runtime/Engine/Classes/Engine/StaticMesh.h
+++ b/Engine/Source/Runtime/Engine/Classes/Engine/StaticMesh.h
@@ -488,7 +488,7 @@ public:
 	ENGINE_API int32 GetNumLODs() const;
 
 	/**
-	 * Returns the number of LODs used by the mesh.
+	 * Returns true if the mesh has data that can be rendered.
 	 */
 	ENGINE_API bool HasValidRenderData() const;
 
