#Merging PlaySlot node parameter by [REDACTED_PERSON]

//depot/UE4-GDC2015/Engine/Source/Runtime/Engine/...

to //depot/UE4-GDC2015/Engine/Source/Runtime/Engine/...

diff --git a/Engine/Source/Runtime/Engine/Classes/Animation/AnimInstance.h b/Engine/Source/Runtime/Engine/Classes/Animation/AnimInstance.h
--- a/Engine/Source/Runtime/Engine/Classes/Animation/AnimInstance.h
+++ b/Engine/Source/Runtime/Engine/Classes/Animation/AnimInstance.h
@@ -409,11 +409,11 @@ public:
 
 	/** Play normal animation asset on the slot node. You can only play one asset (whether montage or animsequence) at a time. */
 	UFUNCTION(BlueprintCallable, Category="Animation")
-	float PlaySlotAnimation(UAnimSequenceBase* Asset, FName SlotNodeName, float BlendInTime = 0.25f, float BlendOutTime = 0.25f, float InPlayRate = 1.f);
+	float PlaySlotAnimation(UAnimSequenceBase* Asset, FName SlotNodeName, float BlendInTime = 0.25f, float BlendOutTime = 0.25f, float InPlayRate = 1.f, int32 LoopCount = 1);
 
-	/** Stops currently playing slot animation */
+	/** Stops currently playing slot animation slot or all*/
 	UFUNCTION(BlueprintCallable, Category="Animation")
-	void StopSlotAnimation(float InBlendOutTime = 0.25f);
+	void StopSlotAnimation(float InBlendOutTime = 0.25f, FName SlotNodeName = NAME_None);
 
 	/** Return true if it's playing the slot animation */
 	UFUNCTION(BlueprintCallable, Category="Animation")
diff --git a/Engine/Source/Runtime/Engine/Private/Animation/AnimInstance.cpp b/Engine/Source/Runtime/Engine/Private/Animation/AnimInstance.cpp
--- a/Engine/Source/Runtime/Engine/Private/Animation/AnimInstance.cpp
+++ b/Engine/Source/Runtime/Engine/Private/Animation/AnimInstance.cpp
@@ -1801,7 +1801,7 @@ void UAnimInstance::Montage_Advance(float DeltaSeconds)
 	}
 }
 
-float UAnimInstance::PlaySlotAnimation(UAnimSequenceBase* Asset, FName SlotNodeName, float BlendInTime, float BlendOutTime, float InPlayRate)
+float UAnimInstance::PlaySlotAnimation(UAnimSequenceBase* Asset, FName SlotNodeName, float BlendInTime, float BlendOutTime, float InPlayRate, int32 LoopCount)
 {
 	// create temporary montage and play
 	bool bValidAsset = Asset && !Asset->IsA(UAnimMontage::StaticClass());
@@ -1839,7 +1839,8 @@ float UAnimInstance::PlaySlotAnimation(UAnimSequenceBase* Asset, FName SlotNodeN
 	NewSegment.AnimEndTime = Asset->SequenceLength;
 	NewSegment.AnimPlayRate = 1.f;
 	NewSegment.StartPos = 0.f;
-	NewMontage->SequenceLength = Asset->SequenceLength;
+	NewSegment.LoopingCount = LoopCount;
+	NewMontage->SequenceLength = NewSegment.GetLength();
 	NewTrack.AnimTrack.AnimSegments.Add(NewSegment);
 		
 	FCompositeSection NewSection;
@@ -1855,11 +1856,42 @@ float UAnimInstance::PlaySlotAnimation(UAnimSequenceBase* Asset, FName SlotNodeN
 	return Montage_Play(NewMontage, InPlayRate);
 }
 
-void UAnimInstance::StopSlotAnimation(float InBlendOutTime)
+void UAnimInstance::StopSlotAnimation(float InBlendOutTime, FName SlotNodeName)
 {
 	// stop temporary montage
 	// when terminate (in the Montage_Advance), we have to lose reference to the temporary montage
-	Montage_Stop(InBlendOutTime);
+	if (SlotNodeName != NAME_None)
+	{
+		for (int32 InstanceIndex = 0; InstanceIndex < MontageInstances.Num(); InstanceIndex++)
+		{
+			// check if this is playing
+			FAnimMontageInstance * MontageInstance = MontageInstances[InstanceIndex];
+			// make sure what is active right now is transient that we created by request
+			if (MontageInstance && MontageInstance->IsActive() && MontageInstance->IsPlaying())
+			{
+				UAnimMontage * CurMontage = MontageInstance->Montage;
+				if (CurMontage && CurMontage->GetOuter() == GetTransientPackage())
+				{
+					// Check each track, in practice there should only be one on these
+					for (int32 SlotTrackIndex = 0; SlotTrackIndex < CurMontage->SlotAnimTracks.Num(); SlotTrackIndex++)
+					{
+						const FSlotAnimationTrack * AnimTrack = &CurMontage->SlotAnimTracks[SlotTrackIndex];
+						if (AnimTrack && AnimTrack->SlotName == SlotNodeName)
+						{
+							// Found it
+							MontageInstance->Stop(InBlendOutTime);
+							break;
+						}
+					}
+				}
+			}
+		}
+	}
+	else
+	{
+		// Stop all
+		Montage_Stop(InBlendOutTime);
+	}
 }
 
 bool UAnimInstance::IsPlayingSlotAnimation(UAnimSequenceBase* Asset, FName SlotNodeName )
