#MATINEE: Simple anim BP support - no weight of slot node yet.

Merging

//depot/UE4-GDC2015/Engine/Source/Runtime/Engine/...

to //depot/UE4/Engine/Source/Runtime/Engine/...

diff --git a/Engine/Source/Runtime/Engine/Classes/Animation/SkeletalMeshActor.h b/Engine/Source/Runtime/Engine/Classes/Animation/SkeletalMeshActor.h
--- a/Engine/Source/Runtime/Engine/Classes/Animation/SkeletalMeshActor.h
+++ b/Engine/Source/Runtime/Engine/Classes/Animation/SkeletalMeshActor.h
@@ -88,9 +88,8 @@ public:
 	// End IMatineeAnimInterface Interface
 
 private:
-	EAnimationMode::Type	SavedAnimationMode;
 	// utility function to see if it can play animation or not
-	bool CanPlayAnimation();
+	bool CanPlayAnimation(class UAnimSequenceBase* AnimAssetBase=NULL) const;
 
 public:
 	/** Returns SkeletalMeshComponent subobject **/
diff --git a/Engine/Source/Runtime/Engine/Private/Interpolation.cpp b/Engine/Source/Runtime/Engine/Private/Interpolation.cpp
--- a/Engine/Source/Runtime/Engine/Private/Interpolation.cpp
+++ b/Engine/Source/Runtime/Engine/Private/Interpolation.cpp
@@ -7042,6 +7042,7 @@ UInterpTrackAnimControl::UInterpTrackAnimControl(const FObjectInitializer& Objec
 	TrackInstClass = UInterpTrackInstAnimControl::StaticClass();
 	TrackTitle = TEXT("Anim");
 	bIsAnimControlTrack = true;
+	SlotName = FAnimSlotGroup::DefaultSlotName;
 #if WITH_EDITORONLY_DATA
 	TrackIcon = Cast<UTexture2D>(StaticLoadObject( UTexture2D::StaticClass(), NULL, TEXT("/Engine/EditorMaterials/MatineeGroups/MAT_Groups_Anim.MAT_Groups_Anim"), NULL, LOAD_None, NULL ));
 #endif // WITH_EDITORONLY_DATA
diff --git a/Engine/Source/Runtime/Engine/Private/SkeletalMesh.cpp b/Engine/Source/Runtime/Engine/Private/SkeletalMesh.cpp
--- a/Engine/Source/Runtime/Engine/Private/SkeletalMesh.cpp
+++ b/Engine/Source/Runtime/Engine/Private/SkeletalMesh.cpp
@@ -27,6 +27,7 @@
 #include "ComponentReregisterContext.h"
 #include "PhysicsEngine/PhysicsAsset.h"
 #include "Engine/AssetUserData.h"
+#include "Animation/AnimMontage.h"
 
 #if WITH_EDITOR
 #include "MeshUtilities.h"
@@ -3666,8 +3667,11 @@ void ASkeletalMeshActor::PreviewBeginAnimControl(UInterpGroup* InInterpGroup)
 {
 	if (CanPlayAnimation())
 	{
-		SavedAnimationMode = SkeletalMeshComponent->GetAnimationMode();
-		SkeletalMeshComponent->SetAnimationMode(EAnimationMode::Type::AnimationSingleNode);
+		UAnimInstance* AnimInst = SkeletalMeshComponent->GetAnimInstance();
+		if (!AnimInst)
+		{
+			SkeletalMeshComponent->SetAnimationMode(EAnimationMode::Type::AnimationSingleNode);
+		}
 	}
 }
 
@@ -3678,46 +3682,80 @@ void ASkeletalMeshActor::PreviewFinishAnimControl(UInterpGroup* InInterpGroup)
 	{
 		// if in editor, reset the Animations, makes easier for artist to see them visually and align them
 		// in game, we keep the last pose that matinee kept. If you'd like it to have animation, you'll need to have AnimTree or AnimGraph to handle correctly
-		SkeletalMeshComponent->SetAnimationMode(SavedAnimationMode);
+		if (SkeletalMeshComponent->GetAnimationMode()== EAnimationMode::Type::AnimationBlueprint)
+		{
+			UAnimInstance* AnimInst = SkeletalMeshComponent->GetAnimInstance();
+			if(AnimInst)
+			{
+				AnimInst->Montage_Stop(0.f);
+				AnimInst->UpdateAnimation(0.f);
+			}
 
+		}
+		
 		// Update space bases to reset it back to ref pose
 		SkeletalMeshComponent->RefreshBoneTransforms();
 		SkeletalMeshComponent->RefreshSlaveComponents();
-		SkeletalMeshComponent->UpdateComponentToWorld();
+		SkeletalMeshComponent->UpdateComponentToWorld();		
 	}
 }
 
 
 void ASkeletalMeshActor::PreviewSetAnimPosition(FName SlotName, int32 ChannelIndex, UAnimSequence* InAnimSequence, float InPosition, bool bLooping, bool bFireNotifies, float DeltaTime)
 {
-	UAnimSingleNodeInstance * SingleNodeInst = SkeletalMeshComponent->GetSingleNodeInstance();
-	if (SingleNodeInst)
+	UAnimInstance* AnimInst = SkeletalMeshComponent->GetAnimInstance();
+	if (AnimInst && CanPlayAnimation(InAnimSequence))
 	{
-		if (SingleNodeInst->CurrentAsset != InAnimSequence)
+		UAnimSingleNodeInstance * SingleNodeInst = SkeletalMeshComponent->GetSingleNodeInstance();
+		if(SingleNodeInst)
 		{
-			SingleNodeInst->SetAnimationAsset(InAnimSequence, bLooping);
+			if(SingleNodeInst->CurrentAsset != InAnimSequence)
+			{
+				SingleNodeInst->SetAnimationAsset(InAnimSequence, bLooping);
+			}
+
+			SingleNodeInst->SetLooping(bLooping);
+			SingleNodeInst->SetPosition(InPosition, bFireNotifies);
 		}
+		else
+		{
+			bool bShouldChange = AnimInst->IsPlayingSlotAnimation(InAnimSequence, SlotName ) == false; 
+			if(bShouldChange)
+			{
+				AnimInst->PlaySlotAnimation(InAnimSequence, SlotName);
+			}
 
-		SingleNodeInst->SetLooping(bLooping);
-		SingleNodeInst->SetPosition(InPosition, bFireNotifies);
-	}
+			struct FAnimMontageInstance* AnimMontageInst = AnimInst->GetActiveMontageInstance();
+			if (AnimMontageInst)
+			{
+				AnimMontageInst->Weight = 1.f;
+				//AnimInst->Montage_Set(bLooping);
+				AnimInst->Montage_SetPosition(NULL, InPosition);
+				AnimInst->UpdateAnimation(DeltaTime);
+			}
+			else
+			{
+				UE_LOG(LogSkeletalMesh, Warning, TEXT("Invalid Slot Node Name: %s"), *SlotName.ToString() );
+			}
+		}
 
-	// Update space bases so new animation position has an effect.
-	SkeletalMeshComponent->UpdateMaterialParameters();
-	SkeletalMeshComponent->RefreshBoneTransforms();
-	SkeletalMeshComponent->RefreshSlaveComponents();
-	SkeletalMeshComponent->UpdateComponentToWorld();
+		// Update space bases so new animation position has an effect.
+		SkeletalMeshComponent->UpdateMaterialParameters();
+		SkeletalMeshComponent->RefreshBoneTransforms();
+		SkeletalMeshComponent->RefreshSlaveComponents();
+		SkeletalMeshComponent->UpdateComponentToWorld();
+	}
 }
 
 
 void ASkeletalMeshActor::PreviewSetAnimWeights(TArray<FAnimSlotInfo>& SlotInfos)
 {
-	// do nothing
+	//no support yet
 }
 
 void ASkeletalMeshActor::SetAnimWeights( const TArray<struct FAnimSlotInfo>& SlotInfos )
 {
-	// do nothing
+	//no support yet
 }
 
 /** Check SkeletalMeshActor for errors. */
@@ -3834,36 +3872,64 @@ void ASkeletalMeshActor::BeginAnimControl(UInterpGroup* InInterpGroup)
 {
 	if (CanPlayAnimation())
 	{
-		SavedAnimationMode = SkeletalMeshComponent->GetAnimationMode();
-		SkeletalMeshComponent->SetAnimationMode(EAnimationMode::Type::AnimationSingleNode);
+		UAnimInstance* AnimInst = SkeletalMeshComponent->GetAnimInstance();
+		if (!AnimInst)
+		{
+			SkeletalMeshComponent->SetAnimationMode(EAnimationMode::Type::AnimationSingleNode);
+		}
 	}
 }
 
-bool ASkeletalMeshActor::CanPlayAnimation()
+bool ASkeletalMeshActor::CanPlayAnimation(class UAnimSequenceBase* AnimAssetBase/*=NULL*/) const
 {
-	return (SkeletalMeshComponent->SkeletalMesh && SkeletalMeshComponent->SkeletalMesh->Skeleton);
+	return (SkeletalMeshComponent->SkeletalMesh && SkeletalMeshComponent->SkeletalMesh->Skeleton && 
+		(!AnimAssetBase || SkeletalMeshComponent->SkeletalMesh->Skeleton->IsCompatible(AnimAssetBase->GetSkeleton())));
 }
 
 // @todo ps4 clang bug: this works around a PS4/clang compiler bug (optimizations)
-void SetAnimPositionInner(USkeletalMeshComponent* SkeletalMeshComponent, UAnimSequence *InAnimSequence, float InPosition, bool bLooping)
+void SetAnimPositionInner(FName SlotName, USkeletalMeshComponent* SkeletalMeshComponent, UAnimSequence *InAnimSequence, float InPosition, bool bLooping)
 {
-	UAnimSingleNodeInstance * SingleNodeInst = SkeletalMeshComponent->GetSingleNodeInstance();
-	if (SingleNodeInst)
+	UAnimInstance* AnimInst = SkeletalMeshComponent->GetAnimInstance();
+	if(AnimInst)
 	{
-		if (SingleNodeInst->CurrentAsset != InAnimSequence)
+		UAnimSingleNodeInstance * SingleNodeInst = SkeletalMeshComponent->GetSingleNodeInstance();
+		if(SingleNodeInst)
 		{
-			SingleNodeInst->SetAnimationAsset(InAnimSequence, bLooping);
-			SingleNodeInst->SetPosition(0.0f);
-			SingleNodeInst->bPlaying = false;
-		}
+			if(SingleNodeInst->CurrentAsset != InAnimSequence)
+			{
+				SingleNodeInst->SetAnimationAsset(InAnimSequence, bLooping);
+				SingleNodeInst->SetPosition(0.0f);
+				SingleNodeInst->bPlaying = false;
+			}
 
-		if (SingleNodeInst->bLooping!=bLooping)
-		{
-			SingleNodeInst->SetLooping(bLooping);
+			if(SingleNodeInst->bLooping!=bLooping)
+			{
+				SingleNodeInst->SetLooping(bLooping);
+			}
+			if(SingleNodeInst->CurrentTime != InPosition)
+			{
+				SingleNodeInst->SetPosition(InPosition);
+			}
 		}
-		if (SingleNodeInst->CurrentTime != InPosition)
+		else
 		{
-			SingleNodeInst->SetPosition(InPosition);
+			bool bShouldChange = AnimInst->IsPlayingSlotAnimation(InAnimSequence, SlotName) == false;
+			if(bShouldChange)
+			{
+				AnimInst->PlaySlotAnimation(InAnimSequence, SlotName, 0.f, 0.f, 0.f);
+			}
+
+			struct FAnimMontageInstance* AnimMontageInst = AnimInst->GetActiveMontageInstance();
+			if (AnimMontageInst)
+			{
+				AnimMontageInst->Weight = 1.f;
+				AnimInst->Montage_SetPosition(NULL, InPosition);
+				AnimInst->UpdateAnimation(0.f);
+			}
+			else
+			{
+				UE_LOG(LogSkeletalMesh, Warning, TEXT("Invalid Slot Node Name: %s"), *SlotName.ToString() );
+			}
 		}
 	}
 }
@@ -3872,20 +3938,28 @@ void SetAnimPositionInner(USkeletalMeshComponent* SkeletalMeshComponent, UAnimSe
 #if !PLATFORM_PS4
 void ASkeletalMeshActor::SetAnimPosition(FName SlotName, int32 ChannelIndex, UAnimSequence *InAnimSequence, float InPosition, bool bFireNotifies, bool bLooping)
 {
-	if (CanPlayAnimation())
+	if (CanPlayAnimation(InAnimSequence))
 	{
-		SetAnimPositionInner(SkeletalMeshComponent, InAnimSequence, InPosition, bLooping);
+		SetAnimPositionInner(SlotName, SkeletalMeshComponent, InAnimSequence, InPosition, bLooping);
 	}
 }
 #endif
 
 void ASkeletalMeshActor::FinishAnimControl(UInterpGroup* InInterpGroup)
 {
-	if (CanPlayAnimation())
+	if(SkeletalMeshComponent->GetAnimationMode()== EAnimationMode::Type::AnimationBlueprint)
 	{
-		// if in editor, reset the Animations, makes easier for artist to see them visually and align them
-		// in game, we keep the last pose that matinee kept. If you'd like it to have animation, you'll need to have AnimTree or AnimGraph to handle correctly
-		SkeletalMeshComponent->SetAnimationMode(SavedAnimationMode);
+		UAnimInstance* AnimInst = SkeletalMeshComponent->GetAnimInstance();
+		if(AnimInst)
+		{
+			AnimInst->Montage_Stop(0.f);
+			AnimInst->UpdateAnimation(0.f);
+		}
+
+		// Update space bases to reset it back to ref pose
+		SkeletalMeshComponent->RefreshBoneTransforms();
+		SkeletalMeshComponent->RefreshSlaveComponents();
+		SkeletalMeshComponent->UpdateComponentToWorld();
 	}
 }
 
