UpdateSceneCaptures now only updates components in their associated scene rather than first scene rendered.  This fixes issues were scenes (moslty preview scenes) could render first and cause an UpdateSceneCapture to update the render targets with data from the wrong world.  This change makes scene captures work in a multi-world environment.

diff --git a/Engine/Source/Runtime/Engine/Private/SceneCaptureComponent.cpp b/Engine/Source/Runtime/Engine/Private/SceneCaptureComponent.cpp
--- a/Engine/Source/Runtime/Engine/Private/SceneCaptureComponent.cpp
+++ b/Engine/Source/Runtime/Engine/Private/SceneCaptureComponent.cpp
@@ -295,24 +295,35 @@ void USceneCaptureComponent2D::TickComponent(float DeltaTime, enum ELevelTick Ti
 	}
 }
 
-static TArray<USceneCaptureComponent2D*> SceneCapturesToUpdate;
+static TMultiMap<TWeakObjectPtr<UWorld>, USceneCaptureComponent2D*> SceneCapturesToUpdateMap;
 
 void USceneCaptureComponent2D::UpdateContent()
 {
 	if (World && World->Scene && IsVisible())
 	{
 		// Defer until after updates finish
-		SceneCapturesToUpdate.AddUnique( this );
+		SceneCapturesToUpdateMap.AddUnique(World, this);
 	}	
 }
 
 void USceneCaptureComponent2D::UpdateDeferredCaptures( FSceneInterface* Scene )
 {
-	for( int32 CaptureIndex = 0; CaptureIndex < SceneCapturesToUpdate.Num(); CaptureIndex++ )
+	UWorld* World = Scene->GetWorld();
+	if( World && SceneCapturesToUpdateMap.Num() > 0 )
 	{
-		Scene->UpdateSceneCaptureContents( SceneCapturesToUpdate[ CaptureIndex ] );
+		// Only update the scene captures assoicated with the current scene.
+		// Updating others not associated with the scene would cause invalid data to be rendered into the target
+		TArray<USceneCaptureComponent2D*> SceneCapturesToUpdate;
+		SceneCapturesToUpdateMap.MultiFind( World, SceneCapturesToUpdate );
+		
+		for( USceneCaptureComponent2D* Component : SceneCapturesToUpdate )
+		{
+			Scene->UpdateSceneCaptureContents( Component );
+		}
+		
+		// All scene captures for this world have been updated
+		SceneCapturesToUpdateMap.Remove( World) ;
 	}
-	SceneCapturesToUpdate.Reset();
 }
 
 #if WITH_EDITOR
@@ -354,24 +365,37 @@ void USceneCaptureComponentCube::TickComponent(float DeltaTime, enum ELevelTick
 	}
 }
 
-static TArray<USceneCaptureComponentCube*> CubedSceneCapturesToUpdate;
+static TMultiMap<TWeakObjectPtr<UWorld>, USceneCaptureComponentCube*> CubedSceneCapturesToUpdateMap;
 
 void USceneCaptureComponentCube::UpdateContent()
 {
 	if (World && World->Scene && IsVisible())
 	{
 		// Defer until after updates finish
-		CubedSceneCapturesToUpdate.AddUnique( this );
+		CubedSceneCapturesToUpdateMap.AddUnique( World, this );
 	}	
 }
 
 void USceneCaptureComponentCube::UpdateDeferredCaptures( FSceneInterface* Scene )
 {
-	for( int32 CaptureIndex = 0; CaptureIndex < CubedSceneCapturesToUpdate.Num(); CaptureIndex++ )
+	UWorld* World = Scene->GetWorld();
+	
+	if( World && CubedSceneCapturesToUpdateMap.Num() > 0 )
 	{
-		Scene->UpdateSceneCaptureContents( CubedSceneCapturesToUpdate[ CaptureIndex ] );
+		// Only update the scene captures assoicated with the current scene.
+		// Updating others not associated with the scene would cause invalid data to be rendered into the target
+		TArray<USceneCaptureComponentCube*> SceneCapturesToUpdate;
+		CubedSceneCapturesToUpdateMap.MultiFind( World, SceneCapturesToUpdate );
+		
+		for( USceneCaptureComponentCube* Component : SceneCapturesToUpdate )
+		{
+			Scene->UpdateSceneCaptureContents( Component );
+		}
+		
+		// All scene captures for this world have been updated
+		CubedSceneCapturesToUpdateMap.Remove( World );
 	}
-	CubedSceneCapturesToUpdate.Reset();
+
 }
 
 #if WITH_EDITOR
