#MATINEE: ANIM: Make sure only one is active at a time

TODO: should support weight change.

Merging

//depot/UE4-GDC2015/Engine/Source/Runtime/Engine/Private/SkeletalMesh.cpp

to //depot/UE4-GDC2015/Engine/Source/Runtime/Engine/Private/SkeletalMesh.cpp

diff --git a/Engine/Source/Runtime/Engine/Private/SkeletalMesh.cpp b/Engine/Source/Runtime/Engine/Private/SkeletalMesh.cpp
--- a/Engine/Source/Runtime/Engine/Private/SkeletalMesh.cpp
+++ b/Engine/Source/Runtime/Engine/Private/SkeletalMesh.cpp
@@ -3722,6 +3722,7 @@ void ASkeletalMeshActor::PreviewSetAnimPosition(FName SlotName, int32 ChannelInd
 			bool bShouldChange = AnimInst->IsPlayingSlotAnimation(InAnimSequence, SlotName ) == false; 
 			if(bShouldChange)
 			{
+				AnimInst->Montage_Stop(0.f);
 				AnimInst->PlaySlotAnimation(InAnimSequence, SlotName);
 			}
 
