Animation: Optimizing state machine ticking
- Prevent most allocations and reallocations for transition evaluation
- Avoid calling the BP transition rule when there is a C++ rule
- Fix an issue where NativeTransitionDelegate wasn't checked when dealing with conduit entry rules
[INTEGRATE] Change 2454954

diff --git a/Engine/Source/Runtime/Engine/Classes/Animation/AnimNode_StateMachine.h b/Engine/Source/Runtime/Engine/Classes/Animation/AnimNode_StateMachine.h
--- a/Engine/Source/Runtime/Engine/Classes/Animation/AnimNode_StateMachine.h
+++ b/Engine/Source/Runtime/Engine/Classes/Animation/AnimNode_StateMachine.h
@@ -52,7 +52,7 @@ struct FAnimationActiveTransitionEntry
 	TArray<FAnimNode_TransitionPoseEvaluator*> PoseEvaluators;
 
 #if WITH_EDITORONLY_DATA
-	TArray<int32> SourceTransitionIndices;
+	TArray<int32, TInlineAllocator<3>> SourceTransitionIndices;
 #endif
 
 
@@ -84,7 +84,7 @@ struct FAnimationPotentialTransition
 	const FBakedStateExitTransition* TransitionRule;
 
 #if WITH_EDITORONLY_DATA
-	TArray<int32> SourceTransitionIndices;
+	TArray<int32, TInlineAllocator<3>> SourceTransitionIndices;
 #endif
 
 public:
@@ -197,7 +197,7 @@ protected:
 	bool FindValidTransition(const FAnimationUpdateContext& Context, 
 							const FBakedAnimationState& StateInfo,
 							/*OUT*/ FAnimationPotentialTransition& OutPotentialTransition,
-							/*OUT*/ TArray<int32>& OutVisitedStateIndices);
+							/*OUT*/ TArray<int32, TInlineAllocator<4>>& OutVisitedStateIndices);
 
 	// Helper function that will update the states associated with a transition
 	void UpdateTransitionStates(const FAnimationUpdateContext& Context, FAnimationActiveTransitionEntry& Transition);
diff --git a/Engine/Source/Runtime/Engine/Private/Animation/AnimNode_StateMachine.cpp b/Engine/Source/Runtime/Engine/Private/Animation/AnimNode_StateMachine.cpp
--- a/Engine/Source/Runtime/Engine/Private/Animation/AnimNode_StateMachine.cpp
+++ b/Engine/Source/Runtime/Engine/Private/Animation/AnimNode_StateMachine.cpp
@@ -148,7 +148,7 @@ void FAnimationPotentialTransition::Clear()
 	TargetState = INDEX_NONE;
 	TransitionRule = NULL;
 #if WITH_EDITORONLY_DATA
-	SourceTransitionIndices.Empty();
+	SourceTransitionIndices.Reset();
 #endif
 }
 
@@ -287,7 +287,7 @@ void FAnimNode_StateMachine::Update(const FAnimationUpdateContext& Context)
 
 			// Evaluate possible transitions out of this state
 			//@TODO: Evaluate if a set is better than an array for the probably low N encountered here
-			TArray<int32> VisitedStateIndices;
+			TArray<int32, TInlineAllocator<4>> VisitedStateIndices;
 			FindValidTransition(Context, GetStateInfo(), /*Out*/ PotentialTransition, /*Out*/ VisitedStateIndices);
 		}
 				
@@ -334,7 +334,7 @@ void FAnimNode_StateMachine::Update(const FAnimationUpdateContext& Context)
 	// so we throw out any transition data at the first update
 	if (bFirstUpdate)
 	{
-		ActiveTransitionArray.Empty();
+		ActiveTransitionArray.Reset();
 		bFirstUpdate = false;
 	}
 
@@ -389,7 +389,7 @@ void FAnimNode_StateMachine::Update(const FAnimationUpdateContext& Context)
 	ElapsedTime += Context.GetDeltaTime();
 }
 
-bool FAnimNode_StateMachine::FindValidTransition(const FAnimationUpdateContext& Context, const FBakedAnimationState& StateInfo, /*out*/ FAnimationPotentialTransition& OutPotentialTransition, /*out*/ TArray<int32>& OutVisitedStateIndices)
+bool FAnimNode_StateMachine::FindValidTransition(const FAnimationUpdateContext& Context, const FBakedAnimationState& StateInfo, /*out*/ FAnimationPotentialTransition& OutPotentialTransition, /*out*/ TArray<int32, TInlineAllocator<4>>& OutVisitedStateIndices)
 {
 	// There is a possibility we'll revisit states connected through conduits,
 	// so we can avoid doing unnecessary work (and infinite loops) by caching off states we have already checked
@@ -406,7 +406,16 @@ bool FAnimNode_StateMachine::FindValidTransition(const FAnimationUpdateContext&
 	//@TODO: It would add flexibility to be able to define this on normal state nodes as well, assuming the dual-graph editing is sorted out
 	if (FAnimNode_TransitionResult* StateEntryRuleNode = GetNodeFromPropertyIndex<FAnimNode_TransitionResult>(Context.AnimInstance, AnimBlueprintClass, StateInfo.EntryRuleNodeIndex))
 	{
-		StateEntryRuleNode->EvaluateGraphExposedInputs.Execute(Context);
+		if (StateEntryRuleNode->NativeTransitionDelegate.IsBound())
+		{
+			// attempt to evaluate native rule
+			StateEntryRuleNode->bCanEnterTransition = StateEntryRuleNode->NativeTransitionDelegate.Execute();
+		}
+		else
+		{
+			// Execute it and see if we can take this rule
+			StateEntryRuleNode->EvaluateGraphExposedInputs.Execute(Context);
+		}
 
 		// not ok, back out
 		if (!StateEntryRuleNode->bCanEnterTransition)
@@ -426,16 +435,15 @@ bool FAnimNode_StateMachine::FindValidTransition(const FAnimationUpdateContext&
 
 		FAnimNode_TransitionResult* ResultNode = GetNodeFromPropertyIndex<FAnimNode_TransitionResult>(Context.AnimInstance, AnimBlueprintClass, TransitionRule.CanTakeDelegateIndex);
 
-		// reset transition flag (this would be always set if we had a bound graph, but if we only have a delegate it will not)
-		ResultNode->bCanEnterTransition = false;
-
-		// Execute it and see if we can take this rule
-		ResultNode->EvaluateGraphExposedInputs.Execute(Context);
-
-		// attempt to evaluate native rule
-		if(ResultNode->NativeTransitionDelegate.IsBound())
+		if (ResultNode->NativeTransitionDelegate.IsBound())
+		{
+			// attempt to evaluate native rule
+			ResultNode->bCanEnterTransition = ResultNode->NativeTransitionDelegate.Execute();
+		}
+		else
 		{
-			ResultNode->bCanEnterTransition |= ResultNode->NativeTransitionDelegate.Execute();
+			// Execute it and see if we can take this rule
+			ResultNode->EvaluateGraphExposedInputs.Execute(Context);
 		}
 
 		if (ResultNode->bCanEnterTransition == TransitionRule.bDesiredTransitionReturnValue)
