UE-9754 FLatentActionManager::RemoveActionsForObject removes actions added after it was called

diff --git a/Engine/Source/Runtime/Engine/Classes/Engine/LatentActionManager.h b/Engine/Source/Runtime/Engine/Classes/Engine/LatentActionManager.h
--- a/Engine/Source/Runtime/Engine/Classes/Engine/LatentActionManager.h
+++ b/Engine/Source/Runtime/Engine/Classes/Engine/LatentActionManager.h
@@ -185,7 +185,9 @@ protected:
 	/** list of objects we have processed the latent action list for this frame. */	
 	TSet<UObject*> ProcessedThisFrame;
 
-	/** For listed objects all actions will be removed */
-	TArray<TWeakObjectPtr<UObject>> ObjectsToRemove;
+	/**List of actions that will be unconditionally removed at the begin of next tick */
+	typedef TPair<int32, class FPendingLatentAction*> FUuidAndAction;
+	typedef TMap< TWeakObjectPtr<UObject>, TSharedPtr<TArray<FUuidAndAction>>> FActionsForObject;
+	FActionsForObject ActionsToRemoveMap;
 };
 
diff --git a/Engine/Source/Runtime/Engine/Private/LatentActionManager.cpp b/Engine/Source/Runtime/Engine/Private/LatentActionManager.cpp
--- a/Engine/Source/Runtime/Engine/Private/LatentActionManager.cpp
+++ b/Engine/Source/Runtime/Engine/Private/LatentActionManager.cpp
@@ -28,31 +28,43 @@ void FLatentActionManager::AddNewAction(UObject* InActionObject, int32 UUID, FPe
 
 void FLatentActionManager::RemoveActionsForObject(TWeakObjectPtr<UObject> InObject)
 {
-	ObjectsToRemove.Add(InObject);
+	auto ObjectActionList = GetActionListForObject(InObject);
+	if (ObjectActionList)
+	{
+		auto ActionToRemoveListPtr = ActionsToRemoveMap.FindOrAdd(InObject);
+		if (!ActionToRemoveListPtr.IsValid())
+		{
+			ActionToRemoveListPtr = MakeShareable(new TArray<FUuidAndAction>());
+		}
+
+		for (FActionList::TConstIterator It(*ObjectActionList); It; ++It)
+		{
+			ActionToRemoveListPtr->Add(*It);
+		}
+	}
 }
 
 void FLatentActionManager::ProcessLatentActions(UObject* InObject, float DeltaTime)
 {
-	if (ObjectsToRemove.Num())
+	for (FActionsForObject::TIterator It(ActionsToRemoveMap); It; ++It)
 	{
-		for (const auto Key : ObjectsToRemove)
+		auto ActionList = GetActionListForObject(InObject);
+		auto ActionToRemoveListPtr = It.Value();
+		if (ActionToRemoveListPtr.IsValid() && ActionList)
 		{
-			FActionList* const ObjectActionList = GetActionListForObject(Key);
-			if (ObjectActionList)
+			for (auto PendingActionToKill : *ActionToRemoveListPtr)
 			{
-				for (TMultiMap<int32, FPendingLatentAction*>::TConstIterator It(*ObjectActionList); It; ++It)
+				const int32 RemovedNum = ActionList->RemoveSingle(PendingActionToKill.Key, PendingActionToKill.Value);
+				auto Action = PendingActionToKill.Value;
+				if (RemovedNum && Action)
 				{
-					FPendingLatentAction* Action = It.Value();
 					Action->NotifyActionAborted();
 					delete Action;
 				}
-				ObjectActionList->Empty();
-
-				ObjectToActionListMap.Remove(Key); //TODO: we should be able to obtain an Iterator from Map::Find function. "Iterator.RemoveCurrent();" should be called here.
 			}
 		}
-		ObjectsToRemove.Empty();
 	}
+	ActionsToRemoveMap.Empty();
 
 	//@TODO: K2: Very inefficient code right now
 	if (InObject != NULL)
