Clean up some more code in FConstraintInstance

diff --git a/Engine/Source/Runtime/Engine/Classes/PhysicsEngine/ConstraintInstance.h b/Engine/Source/Runtime/Engine/Classes/PhysicsEngine/ConstraintInstance.h
--- a/Engine/Source/Runtime/Engine/Classes/PhysicsEngine/ConstraintInstance.h
+++ b/Engine/Source/Runtime/Engine/Classes/PhysicsEngine/ConstraintInstance.h
@@ -492,12 +492,12 @@ public:
 	static void Free(FConstraintInstance * Ptr);
 	static FConstraintInstance * Alloc();
 
-
 private:
 #if WITH_PHYSX 
 	bool CreatePxJoint(physx::PxRigidActor* PActor1, physx::PxRigidActor* PActor2, physx::PxScene* PScene, const float Scale);
 	void UpdateConstraintFlags();
 	void UpdateAverageMass(const physx::PxRigidActor* PActor1, const physx::PxRigidActor* PActor2);
+	physx::PxD6Joint* GetUnbrokenJoint() const;
 #endif
 
 	void UpdateBreakable();
diff --git a/Engine/Source/Runtime/Engine/Private/PhysicsEngine/ConstraintInstance.cpp b/Engine/Source/Runtime/Engine/Private/PhysicsEngine/ConstraintInstance.cpp
--- a/Engine/Source/Runtime/Engine/Private/PhysicsEngine/ConstraintInstance.cpp
+++ b/Engine/Source/Runtime/Engine/Private/PhysicsEngine/ConstraintInstance.cpp
@@ -25,47 +25,84 @@ static TAutoConsoleVariable<float> CVarConstraintStiffnessScale(
 	TEXT("The multiplier of constraint stiffness in simulation. Default: 100000"),
 	ECVF_ReadOnly);
 
+float RevolutionsToRads(const float Revolutions)
+{
+	return Revolutions * 2.f * PI;
+}
+
+FVector RevolutionsToRads(const FVector Revolutions)
+{
+	return Revolutions * 2.f * PI;
+}
 
 #if WITH_PHYSX
 /** Util for setting soft limit params */
-static void SetSoftLimitParams(PxJointLimitParameters* PLimit, bool bSoft, float Spring, float Damping)
+void SetSoftLimitParams(PxJointLimitParameters* PLimit, bool bSoft, float Spring, float Damping)
 {
 	if(bSoft)
 	{
-		PLimit->stiffness = Spring;
-		PLimit->damping = Damping;
+		PLimit->stiffness = Spring * CVarConstraintStiffnessScale.GetValueOnGameThread();
+		PLimit->damping = Damping * CVarConstraintDampingScale.GetValueOnGameThread();
 	}
 }
 
-/** Util for setting linear movement for an axis */
-static void SetLinearMovement(PxD6Joint* PD6Joint, PxD6Axis::Enum PAxis, uint8 Motion, bool bLockLimitSize)
+/** Util for converting from UE motion enum to physx motion enum */
+PxD6Motion::Enum U2PLinearMotion(ELinearConstraintMotion InMotion)
 {
-	if(Motion == LCM_Locked || (Motion == LCM_Limited && bLockLimitSize))
+	switch (InMotion)
 	{
-		PD6Joint->setMotion(PAxis, PxD6Motion::eLOCKED);
+		case ELinearConstraintMotion::LCM_Free: return PxD6Motion::eFREE;
+		case ELinearConstraintMotion::LCM_Limited: return PxD6Motion::eLIMITED;
+		case ELinearConstraintMotion::LCM_Locked: return PxD6Motion::eLOCKED;
+		default: check(0);	//unsupported motion type
 	}
-	else if(Motion == LCM_Limited)
+	
+	return PxD6Motion::eFREE;
+}
+
+/** Util for converting from UE motion enum to physx motion enum */
+PxD6Motion::Enum U2PAngularMotion(EAngularConstraintMotion InMotion)
+{
+	switch (InMotion)
 	{
-		PD6Joint->setMotion(PAxis, PxD6Motion::eLIMITED);
+		case EAngularConstraintMotion::ACM_Free: return PxD6Motion::eFREE;
+		case EAngularConstraintMotion::ACM_Limited: return PxD6Motion::eLIMITED;
+		case EAngularConstraintMotion::ACM_Locked: return PxD6Motion::eLOCKED;
+		default: check(0);	//unsupported motion type
 	}
-	else
+	
+	return PxD6Motion::eFREE;
+}
+
+/** Util for setting linear movement for an axis */
+template <PxD6Axis::Enum PAxis>
+void SetLinearMovement(PxD6Joint* PD6Joint, ELinearConstraintMotion Motion, bool bLockLimitSize)
+{
+	if(Motion == LCM_Locked || (Motion == LCM_Limited && bLockLimitSize))
 	{
-		PD6Joint->setMotion(PAxis, PxD6Motion::eFREE);
+		PD6Joint->setMotion(PAxis, PxD6Motion::eLOCKED);
+	}else
+	{
+		PD6Joint->setMotion(PAxis, U2PLinearMotion(Motion));
 	}
 }
+
+physx::PxD6Joint* FConstraintInstance::GetUnbrokenJoint() const
+{
+	return (ConstraintData && !(ConstraintData->getConstraintFlags()&PxConstraintFlag::eBROKEN)) ? ConstraintData : nullptr;
+}
 #endif //WITH_PHYSX
 
 void FConstraintInstance::UpdateLinearLimit()
 {
 #if WITH_PHYSX
-	PxD6Joint* Joint = ConstraintData;
-	if (Joint && !(Joint->getConstraintFlags()&PxConstraintFlag::eBROKEN))
+	if (PxD6Joint* Joint = GetUnbrokenJoint())
 	{
 		bool bLockLimitSize = (LinearLimitSize < RB_MinSizeToLockDOF);
 
-		SetLinearMovement(Joint, PxD6Axis::eX, LinearXMotion, bLockLimitSize);
-		SetLinearMovement(Joint, PxD6Axis::eY, LinearYMotion, bLockLimitSize);
-		SetLinearMovement(Joint, PxD6Axis::eZ, LinearZMotion, bLockLimitSize);
+		SetLinearMovement<PxD6Axis::eX>(Joint, LinearXMotion, bLockLimitSize);
+		SetLinearMovement<PxD6Axis::eY>(Joint, LinearYMotion, bLockLimitSize);
+		SetLinearMovement<PxD6Axis::eZ>(Joint, LinearZMotion, bLockLimitSize);
 
 		// If any DOF is locked/limited, set up the joint limit
 		if (LinearXMotion != LCM_Free || LinearYMotion != LCM_Free || LinearZMotion != LCM_Free)
@@ -73,18 +110,19 @@ void FConstraintInstance::UpdateLinearLimit()
 			// If limit drops below RB_MinSizeToLockDOF, just pass RB_MinSizeToLockDOF to physics - that axis will be locked anyway, and PhysX dislikes 0 here
 			float LinearLimit = FMath::Max<float>(LinearLimitSize, RB_MinSizeToLockDOF);
 			PxJointLinearLimit PLinearLimit(GPhysXSDK->getTolerancesScale(), LinearLimit, 0.05f * GPhysXSDK->getTolerancesScale().length);
-			SetSoftLimitParams(&PLinearLimit, bLinearLimitSoft, LinearLimitStiffness*AverageMass*CVarConstraintStiffnessScale.GetValueOnGameThread(), LinearLimitDamping*AverageMass*CVarConstraintDampingScale.GetValueOnGameThread());
+			SetSoftLimitParams(&PLinearLimit, bLinearLimitSoft, LinearLimitStiffness*AverageMass, LinearLimitDamping*AverageMass);
 			Joint->setLinearLimit(PLinearLimit);
 		}
 	}
 #endif
 }
 
+
+
 void FConstraintInstance::UpdateAngularLimit()
 {
 #if WITH_PHYSX
-	PxD6Joint* Joint = ConstraintData;
-	if (Joint && !(Joint->getConstraintFlags()&PxConstraintFlag::eBROKEN))
+	if (PxD6Joint* Joint = GetUnbrokenJoint())
 	{
 		/////////////// TWIST LIMIT
 		PxD6Motion::Enum TwistMotion = PxD6Motion::eFREE;
@@ -92,9 +130,9 @@ void FConstraintInstance::UpdateAngularLimit()
 		{
 			TwistMotion = PxD6Motion::eLIMITED;
 			// If angle drops below RB_MinAngleToLockDOF, just pass RB_MinAngleToLockDOF to physics - that axis will be locked anyway, and PhysX dislikes 0 here
-			float TwistLimitRad = TwistLimitAngle * (PI / 180.0f);
-			PxJointAngularLimitPair PTwistLimitPair(-TwistLimitRad, TwistLimitRad, 1.f * (PI / 180.0f));
-			SetSoftLimitParams(&PTwistLimitPair, bTwistLimitSoft, TwistLimitStiffness*AverageMass*CVarConstraintStiffnessScale.GetValueOnGameThread(), TwistLimitDamping*AverageMass*CVarConstraintDampingScale.GetValueOnGameThread());
+			float TwistLimitRad = FMath::DegreesToRadians(TwistLimitAngle);
+			PxJointAngularLimitPair PTwistLimitPair(-TwistLimitRad, TwistLimitRad, FMath::DegreesToRadians(1.f));
+			SetSoftLimitParams(&PTwistLimitPair, bTwistLimitSoft, TwistLimitStiffness*AverageMass, TwistLimitDamping*AverageMass);
 			Joint->setTwistLimit(PTwistLimitPair);
 		}
 		else if (AngularTwistMotion == ACM_Locked)
@@ -104,35 +142,16 @@ void FConstraintInstance::UpdateAngularLimit()
 		Joint->setMotion(PxD6Axis::eTWIST, TwistMotion);
 
 		/////////////// SWING1 LIMIT
-		PxD6Motion::Enum Swing1Motion = PxD6Motion::eFREE;
-		PxD6Motion::Enum Swing2Motion = PxD6Motion::eFREE;
-
-		if (AngularSwing1Motion == ACM_Limited)
-		{
-			Swing1Motion = PxD6Motion::eLIMITED;
-		}
-		else if (AngularSwing1Motion == ACM_Locked)
-		{
-			Swing1Motion = PxD6Motion::eLOCKED;
-		}
-
-		/////////////// SWING2 LIMIT
-		if (AngularSwing2Motion == ACM_Limited)
-		{
-			Swing2Motion = PxD6Motion::eLIMITED;
-		}
-		else if (AngularSwing2Motion == ACM_Locked)
-		{
-			Swing2Motion = PxD6Motion::eLOCKED;
-		}
+		const PxD6Motion::Enum Swing1Motion = U2PAngularMotion(AngularSwing1Motion);
+		const PxD6Motion::Enum Swing2Motion = U2PAngularMotion(AngularSwing2Motion);
 
 		if (AngularSwing1Motion == ACM_Limited || AngularSwing2Motion == ACM_Limited)
 		{
 			//Clamp the limit value to valid range which PhysX won't ignore, both value have to be clamped even there is only one degree limit in constraint
-			float Limit1Rad = FMath::ClampAngle(Swing1LimitAngle, KINDA_SMALL_NUMBER, 179.9999f) * (PI / 180.0f);
-			float Limit2Rad = FMath::ClampAngle(Swing2LimitAngle, KINDA_SMALL_NUMBER, 179.9999f) * (PI / 180.0f);
-			PxJointLimitCone PSwingLimitCone(Limit2Rad, Limit1Rad, 1.f * (PI / 180.0f));
-			SetSoftLimitParams(&PSwingLimitCone, bSwingLimitSoft, SwingLimitStiffness*AverageMass*CVarConstraintStiffnessScale.GetValueOnGameThread(), SwingLimitDamping*AverageMass*CVarConstraintDampingScale.GetValueOnGameThread());
+			float Limit1Rad = FMath::DegreesToRadians(FMath::ClampAngle(Swing1LimitAngle, KINDA_SMALL_NUMBER, 179.9999f));
+			float Limit2Rad = FMath::DegreesToRadians(FMath::ClampAngle(Swing2LimitAngle, KINDA_SMALL_NUMBER, 179.9999f));
+			PxJointLimitCone PSwingLimitCone(Limit2Rad, Limit1Rad, FMath::DegreesToRadians(1.f));
+			SetSoftLimitParams(&PSwingLimitCone, bSwingLimitSoft, SwingLimitStiffness*AverageMass, SwingLimitDamping*AverageMass);
 			Joint->setSwingLimit(PSwingLimitCone);
 		}
 
@@ -156,12 +175,12 @@ void FConstraintInstance::UpdateBreakable()
 void FConstraintInstance::UpdateDriveTarget()
 {
 #if WITH_PHYSX
-	if (ConstraintData && !(ConstraintData->getConstraintFlags()&PxConstraintFlag::eBROKEN))
+	if (PxD6Joint* PJoint = GetUnbrokenJoint())
 	{
 		FQuat OrientationTargetQuat(AngularOrientationTarget);
 
-		ConstraintData->setDrivePosition(PxTransform(U2PVector(LinearPositionTarget), U2PQuat(OrientationTargetQuat)));
-		ConstraintData->setDriveVelocity(U2PVector(LinearVelocityTarget), U2PVector(AngularVelocityTarget * 2 * PI));
+		PJoint->setDrivePosition(PxTransform(U2PVector(LinearPositionTarget), U2PQuat(OrientationTargetQuat)));
+		PJoint->setDriveVelocity(U2PVector(LinearVelocityTarget), U2PVector(RevolutionsToRads(AngularVelocityTarget)));
 	}
 #endif
 }
@@ -227,8 +246,7 @@ void FConstraintInstance::SetDisableCollision(bool InDisableCollision)
 {
 	bDisableCollision = InDisableCollision;
 #if WITH_PHYSX
-	PxD6Joint* Joint = ConstraintData;
-	if (Joint && !(Joint->getConstraintFlags()&PxConstraintFlag::eBROKEN))
+	if (PxD6Joint* Joint = GetUnbrokenJoint())
 	{
 		PxConstraintFlags Flags = Joint->getConstraintFlags();
 		if (bDisableCollision)
@@ -390,7 +408,7 @@ void FConstraintInstance::UpdateConstraintFlags()
 		Flags |= PxConstraintFlag::ePROJECTION;
 
 		ConstraintData->setProjectionLinearTolerance(ProjectionLinearTolerance);
-		ConstraintData->setProjectionAngularTolerance(ProjectionAngularTolerance * ((float)PI / 180.0f));
+		ConstraintData->setProjectionAngularTolerance(FMath::DegreesToRadians(ProjectionAngularTolerance));
 	}
 
 	ConstraintData->setConstraintFlags(Flags);
@@ -411,12 +429,12 @@ void EnsureSleepingActorsStaySleeping(PxRigidActor* PActor1, PxRigidActor* PActo
 	// creation of joints wakes up rigid bodies, so we put them to sleep again if both were initially asleep
 	if (bActor1Asleep && bActor2Asleep)
 	{
-		if (PActor1 != NULL && IsRigidBodyNonKinematic(PActor1->isRigidDynamic()))
+		if (PActor1 && IsRigidBodyNonKinematic(PActor1->isRigidDynamic()))
 		{
 			PActor1->isRigidDynamic()->putToSleep();
 		}
 
-		if (PActor2 != NULL && IsRigidBodyNonKinematic(PActor2->isRigidDynamic()))
+		if (PActor2 && IsRigidBodyNonKinematic(PActor2->isRigidDynamic()))
 		{
 			PActor2->isRigidDynamic()->putToSleep();
 		}
@@ -483,20 +501,19 @@ void FConstraintInstance::TermConstraint()
 	}
 
 	// use correct scene
-	PxScene* PScene = GetPhysXSceneFromIndex(SceneIndex);
-	if(PScene != NULL)
+	if(PxScene* PScene = GetPhysXSceneFromIndex(SceneIndex))
 	{
 		ConstraintData->release();
 	}
 
-	ConstraintData = NULL;
+	ConstraintData = nullptr;
 #endif
 }
 
 bool FConstraintInstance::IsTerminated() const
 {
 #if WITH_PHYSX
-	return (ConstraintData == NULL);
+	return (ConstraintData == nullptr);
 #else 
 	return true;
 #endif //WITH_PHYSX
@@ -505,7 +522,7 @@ bool FConstraintInstance::IsTerminated() const
 bool FConstraintInstance::IsValidConstraintInstance() const
 {
 #if WITH_PHYSX
-	return ConstraintData != NULL;
+	return ConstraintData != nullptr;
 #else
 	return false;
 #endif // WITH_PHYSX
@@ -585,10 +602,10 @@ void FConstraintInstance::SetRefFrame(EConstraintFrame::Type Frame, const FTrans
 	}
 
 #if WITH_PHYSX
-	if (ConstraintData && !(ConstraintData->getConstraintFlags()&PxConstraintFlag::eBROKEN))
+	if (PxD6Joint* Joint = GetUnbrokenJoint())
 	{
-	PxTransform PxRefFrame = U2PTransform(RefFrame);
-	ConstraintData->setLocalPose(PxFrame, PxRefFrame);
+		PxTransform PxRefFrame = U2PTransform(RefFrame);
+		Joint->setLocalPose(PxFrame, PxRefFrame);
 	}
 #endif
 
@@ -610,11 +627,11 @@ void FConstraintInstance::SetRefPosition(EConstraintFrame::Type Frame, const FVe
 	}
 
 #if WITH_PHYSX
-	if (ConstraintData && !(ConstraintData->getConstraintFlags()&PxConstraintFlag::eBROKEN))
+	if (PxD6Joint* Joint = GetUnbrokenJoint())
 	{
 		PxTransform PxRefFrame = ConstraintData->getLocalPose(PxFrame);
 		PxRefFrame.p = U2PVector(RefPosition);
-		ConstraintData->setLocalPose(PxFrame, PxRefFrame);
+		Joint->setLocalPose(PxFrame, PxRefFrame);
 	}
 #endif
 }
@@ -640,11 +657,11 @@ void FConstraintInstance::SetRefOrientation(EConstraintFrame::Type Frame, const
 	}
 	
 #if WITH_PHYSX
-	if (ConstraintData && !(ConstraintData->getConstraintFlags()&PxConstraintFlag::eBROKEN))
+	if (PxD6Joint* Joint = GetUnbrokenJoint())
 	{
 		FTransform URefTransform = FTransform(PriAxis1, SecAxis, PriAxis ^ SecAxis, RefPos);
 		PxTransform PxRefFrame = U2PTransform(URefTransform);
-		ConstraintData->setLocalPose(PxFrame, PxRefFrame);
+		Joint->setLocalPose(PxFrame, PxRefFrame);
 		}
 #endif
 }
@@ -688,11 +705,11 @@ FVector FConstraintInstance::GetConstraintLocation()
 void FConstraintInstance::GetConstraintForce(FVector& OutLinearForce, FVector& OutAngularForce)
 {
 #if WITH_PHYSX
-	if (ConstraintData && !(ConstraintData->getConstraintFlags()&PxConstraintFlag::eBROKEN))
+	if (PxD6Joint* Joint = GetUnbrokenJoint())
 	{
 		PxVec3 PxOutLinearForce;
 		PxVec3 PxOutAngularForce;
-		ConstraintData->getConstraint()->getForce(PxOutLinearForce, PxOutAngularForce);
+		Joint->getConstraint()->getForce(PxOutLinearForce, PxOutAngularForce);
 
 		OutLinearForce = P2UVector(PxOutLinearForce);
 		OutAngularForce = P2UVector(PxOutAngularForce);
@@ -715,9 +732,7 @@ void FConstraintInstance::GetConstraintForce(FVector& OutLinearForce, FVector& O
 void FConstraintInstance::SetLinearPositionDrive(bool bEnableXDrive, bool bEnableYDrive, bool bEnableZDrive)
 {
 #if WITH_PHYSX
-	PxD6Joint* Joint = (PxD6Joint*)ConstraintData;
-	
-	if (Joint && !(Joint->getConstraintFlags()&PxConstraintFlag::eBROKEN))
+	if (PxD6Joint* Joint = GetUnbrokenJoint())
 	{
 		// Get the current drives
 		PxD6JointDrive CurrentDriveX = Joint->getDrive(PxD6Drive::eX);
@@ -767,13 +782,11 @@ void FConstraintInstance::SetLinearVelocityDrive(bool bEnableXDrive, bool bEnabl
 	bLinearVelocityDrive = bEnableXDrive || bEnableYDrive || bEnableZDrive;
 }
 
-
 /** Function for turning angular position drive on and off. */
 void FConstraintInstance::SetAngularPositionDrive(bool bEnableSwingDrive, bool bEnableTwistDrive)
 {
 #if WITH_PHYSX
-	PxD6Joint* Joint = (PxD6Joint*)ConstraintData;
-	if (Joint &&  !(Joint->getConstraintFlags()&PxConstraintFlag::eBROKEN))
+	if (PxD6Joint* Joint = GetUnbrokenJoint())
 	{
 		// Get the current drives
 		PxD6JointDrive CurrentDriveSwing = Joint->getDrive(PxD6Drive::eSWING);
@@ -830,8 +843,7 @@ void FConstraintInstance::SetLinearPositionTarget(const FVector& InPosTarget)
 	}
 
 #if WITH_PHYSX
-	PxD6Joint* Joint = (PxD6Joint*)ConstraintData;
-	if (Joint)
+	if (PxD6Joint* Joint = ConstraintData)
 	{
 		PxVec3 Pos = U2PVector(InPosTarget);
 		Joint->setDrivePosition(PxTransform(Pos, Joint->getDrivePosition().q));
@@ -851,8 +863,7 @@ void FConstraintInstance::SetLinearVelocityTarget(const FVector& InVelTarget)
 	}
 
 #if WITH_PHYSX
-	PxD6Joint* Joint = (PxD6Joint*)ConstraintData;
-	if (Joint)
+	if (PxD6Joint* Joint = ConstraintData)
 	{
 		PxVec3 CurrentLinearVel, CurrentAngVel;
 		Joint->getDriveVelocity(CurrentLinearVel, CurrentAngVel);
@@ -870,26 +881,23 @@ void FConstraintInstance::SetLinearDriveParams(float InSpring, float InDamping,
 #if WITH_PHYSX
 	if (bLinearPositionDrive || bLinearVelocityDrive)
 	{
-		if (ConstraintData &&  !(ConstraintData->getConstraintFlags()&PxConstraintFlag::eBROKEN))
+		if (PxD6Joint* Joint = GetUnbrokenJoint())
 		{
 			// X-Axis linear drive
 			const float DriveSpringX = bLinearPositionDrive && bLinearXPositionDrive ? InSpring : 0.0f;
 			const float DriveDampingX = (bLinearVelocityDrive && LinearVelocityTarget.X != 0.f) ? InDamping : 0.0f;
 			const float LinearForceLimit = InForceLimit > 0.f ? InForceLimit : PX_MAX_F32;
-
-			ConstraintData->setDrive(PxD6Drive::eX, PxD6JointDrive(DriveSpringX, DriveDampingX, LinearForceLimit, bIsAccelerationDrive));
+			Joint->setDrive(PxD6Drive::eX, PxD6JointDrive(DriveSpringX, DriveDampingX, LinearForceLimit, bIsAccelerationDrive));
 
 			// Y-Axis linear drive
 			const float DriveSpringY = bLinearPositionDrive && bLinearYPositionDrive ? InSpring : 0.0f;
 			const float DriveDampingY = (bLinearVelocityDrive && LinearVelocityTarget.Y != 0.f) ? InDamping : 0.0f;
-
-			ConstraintData->setDrive(PxD6Drive::eY, PxD6JointDrive(DriveSpringY, DriveDampingY, LinearForceLimit, bIsAccelerationDrive));
+			Joint->setDrive(PxD6Drive::eY, PxD6JointDrive(DriveSpringY, DriveDampingY, LinearForceLimit, bIsAccelerationDrive));
 
 			// Z-Axis linear drive
 			const float DriveSpringZ = bLinearPositionDrive && bLinearZPositionDrive ? InSpring : 0.0f;
 			const float DriveDampingZ = (bLinearVelocityDrive && LinearVelocityTarget.Z != 0.f) ? InDamping : 0.0f;
-
-			ConstraintData->setDrive(PxD6Drive::eZ, PxD6JointDrive(DriveSpringZ, DriveDampingZ, LinearForceLimit, bIsAccelerationDrive));
+			Joint->setDrive(PxD6Drive::eZ, PxD6JointDrive(DriveSpringZ, DriveDampingZ, LinearForceLimit, bIsAccelerationDrive));
 		}
 	}
 #endif
@@ -911,8 +919,7 @@ void FConstraintInstance::SetAngularOrientationTarget(const FQuat& InPosTarget)
 	}
 	
 #if WITH_PHYSX
-	PxD6Joint* Joint = (PxD6Joint*)ConstraintData;
-	if (Joint)
+	if (PxD6Joint* Joint = ConstraintData)
 	{
 		PxQuat Quat = U2PQuat(InPosTarget);
 		Joint->setDrivePosition(PxTransform(Joint->getDrivePosition().p, Quat));
@@ -926,8 +933,7 @@ float FConstraintInstance::GetCurrentSwing1() const
 {
 	float Swing1 = 0.f;
 #if WITH_PHYSX
-	PxD6Joint* Joint = (PxD6Joint*)ConstraintData;
-	if (Joint)
+	if (PxD6Joint* Joint = ConstraintData)
 	{
 		Swing1 = Joint->getSwingZAngle();
 	}
@@ -940,8 +946,7 @@ float FConstraintInstance::GetCurrentSwing2() const
 {
 	float Swing2 = 0.f;
 #if WITH_PHYSX
-	PxD6Joint* Joint = (PxD6Joint*)ConstraintData;
-	if (Joint)
+	if (PxD6Joint* Joint = ConstraintData)
 	{
 		Swing2 = Joint->getSwingYAngle();
 	}
@@ -954,8 +959,7 @@ float FConstraintInstance::GetCurrentTwist() const
 {
 	float Twist = 0.f;
 #if WITH_PHYSX
-	PxD6Joint* Joint = (PxD6Joint*)ConstraintData;
-	if (Joint)
+	if (PxD6Joint* Joint = ConstraintData)
 	{
 		Twist = Joint->getTwist();
 	}
@@ -975,13 +979,12 @@ void FConstraintInstance::SetAngularVelocityTarget(const FVector& InVelTarget)
 	}
 
 #if WITH_PHYSX
-	PxD6Joint* Joint = (PxD6Joint*)ConstraintData;
-	if (Joint)
+	if (PxD6Joint* Joint = ConstraintData)
 	{
 		PxVec3 CurrentLinearVel, CurrentAngVel;
 		Joint->getDriveVelocity(CurrentLinearVel, CurrentAngVel);
 
-		PxVec3 AngVel = U2PVector(InVelTarget * 2 * (float)PI); // Convert from revs per second to radians
+		PxVec3 AngVel = U2PVector(RevolutionsToRads(InVelTarget));
 		Joint->setDriveVelocity(CurrentLinearVel, AngVel);
 	}
 #endif
@@ -995,7 +998,7 @@ void FConstraintInstance::SetAngularDriveParams(float InSpring, float InDamping,
 #if WITH_PHYSX
 	if (bAngularOrientationDrive || bAngularVelocityDrive)
 	{
-		if (ConstraintData &&  !(ConstraintData->getConstraintFlags()&PxConstraintFlag::eBROKEN))
+		if (PxD6Joint* Joint = GetUnbrokenJoint())
 		{
 			const float AngularForceLimit = InForceLimit > 0.0f ? InForceLimit : PX_MAX_F32;
 			const float DriveSpring = bAngularOrientationDrive ? InSpring : 0.0f;
@@ -1003,15 +1006,15 @@ void FConstraintInstance::SetAngularDriveParams(float InSpring, float InDamping,
 
 			if (AngularDriveMode == EAngularDriveMode::SLERP)
 			{
-				ConstraintData->setDrive(PxD6Drive::eTWIST, PxD6JointDrive());
-				ConstraintData->setDrive(PxD6Drive::eSWING, PxD6JointDrive());
-				ConstraintData->setDrive(PxD6Drive::eSLERP, PxD6JointDrive(DriveSpring, DriveDamping, AngularForceLimit, bIsAccelerationDrive));
+				Joint->setDrive(PxD6Drive::eTWIST, PxD6JointDrive());
+				Joint->setDrive(PxD6Drive::eSWING, PxD6JointDrive());
+				Joint->setDrive(PxD6Drive::eSLERP, PxD6JointDrive(DriveSpring, DriveDamping, AngularForceLimit, bIsAccelerationDrive));
 			}
 			else
 			{
-				ConstraintData->setDrive(PxD6Drive::eTWIST, PxD6JointDrive(DriveSpring, DriveDamping, AngularForceLimit, bIsAccelerationDrive));
-				ConstraintData->setDrive(PxD6Drive::eSWING, PxD6JointDrive(DriveSpring, DriveDamping, AngularForceLimit, bIsAccelerationDrive));
-				ConstraintData->setDrive(PxD6Drive::eSLERP, PxD6JointDrive());
+				Joint->setDrive(PxD6Drive::eTWIST, PxD6JointDrive(DriveSpring, DriveDamping, AngularForceLimit, bIsAccelerationDrive));
+				Joint->setDrive(PxD6Drive::eSWING, PxD6JointDrive(DriveSpring, DriveDamping, AngularForceLimit, bIsAccelerationDrive));
+				Joint->setDrive(PxD6Drive::eSLERP, PxD6JointDrive());
 			}
 		}
 	}
@@ -1026,8 +1029,7 @@ void FConstraintInstance::SetAngularDriveParams(float InSpring, float InDamping,
 void FConstraintInstance::SetAngularDOFLimitScale(float InSwing1LimitScale, float InSwing2LimitScale, float InTwistLimitScale)
 {
 #if WITH_PHYSX
-	PxD6Joint* Joint = (PxD6Joint*)ConstraintData;
-	if (Joint &&  !(Joint->getConstraintFlags()&PxConstraintFlag::eBROKEN))
+	if (PxD6Joint* Joint = GetUnbrokenJoint())
 	{
 		if ( AngularSwing1Motion == ACM_Limited || AngularSwing2Motion == ACM_Limited )
 		{
@@ -1081,8 +1083,7 @@ void FConstraintInstance::SetAngularDOFLimitScale(float InSwing1LimitScale, floa
 void FConstraintInstance::SetLinearLimitSize(float NewLimitSize)
 {
 #if WITH_PHYSX
-	PxD6Joint* Joint = (PxD6Joint*)ConstraintData;
-	if (Joint &&  !(Joint->getConstraintFlags()&PxConstraintFlag::eBROKEN))
+	if (PxD6Joint* Joint = GetUnbrokenJoint())
 	{
 		PxReal LimitContractDistance =  1.f * (PI/180.0f);
 		Joint->setLinearLimit(PxJointLinearLimit(GPhysXSDK->getTolerancesScale(), NewLimitSize, LimitContractDistance * GPhysXSDK->getTolerancesScale().length)); // LOC_MOD33 need to scale the contactDistance if not using its default value
