Fix non-unity build (missing include)
[INTEGRATE] Change 2455328

diff --git a/Engine/Source/Runtime/Engine/Private/Animation/AnimNode_StateMachine.cpp b/Engine/Source/Runtime/Engine/Private/Animation/AnimNode_StateMachine.cpp
--- a/Engine/Source/Runtime/Engine/Private/Animation/AnimNode_StateMachine.cpp
+++ b/Engine/Source/Runtime/Engine/Private/Animation/AnimNode_StateMachine.cpp
@@ -4,6 +4,7 @@
 #include "Animation/AnimNode_StateMachine.h"
 #include "Animation/AnimNode_TransitionResult.h"
 #include "Animation/AnimNode_TransitionPoseEvaluator.h"
+#include "Animation/AnimNode_SequencePlayer.h"
 #include "AnimationRuntime.h"
 #include "AnimTree.h"
 
