Mark objects in world being unloaded as pending kill so that references to them in cross-level elements are broken (this makes LoadMap consistent with seamless travel)

diff --git a/Engine/Source/Runtime/Engine/Private/UnrealEngine.cpp b/Engine/Source/Runtime/Engine/Private/UnrealEngine.cpp
--- a/Engine/Source/Runtime/Engine/Private/UnrealEngine.cpp
+++ b/Engine/Source/Runtime/Engine/Private/UnrealEngine.cpp
@@ -8765,6 +8765,29 @@ bool UEngine::LoadMap( FWorldContext& WorldContext, FURL URL, class UPendingNetG
 		}
 		WorldContext.World()->RemoveFromRoot();
 
+		// mark everything else contained in the world to be deleted
+		TArray<UWorld*> CurrentWorlds;
+		for (auto LevelIt(WorldContext.World()->GetLevelIterator()); LevelIt; ++LevelIt)
+		{
+			const ULevel* Level = *LevelIt;
+			if (Level)
+			{
+				CurrentWorlds.Add(CastChecked<UWorld>(Level->GetOuter()));
+			}
+		}
+
+		for (TObjectIterator<UObject> It; It; ++It)
+		{
+			for (const UWorld* World : CurrentWorlds)
+			{
+				if (It->IsIn(World))
+				{
+					It->MarkPendingKill();
+					break;
+				}
+			}
+		}
+
 		WorldContext.SetCurrentWorld(NULL);
 	}
 
