Ensure gamma correction is disabled on the game viewport since it is handled by the scene rendering. This is was only an issue when capturing matinee movies which requires rendering to a separate render target.

diff --git a/Engine/Source/Runtime/Engine/Private/GameEngine.cpp b/Engine/Source/Runtime/Engine/Private/GameEngine.cpp
--- a/Engine/Source/Runtime/Engine/Private/GameEngine.cpp
+++ b/Engine/Source/Runtime/Engine/Private/GameEngine.cpp
@@ -88,12 +88,15 @@ UGameEngine::UGameEngine(const FObjectInitializer& ObjectInitializer)
 
 void UGameEngine::CreateGameViewportWidget( UGameViewportClient* GameViewportClient )
 {
+	bool bRenderDirectlyToWindow = !GEngine->MatineeScreenshotOptions.bStartWithMatineeCapture && GIsDumpingMovie == 0; 
 	TSharedRef<SOverlay> ViewportOverlayWidgetRef = SNew( SOverlay );
 	TSharedRef<SViewport> GameViewportWidgetRef = 
 		SNew( SViewport )
 			// Render directly to the window backbuffer unless capturing a movie or getting screenshots
 			// @todo TEMP
-			.RenderDirectlyToWindow( !GEngine->MatineeScreenshotOptions.bStartWithMatineeCapture && GIsDumpingMovie == 0 )
+			.RenderDirectlyToWindow(bRenderDirectlyToWindow)
+			//gamma handled by the scene renderer
+			.EnableGammaCorrection(false)
 			.EnableStereoRendering(true)
 			[
 				SNew(SDPIScaler)
