#MATINEE: ANIMBP: Safer turning it off instead of stopping it everything.

Merging

//depot/UE4-GDC2015/Engine/Source/Runtime/Engine/...

to //depot/UE4/Engine/Source/Runtime/Engine/...

diff --git a/Engine/Source/Runtime/Engine/Classes/Animation/SkeletalMeshActor.h b/Engine/Source/Runtime/Engine/Classes/Animation/SkeletalMeshActor.h
--- a/Engine/Source/Runtime/Engine/Classes/Animation/SkeletalMeshActor.h
+++ b/Engine/Source/Runtime/Engine/Classes/Animation/SkeletalMeshActor.h
@@ -91,6 +91,9 @@ private:
 	// utility function to see if it can play animation or not
 	bool CanPlayAnimation(class UAnimSequenceBase* AnimAssetBase=NULL) const;
 
+	// currently actively playing montage
+	TWeakObjectPtr<class UAnimMontage> CurrentlyPlayingMontage;
+
 public:
 	/** Returns SkeletalMeshComponent subobject **/
 	class USkeletalMeshComponent* GetSkeletalMeshComponent();
diff --git a/Engine/Source/Runtime/Engine/Private/SkeletalMesh.cpp b/Engine/Source/Runtime/Engine/Private/SkeletalMesh.cpp
--- a/Engine/Source/Runtime/Engine/Private/SkeletalMesh.cpp
+++ b/Engine/Source/Runtime/Engine/Private/SkeletalMesh.cpp
@@ -3722,8 +3722,18 @@ void ASkeletalMeshActor::PreviewSetAnimPosition(FName SlotName, int32 ChannelInd
 			bool bShouldChange = AnimInst->IsPlayingSlotAnimation(InAnimSequence, SlotName ) == false; 
 			if(bShouldChange)
 			{
-				AnimInst->Montage_Stop(0.f);
+				if (CurrentlyPlayingMontage.IsValid())
+				{
+					// set it's weight to 0
+					struct FAnimMontageInstance* PrevAnimMontageInst = AnimInst->GetActiveInstanceForMontage(*CurrentlyPlayingMontage);
+					if (PrevAnimMontageInst)
+					{
+						PrevAnimMontageInst->Weight=0.f;
+					}
+				}
+
 				AnimInst->PlaySlotAnimation(InAnimSequence, SlotName);
+				CurrentlyPlayingMontage = AnimInst->GetCurrentActiveMontage();
 			}
 
 			struct FAnimMontageInstance* AnimMontageInst = AnimInst->GetActiveMontageInstance();
@@ -3888,7 +3898,7 @@ bool ASkeletalMeshActor::CanPlayAnimation(class UAnimSequenceBase* AnimAssetBase
 }
 
 // @todo ps4 clang bug: this works around a PS4/clang compiler bug (optimizations)
-void SetAnimPositionInner(FName SlotName, USkeletalMeshComponent* SkeletalMeshComponent, UAnimSequence *InAnimSequence, float InPosition, bool bLooping)
+void SetAnimPositionInner(FName SlotName, USkeletalMeshComponent* SkeletalMeshComponent, UAnimSequence *InAnimSequence, TWeakObjectPtr<UAnimMontage>& CurrentlyPlayingMontage, float InPosition, bool bLooping)
 {
 	UAnimInstance* AnimInst = SkeletalMeshComponent->GetAnimInstance();
 	if(AnimInst)
@@ -3917,7 +3927,18 @@ void SetAnimPositionInner(FName SlotName, USkeletalMeshComponent* SkeletalMeshCo
 			bool bShouldChange = AnimInst->IsPlayingSlotAnimation(InAnimSequence, SlotName) == false;
 			if(bShouldChange)
 			{
+				if(CurrentlyPlayingMontage.IsValid())
+				{
+					// set it's weight to 0
+					struct FAnimMontageInstance* PrevAnimMontageInst = AnimInst->GetActiveInstanceForMontage(*CurrentlyPlayingMontage);
+					if(PrevAnimMontageInst)
+					{
+						PrevAnimMontageInst->Weight=0.f;
+					}
+				}
+
 				AnimInst->PlaySlotAnimation(InAnimSequence, SlotName, 0.f, 0.f, 0.f);
+				CurrentlyPlayingMontage = AnimInst->GetCurrentActiveMontage();
 			}
 
 			struct FAnimMontageInstance* AnimMontageInst = AnimInst->GetActiveMontageInstance();
@@ -3941,7 +3962,7 @@ void ASkeletalMeshActor::SetAnimPosition(FName SlotName, int32 ChannelIndex, UAn
 {
 	if (CanPlayAnimation(InAnimSequence))
 	{
-		SetAnimPositionInner(SlotName, SkeletalMeshComponent, InAnimSequence, InPosition, bLooping);
+		SetAnimPositionInner(SlotName, SkeletalMeshComponent, InAnimSequence, CurrentlyPlayingMontage, InPosition, bLooping);
 	}
 }
 #endif
