#ANIM: MATINEE: Fix T-pose issue when using Anim BP for matinee

Merging

//depot/UE4-GDC2015/Engine/Source/Runtime/Engine/Private/...

to //depot/UE4/Engine/Source/Runtime/Engine/Private/...

diff --git a/Engine/Source/Runtime/Engine/Private/Animation/AnimationRuntime.cpp b/Engine/Source/Runtime/Engine/Private/Animation/AnimationRuntime.cpp
--- a/Engine/Source/Runtime/Engine/Private/Animation/AnimationRuntime.cpp
+++ b/Engine/Source/Runtime/Engine/Private/Animation/AnimationRuntime.cpp
@@ -401,6 +401,8 @@ void FAnimationRuntime::GetPoseFromAnimTrack(
 	TArray<float> SourceWeights;
 	float TotalWeight = 0.f;
 
+	float CurrentTime = FMath::Clamp(ExtractionContext.CurrentTime, 0.f, Track.GetLength());
+
 	// first get all the montage instance weight this slot node has
 	for (int32 I=0; I<Track.AnimSegments.Num(); ++I)
 	{
@@ -408,7 +410,7 @@ void FAnimationRuntime::GetPoseFromAnimTrack(
 
 		float PositionInAnim = 0.f;
 		float Weight = 0.f;
-		UAnimSequenceBase* AnimRef = AnimSegment.GetAnimationData(ExtractionContext.CurrentTime, PositionInAnim, Weight);
+		UAnimSequenceBase* AnimRef = AnimSegment.GetAnimationData(CurrentTime, PositionInAnim, Weight);
 		// make this to be 1 function
 		if( AnimRef && (Weight > ZERO_ANIMWEIGHT_THRESH) )
 		{
diff --git a/Engine/Source/Runtime/Engine/Private/SkeletalMesh.cpp b/Engine/Source/Runtime/Engine/Private/SkeletalMesh.cpp
--- a/Engine/Source/Runtime/Engine/Private/SkeletalMesh.cpp
+++ b/Engine/Source/Runtime/Engine/Private/SkeletalMesh.cpp
@@ -3732,16 +3732,18 @@ void ASkeletalMeshActor::PreviewSetAnimPosition(FName SlotName, int32 ChannelInd
 					}
 				}
 
-				AnimInst->PlaySlotAnimation(InAnimSequence, SlotName);
+				AnimInst->PlaySlotAnimation(InAnimSequence, SlotName, 0.0f, 0.0f, 0.f, 1);
 				CurrentlyPlayingMontage = AnimInst->GetCurrentActiveMontage();
 			}
 
-			struct FAnimMontageInstance* AnimMontageInst = AnimInst->GetActiveMontageInstance();
+			ensure(CurrentlyPlayingMontage.IsValid());
+
+			struct FAnimMontageInstance* AnimMontageInst = AnimInst->GetActiveInstanceForMontage(*CurrentlyPlayingMontage);
 			if (AnimMontageInst)
 			{
 				AnimMontageInst->Weight = 1.f;
 				//AnimInst->Montage_Set(bLooping);
-				AnimInst->Montage_SetPosition(NULL, InPosition);
+				AnimInst->Montage_SetPosition(CurrentlyPlayingMontage.Get(), InPosition);
 				AnimInst->UpdateAnimation(DeltaTime);
 			}
 			else
@@ -3937,16 +3939,17 @@ void SetAnimPositionInner(FName SlotName, USkeletalMeshComponent* SkeletalMeshCo
 					}
 				}
 
-				AnimInst->PlaySlotAnimation(InAnimSequence, SlotName, 0.f, 0.f, 0.f);
+				AnimInst->PlaySlotAnimation(InAnimSequence, SlotName, 0.0f, 0.0f, 0.f, 1);
 				CurrentlyPlayingMontage = AnimInst->GetCurrentActiveMontage();
 			}
 
-			struct FAnimMontageInstance* AnimMontageInst = AnimInst->GetActiveMontageInstance();
-			if (AnimMontageInst)
+			ensure(CurrentlyPlayingMontage.IsValid());
+
+			struct FAnimMontageInstance* AnimMontageInst = AnimInst->GetActiveInstanceForMontage(*CurrentlyPlayingMontage);
+			if(AnimMontageInst)
 			{
 				AnimMontageInst->Weight = 1.f;
-				AnimInst->Montage_SetPosition(NULL, InPosition);
-				AnimInst->UpdateAnimation(0.f);
+				AnimInst->Montage_SetPosition(CurrentlyPlayingMontage.Get(), InPosition);
 			}
 			else
 			{
