Include the currently executing timer when looking up handles for K2 delegates

diff --git a/Engine/Source/Runtime/Engine/Private/TimerManager.cpp b/Engine/Source/Runtime/Engine/Private/TimerManager.cpp
--- a/Engine/Source/Runtime/Engine/Private/TimerManager.cpp
+++ b/Engine/Source/Runtime/Engine/Private/TimerManager.cpp
@@ -141,6 +141,11 @@ int32 FTimerManager::FindTimerInList(const TArray<FTimerData> &SearchArray, FTim
 /** Finds a handle to a dynamic timer bound to a particular pointer and function name. */
 FTimerHandle FTimerManager::K2_FindDynamicTimerHandle(FTimerDynamicDelegate InDynamicDelegate) const
 {
+	if (CurrentlyExecutingTimer.TimerDelegate.FuncDynDelegate == InDynamicDelegate)
+	{
+		return CurrentlyExecutingTimer.TimerHandle;
+	}
+
 	if (auto* Timer = ActiveTimerHeap.FindByPredicate([=](const FTimerData& Data){ return Data.TimerDelegate.FuncDynDelegate == InDynamicDelegate; }))
 	{
 		return Timer->TimerHandle;
