Added a fallback font when calling AHUD::GetTextSize

This now mirrors the behaviour of AHUD::DrawText.

diff --git a/Engine/Source/Runtime/Engine/Private/HUD.cpp b/Engine/Source/Runtime/Engine/Private/HUD.cpp
--- a/Engine/Source/Runtime/Engine/Private/HUD.cpp
+++ b/Engine/Source/Runtime/Engine/Private/HUD.cpp
@@ -552,7 +552,7 @@ void AHUD::GetTextSize(const FString& Text, float& OutWidth, float& OutHeight, c
 {
 	if (IsCanvasValid_WarnIfNot())
 	{
-		Canvas->TextSize(Font, Text, OutWidth, OutHeight, Scale, Scale);
+		Canvas->TextSize(Font ? Font : GEngine->GetMediumFont(), Text, OutWidth, OutHeight, Scale, Scale);
 	}
 }
 
