Get rid of old PS4 compiler workaround.

diff --git a/Engine/Source/Runtime/Engine/Private/SkeletalMesh.cpp b/Engine/Source/Runtime/Engine/Private/SkeletalMesh.cpp
--- a/Engine/Source/Runtime/Engine/Private/SkeletalMesh.cpp
+++ b/Engine/Source/Runtime/Engine/Private/SkeletalMesh.cpp
@@ -3959,8 +3959,6 @@ void SetAnimPositionInner(FName SlotName, USkeletalMeshComponent* SkeletalMeshCo
 	}
 }
 
-// @todo ps4 clang bug: this works around a PS4/clang compiler bug (optimizations)
-#if !PLATFORM_PS4
 void ASkeletalMeshActor::SetAnimPosition(FName SlotName, int32 ChannelIndex, UAnimSequence* InAnimSequence, float InPosition, bool bFireNotifies, bool bLooping)
 {
 	if (CanPlayAnimation(InAnimSequence))
@@ -3968,7 +3966,6 @@ void ASkeletalMeshActor::SetAnimPosition(FName SlotName, int32 ChannelIndex, UAn
 		SetAnimPositionInner(SlotName, SkeletalMeshComponent, InAnimSequence, CurrentlyPlayingMontage, InPosition, bLooping);
 	}
 }
-#endif
 
 void ASkeletalMeshActor::FinishAnimControl(UInterpGroup* InInterpGroup)
 {
