Don't crash when a parent class doesn't have a SimpleConstructionScript

diff --git a/Engine/Source/Runtime/Engine/Private/BlueprintGeneratedClass.cpp b/Engine/Source/Runtime/Engine/Private/BlueprintGeneratedClass.cpp
--- a/Engine/Source/Runtime/Engine/Private/BlueprintGeneratedClass.cpp
+++ b/Engine/Source/Runtime/Engine/Private/BlueprintGeneratedClass.cpp
@@ -302,12 +302,17 @@ UObject* UBlueprintGeneratedClass::FindArchetype(UClass* ArchetypeClass, const F
 		while (Class)
 		{
 			USimpleConstructionScript* ClassSCS = Class->SimpleConstructionScript;
-			if (ClassSCS->HasAnyFlags(RF_NeedLoad))
+			USCS_Node* SCSNode = nullptr;
+			if (ClassSCS)
 			{
-				ClassSCS->PreloadChain();
+				if (ClassSCS->HasAnyFlags(RF_NeedLoad))
+				{
+					ClassSCS->PreloadChain();
+				}
+
+				SCSNode = ClassSCS->FindSCSNode(ArchetypeName);
 			}
 
-			USCS_Node* SCSNode = ClassSCS->FindSCSNode(ArchetypeName);
 			if (SCSNode)
 			{
 				Archetype = SCSNode->ComponentTemplate;
