Ensure that the constructed component hierarchy is torn down in attachment order to avoid issues with null root component

diff --git a/Engine/Source/Runtime/Engine/Private/ActorConstruction.cpp b/Engine/Source/Runtime/Engine/Private/ActorConstruction.cpp
--- a/Engine/Source/Runtime/Engine/Private/ActorConstruction.cpp
+++ b/Engine/Source/Runtime/Engine/Private/ActorConstruction.cpp
@@ -87,9 +87,23 @@ void AActor::DestroyConstructedComponents()
 	// Remove all existing components
 	TInlineComponentArray<UActorComponent*> PreviouslyAttachedComponents;
 	GetComponents(PreviouslyAttachedComponents);
-	for (int32 i = 0; i < PreviouslyAttachedComponents.Num(); i++)
+
+	// We need the hierarchy to be torn down in attachment order, so do a quick sort
+	PreviouslyAttachedComponents.Remove(nullptr);
+	PreviouslyAttachedComponents.Sort([](UActorComponent& A, UActorComponent& B)
+	{
+		if (USceneComponent* BSC = Cast<USceneComponent>(&B))
+		{
+			if (BSC->AttachParent == &A)
+			{
+				return false;
+			}
+		}
+		return true;
+	});
+
+	for (UActorComponent* Component : PreviouslyAttachedComponents)
 	{
-		UActorComponent* Component = PreviouslyAttachedComponents[i];
 		if (Component)
 		{
 			bool bDestroyComponent = false;
