Provide a default console bind for russian keyboard layouts

diff --git a/Engine/Source/Runtime/Engine/Private/UserInterface/InputSettings.cpp b/Engine/Source/Runtime/Engine/Private/UserInterface/InputSettings.cpp
--- a/Engine/Source/Runtime/Engine/Private/UserInterface/InputSettings.cpp
+++ b/Engine/Source/Runtime/Engine/Private/UserInterface/InputSettings.cpp
@@ -62,6 +62,7 @@ void UInputSettings::PostInitProperties()
 			break;
 			
 		case LANG_JAPANESE:
+		case LANG_RUSSIAN:
 			DefaultConsoleKey = FInputKeyManager::Get().GetKeyFromCodes(VK_OEM_3, 0);
 			break;
 		}
