Dev to Main integration from CL #2446290

diff --git a/Engine/Source/Runtime/AIModule/Private/Actions/PawnAction_Move.cpp b/Engine/Source/Runtime/AIModule/Private/Actions/PawnAction_Move.cpp
--- a/Engine/Source/Runtime/AIModule/Private/Actions/PawnAction_Move.cpp
+++ b/Engine/Source/Runtime/AIModule/Private/Actions/PawnAction_Move.cpp
@@ -9,15 +9,11 @@ UPawnAction_Move::UPawnAction_Move(const FObjectInitializer& ObjectInitializer)
 	, AcceptableRadius(30.f)
 	, bFinishOnOverlap(true)
 	, bUsePathfinding(true)
+	, bAllowPartialPath(true)
 	, bProjectGoalToNavigation(false)
 	, bUpdatePathToGoal(true)
 	, bAbortChildActionOnPathChange(false)
 {
-	if (HasAnyFlags(RF_ClassDefaultObject) == false)
-	{
-		PathObserver = FNavigationPath::FPathObserverDelegate::FDelegate::CreateUObject(this, &UPawnAction_Move::OnPathUpdated);
-	}
-
 	bShouldPauseMovement = true;
 }
 
@@ -75,19 +71,36 @@ bool UPawnAction_Move::Start()
 EPathFollowingRequestResult::Type UPawnAction_Move::RequestMove(AAIController& Controller)
 {
 	EPathFollowingRequestResult::Type RequestResult = EPathFollowingRequestResult::Failed;
+	
+	FAIMoveRequest MoveReq;
+	MoveReq.SetUsePathfinding(bUsePathfinding);
+	MoveReq.SetAllowPartialPath(bAllowPartialPath);
+	MoveReq.SetProjectGoalLocation(bProjectGoalToNavigation);
+	MoveReq.SetNavigationFilter(FilterClass);
+	MoveReq.SetAcceptanceRadius(AcceptableRadius);
+	MoveReq.SetStopOnOverlap(bFinishOnOverlap);
+	MoveReq.SetCanStrafe(bAllowStrafe);
 
 	if (GoalActor != NULL)
 	{
 		const bool bAtGoal = CheckAlreadyAtGoal(Controller, *GoalActor, AcceptableRadius);
-		RequestResult = bAtGoal ? EPathFollowingRequestResult::AlreadyAtGoal : 
-			bUpdatePathToGoal ? Controller.MoveToActor(GoalActor, AcceptableRadius, bFinishOnOverlap, bUsePathfinding, bAllowStrafe, FilterClass) :
-			Controller.MoveToLocation(GoalActor->GetActorLocation(), AcceptableRadius, bFinishOnOverlap, bUsePathfinding, bProjectGoalToNavigation, bAllowStrafe);
+		if (bUpdatePathToGoal)
+		{
+			MoveReq.SetGoalActor(GoalActor);
+		}
+		else
+		{
+			MoveReq.SetGoalLocation(GoalActor->GetActorLocation());
+		}
+
+		RequestResult = bAtGoal ? EPathFollowingRequestResult::AlreadyAtGoal : Controller.MoveTo(MoveReq);
 	}
 	else if (FAISystem::IsValidLocation(GoalLocation))
 	{
 		const bool bAtGoal = CheckAlreadyAtGoal(Controller, GoalLocation, AcceptableRadius);
-		RequestResult = bAtGoal ? EPathFollowingRequestResult::AlreadyAtGoal :
-			Controller.MoveToLocation(GoalLocation, AcceptableRadius, bFinishOnOverlap, bUsePathfinding, bProjectGoalToNavigation, bAllowStrafe, FilterClass);
+		MoveReq.SetGoalLocation(GoalLocation);
+		
+		RequestResult = bAtGoal ? EPathFollowingRequestResult::AlreadyAtGoal : Controller.MoveTo(MoveReq);
 	}
 	else
 	{
@@ -239,7 +252,7 @@ void UPawnAction_Move::SetPath(FNavPathSharedRef InPath)
 	{
 		ClearPath();
 		Path = InPath;
-		PathObserverDelegateHandle = Path->AddObserver(PathObserver);
+		PathObserverDelegateHandle = Path->AddObserver(FNavigationPath::FPathObserverDelegate::FDelegate::CreateUObject(this, &UPawnAction_Move::OnPathUpdated));
 
 		// skip auto updates, it will be handled manually to include controller's ShouldPostponePathUpdates()
 		Path->EnableRecalculationOnInvalidation(false);
@@ -267,6 +280,15 @@ void UPawnAction_Move::OnPathUpdated(FNavigationPath* UpdatedPath, ENavPathEvent
 	{
 		TryToRepath();
 	}
+
+	// log new path when action is paused, otherwise it will be logged by path following component's update
+	if (Event == ENavPathEvent::UpdatedDueToGoalMoved || Event == ENavPathEvent::UpdatedDueToNavigationChanged)
+	{
+		if (IsPaused() && UpdatedPath)
+		{
+			UPathFollowingComponent::LogPathHelper(MyOwner, UpdatedPath, UpdatedPath->GetGoalActor());
+		}
+	}
 }
 
 void UPawnAction_Move::TryToRepath()
diff --git a/Engine/Source/Runtime/AIModule/Private/Actions/PawnActionsComponent.cpp b/Engine/Source/Runtime/AIModule/Private/Actions/PawnActionsComponent.cpp
--- a/Engine/Source/Runtime/AIModule/Private/Actions/PawnActionsComponent.cpp
+++ b/Engine/Source/Runtime/AIModule/Private/Actions/PawnActionsComponent.cpp
@@ -191,8 +191,13 @@ void UPawnActionsComponent::TickComponent(float DeltaTime, enum ELevelTick TickT
 			switch (Event.EventType)
 			{
 			case EPawnActionEventType::InstantAbort:
-				Event.Action->Abort(EAIForceParam::Force);
-				ActionStacks[Event.Priority].PopAction(*Event.Action);
+				// can result in adding new ActionEvents (from child actions) and reallocating data in ActionEvents array
+				// because of it, we need to operate on copy instead of reference to memory address
+				{
+					FPawnActionEvent EventCopy(Event);
+					EventCopy.Action->Abort(EAIForceParam::Force);
+					ActionStacks[EventCopy.Priority].PopAction(*EventCopy.Action);
+				}
 				break;
 			case EPawnActionEventType::FinishedAborting:
 			case EPawnActionEventType::FinishedExecution:
diff --git a/Engine/Source/Runtime/AIModule/Private/BehaviorTree/BTCompositeNode.cpp b/Engine/Source/Runtime/AIModule/Private/BehaviorTree/BTCompositeNode.cpp
--- a/Engine/Source/Runtime/AIModule/Private/BehaviorTree/BTCompositeNode.cpp
+++ b/Engine/Source/Runtime/AIModule/Private/BehaviorTree/BTCompositeNode.cpp
@@ -467,8 +467,8 @@ int32 UBTCompositeNode::GetMatchingChildIndex(int32 ActiveInstanceIdx, FBTNodeIn
 		// find child outside range
 		for (int32 ChildIndex = 0; ChildIndex < Children.Num(); ChildIndex++)
 		{
-			const UBTNode* ChildNode = GetChildNode(ChildIndex);
-			if (ChildNode != nullptr && ChildNode->GetExecutionIndex() > NodeIdx.ExecutionIndex)
+			const uint16 FirstIndexInBranch = GetChildExecutionIndex(ChildIndex, EBTChildIndex::FirstNode);
+			if (FirstIndexInBranch > NodeIdx.ExecutionIndex)
 			{
 				return ChildIndex ? (ChildIndex - 1) : 0;
 			}
diff --git a/Engine/Source/Runtime/AIModule/Private/BehaviorTree/BehaviorTreeComponent.cpp b/Engine/Source/Runtime/AIModule/Private/BehaviorTree/BehaviorTreeComponent.cpp
--- a/Engine/Source/Runtime/AIModule/Private/BehaviorTree/BehaviorTreeComponent.cpp
+++ b/Engine/Source/Runtime/AIModule/Private/BehaviorTree/BehaviorTreeComponent.cpp
@@ -1058,10 +1058,8 @@ void UBehaviorTreeComponent::ProcessExecutionRequest()
 		if (!ExecutionRequest.bTryNextChild)
 		{
 			// mark all decorators less important than current search start node for removal
-			// (all from first node on branch with search start)
-			const int32 StartNodeIdx = ExecutionRequest.ExecuteNode->GetBranchExecutionIndex(ExecutionRequest.SearchStart.ExecutionIndex) - 1;
-			const FBTNodeIndex DeactivateIdx(ExecutionRequest.ExecuteInstanceIdx, FMath::Max(StartNodeIdx, 0));
-			UnregisterAuxNodesUpTo(DeactivateIdx);
+			const FBTNodeIndex DeactivateIdx(ExecutionRequest.SearchStart.InstanceIndex, ExecutionRequest.SearchStart.ExecutionIndex - 1);
+			UnregisterAuxNodesUpTo(ExecutionRequest.SearchStart.ExecutionIndex ? DeactivateIdx : ExecutionRequest.SearchStart);
 
 			// reactivate top search node, so it could use search range correctly
 			BT_SEARCHLOG(SearchData, Verbose, TEXT("Reactivate node: %s [restart]"), *UBehaviorTreeTypes::DescribeNodeHelper(TestNode));
diff --git a/Engine/Source/Runtime/AIModule/Private/BehaviorTree/Blackboard/BlackboardKeyType_Object.cpp b/Engine/Source/Runtime/AIModule/Private/BehaviorTree/Blackboard/BlackboardKeyType_Object.cpp
--- a/Engine/Source/Runtime/AIModule/Private/BehaviorTree/Blackboard/BlackboardKeyType_Object.cpp
+++ b/Engine/Source/Runtime/AIModule/Private/BehaviorTree/Blackboard/BlackboardKeyType_Object.cpp
@@ -14,6 +14,15 @@ UBlackboardKeyType_Object::UBlackboardKeyType_Object(const FObjectInitializer& O
 
 UObject* UBlackboardKeyType_Object::GetValue(const UBlackboardKeyType_Object* KeyOb, const uint8* RawData)
 {
+	// RawData can be NULL if the key is corrupted by bad data, such as if someone has duplicated a key in a derived
+	// blackboard.  We must handle that bad data case gracefully.  It's likely that we need to handle this for
+	// all blackboard key types, so possibly GetValueFromMemory should handle the NULL case instead.  But for now I'm
+	// just fixing the case I came across.
+	if (RawData == NULL)
+	{
+		return NULL;
+	}
+
 	FWeakObjectPtr WeakObjPtr = GetValueFromMemory<FWeakObjectPtr>(RawData);
 	return WeakObjPtr.Get();
 }
@@ -75,6 +84,11 @@ bool UBlackboardKeyType_Object::GetRotation(const UBlackboardComponent& OwnerCom
 
 bool UBlackboardKeyType_Object::TestBasicOperation(const UBlackboardComponent& OwnerComp, const uint8* MemoryBlock, EBasicKeyOperation::Type Op) const
 {
+	if (MemoryBlock == NULL)
+	{
+		return false;
+	}
+
 	FWeakObjectPtr WeakObjPtr = GetValueFromMemory<FWeakObjectPtr>(MemoryBlock);
 	return (Op == EBasicKeyOperation::Set) ? WeakObjPtr.IsValid() : !WeakObjPtr.IsValid();
 }
diff --git a/Engine/Source/Runtime/AIModule/Private/BehaviorTree/Decorators/BTDecorator_Blackboard.cpp b/Engine/Source/Runtime/AIModule/Private/BehaviorTree/Decorators/BTDecorator_Blackboard.cpp
--- a/Engine/Source/Runtime/AIModule/Private/BehaviorTree/Decorators/BTDecorator_Blackboard.cpp
+++ b/Engine/Source/Runtime/AIModule/Private/BehaviorTree/Decorators/BTDecorator_Blackboard.cpp
@@ -28,9 +28,12 @@ bool UBTDecorator_Blackboard::EvaluateOnBlackboard(const UBlackboardComponent& B
 		UBlackboardKeyType* KeyCDO = BlackboardKey.SelectedKeyType->GetDefaultObject<UBlackboardKeyType>();
 		const uint8* KeyMemory = BlackboardComp.GetKeyRawData(BlackboardKey.GetSelectedKeyID());
 
-		const EBlackboardKeyOperation::Type Op = KeyCDO->GetTestOperation();
-		switch (Op)
+		// KeyMemory can be NULL if the blackboard has its data setup wrong, so we must conditionally handle that case.
+		if (ensure(KeyCDO != NULL) && (KeyMemory != NULL))
 		{
+			const EBlackboardKeyOperation::Type Op = KeyCDO->GetTestOperation();
+			switch (Op)
+			{
 			case EBlackboardKeyOperation::Basic:
 				bResult = KeyCDO->WrappedTestBasicOperation(BlackboardComp, KeyMemory, (EBasicKeyOperation::Type)OperationType);
 				break;
@@ -45,6 +48,7 @@ bool UBTDecorator_Blackboard::EvaluateOnBlackboard(const UBlackboardComponent& B
 
 			default:
 				break;
+			}
 		}
 	}
 
diff --git a/Engine/Source/Runtime/AIModule/Private/BehaviorTree/Tasks/BTTask_MoveTo.cpp b/Engine/Source/Runtime/AIModule/Private/BehaviorTree/Tasks/BTTask_MoveTo.cpp
--- a/Engine/Source/Runtime/AIModule/Private/BehaviorTree/Tasks/BTTask_MoveTo.cpp
+++ b/Engine/Source/Runtime/AIModule/Private/BehaviorTree/Tasks/BTTask_MoveTo.cpp
@@ -11,6 +11,7 @@ UBTTask_MoveTo::UBTTask_MoveTo(const FObjectInitializer& ObjectInitializer)
 	: Super(ObjectInitializer)
 	, AcceptableRadius(50.f)
 	, bAllowStrafe(false)
+	, bAllowPartialPath(true)
 	, bStopOnOverlap(true)
 {
 	NodeName = "Move To";
@@ -52,6 +53,13 @@ EBTNodeResult::Type UBTTask_MoveTo::PerformMoveTask(UBehaviorTreeComponent& Owne
 	if (MyController && MyBlackboard)
 	{
 		EPathFollowingRequestResult::Type RequestResult = EPathFollowingRequestResult::Failed;
+		
+		FAIMoveRequest MoveReq;
+		MoveReq.SetNavigationFilter(FilterClass);
+		MoveReq.SetAllowPartialPath(bAllowPartialPath);
+		MoveReq.SetAcceptanceRadius(AcceptableRadius);
+		MoveReq.SetCanStrafe(bAllowStrafe);
+		MoveReq.SetStopOnOverlap(bStopOnOverlap);
 
 		if (BlackboardKey.SelectedKeyType == UBlackboardKeyType_Object::StaticClass())
 		{
@@ -59,8 +67,9 @@ EBTNodeResult::Type UBTTask_MoveTo::PerformMoveTask(UBehaviorTreeComponent& Owne
 			AActor* TargetActor = Cast<AActor>(KeyValue);
 			if (TargetActor)
 			{
-				RequestResult = MyController->MoveToActor(TargetActor, AcceptableRadius, bStopOnOverlap, /*bUsePathfinding=*/true
-					, bAllowStrafe, FilterClass);
+				MoveReq.SetGoalActor(TargetActor);
+
+				RequestResult = MyController->MoveTo(MoveReq);
 			}
 			else
 			{
@@ -70,8 +79,9 @@ EBTNodeResult::Type UBTTask_MoveTo::PerformMoveTask(UBehaviorTreeComponent& Owne
 		else if (BlackboardKey.SelectedKeyType == UBlackboardKeyType_Vector::StaticClass())
 		{
 			const FVector TargetLocation = MyBlackboard->GetValue<UBlackboardKeyType_Vector>(BlackboardKey.GetSelectedKeyID());
-			RequestResult = MyController->MoveToLocation(TargetLocation, AcceptableRadius, bStopOnOverlap, /*bUsePathfinding=*/true
-				, /*bProjectDestinationToNavigation*/false, bAllowStrafe, FilterClass);
+			MoveReq.SetGoalLocation(TargetLocation);
+
+			RequestResult = MyController->MoveTo(MoveReq);
 		}
 
 		if (RequestResult == EPathFollowingRequestResult::RequestSuccessful)
diff --git a/Engine/Source/Runtime/AIModule/Private/BehaviorTree/Tasks/BTTask_RunBehavior.cpp b/Engine/Source/Runtime/AIModule/Private/BehaviorTree/Tasks/BTTask_RunBehavior.cpp
--- a/Engine/Source/Runtime/AIModule/Private/BehaviorTree/Tasks/BTTask_RunBehavior.cpp
+++ b/Engine/Source/Runtime/AIModule/Private/BehaviorTree/Tasks/BTTask_RunBehavior.cpp
@@ -11,7 +11,7 @@ UBTTask_RunBehavior::UBTTask_RunBehavior(const FObjectInitializer& ObjectInitial
 
 EBTNodeResult::Type UBTTask_RunBehavior::ExecuteTask(UBehaviorTreeComponent& OwnerComp, uint8* NodeMemory)
 {
-	const bool bPushed = OwnerComp.PushInstance(*BehaviorAsset);
+	const bool bPushed = BehaviorAsset != nullptr && OwnerComp.PushInstance(*BehaviorAsset);
 	return bPushed ? EBTNodeResult::InProgress : EBTNodeResult::Failed;
 }
 
diff --git a/Engine/Source/Runtime/AIModule/Private/Blueprint/AIBlueprintHelperLibrary.cpp b/Engine/Source/Runtime/AIModule/Private/Blueprint/AIBlueprintHelperLibrary.cpp
--- a/Engine/Source/Runtime/AIModule/Private/Blueprint/AIBlueprintHelperLibrary.cpp
+++ b/Engine/Source/Runtime/AIModule/Private/Blueprint/AIBlueprintHelperLibrary.cpp
@@ -92,12 +92,25 @@ UAIAsyncTaskBlueprintProxy* UAIBlueprintHelperLibrary::CreateMoveToProxyObject(U
 		UWorld* World = GEngine->GetWorldFromContextObject( WorldContextObject );
 		MyObj = NewObject<UAIAsyncTaskBlueprintProxy>(World);
 
+		FAIMoveRequest MoveReq;
+		MoveReq.SetUsePathfinding(true);
+		MoveReq.SetAcceptanceRadius(AcceptanceRadius);
+		MoveReq.SetStopOnOverlap(bStopOnOverlap);
+		if (TargetActor)
+		{
+			MoveReq.SetGoalActor(TargetActor);
+		}
+		else
+		{
+			MoveReq.SetGoalLocation(Destination);
+		}
+		
 		FPathFindingQuery Query;
-		const bool bValidQuery = AIController->PreparePathfinding(Query, Destination, TargetActor, true);
+		const bool bValidQuery = AIController->PreparePathfinding(MoveReq, Query);
 
 		if (bValidQuery)
 		{
-			const FAIRequestID RequestID = AIController->RequestPathAndMove(Query, TargetActor, AcceptanceRadius, bStopOnOverlap, NULL);
+			const FAIRequestID RequestID = AIController->RequestPathAndMove(MoveReq, Query);
 			if (RequestID.IsValid())
 			{
 				MyObj->AIController = AIController;
diff --git a/Engine/Source/Runtime/AIModule/Private/EnvironmentQuery/Tests/EnvQueryTest_Pathfinding.cpp b/Engine/Source/Runtime/AIModule/Private/EnvironmentQuery/Tests/EnvQueryTest_Pathfinding.cpp
--- a/Engine/Source/Runtime/AIModule/Private/EnvironmentQuery/Tests/EnvQueryTest_Pathfinding.cpp
+++ b/Engine/Source/Runtime/AIModule/Private/EnvironmentQuery/Tests/EnvQueryTest_Pathfinding.cpp
@@ -195,37 +195,37 @@ void UEnvQueryTest_Pathfinding::PostLoad()
 
 bool UEnvQueryTest_Pathfinding::TestPathFrom(const FVector& ItemPos, const FVector& ContextPos, EPathFindingMode::Type Mode, const ANavigationData* NavData, UNavigationSystem* NavSys, const UObject* PathOwner) const
 {
-	const bool bPathExists = NavSys->TestPathSync(FPathFindingQuery(PathOwner, NavData, ItemPos, ContextPos), Mode);
+	const bool bPathExists = NavSys->TestPathSync(FPathFindingQuery(PathOwner, NavData, ItemPos, ContextPos).SetAllowPartialPaths(false), Mode);
 	return bPathExists;
 }
 
 bool UEnvQueryTest_Pathfinding::TestPathTo(const FVector& ItemPos, const FVector& ContextPos, EPathFindingMode::Type Mode, const ANavigationData* NavData, UNavigationSystem* NavSys, const UObject* PathOwner) const
 {
-	const bool bPathExists = NavSys->TestPathSync(FPathFindingQuery(PathOwner, NavData, ContextPos, ItemPos), Mode);
+	const bool bPathExists = NavSys->TestPathSync(FPathFindingQuery(PathOwner, NavData, ContextPos, ItemPos).SetAllowPartialPaths(false), Mode);
 	return bPathExists;
 }
 
 float UEnvQueryTest_Pathfinding::FindPathCostFrom(const FVector& ItemPos, const FVector& ContextPos, EPathFindingMode::Type Mode, const ANavigationData* NavData, UNavigationSystem* NavSys, const UObject* PathOwner) const
 {
-	FPathFindingResult Result = NavSys->FindPathSync(FPathFindingQuery(PathOwner, NavData, ItemPos, ContextPos), Mode);
+	FPathFindingResult Result = NavSys->FindPathSync(FPathFindingQuery(PathOwner, NavData, ItemPos, ContextPos).SetAllowPartialPaths(false), Mode);
 	return (Result.IsSuccessful()) ? Result.Path->GetCost() : BIG_NUMBER;
 }
 
 float UEnvQueryTest_Pathfinding::FindPathCostTo(const FVector& ItemPos, const FVector& ContextPos, EPathFindingMode::Type Mode, const ANavigationData* NavData, UNavigationSystem* NavSys, const UObject* PathOwner) const
 {
-	FPathFindingResult Result = NavSys->FindPathSync(FPathFindingQuery(PathOwner, NavData, ContextPos, ItemPos), Mode);
+	FPathFindingResult Result = NavSys->FindPathSync(FPathFindingQuery(PathOwner, NavData, ContextPos, ItemPos).SetAllowPartialPaths(false), Mode);
 	return (Result.IsSuccessful()) ? Result.Path->GetCost() : BIG_NUMBER;
 }
 
 float UEnvQueryTest_Pathfinding::FindPathLengthFrom(const FVector& ItemPos, const FVector& ContextPos, EPathFindingMode::Type Mode, const ANavigationData* NavData, UNavigationSystem* NavSys, const UObject* PathOwner) const
 {
-	FPathFindingResult Result = NavSys->FindPathSync(FPathFindingQuery(PathOwner, NavData, ItemPos, ContextPos), Mode);
+	FPathFindingResult Result = NavSys->FindPathSync(FPathFindingQuery(PathOwner, NavData, ItemPos, ContextPos).SetAllowPartialPaths(false), Mode);
 	return (Result.IsSuccessful()) ? Result.Path->GetLength() : BIG_NUMBER;
 }
 
 float UEnvQueryTest_Pathfinding::FindPathLengthTo(const FVector& ItemPos, const FVector& ContextPos, EPathFindingMode::Type Mode, const ANavigationData* NavData, UNavigationSystem* NavSys, const UObject* PathOwner) const
 {
-	FPathFindingResult Result = NavSys->FindPathSync(FPathFindingQuery(PathOwner, NavData, ContextPos, ItemPos), Mode);
+	FPathFindingResult Result = NavSys->FindPathSync(FPathFindingQuery(PathOwner, NavData, ContextPos, ItemPos).SetAllowPartialPaths(false), Mode);
 	return (Result.IsSuccessful()) ? Result.Path->GetLength() : BIG_NUMBER;
 }
 
diff --git a/Engine/Source/Runtime/AIModule/Private/Navigation/CrowdManager.cpp b/Engine/Source/Runtime/AIModule/Private/Navigation/CrowdManager.cpp
--- a/Engine/Source/Runtime/AIModule/Private/Navigation/CrowdManager.cpp
+++ b/Engine/Source/Runtime/AIModule/Private/Navigation/CrowdManager.cpp
@@ -31,8 +31,6 @@ DECLARE_CYCLE_STAT(TEXT("Step: movement"), STAT_AI_Crowd_StepMovementTime, STATG
 DECLARE_CYCLE_STAT(TEXT("Agent Update Time"), STAT_AI_Crowd_AgentUpdateTime, STATGROUP_AICrowd);
 DECLARE_DWORD_COUNTER_STAT(TEXT("Num Agents"), STAT_AI_Crowd_NumAgents, STATGROUP_AICrowd);
 
-DEFINE_LOG_CATEGORY_STATIC(LogEngineCrowdFollowing, Warning, All)
-
 namespace CrowdDebugDrawing
 {
 	/** if set, debug information will be displayed for agent selected in editor */
@@ -1130,7 +1128,7 @@ void UCrowdManager::DebugTick() const
 					for (int32 Idx = 0; Idx < CrowdAgent->ncorners; Idx++)
 					{
 						FVector P1 = Recast2UnrealPoint(&CrowdAgent->cornerVerts[Idx * 3]);
-						UE_VLOG_SEGMENT(LogOwner, LogEngineCrowdFollowing, Log, P0 + CrowdDebugDrawing::Offset, P1 + CrowdDebugDrawing::Offset, CrowdDebugDrawing::Corner, TEXT(""));
+						UE_VLOG_SEGMENT(LogOwner, LogCrowdFollowing, Log, P0 + CrowdDebugDrawing::Offset, P1 + CrowdDebugDrawing::Offset, CrowdDebugDrawing::Corner, TEXT(""));
 						P0 = P1;
 					}
 				}
@@ -1138,7 +1136,17 @@ void UCrowdManager::DebugTick() const
 				if (CrowdAgent->ncorners && (CrowdAgent->cornerFlags[CrowdAgent->ncorners - 1] & DT_STRAIGHTPATH_OFFMESH_CONNECTION))
 				{
 					FVector P0 = Recast2UnrealPoint(&CrowdAgent->cornerVerts[(CrowdAgent->ncorners - 1) * 3]);
-					UE_VLOG_SEGMENT(LogOwner, LogEngineCrowdFollowing, Log, P0, P0 + CrowdDebugDrawing::Offset * 2.0f, CrowdDebugDrawing::CornerLink, TEXT(""));
+					UE_VLOG_SEGMENT(LogOwner, LogCrowdFollowing, Log, P0, P0 + CrowdDebugDrawing::Offset * 2.0f, CrowdDebugDrawing::CornerLink, TEXT(""));
+				}
+
+				for (int32 Idx = 0; Idx < CrowdAgent->boundary.getSegmentCount(); Idx++)
+				{
+					const float* s = CrowdAgent->boundary.getSegment(Idx);
+					FColor Color = (dtTriArea2D(CrowdAgent->npos, s, s + 3) < 0.0f) ? CrowdDebugDrawing::CollisionSeg1 : CrowdDebugDrawing::CollisionSeg0;
+					FVector Pt0 = Recast2UnrealPoint(s);
+					FVector Pt1 = Recast2UnrealPoint(s + 3);
+
+					UE_VLOG_SEGMENT(LogOwner, LogCrowdFollowing, Log, Pt0 + CrowdDebugDrawing::Offset, Pt1 + CrowdDebugDrawing::Offset, Color, TEXT(""));
 				}
 			}
 		}
diff --git a/Engine/Source/Runtime/AIModule/Private/Navigation/PathFollowingComponent.cpp b/Engine/Source/Runtime/AIModule/Private/Navigation/PathFollowingComponent.cpp
--- a/Engine/Source/Runtime/AIModule/Private/Navigation/PathFollowingComponent.cpp
+++ b/Engine/Source/Runtime/AIModule/Private/Navigation/PathFollowingComponent.cpp
@@ -51,20 +51,20 @@ UPathFollowingComponent::UPathFollowingComponent(const FObjectInitializer& Objec
 	Status = EPathFollowingStatus::Idle;
 }
 
-void LogPathHelper(AActor* LogOwner, FNavPathSharedPtr Path, const AActor* GoalActor, const int32 CurrentPathPointGoal)
+void UPathFollowingComponent::LogPathHelper(const AActor* LogOwner, FNavigationPath* LogPath, const AActor* LogGoalActor)
 {
 #if ENABLE_VISUAL_LOG
 	FVisualLogger& Vlog = FVisualLogger::Get();
-	if (Vlog.IsRecording() && 
-		Path.IsValid() && Path->IsValid() && Path->GetPathPoints().Num())
+	if (Vlog.IsRecording() &&
+		LogPath && LogPath->IsValid() && LogPath->GetPathPoints().Num())
 	{
 		FVisualLogEntry* Entry = Vlog.GetEntryToWrite(LogOwner, LogOwner->GetWorld()->TimeSeconds);
-		Path->DescribeSelfToVisLog(Entry);
+		LogPath->DescribeSelfToVisLog(Entry);
 
-		const FVector PathEnd = Path->GetPathPoints().Last().Location;
-		if (GoalActor)
+		const FVector PathEnd = LogPath->GetPathPoints().Last().Location;
+		if (LogGoalActor)
 		{
-			const FVector GoalLoc = GoalActor->GetActorLocation();
+			const FVector GoalLoc = LogGoalActor->GetActorLocation();
 			if (FVector::DistSquared(GoalLoc, PathEnd) > 1.0f)
 			{
 				UE_VLOG_LOCATION(LogOwner, LogPathFollowing, Verbose, GoalLoc, 30, FColor::Green, TEXT("GoalActor"));
@@ -77,6 +77,16 @@ void LogPathHelper(AActor* LogOwner, FNavPathSharedPtr Path, const AActor* GoalA
 #endif // ENABLE_VISUAL_LOG
 }
 
+void UPathFollowingComponent::LogPathHelper(const AActor* LogOwner, FNavPathSharedPtr LogPath, const AActor* LogGoalActor)
+{
+#if ENABLE_VISUAL_LOG
+	if (LogPath.IsValid())
+	{
+		LogPathHelper(LogOwner, LogPath.Get(), LogGoalActor);
+	}
+#endif // ENABLE_VISUAL_LOG
+}
+
 void LogBlockHelper(AActor* LogOwner, UNavMovementComponent* MoveComp, float RadiusPct, float HeightPct, const FVector& SegmentStart, const FVector& SegmentEnd)
 {
 #if ENABLE_VISUAL_LOG
@@ -135,7 +145,7 @@ FAIRequestID UPathFollowingComponent::RequestMove(FNavPathSharedPtr InPath, FReq
 		*GetNameSafe(InDestinationActor),
 		!InGameData.IsValid() ? TEXT("missing") : TEXT("valid"));
 
-	LogPathHelper(GetOwner(), InPath, InDestinationActor, GetNextPathIndex());
+	LogPathHelper(GetOwner(), InPath, InDestinationActor);
 
 	if (InAcceptanceRadius == UPathFollowingComponent::DefaultAcceptanceRadius)
 	{
@@ -190,7 +200,7 @@ FAIRequestID UPathFollowingComponent::RequestMove(FNavPathSharedPtr InPath, FReq
 			Path->SetSourceActor(*(MovementComp->GetOwner()));
 		}
 
-		PathTimeWhenPaused = 0;
+		PathTimeWhenPaused = 0.0f;
 		OnPathUpdated();
 
 		AcceptanceRadius = InAcceptanceRadius;
@@ -231,7 +241,7 @@ bool UPathFollowingComponent::UpdateMove(FNavPathSharedPtr InPath, FAIRequestID
 		*GetPathDescHelper(InPath),
 		*GetStatusDesc(), RequestID);
 
-	LogPathHelper(GetOwner(), InPath, DestinationActor.Get(), GetNextPathIndex());
+	LogPathHelper(GetOwner(), InPath, DestinationActor.Get());
 
 	if (!InPath.IsValid() || !InPath->IsValid() || Status == EPathFollowingStatus::Idle 
 		|| RequestID.IsEquivalent(GetCurrentRequestId()) == false)
@@ -308,7 +318,7 @@ void UPathFollowingComponent::PauseMove(FAIRequestID RequestID, bool bResetVeloc
 		}
 
 		LocationWhenPaused = MovementComp ? MovementComp->GetActorFeetLocation() : FVector::ZeroVector;
-		PathTimeWhenPaused = GetWorld()->GetTimeSeconds();
+		PathTimeWhenPaused = Path.IsValid() ? Path->GetTimeStamp() : 0.0f;
 		Status = EPathFollowingStatus::Paused;
 
 		UpdateMoveFocus();
@@ -329,7 +339,7 @@ void UPathFollowingComponent::ResumeMove(FAIRequestID RequestID)
 		{
 			Status = EPathFollowingStatus::Moving;
 
-			const bool bWasPathUpdatedRecently = Path.IsValid() ? (GetWorld()->GetTimeSeconds() > Path->GetTimeStamp()) : false;
+			const bool bWasPathUpdatedRecently = Path.IsValid() ? (Path->GetTimeStamp() > PathTimeWhenPaused) : false;
 			if (bMovedDuringPause || bWasPathUpdatedRecently)
 			{
 				const int32 CurrentSegment = DetermineStartingPathPoint(Path.Get());
diff --git a/Engine/Source/Runtime/Analytics/Analytics/Public/AnalyticsEventAttribute.h b/Engine/Source/Runtime/Analytics/Analytics/Public/AnalyticsEventAttribute.h
--- a/Engine/Source/Runtime/Analytics/Analytics/Public/AnalyticsEventAttribute.h
+++ b/Engine/Source/Runtime/Analytics/Analytics/Public/AnalyticsEventAttribute.h
@@ -61,6 +61,22 @@ struct FAnalyticsEventAttribute
 		AttrValue = AttrValue.LeftChop(1);
 	}
 
+	// special case maps of FString->FString since TTypeToString (ironically) doesn't support FString
+	FAnalyticsEventAttribute(const FString& InName, const TMap<FString, FString>& InValueMap)
+		:AttrName(InName)
+		, AttrValue(FString())
+	{
+		// Serialize the map into "key:value,..." format
+		for (auto& KVP : InValueMap)
+		{
+			AttrValue += KVP.Key + ":" + KVP.Value + ",";
+		}
+
+		// Remove the trailing comma
+		AttrValue = AttrValue.LeftChop(1);
+	}
+
+
 	// Allow maps to be used as values formatted as comma delimited key-value pairs
 	template <typename T>
 	FAnalyticsEventAttribute(const FString& InName, const TMap<FString, T>& InValueMap)
diff --git a/Engine/Source/Runtime/Core/Private/Misc/FeedbackContext.cpp b/Engine/Source/Runtime/Core/Private/Misc/FeedbackContext.cpp
--- a/Engine/Source/Runtime/Core/Private/Misc/FeedbackContext.cpp
+++ b/Engine/Source/Runtime/Core/Private/Misc/FeedbackContext.cpp
@@ -85,7 +85,10 @@ void FScopedSlowTask::MakeDialog(bool bShowCancelButton, bool bAllowInPIE)
 	if (!GIsSlowTask && bIsDialogAllowed)
 	{
 		Context.StartSlowTask(GetCurrentMessage(), bShowCancelButton);
-		bCreatedDialog = true;
+		if (GIsSlowTask)
+		{
+			bCreatedDialog = true;
+		}
 	}
 }
 
diff --git a/Engine/Source/Runtime/Core/Public/Containers/Ticker.h b/Engine/Source/Runtime/Core/Public/Containers/Ticker.h
--- a/Engine/Source/Runtime/Core/Public/Containers/Ticker.h
+++ b/Engine/Source/Runtime/Core/Public/Containers/Ticker.h
@@ -56,7 +56,7 @@ public:
 	DELEGATE_DEPRECATED("This RemoveTicker overload is deprecated - please remove tickers using the FDelegateHandle returned by the AddTicker function.")
 	void RemoveTicker(const FTickerDelegate& Delegate)
 	{
-		for (int32 Index = 0; Index < PriorityQueue.Num(); ++Index)
+		for (int32 Index = PriorityQueue.Num() - 1; Index >= 0; --Index)
 		{
 			if (PriorityQueue[Index].Delegate.DEPRECATED_Compare(Delegate))
 			{
diff --git a/Engine/Source/Runtime/CoreUObject/Private/UObject/ScriptCore.cpp b/Engine/Source/Runtime/CoreUObject/Private/UObject/ScriptCore.cpp
--- a/Engine/Source/Runtime/CoreUObject/Private/UObject/ScriptCore.cpp
+++ b/Engine/Source/Runtime/CoreUObject/Private/UObject/ScriptCore.cpp
@@ -614,9 +614,12 @@ void UObject::ProcessInternal( FFrame& Stack, RESULT_DECL )
 				// Notify anyone who cares that we've had a fatal error, so we can shut down PIE, etc
 				const FString Desc = FString::Printf(TEXT("Runaway loop detected (over %i iterations)"), GMaximumScriptLoopIterations );
 				FBlueprintExceptionInfo RunawayLoopExceptionInfo(EBlueprintExceptionType::InfiniteLoop, Desc);
-				FBlueprintCoreDelegates::ThrowScriptException(this, Stack, RunawayLoopExceptionInfo);
 
+				// Need to reset Runaway counter BEFORE throwing script exception, because the exception causes a modal dialog,
+				// and other scripts running will then erroneously think they are also "runaway".
 				Runaway = 0;
+
+				FBlueprintCoreDelegates::ThrowScriptException(this, Stack, RunawayLoopExceptionInfo);
 				return;
 			}
 #endif
diff --git a/Engine/Source/Runtime/Engine/Classes/AI/Navigation/NavigationData.h b/Engine/Source/Runtime/Engine/Classes/AI/Navigation/NavigationData.h
--- a/Engine/Source/Runtime/Engine/Classes/AI/Navigation/NavigationData.h
+++ b/Engine/Source/Runtime/Engine/Classes/AI/Navigation/NavigationData.h
@@ -240,17 +240,17 @@ struct ENGINE_API FNavigationPath : public TSharedFromThis<FNavigationPath, ESPM
 	}
 
 	/** checks if given path, starting from StartingIndex, intersects with given AABB box */
-	virtual bool DoesIntersectBox(const FBox& Box, uint32 StartingIndex = 0, int32* IntersectingSegmentIndex = NULL) const;
+	virtual bool DoesIntersectBox(const FBox& Box, uint32 StartingIndex = 0, int32* IntersectingSegmentIndex = NULL, FVector* AgentExtent = NULL) const;
 	/** checks if given path, starting from StartingIndex, intersects with given AABB box. This version uses AgentLocation as beginning of the path
 	 *	with segment between AgentLocation and path's StartingIndex-th node treated as first path segment to check */
-	virtual bool DoesIntersectBox(const FBox& Box, const FVector& AgentLocation, uint32 StartingIndex = 0, int32* IntersectingSegmentIndex = NULL) const;
+	virtual bool DoesIntersectBox(const FBox& Box, const FVector& AgentLocation, uint32 StartingIndex = 0, int32* IntersectingSegmentIndex = NULL, FVector* AgentExtent = NULL) const;
 	/** retrieves normalized direction vector to given path segment
 	 *	for '0'-th segment returns same as for 1st segment 
 	 */
 	virtual FVector GetSegmentDirection(uint32 SegmentEndIndex) const;
 
 private:
-	bool DoesPathIntersectBoxImplementation(const FBox& Box, const FVector& StartLocation, uint32 StartingIndex, int32* IntersectingSegmentIndex) const;
+	bool DoesPathIntersectBoxImplementation(const FBox& Box, const FVector& StartLocation, uint32 StartingIndex, int32* IntersectingSegmentIndex, FVector* AgentExtent) const;
 
 public:
 
@@ -870,10 +870,10 @@ protected:
 
 FORCEINLINE bool FPathFindingResult::IsPartial() const
 {
-	return Result == ENavigationQueryResult::Fail && Path.IsValid() && Path->IsPartial();
+	return (Result != ENavigationQueryResult::Error) && Path.IsValid() && Path->IsPartial();
 }
 
 FORCEINLINE void FNavigationPath::SetNavigationDataUsed(const ANavigationData* const NewData)
 {
 	NavigationDataUsed = NewData;
-}
\ No newline at end of file
+}
diff --git a/Engine/Source/Runtime/Engine/Classes/AI/Navigation/NavigationSystem.h b/Engine/Source/Runtime/Engine/Classes/AI/Navigation/NavigationSystem.h
--- a/Engine/Source/Runtime/Engine/Classes/AI/Navigation/NavigationSystem.h
+++ b/Engine/Source/Runtime/Engine/Classes/AI/Navigation/NavigationSystem.h
@@ -481,8 +481,8 @@ public:
 
 	FORCEINLINE bool SupportsNavigationGeneration() const { return bSupportRebuilding; }
 
-	static bool DoesPathIntersectBox(const FNavigationPath* Path, const FBox& Box, uint32 StartingIndex = 0);
-	static bool DoesPathIntersectBox(const FNavigationPath* Path, const FBox& Box, const FVector& AgentLocation, uint32 StartingIndex = 0);
+	static bool DoesPathIntersectBox(const FNavigationPath* Path, const FBox& Box, uint32 StartingIndex = 0, FVector* AgentExtent = NULL);
+	static bool DoesPathIntersectBox(const FNavigationPath* Path, const FBox& Box, const FVector& AgentLocation, uint32 StartingIndex = 0, FVector* AgentExtent = NULL);
 
 	//----------------------------------------------------------------------//
 	// Active tiles
@@ -552,6 +552,7 @@ public:
 
 	void AddDirtyArea(const FBox& NewArea, int32 Flags);
 	void AddDirtyAreas(const TArray<FBox>& NewAreas, int32 Flags);
+	bool HasDirtyAreasQueued() const;
 
 	const FNavigationOctree* GetNavOctree() const { return NavOctree; }
 
diff --git a/Engine/Source/Runtime/Engine/Classes/AI/Navigation/NavigationTypes.h b/Engine/Source/Runtime/Engine/Classes/AI/Navigation/NavigationTypes.h
--- a/Engine/Source/Runtime/Engine/Classes/AI/Navigation/NavigationTypes.h
+++ b/Engine/Source/Runtime/Engine/Classes/AI/Navigation/NavigationTypes.h
@@ -465,6 +465,7 @@ struct ENGINE_API FPathFindingQuery
 	explicit FPathFindingQuery(FNavPathSharedRef PathToRecalculate, const ANavigationData* NavDataOverride = NULL);
 
 	FPathFindingQuery& SetPathInstanceToUpdate(FNavPathSharedPtr InPathInstanceToFill) { PathInstanceToFill = InPathInstanceToFill; return *this; }
+	FPathFindingQuery& SetAllowPartialPaths(bool bAllow) { bAllowPartialPaths = bAllow; return *this; }
 };
 
 namespace EPathFindingMode
diff --git a/Engine/Source/Runtime/Engine/Classes/AI/Navigation/RecastNavMesh.h b/Engine/Source/Runtime/Engine/Classes/AI/Navigation/RecastNavMesh.h
--- a/Engine/Source/Runtime/Engine/Classes/AI/Navigation/RecastNavMesh.h
+++ b/Engine/Source/Runtime/Engine/Classes/AI/Navigation/RecastNavMesh.h
@@ -145,13 +145,13 @@ struct ENGINE_API FNavMeshPath : public FNavigationPath
 
 	bool IsPathSegmentANavLink(const int32 PathSegmentStartIndex) const;
 
-	virtual bool DoesIntersectBox(const FBox& Box, uint32 StartingIndex = 0, int32* IntersectingSegmentIndex = NULL) const override;
-	virtual bool DoesIntersectBox(const FBox& Box, const FVector& AgentLocation, uint32 StartingIndex = 0, int32* IntersectingSegmentIndex = NULL) const override;
+	virtual bool DoesIntersectBox(const FBox& Box, uint32 StartingIndex = 0, int32* IntersectingSegmentIndex = NULL, FVector* AgentExtent = NULL) const override;
+	virtual bool DoesIntersectBox(const FBox& Box, const FVector& AgentLocation, uint32 StartingIndex = 0, int32* IntersectingSegmentIndex = NULL, FVector* AgentExtent = NULL) const override;
 	/** retrieves normalized direction vector to given path segment. If path is not string pulled navigation corridor is being used */
 	virtual FVector GetSegmentDirection(uint32 SegmentEndIndex) const override;
 
 private:
-	bool DoesPathIntersectBoxImplementation(const FBox& Box, const FVector& StartLocation, uint32 StartingIndex, int32* IntersectingSegmentIndex) const;
+	bool DoesPathIntersectBoxImplementation(const FBox& Box, const FVector& StartLocation, uint32 StartingIndex, int32* IntersectingSegmentIndex, FVector* AgentExtent) const;
 public:
 
 #if ENABLE_VISUAL_LOG
diff --git a/Engine/Source/Runtime/Engine/Classes/Components/PrimitiveComponent.h b/Engine/Source/Runtime/Engine/Classes/Components/PrimitiveComponent.h
--- a/Engine/Source/Runtime/Engine/Classes/Components/PrimitiveComponent.h
+++ b/Engine/Source/Runtime/Engine/Classes/Components/PrimitiveComponent.h
@@ -1114,7 +1114,7 @@ public:
 	 * @return		Success if returns > 0.f, if returns 0.f, it is either not convex or inside of the point
 	 *				If returns < 0.f, this primitive does not have collsion
 	 */
-	float GetDistanceToCollision(const FVector& Point, FVector& ClosestPointOnCollision) const;
+	virtual float GetDistanceToCollision(const FVector& Point, FVector& ClosestPointOnCollision) const;
 
 	/**
 	 * Creates a proxy to represent the primitive to the scene manager in the rendering thread.
diff --git a/Engine/Source/Runtime/Engine/Classes/Components/SkeletalMeshComponent.h b/Engine/Source/Runtime/Engine/Classes/Components/SkeletalMeshComponent.h
--- a/Engine/Source/Runtime/Engine/Classes/Components/SkeletalMeshComponent.h
+++ b/Engine/Source/Runtime/Engine/Classes/Components/SkeletalMeshComponent.h
@@ -832,6 +832,7 @@ public:
 	virtual bool IsAnyRigidBodyAwake() override;
 	virtual void OnComponentCollisionSettingsChanged() override;
 	virtual void SetPhysMaterialOverride(UPhysicalMaterial* NewPhysMaterial) override;
+	virtual float GetDistanceToCollision(const FVector& Point, FVector& ClosestPointOnCollision) const;
 	virtual bool LineTraceComponent( FHitResult& OutHit, const FVector Start, const FVector End, const FCollisionQueryParams& Params ) override;
 	virtual bool SweepComponent( FHitResult& OutHit, const FVector Start, const FVector End, const FCollisionShape& CollisionShape, bool bTraceComplex=false) override;
 	virtual bool ComponentOverlapComponent(class UPrimitiveComponent* PrimComp, const FVector Pos, const FRotator FRotator, const FCollisionQueryParams& Params) override;
diff --git a/Engine/Source/Runtime/Engine/Classes/Engine/Engine.h b/Engine/Source/Runtime/Engine/Classes/Engine/Engine.h
--- a/Engine/Source/Runtime/Engine/Classes/Engine/Engine.h
+++ b/Engine/Source/Runtime/Engine/Classes/Engine/Engine.h
@@ -2327,15 +2327,36 @@ public:
 	//virtual ENetMode GetNetMode(FName NetDriverName = NAME_GameNetDriver) const;
 	ENetMode GetNetMode(const UWorld *World) const;
 
+	/**
+	 * Creates a UNetDriver with an engine assigned name
+	 *
+	 * @param InWorld the world context
+	 * @param NetDriverDefinition The name of the definition to use
+	 *
+	 * @return new netdriver if successful, nullptr otherwise
+	 */
+	UNetDriver* CreateNetDriver(UWorld *InWorld, FName NetDriverDefinition);
+
 	/**
 	 * Creates a UNetDriver and associates a name with it.
 	 *
+	 * @param InWorld the world context
 	 * @param NetDriverName The name to associate with the driver.
 	 * @param NetDriverDefinition The name of the definition to use
 	 *
 	 * @return True if the driver was created successfully, false if there was an error.
 	 */
 	bool CreateNamedNetDriver(UWorld *InWorld, FName NetDriverName, FName NetDriverDefinition);
+
+	/**
+	 * Creates a UNetDriver and associates a name with it.
+	 *
+	 * @param PendingNetGame the pending net game context
+	 * @param NetDriverName The name to associate with the driver.
+	 * @param NetDriverDefinition The name of the definition to use
+	 *
+	 * @return True if the driver was created successfully, false if there was an error.
+	 */
 	bool CreateNamedNetDriver(UPendingNetGame *PendingNetGame, FName NetDriverName, FName NetDriverDefinition);
 	
 	/**
diff --git a/Engine/Source/Runtime/Engine/Classes/Engine/LocalPlayer.h b/Engine/Source/Runtime/Engine/Classes/Engine/LocalPlayer.h
--- a/Engine/Source/Runtime/Engine/Classes/Engine/LocalPlayer.h
+++ b/Engine/Source/Runtime/Engine/Classes/Engine/LocalPlayer.h
@@ -5,6 +5,8 @@
 //=============================================================================
 
 #pragma once
+#include "SlateCore.h"
+#include "Reply.h"
 #include "Player.h"
 #include "LocalPlayer.generated.h"
 
@@ -233,8 +235,17 @@ protected:
 	 */
 	void HandleControllerConnectionChange(bool bConnected, int32 InUserId, int32 InControllerId);
 
+	/** FReply used to defer some slate operations. */
+	FReply SlateOperations;
+
 public:
 
+	/**
+	 *  Getter for slate operations.
+	 */
+	FReply& GetSlateOperations() { return SlateOperations; }
+	const FReply& GetSlateOperations() const { return SlateOperations; }
+
 	/**
 	 * Get the world the players actor belongs to
 	 *
diff --git a/Engine/Source/Runtime/Engine/Classes/Engine/NetConnection.h b/Engine/Source/Runtime/Engine/Classes/Engine/NetConnection.h
--- a/Engine/Source/Runtime/Engine/Classes/Engine/NetConnection.h
+++ b/Engine/Source/Runtime/Engine/Classes/Engine/NetConnection.h
@@ -122,9 +122,9 @@ class UNetConnection : public UPlayer
 	UPROPERTY()
 	TArray<class AActor*> SentTemporaries;
 
-	/** @todo document */
+	/** The actor that is currently being viewed/controlled by the owning controller */
 	UPROPERTY()
-	class AActor* Viewer;
+	class AActor* ViewTarget;
 
 	/** Reference to controlling actor (usually PlayerController) */
 	UPROPERTY()
diff --git a/Engine/Source/Runtime/Engine/Classes/Engine/StreamableManager.h b/Engine/Source/Runtime/Engine/Classes/Engine/StreamableManager.h
--- a/Engine/Source/Runtime/Engine/Classes/Engine/StreamableManager.h
+++ b/Engine/Source/Runtime/Engine/Classes/Engine/StreamableManager.h
@@ -41,6 +41,7 @@ struct ENGINE_API FStreamableManager
 	 * Objects will be strongly referenced until the delegate is called, then Unload is called automatically 
 	 */
 	void RequestAsyncLoad(const TArray<FStringAssetReference>& TargetsToStream, FStreamableDelegate DelegateToCall);
+	void RequestAsyncLoad( const FStringAssetReference& TargetToStream, FStreamableDelegate DelegateToCall );
 
 	/** Exposes references to GC system */
 	void AddStructReferencedObjects(class FReferenceCollector& Collector) const;
diff --git a/Engine/Source/Runtime/Engine/Classes/GameFramework/Actor.h b/Engine/Source/Runtime/Engine/Classes/GameFramework/Actor.h
--- a/Engine/Source/Runtime/Engine/Classes/GameFramework/Actor.h
+++ b/Engine/Source/Runtime/Engine/Classes/GameFramework/Actor.h
@@ -1474,17 +1474,22 @@ public:
 #endif		// WITH_EDITOR
 
 	/**
-	* @param ViewPos		Position of the viewer
-	* @param ViewDir		Vector direction of viewer
-	* @param Viewer			PlayerController owned by the client for whom net priority is being determined
-	* @param InChannel		Channel on which this actor is being replicated.
-	* @param Time			Time since actor was last replicated
-	* @param bLowBandwidth True if low bandwith of viewer
-	* @return				Priority of this actor for replication
-	 */
+	 * @param ViewPos		Position of the viewer
+	 * @param ViewDir		Vector direction of viewer
+	 * @param Viewer		PlayerController owned by the client for whom net priority is being determined
+	 * @param ViewTarget	The actor that is currently being viewed/controlled by Viewer, usually a pawn
+	 * @param InChannel		Channel on which this actor is being replicated.
+	 * @param Time			Time since actor was last replicated
+	 * @param bLowBandwidth True if low bandwith of viewer
+	 * @return				Priority of this actor for replication
+	 */
+	virtual float GetNetPriority(const FVector& ViewPos, const FVector& ViewDir, class APlayerController* Viewer, AActor* ViewTarget, UActorChannel* InChannel, float Time, bool bLowBandwidth);
+
+	DEPRECATED(4.8, "GetNetPriority now takes a ViewTarget, please override that version.")
 	virtual float GetNetPriority(const FVector& ViewPos, const FVector& ViewDir, class APlayerController* Viewer, UActorChannel* InChannel, float Time, bool bLowBandwidth);
 
-	virtual bool GetNetDormancy(const FVector& ViewPos, const FVector& ViewDir, class APlayerController* Viewer, UActorChannel* InChannel, float Time, bool bLowBandwidth);
+	/** Returns true if the actor should be dormant for a specific net connection. Only checked for DORM_DormantPartial */
+	virtual bool GetNetDormancy(const FVector& ViewPos, const FVector& ViewDir, class APlayerController* Viewer, AActor* ViewTarget, UActorChannel* InChannel, float Time, bool bLowBandwidth);
 
 	/** 
 	 * Allows for a specific response from the actor when the actor channel is opened (client side)
@@ -1657,12 +1662,12 @@ public:
 
 	/** 
 	  * @param RealViewer - is the PlayerController associated with the client for which network relevancy is being checked
-	  * @param Viewer - is the Actor being used as the point of view for the PlayerController
+	  * @param VieweTarget - is the Actor being used as the point of view for the PlayerController
 	  * @param SrcLocation - is the viewing location
 	  *
 	  * @return bool - true if this actor is network relevant to the client associated with RealViewer 
 	  */
-	virtual bool IsNetRelevantFor(const APlayerController* RealViewer, const AActor* Viewer, const FVector& SrcLocation) const;
+	virtual bool IsNetRelevantFor(const APlayerController* RealViewer, const AActor* ViewTarget, const FVector& SrcLocation) const;
 
 	/**
 	 * Check if this actor is the owner when doing relevancy checks for actors marked bOnlyRelevantToOwner
diff --git a/Engine/Source/Runtime/Engine/Classes/GameFramework/Character.h b/Engine/Source/Runtime/Engine/Classes/GameFramework/Character.h
--- a/Engine/Source/Runtime/Engine/Classes/GameFramework/Character.h
+++ b/Engine/Source/Runtime/Engine/Classes/GameFramework/Character.h
@@ -11,6 +11,8 @@ class UPrimitiveComponent;
 
 DECLARE_DYNAMIC_MULTICAST_DELEGATE_ThreeParams(FMovementModeChangedSignature, class ACharacter*, Character, EMovementMode, PrevMovementMode, uint8, PreviousCustomMode);
 DECLARE_DYNAMIC_MULTICAST_DELEGATE_ThreeParams(FCharacterMovementUpdatedSignature, float, DeltaSeconds, FVector, OldLocation, FVector, OldVelocity);
+DECLARE_DYNAMIC_MULTICAST_DELEGATE(FCharacterReachedApexSignature);
+
 //
 // Forward declarations
 //
@@ -485,7 +487,11 @@ public:
 	virtual void Falling() {}
 
 	/** Called when character's jump reaches Apex. Needs CharacterMovement->bNotifyApex = true */
-	virtual void NotifyJumpApex() {}
+	virtual void NotifyJumpApex();
+
+	/** Broadcast when Character's jump reaches its apex. Needs CharacterMovement->bNotifyApex = true */
+	UPROPERTY(BlueprintAssignable)
+	FCharacterReachedApexSignature OnReachedJumpApex;
 
 	/**
 	 * Called upon landing when falling, to perform actions based on the Hit result. Triggers the OnLanded event.
diff --git a/Engine/Source/Runtime/Engine/Classes/GameFramework/CharacterMovementComponent.h b/Engine/Source/Runtime/Engine/Classes/GameFramework/CharacterMovementComponent.h
--- a/Engine/Source/Runtime/Engine/Classes/GameFramework/CharacterMovementComponent.h
+++ b/Engine/Source/Runtime/Engine/Classes/GameFramework/CharacterMovementComponent.h
@@ -589,7 +589,7 @@ public:
 	uint32 bNetworkMovementModeChanged:1;
 
 	/** if true, event NotifyJumpApex() to CharacterOwner's controller when at apex of jump.  Is cleared when event is triggered. */
-	UPROPERTY()
+	UPROPERTY(Category = "Character Movement", EditAnywhere, BlueprintReadWrite)
 	uint32 bNotifyApex:1;
 
 	/** Instantly stop when in flying mode and no acceleration is being applied. */
@@ -1520,6 +1520,13 @@ protected:
 	 */
 	virtual void OnUnableToFollowBaseMove(const FVector& DeltaPosition, const FVector& OldLocation, const FHitResult& MoveOnBaseHit);
 
+	/** Project a location to navmesh to find adjusted height
+	 * @param TestLocation		Location to project
+	 * @param NavFloorLocation	Location on navmesh
+	 * @return True if projection was performed (successfully or not)
+	 */
+	bool FindNavFloor(const FVector& TestLocation, FNavLocation& NavFloorLocation) const;
+
 public:
 	// Movement functions broken out based on owner's network Role.
 	// TickComponent calls the correct version based on the Role.
diff --git a/Engine/Source/Runtime/Engine/Classes/GameFramework/CheatManager.h b/Engine/Source/Runtime/Engine/Classes/GameFramework/CheatManager.h
--- a/Engine/Source/Runtime/Engine/Classes/GameFramework/CheatManager.h
+++ b/Engine/Source/Runtime/Engine/Classes/GameFramework/CheatManager.h
@@ -67,10 +67,6 @@ class ENGINE_API UCheatManager : public UObject
 	UPROPERTY()
 	TSubclassOf<class ADebugCameraController>  DebugCameraControllerClass;
 
-	/** Holds information about VisualLogger on server - came from replication */
-	UPROPERTY(ReplicatedUsing = OnRep_VisualLoggerActiveOnServer)
-	bool bVisualLoggerActiveOnServer;
-
 	// Trace/Sweep debug start
 	/** If we should should perform a debug capsule trace and draw results. Toggled with DebugCapsuleSweep() */
 	uint32 bDebugCapsuleSweep:1;
@@ -211,10 +207,6 @@ class ENGINE_API UCheatManager : public UObject
 	UFUNCTION(reliable, server, WithValidation)
 	virtual void ServerToggleAILogging();
 
-	/** after bVisualLoggerActiveOnServer replication callback */
-	UFUNCTION()
-	virtual void OnRep_VisualLoggerActiveOnServer();
-
 	/** Toggle capsule trace debugging. Will trace a capsule from current view point and show where it hits the world */
 	UFUNCTION(exec)
 	virtual void DebugCapsuleSweep();
diff --git a/Engine/Source/Runtime/Engine/Classes/GameFramework/HUD.h b/Engine/Source/Runtime/Engine/Classes/GameFramework/HUD.h
--- a/Engine/Source/Runtime/Engine/Classes/GameFramework/HUD.h
+++ b/Engine/Source/Runtime/Engine/Classes/GameFramework/HUD.h
@@ -118,6 +118,23 @@ public:
 	UFUNCTION(exec)
 	void ShowDebugToggleSubCategory(FName Category);
 
+	/** Toggles 'ShowDebug' from showing debug info between reticle target actor (of subclass <DesiredClass>) and camera view target */
+	UFUNCTION(exec)
+	void ShowDebugForReticleTargetToggle(TSubclassOf<AActor> DesiredClass);
+
+private:
+	/** if true show debug info for 'ShowDebugTargetActor', otherwise for Camera Viewtarget */
+	static bool bShowDebugForReticleTarget;
+
+	/** Class filter for selecting 'ShowDebugTargetActor' when 'bShowDebugForReticleTarget' is true. */
+	UPROPERTY(Transient)
+	TSubclassOf<AActor> ShowDebugTargetDesiredClass;
+
+	/** Show Debug Actor used if 'bShowDebugForReticleTarget' is true, only updated if trace from reticle hit a new Actor of class 'ShowDebugTargetDesiredClass'*/
+	UPROPERTY(Transient)
+	AActor* ShowDebugTargetActor;
+
+public:
 	/**
 	 * Add debug text for a specific actor to be displayed via DrawDebugTextList().  If the debug text is invalid then it will
 	 * attempt to remove any previous entries via RemoveDebugText().
diff --git a/Engine/Source/Runtime/Engine/Classes/GameFramework/Pawn.h b/Engine/Source/Runtime/Engine/Classes/GameFramework/Pawn.h
--- a/Engine/Source/Runtime/Engine/Classes/GameFramework/Pawn.h
+++ b/Engine/Source/Runtime/Engine/Classes/GameFramework/Pawn.h
@@ -211,9 +211,8 @@ public:
 	virtual void Tick(float DeltaSeconds) override;
 	virtual void Reset() override;
 	virtual FString GetHumanReadableName() const override;
-	virtual float GetNetPriority(const FVector& ViewPos, const FVector& ViewDir, APlayerController* Viewer, UActorChannel* InChannel, float Time, bool bLowBandwidth) override;
 	virtual bool ShouldTickIfViewportsOnly() const override;
-	virtual bool IsNetRelevantFor(const APlayerController* RealViewer, const AActor* Viewer, const FVector& SrcLocation) const override;
+	virtual bool IsNetRelevantFor(const APlayerController* RealViewer, const AActor* ViewTarget, const FVector& SrcLocation) const override;
 	virtual void PostNetReceiveLocationAndRotation() override;
 	virtual void PostNetReceiveVelocity(const FVector& NewVelocity) override;
 	virtual void DisplayDebug(UCanvas* Canvas, const FDebugDisplayInfo& DebugDisplay, float& YL, float& YPos) override;
diff --git a/Engine/Source/Runtime/Engine/Classes/GameFramework/PlayerController.h b/Engine/Source/Runtime/Engine/Classes/GameFramework/PlayerController.h
--- a/Engine/Source/Runtime/Engine/Classes/GameFramework/PlayerController.h
+++ b/Engine/Source/Runtime/Engine/Classes/GameFramework/PlayerController.h
@@ -1151,7 +1151,7 @@ public:
 	virtual void GetActorEyesViewPoint(FVector& Location, FRotator& Rotation) const override;
 	virtual void CalcCamera(float DeltaTime, struct FMinimalViewInfo& OutResult) override;
 	virtual void TickActor(float DeltaTime, enum ELevelTick TickType, FActorTickFunction& ThisTickFunction) override;
-	virtual bool IsNetRelevantFor(const APlayerController* RealViewer, const AActor* Viewer, const FVector& SrcLocation) const override;
+	virtual bool IsNetRelevantFor(const APlayerController* RealViewer, const AActor* ViewTarget, const FVector& SrcLocation) const override;
 	virtual void FellOutOfWorld(const class UDamageType& dmgType) override;
 	virtual void Reset() override;
 	virtual void Possess(APawn* aPawn) override;
@@ -1162,7 +1162,7 @@ public:
 	virtual void OnActorChannelOpen(class FInBunch& InBunch, class UNetConnection* Connection) override;
 	virtual void OnSerializeNewActor(class FOutBunch& OutBunch) override;
 	virtual void OnNetCleanup(class UNetConnection* Connection) override;
-	virtual float GetNetPriority(const FVector& ViewPos, const FVector& ViewDir, APlayerController* Viewer, UActorChannel* InChannel, float Time, bool bLowBandwidth) override;
+	virtual float GetNetPriority(const FVector& ViewPos, const FVector& ViewDir, APlayerController* Viewer, AActor* ViewTarget, UActorChannel* InChannel, float Time, bool bLowBandwidth) override;
 	virtual class UPlayer* GetNetOwningPlayer() override;
 	virtual class UNetConnection* GetNetConnection() override;
 	virtual void DisplayDebug(class UCanvas* Canvas, const FDebugDisplayInfo& DebugDisplay, float& YL, float& YPos) override;
@@ -1415,6 +1415,12 @@ public:
 	 */
 	virtual void AutoManageActiveCameraTarget(AActor* SuggestedTarget);
 
+	/**
+	 * Notify from server that Visual Logger is recording, to show that information on client about possible performance issues 
+	 */
+	UFUNCTION(Reliable, Client)
+	void OnServerStartedVisualLogger(bool bIsLogging);
+
 protected:
 
 	virtual ACameraActor* GetAutoActivateCameraForPlayer() const;
@@ -1531,6 +1537,4 @@ public:
 	UPROPERTY()
 	uint16		LastCompletedSeamlessTravelCount;
 
-	/** FReply used to defer some slate operations. */
-	TSharedPtr<class FReply> SlateOperations;
 };
diff --git a/Engine/Source/Runtime/Engine/Classes/GameFramework/WorldSettings.h b/Engine/Source/Runtime/Engine/Classes/GameFramework/WorldSettings.h
--- a/Engine/Source/Runtime/Engine/Classes/GameFramework/WorldSettings.h
+++ b/Engine/Source/Runtime/Engine/Classes/GameFramework/WorldSettings.h
@@ -164,22 +164,25 @@ struct ENGINE_API FNetViewer
 {
 	GENERATED_USTRUCT_BODY()
 
+	/** The player controller corresponding to this view */
 	UPROPERTY()
 	class APlayerController* InViewer;
 
+	/** The actor that is being directly viewed, usually a pawn */
 	UPROPERTY()
-	class AActor* Viewer;
+	class AActor* ViewTarget;
 
+	/** Where the viewer is looking from */
 	UPROPERTY()
 	FVector ViewLocation;
 
+	/** Direction the viewer is looking */
 	UPROPERTY()
 	FVector ViewDir;
 
-
 	FNetViewer()
 		: InViewer(NULL)
-		, Viewer(NULL)
+		, ViewTarget(NULL)
 		, ViewLocation(ForceInit)
 		, ViewDir(ForceInit)
 	{
diff --git a/Engine/Source/Runtime/Engine/Classes/Kismet/GameplayStatics.h b/Engine/Source/Runtime/Engine/Classes/Kismet/GameplayStatics.h
--- a/Engine/Source/Runtime/Engine/Classes/Kismet/GameplayStatics.h
+++ b/Engine/Source/Runtime/Engine/Classes/Kismet/GameplayStatics.h
@@ -411,7 +411,7 @@ class ENGINE_API UGameplayStatics : public UBlueprintFunctionLibrary
 	 * @param HitItem		Primitive-specific data recording which item in the primitive was hit
 	 */
 	UFUNCTION(BlueprintPure, Category = "Collision", meta=(NativeBreakFunc))
-	static void BreakHitResult(const struct FHitResult& Hit, bool& bBlockingHit, float& Time, FVector& Location, FVector& Normal, FVector& ImpactPoint, FVector& ImpactNormal, class UPhysicalMaterial*& PhysMat, class AActor*& HitActor, class UPrimitiveComponent*& HitComponent, FName& HitBoneName, int32& HitItem);
+	static void BreakHitResult(const struct FHitResult& Hit, bool& bBlockingHit, float& Time, FVector& Location, FVector& Normal, FVector& ImpactPoint, FVector& ImpactNormal, class UPhysicalMaterial*& PhysMat, class AActor*& HitActor, class UPrimitiveComponent*& HitComponent, FName& HitBoneName, int32& HitItem, FVector& TraceStart, FVector& TraceEnd);
 
 	/** Returns the EPhysicalSurface type of the given Hit. 
 	 * To edit surface type for your project, use ProjectSettings/Physics/PhysicalSurface section
diff --git a/Engine/Source/Runtime/Engine/Classes/Matinee/MatineeActor.h b/Engine/Source/Runtime/Engine/Classes/Matinee/MatineeActor.h
--- a/Engine/Source/Runtime/Engine/Classes/Matinee/MatineeActor.h
+++ b/Engine/Source/Runtime/Engine/Classes/Matinee/MatineeActor.h
@@ -309,7 +309,7 @@ public:
 #endif
 
 	// Begin AActor Interface
-	virtual float GetNetPriority(const FVector& ViewPos, const FVector& ViewDir, APlayerController* Viewer, UActorChannel* InChannel, float Time, bool bLowBandwidth) override;
+	virtual float GetNetPriority(const FVector& ViewPos, const FVector& ViewDir, APlayerController* Viewer, AActor* ViewTarget, UActorChannel* InChannel, float Time, bool bLowBandwidth) override;
 	virtual void Tick(float DeltaTime) override;
 	virtual void PreNetReceive() override;
 	virtual void PostNetReceive() override;
diff --git a/Engine/Source/Runtime/Engine/Private/AI/Navigation/NavigationData.cpp b/Engine/Source/Runtime/Engine/Private/AI/Navigation/NavigationData.cpp
--- a/Engine/Source/Runtime/Engine/Private/AI/Navigation/NavigationData.cpp
+++ b/Engine/Source/Runtime/Engine/Private/AI/Navigation/NavigationData.cpp
@@ -174,7 +174,8 @@ void ANavigationData::PostLoad()
 {
 	Super::PostLoad();
 
-	if (GetLinkerUE4Version() < VER_UE4_ADD_MODIFIERS_RUNTIME_GENERATION)
+	if ((GetLinkerUE4Version() < VER_UE4_ADD_MODIFIERS_RUNTIME_GENERATION) &&
+		(RuntimeGeneration == ERuntimeGenerationType::LegacyGeneration))
 	{
 		RuntimeGeneration = bRebuildAtRuntime_DEPRECATED ? ERuntimeGenerationType::Dynamic : ERuntimeGenerationType::Static;
 	}
@@ -239,6 +240,7 @@ void ANavigationData::TickActor(float DeltaTime, enum ELevelTick TickType, FActo
 	{
 		float TimeStamp = GetWorldTimeStamp();
 		TArray<FNavPathRecalculationRequest> PostponedRequests;
+		const UWorld* World = GetWorld();
 
 		// @todo batch-process it!
 		for (auto RecalcRequest : RepathRequests)
diff --git a/Engine/Source/Runtime/Engine/Private/AI/Navigation/NavigationPath.cpp b/Engine/Source/Runtime/Engine/Private/AI/Navigation/NavigationPath.cpp
--- a/Engine/Source/Runtime/Engine/Private/AI/Navigation/NavigationPath.cpp
+++ b/Engine/Source/Runtime/Engine/Private/AI/Navigation/NavigationPath.cpp
@@ -232,7 +232,7 @@ bool FNavigationPath::ContainsAnyCustomLink() const
 	return false;
 }
 
-FORCEINLINE bool FNavigationPath::DoesPathIntersectBoxImplementation(const FBox& Box, const FVector& StartLocation, uint32 StartingIndex, int32* IntersectingSegmentIndex) const
+FORCEINLINE bool FNavigationPath::DoesPathIntersectBoxImplementation(const FBox& Box, const FVector& StartLocation, uint32 StartingIndex, int32* IntersectingSegmentIndex, FVector* AgentExtent) const
 {	
 	bool bIntersects = false;
 
@@ -244,7 +244,13 @@ FORCEINLINE bool FNavigationPath::DoesPathIntersectBoxImplementation(const FBox&
 		{
 			const FVector Direction = (End - Start);
 
-			if (FMath::LineBoxIntersection(Box, Start, End, Direction))
+			FVector HitLocation, HitNormal;
+			float HitTime;
+
+			// If we have a valid AgentExtent, then we use an extent box to represent the path
+			// Otherwise we use a line to represent the path
+			if ((AgentExtent && FMath::LineExtentBoxIntersection(Box, Start, End, *AgentExtent, HitLocation, HitNormal, HitTime)) ||
+				(!AgentExtent && FMath::LineBoxIntersection(Box, Start, End, Direction)))
 			{
 				bIntersects = true;
 				if (IntersectingSegmentIndex != NULL)
@@ -261,7 +267,7 @@ FORCEINLINE bool FNavigationPath::DoesPathIntersectBoxImplementation(const FBox&
 	return bIntersects;
 }
 
-bool FNavigationPath::DoesIntersectBox(const FBox& Box, uint32 StartingIndex, int32* IntersectingSegmentIndex) const
+bool FNavigationPath::DoesIntersectBox(const FBox& Box, uint32 StartingIndex, int32* IntersectingSegmentIndex, FVector* AgentExtent) const
 {
 	// iterate over all segments and check if any intersects with given box
 	bool bIntersects = false;
@@ -269,13 +275,13 @@ bool FNavigationPath::DoesIntersectBox(const FBox& Box, uint32 StartingIndex, in
 
 	if (PathPoints.Num() > 1 && PathPoints.IsValidIndex(int32(StartingIndex)))
 	{
-		bIntersects = DoesPathIntersectBoxImplementation(Box, PathPoints[StartingIndex].Location, StartingIndex + 1, IntersectingSegmentIndex);
+		bIntersects = DoesPathIntersectBoxImplementation(Box, PathPoints[StartingIndex].Location, StartingIndex + 1, IntersectingSegmentIndex, AgentExtent);
 	}
 
 	return bIntersects;
 }
 
-bool FNavigationPath::DoesIntersectBox(const FBox& Box, const FVector& AgentLocation, uint32 StartingIndex, int32* IntersectingSegmentIndex) const
+bool FNavigationPath::DoesIntersectBox(const FBox& Box, const FVector& AgentLocation, uint32 StartingIndex, int32* IntersectingSegmentIndex, FVector* AgentExtent) const
 {
 	// iterate over all segments and check if any intersects with given box
 	bool bIntersects = false;
@@ -283,7 +289,7 @@ bool FNavigationPath::DoesIntersectBox(const FBox& Box, const FVector& AgentLoca
 
 	if (PathPoints.Num() > 1 && PathPoints.IsValidIndex(int32(StartingIndex)))
 	{
-		bIntersects = DoesPathIntersectBoxImplementation(Box, AgentLocation, StartingIndex, IntersectingSegmentIndex);
+		bIntersects = DoesPathIntersectBoxImplementation(Box, AgentLocation, StartingIndex, IntersectingSegmentIndex, AgentExtent);
 	}
 
 	return bIntersects;
@@ -899,7 +905,7 @@ bool FNavMeshPath::ContainsWithSameEnd(const FNavMeshPath* Other) const
 	return bAreTheSame;
 }
 
-FORCEINLINE_DEBUGGABLE bool FNavMeshPath::DoesPathIntersectBoxImplementation(const FBox& Box, const FVector& StartLocation, uint32 StartingIndex, int32* IntersectingSegmentIndex) const
+bool FNavMeshPath::DoesPathIntersectBoxImplementation(const FBox& Box, const FVector& StartLocation, uint32 StartingIndex, int32* IntersectingSegmentIndex, FVector* AgentExtent) const
 {
 	bool bIntersects = false;	
 	const TArray<FNavigationPortalEdge>& CorridorEdges = GetPathCorridorEdges();
@@ -908,14 +914,22 @@ FORCEINLINE_DEBUGGABLE bool FNavMeshPath::DoesPathIntersectBoxImplementation(con
 	if (CorridorEdges.IsValidIndex(StartingIndex))
 	{
 		FVector Start = StartLocation;
+
 		for (int32 PortalIndex = StartingIndex; PortalIndex < CorridorEdges.Num(); ++PortalIndex)
 		{
 			const FNavigationPortalEdge& Edge = CorridorEdges[PortalIndex];
-			const FVector End = Edge.Right + (Edge.Left - Edge.Right) / 2;
+			const FVector End = Edge.Right + (Edge.Left - Edge.Right) / 2 + (AgentExtent ? FVector(0.f, 0.f, AgentExtent->Z) : FVector::ZeroVector);
 			if (FVector::DistSquared(Start, End) > SMALL_NUMBER)
 			{
 				const FVector Direction = (End - Start);
-				if (FMath::LineBoxIntersection(Box, Start, End, Direction))
+				
+				FVector HitLocation, HitNormal;
+				float HitTime;
+
+				// If we have a valid AgentExtent, then we use an extent box to represent the path
+				// Otherwise we use a line to represent the path
+				if ((AgentExtent && FMath::LineExtentBoxIntersection(Box, Start, End, *AgentExtent, HitLocation, HitNormal, HitTime)) ||
+					(!AgentExtent && FMath::LineBoxIntersection(Box, Start, End, Direction)))
 				{
 					bIntersects = true;
 					if (IntersectingSegmentIndex != NULL)
@@ -933,11 +947,19 @@ FORCEINLINE_DEBUGGABLE bool FNavMeshPath::DoesPathIntersectBoxImplementation(con
 		if (bIntersects == false)
 		{
 			ensure(PathPoints.Num() == 2);
-			const FVector End = PathPoints.Last().Location;
+			const FVector End = PathPoints.Last().Location + (AgentExtent ? FVector(0.f, 0.f, AgentExtent->Z) : FVector::ZeroVector);
+			
 			if (FVector::DistSquared(StartLocation, End) > SMALL_NUMBER)
 			{
 				const FVector Direction = (End - Start);
-				if (FMath::LineBoxIntersection(Box, Start, End, Direction))
+
+				FVector HitLocation, HitNormal;
+				float HitTime;
+
+				// If we have a valid AgentExtent, then we use an extent box to represent the path
+				// Otherwise we use a line to represent the path
+				if ((AgentExtent && FMath::LineExtentBoxIntersection(Box, Start, End, *AgentExtent, HitLocation, HitNormal, HitTime)) ||
+					(!AgentExtent && FMath::LineBoxIntersection(Box, Start, End, Direction)))
 				{
 					bIntersects = true;
 					if (IntersectingSegmentIndex != NULL)
@@ -962,7 +984,7 @@ FORCEINLINE_DEBUGGABLE bool FNavMeshPath::DoesPathIntersectBoxImplementation(con
 	return bIntersects;
 }
 
-bool FNavMeshPath::DoesIntersectBox(const FBox& Box, uint32 StartingIndex, int32* IntersectingSegmentIndex) const
+bool FNavMeshPath::DoesIntersectBox(const FBox& Box, uint32 StartingIndex, int32* IntersectingSegmentIndex, FVector* AgentExtent) const
 {
 	if (IsStringPulled())
 	{
@@ -984,10 +1006,10 @@ bool FNavMeshPath::DoesIntersectBox(const FBox& Box, uint32 StartingIndex, int32
 	}
 	// else will be handled by DoesPathIntersectBoxImplementation
 
-	return bParametersValid && DoesPathIntersectBoxImplementation(Box, StartLocation, StartingIndex, IntersectingSegmentIndex);
+	return bParametersValid && DoesPathIntersectBoxImplementation(Box, StartLocation, StartingIndex, IntersectingSegmentIndex, AgentExtent);
 }
 
-bool FNavMeshPath::DoesIntersectBox(const FBox& Box, const FVector& AgentLocation, uint32 StartingIndex, int32* IntersectingSegmentIndex) const
+bool FNavMeshPath::DoesIntersectBox(const FBox& Box, const FVector& AgentLocation, uint32 StartingIndex, int32* IntersectingSegmentIndex, FVector* AgentExtent) const
 {
 	// trivial scenario check:
 	if (Box.IsInside(AgentLocation))
@@ -997,10 +1019,10 @@ bool FNavMeshPath::DoesIntersectBox(const FBox& Box, const FVector& AgentLocatio
 
 	if (IsStringPulled())
 	{
-		return Super::DoesIntersectBox(Box, AgentLocation, StartingIndex, IntersectingSegmentIndex);
+		return Super::DoesIntersectBox(Box, AgentLocation, StartingIndex, IntersectingSegmentIndex, AgentExtent);
 	}
 
-	return DoesPathIntersectBoxImplementation(Box, AgentLocation, StartingIndex, IntersectingSegmentIndex);
+	return DoesPathIntersectBoxImplementation(Box, AgentLocation, StartingIndex, IntersectingSegmentIndex, AgentExtent);
 }
 
 FVector FNavMeshPath::GetSegmentDirection(uint32 SegmentEndIndex) const
@@ -1271,4 +1293,4 @@ void UNavigationPath::SetPathPointsFromPath(FNavigationPath& NativePath)
 	{
 		PathPoints.Add(PathPoint.Location);
 	}
-}
+}
\ No newline at end of file
diff --git a/Engine/Source/Runtime/Engine/Private/AI/Navigation/NavigationSystem.cpp b/Engine/Source/Runtime/Engine/Private/AI/Navigation/NavigationSystem.cpp
--- a/Engine/Source/Runtime/Engine/Private/AI/Navigation/NavigationSystem.cpp
+++ b/Engine/Source/Runtime/Engine/Private/AI/Navigation/NavigationSystem.cpp
@@ -2106,6 +2106,11 @@ void UNavigationSystem::AddDirtyAreas(const TArray<FBox>& NewAreas, int32 Flags)
 	}
 }
 
+bool UNavigationSystem::HasDirtyAreasQueued() const
+{
+	return DirtyAreas.Num() > 0;
+}
+
 int32 GetDirtyFlagHelper(int32 UpdateFlags, int32 DefaultValue)
 {
 	return ((UpdateFlags & UNavigationSystem::OctreeUpdate_Geometry) != 0) ? ENavigationDirtyFlag::All :
@@ -3411,10 +3416,11 @@ ENavigationQueryResult::Type UNavigationSystem::GetPathLength(UObject* WorldCont
 
 bool UNavigationSystem::IsNavigationBeingBuilt(UObject* WorldContextObject)
 {
-	UWorld* World = GEngine->GetWorldFromContextObject(WorldContextObject); 
-	if (World != NULL && World->GetNavigationSystem() != NULL)
+	UWorld* World = GEngine->GetWorldFromContextObject(WorldContextObject);
+	UNavigationSystem* NavSys = UNavigationSystem::GetCurrent(World);
+	if (NavSys)
 	{
-		return World->GetNavigationSystem()->IsNavigationBuildInProgress();
+		return NavSys->HasDirtyAreasQueued() || NavSys->IsNavigationBuildInProgress();
 	}
 
 	return false;
@@ -3540,14 +3546,14 @@ bool UNavigationSystem::CanRebuildDirtyNavigation() const
 	return true;
 }
 
-bool UNavigationSystem::DoesPathIntersectBox(const FNavigationPath* Path, const FBox& Box, uint32 StartingIndex)
+bool UNavigationSystem::DoesPathIntersectBox(const FNavigationPath* Path, const FBox& Box, uint32 StartingIndex, FVector* AgentExtent)
 {
-	return Path != NULL && Path->DoesIntersectBox(Box, StartingIndex);
+	return Path != NULL && Path->DoesIntersectBox(Box, StartingIndex, NULL, AgentExtent);
 }
 
-bool UNavigationSystem::DoesPathIntersectBox(const FNavigationPath* Path, const FBox& Box, const FVector& AgentLocation, uint32 StartingIndex)
+bool UNavigationSystem::DoesPathIntersectBox(const FNavigationPath* Path, const FBox& Box, const FVector& AgentLocation, uint32 StartingIndex, FVector* AgentExtent)
 {
-	return Path != NULL && Path->DoesIntersectBox(Box, AgentLocation, StartingIndex);
+	return Path != NULL && Path->DoesIntersectBox(Box, AgentLocation, StartingIndex, NULL, AgentExtent);
 }
 
 //----------------------------------------------------------------------//
diff --git a/Engine/Source/Runtime/Engine/Private/AI/Navigation/NavigationTestingActor.cpp b/Engine/Source/Runtime/Engine/Private/AI/Navigation/NavigationTestingActor.cpp
--- a/Engine/Source/Runtime/Engine/Private/AI/Navigation/NavigationTestingActor.cpp
+++ b/Engine/Source/Runtime/Engine/Private/AI/Navigation/NavigationTestingActor.cpp
@@ -357,7 +357,7 @@ void ANavigationTestingActor::SearchPathTo(ANavigationTestingActor* Goal)
 	const float Duration = (EndTime - StartTime);
 	PathfindingTime = Duration * 1000000.0f;			// in micro seconds [us]
 	bPathIsPartial = Result.IsPartial();
-	bPathExist = Result.IsSuccessful() || Result.IsPartial();
+	bPathExist = Result.IsSuccessful();
 	bPathSearchOutOfNodes = bPathExist ? Result.Path->DidSearchReachedLimit() : false;
 	LastPath = Result.Path;
 	PathCost = bPathExist ? Result.Path->GetCost() : 0.0f;
diff --git a/Engine/Source/Runtime/Engine/Private/AI/Navigation/RecastNavMesh.cpp b/Engine/Source/Runtime/Engine/Private/AI/Navigation/RecastNavMesh.cpp
--- a/Engine/Source/Runtime/Engine/Private/AI/Navigation/RecastNavMesh.cpp
+++ b/Engine/Source/Runtime/Engine/Private/AI/Navigation/RecastNavMesh.cpp
@@ -1601,6 +1601,12 @@ FPathFindingResult ARecastNavMesh::FindPath(const FNavAgentProperties& AgentProp
 		{
 			Result.Result = RecastNavMesh->RecastNavMeshImpl->FindPath(Query.StartLocation, Query.EndLocation, *NavMeshPath,
 				*(Query.QueryFilter.Get()), Query.Owner.Get());
+
+			const bool bPartialPath = Result.IsPartial();
+			if (bPartialPath)
+			{
+				Result.Result = Query.bAllowPartialPaths ? ENavigationQueryResult::Success : ENavigationQueryResult::Fail;
+			}
 		}
 		else
 		{
diff --git a/Engine/Source/Runtime/Engine/Private/AI/Navigation/RecastNavMeshGenerator.cpp b/Engine/Source/Runtime/Engine/Private/AI/Navigation/RecastNavMeshGenerator.cpp
--- a/Engine/Source/Runtime/Engine/Private/AI/Navigation/RecastNavMeshGenerator.cpp
+++ b/Engine/Source/Runtime/Engine/Private/AI/Navigation/RecastNavMeshGenerator.cpp
@@ -1156,7 +1156,7 @@ struct FOffMeshData
 		LinkParams.Reserve(ElementsCount);
 	}
 
-	void AddLinks(const TArray<FNavigationLink>& Links, const FTransform& LocalToWorld, uint32 AgentMask)
+	void AddLinks(const TArray<FNavigationLink>& Links, const FTransform& LocalToWorld, uint32 AgentMask, float DefaultSnapHeight)
 	{
 		for (int32 LinkIndex = 0; LinkIndex < Links.Num(); ++LinkIndex)
 		{
@@ -1175,7 +1175,7 @@ struct FOffMeshData
 
 			NewInfo.type = DT_OFFMESH_CON_POINT | (Link.Direction == ENavLinkDirection::BothWays ? DT_OFFMESH_CON_BIDIR : 0);
 			NewInfo.snapRadius = Link.SnapRadius;
-			NewInfo.snapHeight = Link.bUseSnapHeight ? Link.SnapHeight : -1.0f;
+			NewInfo.snapHeight = Link.bUseSnapHeight ? Link.SnapHeight : DefaultSnapHeight;
 			NewInfo.userID = Link.UserId;
 
 			const int32* AreaID = AreaClassToIdMap->Find(Link.AreaClass ? Link.AreaClass : UNavigationSystem::GetDefaultWalkableArea());
@@ -1194,7 +1194,7 @@ struct FOffMeshData
 			LinkParams.Add(NewInfo);
 		}
 	}
-	void AddSegmentLinks(const TArray<FNavigationSegmentLink>& Links, const FTransform& LocalToWorld, uint32 AgentMask)
+	void AddSegmentLinks(const TArray<FNavigationSegmentLink>& Links, const FTransform& LocalToWorld, uint32 AgentMask, float DefaultSnapHeight)
 	{
 		for (int32 LinkIndex = 0; LinkIndex < Links.Num(); ++LinkIndex)
 		{
@@ -1215,7 +1215,7 @@ struct FOffMeshData
 
 			NewInfo.type = DT_OFFMESH_CON_SEGMENT | (Link.Direction == ENavLinkDirection::BothWays ? DT_OFFMESH_CON_BIDIR : 0);
 			NewInfo.snapRadius = Link.SnapRadius;
-			NewInfo.snapHeight = Link.bUseSnapHeight ? Link.SnapHeight : -1.0f;
+			NewInfo.snapHeight = Link.bUseSnapHeight ? Link.SnapHeight : DefaultSnapHeight;
 			NewInfo.userID = Link.UserId;
 
 			const int32* AreaID = AreaClassToIdMap->Find(Link.AreaClass);
@@ -2390,12 +2390,13 @@ bool FRecastTileGenerator::GenerateNavigationData(FNavMeshBuildContext& BuildCon
 				OffMeshData.FlagsPerArea = AdditionalCachedData.FlagsPerOffMeshLinkArea;
 				const uint32 AgentMask = (1 << TileConfig.AgentIndex);
 				const FSimpleLinkNavModifier* LinkModifier = OffmeshLinks.GetData();
+				const float DefaultSnapHeight = TileConfig.walkableClimb * TileConfig.ch;
 
 				for (int32 LinkModifierIndex = 0; LinkModifierIndex < OffmeshLinks.Num(); ++LinkModifierIndex, ++LinkModifier)
 				{
-					OffMeshData.AddLinks(LinkModifier->Links, LinkModifier->LocalToWorld, AgentMask);
+					OffMeshData.AddLinks(LinkModifier->Links, LinkModifier->LocalToWorld, AgentMask, DefaultSnapHeight);
 #if GENERATE_SEGMENT_LINKS
-					OffMeshData.AddSegmentLinks(LinkModifier->SegmentLinks, LinkModifier->LocalToWorld, AgentMask);
+					OffMeshData.AddSegmentLinks(LinkModifier->SegmentLinks, LinkModifier->LocalToWorld, AgentMask, DefaultSnapHeight);
 #endif // GENERATE_SEGMENT_LINKS
 				}
 			}
diff --git a/Engine/Source/Runtime/Engine/Private/ActorReplication.cpp b/Engine/Source/Runtime/Engine/Private/ActorReplication.cpp
--- a/Engine/Source/Runtime/Engine/Private/ActorReplication.cpp
+++ b/Engine/Source/Runtime/Engine/Private/ActorReplication.cpp
@@ -16,35 +16,59 @@
 static bool		SavedbHidden;
 static AActor*	SavedOwner;
 
+#define DEPRECATED_NET_PRIORITY -17.0f // Pick an arbitrary invalid priority, check for that
+
 float AActor::GetNetPriority(const FVector& ViewPos, const FVector& ViewDir, APlayerController* Viewer, UActorChannel* InChannel, float Time, bool bLowBandwidth)
 {
-	if ( bNetUseOwnerRelevancy && Owner)
+	return DEPRECATED_NET_PRIORITY;
+}
+
+float AActor::GetNetPriority(const FVector& ViewPos, const FVector& ViewDir, APlayerController* Viewer, AActor* ViewTarget, UActorChannel* InChannel, float Time, bool bLowBandwidth)
+{
+	PRAGMA_DISABLE_DEPRECATION_WARNINGS
+	// Call the deprecated version so subclasses will still work for a version
+	float DeprecatedValue = AActor::GetNetPriority(ViewPos, ViewDir, Viewer, InChannel, Time, bLowBandwidth);
+	PRAGMA_ENABLE_DEPRECATION_WARNINGS
+	if (DeprecatedValue != DEPRECATED_NET_PRIORITY)
 	{
-		return Owner->GetNetPriority(ViewPos, ViewDir, Viewer, InChannel, Time, bLowBandwidth);
+		return DeprecatedValue;
 	}
 
-	if ( Instigator && (Instigator == Viewer->GetPawn()) )
+	if (bNetUseOwnerRelevancy && Owner)
 	{
-		Time *= 4.f; 
+		// If we should use our owner's priority, pass it through
+		return Owner->GetNetPriority(ViewPos, ViewDir, Viewer, ViewTarget, InChannel, Time, bLowBandwidth);
 	}
-	else if ( !bHidden )
+
+	if (ViewTarget && (this == ViewTarget || Instigator == ViewTarget))
+	{
+		// If we're the view target or owned by the view target, use a high priority
+		Time *= 4.f;
+	}
+	else if (!bHidden)
 	{
 		FVector Dir = GetActorLocation() - ViewPos;
 		float DistSq = Dir.SizeSquared();
-		
-		// adjust priority based on distance and whether actor is in front of viewer
-		if ( (ViewDir | Dir) < 0.f )
+
+		// Adjust priority based on distance and whether actor is in front of viewer
+		if ((ViewDir | Dir) < 0.f)
 		{
-			if ( DistSq > NEARSIGHTTHRESHOLDSQUARED )
+			if (DistSq > NEARSIGHTTHRESHOLDSQUARED)
 			{
 				Time *= 0.2f;
 			}
-			else if ( DistSq > CLOSEPROXIMITYSQUARED )
+			else if (DistSq > CLOSEPROXIMITYSQUARED)
 			{
 				Time *= 0.4f;
 			}
 		}
-		else if ( DistSq > MEDSIGHTTHRESHOLDSQUARED )
+		else if ((DistSq < FARSIGHTTHRESHOLDSQUARED) && (FMath::Square(ViewDir | Dir) > 0.5f * DistSq))
+		{
+			// Compute the amount of distance along the ViewDir vector. Dir is not normalized
+			// Increase priority if we're being looked directly at
+			Time *= 2.f;
+		}
+		else if (DistSq > MEDSIGHTTHRESHOLDSQUARED)
 		{
 			Time *= 0.4f;
 		}
@@ -53,13 +77,12 @@ float AActor::GetNetPriority(const FVector& ViewPos, const FVector& ViewDir, APl
 	return NetPriority * Time;
 }
 
-bool AActor::GetNetDormancy(const FVector& ViewPos, const FVector& ViewDir, class APlayerController* Viewer, UActorChannel* InChannel, float Time, bool bLowBandwidth)
+bool AActor::GetNetDormancy(const FVector& ViewPos, const FVector& ViewDir, class APlayerController* Viewer, AActor* ViewTarget, UActorChannel* InChannel, float Time, bool bLowBandwidth)
 {
 	// For now, per peer dormancy is not supported
 	return false;
 }
 
-
 void AActor::PreNetReceive()
 {
 	SavedbHidden = bHidden;
@@ -122,15 +145,15 @@ void AActor::PostNetReceivePhysicState()
 	}
 }
 
-bool AActor::IsNetRelevantFor(const APlayerController* RealViewer, const AActor* Viewer, const FVector& SrcLocation) const
+bool AActor::IsNetRelevantFor(const APlayerController* RealViewer, const AActor* ViewTarget, const FVector& SrcLocation) const
 {
-	if( bAlwaysRelevant || IsOwnedBy(Viewer) || IsOwnedBy(RealViewer) || this==Viewer || Viewer==Instigator )
+	if (bAlwaysRelevant || IsOwnedBy(ViewTarget) || IsOwnedBy(RealViewer) || this == ViewTarget || ViewTarget == Instigator)
 	{
 		return true;
 	}
 	else if ( bNetUseOwnerRelevancy && Owner)
 	{
-		return Owner->IsNetRelevantFor(RealViewer, Viewer, SrcLocation);
+		return Owner->IsNetRelevantFor(RealViewer, ViewTarget, SrcLocation);
 	}
 	else if ( bOnlyRelevantToOwner )
 	{
@@ -138,7 +161,7 @@ bool AActor::IsNetRelevantFor(const APlayerController* RealViewer, const AActor*
 	}
 	else if ( RootComponent && RootComponent->AttachParent && RootComponent->AttachParent->GetOwner() && (Cast<USkeletalMeshComponent>(RootComponent->AttachParent) || (RootComponent->AttachParent->GetOwner() == Owner)) )
 	{
-		return RootComponent->AttachParent->GetOwner()->IsNetRelevantFor( RealViewer, Viewer, SrcLocation );
+		return RootComponent->AttachParent->GetOwner()->IsNetRelevantFor(RealViewer, ViewTarget, SrcLocation);
 	}
 	else if( bHidden && (!RootComponent || !RootComponent->IsCollisionEnabled()) )
 	{
diff --git a/Engine/Source/Runtime/Engine/Private/Character.cpp b/Engine/Source/Runtime/Engine/Private/Character.cpp
--- a/Engine/Source/Runtime/Engine/Private/Character.cpp
+++ b/Engine/Source/Runtime/Engine/Private/Character.cpp
@@ -210,6 +210,14 @@ void ACharacter::OnWalkingOffLedge_Implementation()
 {
 }
 
+void ACharacter::NotifyJumpApex()
+{
+	// Call delegate callback
+	if (OnReachedJumpApex.IsBound())
+	{
+		OnReachedJumpApex.Broadcast();
+	}
+}
 
 void ACharacter::Landed(const FHitResult& Hit)
 {
diff --git a/Engine/Source/Runtime/Engine/Private/CharacterMovementComponent.cpp b/Engine/Source/Runtime/Engine/Private/CharacterMovementComponent.cpp
--- a/Engine/Source/Runtime/Engine/Private/CharacterMovementComponent.cpp
+++ b/Engine/Source/Runtime/Engine/Private/CharacterMovementComponent.cpp
@@ -3755,38 +3755,6 @@ void UCharacterMovementComponent::PhysNavWalking(float deltaTime, int32 Iteratio
 		return;
 	}
 
-	UNavigationSystem* NavSys = UNavigationSystem::GetCurrent(GetWorld());
-	if (NavSys == NULL)
-	{
-		// can't find navigation system, switch to regular walking
-		SetMovementMode(MOVE_Walking);
-		return;
-	}
-
-	const ANavigationData* NavData = NULL;
-	INavAgentInterface* MyNavAgent = Cast<INavAgentInterface>(CharacterOwner);
-	float SearchRadius = 0.0f;
-	float SearchHeight = 100.0f;
-	if (MyNavAgent)
-	{
-		const FNavAgentProperties& AgentProps = MyNavAgent->GetNavAgentPropertiesRef();
-		NavData = NavSys->GetNavDataForProps(AgentProps);
-		SearchRadius = AgentProps.AgentRadius * 2.0f;
-		SearchHeight = AgentProps.AgentHeight * 0.5f;
-	}
-	if (NavData == NULL)
-	{
-		NavData = NavSys->GetMainNavData();
-	}
-
-	const ARecastNavMesh* NavMeshData = Cast<const ARecastNavMesh>(NavData);
-	if (NavMeshData == NULL)
-	{
-		// can't find navigation data for owning character, switch to regular walking
-		SetMovementMode(MOVE_Walking);
-		return;
-	}
-
 	if ((!CharacterOwner || !CharacterOwner->Controller) && !bRunPhysicsWithNoController && !HasRootMotion())
 	{
 		Acceleration = FVector::ZeroVector;
@@ -3825,8 +3793,13 @@ void UCharacterMovementComponent::PhysNavWalking(float deltaTime, int32 Iteratio
 	}
 	else
 	{
-		// scope for projection perf test
-		NavData->ProjectPoint(AdjustedDest, DestNavLocation, FVector(SearchRadius, SearchRadius, SearchHeight));
+		const bool bHasNavigationData = FindNavFloor(AdjustedDest, DestNavLocation);
+		if (!bHasNavigationData)
+		{
+			SetMovementMode(MOVE_Walking);
+			return;
+		}
+
 		CachedNavLocation = DestNavLocation;
 	}
 
@@ -3856,6 +3829,40 @@ void UCharacterMovementComponent::PhysNavWalking(float deltaTime, int32 Iteratio
 	}
 }
 
+bool UCharacterMovementComponent::FindNavFloor(const FVector& TestLocation, FNavLocation& NavFloorLocation) const
+{
+	UNavigationSystem* NavSys = UNavigationSystem::GetCurrent(GetWorld());
+	if (NavSys == nullptr)
+	{
+		return false;
+	}
+
+	const ANavigationData* NavData = nullptr;
+	INavAgentInterface* MyNavAgent = Cast<INavAgentInterface>(CharacterOwner);
+	float SearchRadius = 0.0f;
+	float SearchHeight = 100.0f;
+	if (MyNavAgent)
+	{
+		const FNavAgentProperties& AgentProps = MyNavAgent->GetNavAgentPropertiesRef();
+		NavData = NavSys->GetNavDataForProps(AgentProps);
+		SearchRadius = AgentProps.AgentRadius * 2.0f;
+		SearchHeight = AgentProps.AgentHeight * 0.5f;
+	}
+	if (NavData == nullptr)
+	{
+		NavData = NavSys->GetMainNavData();
+	}
+
+	const ARecastNavMesh* NavMeshData = Cast<const ARecastNavMesh>(NavData);
+	if (NavMeshData == nullptr)
+	{
+		return false;
+	}
+
+	NavData->ProjectPoint(TestLocation, NavFloorLocation, FVector(SearchRadius, SearchRadius, SearchHeight));
+	return true;
+}
+
 void UCharacterMovementComponent::ProjectLocationFromNavMesh(float DeltaSeconds, FVector& InOutLocation)
 {
 	if(bProjectNavMeshWalking)
@@ -3974,6 +3981,22 @@ void UCharacterMovementComponent::ProcessLanded(const FHitResult& Hit, float rem
 	}
 	if( IsFalling() )
 	{
+		if (GroundMovementMode == MOVE_NavWalking)
+		{
+			// verify navmesh projection and current floor
+			// otherwise movement will be stuck in infinite loop:
+			// navwalking -> (no navmesh) -> falling -> (standing on something) -> navwalking -> ....
+
+			const FVector TestLocation = GetActorFeetLocation();
+			FNavLocation NavLocation;
+
+			const bool bHasNavigationData = FindNavFloor(TestLocation, NavLocation);
+			if (!bHasNavigationData || NavLocation.NodeRef == INVALID_NAVNODEREF)
+			{
+				GroundMovementMode = MOVE_Walking;
+			}
+		}
+
 		SetPostLandedPhysics(Hit);
 	}
 	if (PathFollowingComp.IsValid())
@@ -4044,7 +4067,7 @@ bool UCharacterMovementComponent::TryToLeaveNavWalking()
 	if (CharacterOwner)
 	{
 		FVector CollisionFreeLocation = CharacterOwner->GetActorLocation();
-		const bool bCanTeleport = GetWorld()->FindTeleportSpot(CharacterOwner, CollisionFreeLocation, CharacterOwner->GetActorRotation());
+		bCanTeleport = GetWorld()->FindTeleportSpot(CharacterOwner, CollisionFreeLocation, CharacterOwner->GetActorRotation());
 		if (bCanTeleport)
 		{
 			CharacterOwner->SetActorLocation(CollisionFreeLocation);
diff --git a/Engine/Source/Runtime/Engine/Private/CheatManager.cpp b/Engine/Source/Runtime/Engine/Private/CheatManager.cpp
--- a/Engine/Source/Runtime/Engine/Private/CheatManager.cpp
+++ b/Engine/Source/Runtime/Engine/Private/CheatManager.cpp
@@ -36,14 +36,6 @@ UCheatManager::UCheatManager(const FObjectInitializer& ObjectInitializer)
 	bDebugCapsuleTraceComplex = false;
 }
 
-void UCheatManager::GetLifetimeReplicatedProps(TArray< FLifetimeProperty > & OutLifetimeProps) const
-{
-	Super::GetLifetimeReplicatedProps(OutLifetimeProps);
-#if ENABLE_VISUAL_LOG
-	DOREPLIFETIME(UCheatManager, bVisualLoggerActiveOnServer);
-#endif
-}
-
 void UCheatManager::FreezeFrame(float delay)
 {
 	FCanUnpause DefaultCanUnpause;
@@ -625,7 +617,7 @@ void UCheatManager::BeginDestroy()
 		FVisualLogger::Get().SetIsRecording(false);
 		FVisualLogger::Get().SetIsRecordingToFile(false);
 		bToggleAILogging = false;
-		bVisualLoggerActiveOnServer = false;
+		FVisualLogger::Get().SetIsRecordingOnServer(false);
 	}
 #endif
 	Super::BeginDestroy();
@@ -639,7 +631,6 @@ bool UCheatManager::ServerToggleAILogging_Validate()
 void UCheatManager::ServerToggleAILogging_Implementation()
 {
 #if ENABLE_VISUAL_LOG
-	UWorld *World = GetWorld();
 	if (FVisualLogger::Get().IsRecordingToFile())
 	{
 		// stop recording and dump all remaining logs in a moment
@@ -653,8 +644,24 @@ void UCheatManager::ServerToggleAILogging_Implementation()
 		bToggleAILogging = true;
 	}
 
-	bVisualLoggerActiveOnServer = bToggleAILogging;
-	GetOuterAPlayerController()->ClientMessage(FString::Printf(TEXT("OK! VisLog recording is now %s"), FVisualLogger::Get().IsRecording() ? TEXT("Enabled") : TEXT("Disabled")));
+	FVisualLogger::Get().SetIsRecordingOnServer(bToggleAILogging);
+	UWorld *World = GetWorld();
+	if (World)
+	{
+		for (FConstPlayerControllerIterator Iterator = World->GetPlayerControllerIterator(); Iterator; ++Iterator)
+		{
+			APlayerController* PC = *Iterator;
+			if (PC)
+			{
+				PC->OnServerStartedVisualLogger(bToggleAILogging);
+			}
+		}
+	}
+	else
+	{
+		GetOuterAPlayerController()->OnServerStartedVisualLogger(bToggleAILogging);
+		GetOuterAPlayerController()->ClientMessage(FString::Printf(TEXT("VisLog recording is now %s"), FVisualLogger::Get().IsRecording() ? TEXT("Enabled") : TEXT("Disabled")));
+	}
 #endif
 }
 
@@ -679,13 +686,6 @@ void UCheatManager::ToggleAILogging()
 #endif
 }
 
-void UCheatManager::OnRep_VisualLoggerActiveOnServer()
-{
-#if ENABLE_VISUAL_LOG
-	FVisualLogger::Get().SetIsRecordingOnServer(bVisualLoggerActiveOnServer);
-#endif
-}
-
 #define SAFE_TRACEINDEX_DECREASE(x) ((--x)<0)? 9:(x)
 
 #if !(UE_BUILD_SHIPPING || UE_BUILD_TEST)
diff --git a/Engine/Source/Runtime/Engine/Private/Collision/Collision.cpp b/Engine/Source/Runtime/Engine/Private/Collision/Collision.cpp
--- a/Engine/Source/Runtime/Engine/Private/Collision/Collision.cpp
+++ b/Engine/Source/Runtime/Engine/Private/Collision/Collision.cpp
@@ -36,19 +36,28 @@ bool FHitResult::NetSerialize(FArchive& Ar, class UPackageMap* Map, bool& bOutSu
 	// Most of the time the vectors are the same values, use that as an optimization
 	bool bImpactPointEqualsLocation = 0, bImpactNormalEqualsNormal = 0;
 
+	// Often times the indexes are invalid, use that as an optimization
+	bool bInvalidItem = 0, bInvalidFaceIndex = 0, bNoPenetrationDepth = 0;
+
 	if (Ar.IsSaving())
 	{
 		bImpactPointEqualsLocation = (ImpactPoint == Location);
 		bImpactNormalEqualsNormal = (ImpactNormal == Normal);
+		bInvalidItem = (Item == INDEX_NONE);
+		bInvalidFaceIndex = (FaceIndex == INDEX_NONE);
+		bNoPenetrationDepth = (PenetrationDepth == 0.0f);
 	}
 
 	// pack bitfield with flags
-	uint8 Flags = (bBlockingHit << 0) | (bStartPenetrating << 1) | (bImpactPointEqualsLocation << 2) | (bImpactNormalEqualsNormal << 3);
-	Ar.SerializeBits(&Flags, 4);
+	uint8 Flags = (bBlockingHit << 0) | (bStartPenetrating << 1) | (bImpactPointEqualsLocation << 2) | (bImpactNormalEqualsNormal << 3) | (bInvalidItem << 4) | (bInvalidFaceIndex << 5) | (bInvalidFaceIndex << 6);
+	Ar.SerializeBits(&Flags, 7); 
 	bBlockingHit = (Flags & (1 << 0)) ? 1 : 0;
 	bStartPenetrating = (Flags & (1 << 1)) ? 1 : 0;
 	bImpactPointEqualsLocation = (Flags & (1 << 2)) ? 1 : 0;
 	bImpactNormalEqualsNormal = (Flags & (1 << 3)) ? 1 : 0;
+	bInvalidItem = (Flags & (1 << 4)) ? 1 : 0;
+	bInvalidFaceIndex = (Flags & (1 << 5)) ? 1 : 0;
+	bNoPenetrationDepth = (Flags & (1 << 6)) ? 1 : 0;
 
 	Ar << Time;
 
@@ -85,13 +94,37 @@ bool FHitResult::NetSerialize(FArchive& Ar, class UPackageMap* Map, bool& bOutSu
 	TraceEnd.NetSerialize(Ar, Map, bOutSuccessLocal);
 	bOutSuccess &= bOutSuccessLocal;
 
-	Ar << PenetrationDepth;
-	Ar << Item;
+	if (!bNoPenetrationDepth)
+	{
+		Ar << PenetrationDepth;
+	}
+	else if(Ar.IsLoading())
+	{
+		PenetrationDepth = 0.0f;
+	}
+	
+	if (!bInvalidItem)
+	{
+		Ar << Item;
+	}
+	else if (Ar.IsLoading())
+	{
+		Item = INDEX_NONE;
+	}
+
 	Ar << PhysMaterial;
 	Ar << Actor;
 	Ar << Component;
 	Ar << BoneName;
-	Ar << FaceIndex;
+	if (!bInvalidFaceIndex)
+	{
+		Ar << FaceIndex;
+	}
+	else if (Ar.IsLoading())
+	{
+		FaceIndex = INDEX_NONE;
+	}
+	
 
 	return true;
 }
diff --git a/Engine/Source/Runtime/Engine/Private/DataChannel.cpp b/Engine/Source/Runtime/Engine/Private/DataChannel.cpp
--- a/Engine/Source/Runtime/Engine/Private/DataChannel.cpp
+++ b/Engine/Source/Runtime/Engine/Private/DataChannel.cpp
@@ -1451,9 +1451,9 @@ bool UActorChannel::CleanUp( const bool bForDestroy )
 				{
 					if ( CreateSubObjects[i].IsValid() )
 					{
-						Actor->OnSubobjectDestroyFromReplication( CreateSubObjects[i].Get() );
-
-						CreateSubObjects[i]->MarkPendingKill();
+						UObject *SubObject = CreateSubObjects[i].Get();
+						Actor->OnSubobjectDestroyFromReplication(SubObject);
+						SubObject->MarkPendingKill();
 					}
 				}
 
diff --git a/Engine/Source/Runtime/Engine/Private/DataReplication.cpp b/Engine/Source/Runtime/Engine/Private/DataReplication.cpp
--- a/Engine/Source/Runtime/Engine/Private/DataReplication.cpp
+++ b/Engine/Source/Runtime/Engine/Private/DataReplication.cpp
@@ -12,6 +12,8 @@
 #include "Engine/ActorChannel.h"
 #include "Engine/PackageMapClient.h"
 
+static TAutoConsoleVariable<int32> CVarMaxRPCPerNetUpdate( TEXT( "net.MaxRPCPerNetUpdate" ), 2, TEXT( "Maximum number of RPCs allowed per net update" ) );
+
 class FNetSerializeCB : public INetSerializeCB
 {
 public:
@@ -1005,7 +1007,7 @@ void FObjectReplicator::QueueRemoteFunctionBunch( UFunction* Func, FOutBunch &Bu
 		RemoteFuncInfo[InfoIdx].Calls = 0;
 	}
 	
-	if (++RemoteFuncInfo[InfoIdx].Calls > 2)
+	if (++RemoteFuncInfo[InfoIdx].Calls > CVarMaxRPCPerNetUpdate.GetValueOnGameThread())
 	{
 		UE_LOG(LogNet, Log, TEXT("Too many calls to RPC %s within a single netupdate. Skipping. %s.  LastCallTime: %.2f. CurrentTime: %.2f. LastRelevantTime: %.2f. LastUpdateTime: %.2f "), 
 			*Func->GetName(), *GetObject()->GetName(), RemoteFuncInfo[InfoIdx].LastCallTime, OwningChannel->Connection->Driver->Time, OwningChannel->RelevantTime, OwningChannel->LastUpdateTime );
diff --git a/Engine/Source/Runtime/Engine/Private/GameplayStatics.cpp b/Engine/Source/Runtime/Engine/Private/GameplayStatics.cpp
--- a/Engine/Source/Runtime/Engine/Private/GameplayStatics.cpp
+++ b/Engine/Source/Runtime/Engine/Private/GameplayStatics.cpp
@@ -592,7 +592,7 @@ UParticleSystemComponent* UGameplayStatics::SpawnEmitterAttached(UParticleSystem
 	return PSC;
 }
 
-void UGameplayStatics::BreakHitResult(const FHitResult& Hit, bool& bBlockingHit, float& Time, FVector& Location, FVector& Normal, FVector& ImpactPoint, FVector& ImpactNormal, UPhysicalMaterial*& PhysMat, AActor*& HitActor, UPrimitiveComponent*& HitComponent, FName& HitBoneName, int32& HitItem)
+void UGameplayStatics::BreakHitResult(const FHitResult& Hit, bool& bBlockingHit, float& Time, FVector& Location, FVector& Normal, FVector& ImpactPoint, FVector& ImpactNormal, UPhysicalMaterial*& PhysMat, AActor*& HitActor, UPrimitiveComponent*& HitComponent, FName& HitBoneName, int32& HitItem, FVector& TraceStart, FVector& TraceEnd)
 {
 	bBlockingHit = Hit.bBlockingHit;
 	Time = Hit.Time;
@@ -605,6 +605,8 @@ void UGameplayStatics::BreakHitResult(const FHitResult& Hit, bool& bBlockingHit,
 	HitComponent = Hit.GetComponent();
 	HitBoneName = Hit.BoneName;
 	HitItem = Hit.Item;
+	TraceStart = Hit.TraceStart;
+	TraceEnd = Hit.TraceEnd;
 }
 
 EPhysicalSurface UGameplayStatics::GetSurfaceType(const struct FHitResult& Hit)
diff --git a/Engine/Source/Runtime/Engine/Private/HUD.cpp b/Engine/Source/Runtime/Engine/Private/HUD.cpp
--- a/Engine/Source/Runtime/Engine/Private/HUD.cpp
+++ b/Engine/Source/Runtime/Engine/Private/HUD.cpp
@@ -18,6 +18,8 @@ DEFINE_LOG_CATEGORY_STATIC(LogHUD, Log, All);
 
 #define LOCTEXT_NAMESPACE "HUD"
 
+bool AHUD::bShowDebugForReticleTarget = false;
+
 AHUD::AHUD(const FObjectInitializer& ObjectInitializer)
 	: Super(ObjectInitializer)
 {
@@ -302,6 +304,11 @@ void AHUD::ShowDebugToggleSubCategory(FName Category)
 	}
 }
 
+void AHUD::ShowDebugForReticleTargetToggle(TSubclassOf<AActor> DesiredClass)
+{
+	bShowDebugForReticleTarget = !bShowDebugForReticleTarget;
+	ShowDebugTargetDesiredClass = DesiredClass;
+}
 
 bool AHUD::ShouldDisplayDebug(const FName& DebugType) const
 {
@@ -312,14 +319,36 @@ void AHUD::ShowDebugInfo(float& YL, float& YPos)
 {
 	if (DebugCanvas != nullptr )
 	{
-		if (!DebugDisplay.Contains(TEXT("Bones")))
+		FLinearColor BackgroundColor(0.f, 0.f, 0.f, 0.2f);
+		DebugCanvas->Canvas->DrawTile(0, 0, DebugCanvas->ClipX, DebugCanvas->ClipY, 0.f, 0.f, 0.f, 0.f, BackgroundColor);
+
+		FDebugDisplayInfo DisplayInfo(DebugDisplay, ToggledDebugCategories);
+
+		if (bShowDebugForReticleTarget)
 		{
-			FLinearColor BackgroundColor(0.f, 0.f, 0.f, 0.5f);
-			DebugCanvas->Canvas->DrawTile(0, 0, DebugCanvas->ClipX, DebugCanvas->ClipY, 0.f, 0.f, 0.f, 0.f, BackgroundColor);
+			FRotator CamRot; FVector CamLoc; PlayerOwner->GetPlayerViewPoint(CamLoc, CamRot);
+
+			FCollisionQueryParams TraceParams(NAME_None, true, PlayerOwner->PlayerCameraManager->ViewTarget.Target);
+			FHitResult Hit;
+			bool bHit = GetWorld()->LineTraceSingle(Hit, CamLoc, CamRot.Vector() * 100000.f + CamLoc, ECC_WorldDynamic, TraceParams);
+			if (bHit)
+			{
+				AActor* HitActor = Hit.Actor.Get();
+				if (HitActor && (ShowDebugTargetDesiredClass == NULL || HitActor->IsA(ShowDebugTargetDesiredClass)))
+				{
+					ShowDebugTargetActor = HitActor;
+				}
+			}
+		}
+		else
+		{
+			ShowDebugTargetActor = PlayerOwner->PlayerCameraManager->ViewTarget.Target;
 		}
 
-		FDebugDisplayInfo DisplayInfo(DebugDisplay, ToggledDebugCategories);
-		PlayerOwner->PlayerCameraManager->ViewTarget.Target->DisplayDebug(DebugCanvas, DisplayInfo, YL, YPos);
+		if (ShowDebugTargetActor && !ShowDebugTargetActor->IsPendingKill())
+		{
+			ShowDebugTargetActor->DisplayDebug(DebugCanvas, DisplayInfo, YL, YPos);
+		}
 
 		if (ShouldDisplayDebug(NAME_Game))
 		{
diff --git a/Engine/Source/Runtime/Engine/Private/LevelStreaming.cpp b/Engine/Source/Runtime/Engine/Private/LevelStreaming.cpp
--- a/Engine/Source/Runtime/Engine/Private/LevelStreaming.cpp
+++ b/Engine/Source/Runtime/Engine/Private/LevelStreaming.cpp
@@ -593,11 +593,8 @@ void ULevelStreaming::AsyncLevelLoadComplete(const FName& InPackageName, UPackag
 					}
 
 					// Make sure the redirector is not in the way of the new world.
-					if (WorldRedirector->GetFName() == DestinationWorld->GetFName())
-					{
-						// Pass NULL as the name to make a new unique name and GetTransientPackage() for the outer to remove it from the package.
-						WorldRedirector->Rename(NULL, GetTransientPackage(), REN_DoNotDirty | REN_DontCreateRedirectors | REN_ForceNoResetLoaders | REN_NonTransactional);
-					}
+					// Pass NULL as the name to make a new unique name and GetTransientPackage() for the outer to remove it from the package.
+					WorldRedirector->Rename(NULL, GetTransientPackage(), REN_DoNotDirty | REN_DontCreateRedirectors | REN_ForceNoResetLoaders | REN_NonTransactional);
 				}
 				else
 				{
diff --git a/Engine/Source/Runtime/Engine/Private/LocalPlayer.cpp b/Engine/Source/Runtime/Engine/Private/LocalPlayer.cpp
--- a/Engine/Source/Runtime/Engine/Private/LocalPlayer.cpp
+++ b/Engine/Source/Runtime/Engine/Private/LocalPlayer.cpp
@@ -151,6 +151,7 @@ bool FLocalPlayerContext::IsFromLocalPlayer(const AActor* ActorToTest) const
 
 ULocalPlayer::ULocalPlayer(const FObjectInitializer& ObjectInitializer)
 	: Super(ObjectInitializer)
+	, SlateOperations( FReply::Unhandled() )
 {
 	PendingLevelPlayerControllerClass = APlayerController::StaticClass();
 }
diff --git a/Engine/Source/Runtime/Engine/Private/NetConnection.cpp b/Engine/Source/Runtime/Engine/Private/NetConnection.cpp
--- a/Engine/Source/Runtime/Engine/Private/NetConnection.cpp
+++ b/Engine/Source/Runtime/Engine/Private/NetConnection.cpp
@@ -27,7 +27,7 @@ UNetConnection::UNetConnection(const FObjectInitializer& ObjectInitializer)
 :	UPlayer(ObjectInitializer)
 ,	Driver				( NULL )
 ,	PackageMap			( NULL )
-,	Viewer				( NULL )
+,	ViewTarget			( NULL )
 ,   OwningActor			( NULL )
 ,	MaxPacket			( 0 )
 ,	InternalAck			( false )
@@ -501,6 +501,8 @@ void UNetConnection::FlushNet(bool bIgnoreSimulation)
 
 		NumPaddingBits += SendBuffer.GetNumBits() - NumBitsPrePadding;
 
+		NETWORK_PROFILER(GNetworkProfiler.FlushOutgoingBunches(this));
+
 		// Send now.
 #if DO_ENABLE_NET_TEST
 		// if the connection is closing/being destroyed/etc we need to send immediately regardless of settings
