Default parameter values on parameters defined inside material functions are now obtained correctly, and fixed const correctness on various material methods.

diff --git a/Engine/Source/Runtime/Engine/Classes/Materials/Material.h b/Engine/Source/Runtime/Engine/Classes/Materials/Material.h
--- a/Engine/Source/Runtime/Engine/Classes/Materials/Material.h
+++ b/Engine/Source/Runtime/Engine/Classes/Materials/Material.h
@@ -769,9 +769,9 @@ public:
 	ENGINE_API virtual FMaterialResource* AllocateResource();
 	ENGINE_API virtual FMaterialResource* GetMaterialResource(ERHIFeatureLevel::Type InFeatureLevel, EMaterialQualityLevel::Type QualityLevel = EMaterialQualityLevel::Num) override;
 	ENGINE_API virtual const FMaterialResource* GetMaterialResource(ERHIFeatureLevel::Type InFeatureLevel, EMaterialQualityLevel::Type QualityLevel = EMaterialQualityLevel::Num) const override;
-	ENGINE_API virtual bool GetStaticSwitchParameterValue(FName ParameterName,bool &OutValue,FGuid &OutExpressionGuid) override;
-	ENGINE_API virtual bool GetStaticComponentMaskParameterValue(FName ParameterName, bool &R, bool &G, bool &B, bool &A, FGuid &OutExpressionGuid) override;
-	ENGINE_API virtual bool GetTerrainLayerWeightParameterValue(FName ParameterName, int32& OutWeightmapIndex, FGuid &OutExpressionGuid) override;
+	ENGINE_API virtual bool GetStaticSwitchParameterValue(FName ParameterName,bool &OutValue,FGuid &OutExpressionGuid) const override;
+	ENGINE_API virtual bool GetStaticComponentMaskParameterValue(FName ParameterName, bool &R, bool &G, bool &B, bool &A, FGuid &OutExpressionGuid) const override;
+	ENGINE_API virtual bool GetTerrainLayerWeightParameterValue(FName ParameterName, int32& OutWeightmapIndex, FGuid &OutExpressionGuid) const override;
 	ENGINE_API virtual bool UpdateLightmassTextureTracking() override;
 	ENGINE_API virtual bool GetTexturesInPropertyChain(EMaterialProperty InProperty, TArray<UTexture*>& OutTextures, 
 		TArray<FName>* OutTextureParamNames, class FStaticParameterSet* InStaticParameterSet) override;
@@ -892,11 +892,11 @@ public:
 	 * @return	Returns a array of parameter names used in this material for the specified expression type.
 	 */
 	template<typename ExpressionType>
-	void GetAllParameterNames(TArray<FName>& OutParameterNames, TArray<FGuid>& OutParameterIds)
+	void GetAllParameterNames(TArray<FName>& OutParameterNames, TArray<FGuid>& OutParameterIds) const
 	{
 		for (int32 ExpressionIndex = 0; ExpressionIndex < Expressions.Num(); ExpressionIndex++)
 		{
-			UMaterialExpressionMaterialFunctionCall* FunctionExpression = Cast<UMaterialExpressionMaterialFunctionCall>(Expressions[ExpressionIndex]);
+			const UMaterialExpressionMaterialFunctionCall* FunctionExpression = Cast<const UMaterialExpressionMaterialFunctionCall>(Expressions[ExpressionIndex]);
 
 			if (FunctionExpression)
 			{
@@ -907,7 +907,7 @@ public:
 			}
 			else
 			{
-				ExpressionType* ParameterExpression = Cast<ExpressionType>(Expressions[ExpressionIndex]);
+				const ExpressionType* ParameterExpression = Cast<const ExpressionType>(Expressions[ExpressionIndex]);
 
 				if (ParameterExpression)
 				{
@@ -919,12 +919,12 @@ public:
 		check(OutParameterNames.Num() == OutParameterIds.Num());
 	}
 	
-	ENGINE_API void GetAllVectorParameterNames(TArray<FName> &OutParameterNames, TArray<FGuid> &OutParameterIds);
-	ENGINE_API void GetAllScalarParameterNames(TArray<FName> &OutParameterNames, TArray<FGuid> &OutParameterIds);
-	ENGINE_API void GetAllTextureParameterNames(TArray<FName> &OutParameterNames, TArray<FGuid> &OutParameterIds);
-	ENGINE_API void GetAllFontParameterNames(TArray<FName> &OutParameterNames, TArray<FGuid> &OutParameterIds);
-	void GetAllStaticSwitchParameterNames(TArray<FName> &OutParameterNames, TArray<FGuid> &OutParameterIds);
-	void GetAllStaticComponentMaskParameterNames(TArray<FName> &OutParameterNames, TArray<FGuid> &OutParameterIds);
+	ENGINE_API void GetAllVectorParameterNames(TArray<FName> &OutParameterNames, TArray<FGuid> &OutParameterIds) const;
+	ENGINE_API void GetAllScalarParameterNames(TArray<FName> &OutParameterNames, TArray<FGuid> &OutParameterIds) const;
+	ENGINE_API void GetAllTextureParameterNames(TArray<FName> &OutParameterNames, TArray<FGuid> &OutParameterIds) const;
+	ENGINE_API void GetAllFontParameterNames(TArray<FName> &OutParameterNames, TArray<FGuid> &OutParameterIds) const;
+	void GetAllStaticSwitchParameterNames(TArray<FName> &OutParameterNames, TArray<FGuid> &OutParameterIds) const;
+	void GetAllStaticComponentMaskParameterNames(TArray<FName> &OutParameterNames, TArray<FGuid> &OutParameterIds) const;
 
 	/** Returns the material's decal blend mode, calculated from the DecalBlendMode property and what inputs are connected. */
 	uint32 GetDecalBlendMode() const { return DecalBlendMode; }
diff --git a/Engine/Source/Runtime/Engine/Classes/Materials/MaterialExpressionFontSampleParameter.h b/Engine/Source/Runtime/Engine/Classes/Materials/MaterialExpressionFontSampleParameter.h
--- a/Engine/Source/Runtime/Engine/Classes/Materials/MaterialExpressionFontSampleParameter.h
+++ b/Engine/Source/Runtime/Engine/Classes/Materials/MaterialExpressionFontSampleParameter.h
@@ -33,6 +33,9 @@ class UMaterialExpressionFontSampleParameter : public UMaterialExpressionFontSam
 #endif
 	// End UMaterialExpression Interface
 	
+	/** Return whether this is the named parameter, and fill in its value */
+	bool IsNamedParameter(FName InParameterName, UFont*& OutFontValue, int32& OutFontPage) const;
+
 	/**
 	*	Sets the default Font if none is set
 	*/
@@ -43,7 +46,7 @@ class UMaterialExpressionFontSampleParameter : public UMaterialExpressionFontSam
 		return ExpressionGUID;
 	}
 
-	void GetAllParameterNames(TArray<FName> &OutParameterNames, TArray<FGuid> &OutParameterIds);
+	void GetAllParameterNames(TArray<FName> &OutParameterNames, TArray<FGuid> &OutParameterIds) const;
 };
 
 
diff --git a/Engine/Source/Runtime/Engine/Classes/Materials/MaterialExpressionParameter.h b/Engine/Source/Runtime/Engine/Classes/Materials/MaterialExpressionParameter.h
--- a/Engine/Source/Runtime/Engine/Classes/Materials/MaterialExpressionParameter.h
+++ b/Engine/Source/Runtime/Engine/Classes/Materials/MaterialExpressionParameter.h
@@ -39,7 +39,7 @@ class UMaterialExpressionParameter : public UMaterialExpression
 	/**
 	 * Get list of parameter names for static parameter sets
 	 */
-	void GetAllParameterNames(TArray<FName> &OutParameterNames, TArray<FGuid> &OutParameterIds);
+	void GetAllParameterNames(TArray<FName> &OutParameterNames, TArray<FGuid> &OutParameterIds) const;
 };
 
 
diff --git a/Engine/Source/Runtime/Engine/Classes/Materials/MaterialExpressionScalarParameter.h b/Engine/Source/Runtime/Engine/Classes/Materials/MaterialExpressionScalarParameter.h
--- a/Engine/Source/Runtime/Engine/Classes/Materials/MaterialExpressionScalarParameter.h
+++ b/Engine/Source/Runtime/Engine/Classes/Materials/MaterialExpressionScalarParameter.h
@@ -19,6 +19,9 @@ class UMaterialExpressionScalarParameter : public UMaterialExpressionParameter
 	virtual void GetCaption(TArray<FString>& OutCaptions) const override;
 	// End UMaterialExpression Interface
 
+	/** Return whether this is the named parameter, and fill in its value */
+	bool IsNamedParameter(FName InParameterName, float& OutValue) const;
+
 #if WITH_EDITOR
 	virtual void PostEditChangeProperty(FPropertyChangedEvent& PropertyChangedEvent) override;
 #endif
diff --git a/Engine/Source/Runtime/Engine/Classes/Materials/MaterialExpressionStaticBoolParameter.h b/Engine/Source/Runtime/Engine/Classes/Materials/MaterialExpressionStaticBoolParameter.h
--- a/Engine/Source/Runtime/Engine/Classes/Materials/MaterialExpressionStaticBoolParameter.h
+++ b/Engine/Source/Runtime/Engine/Classes/Materials/MaterialExpressionStaticBoolParameter.h
@@ -23,6 +23,8 @@ public:
 #endif
 	// End UMaterialExpression Interface
 
+	/** Return whether this is the named parameter, and fill in its value */
+	bool IsNamedParameter(FName InParameterName, bool& OutValue, FGuid& OutExpressionGuid) const;
 };
 
 
diff --git a/Engine/Source/Runtime/Engine/Classes/Materials/MaterialExpressionStaticComponentMaskParameter.h b/Engine/Source/Runtime/Engine/Classes/Materials/MaterialExpressionStaticComponentMaskParameter.h
--- a/Engine/Source/Runtime/Engine/Classes/Materials/MaterialExpressionStaticComponentMaskParameter.h
+++ b/Engine/Source/Runtime/Engine/Classes/Materials/MaterialExpressionStaticComponentMaskParameter.h
@@ -32,6 +32,9 @@ public:
 	virtual int32 Compile(class FMaterialCompiler* Compiler, int32 OutputIndex, int32 MultiplexIndex) override;
 	virtual void GetCaption(TArray<FString>& OutCaptions) const override;
 	// End UMaterialExpression Interface
+
+	/** Return whether this is the named parameter, and fill in its value */
+	bool IsNamedParameter(FName InParameterName, bool& OutR, bool& OutG, bool& OutB, bool& OutA, FGuid&OutExpressionGuid) const;
 };
 
 
diff --git a/Engine/Source/Runtime/Engine/Classes/Materials/MaterialExpressionTextureSampleParameter.h b/Engine/Source/Runtime/Engine/Classes/Materials/MaterialExpressionTextureSampleParameter.h
--- a/Engine/Source/Runtime/Engine/Classes/Materials/MaterialExpressionTextureSampleParameter.h
+++ b/Engine/Source/Runtime/Engine/Classes/Materials/MaterialExpressionTextureSampleParameter.h
@@ -32,6 +32,9 @@ class ENGINE_API UMaterialExpressionTextureSampleParameter : public UMaterialExp
 #endif
 	// End UMaterialExpression Interface
 
+	/** Return whether this is the named parameter, and fill in its value */
+	bool IsNamedParameter(FName InParameterName, UTexture*& OutValue) const;
+
 	/**
 	 * Return true if the texture is a movie texture
 	 *
@@ -57,5 +60,5 @@ class ENGINE_API UMaterialExpressionTextureSampleParameter : public UMaterialExp
 		return ExpressionGUID;
 	}
 
-	void GetAllParameterNames(TArray<FName> &OutParameterNames, TArray<FGuid> &OutParameterIds);
+	void GetAllParameterNames(TArray<FName> &OutParameterNames, TArray<FGuid> &OutParameterIds) const;
 };
diff --git a/Engine/Source/Runtime/Engine/Classes/Materials/MaterialExpressionVectorParameter.h b/Engine/Source/Runtime/Engine/Classes/Materials/MaterialExpressionVectorParameter.h
--- a/Engine/Source/Runtime/Engine/Classes/Materials/MaterialExpressionVectorParameter.h
+++ b/Engine/Source/Runtime/Engine/Classes/Materials/MaterialExpressionVectorParameter.h
@@ -18,6 +18,9 @@ class UMaterialExpressionVectorParameter : public UMaterialExpressionParameter
 	virtual void GetCaption(TArray<FString>& OutCaptions) const override;
 	// End UMaterialExpression Interface
 
+	/** Return whether this is the named parameter, and fill in its value */
+	bool IsNamedParameter(FName InParameterName, FLinearColor& OutValue) const;
+
 #if WITH_EDITOR
 	virtual void PostEditChangeProperty(FPropertyChangedEvent& PropertyChangedEvent) override;
 #endif
diff --git a/Engine/Source/Runtime/Engine/Classes/Materials/MaterialFunction.h b/Engine/Source/Runtime/Engine/Classes/Materials/MaterialFunction.h
--- a/Engine/Source/Runtime/Engine/Classes/Materials/MaterialFunction.h
+++ b/Engine/Source/Runtime/Engine/Classes/Materials/MaterialFunction.h
@@ -82,22 +82,22 @@ public:
 	ENGINE_API void AppendReferencedTextures(TArray<UTexture*>& InOutTextures) const;
 
 	template<typename ExpressionType>
-	void GetAllParameterNames(TArray<FName> &OutParameterNames, TArray<FGuid> &OutParameterIds)
+	void GetAllParameterNames(TArray<FName> &OutParameterNames, TArray<FGuid> &OutParameterIds) const
 	{
 		for (int32 ExpressionIndex = 0; ExpressionIndex < FunctionExpressions.Num(); ExpressionIndex++)
 		{
-			UMaterialExpressionMaterialFunctionCall* FunctionExpression = Cast<UMaterialExpressionMaterialFunctionCall>(FunctionExpressions[ExpressionIndex]);
+			const UMaterialExpressionMaterialFunctionCall* FunctionExpression = Cast<const UMaterialExpressionMaterialFunctionCall>(FunctionExpressions[ExpressionIndex]);
 
 			if (FunctionExpression)
 			{
 				if (FunctionExpression->MaterialFunction)
 				{
-					FunctionExpression->MaterialFunction->GetAllParameterNames<ExpressionType>(OutParameterNames, OutParameterIds);
+					FunctionExpression->MaterialFunction->GetAllParameterNames<const ExpressionType>(OutParameterNames, OutParameterIds);
 				}
 			}
 			else
 			{
-				ExpressionType* ParameterExpression = Cast<ExpressionType>(FunctionExpressions[ExpressionIndex]);
+				const ExpressionType* ParameterExpression = Cast<const ExpressionType>(FunctionExpressions[ExpressionIndex]);
 
 				if (ParameterExpression)
 				{
@@ -109,6 +109,14 @@ public:
 		check(OutParameterNames.Num() == OutParameterIds.Num());
 	}
 
+	/** Searches the material function for a parameter of the given name and returns whether it was found and its value */
+	bool GetVectorParameterValue(FName ParameterName, FLinearColor& OutValue) const;
+	bool GetScalarParameterValue(FName ParameterName, float& OutValue) const;
+	bool GetTextureParameterValue(FName ParameterName, UTexture*& OutValue) const;
+	bool GetFontParameterValue(FName ParameterName, class UFont*& OutFontValue, int32& OutFontPage) const;
+	bool GetStaticSwitchParameterValue(FName ParameterName, bool& OutValue, FGuid& OutExpressionGuid) const;
+	bool GetStaticComponentMaskParameterValue(FName ParameterName, bool& OutR, bool& OutG, bool& OutB, bool& OutA, FGuid& OutExpressionGuid) const;
+
 #if WITH_EDITOR
 	ENGINE_API UMaterial* GetPreviewMaterial();
 
diff --git a/Engine/Source/Runtime/Engine/Classes/Materials/MaterialInstance.h b/Engine/Source/Runtime/Engine/Classes/Materials/MaterialInstance.h
--- a/Engine/Source/Runtime/Engine/Classes/Materials/MaterialInstance.h
+++ b/Engine/Source/Runtime/Engine/Classes/Materials/MaterialInstance.h
@@ -259,9 +259,9 @@ public:
 	virtual ENGINE_API void OverrideScalarParameterDefault(FName ParameterName, float Value, bool bOverride, ERHIFeatureLevel::Type FeatureLevel) override;
 	virtual ENGINE_API bool CheckMaterialUsage(const EMaterialUsage Usage, const bool bSkipPrim = false) override;
 	virtual ENGINE_API bool CheckMaterialUsage_Concurrent(const EMaterialUsage Usage, const bool bSkipPrim = false) const override;
-	virtual ENGINE_API bool GetStaticSwitchParameterValue(FName ParameterName, bool &OutValue, FGuid &OutExpressionGuid) override;
-	virtual ENGINE_API bool GetStaticComponentMaskParameterValue(FName ParameterName, bool &R, bool &G, bool &B, bool &A, FGuid &OutExpressionGuid) override;
-	virtual ENGINE_API bool GetTerrainLayerWeightParameterValue(FName ParameterName, int32& OutWeightmapIndex, FGuid &OutExpressionGuid) override;
+	virtual ENGINE_API bool GetStaticSwitchParameterValue(FName ParameterName, bool &OutValue, FGuid &OutExpressionGuid) const override;
+	virtual ENGINE_API bool GetStaticComponentMaskParameterValue(FName ParameterName, bool &R, bool &G, bool &B, bool &A, FGuid &OutExpressionGuid) const override;
+	virtual ENGINE_API bool GetTerrainLayerWeightParameterValue(FName ParameterName, int32& OutWeightmapIndex, FGuid &OutExpressionGuid) const override;
 	virtual ENGINE_API bool IsDependent(UMaterialInterface* TestDependency) override;
 	virtual ENGINE_API FMaterialRenderProxy* GetRenderProxy(bool Selected, bool bHovered = false) const override;
 	virtual ENGINE_API UPhysicalMaterial* GetPhysicalMaterial() const override;
diff --git a/Engine/Source/Runtime/Engine/Classes/Materials/MaterialInterface.h b/Engine/Source/Runtime/Engine/Classes/Materials/MaterialInterface.h
--- a/Engine/Source/Runtime/Engine/Classes/Materials/MaterialInterface.h
+++ b/Engine/Source/Runtime/Engine/Classes/Materials/MaterialInterface.h
@@ -325,7 +325,7 @@ public:
 	* @param	OutValue		Will contain the value of the parameter if successful
 	* @return					True if successful
 	*/
-	virtual bool GetStaticSwitchParameterValue(FName ParameterName,bool &OutValue,FGuid &OutExpressionGuid) 
+	virtual bool GetStaticSwitchParameterValue(FName ParameterName,bool &OutValue,FGuid &OutExpressionGuid) const
 		PURE_VIRTUAL(UMaterialInterface::GetStaticSwitchParameterValue,return false;);
 
 	/**
@@ -335,7 +335,7 @@ public:
 	* @param	R, G, B, A		Will contain the values of the parameter if successful
 	* @return					True if successful
 	*/
-	virtual bool GetStaticComponentMaskParameterValue(FName ParameterName, bool &R, bool &G, bool &B, bool &A, FGuid &OutExpressionGuid) 
+	virtual bool GetStaticComponentMaskParameterValue(FName ParameterName, bool &R, bool &G, bool &B, bool &A, FGuid &OutExpressionGuid) const
 		PURE_VIRTUAL(UMaterialInterface::GetStaticComponentMaskParameterValue,return false;);
 
 	/**
@@ -345,7 +345,7 @@ public:
 	* @param	OutWeightmapIndex	Will contain the values of the parameter if successful
 	* @return					True if successful
 	*/
-	virtual bool GetTerrainLayerWeightParameterValue(FName ParameterName, int32& OutWeightmapIndex, FGuid &OutExpressionGuid)
+	virtual bool GetTerrainLayerWeightParameterValue(FName ParameterName, int32& OutWeightmapIndex, FGuid &OutExpressionGuid) const
 		PURE_VIRTUAL(UMaterialInterface::GetTerrainLayerWeightParameterValue,return false;);
 
 	/** @return The material's relevance. */
diff --git a/Engine/Source/Runtime/Engine/Private/Materials/Material.cpp b/Engine/Source/Runtime/Engine/Private/Materials/Material.cpp
--- a/Engine/Source/Runtime/Engine/Private/Materials/Material.cpp
+++ b/Engine/Source/Runtime/Engine/Private/Materials/Material.cpp
@@ -1058,40 +1058,40 @@ bool UMaterial::SetMaterialUsage(bool &bNeedsRecompile, EMaterialUsage Usage, co
 	return true;
 }
 
-void UMaterial::GetAllVectorParameterNames(TArray<FName> &OutParameterNames, TArray<FGuid> &OutParameterIds)
+void UMaterial::GetAllVectorParameterNames(TArray<FName> &OutParameterNames, TArray<FGuid> &OutParameterIds) const
 {
 	OutParameterNames.Empty();
 	OutParameterIds.Empty();
 	GetAllParameterNames<UMaterialExpressionVectorParameter>(OutParameterNames, OutParameterIds);
 }
-void UMaterial::GetAllScalarParameterNames(TArray<FName> &OutParameterNames, TArray<FGuid> &OutParameterIds)
+void UMaterial::GetAllScalarParameterNames(TArray<FName> &OutParameterNames, TArray<FGuid> &OutParameterIds) const
 {
 	OutParameterNames.Empty();
 	OutParameterIds.Empty();
 	GetAllParameterNames<UMaterialExpressionScalarParameter>(OutParameterNames, OutParameterIds);
 }
-void UMaterial::GetAllTextureParameterNames(TArray<FName> &OutParameterNames, TArray<FGuid> &OutParameterIds)
+void UMaterial::GetAllTextureParameterNames(TArray<FName> &OutParameterNames, TArray<FGuid> &OutParameterIds) const
 {
 	OutParameterNames.Empty();
 	OutParameterIds.Empty();
 	GetAllParameterNames<UMaterialExpressionTextureSampleParameter>(OutParameterNames, OutParameterIds);
 }
 
-void UMaterial::GetAllFontParameterNames(TArray<FName> &OutParameterNames, TArray<FGuid> &OutParameterIds)
+void UMaterial::GetAllFontParameterNames(TArray<FName> &OutParameterNames, TArray<FGuid> &OutParameterIds) const
 {
 	OutParameterNames.Empty();
 	OutParameterIds.Empty();
 	GetAllParameterNames<UMaterialExpressionFontSampleParameter>(OutParameterNames, OutParameterIds);
 }
 
-void UMaterial::GetAllStaticSwitchParameterNames(TArray<FName> &OutParameterNames, TArray<FGuid> &OutParameterIds)
+void UMaterial::GetAllStaticSwitchParameterNames(TArray<FName> &OutParameterNames, TArray<FGuid> &OutParameterIds) const
 {
 	OutParameterNames.Empty();
 	OutParameterIds.Empty();
 	GetAllParameterNames<UMaterialExpressionStaticBoolParameter>(OutParameterNames, OutParameterIds);
 }
 
-void UMaterial::GetAllStaticComponentMaskParameterNames(TArray<FName> &OutParameterNames, TArray<FGuid> &OutParameterIds)
+void UMaterial::GetAllStaticComponentMaskParameterNames(TArray<FName> &OutParameterNames, TArray<FGuid> &OutParameterIds) const
 {
 	OutParameterNames.Empty();
 	OutParameterIds.Empty();
@@ -1246,122 +1246,146 @@ bool UMaterial::GetParameterDesc(FName ParameterName, FString& OutDesc) const
 
 bool UMaterial::GetVectorParameterValue(FName ParameterName, FLinearColor& OutValue) const
 {
-	bool bSuccess = false;
-	for(int32 ExpressionIndex = 0;ExpressionIndex < Expressions.Num();ExpressionIndex++)
+	for (const UMaterialExpression* Expression : Expressions)
 	{
-		UMaterialExpressionVectorParameter* VectorParameter =
-			Cast<UMaterialExpressionVectorParameter>(Expressions[ExpressionIndex]);
-
-		if(VectorParameter && VectorParameter->ParameterName == ParameterName)
+		if (const UMaterialExpressionVectorParameter* Parameter = Cast<const UMaterialExpressionVectorParameter>(Expression))
 		{
-			OutValue = VectorParameter->DefaultValue;
-			bSuccess = true;
-			break;
+			if (Parameter->IsNamedParameter(ParameterName, OutValue))
+			{
+				return true;
+			}
+		}
+		else if (const UMaterialExpressionMaterialFunctionCall* FunctionCall = Cast<UMaterialExpressionMaterialFunctionCall>(Expression))
+		{
+			if (FunctionCall->MaterialFunction && FunctionCall->MaterialFunction->GetVectorParameterValue(ParameterName, OutValue))
+			{
+				return true;
+			}
 		}
 	}
-	return bSuccess;
+
+	return false;
 }
 
 bool UMaterial::GetScalarParameterValue(FName ParameterName, float& OutValue) const
 {
-	bool bSuccess = false;
-	for(int32 ExpressionIndex = 0;ExpressionIndex < Expressions.Num();ExpressionIndex++)
+	for (const UMaterialExpression* Expression : Expressions)
 	{
-		UMaterialExpressionScalarParameter* ScalarParameter =
-			Cast<UMaterialExpressionScalarParameter>(Expressions[ExpressionIndex]);
-
-		if(ScalarParameter && ScalarParameter->ParameterName == ParameterName)
+		if (const UMaterialExpressionScalarParameter* Parameter = Cast<const UMaterialExpressionScalarParameter>(Expression))
 		{
-			OutValue = ScalarParameter->DefaultValue;
-			bSuccess = true;
-			break;
+			if (Parameter->IsNamedParameter(ParameterName, OutValue))
+			{
+				return true;
+			}
+		}
+		else if (const UMaterialExpressionMaterialFunctionCall* FunctionCall = Cast<UMaterialExpressionMaterialFunctionCall>(Expression))
+		{
+			if (FunctionCall->MaterialFunction && FunctionCall->MaterialFunction->GetScalarParameterValue(ParameterName, OutValue))
+			{
+				return true;
+			}
 		}
 	}
-	return bSuccess;
+
+	return false;
 }
 
 bool UMaterial::GetTextureParameterValue(FName ParameterName, UTexture*& OutValue) const
 {
-	bool bSuccess = false;
-	for(int32 ExpressionIndex = 0;ExpressionIndex < Expressions.Num();ExpressionIndex++)
+	for (const UMaterialExpression* Expression : Expressions)
 	{
-		UMaterialExpressionTextureSampleParameter* TextureSampleParameter =
-			Cast<UMaterialExpressionTextureSampleParameter>(Expressions[ExpressionIndex]);
-
-		if(TextureSampleParameter && TextureSampleParameter->ParameterName == ParameterName)
+		if (const UMaterialExpressionTextureSampleParameter* Parameter = Cast<const UMaterialExpressionTextureSampleParameter>(Expression))
 		{
-			OutValue = TextureSampleParameter->Texture;
-			bSuccess = true;
-			break;
+			if (Parameter->IsNamedParameter(ParameterName, OutValue))
+			{
+				return true;
+			}
+		}
+		else if (const UMaterialExpressionMaterialFunctionCall* FunctionCall = Cast<UMaterialExpressionMaterialFunctionCall>(Expression))
+		{
+			if (FunctionCall->MaterialFunction && FunctionCall->MaterialFunction->GetTextureParameterValue(ParameterName, OutValue))
+			{
+				return true;
+			}
 		}
 	}
-	return bSuccess;
+
+	return false;
 }
 
-bool UMaterial::GetFontParameterValue(FName ParameterName,class UFont*& OutFontValue,int32& OutFontPage) const
+bool UMaterial::GetFontParameterValue(FName ParameterName, UFont*& OutFontValue, int32& OutFontPage) const
 {
-	bool bSuccess = false;
-	for(int32 ExpressionIndex = 0;ExpressionIndex < Expressions.Num();ExpressionIndex++)
+	for (const UMaterialExpression* Expression : Expressions)
 	{
-		UMaterialExpressionFontSampleParameter* FontSampleParameter =
-			Cast<UMaterialExpressionFontSampleParameter>(Expressions[ExpressionIndex]);
-
-		if(FontSampleParameter && FontSampleParameter->ParameterName == ParameterName)
+		if (const UMaterialExpressionFontSampleParameter* Parameter = Cast<const UMaterialExpressionFontSampleParameter>(Expression))
 		{
-			OutFontValue = FontSampleParameter->Font;
-			OutFontPage = FontSampleParameter->FontTexturePage;
-			bSuccess = true;
-			break;
+			if (Parameter->IsNamedParameter(ParameterName, OutFontValue, OutFontPage))
+			{
+				return true;
+			}
+		}
+		else if (const UMaterialExpressionMaterialFunctionCall* FunctionCall = Cast<UMaterialExpressionMaterialFunctionCall>(Expression))
+		{
+			if (FunctionCall->MaterialFunction && FunctionCall->MaterialFunction->GetFontParameterValue(ParameterName, OutFontValue, OutFontPage))
+			{
+				return true;
+			}
 		}
 	}
-	return bSuccess;
+
+	return false;
 }
 
 
-bool UMaterial::GetStaticSwitchParameterValue(FName ParameterName,bool &OutValue,FGuid &OutExpressionGuid)
+bool UMaterial::GetStaticSwitchParameterValue(FName ParameterName, bool& OutValue, FGuid& OutExpressionGuid) const
 {
-	bool bSuccess = false;
-	for(int32 ExpressionIndex = 0;ExpressionIndex < Expressions.Num();ExpressionIndex++)
+	for (const UMaterialExpression* Expression : Expressions)
 	{
-		UMaterialExpressionStaticBoolParameter* StaticSwitchParameter =
-			Cast<UMaterialExpressionStaticBoolParameter>(Expressions[ExpressionIndex]);
-
-		if(StaticSwitchParameter && StaticSwitchParameter->ParameterName == ParameterName)
+		if (const UMaterialExpressionStaticBoolParameter* Parameter = Cast<const UMaterialExpressionStaticBoolParameter>(Expression))
 		{
-			OutValue = StaticSwitchParameter->DefaultValue;
-			OutExpressionGuid = StaticSwitchParameter->ExpressionGUID;
-			bSuccess = true;
-			break;
+			if (Parameter->IsNamedParameter(ParameterName, OutValue, OutExpressionGuid))
+			{
+				return true;
+			}
+		}
+		else if (const UMaterialExpressionMaterialFunctionCall* FunctionCall = Cast<UMaterialExpressionMaterialFunctionCall>(Expression))
+		{
+			if (FunctionCall->MaterialFunction && FunctionCall->MaterialFunction->GetStaticSwitchParameterValue(ParameterName, OutValue, OutExpressionGuid))
+			{
+				return true;
+			}
 		}
 	}
-	return bSuccess;
+
+	return false;
 }
 
 
-bool UMaterial::GetStaticComponentMaskParameterValue(FName ParameterName, bool &OutR, bool &OutG, bool &OutB, bool &OutA, FGuid &OutExpressionGuid)
+bool UMaterial::GetStaticComponentMaskParameterValue(FName ParameterName, bool& OutR, bool& OutG, bool& OutB, bool& OutA, FGuid& OutExpressionGuid) const
 {
-	bool bSuccess = false;
-	for(int32 ExpressionIndex = 0;ExpressionIndex < Expressions.Num();ExpressionIndex++)
+	for (const UMaterialExpression* Expression : Expressions)
 	{
-		UMaterialExpressionStaticComponentMaskParameter* StaticComponentMaskParameter =
-			Cast<UMaterialExpressionStaticComponentMaskParameter>(Expressions[ExpressionIndex]);
-
-		if(StaticComponentMaskParameter && StaticComponentMaskParameter->ParameterName == ParameterName)
+		if (const UMaterialExpressionStaticComponentMaskParameter* Parameter = Cast<const UMaterialExpressionStaticComponentMaskParameter>(Expression))
 		{
-			OutR = StaticComponentMaskParameter->DefaultR;
-			OutG = StaticComponentMaskParameter->DefaultG;
-			OutB = StaticComponentMaskParameter->DefaultB;
-			OutA = StaticComponentMaskParameter->DefaultA;
-			OutExpressionGuid = StaticComponentMaskParameter->ExpressionGUID;
-			bSuccess = true;
-			break;
+			if (Parameter->IsNamedParameter(ParameterName, OutR, OutG, OutB, OutA, OutExpressionGuid))
+			{
+				return true;
+			}
+		}
+		else if (const UMaterialExpressionMaterialFunctionCall* FunctionCall = Cast<UMaterialExpressionMaterialFunctionCall>(Expression))
+		{
+			if (FunctionCall->MaterialFunction && FunctionCall->MaterialFunction->GetStaticComponentMaskParameterValue(ParameterName, OutR, OutG, OutB, OutA, OutExpressionGuid))
+			{
+				return true;
+			}
 		}
 	}
-	return bSuccess;
+
+	return false;
 }
 
 
-bool UMaterial::GetTerrainLayerWeightParameterValue(FName ParameterName, int32& OutWeightmapIndex, FGuid &OutExpressionGuid)
+bool UMaterial::GetTerrainLayerWeightParameterValue(FName ParameterName, int32& OutWeightmapIndex, FGuid &OutExpressionGuid) const
 {
 	bool bSuccess = false;
 	OutWeightmapIndex = INDEX_NONE;
diff --git a/Engine/Source/Runtime/Engine/Private/Materials/MaterialExpressions.cpp b/Engine/Source/Runtime/Engine/Private/Materials/MaterialExpressions.cpp
--- a/Engine/Source/Runtime/Engine/Private/Materials/MaterialExpressions.cpp
+++ b/Engine/Source/Runtime/Engine/Private/Materials/MaterialExpressions.cpp
@@ -1489,6 +1489,17 @@ void UMaterialExpressionTextureSampleParameter::GetCaption(TArray<FString>& OutC
 	OutCaptions.Add(FString::Printf(TEXT("'%s'"), *ParameterName.ToString()));
 }
 
+bool UMaterialExpressionTextureSampleParameter::IsNamedParameter(FName InParameterName, UTexture*& OutValue) const
+{
+	if (InParameterName == ParameterName)
+	{
+		OutValue = Texture;
+		return true;
+	}
+
+	return false;
+}
+
 bool UMaterialExpressionTextureSampleParameter::TextureIsValid( UTexture* /*InTexture*/ )
 {
 	return false;
@@ -1505,7 +1516,7 @@ void UMaterialExpressionTextureSampleParameter::SetDefaultTexture()
 	// Does nothing in the base case...
 }
 
-void UMaterialExpressionTextureSampleParameter::GetAllParameterNames(TArray<FName> &OutParameterNames, TArray<FGuid> &OutParameterIds)
+void UMaterialExpressionTextureSampleParameter::GetAllParameterNames(TArray<FName> &OutParameterNames, TArray<FGuid> &OutParameterIds) const
 {
 	int32 CurrentSize = OutParameterNames.Num();
 	OutParameterNames.AddUnique(ParameterName);
@@ -2696,6 +2707,20 @@ void UMaterialExpressionStaticComponentMaskParameter::GetCaption(TArray<FString>
 	OutCaptions.Add(FString::Printf(TEXT("'%s'"), *ParameterName.ToString()));
 }
 
+bool UMaterialExpressionStaticComponentMaskParameter::IsNamedParameter(FName InParameterName, bool& OutR, bool& OutG, bool& OutB, bool& OutA, FGuid&OutExpressionGuid) const
+{
+	if (InParameterName == ParameterName)
+	{
+		OutR = DefaultR;
+		OutG = DefaultG;
+		OutB = DefaultB;
+		OutA = DefaultA;
+		OutExpressionGuid = ExpressionGUID;
+		return true;
+	}
+
+	return false;
+}
 
 //
 //	UMaterialExpressionTime
@@ -3531,7 +3556,7 @@ void UMaterialExpressionParameter::SetEditableName(const FString& NewName)
 
 #endif
 
-void UMaterialExpressionParameter::GetAllParameterNames(TArray<FName> &OutParameterNames, TArray<FGuid> &OutParameterIds)
+void UMaterialExpressionParameter::GetAllParameterNames(TArray<FName> &OutParameterNames, TArray<FGuid> &OutParameterIds) const
 {
 	int32 CurrentSize = OutParameterNames.Num();
 	OutParameterNames.AddUnique(ParameterName);
@@ -3574,6 +3599,17 @@ void UMaterialExpressionVectorParameter::GetCaption(TArray<FString>& OutCaptions
 	OutCaptions.Add(FString::Printf(TEXT("'%s'"), *ParameterName.ToString())); 
 }
 
+bool UMaterialExpressionVectorParameter::IsNamedParameter(FName InParameterName, FLinearColor& OutValue) const
+{
+	if (InParameterName == ParameterName)
+	{
+		OutValue = DefaultValue;
+		return true;
+	}
+
+	return false;
+}
+
 #if WITH_EDITOR
 
 void UMaterialExpressionVectorParameter::PostEditChangeProperty(FPropertyChangedEvent& PropertyChangedEvent)
@@ -3618,6 +3654,17 @@ void UMaterialExpressionScalarParameter::GetCaption(TArray<FString>& OutCaptions
 	 OutCaptions.Add(FString::Printf(TEXT("'%s'"), *ParameterName.ToString())); 
 }
 
+bool UMaterialExpressionScalarParameter::IsNamedParameter(FName InParameterName, float& OutValue) const
+{
+	if (InParameterName == ParameterName)
+	{
+		OutValue = DefaultValue;
+		return true;
+	}
+
+	return false;
+}
+
 #if WITH_EDITOR
 
 void UMaterialExpressionScalarParameter::PostEditChangeProperty(FPropertyChangedEvent& PropertyChangedEvent)
@@ -3736,6 +3783,17 @@ void UMaterialExpressionStaticBoolParameter::GetCaption(TArray<FString>& OutCapt
 	OutCaptions.Add(FString::Printf(TEXT("'%s'"), *ParameterName.ToString())); 
 }
 
+bool UMaterialExpressionStaticBoolParameter::IsNamedParameter(FName InParameterName, bool& OutValue, FGuid& OutExpressionGuid) const
+{
+	if (InParameterName == ParameterName)
+	{
+		OutValue = DefaultValue;
+		OutExpressionGuid = ExpressionGUID;
+		return true;
+	}
+
+	return false;
+}
 
 //
 //	UMaterialExpressionStaticBool
@@ -5565,6 +5623,17 @@ void UMaterialExpressionFontSampleParameter::GetCaption(TArray<FString>& OutCapt
 	OutCaptions.Add(FString::Printf(TEXT("'%s'"), *ParameterName.ToString()));
 }
 
+bool UMaterialExpressionFontSampleParameter::IsNamedParameter(FName InParameterName, UFont*& OutFontValue, int32& OutFontPage) const
+{
+	if (InParameterName == ParameterName)
+	{
+		OutFontValue = Font;
+		OutFontPage = FontTexturePage;
+		return true;
+	}
+
+	return false;
+}
 
 void UMaterialExpressionFontSampleParameter::SetDefaultFont()
 {
@@ -5593,7 +5662,7 @@ void UMaterialExpressionFontSampleParameter::SetEditableName(const FString& NewN
 }
 #endif
 
-void UMaterialExpressionFontSampleParameter::GetAllParameterNames(TArray<FName> &OutParameterNames, TArray<FGuid> &OutParameterIds)
+void UMaterialExpressionFontSampleParameter::GetAllParameterNames(TArray<FName> &OutParameterNames, TArray<FGuid> &OutParameterIds) const
 {
 	int32 CurrentSize = OutParameterNames.Num();
 	OutParameterNames.AddUnique(ParameterName);
@@ -6570,6 +6639,147 @@ bool UMaterialFunction::HasFlippedCoordinates() const
 }
 #endif //WITH_EDITORONLY_DATA
 
+bool UMaterialFunction::GetVectorParameterValue(FName ParameterName, FLinearColor& OutValue) const
+{
+	for (const UMaterialExpression* Expression : FunctionExpressions)
+	{
+		if (const UMaterialExpressionVectorParameter* Parameter = Cast<const UMaterialExpressionVectorParameter>(Expression))
+		{
+			if (Parameter->IsNamedParameter(ParameterName, OutValue))
+			{
+				return true;
+			}
+		}
+		else if (const UMaterialExpressionMaterialFunctionCall* FunctionCall = Cast<UMaterialExpressionMaterialFunctionCall>(Expression))
+		{
+			if (FunctionCall->MaterialFunction && FunctionCall->MaterialFunction->GetVectorParameterValue(ParameterName, OutValue))
+			{
+				return true;
+			}
+		}
+	}
+
+	return false;
+}
+
+bool UMaterialFunction::GetScalarParameterValue(FName ParameterName, float& OutValue) const
+{
+	for (const UMaterialExpression* Expression : FunctionExpressions)
+	{
+		if (const UMaterialExpressionScalarParameter* Parameter = Cast<const UMaterialExpressionScalarParameter>(Expression))
+		{
+			if (Parameter->IsNamedParameter(ParameterName, OutValue))
+			{
+				return true;
+			}
+		}
+		else if (const UMaterialExpressionMaterialFunctionCall* FunctionCall = Cast<UMaterialExpressionMaterialFunctionCall>(Expression))
+		{
+			if (FunctionCall->MaterialFunction && FunctionCall->MaterialFunction->GetScalarParameterValue(ParameterName, OutValue))
+			{
+				return true;
+			}
+		}
+	}
+
+	return false;
+}
+
+bool UMaterialFunction::GetTextureParameterValue(FName ParameterName, UTexture*& OutValue) const
+{
+	for (const UMaterialExpression* Expression : FunctionExpressions)
+	{
+		if (const UMaterialExpressionTextureSampleParameter* Parameter = Cast<const UMaterialExpressionTextureSampleParameter>(Expression))
+		{
+			if (Parameter->IsNamedParameter(ParameterName, OutValue))
+			{
+				return true;
+			}
+		}
+		else if (const UMaterialExpressionMaterialFunctionCall* FunctionCall = Cast<UMaterialExpressionMaterialFunctionCall>(Expression))
+		{
+			if (FunctionCall->MaterialFunction && FunctionCall->MaterialFunction->GetTextureParameterValue(ParameterName, OutValue))
+			{
+				return true;
+			}
+		}
+	}
+
+	return false;
+}
+
+bool UMaterialFunction::GetFontParameterValue(FName ParameterName, UFont*& OutFontValue, int32& OutFontPage) const
+{
+	for (const UMaterialExpression* Expression : FunctionExpressions)
+	{
+		if (const UMaterialExpressionFontSampleParameter* Parameter = Cast<const UMaterialExpressionFontSampleParameter>(Expression))
+		{
+			if (Parameter->IsNamedParameter(ParameterName, OutFontValue, OutFontPage))
+			{
+				return true;
+			}
+		}
+		else if (const UMaterialExpressionMaterialFunctionCall* FunctionCall = Cast<UMaterialExpressionMaterialFunctionCall>(Expression))
+		{
+			if (FunctionCall->MaterialFunction && FunctionCall->MaterialFunction->GetFontParameterValue(ParameterName, OutFontValue, OutFontPage))
+			{
+				return true;
+			}
+		}
+	}
+
+	return false;
+}
+
+
+bool UMaterialFunction::GetStaticSwitchParameterValue(FName ParameterName, bool& OutValue, FGuid& OutExpressionGuid) const
+{
+	for (const UMaterialExpression* Expression : FunctionExpressions)
+	{
+		if (const UMaterialExpressionStaticBoolParameter* Parameter = Cast<const UMaterialExpressionStaticBoolParameter>(Expression))
+		{
+			if (Parameter->IsNamedParameter(ParameterName, OutValue, OutExpressionGuid))
+			{
+				return true;
+			}
+		}
+		else if (const UMaterialExpressionMaterialFunctionCall* FunctionCall = Cast<UMaterialExpressionMaterialFunctionCall>(Expression))
+		{
+			if (FunctionCall->MaterialFunction && FunctionCall->MaterialFunction->GetStaticSwitchParameterValue(ParameterName, OutValue, OutExpressionGuid))
+			{
+				return true;
+			}
+		}
+	}
+
+	return false;
+}
+
+
+bool UMaterialFunction::GetStaticComponentMaskParameterValue(FName ParameterName, bool& OutR, bool& OutG, bool& OutB, bool& OutA, FGuid& OutExpressionGuid) const
+{
+	for (const UMaterialExpression* Expression : FunctionExpressions)
+	{
+		if (const UMaterialExpressionStaticComponentMaskParameter* Parameter = Cast<const UMaterialExpressionStaticComponentMaskParameter>(Expression))
+		{
+			if (Parameter->IsNamedParameter(ParameterName, OutR, OutG, OutB, OutA, OutExpressionGuid))
+			{
+				return true;
+			}
+		}
+		else if (const UMaterialExpressionMaterialFunctionCall* FunctionCall = Cast<UMaterialExpressionMaterialFunctionCall>(Expression))
+		{
+			if (FunctionCall->MaterialFunction && FunctionCall->MaterialFunction->GetStaticComponentMaskParameterValue(ParameterName, OutR, OutG, OutB, OutA, OutExpressionGuid))
+			{
+				return true;
+			}
+		}
+	}
+
+	return false;
+}
+
+
 ///////////////////////////////////////////////////////////////////////////////
 // UMaterialExpressionMaterialFunctionCall
 ///////////////////////////////////////////////////////////////////////////////
@@ -7107,6 +7317,7 @@ uint32 UMaterialExpressionMaterialFunctionCall::GetInputType(int32 InputIndex)
 }
 #endif // WITH_EDITOR
 
+
 ///////////////////////////////////////////////////////////////////////////////
 // UMaterialExpressionFunctionInput
 ///////////////////////////////////////////////////////////////////////////////
diff --git a/Engine/Source/Runtime/Engine/Private/Materials/MaterialInstance.cpp b/Engine/Source/Runtime/Engine/Private/Materials/MaterialInstance.cpp
--- a/Engine/Source/Runtime/Engine/Private/Materials/MaterialInstance.cpp
+++ b/Engine/Source/Runtime/Engine/Private/Materials/MaterialInstance.cpp
@@ -1414,31 +1414,24 @@ void UMaterialInstance::CacheShadersForResources(EShaderPlatform ShaderPlatform,
 	}
 }
 
-bool UMaterialInstance::GetStaticSwitchParameterValue(FName ParameterName, bool &OutValue,FGuid &OutExpressionGuid)
+bool UMaterialInstance::GetStaticSwitchParameterValue(FName ParameterName, bool &OutValue,FGuid &OutExpressionGuid) const
 {
 	if(ReentrantFlag)
 	{
 		return false;
 	}
 
-	bool* Value = NULL;
-	FGuid* Guid = NULL;
 	for (int32 ValueIndex = 0;ValueIndex < StaticParameters.StaticSwitchParameters.Num();ValueIndex++)
 	{
 		if (StaticParameters.StaticSwitchParameters[ValueIndex].ParameterName == ParameterName)
 		{
-			Value = &StaticParameters.StaticSwitchParameters[ValueIndex].Value;
-			Guid = &StaticParameters.StaticSwitchParameters[ValueIndex].ExpressionGUID;
-			break;
+			OutValue = StaticParameters.StaticSwitchParameters[ValueIndex].Value;
+			OutExpressionGuid = StaticParameters.StaticSwitchParameters[ValueIndex].ExpressionGUID;
+			return true;
 		}
 	}
-	if(Value)
-	{
-		OutValue = *Value;
-		OutExpressionGuid = *Guid;
-		return true;
-	}
-	else if(Parent)
+
+	if(Parent)
 	{
 		FMICReentranceGuard	Guard(this);
 		return Parent->GetStaticSwitchParameterValue(ParameterName,OutValue,OutExpressionGuid);
@@ -1449,40 +1442,27 @@ bool UMaterialInstance::GetStaticSwitchParameterValue(FName ParameterName, bool
 	}
 }
 
-bool UMaterialInstance::GetStaticComponentMaskParameterValue(FName ParameterName, bool &OutR, bool &OutG, bool &OutB, bool &OutA, FGuid &OutExpressionGuid)
+bool UMaterialInstance::GetStaticComponentMaskParameterValue(FName ParameterName, bool &OutR, bool &OutG, bool &OutB, bool &OutA, FGuid &OutExpressionGuid) const
 {
 	if(ReentrantFlag)
 	{
 		return false;
 	}
 
-	bool* R = NULL;
-	bool* G = NULL;
-	bool* B = NULL;
-	bool* A = NULL;
-	FGuid* ExpressionId = NULL;
 	for (int32 ValueIndex = 0;ValueIndex < StaticParameters.StaticComponentMaskParameters.Num();ValueIndex++)
 	{
 		if (StaticParameters.StaticComponentMaskParameters[ValueIndex].ParameterName == ParameterName)
 		{
-			R = &StaticParameters.StaticComponentMaskParameters[ValueIndex].R;
-			G = &StaticParameters.StaticComponentMaskParameters[ValueIndex].G;
-			B = &StaticParameters.StaticComponentMaskParameters[ValueIndex].B;
-			A = &StaticParameters.StaticComponentMaskParameters[ValueIndex].A;
-			ExpressionId = &StaticParameters.StaticComponentMaskParameters[ValueIndex].ExpressionGUID;
-			break;
+			OutR = StaticParameters.StaticComponentMaskParameters[ValueIndex].R;
+			OutG = StaticParameters.StaticComponentMaskParameters[ValueIndex].G;
+			OutB = StaticParameters.StaticComponentMaskParameters[ValueIndex].B;
+			OutA = StaticParameters.StaticComponentMaskParameters[ValueIndex].A;
+			OutExpressionGuid = StaticParameters.StaticComponentMaskParameters[ValueIndex].ExpressionGUID;
+			return true;
 		}
 	}
-	if(R && G && B && A)
-	{
-		OutR = *R;
-		OutG = *G;
-		OutB = *B;
-		OutA = *A;
-		OutExpressionGuid = *ExpressionId;
-		return true;
-	}
-	else if(Parent)
+
+	if(Parent)
 	{
 		FMICReentranceGuard	Guard(this);
 		return Parent->GetStaticComponentMaskParameterValue(ParameterName, OutR, OutG, OutB, OutA, OutExpressionGuid);
@@ -1493,32 +1473,24 @@ bool UMaterialInstance::GetStaticComponentMaskParameterValue(FName ParameterName
 	}
 }
 
-bool UMaterialInstance::GetTerrainLayerWeightParameterValue(FName ParameterName, int32& OutWeightmapIndex, FGuid &OutExpressionGuid)
+bool UMaterialInstance::GetTerrainLayerWeightParameterValue(FName ParameterName, int32& OutWeightmapIndex, FGuid &OutExpressionGuid) const
 {
 	if(ReentrantFlag)
 	{
 		return false;
 	}
 
-	int32 WeightmapIndex = INDEX_NONE;
-
-	FGuid* ExpressionId = NULL;
 	for (int32 ValueIndex = 0;ValueIndex < StaticParameters.TerrainLayerWeightParameters.Num();ValueIndex++)
 	{
 		if (StaticParameters.TerrainLayerWeightParameters[ValueIndex].ParameterName == ParameterName)
 		{
-			WeightmapIndex = StaticParameters.TerrainLayerWeightParameters[ValueIndex].WeightmapIndex;
-			ExpressionId = &StaticParameters.TerrainLayerWeightParameters[ValueIndex].ExpressionGUID;
-			break;
+			OutWeightmapIndex = StaticParameters.TerrainLayerWeightParameters[ValueIndex].WeightmapIndex;
+			OutExpressionGuid = StaticParameters.TerrainLayerWeightParameters[ValueIndex].ExpressionGUID;
+			return true;
 		}
 	}
-	if (WeightmapIndex >= 0)
-	{
-		OutWeightmapIndex = WeightmapIndex;
-		OutExpressionGuid = *ExpressionId;
-		return true;
-	}
-	else if(Parent)
+
+	if(Parent)
 	{
 		FMICReentranceGuard	Guard(this);
 		return Parent->GetTerrainLayerWeightParameterValue(ParameterName, OutWeightmapIndex, OutExpressionGuid);
diff --git a/Engine/Source/Runtime/Landscape/Classes/Materials/MaterialExpressionLandscapeLayerBlend.h b/Engine/Source/Runtime/Landscape/Classes/Materials/MaterialExpressionLandscapeLayerBlend.h
--- a/Engine/Source/Runtime/Landscape/Classes/Materials/MaterialExpressionLandscapeLayerBlend.h
+++ b/Engine/Source/Runtime/Landscape/Classes/Materials/MaterialExpressionLandscapeLayerBlend.h
@@ -86,7 +86,7 @@ class UMaterialExpressionLandscapeLayerBlend : public UMaterialExpression
 	/**
 	 * Get list of parameter names for static parameter sets
 	 */
-	void GetAllParameterNames(TArray<FName> &OutParameterNames, TArray<FGuid> &OutParameterIds);
+	void GetAllParameterNames(TArray<FName> &OutParameterNames, TArray<FGuid> &OutParameterIds) const;
 };
 
 
diff --git a/Engine/Source/Runtime/Landscape/Classes/Materials/MaterialExpressionLandscapeLayerSample.h b/Engine/Source/Runtime/Landscape/Classes/Materials/MaterialExpressionLandscapeLayerSample.h
--- a/Engine/Source/Runtime/Landscape/Classes/Materials/MaterialExpressionLandscapeLayerSample.h
+++ b/Engine/Source/Runtime/Landscape/Classes/Materials/MaterialExpressionLandscapeLayerSample.h
@@ -34,7 +34,7 @@ public:
 	/**
 	 * Called to get list of parameter names for static parameter sets
 	 */
-	void GetAllParameterNames(TArray<FName> &OutParameterNames, TArray<FGuid> &OutParameterIds);
+	void GetAllParameterNames(TArray<FName> &OutParameterNames, TArray<FGuid> &OutParameterIds) const;
 };
 
 
diff --git a/Engine/Source/Runtime/Landscape/Classes/Materials/MaterialExpressionLandscapeLayerSwitch.h b/Engine/Source/Runtime/Landscape/Classes/Materials/MaterialExpressionLandscapeLayerSwitch.h
--- a/Engine/Source/Runtime/Landscape/Classes/Materials/MaterialExpressionLandscapeLayerSwitch.h
+++ b/Engine/Source/Runtime/Landscape/Classes/Materials/MaterialExpressionLandscapeLayerSwitch.h
@@ -46,7 +46,7 @@ public:
 
 	LANDSCAPE_API virtual FGuid& GetParameterExpressionId() override;
 
-	void GetAllParameterNames(TArray<FName> &OutParameterNames, TArray<FGuid> &OutParameterIds);
+	void GetAllParameterNames(TArray<FName> &OutParameterNames, TArray<FGuid> &OutParameterIds) const;
 
 	/**
 	 * Do any object-specific cleanup required immediately after loading an object,
diff --git a/Engine/Source/Runtime/Landscape/Classes/Materials/MaterialExpressionLandscapeLayerWeight.h b/Engine/Source/Runtime/Landscape/Classes/Materials/MaterialExpressionLandscapeLayerWeight.h
--- a/Engine/Source/Runtime/Landscape/Classes/Materials/MaterialExpressionLandscapeLayerWeight.h
+++ b/Engine/Source/Runtime/Landscape/Classes/Materials/MaterialExpressionLandscapeLayerWeight.h
@@ -47,7 +47,7 @@ public:
 	/**
 	 * Called to get list of parameter names for static parameter sets
 	 */
-	void GetAllParameterNames(TArray<FName> &OutParameterNames, TArray<FGuid> &OutParameterIds);
+	void GetAllParameterNames(TArray<FName> &OutParameterNames, TArray<FGuid> &OutParameterIds) const;
 
 	/**
 	 * Do any object-specific cleanup required immediately after loading an object,
diff --git a/Engine/Source/Runtime/Landscape/Classes/Materials/MaterialExpressionLandscapeVisibilityMask.h b/Engine/Source/Runtime/Landscape/Classes/Materials/MaterialExpressionLandscapeVisibilityMask.h
--- a/Engine/Source/Runtime/Landscape/Classes/Materials/MaterialExpressionLandscapeVisibilityMask.h
+++ b/Engine/Source/Runtime/Landscape/Classes/Materials/MaterialExpressionLandscapeVisibilityMask.h
@@ -29,7 +29,7 @@ public:
 	/**
 	 * Called to get list of parameter names for static parameter sets
 	 */
-	void GetAllParameterNames(TArray<FName> &OutParameterNames, TArray<FGuid> &OutParameterIds);
+	void GetAllParameterNames(TArray<FName> &OutParameterNames, TArray<FGuid> &OutParameterIds) const;
 };
 
 
diff --git a/Engine/Source/Runtime/Landscape/Private/Materials/MaterialExpressionLandscapeLayerBlend.cpp b/Engine/Source/Runtime/Landscape/Private/Materials/MaterialExpressionLandscapeLayerBlend.cpp
--- a/Engine/Source/Runtime/Landscape/Private/Materials/MaterialExpressionLandscapeLayerBlend.cpp
+++ b/Engine/Source/Runtime/Landscape/Private/Materials/MaterialExpressionLandscapeLayerBlend.cpp
@@ -304,11 +304,11 @@ void UMaterialExpressionLandscapeLayerBlend::PostEditChangeProperty(FPropertyCha
 #endif
 
 
-void UMaterialExpressionLandscapeLayerBlend::GetAllParameterNames(TArray<FName> &OutParameterNames, TArray<FGuid> &OutParameterIds)
+void UMaterialExpressionLandscapeLayerBlend::GetAllParameterNames(TArray<FName> &OutParameterNames, TArray<FGuid> &OutParameterIds) const
 {
 	for (int32 LayerIdx = 0; LayerIdx<Layers.Num(); LayerIdx++)
 	{
-		FLayerBlendInput& Layer = Layers[LayerIdx];
+		const FLayerBlendInput& Layer = Layers[LayerIdx];
 
 		int32 CurrentSize = OutParameterNames.Num();
 		OutParameterNames.AddUnique(Layer.LayerName);
diff --git a/Engine/Source/Runtime/Landscape/Private/Materials/MaterialExpressionLandscapeLayerSample.cpp b/Engine/Source/Runtime/Landscape/Private/Materials/MaterialExpressionLandscapeLayerSample.cpp
--- a/Engine/Source/Runtime/Landscape/Private/Materials/MaterialExpressionLandscapeLayerSample.cpp
+++ b/Engine/Source/Runtime/Landscape/Private/Materials/MaterialExpressionLandscapeLayerSample.cpp
@@ -65,7 +65,7 @@ void UMaterialExpressionLandscapeLayerSample::GetCaption(TArray<FString>& OutCap
 }
 
 
-void UMaterialExpressionLandscapeLayerSample::GetAllParameterNames(TArray<FName> &OutParameterNames, TArray<FGuid> &OutParameterIds)
+void UMaterialExpressionLandscapeLayerSample::GetAllParameterNames(TArray<FName> &OutParameterNames, TArray<FGuid> &OutParameterIds) const
 {
 	if (!OutParameterNames.Contains(ParameterName))
 	{
diff --git a/Engine/Source/Runtime/Landscape/Private/Materials/MaterialExpressionLandscapeLayerSwitch.cpp b/Engine/Source/Runtime/Landscape/Private/Materials/MaterialExpressionLandscapeLayerSwitch.cpp
--- a/Engine/Source/Runtime/Landscape/Private/Materials/MaterialExpressionLandscapeLayerSwitch.cpp
+++ b/Engine/Source/Runtime/Landscape/Private/Materials/MaterialExpressionLandscapeLayerSwitch.cpp
@@ -116,7 +116,7 @@ FGuid& UMaterialExpressionLandscapeLayerSwitch::GetParameterExpressionId()
 }
 
 
-void UMaterialExpressionLandscapeLayerSwitch::GetAllParameterNames(TArray<FName> &OutParameterNames, TArray<FGuid> &OutParameterIds)
+void UMaterialExpressionLandscapeLayerSwitch::GetAllParameterNames(TArray<FName> &OutParameterNames, TArray<FGuid> &OutParameterIds) const
 {
 	int32 CurrentSize = OutParameterNames.Num();
 	OutParameterNames.AddUnique(ParameterName);
diff --git a/Engine/Source/Runtime/Landscape/Private/Materials/MaterialExpressionLandscapeLayerWeight.cpp b/Engine/Source/Runtime/Landscape/Private/Materials/MaterialExpressionLandscapeLayerWeight.cpp
--- a/Engine/Source/Runtime/Landscape/Private/Materials/MaterialExpressionLandscapeLayerWeight.cpp
+++ b/Engine/Source/Runtime/Landscape/Private/Materials/MaterialExpressionLandscapeLayerWeight.cpp
@@ -103,7 +103,7 @@ void UMaterialExpressionLandscapeLayerWeight::GetCaption(TArray<FString>& OutCap
 }
 
 
-void UMaterialExpressionLandscapeLayerWeight::GetAllParameterNames(TArray<FName> &OutParameterNames, TArray<FGuid> &OutParameterIds)
+void UMaterialExpressionLandscapeLayerWeight::GetAllParameterNames(TArray<FName> &OutParameterNames, TArray<FGuid> &OutParameterIds) const
 {
 	int32 CurrentSize = OutParameterNames.Num();
 	OutParameterNames.AddUnique(ParameterName);
diff --git a/Engine/Source/Runtime/Landscape/Private/Materials/MaterialExpressionLandscapeVisibilityMask.cpp b/Engine/Source/Runtime/Landscape/Private/Materials/MaterialExpressionLandscapeVisibilityMask.cpp
--- a/Engine/Source/Runtime/Landscape/Private/Materials/MaterialExpressionLandscapeVisibilityMask.cpp
+++ b/Engine/Source/Runtime/Landscape/Private/Materials/MaterialExpressionLandscapeVisibilityMask.cpp
@@ -49,7 +49,7 @@ UTexture* UMaterialExpressionLandscapeVisibilityMask::GetReferencedTexture()
 	return GEngine->WeightMapPlaceholderTexture;
 }
 
-void UMaterialExpressionLandscapeVisibilityMask::GetAllParameterNames(TArray<FName> &OutParameterNames, TArray<FGuid> &OutParameterIds)
+void UMaterialExpressionLandscapeVisibilityMask::GetAllParameterNames(TArray<FName> &OutParameterNames, TArray<FGuid> &OutParameterIds) const
 {
 	int32 CurrentSize = OutParameterNames.Num();
 	OutParameterNames.AddUnique(ParameterName);
