Fixing ConstrucObject deprecation warning

diff --git a/Engine/Source/Runtime/Engine/Private/UnrealEngine.cpp b/Engine/Source/Runtime/Engine/Private/UnrealEngine.cpp
--- a/Engine/Source/Runtime/Engine/Private/UnrealEngine.cpp
+++ b/Engine/Source/Runtime/Engine/Private/UnrealEngine.cpp
@@ -7616,7 +7616,7 @@ UNetDriver* CreateNetDriver_Local(UEngine *Engine, FWorldContext &Context, FName
 			}
 
 			// Try to create network driver.
-			NetDriver = ConstructObject<UNetDriver>(NetDriverClass);
+			NetDriver = NewObject<UNetDriver>(GetTransientPackage(), NetDriverClass);
 			check(NetDriver);
 			NetDriver->NetDriverName = NetDriver->GetFName();
 
