Reduce weak pointer accesses in sound code

diff --git a/Engine/Source/Runtime/Engine/Private/ActiveSound.cpp b/Engine/Source/Runtime/Engine/Private/ActiveSound.cpp
--- a/Engine/Source/Runtime/Engine/Private/ActiveSound.cpp
+++ b/Engine/Source/Runtime/Engine/Private/ActiveSound.cpp
@@ -163,7 +163,7 @@ void FActiveSound::UpdateWaveInstances( FAudioDevice* AudioDevice, TArray<FWaveI
 	const FListener& ClosestListener = AudioDevice->Listeners[ ClosestListenerIndex ];
 
 	// Process occlusion before shifting the sounds position
-	if (World.IsValid())
+	if (OcclusionCheckInterval > 0.f)
 	{
 		CheckOcclusion( ClosestListener.Transform.GetTranslation(), ParseParams.Transform.GetTranslation() );
 	}
@@ -178,7 +178,7 @@ void FActiveSound::UpdateWaveInstances( FAudioDevice* AudioDevice, TArray<FWaveI
 	ParseParams.HighFrequencyGain *= HighFrequencyGainMultiplier;
 
 	ParseParams.SoundClass = GetSoundClass();
-	if( ParseParams.SoundClass && ParseParams.SoundClass->Properties.bApplyAmbientVolumes && World.IsValid())
+	if( ParseParams.SoundClass && ParseParams.SoundClass->Properties.bApplyAmbientVolumes)
 	{
 		// Additional inside/outside processing for ambient sounds
 		// If we aren't in a world there is no interior volumes to be handled.
@@ -281,11 +281,12 @@ void FActiveSound::CheckOcclusion( const FVector ListenerLocation, const FVector
 {
 	// @optimization: dont do this if listener and current locations haven't changed much?
 	// also, should this use getaudiotimeseconds?
-	if( OcclusionCheckInterval > 0.0f && World->GetTimeSeconds() - LastOcclusionCheckTime > OcclusionCheckInterval && Sound->GetMaxAudibleDistance() != WORLD_MAX )
+	UWorld* WorldPtr = World.Get();
+	if (WorldPtr->GetTimeSeconds() - LastOcclusionCheckTime > OcclusionCheckInterval && Sound->GetMaxAudibleDistance() != WORLD_MAX )
 	{
-		LastOcclusionCheckTime = World->GetTimeSeconds();
+		LastOcclusionCheckTime = WorldPtr->GetTimeSeconds();
 		static FName NAME_SoundOcclusion = FName(TEXT("SoundOcclusion"));
-		bool bNowOccluded = World->LineTraceTest(SoundLocation, ListenerLocation, ECC_Visibility, FCollisionQueryParams(NAME_SoundOcclusion, true));
+		const bool bNowOccluded = WorldPtr->LineTraceTest(SoundLocation, ListenerLocation, ECC_Visibility, FCollisionQueryParams(NAME_SoundOcclusion, true));
 		if( bNowOccluded != bOccluded )
 		{
 			bOccluded = bNowOccluded;
@@ -302,12 +303,19 @@ const TCHAR* GetAWaveName(TMap<UPTRINT, struct FWaveInstance*> WaveInstances)
 
 void FActiveSound::HandleInteriorVolumes( const FListener& Listener, FSoundParseParameters& ParseParams )
 {
+	UWorld* WorldPtr = World.Get();
+	if (WorldPtr == nullptr)
+	{
+		return;
+	}
+
 	// Get the settings of the ambient sound
 	FInteriorSettings Ambient;
 	class AAudioVolume* AudioVolume;
-	if (!bGotInteriorSettings || (GIsEditor && !World->IsGameWorld()) || (ParseParams.Transform.GetTranslation() - LastLocation).SizeSquared() > KINDA_SMALL_NUMBER)
+
+	if (!bGotInteriorSettings || (GIsEditor && !WorldPtr->IsGameWorld()) || (ParseParams.Transform.GetTranslation() - LastLocation).SizeSquared() > KINDA_SMALL_NUMBER)
 	{
-		AudioVolume = World->GetAudioSettings(ParseParams.Transform.GetTranslation(), NULL, &Ambient);
+		AudioVolume = WorldPtr->GetAudioSettings(ParseParams.Transform.GetTranslation(), NULL, &Ambient);
 		LastInteriorSettings = Ambient;
 		LastAudioVolume = AudioVolume;
 		bGotInteriorSettings = true;
diff --git a/Engine/Source/Runtime/Engine/Private/SoundWave.cpp b/Engine/Source/Runtime/Engine/Private/SoundWave.cpp
--- a/Engine/Source/Runtime/Engine/Private/SoundWave.cpp
+++ b/Engine/Source/Runtime/Engine/Private/SoundWave.cpp
@@ -491,18 +491,18 @@ FWaveInstance* USoundWave::HandleStart( FActiveSound& ActiveSound, const UPTRINT
 	// Add in the subtitle if they exist
 	if (ActiveSound.bHandleSubtitles && Subtitles.Num() > 0)
 	{
-		// TODO - Audio Threading. This would need to be a call back to the main thread.
-		if (ActiveSound.AudioComponent.IsValid() && ActiveSound.AudioComponent->OnQueueSubtitles.IsBound())
+		if (UAudioComponent* AudioComponent = ActiveSound.AudioComponent.Get())
 		{
-			// intercept the subtitles if the delegate is set
-			ActiveSound.AudioComponent->OnQueueSubtitles.ExecuteIfBound( Subtitles, Duration );
-		}
-		else
-		{
-			// otherwise, pass them on to the subtitle manager for display
-			// Subtitles are hashed based on the associated sound (wave instance).
-			if( ActiveSound.World.IsValid() )
+			// TODO - Audio Threading. This would need to be a call back to the main thread.
+			if (AudioComponent->OnQueueSubtitles.IsBound())
+			{
+				// intercept the subtitles if the delegate is set
+				ActiveSound.AudioComponent->OnQueueSubtitles.ExecuteIfBound( Subtitles, Duration );
+			}
+			else
 			{
+				// otherwise, pass them on to the subtitle manager for display
+				// Subtitles are hashed based on the associated sound (wave instance).
 				FSubtitleManager::GetSubtitleManager()->QueueSubtitles( ( PTRINT )WaveInstance, ActiveSound.SubtitlePriority, bManualWordWrap, bSingleLine, Duration, Subtitles, ActiveSound.World->GetAudioTimeSeconds() );
 			}
 		}
