Fix for pooled resources getting merged back into the main pool on the same frame they were released - they are now merged NumSafeFrames later. This fixes some random rendering issues related to the DynamicMeshBuilder and the flickering of skeletalmeshes when the games is paused.

diff --git a/Engine/Source/Runtime/Engine/Public/ResourcePool.h b/Engine/Source/Runtime/Engine/Public/ResourcePool.h
--- a/Engine/Source/Runtime/Engine/Public/ResourcePool.h
+++ b/Engine/Source/Runtime/Engine/Public/ResourcePool.h
@@ -93,7 +93,7 @@ public:
 		if(!bForceDrainAll)
 		{
 			// Index of the bucket that is now old enough to be reused
-			const int32 SafeFrameIndex = GFrameNumberRenderThread % ResourcePoolPolicy::NumSafeFrames;
+			const int32 SafeFrameIndex = (GFrameNumberRenderThread + 1) % ResourcePoolPolicy::NumSafeFrames;
 			
 			// Merge the bucket into the free pool array
 			for (int32 BucketIndex = 0; BucketIndex < ResourcePoolPolicy::NumPoolBuckets; BucketIndex++)
