remove bogus comment.

diff --git a/Engine/Source/Runtime/Engine/Private/SkeletalMesh.cpp b/Engine/Source/Runtime/Engine/Private/SkeletalMesh.cpp
--- a/Engine/Source/Runtime/Engine/Private/SkeletalMesh.cpp
+++ b/Engine/Source/Runtime/Engine/Private/SkeletalMesh.cpp
@@ -3899,7 +3899,6 @@ bool ASkeletalMeshActor::CanPlayAnimation(class UAnimSequenceBase* AnimAssetBase
 		(!AnimAssetBase || SkeletalMeshComponent->SkeletalMesh->Skeleton->IsCompatible(AnimAssetBase->GetSkeleton())));
 }
 
-// @todo ps4 clang bug: this works around a PS4/clang compiler bug (optimizations)
 void SetAnimPositionInner(FName SlotName, USkeletalMeshComponent* SkeletalMeshComponent, UAnimSequence* InAnimSequence, TWeakObjectPtr<UAnimMontage>& CurrentlyPlayingMontage, float InPosition, bool bLooping)
 {
 	UAnimInstance* AnimInst = SkeletalMeshComponent->GetAnimInstance();
