#ue4 - Merge select NavWalking changes from GDC branch.

- Whether Character sweeps is determined by bGenerateOverlaps on the capsule
- Add stat counters
- Ensure Z velocity is zero after update.

diff --git a/Engine/Source/Runtime/Engine/Classes/AI/Navigation/NavigationTypes.h b/Engine/Source/Runtime/Engine/Classes/AI/Navigation/NavigationTypes.h
--- a/Engine/Source/Runtime/Engine/Classes/AI/Navigation/NavigationTypes.h
+++ b/Engine/Source/Runtime/Engine/Classes/AI/Navigation/NavigationTypes.h
@@ -347,8 +347,12 @@ struct ENGINE_API FNavAgentProperties : public FMovementProperties
 	UPROPERTY(EditAnywhere, BlueprintReadWrite, Category = MovementProperties)
 	float AgentStepHeight;
 
+	/** Scale factor to apply to height of bounds when searching for navmesh to project to when nav walking */
+	UPROPERTY(EditAnywhere, BlueprintReadWrite, Category = MovementProperties)
+	float NavWalkingSearchHeightScale;
+
 	FNavAgentProperties(float Radius = -1.f, float Height = -1.f)
-		: AgentRadius(Radius), AgentHeight(Height), AgentStepHeight(-1)
+		: AgentRadius(Radius), AgentHeight(Height), AgentStepHeight(-1), NavWalkingSearchHeightScale(0.5f)
 	{
 	}
 
diff --git a/Engine/Source/Runtime/Engine/Classes/GameFramework/CharacterMovementComponent.h b/Engine/Source/Runtime/Engine/Classes/GameFramework/CharacterMovementComponent.h
--- a/Engine/Source/Runtime/Engine/Classes/GameFramework/CharacterMovementComponent.h
+++ b/Engine/Source/Runtime/Engine/Classes/GameFramework/CharacterMovementComponent.h
@@ -1129,7 +1129,7 @@ protected:
 	 * navmesh can differ quite significantly from geometry).
 	 * Updates CachedProjectedNavMeshHitResult, access this for more info about hits.
 	 */
-	void ProjectLocationFromNavMesh(float DeltaSeconds, FVector& InOutLocation);
+	virtual void ProjectLocationFromNavMesh(float DeltaSeconds, FVector& InOutLocation);
 
 public:
 
diff --git a/Engine/Source/Runtime/Engine/Private/CharacterMovementComponent.cpp b/Engine/Source/Runtime/Engine/Private/CharacterMovementComponent.cpp
--- a/Engine/Source/Runtime/Engine/Private/CharacterMovementComponent.cpp
+++ b/Engine/Source/Runtime/Engine/Private/CharacterMovementComponent.cpp
@@ -38,6 +38,9 @@ DECLARE_CYCLE_STAT(TEXT("Char Update Acceleration"), STAT_CharUpdateAcceleration
 DECLARE_CYCLE_STAT(TEXT("Char MoveUpdateDelegate"), STAT_CharMoveUpdateDelegate, STATGROUP_Character);
 DECLARE_CYCLE_STAT(TEXT("PhysWalking"), STAT_CharPhysWalking, STATGROUP_Character);
 DECLARE_CYCLE_STAT(TEXT("PhysFalling"), STAT_CharPhysFalling, STATGROUP_Character);
+DECLARE_CYCLE_STAT(TEXT("PhysNavWalking"), STAT_CharPhysNavWalking, STATGROUP_Character);
+DECLARE_CYCLE_STAT(TEXT("NavProjectPoint"), STAT_CharNavProjectPoint, STATGROUP_Character);
+DECLARE_CYCLE_STAT(TEXT("NavProjectLocation"), STAT_CharNavProjectLocation, STATGROUP_Character);
 
 // MAGIC NUMBERS
 const float MAX_STEP_SIDE_Z = 0.08f;	// maximum z value for the normal on the vertical side of steps
@@ -3750,6 +3753,8 @@ void UCharacterMovementComponent::PhysWalking(float deltaTime, int32 Iterations)
 
 void UCharacterMovementComponent::PhysNavWalking(float deltaTime, int32 Iterations)
 {
+	SCOPE_CYCLE_COUNTER(STAT_CharPhysNavWalking);
+
 	if (deltaTime < MIN_TICK_TIME)
 	{
 		return;
@@ -3793,6 +3798,7 @@ void UCharacterMovementComponent::PhysNavWalking(float deltaTime, int32 Iteratio
 	}
 	else
 	{
+		SCOPE_CYCLE_COUNTER(STAT_CharNavProjectPoint);
 		const bool bHasNavigationData = FindNavFloor(AdjustedDest, DestNavLocation);
 		if (!bHasNavigationData)
 		{
@@ -3806,19 +3812,26 @@ void UCharacterMovementComponent::PhysNavWalking(float deltaTime, int32 Iteratio
 	if (DestNavLocation.NodeRef != INVALID_NAVNODEREF)
 	{
 		FVector NewLocation(AdjustedDest.X, AdjustedDest.Y, DestNavLocation.Location.Z);
-		ProjectLocationFromNavMesh(deltaTime, NewLocation);
+		if (bProjectNavMeshWalking)
+		{
+			SCOPE_CYCLE_COUNTER(STAT_CharNavProjectLocation);
+			ProjectLocationFromNavMesh(deltaTime, NewLocation);
+		}
 
 		FVector AdjustedDelta = NewLocation - OldLocation;
 
 		if (!AdjustedDelta.IsNearlyZero())
 		{
-			MoveUpdatedComponent(AdjustedDelta, CharacterOwner->GetActorRotation(), true);
+			const bool bSweep = UpdatedComponent ? UpdatedComponent->bGenerateOverlapEvents : false;
+			FHitResult HitResult;
+			SafeMoveUpdatedComponent(AdjustedDelta, CharacterOwner->GetActorRotation(), bSweep, HitResult);
 		}
 
 		// Update velocity to reflect actual move
 		if (!bJustTeleported && !HasRootMotion())
 		{
 			Velocity = (GetActorFeetLocation() - OldLocation) / deltaTime;
+			MaintainHorizontalGroundVelocity();
 		}
 
 		bJustTeleported = false;
@@ -3846,7 +3859,7 @@ bool UCharacterMovementComponent::FindNavFloor(const FVector& TestLocation, FNav
 		const FNavAgentProperties& AgentProps = MyNavAgent->GetNavAgentPropertiesRef();
 		NavData = NavSys->GetNavDataForProps(AgentProps);
 		SearchRadius = AgentProps.AgentRadius * 2.0f;
-		SearchHeight = AgentProps.AgentHeight * 0.5f;
+		SearchHeight = AgentProps.AgentHeight * AgentProps.NavWalkingSearchHeightScale;
 	}
 	if (NavData == nullptr)
 	{
@@ -3865,36 +3878,36 @@ bool UCharacterMovementComponent::FindNavFloor(const FVector& TestLocation, FNav
 
 void UCharacterMovementComponent::ProjectLocationFromNavMesh(float DeltaSeconds, FVector& InOutLocation)
 {
-	if(bProjectNavMeshWalking)
-	{
-		NavMeshProjectionTimer -= DeltaSeconds;
+	SCOPE_CYCLE_COUNTER(STAT_CharNavProjectLocation);
 
-		const FVector RayCastOffset(0.0f, 0.0f, CharacterOwner->GetCapsuleComponent()->GetScaledCapsuleHalfHeight() * 2.0f);
+	const FVector RayCastOffset(0.0f, 0.0f, CharacterOwner->GetCapsuleComponent()->GetScaledCapsuleHalfHeight() * 2.0f);
 
-		if(NavMeshProjectionTimer <= 0.0f)
-		{
-			// raycast to underlying mesh to allow us to more closely follow geometry
-			// we use static objects here as a best approximation to accept only objects that
-			// influence navmesh generation
-			FCollisionQueryParams Params;
-			FCollisionObjectQueryParams ObjectQueryParams(FCollisionObjectQueryParams::AllStaticObjects);
-			GetWorld()->LineTraceSingle(CachedProjectedNavMeshHitResult, InOutLocation + RayCastOffset, InOutLocation - RayCastOffset, Params, ObjectQueryParams);
+	NavMeshProjectionTimer -= DeltaSeconds;
+	if (NavMeshProjectionTimer <= 0.0f)
+	{
+		UE_LOG(LogCharacterMovement, Verbose, TEXT("ProjectLocationFromNavMesh(): %s.%s interval: %.3f velocity: %s"), *GetNameSafe(CharacterOwner), *GetName(), NavMeshProjectionInterval, *Velocity.ToString());
 
-			// discard result if we were already inside something
-			if(CachedProjectedNavMeshHitResult.bStartPenetrating)
-			{
-				CachedProjectedNavMeshHitResult.Reset();
-			}
+		// raycast to underlying mesh to allow us to more closely follow geometry
+		// we use static objects here as a best approximation to accept only objects that
+		// influence navmesh generation
+		FCollisionQueryParams Params;
+		FCollisionObjectQueryParams ObjectQueryParams(FCollisionObjectQueryParams::AllStaticObjects);
+		GetWorld()->LineTraceSingle(CachedProjectedNavMeshHitResult, InOutLocation + RayCastOffset, InOutLocation - RayCastOffset, Params, ObjectQueryParams);
 
-			NavMeshProjectionTimer = NavMeshProjectionInterval;
-		}
-		
-		// project to last plane we found
-		if(CachedProjectedNavMeshHitResult.bBlockingHit)
+		// discard result if we were already inside something
+		if (CachedProjectedNavMeshHitResult.bStartPenetrating)
 		{
-			FVector ProjectedPoint = FMath::LinePlaneIntersection(InOutLocation, InOutLocation - RayCastOffset, CachedProjectedNavMeshHitResult.Location, CachedProjectedNavMeshHitResult.Normal);
-			InOutLocation.Z = ProjectedPoint.Z;
+			CachedProjectedNavMeshHitResult.Reset();
 		}
+
+		NavMeshProjectionTimer = NavMeshProjectionInterval;
+	}
+		
+	// project to last plane we found
+	if (CachedProjectedNavMeshHitResult.bBlockingHit)
+	{
+		FVector ProjectedPoint = FMath::LinePlaneIntersection(InOutLocation, InOutLocation - RayCastOffset, CachedProjectedNavMeshHitResult.Location, CachedProjectedNavMeshHitResult.Normal);
+		InOutLocation.Z = ProjectedPoint.Z;
 	}
 }
 
