#ue4 - Remove null entries from NonDefaultPhysicsVolumeList in UWorld::RemovePhysicsVolume() in case any are accumlated.

diff --git a/Engine/Source/Runtime/Engine/Private/World.cpp b/Engine/Source/Runtime/Engine/Private/World.cpp
--- a/Engine/Source/Runtime/Engine/Private/World.cpp
+++ b/Engine/Source/Runtime/Engine/Private/World.cpp
@@ -3275,7 +3275,9 @@ void UWorld::AddPhysicsVolume(APhysicsVolume* Volume)
 
 void UWorld::RemovePhysicsVolume(APhysicsVolume* Volume)
 {
-	NonDefaultPhysicsVolumeList.Remove(Volume);
+	NonDefaultPhysicsVolumeList.RemoveSwap(Volume);
+	// Also remove null entries that may accumulate as items become invalidated
+	NonDefaultPhysicsVolumeList.RemoveSwap(nullptr);
 }
 
 FConstPhysicsVolumeIterator UWorld::GetNonDefaultPhysicsVolumeIterator() const
