Merging CL#2452632 using UT_To_UE4
#UE4 Expanded FPSChart to record frame rates up to 120

diff --git a/Engine/Source/Runtime/Engine/Private/ChartCreation.cpp b/Engine/Source/Runtime/Engine/Private/ChartCreation.cpp
--- a/Engine/Source/Runtime/Engine/Private/ChartCreation.cpp
+++ b/Engine/Source/Runtime/Engine/Private/ChartCreation.cpp
@@ -211,10 +211,70 @@ void UEngine::TickFPSChart( float DeltaSeconds )
 			GFPSChart[ STAT_FPSChart_55_60 - STAT_FPSChartFirstStat ].Count++;
 			GFPSChart[ STAT_FPSChart_55_60 - STAT_FPSChartFirstStat ].CummulativeTime += DeltaSeconds;
 		}
+		else if (CurrentFPS < 65)
+		{
+			GFPSChart[STAT_FPSChart_60_65 - STAT_FPSChartFirstStat].Count++;
+			GFPSChart[STAT_FPSChart_60_65 - STAT_FPSChartFirstStat].CummulativeTime += DeltaSeconds;
+		}
+		else if (CurrentFPS < 70)
+		{
+			GFPSChart[STAT_FPSChart_65_70 - STAT_FPSChartFirstStat].Count++;
+			GFPSChart[STAT_FPSChart_65_70 - STAT_FPSChartFirstStat].CummulativeTime += DeltaSeconds;
+		}
+		else if (CurrentFPS < 75)
+		{
+			GFPSChart[STAT_FPSChart_70_75 - STAT_FPSChartFirstStat].Count++;
+			GFPSChart[STAT_FPSChart_70_75 - STAT_FPSChartFirstStat].CummulativeTime += DeltaSeconds;
+		}
+		else if (CurrentFPS < 80)
+		{
+			GFPSChart[STAT_FPSChart_75_80 - STAT_FPSChartFirstStat].Count++;
+			GFPSChart[STAT_FPSChart_75_80 - STAT_FPSChartFirstStat].CummulativeTime += DeltaSeconds;
+		}
+		else if (CurrentFPS < 85)
+		{
+			GFPSChart[STAT_FPSChart_80_85 - STAT_FPSChartFirstStat].Count++;
+			GFPSChart[STAT_FPSChart_80_85 - STAT_FPSChartFirstStat].CummulativeTime += DeltaSeconds;
+		}
+		else if (CurrentFPS < 90)
+		{
+			GFPSChart[STAT_FPSChart_85_90 - STAT_FPSChartFirstStat].Count++;
+			GFPSChart[STAT_FPSChart_85_90 - STAT_FPSChartFirstStat].CummulativeTime += DeltaSeconds;
+		}
+		else if (CurrentFPS < 95)
+		{
+			GFPSChart[STAT_FPSChart_90_95 - STAT_FPSChartFirstStat].Count++;
+			GFPSChart[STAT_FPSChart_90_95 - STAT_FPSChartFirstStat].CummulativeTime += DeltaSeconds;
+		}
+		else if (CurrentFPS < 100)
+		{
+			GFPSChart[STAT_FPSChart_95_100 - STAT_FPSChartFirstStat].Count++;
+			GFPSChart[STAT_FPSChart_95_100 - STAT_FPSChartFirstStat].CummulativeTime += DeltaSeconds;
+		}
+		else if (CurrentFPS < 105)
+		{
+			GFPSChart[STAT_FPSChart_100_105 - STAT_FPSChartFirstStat].Count++;
+			GFPSChart[STAT_FPSChart_100_105 - STAT_FPSChartFirstStat].CummulativeTime += DeltaSeconds;
+		}
+		else if (CurrentFPS < 110)
+		{
+			GFPSChart[STAT_FPSChart_105_110 - STAT_FPSChartFirstStat].Count++;
+			GFPSChart[STAT_FPSChart_105_110 - STAT_FPSChartFirstStat].CummulativeTime += DeltaSeconds;
+		}
+		else if (CurrentFPS < 115)
+		{
+			GFPSChart[STAT_FPSChart_110_115 - STAT_FPSChartFirstStat].Count++;
+			GFPSChart[STAT_FPSChart_110_115 - STAT_FPSChartFirstStat].CummulativeTime += DeltaSeconds;
+		}
+		else if (CurrentFPS < 120)
+		{
+			GFPSChart[STAT_FPSChart_115_120 - STAT_FPSChartFirstStat].Count++;
+			GFPSChart[STAT_FPSChart_115_120 - STAT_FPSChartFirstStat].CummulativeTime += DeltaSeconds;
+		}
 		else
 		{
-			GFPSChart[ STAT_FPSChart_60_INF - STAT_FPSChartFirstStat ].Count++;
-			GFPSChart[ STAT_FPSChart_60_INF - STAT_FPSChartFirstStat ].CummulativeTime += DeltaSeconds;
+			GFPSChart[ STAT_FPSChart_120_INF - STAT_FPSChartFirstStat ].Count++;
+			GFPSChart[ STAT_FPSChart_120_INF - STAT_FPSChartFirstStat ].CummulativeTime += DeltaSeconds;
 		}
 
 		GTotalGPUTime += FPlatformTime::ToSeconds(LocalGPUFrameTime);
@@ -405,9 +465,9 @@ void UEngine::DumpFPSChartToLog( float TotalTime, float DeltaTime, int32 NumFram
 		// Figure out bucket range.
 		const int32 StartFPS	= BucketIndex * 5;
 		int32 EndFPS		= StartFPS + 5;
-		if( BucketIndex + STAT_FPSChartFirstStat == STAT_FPSChart_60_INF )
+		if( BucketIndex + STAT_FPSChartFirstStat == STAT_FPSChart_120_INF )
 		{
-			EndFPS = 99;
+			EndFPS = 999;
 		}
 
 		// Keep track of time spent at 30+ FPS.
@@ -641,9 +701,9 @@ void UEngine::DumpFPSChartToStatsLog( float TotalTime, float DeltaTime, int32 Nu
 			// Figure out bucket range.
 			const int32 StartFPS	= BucketIndex * 5;
 			int32 EndFPS		= StartFPS + 5;
-			if( BucketIndex + STAT_FPSChartFirstStat == STAT_FPSChart_60_INF )
+			if( BucketIndex + STAT_FPSChartFirstStat == STAT_FPSChart_120_INF )
 			{
-				EndFPS = 99;
+				EndFPS = 999;
 			}
 
 			// Keep track of time spent at 30+ FPS.
@@ -763,9 +823,9 @@ void UEngine::DumpFPSChartToHTML( float TotalTime, float DeltaTime, int32 NumFra
 			// Figure out bucket range.
 			const int32 StartFPS	= BucketIndex * 5;
 			int32 EndFPS		= StartFPS + 5;
-			if( BucketIndex + STAT_FPSChartFirstStat == STAT_FPSChart_60_INF )
+			if( BucketIndex + STAT_FPSChartFirstStat == STAT_FPSChart_120_INF )
 			{
-				EndFPS = 99;
+				EndFPS = 999;
 			}
 
 			// Keep track of time spent at 30+ FPS.
diff --git a/Engine/Source/Runtime/Engine/Public/ChartCreation.h b/Engine/Source/Runtime/Engine/Public/ChartCreation.h
--- a/Engine/Source/Runtime/Engine/Public/ChartCreation.h
+++ b/Engine/Source/Runtime/Engine/Public/ChartCreation.h
@@ -40,7 +40,19 @@ enum FPSChartStats
 	STAT_FPSChart_45_50,
 	STAT_FPSChart_50_55,
 	STAT_FPSChart_55_60,
-	STAT_FPSChart_60_INF,
+	STAT_FPSChart_60_65,
+	STAT_FPSChart_65_70,
+	STAT_FPSChart_70_75,
+	STAT_FPSChart_75_80,
+	STAT_FPSChart_80_85,
+	STAT_FPSChart_85_90,
+	STAT_FPSChart_90_95,
+	STAT_FPSChart_95_100,
+	STAT_FPSChart_100_105,
+	STAT_FPSChart_105_110,
+	STAT_FPSChart_110_115,
+	STAT_FPSChart_115_120,
+	STAT_FPSChart_120_INF,
 	STAT_FPSChartLastBucketStat,
 	STAT_FPSChart_30Plus,
 	STAT_FPSChart_UnaccountedTime,
