Merging CL#2458915 using UT_To_UE4
#UE4 FPSChart now buckets frame rate by intervals of 5 until 30fps, then it buckets at intervals of 10. This will reduce the number of buckets and make the chart cleaner.

diff --git a/Engine/Source/Runtime/Engine/Private/ChartCreation.cpp b/Engine/Source/Runtime/Engine/Private/ChartCreation.cpp
--- a/Engine/Source/Runtime/Engine/Private/ChartCreation.cpp
+++ b/Engine/Source/Runtime/Engine/Private/ChartCreation.cpp
@@ -181,95 +181,50 @@ void UEngine::TickFPSChart( float DeltaSeconds )
 			GFPSChart[ STAT_FPSChart_25_30 - STAT_FPSChartFirstStat ].Count++;
 			GFPSChart[ STAT_FPSChart_25_30 - STAT_FPSChartFirstStat ].CummulativeTime += DeltaSeconds;
 		}
-		else if( CurrentFPS < 35 )
-		{
-			GFPSChart[ STAT_FPSChart_30_35 - STAT_FPSChartFirstStat ].Count++;
-			GFPSChart[ STAT_FPSChart_30_35 - STAT_FPSChartFirstStat ].CummulativeTime += DeltaSeconds;
-		}
 		else if( CurrentFPS < 40 )
 		{
-			GFPSChart[ STAT_FPSChart_35_40 - STAT_FPSChartFirstStat ].Count++;
-			GFPSChart[ STAT_FPSChart_35_40 - STAT_FPSChartFirstStat ].CummulativeTime += DeltaSeconds;
-		}
-		else if( CurrentFPS < 45 )
-		{
-			GFPSChart[ STAT_FPSChart_40_45 - STAT_FPSChartFirstStat ].Count++;
-			GFPSChart[ STAT_FPSChart_40_45 - STAT_FPSChartFirstStat ].CummulativeTime += DeltaSeconds;
+			GFPSChart[ STAT_FPSChart_30_40 - STAT_FPSChartFirstStat ].Count++;
+			GFPSChart[ STAT_FPSChart_30_40 - STAT_FPSChartFirstStat ].CummulativeTime += DeltaSeconds;
 		}
 		else if( CurrentFPS < 50 )
 		{
-			GFPSChart[ STAT_FPSChart_45_50 - STAT_FPSChartFirstStat ].Count++;
-			GFPSChart[ STAT_FPSChart_45_50 - STAT_FPSChartFirstStat ].CummulativeTime += DeltaSeconds;
-		}
-		else if( CurrentFPS < 55 )
-		{
-			GFPSChart[ STAT_FPSChart_50_55 - STAT_FPSChartFirstStat ].Count++;
-			GFPSChart[ STAT_FPSChart_50_55 - STAT_FPSChartFirstStat ].CummulativeTime += DeltaSeconds;
+			GFPSChart[ STAT_FPSChart_40_50 - STAT_FPSChartFirstStat ].Count++;
+			GFPSChart[ STAT_FPSChart_40_50 - STAT_FPSChartFirstStat ].CummulativeTime += DeltaSeconds;
 		}
 		else if( CurrentFPS < 60 )
 		{
-			GFPSChart[ STAT_FPSChart_55_60 - STAT_FPSChartFirstStat ].Count++;
-			GFPSChart[ STAT_FPSChart_55_60 - STAT_FPSChartFirstStat ].CummulativeTime += DeltaSeconds;
-		}
-		else if (CurrentFPS < 65)
-		{
-			GFPSChart[STAT_FPSChart_60_65 - STAT_FPSChartFirstStat].Count++;
-			GFPSChart[STAT_FPSChart_60_65 - STAT_FPSChartFirstStat].CummulativeTime += DeltaSeconds;
+			GFPSChart[ STAT_FPSChart_50_60 - STAT_FPSChartFirstStat ].Count++;
+			GFPSChart[ STAT_FPSChart_50_60 - STAT_FPSChartFirstStat ].CummulativeTime += DeltaSeconds;
 		}
 		else if (CurrentFPS < 70)
 		{
-			GFPSChart[STAT_FPSChart_65_70 - STAT_FPSChartFirstStat].Count++;
-			GFPSChart[STAT_FPSChart_65_70 - STAT_FPSChartFirstStat].CummulativeTime += DeltaSeconds;
-		}
-		else if (CurrentFPS < 75)
-		{
-			GFPSChart[STAT_FPSChart_70_75 - STAT_FPSChartFirstStat].Count++;
-			GFPSChart[STAT_FPSChart_70_75 - STAT_FPSChartFirstStat].CummulativeTime += DeltaSeconds;
+			GFPSChart[STAT_FPSChart_60_70 - STAT_FPSChartFirstStat].Count++;
+			GFPSChart[STAT_FPSChart_60_70 - STAT_FPSChartFirstStat].CummulativeTime += DeltaSeconds;
 		}
 		else if (CurrentFPS < 80)
 		{
-			GFPSChart[STAT_FPSChart_75_80 - STAT_FPSChartFirstStat].Count++;
-			GFPSChart[STAT_FPSChart_75_80 - STAT_FPSChartFirstStat].CummulativeTime += DeltaSeconds;
-		}
-		else if (CurrentFPS < 85)
-		{
-			GFPSChart[STAT_FPSChart_80_85 - STAT_FPSChartFirstStat].Count++;
-			GFPSChart[STAT_FPSChart_80_85 - STAT_FPSChartFirstStat].CummulativeTime += DeltaSeconds;
+			GFPSChart[STAT_FPSChart_70_80 - STAT_FPSChartFirstStat].Count++;
+			GFPSChart[STAT_FPSChart_70_80 - STAT_FPSChartFirstStat].CummulativeTime += DeltaSeconds;
 		}
 		else if (CurrentFPS < 90)
 		{
-			GFPSChart[STAT_FPSChart_85_90 - STAT_FPSChartFirstStat].Count++;
-			GFPSChart[STAT_FPSChart_85_90 - STAT_FPSChartFirstStat].CummulativeTime += DeltaSeconds;
-		}
-		else if (CurrentFPS < 95)
-		{
-			GFPSChart[STAT_FPSChart_90_95 - STAT_FPSChartFirstStat].Count++;
-			GFPSChart[STAT_FPSChart_90_95 - STAT_FPSChartFirstStat].CummulativeTime += DeltaSeconds;
+			GFPSChart[STAT_FPSChart_80_90 - STAT_FPSChartFirstStat].Count++;
+			GFPSChart[STAT_FPSChart_80_90 - STAT_FPSChartFirstStat].CummulativeTime += DeltaSeconds;
 		}
 		else if (CurrentFPS < 100)
 		{
-			GFPSChart[STAT_FPSChart_95_100 - STAT_FPSChartFirstStat].Count++;
-			GFPSChart[STAT_FPSChart_95_100 - STAT_FPSChartFirstStat].CummulativeTime += DeltaSeconds;
-		}
-		else if (CurrentFPS < 105)
-		{
-			GFPSChart[STAT_FPSChart_100_105 - STAT_FPSChartFirstStat].Count++;
-			GFPSChart[STAT_FPSChart_100_105 - STAT_FPSChartFirstStat].CummulativeTime += DeltaSeconds;
+			GFPSChart[STAT_FPSChart_90_100 - STAT_FPSChartFirstStat].Count++;
+			GFPSChart[STAT_FPSChart_90_100 - STAT_FPSChartFirstStat].CummulativeTime += DeltaSeconds;
 		}
 		else if (CurrentFPS < 110)
 		{
-			GFPSChart[STAT_FPSChart_105_110 - STAT_FPSChartFirstStat].Count++;
-			GFPSChart[STAT_FPSChart_105_110 - STAT_FPSChartFirstStat].CummulativeTime += DeltaSeconds;
-		}
-		else if (CurrentFPS < 115)
-		{
-			GFPSChart[STAT_FPSChart_110_115 - STAT_FPSChartFirstStat].Count++;
-			GFPSChart[STAT_FPSChart_110_115 - STAT_FPSChartFirstStat].CummulativeTime += DeltaSeconds;
+			GFPSChart[STAT_FPSChart_100_110 - STAT_FPSChartFirstStat].Count++;
+			GFPSChart[STAT_FPSChart_100_110 - STAT_FPSChartFirstStat].CummulativeTime += DeltaSeconds;
 		}
 		else if (CurrentFPS < 120)
 		{
-			GFPSChart[STAT_FPSChart_115_120 - STAT_FPSChartFirstStat].Count++;
-			GFPSChart[STAT_FPSChart_115_120 - STAT_FPSChartFirstStat].CummulativeTime += DeltaSeconds;
+			GFPSChart[STAT_FPSChart_110_120 - STAT_FPSChartFirstStat].Count++;
+			GFPSChart[STAT_FPSChart_110_120 - STAT_FPSChartFirstStat].CummulativeTime += DeltaSeconds;
 		}
 		else
 		{
@@ -462,9 +417,23 @@ void UEngine::DumpFPSChartToLog( float TotalTime, float DeltaTime, int32 NumFram
 		const float BucketTimePercentage  = 100.f * GFPSChart[BucketIndex].CummulativeTime / TotalTime;
 		const float BucketFramePercentage = 100.f * GFPSChart[BucketIndex].Count / NumFrames;
 
-		// Figure out bucket range.
-		const int32 StartFPS	= BucketIndex * 5;
-		int32 EndFPS		= StartFPS + 5;
+		// Figure out bucket range. Buckets start at 5 frame intervals then change to 10.
+		int32 StartFPS = 0;
+		int32 EndFPS = 0;
+		if ( BucketIndex < STAT_FPSChart_30_40 )
+		{
+			// Still incrementing by 5
+			StartFPS = BucketIndex * 5;
+			EndFPS = StartFPS + 5;
+		}
+		else
+		{
+			// Now incrementing by 10
+			StartFPS = STAT_FPSChart_30_40 * 5;
+			StartFPS += (BucketIndex - STAT_FPSChart_30_40) * 10;
+			EndFPS = StartFPS + 10;
+		}
+		
 		if( BucketIndex + STAT_FPSChartFirstStat == STAT_FPSChart_120_INF )
 		{
 			EndFPS = 999;
@@ -698,9 +667,23 @@ void UEngine::DumpFPSChartToStatsLog( float TotalTime, float DeltaTime, int32 Nu
 			const float BucketTimePercentage  = 100.f * GFPSChart[BucketIndex].CummulativeTime / TotalTime;
 			const float BucketFramePercentage = 100.f * GFPSChart[BucketIndex].Count / NumFrames;
 
-			// Figure out bucket range.
-			const int32 StartFPS	= BucketIndex * 5;
-			int32 EndFPS		= StartFPS + 5;
+			// Figure out bucket range. Buckets start at 5 frame intervals then change to 10.
+			int32 StartFPS = 0;
+			int32 EndFPS = 0;
+			if (BucketIndex < STAT_FPSChart_30_40)
+			{
+				// Still incrementing by 5
+				StartFPS = BucketIndex * 5;
+				EndFPS = StartFPS + 5;
+			}
+			else
+			{
+				// Now incrementing by 10
+				StartFPS = STAT_FPSChart_30_40 * 5;
+				StartFPS += (BucketIndex - STAT_FPSChart_30_40) * 10;
+				EndFPS = StartFPS + 10;
+			}
+
 			if( BucketIndex + STAT_FPSChartFirstStat == STAT_FPSChart_120_INF )
 			{
 				EndFPS = 999;
@@ -820,9 +803,23 @@ void UEngine::DumpFPSChartToHTML( float TotalTime, float DeltaTime, int32 NumFra
 			const float BucketTimePercentage  = 100.f * GFPSChart[BucketIndex].CummulativeTime / TotalTime;
 			const float BucketFramePercentage = 100.f * GFPSChart[BucketIndex].Count / NumFrames;
 
-			// Figure out bucket range.
-			const int32 StartFPS	= BucketIndex * 5;
-			int32 EndFPS		= StartFPS + 5;
+			// Figure out bucket range. Buckets start at 5 frame intervals then change to 10.
+			int32 StartFPS = 0;
+			int32 EndFPS = 0;
+			if (BucketIndex < STAT_FPSChart_30_40)
+			{
+				// Still incrementing by 5
+				StartFPS = BucketIndex * 5;
+				EndFPS = StartFPS + 5;
+			}
+			else
+			{
+				// Now incrementing by 10
+				StartFPS = STAT_FPSChart_30_40 * 5;
+				StartFPS += (BucketIndex - STAT_FPSChart_30_40) * 10;
+				EndFPS = StartFPS + 10;
+			}
+
 			if( BucketIndex + STAT_FPSChartFirstStat == STAT_FPSChart_120_INF )
 			{
 				EndFPS = 999;
diff --git a/Engine/Source/Runtime/Engine/Public/ChartCreation.h b/Engine/Source/Runtime/Engine/Public/ChartCreation.h
--- a/Engine/Source/Runtime/Engine/Public/ChartCreation.h
+++ b/Engine/Source/Runtime/Engine/Public/ChartCreation.h
@@ -34,24 +34,15 @@ enum FPSChartStats
 	STAT_FPSChart_15_20,
 	STAT_FPSChart_20_25,
 	STAT_FPSChart_25_30,
-	STAT_FPSChart_30_35,
-	STAT_FPSChart_35_40,
-	STAT_FPSChart_40_45,
-	STAT_FPSChart_45_50,
-	STAT_FPSChart_50_55,
-	STAT_FPSChart_55_60,
-	STAT_FPSChart_60_65,
-	STAT_FPSChart_65_70,
-	STAT_FPSChart_70_75,
-	STAT_FPSChart_75_80,
-	STAT_FPSChart_80_85,
-	STAT_FPSChart_85_90,
-	STAT_FPSChart_90_95,
-	STAT_FPSChart_95_100,
-	STAT_FPSChart_100_105,
-	STAT_FPSChart_105_110,
-	STAT_FPSChart_110_115,
-	STAT_FPSChart_115_120,
+	STAT_FPSChart_30_40,
+	STAT_FPSChart_40_50,
+	STAT_FPSChart_50_60,
+	STAT_FPSChart_60_70,
+	STAT_FPSChart_70_80,
+	STAT_FPSChart_80_90,
+	STAT_FPSChart_90_100,
+	STAT_FPSChart_100_110,
+	STAT_FPSChart_110_120,
 	STAT_FPSChart_120_INF,
 	STAT_FPSChartLastBucketStat,
 	STAT_FPSChart_30Plus,
