Merging CL#2459028 using UT_To_UE4
#UE4 FPS charts now include OS, CPU, GPU, and quality settings

diff --git a/Engine/Source/Runtime/Engine/Private/ChartCreation.cpp b/Engine/Source/Runtime/Engine/Private/ChartCreation.cpp
--- a/Engine/Source/Runtime/Engine/Private/ChartCreation.cpp
+++ b/Engine/Source/Runtime/Engine/Private/ChartCreation.cpp
@@ -8,6 +8,8 @@
 #include "ChartCreation.h"
 #include "Database.h"
 #include "RenderCore.h"
+#include "Scalability.h"
+#include "GameFramework/GameUserSettings.h"
 
 DEFINE_LOG_CATEGORY_STATIC(LogChartCreation, Log, All);
 
@@ -657,6 +659,26 @@ void UEngine::DumpFPSChartToStatsLog( float TotalTime, float DeltaTime, int32 Nu
 	{
 		OutputFile->Logf(TEXT("Dumping FPS chart at %s using build %s built from changelist %i"), *FDateTime::Now().ToString(), *GEngineVersion.ToString(), GetChangeListNumberForPerfTesting() );
 
+		// Get OS info
+		FString OSMajor;
+		FString OSMinor;
+		FPlatformMisc::GetOSVersions(OSMajor, OSMinor);
+
+		// Get settings info
+		const Scalability::FQualityLevels& Quality = GEngine->GetGameUserSettings()->ScalabilityQuality;
+
+		OutputFile->Logf(TEXT("Machine info:"));
+		OutputFile->Logf(TEXT("\tOS: %s %s"), *OSMajor, *OSMinor);
+		OutputFile->Logf(TEXT("\tCPU: %s %s"), *FPlatformMisc::GetCPUVendor(), *FPlatformMisc::GetCPUBrand());
+		OutputFile->Logf(TEXT("\tGPU: %s"), *FPlatformMisc::GetPrimaryGPUBrand());
+		OutputFile->Logf(TEXT("\tResolution Quality: %d"), Quality.ResolutionQuality);
+		OutputFile->Logf(TEXT("\tView Distance Quality: %d"), Quality.ViewDistanceQuality);
+		OutputFile->Logf(TEXT("\tAnti-Aliasing Quality: %d"), Quality.AntiAliasingQuality);
+		OutputFile->Logf(TEXT("\tShadow Quality: %d"), Quality.ShadowQuality);
+		OutputFile->Logf(TEXT("\tPost-Process Quality: %d"), Quality.PostProcessQuality);
+		OutputFile->Logf(TEXT("\tTexture Quality: %d"), Quality.TextureQuality);
+		OutputFile->Logf(TEXT("\tEffects Quality: %d"), Quality.EffectsQuality);
+
 		int32 NumFramesBelow30 = 0; // keep track of the number of frames below 30 FPS
 		float PctTimeAbove30 = 0; // Keep track of percentage of time at 30+ FPS.
 
@@ -876,12 +898,30 @@ void UEngine::DumpFPSChartToHTML( float TotalTime, float DeltaTime, int32 NumFra
 			FPSChartRow = FPSChartRow.Replace( TEXT("TOKEN_HITCH_GPU_BOUND_COUNT"), *FString::Printf(TEXT("%i"), TotalGPUBoundHitches), ESearchCase::CaseSensitive );
 		}
 
+		// Get OS info
+		FString OSMajor;
+		FString OSMinor;
+		FPlatformMisc::GetOSVersions(OSMajor, OSMinor);
+
+		// Get settings info
+		const Scalability::FQualityLevels& Quality = GEngine->GetGameUserSettings()->ScalabilityQuality;
 
 		// Update non- bucket stats.
 		FPSChartRow = FPSChartRow.Replace( TEXT("TOKEN_MAPNAME"),		    *FString::Printf(TEXT("%s"), *InMapName ), ESearchCase::CaseSensitive );
 		FPSChartRow = FPSChartRow.Replace( TEXT("TOKEN_CHANGELIST"),		*FString::Printf(TEXT("%i"), GetChangeListNumberForPerfTesting() ), ESearchCase::CaseSensitive );
 		FPSChartRow = FPSChartRow.Replace( TEXT("TOKEN_DATESTAMP"),         *FString::Printf(TEXT("%s"), *FDateTime::Now().ToString() ), ESearchCase::CaseSensitive );
 
+		FPSChartRow = FPSChartRow.Replace( TEXT("TOKEN_OS"),			    *FString::Printf(TEXT("%s %s"), *OSMajor, *OSMinor ), ESearchCase::CaseSensitive );
+		FPSChartRow = FPSChartRow.Replace( TEXT("TOKEN_CPU"),			    *FString::Printf(TEXT("%s %s"), *FPlatformMisc::GetCPUVendor(), *FPlatformMisc::GetCPUBrand() ), ESearchCase::CaseSensitive );
+		FPSChartRow = FPSChartRow.Replace( TEXT("TOKEN_GPU"),			    *FString::Printf(TEXT("%s"), *FPlatformMisc::GetPrimaryGPUBrand() ), ESearchCase::CaseSensitive );
+		FPSChartRow = FPSChartRow.Replace( TEXT("TOKEN_SETTINGS_RES"),	    *FString::Printf(TEXT("%d"), Quality.ResolutionQuality ), ESearchCase::CaseSensitive );
+		FPSChartRow = FPSChartRow.Replace( TEXT("TOKEN_SETTINGS_VD"),	    *FString::Printf(TEXT("%d"), Quality.ViewDistanceQuality ), ESearchCase::CaseSensitive );
+		FPSChartRow = FPSChartRow.Replace( TEXT("TOKEN_SETTINGS_AA"),	    *FString::Printf(TEXT("%d"), Quality.AntiAliasingQuality ), ESearchCase::CaseSensitive );
+		FPSChartRow = FPSChartRow.Replace( TEXT("TOKEN_SETTINGS_SHADOW"),   *FString::Printf(TEXT("%d"), Quality.ShadowQuality ), ESearchCase::CaseSensitive );
+		FPSChartRow = FPSChartRow.Replace( TEXT("TOKEN_SETTINGS_PP"),	    *FString::Printf(TEXT("%d"), Quality.PostProcessQuality ), ESearchCase::CaseSensitive );
+		FPSChartRow = FPSChartRow.Replace( TEXT("TOKEN_SETTINGS_TEX"),	    *FString::Printf(TEXT("%d"), Quality.TextureQuality ), ESearchCase::CaseSensitive );
+		FPSChartRow = FPSChartRow.Replace( TEXT("TOKEN_SETTINGS_FX"),	    *FString::Printf(TEXT("%d"), Quality.EffectsQuality ), ESearchCase::CaseSensitive );
+
 		FPSChartRow = FPSChartRow.Replace( TEXT("TOKEN_AVG_FPS"),			*FString::Printf(TEXT("%4.2f"), NumFrames / TotalTime), ESearchCase::CaseSensitive );
 		FPSChartRow = FPSChartRow.Replace( TEXT("TOKEN_PCT_ABOVE_30"),		*FString::Printf(TEXT("%4.2f"), PctTimeAbove30), ESearchCase::CaseSensitive );
 		FPSChartRow = FPSChartRow.Replace( TEXT("TOKEN_TIME_DISREGARDED"),	*FString::Printf(TEXT("%4.2f"), FMath::Max<float>( 0, DeltaTime - TotalTime ) ), ESearchCase::CaseSensitive );
