PR #849: Improved default values for various blueprint math functions (Contributed by [REDACTED_PERSON])

https://github.com/EpicGames/UnrealEngine/pull/849

diff --git a/Engine/Source/Runtime/Engine/Classes/Kismet/KismetMathLibrary.h b/Engine/Source/Runtime/Engine/Classes/Kismet/KismetMathLibrary.h
--- a/Engine/Source/Runtime/Engine/Classes/Kismet/KismetMathLibrary.h
+++ b/Engine/Source/Runtime/Engine/Classes/Kismet/KismetMathLibrary.h
@@ -102,19 +102,19 @@ class UKismetMathLibrary : public UBlueprintFunctionLibrary
 
 	/* Division (A / B) */
 	UFUNCTION(BlueprintPure, meta=(FriendlyName = "Byte / Byte", CompactNodeTitle = "/", Keywords = "/ divide division"), Category="Math|Byte")
-	static uint8 Divide_ByteByte(uint8 A, uint8 B);
+	static uint8 Divide_ByteByte(uint8 A, uint8 B = 1);
 
 	/* Modulo (A % B) */
 	UFUNCTION(BlueprintPure, meta=(FriendlyName = "% (Byte)", CompactNodeTitle = "%", Keywords = "% modulus"), Category="Math|Byte")
-	static uint8 Percent_ByteByte(uint8 A, uint8 B);
+	static uint8 Percent_ByteByte(uint8 A, uint8 B = 1);
 
 	/* Addition (A + B) */
 	UFUNCTION(BlueprintPure, meta=(FriendlyName = "Byte + Byte", CompactNodeTitle = "+", Keywords = "+ add plus", CommutativeAssociativeBinaryOperator = "true"), Category="Math|Byte")
-	static uint8 Add_ByteByte(uint8 A, uint8 B);
+	static uint8 Add_ByteByte(uint8 A, uint8 B = 1);
 
 	/* Subtraction (A - B) */
 	UFUNCTION(BlueprintPure, meta=(FriendlyName = "Byte - Byte", CompactNodeTitle = "-", Keywords = "- subtract minus"), Category="Math|Byte")
-	static uint8 Subtract_ByteByte(uint8 A, uint8 B);
+	static uint8 Subtract_ByteByte(uint8 A, uint8 B = 1);
 
 	/* Returns true if A is less than B (A < B) */
 	UFUNCTION(BlueprintPure, meta=(FriendlyName = "Byte < Byte", CompactNodeTitle = "<", Keywords = "< less"), Category="Math|Byte")
@@ -150,19 +150,19 @@ class UKismetMathLibrary : public UBlueprintFunctionLibrary
 
 	/* Division (A / B) */
 	UFUNCTION(BlueprintPure, meta=(FriendlyName = "integer / integer", CompactNodeTitle = "/", Keywords = "/ divide division"), Category="Math|Integer")
-	static int32 Divide_IntInt(int32 A, int32 B);
+	static int32 Divide_IntInt(int32 A, int32 B = 1);
 
 	/* Modulo (A % B) */
 	UFUNCTION(BlueprintPure, meta=(FriendlyName = "% (integer)", CompactNodeTitle = "%", Keywords = "% modulus"), Category="Math|Integer")
-	static int32 Percent_IntInt(int32 A, int32 B);
+	static int32 Percent_IntInt(int32 A, int32 B = 1);
 
 	/* Addition (A + B) */
 	UFUNCTION(BlueprintPure, meta=(FriendlyName = "integer + integer", CompactNodeTitle = "+", Keywords = "+ add plus", CommutativeAssociativeBinaryOperator = "true"), Category="Math|Integer")
-	static int32 Add_IntInt(int32 A, int32 B);
+	static int32 Add_IntInt(int32 A, int32 B = 1);
 
 	/* Subtraction (A - B) */
 	UFUNCTION(BlueprintPure, meta=(FriendlyName = "integer - integer", CompactNodeTitle = "-", Keywords = "- subtract minus"), Category="Math|Integer")
-	static int32 Subtract_IntInt(int32 A, int32 B);
+	static int32 Subtract_IntInt(int32 A, int32 B = 1);
 
 	/* Returns true if A is less than B (A < B) */
 	UFUNCTION(BlueprintPure, meta=(FriendlyName = "integer < integer", CompactNodeTitle = "<", Keywords = "< less"), Category="Math|Integer")
@@ -246,11 +246,11 @@ class UKismetMathLibrary : public UBlueprintFunctionLibrary
 
 	/* Division (A / B) */
 	UFUNCTION(BlueprintPure, meta=(FriendlyName = "float / float", CompactNodeTitle = "/", Keywords = "/ divide division"), Category="Math|Float")
-	static float Divide_FloatFloat(float A, float B);
+	static float Divide_FloatFloat(float A, float B = 1.f);
 
 	/* Modulo (A % B) */
 	UFUNCTION(BlueprintPure, meta=(FriendlyName = "% (float)", CompactNodeTitle = "%", Keywords = "% modulus"), Category="Math|Float")
-	static float Percent_FloatFloat(float A, float B);
+	static float Percent_FloatFloat(float A, float B = 1.f);
 
 	/** Returns the fractional part of a float. */
 	UFUNCTION(BlueprintPure, Category="Math|Float")
@@ -258,11 +258,11 @@ class UKismetMathLibrary : public UBlueprintFunctionLibrary
 
 	/* Addition (A + B) */
 	UFUNCTION(BlueprintPure, meta=(FriendlyName = "float + float", CompactNodeTitle = "+", Keywords = "+ add plus", CommutativeAssociativeBinaryOperator = "true"), Category="Math|Float")
-	static float Add_FloatFloat(float A, float B);
+	static float Add_FloatFloat(float A, float B = 1.f);
 
 	/* Subtraction (A - B) */
 	UFUNCTION(BlueprintPure, meta=(FriendlyName = "float - float", CompactNodeTitle = "-", Keywords = "- subtract minus"), Category="Math|Float")
-	static float Subtract_FloatFloat(float A, float B);
+	static float Subtract_FloatFloat(float A, float B = 1.f);
 
 	/*Returns true if A is Less than B (A < B) */
 	UFUNCTION(BlueprintPure, meta=(FriendlyName = "float < float", CompactNodeTitle = "<", Keywords = "< less"), Category="Math|Float")
@@ -513,11 +513,11 @@ class UKismetMathLibrary : public UBlueprintFunctionLibrary
 
 	/* Vector divide by a float */
 	UFUNCTION(BlueprintPure, meta=(FriendlyName = "vector / float", CompactNodeTitle = "/", Keywords = "/ divide division"), Category="Math|Vector")
-	static FVector Divide_VectorFloat(FVector A, float B);
+	static FVector Divide_VectorFloat(FVector A, float B = 1.f);
 
 	/* Vector divide by vector */
 	UFUNCTION(BlueprintPure, meta=(FriendlyName = "vector / vector", CompactNodeTitle = "/", Keywords = "/ divide division"), Category="Math|Vector")
-	static FVector Divide_VectorVector(FVector A, FVector B);
+	static FVector Divide_VectorVector(FVector A, FVector B = FVector(1.f,1.f,1.f));
 
 	/* Vector addition */
 	UFUNCTION(BlueprintPure, meta=(FriendlyName = "vector + vector", CompactNodeTitle = "+", Keywords = "+ add plus", CommutativeAssociativeBinaryOperator = "true"), Category="Math|Vector")
@@ -1350,7 +1350,7 @@ class UKismetMathLibrary : public UBlueprintFunctionLibrary
 
 	/** Returns Vector A divided by B */
 	UFUNCTION(BlueprintPure, meta=(FriendlyName = "vector2d / float", CompactNodeTitle = "/", Keywords = "/ divide division"), Category="Math|Vector2D")
-	static FVector2D Divide_Vector2DFloat(FVector2D A, float B);
+	static FVector2D Divide_Vector2DFloat(FVector2D A, float B = 1.f);
 
 	/** Returns Vector A added by B */
 	UFUNCTION(BlueprintPure, meta=(FriendlyName = "vector2d + float", CompactNodeTitle = "+", Keywords = "+ add plus"), Category="Math|Vector2D")
