#ue4 - MovementComponents now work with SceneComponents, rather than being restricted to updating a PrimitiveComponent.

- MovementComponent::UpdatedComponent changed from UPrimitiveComponent* to USceneComponent*.
- Added MovementComponent::UpdatedPrimitive, which is a version of UpdatedComponent cast as a PrimitiveComponent.
- Existing blueprints that are referencing 'UpdatedComponent' will be converted on load to use 'UpdatedPrimitive' if the variable link is connected to a pin that requires a PrimitiveComponent or child.

UE-8602

(Integrate 2433859 from 4.7)

diff --git a/Engine/Source/Editor/BlueprintGraph/Classes/K2Node_Variable.h b/Engine/Source/Editor/BlueprintGraph/Classes/K2Node_Variable.h
--- a/Engine/Source/Editor/BlueprintGraph/Classes/K2Node_Variable.h
+++ b/Engine/Source/Editor/BlueprintGraph/Classes/K2Node_Variable.h
@@ -43,6 +43,12 @@ protected:
 	UPROPERTY()
 	uint32 bSelfContext_DEPRECATED:1;
 
+	/**
+	 * Remap a reference from one variable to another, if this variable is of class type 'MatchInVariableClass', and if linked to anything that is a child of 'RemapIfLinkedToClass'.
+	 * Only intended for versioned fixup where redirects can't be applied.
+	 */
+	bool RemapRestrictedLinkReference(FName OldVariableName, FName NewVariableName, const UClass* MatchInVariableClass, const UClass* RemapIfLinkedToClass, bool bLogWarning);
+
 public:
 	// Begin UObject interface
 	virtual void Serialize(FArchive& Ar) override;
diff --git a/Engine/Source/Editor/BlueprintGraph/Private/K2Node_Variable.cpp b/Engine/Source/Editor/BlueprintGraph/Private/K2Node_Variable.cpp
--- a/Engine/Source/Editor/BlueprintGraph/Private/K2Node_Variable.cpp
+++ b/Engine/Source/Editor/BlueprintGraph/Private/K2Node_Variable.cpp
@@ -5,6 +5,7 @@
 
 #include "CompilerResultsLog.h"
 #include "ClassIconFinder.h"
+#include "MessageLog.h"
 
 #define LOCTEXT_NAMESPACE "K2Node"
 
@@ -465,6 +466,7 @@ void UK2Node_Variable::ReconstructNode()
 	UClass* const VarClass = GetVariableSourceClass();
 	if (VarClass)
 	{
+		bool bRemappedProperty = false;
 		UClass* SearchClass = VarClass;
 		while (SearchClass != NULL)
 		{
@@ -484,12 +486,20 @@ void UK2Node_Variable::ReconstructNode()
 					}
 
 					// found, can break
+					bRemappedProperty = true;
 					break;
 				}
 			}
 
 			SearchClass = SearchClass->GetSuperClass();
 		}
+
+		if (!bRemappedProperty)
+		{
+			static FName OldVariableName(TEXT("UpdatedComponent"));
+			static FName NewVariableName(TEXT("UpdatedPrimitive"));
+			bRemappedProperty = RemapRestrictedLinkReference(OldVariableName, NewVariableName, UMovementComponent::StaticClass(), UPrimitiveComponent::StaticClass(), true);
+		}
 	}
 
 	const FGuid VarGuid = VariableReference.GetMemberGuid();
@@ -512,6 +522,58 @@ void UK2Node_Variable::ReconstructNode()
 	Super::ReconstructNode();
 }
 
+
+bool UK2Node_Variable::RemapRestrictedLinkReference(FName OldVariableName, FName NewVariableName, const UClass* MatchInVariableClass, const UClass* RemapIfLinkedToClass, bool bLogWarning)
+{
+	bool bRemapped = false;
+	if (VariableReference.GetMemberName() == OldVariableName)
+	{
+		UClass* const VarClass = GetVariableSourceClass();
+		if (VarClass->IsChildOf(MatchInVariableClass))
+		{
+			UEdGraphPin* VariablePin = GetValuePin();
+			if (VariablePin)
+			{
+				for (UEdGraphPin* OtherPin : VariablePin->LinkedTo)
+				{
+					if (OtherPin != nullptr && VariablePin->PinType.PinCategory == OtherPin->PinType.PinCategory)
+					{
+						// If any other pin we are linked to is a more restricted type, we need to do the remap.
+						const UClass* OtherPinClass = Cast<UClass>(OtherPin->PinType.PinSubCategoryObject.Get());
+						if (OtherPinClass && OtherPinClass->IsChildOf(RemapIfLinkedToClass))
+						{
+							if (VariableReference.IsSelfContext())
+							{
+								VariableReference.SetSelfMember(NewVariableName);
+							}
+							else
+							{
+								VariableReference.SetExternalMember(NewVariableName, VarClass);
+							}
+							bRemapped = true;
+							break;
+						}
+					}
+				}
+			}
+		}
+	}
+
+	if (bRemapped && bLogWarning && GetBlueprint())
+	{
+		FMessageLog("BlueprintLog").Warning(
+			FText::Format(
+			LOCTEXT("RemapRestrictedLinkReference", "{0}: Variable '{1}' was automatically changed to '{2}'. Verify that logic works as intended. (This warning will disappear once the blueprint has been resaved)"),
+			FText::FromString(GetBlueprint()->GetPathName()),
+			FText::FromString(MatchInVariableClass->GetName() + TEXT(".") + OldVariableName.ToString()),
+			FText::FromString(MatchInVariableClass->GetName() + TEXT(".") + NewVariableName.ToString())
+			));
+	}
+
+	return bRemapped;
+}
+
+
 FName UK2Node_Variable::GetCornerIcon() const
 {
 	const UProperty* VariableProperty = VariableReference.ResolveMember<UProperty>(this);
diff --git a/Engine/Source/Runtime/Engine/Classes/GameFramework/CharacterMovementComponent.h b/Engine/Source/Runtime/Engine/Classes/GameFramework/CharacterMovementComponent.h
--- a/Engine/Source/Runtime/Engine/Classes/GameFramework/CharacterMovementComponent.h
+++ b/Engine/Source/Runtime/Engine/Classes/GameFramework/CharacterMovementComponent.h
@@ -373,7 +373,7 @@ public:
 
 	/** What to update CharacterOwner and UpdatedComponent after movement ends */
 	UPROPERTY()
-	UPrimitiveComponent* DeferredUpdatedMoveComponent;
+	USceneComponent* DeferredUpdatedMoveComponent;
 
 	/** Maximum step height for getting out of water */
 	UPROPERTY(Category="Character Movement", EditAnywhere, BlueprintReadWrite, AdvancedDisplay, meta=(ClampMin="0", UIMin="0"))
@@ -1189,7 +1189,7 @@ public:
 	virtual void MoveSmooth(const FVector& InVelocity, const float DeltaSeconds, FStepDownResult* OutStepDownResult = NULL );
 
 	
-	virtual void SetUpdatedComponent(UPrimitiveComponent* NewUpdatedComponent) override;
+	virtual void SetUpdatedComponent(USceneComponent* NewUpdatedComponent) override;
 	
 	/** @Return MovementMode string */
 	virtual FString GetMovementName();
diff --git a/Engine/Source/Runtime/Engine/Classes/GameFramework/MovementComponent.h b/Engine/Source/Runtime/Engine/Classes/GameFramework/MovementComponent.h
--- a/Engine/Source/Runtime/Engine/Classes/GameFramework/MovementComponent.h
+++ b/Engine/Source/Runtime/Engine/Classes/GameFramework/MovementComponent.h
@@ -12,6 +12,7 @@
 struct FCollisionQueryParams;
 struct FCollisionResponseParams;
 struct FCollisionShape;
+class USceneComponent;
 
 /**
  * Setting that controls behavior when movement is restricted to a 2D plane defined by a specific axis/normal,
@@ -50,10 +51,16 @@ class ENGINE_API UMovementComponent : public UActorComponent
 	/**
 	 * The component we move and update.
 	 * If this is null at startup and bAutoRegisterUpdatedComponent is true, the owning Actor's root component will automatically be set as our UpdatedComponent at startup.
-	 * @see bAutoRegisterUpdatedComponent, SetUpdatedComponent()
+	 * @see bAutoRegisterUpdatedComponent, SetUpdatedComponent(), UpdatedPrimitive
 	 */
 	UPROPERTY(BlueprintReadOnly, DuplicateTransient, Category=MovementComponent)
-	UPrimitiveComponent* UpdatedComponent;
+	USceneComponent* UpdatedComponent;
+
+	/**
+	 * UpdatedComponent, cast as a UPrimitiveComponent. May be invalid if UpdatedComponent was null or not a UPrimitiveComponent.
+	 */
+	UPROPERTY(BlueprintReadOnly, DuplicateTransient, Category=MovementComponent)
+	UPrimitiveComponent* UpdatedPrimitive;
 
 	/**
 	 * Flags that control the behavior of calls to MoveComponent() on our UpdatedComponent.
@@ -206,7 +213,7 @@ public:
 
 	/** Assign the component we move and update. */
 	UFUNCTION(BlueprintCallable, Category="Components|Movement")
-	virtual void SetUpdatedComponent(class UPrimitiveComponent* NewUpdatedComponent);
+	virtual void SetUpdatedComponent(USceneComponent* NewUpdatedComponent);
 
 	/** return true if it's in PhysicsVolume with water flag **/
 	virtual bool IsInWater() const;
diff --git a/Engine/Source/Runtime/Engine/Classes/GameFramework/PawnMovementComponent.h b/Engine/Source/Runtime/Engine/Classes/GameFramework/PawnMovementComponent.h
--- a/Engine/Source/Runtime/Engine/Classes/GameFramework/PawnMovementComponent.h
+++ b/Engine/Source/Runtime/Engine/Classes/GameFramework/PawnMovementComponent.h
@@ -21,7 +21,7 @@ class ENGINE_API UPawnMovementComponent : public UNavMovementComponent
 public:
 
 	/** Overridden to only allow registration with components owned by a Pawn. */
-	virtual void SetUpdatedComponent(class UPrimitiveComponent* NewUpdatedComponent) override;
+	virtual void SetUpdatedComponent(USceneComponent* NewUpdatedComponent) override;
 
 	/**
 	 * Adds the given vector to the accumulated input in world space. Input vectors are usually between 0 and 1 in magnitude. 
diff --git a/Engine/Source/Runtime/Engine/Private/CharacterMovementComponent.cpp b/Engine/Source/Runtime/Engine/Private/CharacterMovementComponent.cpp
--- a/Engine/Source/Runtime/Engine/Private/CharacterMovementComponent.cpp
+++ b/Engine/Source/Runtime/Engine/Private/CharacterMovementComponent.cpp
@@ -365,7 +365,7 @@ void UCharacterMovementComponent::BeginDestroy()
 	Super::BeginDestroy();
 }
 
-void UCharacterMovementComponent::SetUpdatedComponent(UPrimitiveComponent* NewUpdatedComponent)
+void UCharacterMovementComponent::SetUpdatedComponent(USceneComponent* NewUpdatedComponent)
 {
 	if (NewUpdatedComponent)
 	{
@@ -394,9 +394,10 @@ void UCharacterMovementComponent::SetUpdatedComponent(UPrimitiveComponent* NewUp
 	bDeferUpdateMoveComponent = false;
 	DeferredUpdatedMoveComponent = NULL;
 
-	if (IsValid(UpdatedComponent) && UpdatedComponent->OnComponentBeginOverlap.IsBound())
+	UPrimitiveComponent* OldPrimitive = Cast<UPrimitiveComponent>(UpdatedComponent);
+	if (IsValid(OldPrimitive) && OldPrimitive->OnComponentBeginOverlap.IsBound())
 	{
-		UpdatedComponent->OnComponentBeginOverlap.RemoveDynamic(this, &UCharacterMovementComponent::CapsuleTouched);
+		OldPrimitive->OnComponentBeginOverlap.RemoveDynamic(this, &UCharacterMovementComponent::CapsuleTouched);
 	}
 	
 	Super::SetUpdatedComponent(NewUpdatedComponent);
@@ -407,9 +408,9 @@ void UCharacterMovementComponent::SetUpdatedComponent(UPrimitiveComponent* NewUp
 		StopActiveMovement();
 	}
 
-	if (IsValid(UpdatedComponent) && bEnablePhysicsInteraction)
+	if (IsValid(UpdatedPrimitive) && bEnablePhysicsInteraction)
 	{
-		UpdatedComponent->OnComponentBeginOverlap.AddUniqueDynamic(this, &UCharacterMovementComponent::CapsuleTouched);
+		UpdatedPrimitive->OnComponentBeginOverlap.AddUniqueDynamic(this, &UCharacterMovementComponent::CapsuleTouched);
 	}
 
 	if (bUseRVOAvoidance)
@@ -3822,7 +3823,7 @@ void UCharacterMovementComponent::PhysNavWalking(float deltaTime, int32 Iteratio
 
 		if (!AdjustedDelta.IsNearlyZero())
 		{
-			const bool bSweep = UpdatedComponent ? UpdatedComponent->bGenerateOverlapEvents : false;
+			const bool bSweep = UpdatedPrimitive ? UpdatedPrimitive->bGenerateOverlapEvents : false;
 			FHitResult HitResult;
 			SafeMoveUpdatedComponent(AdjustedDelta, CharacterOwner->GetActorRotation(), bSweep, HitResult);
 		}
@@ -4040,12 +4041,12 @@ void UCharacterMovementComponent::SetPostLandedPhysics(const FHitResult& Hit)
 
 void UCharacterMovementComponent::SetNavWalkingPhysics(bool bEnable)
 {
-	if (UpdatedComponent)
+	if (UpdatedPrimitive)
 	{
 		if (bEnable)
 		{
-			UpdatedComponent->SetCollisionResponseToChannel(ECC_WorldStatic, ECR_Ignore);
-			UpdatedComponent->SetCollisionResponseToChannel(ECC_WorldDynamic, ECR_Ignore);
+			UpdatedPrimitive->SetCollisionResponseToChannel(ECC_WorldStatic, ECR_Ignore);
+			UpdatedPrimitive->SetCollisionResponseToChannel(ECC_WorldDynamic, ECR_Ignore);
 			CachedProjectedNavMeshHitResult.Reset();
 			NavMeshProjectionTimer = 0.0f;
 		}
@@ -4060,8 +4061,8 @@ void UCharacterMovementComponent::SetNavWalkingPhysics(bool bEnable)
 
 			if (DefaultCapsule)
 			{
-				UpdatedComponent->SetCollisionResponseToChannel(ECC_WorldStatic, DefaultCapsule->GetCollisionResponseToChannel(ECC_WorldStatic));
-				UpdatedComponent->SetCollisionResponseToChannel(ECC_WorldDynamic, DefaultCapsule->GetCollisionResponseToChannel(ECC_WorldDynamic));
+				UpdatedPrimitive->SetCollisionResponseToChannel(ECC_WorldStatic, DefaultCapsule->GetCollisionResponseToChannel(ECC_WorldStatic));
+				UpdatedPrimitive->SetCollisionResponseToChannel(ECC_WorldDynamic, DefaultCapsule->GetCollisionResponseToChannel(ECC_WorldDynamic));
 			}
 			else
 			{
@@ -6594,20 +6595,20 @@ void UCharacterMovementComponent::SetAvoidanceEnabled(bool bEnable)
 
 void UCharacterMovementComponent::ApplyRepulsionForce(float DeltaSeconds)
 {
-	if (UpdatedComponent && RepulsionForce > 0.0f)
+	if (UpdatedPrimitive && RepulsionForce > 0.0f)
 	{
 		FCollisionQueryParams QueryParams;
 		QueryParams.bReturnFaceIndex = false;
 		QueryParams.bReturnPhysicalMaterial = false;
 
-		const FCollisionShape CollisionShape = UpdatedComponent->GetCollisionShape();
+		const FCollisionShape CollisionShape = UpdatedPrimitive->GetCollisionShape();
 		const float CapsuleRadius = CollisionShape.GetCapsuleRadius();
 		const float CapsuleHalfHeight = CollisionShape.GetCapsuleHalfHeight();
 		const float RepulsionForceRadius = CapsuleRadius * 1.2f;
 		const float StopBodyDistance = 2.5f;
 
-		const TArray<FOverlapInfo>& Overlaps = UpdatedComponent->GetOverlapInfos();
-		const FVector MyLocation = UpdatedComponent->GetComponentLocation();
+		const TArray<FOverlapInfo>& Overlaps = UpdatedPrimitive->GetOverlapInfos();
+		const FVector MyLocation = UpdatedPrimitive->GetComponentLocation();
 
 		for (int32 i=0; i < Overlaps.Num(); i++)
 		{
@@ -6648,7 +6649,7 @@ void UCharacterMovementComponent::ApplyRepulsionForce(float DeltaSeconds)
 
 			// Trace to get the hit location on the capsule
 			FHitResult Hit;
-			bool bHasHit = UpdatedComponent->LineTraceComponent(Hit, BodyLocation, 
+			bool bHasHit = UpdatedPrimitive->LineTraceComponent(Hit, BodyLocation,
 																FVector(MyLocation.X, MyLocation.Y, BodyLocation.Z),
 																QueryParams);
 
diff --git a/Engine/Source/Runtime/Engine/Private/MovementComponent.cpp b/Engine/Source/Runtime/Engine/Private/MovementComponent.cpp
--- a/Engine/Source/Runtime/Engine/Private/MovementComponent.cpp
+++ b/Engine/Source/Runtime/Engine/Private/MovementComponent.cpp
@@ -5,6 +5,7 @@
 #include "GameFramework/MovementComponent.h"
 #include "GameFramework/PhysicsVolume.h"
 #include "Components/SceneComponent.h"
+#include "Components/PrimitiveComponent.h"
 #include "MessageLog.h"
 #include "PhysicsEngine/PhysicsSettings.h"
 
@@ -38,7 +39,7 @@ UMovementComponent::UMovementComponent(const FObjectInitializer& ObjectInitializ
 }
 
 
-void UMovementComponent::SetUpdatedComponent(UPrimitiveComponent* NewUpdatedComponent)
+void UMovementComponent::SetUpdatedComponent(USceneComponent* NewUpdatedComponent)
 {
 	if (UpdatedComponent && UpdatedComponent != NewUpdatedComponent)
 	{
@@ -55,7 +56,9 @@ void UMovementComponent::SetUpdatedComponent(UPrimitiveComponent* NewUpdatedComp
 
 	// Don't assign pending kill components, but allow those to null out previous UpdatedComponent.
 	UpdatedComponent = IsValid(NewUpdatedComponent) ? NewUpdatedComponent : NULL;
+	UpdatedPrimitive = Cast<UPrimitiveComponent>(UpdatedComponent);
 
+	// Assign delegates
 	if (IsValid(UpdatedComponent))
 	{
 		UpdatedComponent->bShouldUpdatePhysicsVolume = true;
@@ -85,29 +88,18 @@ void UMovementComponent::InitializeComponent()
 	TGuardValue<bool> InInitializeComponentGuard(bInInitializeComponent, true);
 	Super::InitializeComponent();
 
-	UPrimitiveComponent* NewUpdatedComponent = NULL;
+	USceneComponent* NewUpdatedComponent = NULL;
 	if (UpdatedComponent != NULL)
 	{
 		NewUpdatedComponent = UpdatedComponent;
 	}
 	else if (bAutoRegisterUpdatedComponent)
 	{
-		// Auto-register owner's root primitive component if found.
+		// Auto-register owner's root component if found.
 		AActor* MyActor = GetOwner();
 		if (MyActor)
 		{
-			NewUpdatedComponent = Cast<UPrimitiveComponent>(MyActor->GetRootComponent());
-			if (!NewUpdatedComponent)
-			{
-				FMessageLog("PIE").Warning(
-					FText::Format(
-						LOCTEXT("NoRootPrimitiveWarning",
-								"Movement component {0} must update a PrimitiveComponent. Consider changing root component of owning actor '{1}' to a shape (Box, Capsule, etc) or another component that can handle collision."),
-						FText::FromString(GetName()),
-						FText::FromString(MyActor->GetName())
-					)
-				);
-			}
+			NewUpdatedComponent = MyActor->GetRootComponent();
 		}
 	}
 
@@ -118,6 +110,8 @@ void UMovementComponent::InitializeComponent()
 void UMovementComponent::OnRegister()
 {
 	TGuardValue<bool> InOnRegisterGuard(bInOnRegister, true);
+
+	UpdatedPrimitive = Cast<UPrimitiveComponent>(UpdatedComponent);
 	Super::OnRegister();
 
 	if (PlaneConstraintAxisSetting != EPlaneConstraintAxisSetting::Custom)
@@ -179,6 +173,8 @@ void UMovementComponent::PostLoad()
 		// Make sure to use the most up-to-date project setting in case it has changed.
 		PlaneConstraintNormal = GetPlaneConstraintNormalFromAxisSetting(PlaneConstraintAxisSetting);
 	}
+
+	UpdatedPrimitive = Cast<UPrimitiveComponent>(UpdatedComponent);
 }
 
 
@@ -254,7 +250,7 @@ bool UMovementComponent::ShouldSkipUpdate(float DeltaTime) const
 
 		const float RenderTimeThreshold = 0.41f;
 		UWorld* TheWorld = GetWorld();
-		if (TheWorld->TimeSince(UpdatedComponent->LastRenderTime) <= RenderTimeThreshold)
+		if (UpdatedPrimitive && TheWorld->TimeSince(UpdatedPrimitive->LastRenderTime) <= RenderTimeThreshold)
 		{
 			return false; // Rendered, don't skip it.
 		}
@@ -302,9 +298,9 @@ void UMovementComponent::UpdateComponentVelocity()
 
 void UMovementComponent::InitCollisionParams(FCollisionQueryParams &OutParams, FCollisionResponseParams& OutResponseParam) const
 {
-	if (UpdatedComponent)
+	if (UpdatedPrimitive)
 	{
-		UpdatedComponent->InitSweepCollisionParams(OutParams, OutResponseParam);
+		UpdatedPrimitive->InitSweepCollisionParams(OutParams, OutResponseParam);
 	}
 }
 
@@ -498,8 +494,9 @@ FVector UMovementComponent::GetPenetrationAdjustment(const FHitResult& Hit) cons
 
 bool UMovementComponent::ResolvePenetration(const FVector& ProposedAdjustment, const FHitResult& Hit, const FRotator& NewRotation)
 {
+	// SceneComponent can't be in penetration, so this function really only applies to PrimitiveComponent.
 	const FVector Adjustment = ConstrainDirectionToPlane(ProposedAdjustment);
-	if (!Adjustment.IsZero() && UpdatedComponent)
+	if (!Adjustment.IsZero() && UpdatedPrimitive)
 	{
 		// See if we can fit at the adjusted location without overlapping anything.
 		AActor* ActorOwner = UpdatedComponent->GetOwner();
@@ -511,7 +508,7 @@ bool UMovementComponent::ResolvePenetration(const FVector& ProposedAdjustment, c
 		// We really want to make sure that precision differences or differences between the overlap test and sweep tests don't put us into another overlap,
 		// so make the overlap test a bit more restrictive.
 		const float OverlapInflation = 0.10f;
-		bool bEncroached = OverlapTest(Hit.TraceStart + Adjustment, NewRotation.Quaternion(), UpdatedComponent->GetCollisionObjectType(), UpdatedComponent->GetCollisionShape(OverlapInflation), ActorOwner);
+		bool bEncroached = OverlapTest(Hit.TraceStart + Adjustment, NewRotation.Quaternion(), UpdatedPrimitive->GetCollisionObjectType(), UpdatedPrimitive->GetCollisionShape(OverlapInflation), ActorOwner);
 		if (!bEncroached)
 		{
 			// Move without sweeping.
diff --git a/Engine/Source/Runtime/Engine/Private/PawnMovementComponent.cpp b/Engine/Source/Runtime/Engine/Private/PawnMovementComponent.cpp
--- a/Engine/Source/Runtime/Engine/Private/PawnMovementComponent.cpp
+++ b/Engine/Source/Runtime/Engine/Private/PawnMovementComponent.cpp
@@ -9,7 +9,7 @@ DEFINE_LOG_CATEGORY_STATIC(LogPawnMovementComponent, Log, All);
 //----------------------------------------------------------------------//
 // UPawnMovementComponent
 //----------------------------------------------------------------------//
-void UPawnMovementComponent::SetUpdatedComponent(UPrimitiveComponent* NewUpdatedComponent)
+void UPawnMovementComponent::SetUpdatedComponent(USceneComponent* NewUpdatedComponent)
 {
 	if (NewUpdatedComponent)
 	{
diff --git a/Engine/Source/Runtime/Engine/Private/ProjectileMovementComponent.cpp b/Engine/Source/Runtime/Engine/Private/ProjectileMovementComponent.cpp
--- a/Engine/Source/Runtime/Engine/Private/ProjectileMovementComponent.cpp
+++ b/Engine/Source/Runtime/Engine/Private/ProjectileMovementComponent.cpp
@@ -80,9 +80,9 @@ void UProjectileMovementComponent::InitializeComponent()
 
 		UpdateComponentVelocity();
 		
-		if (UpdatedComponent && UpdatedComponent->IsSimulatingPhysics())
+		if (UpdatedPrimitive && UpdatedPrimitive->IsSimulatingPhysics())
 		{
-			UpdatedComponent->SetPhysicsLinearVelocity(Velocity);
+			UpdatedPrimitive->SetPhysicsLinearVelocity(Velocity);
 		}
 	}
 }
diff --git a/Engine/Source/Runtime/Engine/Private/Vehicles/WheeledVehicleMovementComponent.cpp b/Engine/Source/Runtime/Engine/Private/Vehicles/WheeledVehicleMovementComponent.cpp
--- a/Engine/Source/Runtime/Engine/Private/Vehicles/WheeledVehicleMovementComponent.cpp
+++ b/Engine/Source/Runtime/Engine/Private/Vehicles/WheeledVehicleMovementComponent.cpp
@@ -143,7 +143,14 @@ bool UWheeledVehicleMovementComponent::CanCreateVehicle() const
 		return false;
 	}
 
-	if ( UpdatedComponent->GetBodyInstance() == NULL )
+	if (UpdatedPrimitive == NULL)
+	{
+		UE_LOG(LogVehicles, Warning, TEXT("Cannot create vehicle for %s. UpdatedComponent is not a PrimitiveComponent."),
+			*GetPathName());
+		return false;
+	}
+
+	if (UpdatedPrimitive->GetBodyInstance() == NULL)
 	{
 		UE_LOG( LogVehicles, Warning, TEXT("Cannot create vehicle for %s. UpdatedComponent has not initialized its rigid body actor."),
 			*GetPathName() );
@@ -172,14 +179,16 @@ void UWheeledVehicleMovementComponent::CreateVehicle()
 		if ( CanCreateVehicle() )
 		{
 			check(UpdatedComponent);
-			check(UpdatedComponent->GetBodyInstance()->GetPxRigidDynamic());
-
-
-			SetupVehicle();
-
-			if ( PVehicle != NULL )
+			if (ensure(UpdatedPrimitive != nullptr))
 			{
-				PostSetupVehicle();
+				check(UpdatedPrimitive->GetBodyInstance()->GetPxRigidDynamic());
+				
+				SetupVehicle();
+
+				if ( PVehicle != NULL )
+				{
+					PostSetupVehicle();
+				}
 			}
 		}
 	}
@@ -187,7 +196,12 @@ void UWheeledVehicleMovementComponent::CreateVehicle()
 
 void UWheeledVehicleMovementComponent::SetupVehicleShapes()
 {
-	PxRigidDynamic* PVehicleActor = UpdatedComponent->GetBodyInstance()->GetPxRigidDynamic();
+	if (!UpdatedPrimitive)
+	{
+		return;
+	}
+
+	PxRigidDynamic* PVehicleActor = UpdatedPrimitive->GetBodyInstance()->GetPxRigidDynamic();
 
 	static PxMaterial* WheelMaterial = GPhysXSDK->createMaterial(0.0f, 0.0f, 0.0f);
 
@@ -253,7 +267,12 @@ void UWheeledVehicleMovementComponent::SetupVehicleShapes()
 
 void UWheeledVehicleMovementComponent::SetupVehicleMass()
 {
-	const FBodyInstance& BodyInst = *UpdatedComponent->GetBodyInstance();
+	if (!UpdatedPrimitive)
+	{
+		return;
+	}
+
+	const FBodyInstance& BodyInst = *UpdatedPrimitive->GetBodyInstance();
 	PxRigidDynamic* PVehicleActor = BodyInst.GetPxRigidDynamic();
 
 	// Override mass
@@ -274,7 +293,12 @@ void UWheeledVehicleMovementComponent::SetupVehicleMass()
 
 void UWheeledVehicleMovementComponent::SetupWheels( PxVehicleWheelsSimData* PWheelsSimData )
 {
-	const FBodyInstance& VehicleBodyInstance = *UpdatedComponent->GetBodyInstance();
+	if (!UpdatedPrimitive)
+	{
+		return;
+	}
+
+	const FBodyInstance& VehicleBodyInstance = *UpdatedPrimitive->GetBodyInstance();
 	PxRigidDynamic* PVehicleActor = VehicleBodyInstance.GetPxRigidDynamic();
 
 	const PxReal LengthScale = 100.f; // Convert default from m to cm
@@ -525,12 +549,15 @@ FVector UWheeledVehicleMovementComponent::GetWheelRestingPosition( const FWheelS
 
 FVector UWheeledVehicleMovementComponent::GetLocalCOM() const
 {
-	if (const FBodyInstance* BodyInst = UpdatedComponent->GetBodyInstance())
+	if (UpdatedPrimitive)
 	{
-		if (PxRigidDynamic* PVehicleActor = BodyInst->GetPxRigidDynamic())
+		if (const FBodyInstance* BodyInst = UpdatedPrimitive->GetBodyInstance())
 		{
-			PxTransform PCOMTransform = PVehicleActor->getCMassLocalPose();
-			return P2UVector(PCOMTransform.p);
+			if (PxRigidDynamic* PVehicleActor = BodyInst->GetPxRigidDynamic())
+			{
+				PxTransform PCOMTransform = PVehicleActor->getCMassLocalPose();
+				return P2UVector(PCOMTransform.p);
+			}
 		}
 	}
 
@@ -711,7 +738,7 @@ void UWheeledVehicleMovementComponent::TickVehicle( float DeltaTime )
 
 void UWheeledVehicleMovementComponent::UpdateDrag(float DeltaTime)
 {
-	if (PVehicle && UpdatedComponent)
+	if (PVehicle && UpdatedPrimitive)
 	{
 		float ForwardSpeed = GetForwardSpeed();
 		if (FMath::Abs(ForwardSpeed) > 1.f)
@@ -724,7 +751,7 @@ void UWheeledVehicleMovementComponent::UpdateDrag(float DeltaTime)
 			float DragMag = 0.5f * AirDensity * SpeedSquared * DragCoefficient * ChassisDragArea;
 			DebugDragMagnitude = DragMag;
 			DragVector *= DragMag;
-			FBodyInstance* BodyInstance = UpdatedComponent->GetBodyInstance();
+			FBodyInstance* BodyInstance = UpdatedPrimitive->GetBodyInstance();
 			BodyInstance->AddForce(DragVector, false);
 		}
 		
diff --git a/Engine/Source/Runtime/Engine/Private/Vehicles/WheeledVehicleMovementComponent4W.cpp b/Engine/Source/Runtime/Engine/Private/Vehicles/WheeledVehicleMovementComponent4W.cpp
--- a/Engine/Source/Runtime/Engine/Private/Vehicles/WheeledVehicleMovementComponent4W.cpp
+++ b/Engine/Source/Runtime/Engine/Private/Vehicles/WheeledVehicleMovementComponent4W.cpp
@@ -244,6 +244,11 @@ void SetupDriveHelper(const UWheeledVehicleMovementComponent4W* VehicleData, con
 
 void UWheeledVehicleMovementComponent4W::SetupVehicle()
 {
+	if (!UpdatedPrimitive)
+	{
+		return;
+	}
+
 	if (WheelSetups.Num() != 4)
 	{
 		PVehicle = NULL;
@@ -278,7 +283,7 @@ void UWheeledVehicleMovementComponent4W::SetupVehicle()
 	PxVehicleDrive4W* PVehicleDrive4W = PxVehicleDrive4W::allocate(4);
 	check(PVehicleDrive4W);
 
-	PVehicleDrive4W->setup( GPhysXSDK, UpdatedComponent->GetBodyInstance()->GetPxRigidDynamic(), *PWheelsSimData, DriveData, 0);
+	PVehicleDrive4W->setup( GPhysXSDK, UpdatedPrimitive->GetBodyInstance()->GetPxRigidDynamic(), *PWheelsSimData, DriveData, 0);
 	PVehicleDrive4W->setToRestState();
 
 	// cleanup
