Fix attachment issues between Actors when the parent Actor has a blueprint created root component

diff --git a/Engine/Source/Runtime/Engine/Private/SceneComponent.cpp b/Engine/Source/Runtime/Engine/Private/SceneComponent.cpp
--- a/Engine/Source/Runtime/Engine/Private/SceneComponent.cpp
+++ b/Engine/Source/Runtime/Engine/Private/SceneComponent.cpp
@@ -1165,10 +1165,11 @@ bool USceneComponent::IsAttachedTo(class USceneComponent* TestComp) const
 FSceneComponentInstanceData::FSceneComponentInstanceData(const USceneComponent* SourceComponent)
 	: FActorComponentInstanceData(SourceComponent)
 {
+	AActor* SourceOwner = SourceComponent->GetOwner();
 	for (int32 i = SourceComponent->AttachChildren.Num()-1; i >= 0; --i)
 	{
 		USceneComponent* SceneComponent = SourceComponent->AttachChildren[i];
-		if (SceneComponent && !SceneComponent->IsCreatedByConstructionScript())
+		if (SceneComponent && SceneComponent->GetOwner() == SourceOwner && !SceneComponent->IsCreatedByConstructionScript())
 		{
 			AttachedInstanceComponents.Add(SceneComponent);
 		}
