#ue4 - Characters using NavWalking movement mode now stagger their line check updates to avoid running on the same frame if possible.

- Time updates account for rollover to make delta times accurate. Adds protection so a frame rate / delta time spike doesn't realign them all to the same frame.

diff --git a/Engine/Source/Runtime/Engine/Private/CharacterMovementComponent.cpp b/Engine/Source/Runtime/Engine/Private/CharacterMovementComponent.cpp
--- a/Engine/Source/Runtime/Engine/Private/CharacterMovementComponent.cpp
+++ b/Engine/Source/Runtime/Engine/Private/CharacterMovementComponent.cpp
@@ -3901,7 +3901,15 @@ void UCharacterMovementComponent::ProjectLocationFromNavMesh(float DeltaSeconds,
 			CachedProjectedNavMeshHitResult.Reset();
 		}
 
-		NavMeshProjectionTimer = NavMeshProjectionInterval;
+		// Wrap around to maintain same relative offset to tick time changes.
+		// Prevents large framerate spikes from aligning multiple characters to the same frame (if they start staggered, they will now remain staggered).
+		float ModTime = 0.f;
+		if (NavMeshProjectionInterval > SMALL_NUMBER)
+		{
+			ModTime = FMath::Fmod(FMath::Abs(NavMeshProjectionTimer), NavMeshProjectionInterval);
+		}
+
+		NavMeshProjectionTimer = NavMeshProjectionInterval - ModTime;
 	}
 		
 	// project to last plane we found
@@ -4048,7 +4056,10 @@ void UCharacterMovementComponent::SetNavWalkingPhysics(bool bEnable)
 			UpdatedPrimitive->SetCollisionResponseToChannel(ECC_WorldStatic, ECR_Ignore);
 			UpdatedPrimitive->SetCollisionResponseToChannel(ECC_WorldDynamic, ECR_Ignore);
 			CachedProjectedNavMeshHitResult.Reset();
-			NavMeshProjectionTimer = 0.0f;
+
+			// Stagger timed updates so many different characters spawned at the same time don't update on the same frame.
+			// Initially we want an immediate update though, so set time to a negative randomized range.
+			NavMeshProjectionTimer = (NavMeshProjectionInterval > 0.f) ? FMath::FRandRange(-NavMeshProjectionInterval, 0.f) : 0.f;
 		}
 		else
 		{
