Framework: Add support for visualizing the frustum of an orthographic camera

diff --git a/Engine/Source/Runtime/Engine/Private/Camera/CameraComponent.cpp b/Engine/Source/Runtime/Engine/Private/Camera/CameraComponent.cpp
--- a/Engine/Source/Runtime/Engine/Private/Camera/CameraComponent.cpp
+++ b/Engine/Source/Runtime/Engine/Private/Camera/CameraComponent.cpp
@@ -140,9 +140,19 @@ void UCameraComponent::RefreshVisualRepresentation()
 {
 	if (DrawFrustum != NULL)
 	{
-		DrawFrustum->FrustumAngle = (ProjectionMode == ECameraProjectionMode::Perspective) ? FieldOfView : 0.0f;
-		DrawFrustum->FrustumStartDist = 10.f;
-		DrawFrustum->FrustumEndDist = 1000.f;
+		const float FrustumDrawDistance = 1000.0f;
+		if (ProjectionMode == ECameraProjectionMode::Perspective)
+		{
+			DrawFrustum->FrustumAngle = FieldOfView;
+			DrawFrustum->FrustumStartDist = 10.f;
+			DrawFrustum->FrustumEndDist = DrawFrustum->FrustumStartDist + FrustumDrawDistance;
+		}
+		else
+		{
+			DrawFrustum->FrustumAngle = -OrthoWidth;
+			DrawFrustum->FrustumStartDist = OrthoNearClipPlane;
+			DrawFrustum->FrustumEndDist = FMath::Min(OrthoFarClipPlane - OrthoNearClipPlane, FrustumDrawDistance);
+		}
 		DrawFrustum->FrustumAspectRatio = AspectRatio;
 		DrawFrustum->MarkRenderStateDirty();
 	}
diff --git a/Engine/Source/Runtime/Engine/Private/DrawFrustumComponent.cpp b/Engine/Source/Runtime/Engine/Private/DrawFrustumComponent.cpp
--- a/Engine/Source/Runtime/Engine/Private/DrawFrustumComponent.cpp
+++ b/Engine/Source/Runtime/Engine/Private/DrawFrustumComponent.cpp
@@ -42,9 +42,22 @@ public:
 		FVector Verts[8];
 
 		// FOVAngle controls the horizontal angle.
-		float HozHalfAngle = (FrustumAngle) * ((float)PI/360.f);
-		float HozLength = FrustumStartDist * FMath::Tan(HozHalfAngle);
-		float VertLength = HozLength/FrustumAspectRatio;
+		const float HozHalfAngleInRadians = FMath::DegreesToRadians(FrustumAngle * 0.5f);
+
+		float HozLength = 0.0f;
+		float VertLength = 0.0f;
+		
+		if (FrustumAngle > 0.0f)
+		{
+			HozLength = FrustumStartDist * FMath::Tan(HozHalfAngleInRadians);
+			VertLength = HozLength / FrustumAspectRatio;
+		}
+		else
+		{
+			const float OrthoWidth = (FrustumAngle == 0.0f) ? 1000.0f : -FrustumAngle;
+			HozLength = OrthoWidth * 0.5f;
+			VertLength = HozLength / FrustumAspectRatio;
+		}
 
 		// near plane verts
 		Verts[0] = (Direction * FrustumStartDist) + (UpVector * VertLength) + (LeftVector * HozLength);
@@ -52,8 +65,11 @@ public:
 		Verts[2] = (Direction * FrustumStartDist) - (UpVector * VertLength) - (LeftVector * HozLength);
 		Verts[3] = (Direction * FrustumStartDist) - (UpVector * VertLength) + (LeftVector * HozLength);
 
-		HozLength = FrustumEndDist * FMath::Tan(HozHalfAngle);
-		VertLength = HozLength/FrustumAspectRatio;
+		if (FrustumAngle > 0.0f)
+		{
+			HozLength = FrustumEndDist * FMath::Tan(HozHalfAngleInRadians);
+			VertLength = HozLength / FrustumAspectRatio;
+		}
 
 		// far plane verts
 		Verts[4] = (Direction * FrustumEndDist) + (UpVector * VertLength) + (LeftVector * HozLength);
@@ -61,9 +77,9 @@ public:
 		Verts[6] = (Direction * FrustumEndDist) - (UpVector * VertLength) - (LeftVector * HozLength);
 		Verts[7] = (Direction * FrustumEndDist) - (UpVector * VertLength) + (LeftVector * HozLength);
 
-		for( int32 x = 0 ; x < 8 ; ++x )
+		for (int32 X = 0; X < 8; ++X)
 		{
-			Verts[x] = GetLocalToWorld().TransformPosition( Verts[x] );
+			Verts[X] = GetLocalToWorld().TransformPosition(Verts[X]);
 		}
 
 		for (int32 ViewIndex = 0; ViewIndex < Views.Num(); ViewIndex++)
