#ue4 - Character NavMesh movement avoids ProjectLocationFromNavMesh() line trace when location has not changed since last valid check and bAlwaysCheckFloor is not true.

diff --git a/Engine/Source/Runtime/Engine/Private/CharacterMovementComponent.cpp b/Engine/Source/Runtime/Engine/Private/CharacterMovementComponent.cpp
--- a/Engine/Source/Runtime/Engine/Private/CharacterMovementComponent.cpp
+++ b/Engine/Source/Runtime/Engine/Private/CharacterMovementComponent.cpp
@@ -3886,19 +3886,31 @@ void UCharacterMovementComponent::ProjectLocationFromNavMesh(float DeltaSeconds,
 	NavMeshProjectionTimer -= DeltaSeconds;
 	if (NavMeshProjectionTimer <= 0.0f)
 	{
-		UE_LOG(LogCharacterMovement, Verbose, TEXT("ProjectLocationFromNavMesh(): %s.%s interval: %.3f velocity: %s"), *GetNameSafe(CharacterOwner), *GetName(), NavMeshProjectionInterval, *Velocity.ToString());
+		// We can skip this trace if we are at the same location as the last trace (ie, we haven't moved).
+		const bool bCachedLocationStillValid = (!bAlwaysCheckFloor &&
+												CachedProjectedNavMeshHitResult.bBlockingHit &&
+												(CachedProjectedNavMeshHitResult.TraceStart - (InOutLocation + RayCastOffset)).SizeSquared() < KINDA_SMALL_NUMBER);
 
-		// raycast to underlying mesh to allow us to more closely follow geometry
-		// we use static objects here as a best approximation to accept only objects that
-		// influence navmesh generation
-		FCollisionQueryParams Params;
-		FCollisionObjectQueryParams ObjectQueryParams(FCollisionObjectQueryParams::AllStaticObjects);
-		GetWorld()->LineTraceSingle(CachedProjectedNavMeshHitResult, InOutLocation + RayCastOffset, InOutLocation - RayCastOffset, Params, ObjectQueryParams);
+		if (!bCachedLocationStillValid)
+		{
+			UE_LOG(LogCharacterMovement, Verbose, TEXT("ProjectLocationFromNavMesh(): %s interval: %.3f velocity: %s"), *GetNameSafe(CharacterOwner), NavMeshProjectionInterval, *Velocity.ToString());
+
+			// raycast to underlying mesh to allow us to more closely follow geometry
+			// we use static objects here as a best approximation to accept only objects that
+			// influence navmesh generation
+			FCollisionQueryParams Params;
+			FCollisionObjectQueryParams ObjectQueryParams(FCollisionObjectQueryParams::AllStaticObjects);
+			GetWorld()->LineTraceSingle(CachedProjectedNavMeshHitResult, InOutLocation + RayCastOffset, InOutLocation - RayCastOffset, Params, ObjectQueryParams);
 
-		// discard result if we were already inside something
-		if (CachedProjectedNavMeshHitResult.bStartPenetrating)
+			// discard result if we were already inside something
+			if (CachedProjectedNavMeshHitResult.bStartPenetrating)
+			{
+				CachedProjectedNavMeshHitResult.Reset();
+			}
+		}
+		else
 		{
-			CachedProjectedNavMeshHitResult.Reset();
+			UE_LOG(LogCharacterMovement, Verbose, TEXT("ProjectLocationFromNavMesh(): %s interval: %.3f velocity: %s [SKIP TRACE]"), *GetNameSafe(CharacterOwner), NavMeshProjectionInterval, *Velocity.ToString());
 		}
 
 		// Wrap around to maintain same relative offset to tick time changes.
