Fix ISO C++ enum forward decl error

diff --git a/Engine/Source/Runtime/Engine/Classes/Camera/CameraTypes.h b/Engine/Source/Runtime/Engine/Classes/Camera/CameraTypes.h
--- a/Engine/Source/Runtime/Engine/Classes/Camera/CameraTypes.h
+++ b/Engine/Source/Runtime/Engine/Classes/Camera/CameraTypes.h
@@ -3,6 +3,7 @@
 #pragma once
 
 #include "Engine/Scene.h"
+#include "Engine/EngineTypes.h"
 #include "CameraTypes.generated.h"
 
 //@TODO: Document
