[UE-9929] Fixed up name generation of the component template in the SimpleConstructionScript. Just using a guid instead of trying to a generate a name that hopefully doesn't conflict with an SCS_Node (UE-9929 specifically results from creating a component template with a name that is used by a SCS node that is created later)

diff --git a/Engine/Source/Editor/Kismet/Private/SSCSEditor.cpp b/Engine/Source/Editor/Kismet/Private/SSCSEditor.cpp
--- a/Engine/Source/Editor/Kismet/Private/SSCSEditor.cpp
+++ b/Engine/Source/Editor/Kismet/Private/SSCSEditor.cpp
@@ -4711,11 +4711,8 @@ void SSCSEditor::PasteNodes()
 			UActorComponent* NewActorComponent = NewObjectIt->Value;
 			check(NewActorComponent);
 
-			// Relocate the instance from the transient package to the BPGC and assign it a unique object name
-			NewActorComponent->Rename(NULL, Blueprint->GeneratedClass, REN_DontCreateRedirectors | REN_DoNotDirty);
-
 			// Create a new SCS node to contain the new component and add it to the tree
-			NewActorComponent = AddNewNode(Blueprint->SimpleConstructionScript->CreateNode(NewActorComponent), NULL, false, false);
+			NewActorComponent = AddNewNode(Blueprint->SimpleConstructionScript->CreateNodeAndRenameComponent(NewActorComponent), NULL, false, false);
 
 			if (NewActorComponent)
 			{
diff --git a/Engine/Source/Editor/KismetCompiler/Private/KismetCompilerMisc.cpp b/Engine/Source/Editor/KismetCompiler/Private/KismetCompilerMisc.cpp
--- a/Engine/Source/Editor/KismetCompiler/Private/KismetCompilerMisc.cpp
+++ b/Engine/Source/Editor/KismetCompiler/Private/KismetCompilerMisc.cpp
@@ -656,7 +656,7 @@ UProperty* FKismetCompilerUtilities::CreatePropertyOnScope(UStruct* Scope, const
 	{
 		if (UObject* ExistingObject = FindObject<UObject>(Scope, *PropertyName.ToString(), false))
 		{
-			MessageLog.Error(*FString::Printf(TEXT("Internal Compiler Error:  Tried to create a property %s in scope %s, but another object of type %s already already exists there."), *PropertyName.ToString(), (Scope ? *Scope->GetName() : TEXT("None"))), *ExistingObject->GetFullName(Scope));
+			MessageLog.Error(*FString::Printf(TEXT("Internal Compiler Error:  Tried to create a property %s in scope %s, but another object of type %s already exists there."), *PropertyName.ToString(), (Scope ? *Scope->GetName() : TEXT("None"))), *ExistingObject->GetFullName(Scope));
 
 			// Find a free name, so we can still create the property to make it easier to spot the duplicates, and avoid crashing
 			uint32 Counter = 0;
diff --git a/Engine/Source/Runtime/Engine/Classes/Engine/SCS_Node.h b/Engine/Source/Runtime/Engine/Classes/Engine/SCS_Node.h
--- a/Engine/Source/Runtime/Engine/Classes/Engine/SCS_Node.h
+++ b/Engine/Source/Runtime/Engine/Classes/Engine/SCS_Node.h
@@ -74,23 +74,6 @@ class USCS_Node : public UObject
 	virtual void PostLoad() override;
 	// End UObject Interface
 
-	/** 
-	 * Create a list of existing names used in this function. 
-	 *
-	 * @param	CurrentNames		Array to populate with the names
-	 *
-	 */
-	ENGINE_API void GenerateListOfExistingNames( TArray<FName>& CurrentNames ) const;
-
-	/** 
-	 * Create a new unique name for a component. 
-	 *
-	 * @param	CurrentNames		List of current names used in the blueprint
-	 * @param	DesiredName			Optional desired name to use; if NAME_None, a default name will be generated based on the component type.
-	 *
-	 */	 
-	FName GenerateNewComponentName( TArray<FName>& CurrentNames, FName DesiredName = NAME_None ) const;
-
 	/** The scene component constructed for component editing in the SCS editor */
 	class USceneComponent* EditorComponentInstance;
 
diff --git a/Engine/Source/Runtime/Engine/Classes/Engine/SimpleConstructionScript.h b/Engine/Source/Runtime/Engine/Classes/Engine/SimpleConstructionScript.h
--- a/Engine/Source/Runtime/Engine/Classes/Engine/SimpleConstructionScript.h
+++ b/Engine/Source/Runtime/Engine/Classes/Engine/SimpleConstructionScript.h
@@ -40,6 +40,9 @@ class USimpleConstructionScript : public UObject
 	/** Return all nodes in tree as a flat list */
 	ENGINE_API TArray<USCS_Node*> GetAllNodes() const;
 
+	/** Return immutable references to nodes in tree as a flat list */
+	ENGINE_API TArray<const USCS_Node*> GetAllNodesConst() const;
+
 	/** Provides read-only access to the root node set */
 	const ENGINE_API TArray<USCS_Node*>& GetRootNodes() const { return RootNodes; }
 
@@ -87,12 +90,21 @@ private:
 	TArray<USCS_Node*> ActorComponentNodes_DEPRECATED;
 
 #if WITH_EDITOR
+	/** Helper function for generating list of currently used names */
+	void GenerateListOfExistingNames(TArray<FName>& CurrentNames) const;
+
+	/** Helper function for generating a name for a new node, DesiredName can be used to optionally request a name, ComponentClass is mandatory */
+	FName GenerateNewComponentName(const UClass* ComponentClass, FName DesiredName = NAME_None) const;
+
+	/** Helper function to perform actions that all node creation functions have in common */
+	USCS_Node* CreateNodeImpl(UActorComponent* NewComponentTemplate, FName ComponentVariableName);
+
 public:
 	/** Creates a new SCS node using the given class to create the component template */
 	ENGINE_API USCS_Node* CreateNode(class UClass* NewComponentClass, FName NewComponentVariableName = NAME_None);
 
 	/** Creates a new SCS node using the given component template instance */
-	ENGINE_API USCS_Node* CreateNode(UActorComponent* NewComponentTemplate, FName NewComponentVariableName = NAME_None);
+	ENGINE_API USCS_Node* CreateNodeAndRenameComponent(UActorComponent* ExistingTemplate);
 
 	/** Ensures that all nodes in the SCS have valid names for compilation/replication */
 	ENGINE_API void ValidateNodeVariableNames(class FCompilerResultsLog& MessageLog);
diff --git a/Engine/Source/Runtime/Engine/Private/SCS_Node.cpp b/Engine/Source/Runtime/Engine/Private/SCS_Node.cpp
--- a/Engine/Source/Runtime/Engine/Private/SCS_Node.cpp
+++ b/Engine/Source/Runtime/Engine/Private/SCS_Node.cpp
@@ -391,101 +391,6 @@ USceneComponent* USCS_Node::GetParentComponentTemplate(UBlueprint* InBlueprint)
 	return ParentComponentTemplate;
 }
 
-void USCS_Node::GenerateListOfExistingNames( TArray<FName>& CurrentNames ) const
-{
-	const USimpleConstructionScript* SCS = GetSCS();
-	check(SCS);
-
-	const UBlueprintGeneratedClass* OwnerClass = Cast<const UBlueprintGeneratedClass>(SCS->GetOuter());
-	const UBlueprint* Blueprint = Cast<const UBlueprint>(OwnerClass ? OwnerClass->ClassGeneratedBy : NULL);
-	// >>> Backwards Compatibility:  VER_UE4_EDITORONLY_BLUEPRINTS
-	if(!Blueprint)
-	{
-		Blueprint = Cast<UBlueprint>(SCS->GetOuter());
-	}
-	// <<< End Backwards Compatibility
-	check(Blueprint);
-
-	TArray<UObject*> NativeCDOChildren;
-	UClass* FirstNativeClass = FBlueprintEditorUtils::FindFirstNativeClass(Blueprint->ParentClass);
-	GetObjectsWithOuter(FirstNativeClass->GetDefaultObject(), NativeCDOChildren, false);
-
-	for (UObject* NativeCDOChild : NativeCDOChildren)
-	{
-		CurrentNames.Add(NativeCDOChild->GetFName());
-	}
-
-	if(Blueprint->SkeletonGeneratedClass)
-	{
-		// First add the class variables.
-		FBlueprintEditorUtils::GetClassVariableList(Blueprint, CurrentNames, true);
-		// Then the function names.
-		FBlueprintEditorUtils::GetFunctionNameList(Blueprint, CurrentNames);
-	}
-	
-	// Get the list of child nodes and add each child (including all their children)
-	TArray<const USCS_Node*> ChildrenNodes = GetAllNodes();
-	// And add their names
-	for (int32 NodeIndex=0; NodeIndex < ChildrenNodes.Num(); ++NodeIndex)
-	{
-		const USCS_Node* ChildNode = ChildrenNodes[NodeIndex];
-		if( ChildNode )
-		{
-			if( ChildNode->VariableName != NAME_None )
-			{
-				CurrentNames.Add( ChildNode->VariableName );
-			}
-		}
-	}
-
-	if(SCS->GetDefaultSceneRootNode())
-	{
-		CurrentNames.AddUnique(SCS->GetDefaultSceneRootNode()->GetVariableName());
-	}
-}
-
-FName USCS_Node::GenerateNewComponentName( TArray<FName>& CurrentNames, FName DesiredName ) const
-{
-	FName NewName;
-	if( ComponentTemplate )
-	{
-		if(DesiredName != NAME_None && !CurrentNames.Contains(DesiredName))
-		{
-			NewName = DesiredName;
-		}
-		else
-		{
-			FString ComponentName;
-			if(DesiredName != NAME_None)
-			{
-				ComponentName = DesiredName.ToString();
-			}
-			else
-			{
-				const UClass* ComponentClass = ComponentTemplate->GetClass();
-				ComponentName = ComponentClass->GetName();
-
-				if (!ComponentClass->HasAnyClassFlags(CLASS_CompiledFromBlueprint))
-				{
-					ComponentName.RemoveFromEnd(TEXT("Component"));
-				}
-				else
-				{
-					ComponentName.RemoveFromEnd("_C");
-				}
-			}
-			
-			NewName = *ComponentName;
-			int32 Counter = 1;
-			while (CurrentNames.Contains(NewName))
-			{
-				NewName = FName(*( FString::Printf(TEXT("%s%d"), *ComponentName, Counter++) ));		
-			} 
-		}
-	}
-	return NewName;
-}
-
 void USCS_Node::PostLoad()
 {
 	Super::PostLoad();
diff --git a/Engine/Source/Runtime/Engine/Private/SimpleConstructionScript.cpp b/Engine/Source/Runtime/Engine/Private/SimpleConstructionScript.cpp
--- a/Engine/Source/Runtime/Engine/Private/SimpleConstructionScript.cpp
+++ b/Engine/Source/Runtime/Engine/Private/SimpleConstructionScript.cpp
@@ -448,6 +448,11 @@ TArray<USCS_Node*> USimpleConstructionScript::GetAllNodes() const
 	return AllNodes;
 }
 
+TArray<const USCS_Node*> USimpleConstructionScript::GetAllNodesConst() const
+{
+	return TArray<const USCS_Node*>(GetAllNodes());
+}
+
 void USimpleConstructionScript::AddNode(USCS_Node* Node)
 {
 	if(!RootNodes.Contains(Node))
@@ -731,63 +736,135 @@ void USimpleConstructionScript::ValidateSceneRootNodes()
 }
 
 #if WITH_EDITOR
-
-USCS_Node* USimpleConstructionScript::CreateNode(UClass* NewComponentClass, FName NewComponentVariableName)
+void USimpleConstructionScript::GenerateListOfExistingNames(TArray<FName>& CurrentNames) const
 {
-	UBlueprint* Blueprint = GetBlueprint();
-	check(Blueprint != NULL);
+	TArray<const USCS_Node*> ChildrenNodes = GetAllNodesConst();
+	const UBlueprintGeneratedClass* OwnerClass = Cast<const UBlueprintGeneratedClass>(GetOuter());
+	const UBlueprint* Blueprint = Cast<const UBlueprint>(OwnerClass ? OwnerClass->ClassGeneratedBy : NULL);
+	// >>> Backwards Compatibility:  VER_UE4_EDITORONLY_BLUEPRINTS
+	if (!Blueprint)
+	{
+		Blueprint = Cast<UBlueprint>(GetOuter());
+	}
+	// <<< End Backwards Compatibility
+	check(Blueprint);
 
-	// Ensure that the given class is of type UActorComponent
-	check(NewComponentClass->IsChildOf(UActorComponent::StaticClass()));
+	TArray<UObject*> NativeCDOChildren;
+	UClass* FirstNativeClass = FBlueprintEditorUtils::FindFirstNativeClass(Blueprint->ParentClass);
+	GetObjectsWithOuter(FirstNativeClass->GetDefaultObject(), NativeCDOChildren, false);
 
-	ensure(NULL != Cast<UBlueprintGeneratedClass>(Blueprint->GeneratedClass));
+	for (UObject* NativeCDOChild : NativeCDOChildren)
+	{
+		CurrentNames.Add(NativeCDOChild->GetFName());
+	}
 
-	FName NewComponentName(NAME_None);
-	if (NewComponentClass->ClassGeneratedBy != nullptr)
+	if (Blueprint->SkeletonGeneratedClass)
 	{
-		const FString NewClassName = FBlueprintEditorUtils::GetClassNameWithoutSuffix(NewComponentClass);
-		NewComponentName = MakeUniqueObjectName(Blueprint->GeneratedClass, NewComponentClass, FName(*NewClassName));
+		// First add the class variables.
+		FBlueprintEditorUtils::GetClassVariableList(Blueprint, CurrentNames, true);
+		// Then the function names.
+		FBlueprintEditorUtils::GetFunctionNameList(Blueprint, CurrentNames);
 	}
 
-	auto NewComponentTemplate = NewObject<UActorComponent>(Blueprint->GeneratedClass, NewComponentClass, NewComponentName);
-	NewComponentTemplate->SetFlags(RF_ArchetypeObject|RF_Transactional|RF_Public);
+	// And add their names
+	for (int32 NodeIndex = 0; NodeIndex < ChildrenNodes.Num(); ++NodeIndex)
+	{
+		const USCS_Node* ChildNode = ChildrenNodes[NodeIndex];
+		if (ChildNode)
+		{
+			if (ChildNode->VariableName != NAME_None)
+			{
+				CurrentNames.Add(ChildNode->VariableName);
+			}
+		}
+	}
 
-	return CreateNode(NewComponentTemplate, NewComponentVariableName);
+	if (GetDefaultSceneRootNode())
+	{
+		CurrentNames.AddUnique(GetDefaultSceneRootNode()->GetVariableName());
+	}
 }
 
-USCS_Node* USimpleConstructionScript::CreateNode(UActorComponent* NewComponentTemplate, FName NewComponentVariableName)
+FName USimpleConstructionScript::GenerateNewComponentName(const UClass* ComponentClass, FName DesiredName ) const
 {
-	if(NewComponentTemplate)
+	TArray<FName> CurrentNames;
+	GenerateListOfExistingNames(CurrentNames);
+
+	FName NewName;
+	if (ComponentClass)
 	{
-		// Create a node for the script, and save a pointer to the template
-		// NewNamedObject to work around the fact we shouldn't use NewObject for default subobjects
-		auto NewNode = NewObject<USCS_Node>(this, MakeUniqueObjectName(this, USCS_Node::StaticClass()));
-		NewNode->SetFlags(RF_Transactional);
-		NewNode->ComponentTemplate = NewComponentTemplate;
+		if (DesiredName != NAME_None && !CurrentNames.Contains(DesiredName))
+		{
+			NewName = DesiredName;
+		}
+		else
+		{
+			FString ComponentName;
+			if (DesiredName != NAME_None)
+			{
+				ComponentName = DesiredName.ToString();
+			}
+			else
+			{
+				ComponentName = ComponentClass->GetName();
 
-		// Get a list of names currently in use.
-		TArray<FName> CurrentNames;
-		NewNode->GenerateListOfExistingNames( CurrentNames );
+				if (!ComponentClass->HasAnyClassFlags(CLASS_CompiledFromBlueprint))
+				{
+					ComponentName.RemoveFromEnd(TEXT("Component"));
+				}
+				else
+				{
+					ComponentName.RemoveFromEnd("_C");
+				}
+			}
 
-		if (NewComponentVariableName.ToString().EndsWith(TEXT("_C")))
-		{
-			const FString NameAsString = NewComponentVariableName.ToString();
-			const int32 NewStrLen = NameAsString.Len() - 2;
-			NewComponentVariableName = FName(*NameAsString.Left(NewStrLen));
+			NewName = *ComponentName;
+			int32 Counter = 1;
+			while (CurrentNames.Contains(NewName))
+			{
+				NewName = FName(*(FString::Printf(TEXT("%s%d"), *ComponentName, Counter++)));
+			}
 		}
+	}
+	return NewName;
+}
 
-		// Now create a name for the new component.
-		NewNode->VariableName = NewNode->GenerateNewComponentName( CurrentNames, NewComponentVariableName );
+USCS_Node* USimpleConstructionScript::CreateNodeImpl(UActorComponent* NewComponentTemplate, FName ComponentVariableName)
+{
+	auto NewNode = NewObject<USCS_Node>(this, MakeUniqueObjectName(this, USCS_Node::StaticClass()));
+	NewNode->SetFlags(RF_Transactional);
+	NewNode->ComponentTemplate = NewComponentTemplate;
 
-		// Note: This should match up with UEdGraphSchema_K2::VR_DefaultCategory
-		NewNode->CategoryName = TEXT("Default");
+	// Now create a name for the new component.
+	NewNode->VariableName = GenerateNewComponentName(NewComponentTemplate->GetClass(), ComponentVariableName);
 
-		NewNode->VariableGuid = FGuid::NewGuid();
+	// Note: This should match up with UEdGraphSchema_K2::VR_DefaultCategory
+	NewNode->CategoryName = TEXT("Default");
+	NewNode->VariableGuid = FGuid::NewGuid();
+	return NewNode;
+}
 
-		return NewNode;
-	}
+USCS_Node* USimpleConstructionScript::CreateNode(UClass* NewComponentClass, FName NewComponentVariableName)
+{
+	UBlueprint* Blueprint = GetBlueprint();
+	check(Blueprint);
+	check(NewComponentClass->IsChildOf(UActorComponent::StaticClass()));
+	ensure(Cast<UBlueprintGeneratedClass>(Blueprint->GeneratedClass));
 
-	return NULL;
+	auto NewComponentTemplate = NewObject<UActorComponent>(Blueprint->GeneratedClass, NewComponentClass, *(FGuid::NewGuid().ToString()) /*NewComponentName*/);
+	NewComponentTemplate->SetFlags(RF_ArchetypeObject|RF_Transactional|RF_Public);
+
+	return CreateNodeImpl(NewComponentTemplate, NewComponentVariableName);
+}
+
+USCS_Node* USimpleConstructionScript::CreateNodeAndRenameComponent(UActorComponent* NewComponentTemplate)
+{
+	check(NewComponentTemplate);
+
+	// Relocate the instance from the transient package to the BPGC and assign it a unique object name
+	NewComponentTemplate->Rename(*(FGuid::NewGuid().ToString()), GetBlueprint()->GeneratedClass, REN_DontCreateRedirectors | REN_DoNotDirty);
+
+	return CreateNodeImpl(NewComponentTemplate, NAME_None);
 }
 
 void USimpleConstructionScript::ValidateNodeVariableNames(FCompilerResultsLog& MessageLog)
@@ -822,12 +899,8 @@ void USimpleConstructionScript::ValidateNodeVariableNames(FCompilerResultsLog& M
 			{
 				FName OldName = Node->VariableName;
 
-				// Get a list of names currently in use.
-				TArray<FName> CurrentNames;
-				Node->GenerateListOfExistingNames( CurrentNames );
-
 				// Generate a new default variable name for the component.
-				Node->VariableName = Node->GenerateNewComponentName( CurrentNames );
+				Node->VariableName = GenerateNewComponentName(Node->ComponentTemplate->GetClass());
 				Node->bVariableNameAutoGenerated_DEPRECATED = false;
 
 				if( OldName != NAME_None )
