PR #842: Inverse Lerp function for blueprints from user JaredTherriault

Made some minor changes to the documenting comment

diff --git a/Engine/Source/Runtime/Engine/Classes/Kismet/KismetMathLibrary.h b/Engine/Source/Runtime/Engine/Classes/Kismet/KismetMathLibrary.h
--- a/Engine/Source/Runtime/Engine/Classes/Kismet/KismetMathLibrary.h
+++ b/Engine/Source/Runtime/Engine/Classes/Kismet/KismetMathLibrary.h
@@ -441,6 +441,19 @@ class UKismetMathLibrary : public UBlueprintFunctionLibrary
 	/* Linearly interpolates between A and B based on Alpha (100% of A when Alpha=0 and 100% of B when Alpha=1) */
 	UFUNCTION(BlueprintPure, Category="Math|Float")
 	static float Lerp(float A, float B, float Alpha);
+	
+	/**
+	* Returns the percentage of the range B-A that corresponds to Value. E.g.,
+	*	A = 0, B = 8, Value = 3 : return value = 3/8, indicating Value is 3/8 from A to B 
+	*	A = 8, B = 0, Value = 3 : return value = 5/8, indicating Value is 5/8 from A to B
+	* Named InverseLerp because Lerp( A, B, InverseLerp(A, B, Value) ) == Value
+	* @param A The "from" value this float could be, usually but not necessarily a minimum. Returned as 0.
+	* @param B The "to" value this float could be, usually but not necessarily a maximum. Returned as 1.
+	* @param Value A value intended to be normalized relative to B-A
+	* @return A normalized value considering A and B.
+	*/
+	UFUNCTION(BlueprintPure, meta = (Keywords = "percentage normalize range"), Category = "Math|Float")
+	static float InverseLerp(float A, float B, float Value);
 
 	/* Easeing  between A and B using a specified easing function */
 	UFUNCTION(BlueprintPure, meta = (FriendlyName = "Ease", BlueprintInternalUseOnly = "true"), Category = "Math|Interpolation")
diff --git a/Engine/Source/Runtime/Engine/Private/KismetMathLibrary.cpp b/Engine/Source/Runtime/Engine/Private/KismetMathLibrary.cpp
--- a/Engine/Source/Runtime/Engine/Private/KismetMathLibrary.cpp
+++ b/Engine/Source/Runtime/Engine/Private/KismetMathLibrary.cpp
@@ -618,6 +618,25 @@ float UKismetMathLibrary::Lerp(float A, float B, float V)
 	return A + V*(B-A);
 }	
 
+float UKismetMathLibrary::InverseLerp(float A, float B, float Value)
+{
+	if (FMath::IsNearlyEqual(A, B))
+	{
+		if (Value < A)
+		{
+			return 0;
+		}
+		else
+		{
+			return 1;
+		}
+	}
+	else
+	{
+		return ((Value - A) / (B - A));
+	}
+}
+
 float UKismetMathLibrary::Ease(float A, float B, float Alpha, TEnumAsByte<EEasingFunc::Type> EasingFunc, float BlendExp, int32 Steps)
 {
 	return Lerp(A, B, EaseAlpha(Alpha, EasingFunc, BlendExp, Steps));
