Fixed bug when fading in with a start time set - PlaybackTime always begins from 0 so adding StartTime isn't necessary

Found on Answerhub: https://answers.unrealengine.com/questions/177703/fade-in-audio-bug.html

diff --git a/Engine/Source/Runtime/Engine/Private/AudioComponent.cpp b/Engine/Source/Runtime/Engine/Private/AudioComponent.cpp
--- a/Engine/Source/Runtime/Engine/Private/AudioComponent.cpp
+++ b/Engine/Source/Runtime/Engine/Private/AudioComponent.cpp
@@ -177,7 +177,7 @@ void UAudioComponent::PlayInternal(const float StartTime, const float FadeInDura
 			if (FadeInDuration > 0.0f)
 			{
 				NewActiveSound.CurrentAdjustVolumeMultiplier = 0.f;
-				NewActiveSound.TargetAdjustVolumeStopTime = FMath::Max(0.f, StartTime) + FadeInDuration;
+				NewActiveSound.TargetAdjustVolumeStopTime = FadeInDuration;
 			}
 			else
 			{
