[INTEGRATE] Change 2456645 by [REDACTED_PERSON][REDACTED_HANDLE] on 2015/02/23 14:33:22

	Jira UE - 10313 : Crash playing animation montage in reverse

	Floating point error can trip this up if the planets align, clamp instead.

diff --git a/Engine/Source/Runtime/Engine/Private/Animation/AnimCompositeBase.cpp b/Engine/Source/Runtime/Engine/Private/Animation/AnimCompositeBase.cpp
--- a/Engine/Source/Runtime/Engine/Private/Animation/AnimCompositeBase.cpp
+++ b/Engine/Source/Runtime/Engine/Private/Animation/AnimCompositeBase.cpp
@@ -92,7 +92,7 @@ void FAnimSegment::GetAnimNotifiesFromTrackPositions(const float& PreviousTrackP
 
 			// Get starting position, closest overlap.
 			float AnimStartPosition = ConvertTrackPosToAnimPos( bTrackPlayingBackwards ? FMath::Min(PreviousTrackPosition, SegmentEndPos) : FMath::Max(PreviousTrackPosition, SegmentStartPos) );
-			check( (AnimStartPosition >= AnimStartTime) && (AnimStartPosition <= AnimEndTime) );
+			FMath::Clamp(AnimStartPosition, AnimStartTime, AnimEndTime);
 			float TrackTimeToGo = FMath::Abs(CurrentTrackPosition - PreviousTrackPosition);
 
 			// The track can be playing backwards and the animation can be playing backwards, so we
