[INTEGRATE] Change 2458903 by bruce.nesbit@bnesbit_Releases on 2015/02/24 14:16:06

	Fixed possible crash in USceneComponent::UpdatePhysicsVolume

diff --git a/Engine/Source/Runtime/Engine/Private/SceneComponent.cpp b/Engine/Source/Runtime/Engine/Private/SceneComponent.cpp
--- a/Engine/Source/Runtime/Engine/Private/SceneComponent.cpp
+++ b/Engine/Source/Runtime/Engine/Private/SceneComponent.cpp
@@ -1443,8 +1443,7 @@ void USceneComponent::UpdatePhysicsVolume( bool bTriggerNotifiers )
 		SCOPE_CYCLE_COUNTER(STAT_UpdatePhysicsVolume);
 
 		UWorld* const MyWorld = GetWorld();
-		APhysicsVolume *NewVolume = MyWorld->GetDefaultPhysicsVolume();
-		
+		APhysicsVolume* NewVolume = MyWorld->GetDefaultPhysicsVolume();
 		// Avoid doing anything if there are no other physics volumes in the world.
 		if (MyWorld->GetNonDefaultPhysicsVolumeCount() > 0)
 		{
@@ -1455,14 +1454,17 @@ void USceneComponent::UpdatePhysicsVolume( bool bTriggerNotifiers )
 			for (auto VolumeIter = MyWorld->GetNonDefaultPhysicsVolumeIterator(); VolumeIter && !bAnyPotentialOverlap; ++VolumeIter, ++VolumeIndex)
 			{
 				const APhysicsVolume* Volume = *VolumeIter;
-				const USceneComponent* VolumeRoot = Volume->GetRootComponent();
-				if (VolumeRoot)
+				if( Volume != nullptr )
 				{
-					if (FBoxSphereBounds::SpheresIntersect(VolumeRoot->Bounds, Bounds))
+					const USceneComponent* VolumeRoot = Volume->GetRootComponent();
+					if (VolumeRoot)
 					{
-						if (FBoxSphereBounds::BoxesIntersect(VolumeRoot->Bounds, Bounds))
+						if (FBoxSphereBounds::SpheresIntersect(VolumeRoot->Bounds, Bounds))
 						{
-							bAnyPotentialOverlap = true;
+							if (FBoxSphereBounds::BoxesIntersect(VolumeRoot->Bounds, Bounds))
+							{
+								bAnyPotentialOverlap = true;
+							}
 						}
 					}
 				}
@@ -1492,8 +1494,8 @@ void USceneComponent::UpdatePhysicsVolume( bool bTriggerNotifiers )
 				{
 					bOverlappedOrigin = true;
 					MyWorld->OverlapMulti(Hits, GetComponentLocation(), FQuat::Identity, GetCollisionObjectType(), FCollisionShape::MakeSphere(0.f), Params);
-				}
-
+				}				
+				
 				for (int32 HitIdx = 0; HitIdx < Hits.Num(); HitIdx++)
 				{
 					const FOverlapResult& Link = Hits[HitIdx];
