[INTEGRATE] Change 2460110 by [REDACTED_PERSON][REDACTED_HANDLE] on 2015/02/25 08:18:29

	Fix CL 2456645 - FMath::Clamp() returns a value, does not operate on references.

diff --git a/Engine/Source/Runtime/Engine/Private/Animation/AnimCompositeBase.cpp b/Engine/Source/Runtime/Engine/Private/Animation/AnimCompositeBase.cpp
--- a/Engine/Source/Runtime/Engine/Private/Animation/AnimCompositeBase.cpp
+++ b/Engine/Source/Runtime/Engine/Private/Animation/AnimCompositeBase.cpp
@@ -92,7 +92,7 @@ void FAnimSegment::GetAnimNotifiesFromTrackPositions(const float& PreviousTrackP
 
 			// Get starting position, closest overlap.
 			float AnimStartPosition = ConvertTrackPosToAnimPos( bTrackPlayingBackwards ? FMath::Min(PreviousTrackPosition, SegmentEndPos) : FMath::Max(PreviousTrackPosition, SegmentStartPos) );
-			FMath::Clamp(AnimStartPosition, AnimStartTime, AnimEndTime);
+			AnimStartPosition = FMath::Clamp(AnimStartPosition, AnimStartTime, AnimEndTime);
 			float TrackTimeToGo = FMath::Abs(CurrentTrackPosition - PreviousTrackPosition);
 
 			// The track can be playing backwards and the animation can be playing backwards, so we
