[INTEGRATE] Change 2454402 by [REDACTED_PERSON][REDACTED_HANDLE] on 2015/02/20 22:32:02

	Fix for UE-10260 - build a a HISMC cluster tree if it's empty in PostLoad.

diff --git a/Engine/Source/Runtime/Engine/Classes/Components/HierarchicalInstancedStaticMeshComponent.h b/Engine/Source/Runtime/Engine/Classes/Components/HierarchicalInstancedStaticMeshComponent.h
--- a/Engine/Source/Runtime/Engine/Classes/Components/HierarchicalInstancedStaticMeshComponent.h
+++ b/Engine/Source/Runtime/Engine/Classes/Components/HierarchicalInstancedStaticMeshComponent.h
@@ -92,6 +92,7 @@ public:
 
 	virtual FPrimitiveSceneProxy* CreateSceneProxy() override;
 	virtual FBoxSphereBounds CalcBounds(const FTransform& BoundTransform) const override;
+	virtual void PostLoad() override;
 
 	// UInstancedStaticMesh interface
 	virtual int32 AddInstance(const FTransform& InstanceTransform) override;
diff --git a/Engine/Source/Runtime/Engine/Private/HierarchicalInstancedStaticMesh.cpp b/Engine/Source/Runtime/Engine/Private/HierarchicalInstancedStaticMesh.cpp
--- a/Engine/Source/Runtime/Engine/Private/HierarchicalInstancedStaticMesh.cpp
+++ b/Engine/Source/Runtime/Engine/Private/HierarchicalInstancedStaticMesh.cpp
@@ -1369,7 +1369,7 @@ UHierarchicalInstancedStaticMeshComponent::~UHierarchicalInstancedStaticMeshComp
 	if (ProxySize)
 	{
 		DEC_DWORD_STAT_BY(STAT_FoliageInstanceBuffers, ProxySize);
-	}
+}
 	ProxySize = 0;
 }
 
@@ -1386,6 +1386,17 @@ void UHierarchicalInstancedStaticMeshComponent::Serialize(FArchive& Ar)
 	ClusterTree.BulkSerialize(Ar);
 }
 
+void UHierarchicalInstancedStaticMeshComponent::PostLoad()
+{
+	Super::PostLoad();
+
+	// For some reason we don't have a tree. Build one now!
+	if (PerInstanceSMData.Num() > 0 && ClusterTreePtr.IsValid() && ClusterTreePtr->Num() == 0 && !IsAsyncBuilding())
+	{
+		BuildTreeAsync();
+	}
+}
+
 SIZE_T UHierarchicalInstancedStaticMeshComponent::GetResourceSize( EResourceSizeMode::Type Mode )
 {
 	SIZE_T ResSize = Super::GetResourceSize(Mode);
@@ -1473,7 +1484,7 @@ bool UHierarchicalInstancedStaticMeshComponent::UpdateInstanceTransform(int32 In
 
 	if (StaticMesh)
 	{
-		UnbuiltInstanceBounds += StaticMesh->GetBounds().GetBox().TransformBy(NewInstanceTransform);
+	UnbuiltInstanceBounds += StaticMesh->GetBounds().GetBox().TransformBy(NewInstanceTransform);
 	}
 
 	if (!IsAsyncBuilding())
