[INTEGRATE] Change 2454747 by [REDACTED_PERSON][REDACTED_HANDLE] on 2015/02/21 13:20:48

	Fix for UE-10260. Make sure there is always a valid tree for hierarchical instanced static meshes.

diff --git a/Engine/Source/Runtime/Engine/Private/HierarchicalInstancedStaticMesh.cpp b/Engine/Source/Runtime/Engine/Private/HierarchicalInstancedStaticMesh.cpp
--- a/Engine/Source/Runtime/Engine/Private/HierarchicalInstancedStaticMesh.cpp
+++ b/Engine/Source/Runtime/Engine/Private/HierarchicalInstancedStaticMesh.cpp
@@ -1391,9 +1391,10 @@ void UHierarchicalInstancedStaticMeshComponent::PostLoad()
 	Super::PostLoad();
 
 	// For some reason we don't have a tree. Build one now!
-	if (PerInstanceSMData.Num() > 0 && ClusterTreePtr.IsValid() && ClusterTreePtr->Num() == 0 && !IsAsyncBuilding())
+	if (PerInstanceSMData.Num() > 0 && (!ClusterTreePtr.IsValid() || ClusterTreePtr->Num() == 0))
 	{
-		BuildTreeAsync();
+		check(!IsAsyncBuilding());
+		BuildTree();
 	}
 }
 
