#ue4 -  Fix for "penetration" causing character movement to hitch. Rather than stopping all movement that update, deflect off the depenetration vector and continue movement.

- Add "p.PenetrationOverlapCheckInflation" to replace previously hardcoded value, for easier tuning of inflation used in overlap check when resolving overlap.
- Fix jumping when pressing up against corners causing a very quick landing, caused by using invalid penetrating sweep result.

diff --git a/Engine/Source/Runtime/Engine/Private/CharacterMovementComponent.cpp b/Engine/Source/Runtime/Engine/Private/CharacterMovementComponent.cpp
--- a/Engine/Source/Runtime/Engine/Private/CharacterMovementComponent.cpp
+++ b/Engine/Source/Runtime/Engine/Private/CharacterMovementComponent.cpp
@@ -3474,10 +3474,16 @@ void UCharacterMovementComponent::MoveAlongFloor(const FVector& InVelocity, floa
 	
 	if (Hit.bStartPenetrating)
 	{
-		OnCharacterStuckInGeometry();
-	}
+		// Allow this hit to be used as an impact we can deflect off, otherwise we do nothing the rest of the update and appear to hitch.
+		HandleImpact(Hit);
+		SlideAlongSurface(Delta, 1.f, Hit.Normal, Hit, true);
 
-	if (Hit.IsValidBlockingHit())
+		if (Hit.bStartPenetrating)
+		{
+			OnCharacterStuckInGeometry();
+		}
+	}
+	else if (Hit.IsValidBlockingHit())
 	{
 		// We impacted something (most likely another ramp, but possibly a barrier).
 		float PercentTimeApplied = Hit.Time;
@@ -4857,6 +4863,15 @@ bool UCharacterMovementComponent::IsValidLandingSpot(const FVector& CapsuleLocat
 			return false;
 		}
 	}
+	else
+	{
+		// Penetrating
+		if (Hit.Normal.Z < KINDA_SMALL_NUMBER)
+		{
+			// Normal is nearly horizontal or downward, that's a penetration adjustment next to a vertical or overhanging wall. Don't pop to the floor.
+			return false;
+		}
+	}
 
 	FFindFloorResult FloorResult;
 	FindFloor(CapsuleLocation, FloorResult, false, &Hit);
diff --git a/Engine/Source/Runtime/Engine/Private/MovementComponent.cpp b/Engine/Source/Runtime/Engine/Private/MovementComponent.cpp
--- a/Engine/Source/Runtime/Engine/Private/MovementComponent.cpp
+++ b/Engine/Source/Runtime/Engine/Private/MovementComponent.cpp
@@ -476,6 +476,12 @@ static TAutoConsoleVariable<float> CVarPenetrationPullbackDistance(TEXT("p.Penet
 	TEXT("Distance added to penetration fix-ups."),
 	ECVF_Default);
 
+static TAutoConsoleVariable<float> CVarPenetrationOverlapCheckInflation(TEXT("p.PenetrationOverlapCheckInflation"),
+	0.100,
+	TEXT("Inflation added to object when checking if a location is free of blocking collision.\n")
+	TEXT("Distance added to inflation in penetration overlap check."),
+	ECVF_Default);
+
 FVector UMovementComponent::GetPenetrationAdjustment(const FHitResult& Hit) const
 {
 	if (!Hit.bStartPenetrating)
@@ -507,7 +513,7 @@ bool UMovementComponent::ResolvePenetration(const FVector& ProposedAdjustment, c
 
 		// We really want to make sure that precision differences or differences between the overlap test and sweep tests don't put us into another overlap,
 		// so make the overlap test a bit more restrictive.
-		const float OverlapInflation = 0.10f;
+		const float OverlapInflation = CVarPenetrationOverlapCheckInflation.GetValueOnGameThread();
 		bool bEncroached = OverlapTest(Hit.TraceStart + Adjustment, NewRotation.Quaternion(), UpdatedPrimitive->GetCollisionObjectType(), UpdatedPrimitive->GetCollisionShape(OverlapInflation), ActorOwner);
 		if (!bEncroached)
 		{
