Engine: Add Thickness option to DrawFlatArrow, DrawDirectionalArrow, and DrawCoordinateSystem
- Moved comments on drawing utility functions from the cpp to the header

diff --git a/Engine/Source/Runtime/Engine/Private/PrimitiveDrawingUtils.cpp b/Engine/Source/Runtime/Engine/Private/PrimitiveDrawingUtils.cpp
--- a/Engine/Source/Runtime/Engine/Private/PrimitiveDrawingUtils.cpp
+++ b/Engine/Source/Runtime/Engine/Private/PrimitiveDrawingUtils.cpp
@@ -655,19 +655,6 @@ void DrawCylinder(FPrimitiveDrawInterface* PDI, const FMatrix& CylToWorld, const
 	MeshBuilder.Draw(PDI, CylToWorld, MaterialRenderProxy, DepthPriority,0.f);
 }
 
-/**
- * Draws a circle using triangles.
- *
- * @param	PDI						Draw interface.
- * @param	Base					Center of the circle.
- * @param	XAxis					X alignment axis to draw along.
- * @param	YAxis					Y alignment axis to draw along.
- * @param	Color					Color of the circle.
- * @param	Radius					Radius of the circle.
- * @param	NumSides				Numbers of sides that the circle has.
- * @param	MaterialRenderProxy		Material to use for render 
- * @param	DepthPriority			Depth priority for the circle.
- */
 void DrawDisc(class FPrimitiveDrawInterface* PDI,const FVector& Base,const FVector& XAxis,const FVector& YAxis,FColor Color,float Radius,int32 NumSides,const FMaterialRenderProxy* MaterialRenderProxy, uint8 DepthPriority)
 {
 	check (NumSides >= 3);
@@ -717,34 +704,7 @@ void DrawDisc(class FPrimitiveDrawInterface* PDI,const FVector& Base,const FVect
 	MeshBuilder.Draw(PDI, FMatrix::Identity, MaterialRenderProxy, DepthPriority,0.f);
 }
 
-/**
- * Draws a flat arrow with an outline.
- *
- * @param	PDI						Draw interface.
- * @param	Base					Base of the arrow.
- * @param	XAxis					X alignment axis to draw along.
- * @param	YAxis					Y alignment axis to draw along.
- * @param	Color					Color of the circle.
- * @param	Length					Length of the arrow, from base to tip.
- * @param	Width					Width of the base of the arrow, head of the arrow will be 2x.
- * @param	MaterialRenderProxy		Material to use for render 
- * @param	DepthPriority			Depth priority for the circle.
- */
-
-/*
-x-axis is from point 0 to point 2
-y-axis is from point 0 to point 1
-		6
-		/\
-	   /  \
-	  /    \
-	 4_2  3_5
-	   |  |
-	   0__1
-*/
-
-
-void DrawFlatArrow(class FPrimitiveDrawInterface* PDI,const FVector& Base,const FVector& XAxis,const FVector& YAxis,FColor Color,float Length,int32 Width, const FMaterialRenderProxy* MaterialRenderProxy, uint8 DepthPriority)
+void DrawFlatArrow(class FPrimitiveDrawInterface* PDI,const FVector& Base,const FVector& XAxis,const FVector& YAxis,FColor Color,float Length,int32 Width, const FMaterialRenderProxy* MaterialRenderProxy, uint8 DepthPriority, float Thickness)
 {
 	float DistanceFromBaseToHead = Length/3.0f;
 	float DistanceFromBaseToTip = DistanceFromBaseToHead*2.0f;
@@ -808,14 +768,6 @@ void DrawFlatArrow(class FPrimitiveDrawInterface* PDI,const FVector& Base,const
 
 // Line drawing utility functions.
 
-/**
- * Draws a wireframe box.
- *
- * @param	PDI				Draw interface.
- * @param	Box				The FBox to use for drawing.
- * @param	Color			Color of the box.
- * @param	DepthPriority	Depth priority for the circle.
- */
 void DrawWireBox(FPrimitiveDrawInterface* PDI, const FBox& Box, const FLinearColor& Color, uint8 DepthPriority, float Thickness, float DepthBias, bool bScreenSpace)
 {
 	FVector	B[2],P,Q;
@@ -878,19 +830,6 @@ void DrawWireBox(FPrimitiveDrawInterface* PDI, const FMatrix& Matrix, const FBox
 }
 
 
-/**
- * Draws a circle using lines.
- *
- * @param	PDI				Draw interface.
- * @param	Base			Center of the circle.
- * @param	X				X alignment axis to draw along.
- * @param	Y				Y alignment axis to draw along.
- * @param	Z				Z alignment axis to draw along.
- * @param	Color			Color of the circle.
- * @param	Radius			Radius of the circle.
- * @param	NumSides		Numbers of sides that the circle has.
- * @param	DepthPriority	Depth priority for the circle.
- */
 void DrawCircle(FPrimitiveDrawInterface* PDI, const FVector& Base, const FVector& X, const FVector& Y, const FLinearColor& Color, float Radius, int32 NumSides, uint8 DepthPriority, float Thickness, float DepthBias, bool bScreenSpace)
 {
 	const float	AngleDelta = 2.0f * PI / NumSides;
@@ -904,20 +843,6 @@ void DrawCircle(FPrimitiveDrawInterface* PDI, const FVector& Base, const FVector
 	}
 }
 
-/**
- * Draws an arc using lines.
- *
- * @param	PDI				Draw interface.
- * @param	Base			Center of the circle.
- * @param	X				Normalized axis from one point to the center
- * @param	Y				Normalized axis from other point to the center
- * @param   MinAngle        The minimum angle
- * @param   MinAngle        The maximum angle
- * @param   Radius          Radius of the arc
- * @param	Sections		Numbers of sides that the circle has.
- * @param	Color			Color of the circle.
- * @param	DepthPriority	Depth priority for the circle.
- */
 void DrawArc(FPrimitiveDrawInterface* PDI, const FVector Base, const FVector X, const FVector Y, const float MinAngle, const float MaxAngle, const float Radius, const int32 Sections, const FLinearColor& Color, uint8 DepthPriority)
 {
 	float AngleStep = (MaxAngle - MinAngle)/((float)(Sections));
@@ -935,16 +860,6 @@ void DrawArc(FPrimitiveDrawInterface* PDI, const FVector Base, const FVector X,
 	}
 }
 
-/**
- * Draws a sphere using circles.
- *
- * @param	PDI				Draw interface.
- * @param	Base			Center of the sphere.
- * @param	Color			Color of the sphere.
- * @param	Radius			Radius of the sphere.
- * @param	NumSides		Numbers of sides that the circle has.
- * @param	DepthPriority	Depth priority for the circle.
- */
 void DrawWireSphere(class FPrimitiveDrawInterface* PDI, const FVector& Base, const FLinearColor& Color, float Radius, int32 NumSides, uint8 DepthPriority, float Thickness, float DepthBias, bool bScreenSpace)
 {
 	DrawCircle(PDI, Base, FVector(1,0,0), FVector(0,1,0), Color, Radius, NumSides, DepthPriority, Thickness, DepthBias, bScreenSpace);
@@ -952,15 +867,6 @@ void DrawWireSphere(class FPrimitiveDrawInterface* PDI, const FVector& Base, con
 	DrawCircle(PDI, Base, FVector(0, 1, 0), FVector(0, 0, 1), Color, Radius, NumSides, DepthPriority, Thickness, DepthBias, bScreenSpace);
 }
 
-/**
- * Draws a sphere using circles, automatically calculating a reasonable number of sides
- *
- * @param	PDI				Draw interface.
- * @param	Base			Center of the sphere.
- * @param	Color			Color of the sphere.
- * @param	Radius			Radius of the sphere.
- * @param	DepthPriority	Depth priority for the circle.
- */
 void DrawWireSphereAutoSides(class FPrimitiveDrawInterface* PDI, const FVector& Base, const FLinearColor& Color, float Radius, uint8 DepthPriority, float Thickness, float DepthBias, bool bScreenSpace)
 {
 	// Guess a good number of sides
@@ -983,20 +889,6 @@ void DrawWireSphereAutoSides(class FPrimitiveDrawInterface* PDI, const FTransfor
 	DrawWireSphere(PDI, Transform, Color, Radius, NumSides, DepthPriority, Thickness, DepthBias, bScreenSpace);
 }
 
-/**
- * Draws a wireframe cylinder.
- *
- * @param	PDI				Draw interface.
- * @param	Base			Center pointer of the base of the cylinder.
- * @param	X				X alignment axis to draw along.
- * @param	Y				Y alignment axis to draw along.
- * @param	Z				Z alignment axis to draw along.
- * @param	Color			Color of the cylinder.
- * @param	Radius			Radius of the cylinder.
- * @param	HalfHeight		Half of the height of the cylinder.
- * @param	NumSides		Numbers of sides that the cylinder has.
- * @param	DepthPriority	Depth priority for the cylinder.
- */
 void DrawWireCylinder(FPrimitiveDrawInterface* PDI, const FVector& Base, const FVector& X, const FVector& Y, const FVector& Z, const FLinearColor& Color, float Radius, float HalfHeight, int32 NumSides, uint8 DepthPriority, float Thickness, float DepthBias, bool bScreenSpace)
 {
 	const float	AngleDelta = 2.0f * PI / NumSides;
@@ -1028,21 +920,6 @@ static void DrawHalfCircle(FPrimitiveDrawInterface* PDI, const FVector& Base, co
 	}	
 }
 
-
-/**
- * Draws a wireframe capsule.
- *
- * @param	PDI				Draw interface.
- * @param	Base			Center pointer of the base of the cylinder.
- * @param	X				X alignment axis to draw along.
- * @param	Y				Y alignment axis to draw along.
- * @param	Z				Z alignment axis to draw along.
- * @param	Color			Color of the cylinder.
- * @param	Radius			Radius of the cylinder.
- * @param	HalfHeight		Half of the height of the cylinder.
- * @param	NumSides		Numbers of sides that the cylinder has.
- * @param	DepthPriority	Depth priority for the cylinder.
- */
 void DrawWireCapsule(FPrimitiveDrawInterface* PDI, const FVector& Base, const FVector& X, const FVector& Y, const FVector& Z, const FLinearColor& Color, float Radius, float HalfHeight, int32 NumSides, uint8 DepthPriority, float Thickness, float DepthBias, bool bScreenSpace)
 {
 	const FVector Origin = Base;
@@ -1093,20 +970,6 @@ void DrawWireCapsule(FPrimitiveDrawInterface* PDI, const FVector& Base, const FV
 	}
 }
 
-
-
-/**
- * Draws a wireframe cone
- *
- * @param	PDI				Draw interface.
- * @param	Transform		Generic transform to apply (ex. a local-to-world transform).
- * @param	ConeRadius		Radius of the cone.
- * @param	ConeAngle		Angle of the cone.
- * @param	ConeSides		Numbers of sides that the cone has.
- * @param	Color			Color of the cone.
- * @param	DepthPriority	Depth priority for the cone.
- * @param	Verts			Out param, the positions of the verts at the cone's base.
- */
 void DrawWireCone(FPrimitiveDrawInterface* PDI, TArray<FVector>& Verts, const FTransform& Transform, float ConeRadius, float ConeAngle, int32 ConeSides, const FLinearColor& Color, uint8 DepthPriority, float Thickness, float DepthBias, bool bScreenSpace)
 {
 	static const float TwoPI = 2.0f * PI;
@@ -1154,19 +1017,6 @@ void DrawWireCone(FPrimitiveDrawInterface* PDI, TArray<FVector>& Verts, const FM
 	DrawWireCone(PDI, Verts, FTransform(Transform), ConeRadius, ConeAngle, ConeSides, Color, DepthPriority, Thickness, DepthBias, bScreenSpace);
 }
 
-/**
- * Draws a wireframe cone with a arcs on the cap
- *
- * @param	PDI				Draw interface.
- * @param	Transform		Generic transform to apply (ex. a local-to-world transform).
- * @param	ConeRadius		Radius of the cone.
- * @param	ConeAngle		Angle of the cone.
- * @param	ConeSides		Numbers of sides that the cone has.
- * @param   ArcFrequency    How frequently to draw an arc (1 means every vertex, 2 every 2nd etc.)
- * @param	CapSegments		How many lines to use to make the arc
- * @param	Color			Color of the cone.
- * @param	DepthPriority	Depth priority for the cone.
- */
 void DrawWireSphereCappedCone(FPrimitiveDrawInterface* PDI, const FTransform& Transform, float ConeRadius, float ConeAngle, int32 ConeSides, int32 ArcFrequency, int32 CapSegments, const FLinearColor& Color, uint8 DepthPriority)
 {
 	// The cap only works if there are an even number of verts generated so add another if needed 
@@ -1190,21 +1040,6 @@ void DrawWireSphereCappedCone(FPrimitiveDrawInterface* PDI, const FTransform& Tr
 }
 
 
-/**
- * Draws a wireframe chopped cone(cylinder with independant top and bottom radius).
- *
- * @param	PDI				Draw interface.
- * @param	Base			Center pointer of the base of the cone.
- * @param	X				X alignment axis to draw along.
- * @param	Y				Y alignment axis to draw along.
- * @param	Z				Z alignment axis to draw along.
- * @param	Color			Color of the cone.
- * @param	Radius			Radius of the cone at the bottom.
- * @param	TopRadius		Radius of the cone at the top.
- * @param	HalfHeight		Half of the height of the cone.
- * @param	NumSides		Numbers of sides that the cone has.
- * @param	DepthPriority	Depth priority for the cone.
- */
 void DrawWireChoppedCone(FPrimitiveDrawInterface* PDI,const FVector& Base,const FVector& X,const FVector& Y,const FVector& Z,const FLinearColor& Color,float Radius, float TopRadius,float HalfHeight,int32 NumSides,uint8 DepthPriority)
 {
 	const float	AngleDelta = 2.0f * PI / NumSides;
@@ -1225,19 +1060,6 @@ void DrawWireChoppedCone(FPrimitiveDrawInterface* PDI,const FVector& Base,const
 	}
 }
 
-/**
- * Draws an oriented box.
- *
- * @param	PDI				Draw interface.
- * @param	Base			Center point of the box.
- * @param	X				X alignment axis to draw along.
- * @param	Y				Y alignment axis to draw along.
- * @param	Z				Z alignment axis to draw along.
- * @param	Color			Color of the box.
- * @param	Extent			Vector with the half-sizes of the box.
- * @param	DepthPriority	Depth priority for the cone.
- */
-
 void DrawOrientedWireBox(FPrimitiveDrawInterface* PDI, const FVector& Base, const FVector& X, const FVector& Y, const FVector& Z, FVector Extent, const FLinearColor& Color, uint8 DepthPriority, float Thickness, float DepthBias, bool bScreenSpace)
 {
 	FVector	B[2],P,Q;
@@ -1275,48 +1097,27 @@ void DrawOrientedWireBox(FPrimitiveDrawInterface* PDI, const FVector& Base, cons
 }
 
 
-void DrawCoordinateSystem(FPrimitiveDrawInterface* PDI, FVector const& AxisLoc, FRotator const& AxisRot, float Scale, uint8 DepthPriority)
+void DrawCoordinateSystem(FPrimitiveDrawInterface* PDI, FVector const& AxisLoc, FRotator const& AxisRot, float Scale, uint8 DepthPriority, float Thickness)
 {
 	FRotationMatrix R(AxisRot);
 	FVector const X = R.GetScaledAxis( EAxis::X );
 	FVector const Y = R.GetScaledAxis( EAxis::Y );
 	FVector const Z = R.GetScaledAxis( EAxis::Z );
 
-	PDI->DrawLine(AxisLoc, AxisLoc + X*Scale, FLinearColor::Red, DepthPriority );
-	PDI->DrawLine(AxisLoc, AxisLoc + Y*Scale, FLinearColor::Green, DepthPriority );
-	PDI->DrawLine(AxisLoc, AxisLoc + Z*Scale, FLinearColor::Blue, DepthPriority );
+	PDI->DrawLine(AxisLoc, AxisLoc + X*Scale, FLinearColor::Red, DepthPriority, Thickness );
+	PDI->DrawLine(AxisLoc, AxisLoc + Y*Scale, FLinearColor::Green, DepthPriority, Thickness );
+	PDI->DrawLine(AxisLoc, AxisLoc + Z*Scale, FLinearColor::Blue, DepthPriority, Thickness );
 }
-/**
- * Draws a directional arrow.
- *
- * @param	PDI				Draw interface.
- * @param	ArrowToWorld	Transform matrix for the arrow.
- * @param	InColor			Color of the arrow.
- * @param	Length			Length of the arrow
- * @param	ArrowSize		Size of the arrow head.
- * @param	DepthPriority	Depth priority for the arrow.
- */
-void DrawDirectionalArrow(FPrimitiveDrawInterface* PDI,const FMatrix& ArrowToWorld,const FLinearColor& InColor,float Length,float ArrowSize,uint8 DepthPriority)
+
+void DrawDirectionalArrow(FPrimitiveDrawInterface* PDI,const FMatrix& ArrowToWorld,const FLinearColor& InColor,float Length,float ArrowSize,uint8 DepthPriority,float Thickness)
 {
-	PDI->DrawLine(ArrowToWorld.TransformPosition(FVector(Length,0,0)),ArrowToWorld.TransformPosition(FVector::ZeroVector),InColor,DepthPriority);
-	PDI->DrawLine(ArrowToWorld.TransformPosition(FVector(Length,0,0)),ArrowToWorld.TransformPosition(FVector(Length-ArrowSize,+ArrowSize,+ArrowSize)),InColor,DepthPriority);
-	PDI->DrawLine(ArrowToWorld.TransformPosition(FVector(Length,0,0)),ArrowToWorld.TransformPosition(FVector(Length-ArrowSize,+ArrowSize,-ArrowSize)),InColor,DepthPriority);
-	PDI->DrawLine(ArrowToWorld.TransformPosition(FVector(Length,0,0)),ArrowToWorld.TransformPosition(FVector(Length-ArrowSize,-ArrowSize,+ArrowSize)),InColor,DepthPriority);
-	PDI->DrawLine(ArrowToWorld.TransformPosition(FVector(Length,0,0)),ArrowToWorld.TransformPosition(FVector(Length-ArrowSize,-ArrowSize,-ArrowSize)),InColor,DepthPriority);
+	PDI->DrawLine(ArrowToWorld.TransformPosition(FVector(Length,0,0)),ArrowToWorld.TransformPosition(FVector::ZeroVector),InColor,DepthPriority,Thickness);
+	PDI->DrawLine(ArrowToWorld.TransformPosition(FVector(Length,0,0)),ArrowToWorld.TransformPosition(FVector(Length-ArrowSize,+ArrowSize,+ArrowSize)),InColor,DepthPriority,Thickness);
+	PDI->DrawLine(ArrowToWorld.TransformPosition(FVector(Length,0,0)),ArrowToWorld.TransformPosition(FVector(Length-ArrowSize,+ArrowSize,-ArrowSize)),InColor,DepthPriority,Thickness);
+	PDI->DrawLine(ArrowToWorld.TransformPosition(FVector(Length,0,0)),ArrowToWorld.TransformPosition(FVector(Length-ArrowSize,-ArrowSize,+ArrowSize)),InColor,DepthPriority,Thickness);
+	PDI->DrawLine(ArrowToWorld.TransformPosition(FVector(Length,0,0)),ArrowToWorld.TransformPosition(FVector(Length-ArrowSize,-ArrowSize,-ArrowSize)),InColor,DepthPriority,Thickness);
 }
 
-/**
- * Draws a directional arrow with connected spokes.
- *
- * @param	PDI				Draw interface.
- * @param	ArrowToWorld	Transform matrix for the arrow.
- * @param	Color			Color of the arrow.
- * @param	ArrowHeight		Height of the the arrow head.
- * @param	ArrowWidth		Width of the arrow head.
- * @param	DepthPriority	Depth priority for the arrow.
- * @param	Thickness		Thickness of the lines used to draw the arrow.
- * @param	NumSpokes		Number of spokes used to make the arrow head.
- */
 void DrawConnectedArrow(class FPrimitiveDrawInterface* PDI, const FMatrix& ArrowToWorld, const FLinearColor& Color, float ArrowHeight, float ArrowWidth, uint8 DepthPriority, float Thickness, int32 NumSpokes)
 {
 	float RotPerSpoke = (2.0f * PI) / (float)NumSpokes;
@@ -1333,15 +1134,6 @@ void DrawConnectedArrow(class FPrimitiveDrawInterface* PDI, const FMatrix& Arrow
 	}
 }
 
-/**
- * Draws a axis-aligned 3 line star.
- *
- * @param	PDI				Draw interface.
- * @param	Position		Position of the star.
- * @param	Size			Size of the star
- * @param	InColor			Color of the arrow.
- * @param	DepthPriority	Depth priority for the star.
- */
 void DrawWireStar(FPrimitiveDrawInterface* PDI,const FVector& Position, float Size, const FLinearColor& Color,uint8 DepthPriority)
 {
 	PDI->DrawLine(Position + Size * FVector(1,0,0), Position - Size * FVector(1,0,0), Color, DepthPriority);
@@ -1349,16 +1141,6 @@ void DrawWireStar(FPrimitiveDrawInterface* PDI,const FVector& Position, float Si
 	PDI->DrawLine(Position + Size * FVector(0,0,1), Position - Size * FVector(0,0,1), Color, DepthPriority);
 }
 
-/**
- * Draws a dashed line.
- *
- * @param	PDI				Draw interface.
- * @param	Start			Start position of the line.
- * @param	End				End position of the line.
- * @param	Color			Color of the arrow.
- * @param	DashSize		Size of each of the dashes that makes up the line.
- * @param	DepthPriority	Depth priority for the line.
- */
 void DrawDashedLine(FPrimitiveDrawInterface* PDI, const FVector& Start, const FVector& End, const FLinearColor& Color, float DashSize, uint8 DepthPriority, float DepthBias)
 {
 	FVector LineDir = End - Start;
@@ -1391,15 +1173,6 @@ void DrawDashedLine(FPrimitiveDrawInterface* PDI, const FVector& Start, const FV
 	}
 }
 
-/**
- * Draws a wireframe diamond.
- *
- * @param	PDI				Draw interface.
- * @param	DiamondMatrix	Transform Matrix for the diamond.
- * @param	Size			Size of the diamond.
- * @param	InColor			Color of the diamond.
- * @param	DepthPriority	Depth priority for the diamond.
- */
 void DrawWireDiamond(FPrimitiveDrawInterface* PDI,const FMatrix& DiamondMatrix, float Size, const FLinearColor& InColor,uint8 DepthPriority)
 {
 	const FVector TopPoint = DiamondMatrix.TransformPosition( FVector(0,0,1) * Size );
diff --git a/Engine/Source/Runtime/Engine/Public/SceneManagement.h b/Engine/Source/Runtime/Engine/Public/SceneManagement.h
--- a/Engine/Source/Runtime/Engine/Public/SceneManagement.h
+++ b/Engine/Source/Runtime/Engine/Public/SceneManagement.h
@@ -1541,7 +1541,7 @@ private:
 //
 
 // Solid shape drawing utility functions. Not really designed for speed - more for debugging.
-// These utilities functions are implemented in UnScene.cpp using GetTRI.
+// These utilities functions are implemented in PrimitiveDrawingUtils.cpp.
 
 // 10x10 tessellated plane at x=-1..1 y=-1...1 z=0
 extern ENGINE_API void DrawPlane10x10(class FPrimitiveDrawInterface* PDI,const FMatrix& ObjectToWorld,float Radii,FVector2D UVMin, FVector2D UVMax,const FMaterialRenderProxy* MaterialRenderProxy,uint8 DepthPriority);
@@ -1572,31 +1572,283 @@ extern ENGINE_API void GetCapsuleMesh(const FVector& Origin, const FVector& XAxi
 									const FMaterialRenderProxy* MaterialRenderProxy, uint8 DepthPriority, bool bDisableBackfaceCulling, int32 ViewIndex, FMeshElementCollector& Collector);
 
 
+/**
+ * Draws a circle using triangles.
+ *
+ * @param	PDI						Draw interface.
+ * @param	Base					Center of the circle.
+ * @param	XAxis					X alignment axis to draw along.
+ * @param	YAxis					Y alignment axis to draw along.
+ * @param	Color					Color of the circle.
+ * @param	Radius					Radius of the circle.
+ * @param	NumSides				Numbers of sides that the circle has.
+ * @param	MaterialRenderProxy		Material to use for render 
+ * @param	DepthPriority			Depth priority for the circle.
+ */
 extern ENGINE_API void DrawDisc(class FPrimitiveDrawInterface* PDI,const FVector& Base,const FVector& XAxis,const FVector& YAxis,FColor Color,float Radius,int32 NumSides, const FMaterialRenderProxy* MaterialRenderProxy, uint8 DepthPriority);
-extern ENGINE_API void DrawFlatArrow(class FPrimitiveDrawInterface* PDI,const FVector& Base,const FVector& XAxis,const FVector& YAxis,FColor Color,float Length,int32 Width, const FMaterialRenderProxy* MaterialRenderProxy, uint8 DepthPriority);
+
+
+/**
+ * Draws a flat arrow with an outline.
+ *
+ * @param	PDI						Draw interface.
+ * @param	Base					Base of the arrow.
+ * @param	XAxis					X alignment axis to draw along.
+ * @param	YAxis					Y alignment axis to draw along.
+ * @param	Color					Color of the circle.
+ * @param	Length					Length of the arrow, from base to tip.
+ * @param	Width					Width of the base of the arrow, head of the arrow will be 2x.
+ * @param	MaterialRenderProxy		Material to use for render 
+ * @param	DepthPriority			Depth priority for the circle.
+ * @param	Thickness				Thickness of the lines comprising the arrow
+ */
+
+/*
+x-axis is from point 0 to point 2
+y-axis is from point 0 to point 1
+		6
+		/\
+	   /  \
+	  /    \
+	 4_2  3_5
+	   |  |
+	   0__1
+*/
+extern ENGINE_API void DrawFlatArrow(class FPrimitiveDrawInterface* PDI,const FVector& Base,const FVector& XAxis,const FVector& YAxis,FColor Color,float Length,int32 Width, const FMaterialRenderProxy* MaterialRenderProxy, uint8 DepthPriority, float Thickness = 0.0f);
 
 // Line drawing utility functions.
-extern ENGINE_API void DrawWireBox(class FPrimitiveDrawInterface* PDI, const FBox& Box, const FLinearColor& Color, uint8 DepthPriority, float Thickness = 0, float DepthBias = 0.0f, bool bScreenSpace = false);
-extern ENGINE_API void DrawWireBox(class FPrimitiveDrawInterface* PDI, const FMatrix& Matrix, const FBox& Box, const FLinearColor& Color, uint8 DepthPriority, float Thickness = 0, float DepthBias = 0.0f, bool bScreenSpace = false);
-extern ENGINE_API void DrawCircle(class FPrimitiveDrawInterface* PDI, const FVector& Base, const FVector& X, const FVector& Y, const FLinearColor& Color, float Radius, int32 NumSides, uint8 DepthPriority, float Thickness = 0, float DepthBias = 0.0f, bool bScreenSpace = false);
+
+/**
+ * Draws a wireframe box.
+ *
+ * @param	PDI				Draw interface.
+ * @param	Box				The FBox to use for drawing.
+ * @param	Color			Color of the box.
+ * @param	DepthPriority	Depth priority for the circle.
+ * @param	Thickness		Thickness of the lines comprising the box
+ */
+extern ENGINE_API void DrawWireBox(class FPrimitiveDrawInterface* PDI, const FBox& Box, const FLinearColor& Color, uint8 DepthPriority, float Thickness = 0.0f, float DepthBias = 0.0f, bool bScreenSpace = false);
+extern ENGINE_API void DrawWireBox(class FPrimitiveDrawInterface* PDI, const FMatrix& Matrix, const FBox& Box, const FLinearColor& Color, uint8 DepthPriority, float Thickness = 0.0f, float DepthBias = 0.0f, bool bScreenSpace = false);
+
+/**
+ * Draws a circle using lines.
+ *
+ * @param	PDI				Draw interface.
+ * @param	Base			Center of the circle.
+ * @param	X				X alignment axis to draw along.
+ * @param	Y				Y alignment axis to draw along.
+ * @param	Z				Z alignment axis to draw along.
+ * @param	Color			Color of the circle.
+ * @param	Radius			Radius of the circle.
+ * @param	NumSides		Numbers of sides that the circle has.
+ * @param	DepthPriority	Depth priority for the circle.
+ * @param	Thickness		Thickness of the lines comprising the circle
+ */
+extern ENGINE_API void DrawCircle(class FPrimitiveDrawInterface* PDI, const FVector& Base, const FVector& X, const FVector& Y, const FLinearColor& Color, float Radius, int32 NumSides, uint8 DepthPriority, float Thickness = 0.0f, float DepthBias = 0.0f, bool bScreenSpace = false);
+
+
+/**
+ * Draws an arc using lines.
+ *
+ * @param	PDI				Draw interface.
+ * @param	Base			Center of the circle.
+ * @param	X				Normalized axis from one point to the center
+ * @param	Y				Normalized axis from other point to the center
+ * @param   MinAngle        The minimum angle
+ * @param   MinAngle        The maximum angle
+ * @param   Radius          Radius of the arc
+ * @param	Sections		Numbers of sides that the circle has.
+ * @param	Color			Color of the circle.
+ * @param	DepthPriority	Depth priority for the circle.
+ */
 extern ENGINE_API void DrawArc(FPrimitiveDrawInterface* PDI, const FVector Base, const FVector X, const FVector Y, const float MinAngle, const float MaxAngle, const float Radius, const int32 Sections, const FLinearColor& Color, uint8 DepthPriority);
-extern ENGINE_API void DrawWireSphere(class FPrimitiveDrawInterface* PDI, const FVector& Base, const FLinearColor& Color, float Radius, int32 NumSides, uint8 DepthPriority, float Thickness = 0, float DepthBias = 0.0f, bool bScreenSpace = false);
-extern ENGINE_API void DrawWireSphere(class FPrimitiveDrawInterface* PDI, const FTransform& Transform, const FLinearColor& Color, float Radius, int32 NumSides, uint8 DepthPriority, float Thickness = 0, float DepthBias = 0.0f, bool bScreenSpace = false);
-extern ENGINE_API void DrawWireSphereAutoSides(class FPrimitiveDrawInterface* PDI, const FVector& Base, const FLinearColor& Color, float Radius, uint8 DepthPriority, float Thickness = 0, float DepthBias = 0.0f, bool bScreenSpace = false);
-extern ENGINE_API void DrawWireSphereAutoSides(class FPrimitiveDrawInterface* PDI, const FTransform& Transform, const FLinearColor& Color, float Radius, uint8 DepthPriority, float Thickness = 0, float DepthBias = 0.0f, bool bScreenSpace = false);
+
+/**
+ * Draws a sphere using circles.
+ *
+ * @param	PDI				Draw interface.
+ * @param	Base			Center of the sphere.
+ * @param	Color			Color of the sphere.
+ * @param	Radius			Radius of the sphere.
+ * @param	NumSides		Numbers of sides that the circle has.
+ * @param	DepthPriority	Depth priority for the circle.
+ * @param	Thickness		Thickness of the lines comprising the sphere
+ */
+extern ENGINE_API void DrawWireSphere(class FPrimitiveDrawInterface* PDI, const FVector& Base, const FLinearColor& Color, float Radius, int32 NumSides, uint8 DepthPriority, float Thickness = 0.0f, float DepthBias = 0.0f, bool bScreenSpace = false);
+extern ENGINE_API void DrawWireSphere(class FPrimitiveDrawInterface* PDI, const FTransform& Transform, const FLinearColor& Color, float Radius, int32 NumSides, uint8 DepthPriority, float Thickness = 0.0f, float DepthBias = 0.0f, bool bScreenSpace = false);
+
+/**
+ * Draws a sphere using circles, automatically calculating a reasonable number of sides
+ *
+ * @param	PDI				Draw interface.
+ * @param	Base			Center of the sphere.
+ * @param	Color			Color of the sphere.
+ * @param	Radius			Radius of the sphere.
+ * @param	DepthPriority	Depth priority for the circle.
+ * @param	Thickness		Thickness of the lines comprising the sphere
+ */
+extern ENGINE_API void DrawWireSphereAutoSides(class FPrimitiveDrawInterface* PDI, const FVector& Base, const FLinearColor& Color, float Radius, uint8 DepthPriority, float Thickness = 0.0f, float DepthBias = 0.0f, bool bScreenSpace = false);
+extern ENGINE_API void DrawWireSphereAutoSides(class FPrimitiveDrawInterface* PDI, const FTransform& Transform, const FLinearColor& Color, float Radius, uint8 DepthPriority, float Thickness = 0.0f, float DepthBias = 0.0f, bool bScreenSpace = false);
+
+/**
+ * Draws a wireframe cylinder.
+ *
+ * @param	PDI				Draw interface.
+ * @param	Base			Center pointer of the base of the cylinder.
+ * @param	X				X alignment axis to draw along.
+ * @param	Y				Y alignment axis to draw along.
+ * @param	Z				Z alignment axis to draw along.
+ * @param	Color			Color of the cylinder.
+ * @param	Radius			Radius of the cylinder.
+ * @param	HalfHeight		Half of the height of the cylinder.
+ * @param	NumSides		Numbers of sides that the cylinder has.
+ * @param	DepthPriority	Depth priority for the cylinder.
+ * @param	Thickness		Thickness of the lines comprising the cylinder
+ */
 extern ENGINE_API void DrawWireCylinder(class FPrimitiveDrawInterface* PDI, const FVector& Base, const FVector& X, const FVector& Y, const FVector& Z, const FLinearColor& Color, float Radius, float HalfHeight, int32 NumSides, uint8 DepthPriority, float Thickness = 0.0f, float DepthBias = 0.0f, bool bScreenSpace = false);
+
+/**
+ * Draws a wireframe capsule.
+ *
+ * @param	PDI				Draw interface.
+ * @param	Base			Center pointer of the base of the cylinder.
+ * @param	X				X alignment axis to draw along.
+ * @param	Y				Y alignment axis to draw along.
+ * @param	Z				Z alignment axis to draw along.
+ * @param	Color			Color of the cylinder.
+ * @param	Radius			Radius of the cylinder.
+ * @param	HalfHeight		Half of the height of the cylinder.
+ * @param	NumSides		Numbers of sides that the cylinder has.
+ * @param	DepthPriority	Depth priority for the cylinder.
+ * @param	Thickness		Thickness of the lines comprising the cylinder
+ */
 extern ENGINE_API void DrawWireCapsule(class FPrimitiveDrawInterface* PDI, const FVector& Base, const FVector& X, const FVector& Y, const FVector& Z, const FLinearColor& Color, float Radius, float HalfHeight, int32 NumSides, uint8 DepthPriority, float Thickness = 0.0f, float DepthBias = 0.0f, bool bScreenSpace = false);
+
+/**
+ * Draws a wireframe chopped cone (cylinder with independent top and bottom radius).
+ *
+ * @param	PDI				Draw interface.
+ * @param	Base			Center pointer of the base of the cone.
+ * @param	X				X alignment axis to draw along.
+ * @param	Y				Y alignment axis to draw along.
+ * @param	Z				Z alignment axis to draw along.
+ * @param	Color			Color of the cone.
+ * @param	Radius			Radius of the cone at the bottom.
+ * @param	TopRadius		Radius of the cone at the top.
+ * @param	HalfHeight		Half of the height of the cone.
+ * @param	NumSides		Numbers of sides that the cone has.
+ * @param	DepthPriority	Depth priority for the cone.
+ */
 extern ENGINE_API void DrawWireChoppedCone(class FPrimitiveDrawInterface* PDI,const FVector& Base,const FVector& X,const FVector& Y,const FVector& Z,const FLinearColor& Color,float Radius,float TopRadius,float HalfHeight,int32 NumSides,uint8 DepthPriority);
+
+/**
+ * Draws a wireframe cone
+ *
+ * @param	PDI				Draw interface.
+ * @param	Transform		Generic transform to apply (ex. a local-to-world transform).
+ * @param	ConeRadius		Radius of the cone.
+ * @param	ConeAngle		Angle of the cone.
+ * @param	ConeSides		Numbers of sides that the cone has.
+ * @param	Color			Color of the cone.
+ * @param	DepthPriority	Depth priority for the cone.
+ * @param	Verts			Out param, the positions of the verts at the cone's base.
+ * @param	Thickness		Thickness of the lines comprising the cone
+ */
 extern ENGINE_API void DrawWireCone(class FPrimitiveDrawInterface* PDI, TArray<FVector>& Verts, const FMatrix& Transform, float ConeRadius, float ConeAngle, int32 ConeSides, const FLinearColor& Color, uint8 DepthPriority, float Thickness = 0.0f, float DepthBias = 0.0f, bool bScreenSpace = false);
 extern ENGINE_API void DrawWireCone(class FPrimitiveDrawInterface* PDI, TArray<FVector>& Verts, const FTransform& Transform, float ConeRadius, float ConeAngle, int32 ConeSides, const FLinearColor& Color, uint8 DepthPriority, float Thickness = 0.0f, float DepthBias = 0.0f, bool bScreenSpace = false);
+
+/**
+ * Draws a wireframe cone with a arcs on the cap
+ *
+ * @param	PDI				Draw interface.
+ * @param	Transform		Generic transform to apply (ex. a local-to-world transform).
+ * @param	ConeRadius		Radius of the cone.
+ * @param	ConeAngle		Angle of the cone.
+ * @param	ConeSides		Numbers of sides that the cone has.
+ * @param   ArcFrequency    How frequently to draw an arc (1 means every vertex, 2 every 2nd etc.)
+ * @param	CapSegments		How many lines to use to make the arc
+ * @param	Color			Color of the cone.
+ * @param	DepthPriority	Depth priority for the cone.
+ */
 extern ENGINE_API void DrawWireSphereCappedCone(FPrimitiveDrawInterface* PDI, const FTransform& Transform, float ConeRadius, float ConeAngle, int32 ConeSides, int32 ArcFrequency, int32 CapSegments, const FLinearColor& Color, uint8 DepthPriority);
+
+/**
+ * Draws an oriented box.
+ *
+ * @param	PDI				Draw interface.
+ * @param	Base			Center point of the box.
+ * @param	X				X alignment axis to draw along.
+ * @param	Y				Y alignment axis to draw along.
+ * @param	Z				Z alignment axis to draw along.
+ * @param	Color			Color of the box.
+ * @param	Extent			Vector with the half-sizes of the box.
+ * @param	DepthPriority	Depth priority for the cone.
+ * @param	Thickness		Thickness of the lines comprising the box
+ */
 extern ENGINE_API void DrawOrientedWireBox(class FPrimitiveDrawInterface* PDI, const FVector& Base, const FVector& X, const FVector& Y, const FVector& Z, FVector Extent, const FLinearColor& Color, uint8 DepthPriority, float Thickness = 0.0f, float DepthBias = 0.0f, bool bScreenSpace = false);
-extern ENGINE_API void DrawDirectionalArrow(class FPrimitiveDrawInterface* PDI,const FMatrix& ArrowToWorld,const FLinearColor& InColor,float Length,float ArrowSize,uint8 DepthPriority);
+
+/**
+ * Draws a directional arrow (starting at ArrowToWorld.Origin and continuing for Length units in the X direction of ArrowToWorld).
+ *
+ * @param	PDI				Draw interface.
+ * @param	ArrowToWorld	Transform matrix for the arrow.
+ * @param	InColor			Color of the arrow.
+ * @param	Length			Length of the arrow
+ * @param	ArrowSize		Size of the arrow head.
+ * @param	DepthPriority	Depth priority for the arrow.
+ * @param	Thickness		Thickness of the lines comprising the arrow
+ */
+extern ENGINE_API void DrawDirectionalArrow(class FPrimitiveDrawInterface* PDI, const FMatrix& ArrowToWorld, const FLinearColor& InColor, float Length, float ArrowSize, uint8 DepthPriority, float Thickness = 0.0f);
+
+/**
+ * Draws a directional arrow with connected spokes.
+ *
+ * @param	PDI				Draw interface.
+ * @param	ArrowToWorld	Transform matrix for the arrow.
+ * @param	Color			Color of the arrow.
+ * @param	ArrowHeight		Height of the the arrow head.
+ * @param	ArrowWidth		Width of the arrow head.
+ * @param	DepthPriority	Depth priority for the arrow.
+ * @param	Thickness		Thickness of the lines used to draw the arrow.
+ * @param	NumSpokes		Number of spokes used to make the arrow head.
+ */
 extern ENGINE_API void DrawConnectedArrow(class FPrimitiveDrawInterface* PDI, const FMatrix& ArrowToWorld, const FLinearColor& Color, float ArrowHeight, float ArrowWidth, uint8 DepthPriority, float Thickness = 0.5f, int32 NumSpokes = 6);
-extern ENGINE_API void DrawWireStar(class FPrimitiveDrawInterface* PDI,const FVector& Position, float Size, const FLinearColor& Color,uint8 DepthPriority);
+
+/**
+ * Draws a axis-aligned 3 line star.
+ *
+ * @param	PDI				Draw interface.
+ * @param	Position		Position of the star.
+ * @param	Size			Size of the star
+ * @param	InColor			Color of the arrow.
+ * @param	DepthPriority	Depth priority for the star.
+ */
+extern ENGINE_API void DrawWireStar(class FPrimitiveDrawInterface* PDI, const FVector& Position, float Size, const FLinearColor& Color, uint8 DepthPriority);
+
+/**
+ * Draws a dashed line.
+ *
+ * @param	PDI				Draw interface.
+ * @param	Start			Start position of the line.
+ * @param	End				End position of the line.
+ * @param	Color			Color of the arrow.
+ * @param	DashSize		Size of each of the dashes that makes up the line.
+ * @param	DepthPriority	Depth priority for the line.
+ */
 extern ENGINE_API void DrawDashedLine(class FPrimitiveDrawInterface* PDI, const FVector& Start, const FVector& End, const FLinearColor& Color, float DashSize, uint8 DepthPriority, float DepthBias = 0.0f);
-extern ENGINE_API void DrawWireDiamond(class FPrimitiveDrawInterface* PDI,const FMatrix& DiamondMatrix, float Size, const FLinearColor& InColor,uint8 DepthPriority);
-extern ENGINE_API void DrawCoordinateSystem(FPrimitiveDrawInterface* PDI, FVector const& AxisLoc, FRotator const& AxisRot, float Scale, uint8 DepthPriority);
+
+/**
+ * Draws a wireframe diamond.
+ *
+ * @param	PDI				Draw interface.
+ * @param	DiamondMatrix	Transform Matrix for the diamond.
+ * @param	Size			Size of the diamond.
+ * @param	InColor			Color of the diamond.
+ * @param	DepthPriority	Depth priority for the diamond.
+ */
+extern ENGINE_API void DrawWireDiamond(class FPrimitiveDrawInterface* PDI, const FMatrix& DiamondMatrix, float Size, const FLinearColor& InColor, uint8 DepthPriority);
+
+extern ENGINE_API void DrawCoordinateSystem(FPrimitiveDrawInterface* PDI, FVector const& AxisLoc, FRotator const& AxisRot, float Scale, uint8 DepthPriority, float Thickness = 0.0f);
 
 /**
  * Draws a wireframe of the bounds of a frustum as defined by a transform from clip-space into world-space.
