#ue4 - Zero Velocity.Z when switching to NavWalking movement.

diff --git a/Engine/Source/Runtime/Engine/Private/CharacterMovementComponent.cpp b/Engine/Source/Runtime/Engine/Private/CharacterMovementComponent.cpp
--- a/Engine/Source/Runtime/Engine/Private/CharacterMovementComponent.cpp
+++ b/Engine/Source/Runtime/Engine/Private/CharacterMovementComponent.cpp
@@ -673,6 +673,8 @@ void UCharacterMovementComponent::OnMovementModeChanged(EMovementMode PreviousMo
 	{
 		SetNavWalkingPhysics(true);
 		GroundMovementMode = MovementMode;
+		// Walking uses only XY velocity
+		Velocity.Z = 0.f;
 	}
 	else if (PreviousMovementMode == MOVE_NavWalking)
 	{
