Refactor overlap/sweep/raycast so that they explicitly use channel object type or profile.
Old functionality still works, but deprecated. Note that for the object type the object params have been moved to the same argument position as ByChannel and byProfile

diff --git a/Engine/Source/Developer/LogVisualizer/Private/VisualLoggerHUD.cpp b/Engine/Source/Developer/LogVisualizer/Private/VisualLoggerHUD.cpp
--- a/Engine/Source/Developer/LogVisualizer/Private/VisualLoggerHUD.cpp
+++ b/Engine/Source/Developer/LogVisualizer/Private/VisualLoggerHUD.cpp
@@ -44,7 +44,7 @@ void AVisualLoggerHUD::PostRender()
 
 			FCollisionQueryParams TraceParams(NAME_None, true, this);
 			FHitResult Hit;
-			bool bHit = GetWorld()->LineTraceSingle(Hit, CamLoc, CamRot.Vector() * 100000.f + CamLoc, ECC_Pawn, TraceParams);
+			bool bHit = GetWorld()->LineTraceSingleByChannel(Hit, CamLoc, CamRot.Vector() * 100000.f + CamLoc, ECC_Pawn, TraceParams);
 			if( bHit )
 			{
 				TextItem.Text = FText::FromString(FString::Printf(TEXT("Under cursor: '%s'"), *Hit.GetActor()->GetName()));
diff --git a/Engine/Source/Editor/LandscapeEditor/Private/LandscapeEdMode.cpp b/Engine/Source/Editor/LandscapeEditor/Private/LandscapeEdMode.cpp
--- a/Engine/Source/Editor/LandscapeEditor/Private/LandscapeEdMode.cpp
+++ b/Engine/Source/Editor/LandscapeEditor/Private/LandscapeEdMode.cpp
@@ -710,7 +710,7 @@ bool FEdModeLandscape::LandscapeMouseTrace(FEditorViewportClient* ViewportClient
 	TArray<FHitResult> Results;
 	// Each landscape component has 2 collision shapes, 1 of them is specific to landscape editor
 	// Trace only ECC_Visibility channel, so we do hit only Editor specific shape
-	World->LineTraceMulti(Results, Start, End, FCollisionQueryParams(TraceTag, true), FCollisionObjectQueryParams(ECollisionChannel::ECC_Visibility));
+	World->LineTraceMultiByObjectType(Results, Start, End, FCollisionObjectQueryParams(ECollisionChannel::ECC_Visibility), FCollisionQueryParams(TraceTag, true));
 
 	for (int32 i = 0; i < Results.Num(); i++)
 	{
diff --git a/Engine/Source/Editor/LandscapeEditor/Private/LandscapeEdModeSplineTools.cpp b/Engine/Source/Editor/LandscapeEditor/Private/LandscapeEdModeSplineTools.cpp
--- a/Engine/Source/Editor/LandscapeEditor/Private/LandscapeEdModeSplineTools.cpp
+++ b/Engine/Source/Editor/LandscapeEditor/Private/LandscapeEdModeSplineTools.cpp
@@ -587,7 +587,7 @@ public:
 		FHitResult Hit;
 		UWorld* World = SplinesComponent->GetWorld();
 		check(World);
-		if (World->LineTraceSingle(Hit, Start, End, FCollisionQueryParams(true), FCollisionObjectQueryParams(ECC_WorldStatic)))
+		if (World->LineTraceSingleByObjectType(Hit, Start, End, FCollisionObjectQueryParams(ECC_WorldStatic), FCollisionQueryParams(true)))
 		{
 			ControlPoint->Location = LocalToWorld.InverseTransformPosition(Hit.Location);
 			ControlPoint->UpdateSplinePoints();
diff --git a/Engine/Source/Editor/LevelEditor/Private/LevelEditorActions.cpp b/Engine/Source/Editor/LevelEditor/Private/LevelEditorActions.cpp
--- a/Engine/Source/Editor/LevelEditor/Private/LevelEditorActions.cpp
+++ b/Engine/Source/Editor/LevelEditor/Private/LevelEditorActions.cpp
@@ -1179,7 +1179,7 @@ void FLevelEditorActionCallbacks::GoHere_Clicked( const FVector* Point )
 				static FName FocusOnPoint = FName(TEXT("FocusOnPoint"));
 				FCollisionQueryParams LineParams(FocusOnPoint, true);
 
-				if(GCurrentLevelEditingViewportClient->GetWorld()->LineTraceSingle(HitResult, WorldOrigin, WorldOrigin + WorldDirection * HALF_WORLD_MAX, LineParams, FCollisionObjectQueryParams(ECC_WorldStatic)))
+				if(GCurrentLevelEditingViewportClient->GetWorld()->LineTraceSingleByObjectType(HitResult, WorldOrigin, WorldOrigin + WorldDirection * HALF_WORLD_MAX, FCollisionObjectQueryParams(ECC_WorldStatic), LineParams))
 				{
 					ZoomToPoint = HitResult.ImpactPoint;
 				}
diff --git a/Engine/Source/Editor/UnrealEd/Private/DebugToolExec.cpp b/Engine/Source/Editor/UnrealEd/Private/DebugToolExec.cpp
--- a/Engine/Source/Editor/UnrealEd/Private/DebugToolExec.cpp
+++ b/Engine/Source/Editor/UnrealEd/Private/DebugToolExec.cpp
@@ -202,7 +202,7 @@ bool FDebugToolExec::Exec( UWorld* InWorld, const TCHAR* Cmd, FOutputDevice& Ar
 				FRotator PlayerRotation;
 				PlayerController->GetPlayerViewPoint(PlayerLocation, PlayerRotation);
 				FHitResult Hit(1.0f);
-				PlayerController->GetWorld()->LineTraceSingle(Hit, PlayerLocation, PlayerLocation + PlayerRotation.Vector() * 10000.f, ECC_Pawn, FCollisionQueryParams(NAME_None, true, PlayerController->GetPawn()));
+				PlayerController->GetWorld()->LineTraceSingleByChannel(Hit, PlayerLocation, PlayerLocation + PlayerRotation.Vector() * 10000.f, ECC_Pawn, FCollisionQueryParams(NAME_None, true, PlayerController->GetPawn()));
 				Found = Hit.GetActor();
 			}
 		}
diff --git a/Engine/Source/Editor/UnrealEd/Private/Editor.cpp b/Engine/Source/Editor/UnrealEd/Private/Editor.cpp
--- a/Engine/Source/Editor/UnrealEd/Private/Editor.cpp
+++ b/Engine/Source/Editor/UnrealEd/Private/Editor.cpp
@@ -6371,7 +6371,7 @@ void UEditorEngine::UpdatePreviewMesh()
 		static FName UpdatePreviewMeshTrace = FName(TEXT("UpdatePreviewMeshTrace"));
 		FCollisionQueryParams LineParams(UpdatePreviewMeshTrace, true);
 		LineParams.bTraceComplex = false;
-		if ( GWorld->LineTraceSingle(Hit, LineCheckStart, LineCheckEnd, LineParams, FCollisionObjectQueryParams(ECC_WorldStatic)) ) 
+		if ( GWorld->LineTraceSingleByObjectType(Hit, LineCheckStart, LineCheckEnd, FCollisionObjectQueryParams(ECC_WorldStatic), LineParams) ) 
 		{
 			// Dirty the transform so UpdateComponent will actually update the transforms. 
 			PreviewMeshComp->SetRelativeLocation(Hit.Location);
diff --git a/Engine/Source/Editor/UnrealEd/Private/Editor/ActorPositioning.cpp b/Engine/Source/Editor/UnrealEd/Private/Editor/ActorPositioning.cpp
--- a/Engine/Source/Editor/UnrealEd/Private/Editor/ActorPositioning.cpp
+++ b/Engine/Source/Editor/UnrealEd/Private/Editor/ActorPositioning.cpp
@@ -99,7 +99,7 @@ FActorPositionTraceResult FActorPositioning::TraceWorldForPosition(const UWorld&
 	}
 
 	FActorPositionTraceResult Results;
-	if ( InWorld.LineTraceMulti(Hits, RayStart, RayEnd, Param, FCollisionObjectQueryParams(FCollisionObjectQueryParams::InitType::AllObjects)) )
+	if ( InWorld.LineTraceMultiByObjectType(Hits, RayStart, RayEnd, FCollisionObjectQueryParams(FCollisionObjectQueryParams::InitType::AllObjects), Param) )
 	{
 		{
 			// Filter out anything that should be ignored
diff --git a/Engine/Source/Editor/UnrealEd/Private/EditorServer.cpp b/Engine/Source/Editor/UnrealEd/Private/EditorServer.cpp
--- a/Engine/Source/Editor/UnrealEd/Private/EditorServer.cpp
+++ b/Engine/Source/Editor/UnrealEd/Private/EditorServer.cpp
@@ -4595,7 +4595,7 @@ bool UEditorEngine::SnapObjectTo( FActorOrComponent Object, const bool InAlign,
 		Params.AddIgnoredComponent( Cast<UPrimitiveComponent>(Object.Component) );
 	}
 
-	if (Object.GetWorld()->SweepSingle(Hit, StartLocation, StartLocation + Direction*WORLD_MAX, FQuat::Identity, ECC_WorldStatic, FCollisionShape::MakeBox(Extent), Params))
+	if (Object.GetWorld()->SweepSingleByChannel(Hit, StartLocation, StartLocation + Direction*WORLD_MAX, FQuat::Identity, ECC_WorldStatic, FCollisionShape::MakeBox(Extent), Params))
 	{
 		FVector NewLocation = Hit.Location - LocationOffset;
 		NewLocation.Z += KINDA_SMALL_NUMBER;	// Move the new desired location up by an error tolerance
diff --git a/Engine/Source/Editor/UnrealEd/Private/LevelEditorViewport.cpp b/Engine/Source/Editor/UnrealEd/Private/LevelEditorViewport.cpp
--- a/Engine/Source/Editor/UnrealEd/Private/LevelEditorViewport.cpp
+++ b/Engine/Source/Editor/UnrealEd/Private/LevelEditorViewport.cpp
@@ -1870,7 +1870,7 @@ void FLevelEditorViewportClient::ProcessClick(FSceneView& View, HHitProxy* HitPr
 		{
 			FHitResult CheckResult(ForceInit);
 			FCollisionQueryParams BoxParams(ProcessClickTrace, false, ((HGeomPolyProxy*)HitProxy)->GeomObject->ActualBrush);
-			bool bHit = GWorld->SweepSingle(CheckResult, Click.GetOrigin(), Click.GetOrigin() + Click.GetDirection() * HALF_WORLD_MAX, FQuat::Identity, FCollisionShape::MakeBox(FVector(1.f)), BoxParams, FCollisionObjectQueryParams(ECC_WorldStatic));
+			bool bHit = GWorld->SweepSingleByObjectType(CheckResult, Click.GetOrigin(), Click.GetOrigin() + Click.GetDirection() * HALF_WORLD_MAX, FQuat::Identity,FCollisionObjectQueryParams(ECC_WorldStatic), FCollisionShape::MakeBox(FVector(1.f)), BoxParams);
 
 			if( bHit )
 			{
diff --git a/Engine/Source/Runtime/AIModule/Private/AIController.cpp b/Engine/Source/Runtime/AIModule/Private/AIController.cpp
--- a/Engine/Source/Runtime/AIModule/Private/AIController.cpp
+++ b/Engine/Source/Runtime/AIModule/Private/AIController.cpp
@@ -301,7 +301,7 @@ bool AAIController::LineOfSightTo(const AActor* Other, FVector ViewPoint, bool b
 	FCollisionQueryParams CollisionParams(NAME_LineOfSight, true, this->GetPawn());
 	CollisionParams.AddIgnoredActor(Other);
 
-	bool bHit = GetWorld()->LineTraceTest(ViewPoint, TargetLocation, ECC_Visibility, CollisionParams);
+	bool bHit = GetWorld()->LineTraceTestByChannel(ViewPoint, TargetLocation, ECC_Visibility, CollisionParams);
 	if (!bHit)
 	{
 		return true;
@@ -333,7 +333,7 @@ bool AAIController::LineOfSightTo(const AActor* Other, FVector ViewPoint, bool b
 	if (!bAlternateChecks || !bLOSflag)
 	{
 		//try viewpoint to head
-		bHit = GetWorld()->LineTraceTest(ViewPoint, OtherActorLocation + FVector(0.f, 0.f, OtherHeight), ECC_Visibility, CollisionParams);
+		bHit = GetWorld()->LineTraceTestByChannel(ViewPoint, OtherActorLocation + FVector(0.f, 0.f, OtherHeight), ECC_Visibility, CollisionParams);
 		if (!bHit)
 		{
 			return true;
@@ -376,7 +376,7 @@ bool AAIController::LineOfSightTo(const AActor* Other, FVector ViewPoint, bool b
 		{
 			if ((PointIndex != IndexMin) && (PointIndex != IndexMax))
 			{
-				bHit = GetWorld()->LineTraceTest(ViewPoint, Points[PointIndex], ECC_Visibility, CollisionParams);
+				bHit = GetWorld()->LineTraceTestByChannel(ViewPoint, Points[PointIndex], ECC_Visibility, CollisionParams);
 				if (!bHit)
 				{
 					return true;
diff --git a/Engine/Source/Runtime/AIModule/Private/EnvironmentQuery/Generators/EnvQueryGenerator_OnCircle.cpp b/Engine/Source/Runtime/AIModule/Private/EnvironmentQuery/Generators/EnvQueryGenerator_OnCircle.cpp
--- a/Engine/Source/Runtime/AIModule/Private/EnvironmentQuery/Generators/EnvQueryGenerator_OnCircle.cpp
+++ b/Engine/Source/Runtime/AIModule/Private/EnvironmentQuery/Generators/EnvQueryGenerator_OnCircle.cpp
@@ -24,28 +24,28 @@ struct FBatchTracingHelper
 	FORCEINLINE_DEBUGGABLE FVector RunLineTrace(const FVector& StartPos, const FVector& EndPos)
 	{
 		FHitResult OutHit;
-		const bool bHit = World->LineTraceSingle(OutHit, StartPos, EndPos, Channel, Params);
+		const bool bHit = World->LineTraceSingleByChannel(OutHit, StartPos, EndPos, Channel, Params);
 		return bHit ? OutHit.ImpactPoint : EndPos;
 	}
 
 	FORCEINLINE_DEBUGGABLE FVector RunSphereTrace(const FVector& StartPos, const FVector& EndPos)
 	{
 		FHitResult OutHit;
-		const bool bHit = World->SweepSingle(OutHit, StartPos, EndPos, FQuat::Identity, Channel, FCollisionShape::MakeSphere(Extent.X), Params);
+		const bool bHit = World->SweepSingleByChannel(OutHit, StartPos, EndPos, FQuat::Identity, Channel, FCollisionShape::MakeSphere(Extent.X), Params);
 		return bHit ? OutHit.ImpactPoint : EndPos;
 	}
 
 	FORCEINLINE_DEBUGGABLE FVector RunCapsuleTrace(const FVector& StartPos, const FVector& EndPos)
 	{
 		FHitResult OutHit;
-		const bool bHit = World->SweepSingle(OutHit, StartPos, EndPos, FQuat::Identity, Channel, FCollisionShape::MakeCapsule(Extent.X, Extent.Z), Params);
+		const bool bHit = World->SweepSingleByChannel(OutHit, StartPos, EndPos, FQuat::Identity, Channel, FCollisionShape::MakeCapsule(Extent.X, Extent.Z), Params);
 		return bHit ? OutHit.ImpactPoint : EndPos;
 	}
 
 	FORCEINLINE_DEBUGGABLE FVector RunBoxTrace(const FVector& StartPos, const FVector& EndPos)
 	{
 		FHitResult OutHit;
-		const bool bHit = World->SweepSingle(OutHit, StartPos, EndPos, FQuat((EndPos - StartPos).Rotation()), Channel, FCollisionShape::MakeBox(Extent), Params);
+		const bool bHit = World->SweepSingleByChannel(OutHit, StartPos, EndPos, FQuat((EndPos - StartPos).Rotation()), Channel, FCollisionShape::MakeBox(Extent), Params);
 		return bHit ? OutHit.ImpactPoint : EndPos;
 	}
 
diff --git a/Engine/Source/Runtime/AIModule/Private/EnvironmentQuery/Tests/EnvQueryTest_Trace.cpp b/Engine/Source/Runtime/AIModule/Private/EnvironmentQuery/Tests/EnvQueryTest_Trace.cpp
--- a/Engine/Source/Runtime/AIModule/Private/EnvironmentQuery/Tests/EnvQueryTest_Trace.cpp
+++ b/Engine/Source/Runtime/AIModule/Private/EnvironmentQuery/Tests/EnvQueryTest_Trace.cpp
@@ -112,7 +112,7 @@ bool UEnvQueryTest_Trace::RunLineTraceTo(const FVector& ItemPos, const FVector&
 	FCollisionQueryParams TraceParams(Params);
 	TraceParams.AddIgnoredActor(ItemActor);
 
-	const bool bHit = World->LineTraceTest(ContextPos, ItemPos, Channel, TraceParams);
+	const bool bHit = World->LineTraceTestByChannel(ContextPos, ItemPos, Channel, TraceParams);
 	return bHit;
 }
 
@@ -121,7 +121,7 @@ bool UEnvQueryTest_Trace::RunLineTraceFrom(const FVector& ItemPos, const FVector
 	FCollisionQueryParams TraceParams(Params);
 	TraceParams.AddIgnoredActor(ItemActor);
 
-	const bool bHit = World->LineTraceTest(ItemPos, ContextPos, Channel, TraceParams);
+	const bool bHit = World->LineTraceTestByChannel(ItemPos, ContextPos, Channel, TraceParams);
 	return bHit;
 }
 
@@ -130,7 +130,7 @@ bool UEnvQueryTest_Trace::RunBoxTraceTo(const FVector& ItemPos, const FVector& C
 	FCollisionQueryParams TraceParams(Params);
 	TraceParams.AddIgnoredActor(ItemActor);
 
-	const bool bHit = World->SweepTest(ContextPos, ItemPos, FQuat((ItemPos - ContextPos).Rotation()), Channel, FCollisionShape::MakeBox(Extent), TraceParams);
+	const bool bHit = World->SweepTestByChannel(ContextPos, ItemPos, FQuat((ItemPos - ContextPos).Rotation()), Channel, FCollisionShape::MakeBox(Extent), TraceParams);
 	return bHit;
 }
 
@@ -139,7 +139,7 @@ bool UEnvQueryTest_Trace::RunBoxTraceFrom(const FVector& ItemPos, const FVector&
 	FCollisionQueryParams TraceParams(Params);
 	TraceParams.AddIgnoredActor(ItemActor);
 
-	const bool bHit = World->SweepTest(ItemPos, ContextPos, FQuat((ContextPos - ItemPos).Rotation()), Channel, FCollisionShape::MakeBox(Extent), TraceParams);
+	const bool bHit = World->SweepTestByChannel(ItemPos, ContextPos, FQuat((ContextPos - ItemPos).Rotation()), Channel, FCollisionShape::MakeBox(Extent), TraceParams);
 	return bHit;
 }
 
@@ -148,7 +148,7 @@ bool UEnvQueryTest_Trace::RunSphereTraceTo(const FVector& ItemPos, const FVector
 	FCollisionQueryParams TraceParams(Params);
 	TraceParams.AddIgnoredActor(ItemActor);
 
-	const bool bHit = World->SweepTest(ContextPos, ItemPos, FQuat::Identity, Channel, FCollisionShape::MakeSphere(Extent.X), TraceParams);
+	const bool bHit = World->SweepTestByChannel(ContextPos, ItemPos, FQuat::Identity, Channel, FCollisionShape::MakeSphere(Extent.X), TraceParams);
 	return bHit;
 }
 
@@ -157,7 +157,7 @@ bool UEnvQueryTest_Trace::RunSphereTraceFrom(const FVector& ItemPos, const FVect
 	FCollisionQueryParams TraceParams(Params);
 	TraceParams.AddIgnoredActor(ItemActor);
 
-	const bool bHit = World->SweepTest(ItemPos, ContextPos, FQuat::Identity, Channel, FCollisionShape::MakeSphere(Extent.X), TraceParams);
+	const bool bHit = World->SweepTestByChannel(ItemPos, ContextPos, FQuat::Identity, Channel, FCollisionShape::MakeSphere(Extent.X), TraceParams);
 	return bHit;
 }
 
@@ -166,7 +166,7 @@ bool UEnvQueryTest_Trace::RunCapsuleTraceTo(const FVector& ItemPos, const FVecto
 	FCollisionQueryParams TraceParams(Params);
 	TraceParams.AddIgnoredActor(ItemActor);
 
-	const bool bHit = World->SweepTest(ContextPos, ItemPos, FQuat::Identity, Channel, FCollisionShape::MakeCapsule(Extent.X, Extent.Z), TraceParams);
+	const bool bHit = World->SweepTestByChannel(ContextPos, ItemPos, FQuat::Identity, Channel, FCollisionShape::MakeCapsule(Extent.X, Extent.Z), TraceParams);
 	return bHit;
 }
 
@@ -175,7 +175,7 @@ bool UEnvQueryTest_Trace::RunCapsuleTraceFrom(const FVector& ItemPos, const FVec
 	FCollisionQueryParams TraceParams(Params);
 	TraceParams.AddIgnoredActor(ItemActor);
 
-	const bool bHit = World->SweepTest(ItemPos, ContextPos, FQuat::Identity, Channel, FCollisionShape::MakeCapsule(Extent.X, Extent.Z), TraceParams);
+	const bool bHit = World->SweepTestByChannel(ItemPos, ContextPos, FQuat::Identity, Channel, FCollisionShape::MakeCapsule(Extent.X, Extent.Z), TraceParams);
 	return bHit;
 }
 
diff --git a/Engine/Source/Runtime/AIModule/Private/Navigation/PathFollowingComponent.cpp b/Engine/Source/Runtime/AIModule/Private/Navigation/PathFollowingComponent.cpp
--- a/Engine/Source/Runtime/AIModule/Private/Navigation/PathFollowingComponent.cpp
+++ b/Engine/Source/Runtime/AIModule/Private/Navigation/PathFollowingComponent.cpp
@@ -713,7 +713,7 @@ int32 UPathFollowingComponent::OptimizeSegmentVisibility(int32 StartIndex)
 		else
 		{
 			FHitResult Result;
-			if (GetWorld()->SweepSingle(Result, PawnLocation, AdjustedDestination, FQuat::Identity, ECC_WorldStatic, CollisionShape, CapsuleParams))
+			if (GetWorld()->SweepSingleByChannel(Result, PawnLocation, AdjustedDestination, FQuat::Identity, ECC_WorldStatic, CollisionShape, CapsuleParams))
 			{
 				break;
 			}
diff --git a/Engine/Source/Runtime/AIModule/Private/Perception/AISense_Sight.cpp b/Engine/Source/Runtime/AIModule/Private/Perception/AISense_Sight.cpp
--- a/Engine/Source/Runtime/AIModule/Private/Perception/AISense_Sight.cpp
+++ b/Engine/Source/Runtime/AIModule/Private/Perception/AISense_Sight.cpp
@@ -163,9 +163,9 @@ float UAISense_Sight::Update()
 							, FCollisionQueryParams(NAME_AILineOfSight, true, Listener.Listener->GetBodyActor())
 							, FCollisionObjectQueryParams(ECC_WorldStatic));*/
 						FHitResult HitResult;
-						const bool bHit = World->LineTraceSingle(HitResult, Listener.CachedLocation, TargetLocation
-							, FCollisionQueryParams(NAME_AILineOfSight, true, Listener.Listener->GetBodyActor())
-							, FCollisionObjectQueryParams(ECC_WorldStatic));
+						const bool bHit = World->LineTraceSingleByObjectType(HitResult, Listener.CachedLocation, TargetLocation
+							, FCollisionObjectQueryParams(ECC_WorldStatic)
+							, FCollisionQueryParams(NAME_AILineOfSight, true, Listener.Listener->GetBodyActor()));
 
 						++TracesCount;
 
diff --git a/Engine/Source/Runtime/AIModule/Private/Perception/PawnSensingComponent.cpp b/Engine/Source/Runtime/AIModule/Private/Perception/PawnSensingComponent.cpp
--- a/Engine/Source/Runtime/AIModule/Private/Perception/PawnSensingComponent.cpp
+++ b/Engine/Source/Runtime/AIModule/Private/Perception/PawnSensingComponent.cpp
@@ -425,7 +425,7 @@ bool UPawnSensingComponent::CanHear(const FVector& NoiseLoc, float Loudness, boo
 
 	// check if sound is occluded
 	static FName NAME_CanHear = FName(TEXT("CanHear"));
-	return !Owner->GetWorld()->LineTraceTest(HearingLocation, NoiseLoc, ECC_Visibility, FCollisionQueryParams(NAME_CanHear, true, Owner));
+	return !Owner->GetWorld()->LineTraceTestByChannel(HearingLocation, NoiseLoc, ECC_Visibility, FCollisionQueryParams(NAME_CanHear, true, Owner));
 }
 
 
diff --git a/Engine/Source/Runtime/Engine/Classes/Engine/World.h b/Engine/Source/Runtime/Engine/Classes/Engine/World.h
--- a/Engine/Source/Runtime/Engine/Classes/Engine/World.h
+++ b/Engine/Source/Runtime/Engine/Classes/Engine/World.h
@@ -1022,8 +1022,23 @@ public:
 	 * 	@param 	ResponseParam	ResponseContainer to be used for this trace
 	 *  @return TRUE if a blocking hit is found
 	 */
-	bool LineTraceTest(const FVector& Start,const FVector& End,ECollisionChannel TraceChannel,const struct FCollisionQueryParams& Params, const struct FCollisionResponseParams& ResponseParam = FCollisionResponseParams::DefaultResponseParam) const;
+	DEPRECATED(4.8, "Use LineTraceTestByChannel instead.")
+	bool LineTraceTest(const FVector& Start,const FVector& End,ECollisionChannel TraceChannel,const FCollisionQueryParams& Params, const FCollisionResponseParams& ResponseParam = FCollisionResponseParams::DefaultResponseParam) const
+	{
+		return LineTraceTestByChannel(Start, End, TraceChannel, Params, ResponseParam);
+	}
 
+	/**
+	 *  Trace a ray against the world using a specific channel and return if a blocking hit is found.
+	 *  @param  Start           Start location of the ray
+	 *  @param  End             End location of the ray
+	 *  @param  TraceChannel    The 'channel' that this ray is in, used to determine which components to hit
+	 *  @param  Params          Additional parameters used for the trace
+	 * 	@param 	ResponseParam	ResponseContainer to be used for this trace
+	 *  @return TRUE if a blocking hit is found
+	 */
+	bool LineTraceTestByChannel(const FVector& Start,const FVector& End,ECollisionChannel TraceChannel, const FCollisionQueryParams& Params = FCollisionQueryParams::DefaultQueryParam, const FCollisionResponseParams& ResponseParam = FCollisionResponseParams::DefaultResponseParam) const;
+	
 	/**
 	 *  Trace a ray against the world using object types and return if a blocking hit is found.
 	 *  @param  Start           	Start location of the ray
@@ -1032,7 +1047,21 @@ public:
 	 *	@param	ObjectQueryParams	List of object types it's looking for
 	 *  @return TRUE if any hit is found
 	 */
-	bool LineTraceTest(const FVector& Start,const FVector& End,const struct FCollisionQueryParams& Params, const struct FCollisionObjectQueryParams& ObjectQueryParams) const;
+	DEPRECATED(4.8, "Use LineTraceTestByObjectType instead.")
+	bool LineTraceTest(const FVector& Start,const FVector& End,const FCollisionQueryParams& Params, const FCollisionObjectQueryParams& ObjectQueryParams) const
+	{
+		return LineTraceTestByObjectType(Start, End, ObjectQueryParams, Params);
+	}
+
+	/**
+	 *  Trace a ray against the world using object types and return if a blocking hit is found.
+	 *  @param  Start           	Start location of the ray
+	 *  @param  End             	End location of the ray
+	 *	@param	ObjectQueryParams	List of object types it's looking for
+	 *  @param  Params          	Additional parameters used for the trace
+	 *  @return TRUE if any hit is found
+	 */
+	bool LineTraceTestByObjectType(const FVector& Start,const FVector& End,const FCollisionObjectQueryParams& ObjectQueryParams, const FCollisionQueryParams& Params = FCollisionQueryParams::DefaultQueryParam) const;
 
 	/**
 	 *  Trace a ray against the world using a specific profile and return if a blocking hit is found.
@@ -1042,7 +1071,7 @@ public:
 	 *  @param  Params          Additional parameters used for the trace
 	 *  @return TRUE if a blocking hit is found
 	 */
-	bool LineTraceTestByProfile(const FVector& Start, const FVector& End, FName ProfileName, const struct FCollisionQueryParams& Params) const;
+	bool LineTraceTestByProfile(const FVector& Start, const FVector& End, FName ProfileName, const FCollisionQueryParams& Params = FCollisionQueryParams::DefaultQueryParam) const;
 
 	/**
 	 *  Trace a ray against the world using a specific channel and return the first blocking hit
@@ -1054,7 +1083,22 @@ public:
 	 * 	@param 	ResponseParam	ResponseContainer to be used for this trace	 
 	 *  @return TRUE if a blocking hit is found
 	 */
-	bool LineTraceSingle(struct FHitResult& OutHit,const FVector& Start,const FVector& End,ECollisionChannel TraceChannel,const struct FCollisionQueryParams& Params, const struct FCollisionResponseParams& ResponseParam = FCollisionResponseParams::DefaultResponseParam) const;
+	DEPRECATED(4.8, "Use LineTraceSingleByChannel instead.")
+	bool LineTraceSingle(struct FHitResult& OutHit,const FVector& Start,const FVector& End,ECollisionChannel TraceChannel,const FCollisionQueryParams& Params, const FCollisionResponseParams& ResponseParam = FCollisionResponseParams::DefaultResponseParam) const
+	{
+		return LineTraceSingleByChannel(OutHit, Start, End, TraceChannel, Params, ResponseParam);
+	}
+	/**
+	 *  Trace a ray against the world using a specific channel and return the first blocking hit
+	 *  @param  OutHit          First blocking hit found
+	 *  @param  Start           Start location of the ray
+	 *  @param  End             End location of the ray
+	 *  @param  TraceChannel    The 'channel' that this ray is in, used to determine which components to hit
+	 *  @param  Params          Additional parameters used for the trace
+	 * 	@param 	ResponseParam	ResponseContainer to be used for this trace	 
+	 *  @return TRUE if a blocking hit is found
+	 */
+	bool LineTraceSingleByChannel(struct FHitResult& OutHit,const FVector& Start,const FVector& End,ECollisionChannel TraceChannel,const FCollisionQueryParams& Params = FCollisionQueryParams::DefaultQueryParam, const FCollisionResponseParams& ResponseParam = FCollisionResponseParams::DefaultResponseParam) const;
 
 	/**
 	 *  Trace a ray against the world using object types and return the first blocking hit
@@ -1065,7 +1109,22 @@ public:
 	 *	@param	ObjectQueryParams	List of object types it's looking for
 	 *  @return TRUE if any hit is found
 	 */
-	bool LineTraceSingle(struct FHitResult& OutHit,const FVector& Start,const FVector& End,const struct FCollisionQueryParams& Params, const struct FCollisionObjectQueryParams& ObjectQueryParams) const;
+	DEPRECATED(4.8, "Use LineTraceSingleByObjectType instead.")
+	bool LineTraceSingle(struct FHitResult& OutHit,const FVector& Start,const FVector& End,const FCollisionQueryParams& Params, const FCollisionObjectQueryParams& ObjectQueryParams) const
+	{
+		return LineTraceSingleByObjectType(OutHit, Start, End, ObjectQueryParams, Params);
+	}
+
+	/**
+	 *  Trace a ray against the world using object types and return the first blocking hit
+	 *  @param  OutHit          First blocking hit found
+	 *  @param  Start           Start location of the ray
+	 *  @param  End             End location of the ray
+	 *	@param	ObjectQueryParams	List of object types it's looking for
+	 *  @param  Params          Additional parameters used for the trace
+	 *  @return TRUE if any hit is found
+	 */
+	bool LineTraceSingleByObjectType(struct FHitResult& OutHit,const FVector& Start,const FVector& End,const FCollisionObjectQueryParams& ObjectQueryParams, const FCollisionQueryParams& Params = FCollisionQueryParams::DefaultQueryParam) const;
 
 	/**
 	 *  Trace a ray against the world using a specific profile and return the first blocking hit
@@ -1076,7 +1135,25 @@ public:
 	 *  @param  Params          Additional parameters used for the trace
 	 *  @return TRUE if a blocking hit is found
 	 */
-	bool LineTraceSingleByProfile(struct FHitResult& OutHit, const FVector& Start, const FVector& End, FName ProfileName, const struct FCollisionQueryParams& Params) const;
+	bool LineTraceSingleByProfile(struct FHitResult& OutHit, const FVector& Start, const FVector& End, FName ProfileName, const FCollisionQueryParams& Params = FCollisionQueryParams::DefaultQueryParam) const;
+
+	/**
+	 *  Trace a ray against the world using a specific channel and return overlapping hits and then first blocking hit
+	 *  Results are sorted, so a blocking hit (if found) will be the last element of the array
+	 *  Only the single closest blocking result will be generated, no tests will be done after that
+	 *  @param  OutHits         Array of hits found between ray and the world
+	 *  @param  Start           Start location of the ray
+	 *  @param  End             End location of the ray
+	 *  @param  TraceChannel    The 'channel' that this ray is in, used to determine which components to hit
+	 *  @param  Params          Additional parameters used for the trace
+	 * 	@param 	ResponseParam	ResponseContainer to be used for this trace	 
+	 *  @return TRUE if OutHits contains any blocking hit entries
+	 */
+	DEPRECATED(4.8, "Use LineTraceMultiByChannel instead.")
+	bool LineTraceMulti(TArray<struct FHitResult>& OutHits,const FVector& Start,const FVector& End,ECollisionChannel TraceChannel,const FCollisionQueryParams& Params, const FCollisionResponseParams& ResponseParam = FCollisionResponseParams::DefaultResponseParam) const
+	{
+		return LineTraceMultiByChannel(OutHits, Start, End, TraceChannel, Params, ResponseParam);
+	}
 
 	/**
 	 *  Trace a ray against the world using a specific channel and return overlapping hits and then first blocking hit
@@ -1090,7 +1167,8 @@ public:
 	 * 	@param 	ResponseParam	ResponseContainer to be used for this trace	 
 	 *  @return TRUE if OutHits contains any blocking hit entries
 	 */
-	bool LineTraceMulti(TArray<struct FHitResult>& OutHits,const FVector& Start,const FVector& End,ECollisionChannel TraceChannel,const struct FCollisionQueryParams& Params, const struct FCollisionResponseParams& ResponseParam = FCollisionResponseParams::DefaultResponseParam) const;
+	bool LineTraceMultiByChannel(TArray<struct FHitResult>& OutHits,const FVector& Start,const FVector& End,ECollisionChannel TraceChannel,const FCollisionQueryParams& Params = FCollisionQueryParams::DefaultQueryParam, const FCollisionResponseParams& ResponseParam = FCollisionResponseParams::DefaultResponseParam) const;
+
 
 	/**
 	 *  Trace a ray against the world using object types and return overlapping hits and then first blocking hit
@@ -1103,7 +1181,25 @@ public:
 	 *	@param	ObjectQueryParams	List of object types it's looking for
 	 *  @return TRUE if any hit is found
 	 */
-	bool LineTraceMulti(TArray<struct FHitResult>& OutHits,const FVector& Start,const FVector& End,const struct FCollisionQueryParams& Params, const struct FCollisionObjectQueryParams& ObjectQueryParams) const;
+	DEPRECATED(4.8, "Use LineTraceMultiByObjectType instead.")
+	bool LineTraceMulti(TArray<struct FHitResult>& OutHits,const FVector& Start,const FVector& End,const FCollisionQueryParams& Params, const FCollisionObjectQueryParams& ObjectQueryParams) const
+	{
+		return LineTraceMultiByObjectType(OutHits, Start, End, ObjectQueryParams, Params);
+	}
+
+	/**
+	 *  Trace a ray against the world using object types and return overlapping hits and then first blocking hit
+	 *  Results are sorted, so a blocking hit (if found) will be the last element of the array
+	 *  Only the single closest blocking result will be generated, no tests will be done after that
+	 *  @param  OutHits         Array of hits found between ray and the world
+	 *  @param  Start           Start location of the ray
+	 *  @param  End             End location of the ray
+	 *	@param	ObjectQueryParams	List of object types it's looking for
+	 *  @param  Params          Additional parameters used for the trace
+	 *  @return TRUE if any hit is found
+	 */
+	bool LineTraceMultiByObjectType(TArray<struct FHitResult>& OutHits,const FVector& Start,const FVector& End,const FCollisionObjectQueryParams& ObjectQueryParams, const FCollisionQueryParams& Params = FCollisionQueryParams::DefaultQueryParam) const;
+
 
 	/**
 	 *  Trace a ray against the world using a specific profile and return overlapping hits and then first blocking hit
@@ -1116,7 +1212,7 @@ public:
 	 *  @param  Params          Additional parameters used for the trace
 	 *  @return TRUE if OutHits contains any blocking hit entries
 	 */
-	bool LineTraceMultiByProfile(TArray<struct FHitResult>& OutHits, const FVector& Start, const FVector& End, FName ProfileName, const struct FCollisionQueryParams& Params) const;
+	bool LineTraceMultiByProfile(TArray<struct FHitResult>& OutHits, const FVector& Start, const FVector& End, FName ProfileName, const FCollisionQueryParams& Params = FCollisionQueryParams::DefaultQueryParam) const;
 
 	/**
 	 *  Sweep a sphere against the world using a specific channel and return if a blocking hit is found.
@@ -1128,7 +1224,23 @@ public:
 	 * 	@param 	ResponseParam	ResponseContainer to be used for this trace	 	 
 	 *  @return TRUE if a blocking hit is found
 	 */
-	bool SweepTest(const FVector& Start, const FVector& End, const FQuat& Rot, ECollisionChannel TraceChannel, const struct FCollisionShape& CollisionShape, const struct FCollisionQueryParams& Params, const struct FCollisionResponseParams& ResponseParam = FCollisionResponseParams::DefaultResponseParam) const;
+	DEPRECATED(4.8, "Use SweepTestByChannel instead.")
+	bool SweepTest(const FVector& Start, const FVector& End, const FQuat& Rot, ECollisionChannel TraceChannel, const FCollisionShape& CollisionShape, const FCollisionQueryParams& Params, const FCollisionResponseParams& ResponseParam = FCollisionResponseParams::DefaultResponseParam) const
+	{
+		return SweepTestByChannel(Start, End, Rot, TraceChannel, CollisionShape, Params, ResponseParam);
+	}
+
+	/**
+	 *  Sweep a sphere against the world using a specific channel and return if a blocking hit is found.
+	 *  @param  Start           Start location of the sphere
+	 *  @param  End             End location of the sphere
+	 *  @param  TraceChannel    The 'channel' that this trace uses, used to determine which components to hit
+	 *  @param	CollisionShape	CollisionShape - supports Box, Sphere, Capsule
+	 *  @param  Params          Additional parameters used for the trace
+	 * 	@param 	ResponseParam	ResponseContainer to be used for this trace	 	 
+	 *  @return TRUE if a blocking hit is found
+	 */
+	bool SweepTestByChannel(const FVector& Start, const FVector& End, const FQuat& Rot, ECollisionChannel TraceChannel, const FCollisionShape& CollisionShape, const FCollisionQueryParams& Params = FCollisionQueryParams::DefaultQueryParam, const FCollisionResponseParams& ResponseParam = FCollisionResponseParams::DefaultResponseParam) const;
 
 	/**
 	 *  Sweep a sphere against the world using object types and return if a blocking hit is found.
@@ -1139,7 +1251,23 @@ public:
 	 *	@param	ObjectQueryParams	List of object types it's looking for
 	 *  @return TRUE if any hit is found
 	 */
-	bool SweepTest(const FVector& Start, const FVector& End, const FQuat& Rot, const struct FCollisionShape& CollisionShape, const struct FCollisionQueryParams& Params, const struct FCollisionObjectQueryParams& ObjectQueryParams) const;
+	DEPRECATED(4.8, "Use SweepTestByObjectType instead.")
+	bool SweepTest(const FVector& Start, const FVector& End, const FQuat& Rot, const FCollisionShape& CollisionShape, const FCollisionQueryParams& Params, const FCollisionObjectQueryParams& ObjectQueryParams) const
+	{
+		return SweepTestByObjectType(Start, End, Rot, ObjectQueryParams, CollisionShape, Params);
+	}
+
+	/**
+	 *  Sweep a sphere against the world using object types and return if a blocking hit is found.
+	 *  @param  Start           Start location of the sphere
+	 *  @param  End             End location of the sphere
+	 *	@param	ObjectQueryParams	List of object types it's looking for
+	 *  @param	CollisionShape	CollisionShape - supports Box, Sphere, Capsule
+	 *  @param  Params          Additional parameters used for the trace
+	 *  @return TRUE if any hit is found
+	 */
+	bool SweepTestByObjectType(const FVector& Start, const FVector& End, const FQuat& Rot, const FCollisionObjectQueryParams& ObjectQueryParams, const FCollisionShape& CollisionShape, const FCollisionQueryParams& Params = FCollisionQueryParams::DefaultQueryParam) const;
+
 
 	/**
 	 *  Sweep a sphere against the world using a specific profile and return if a blocking hit is found.
@@ -1150,7 +1278,24 @@ public:
 	 *  @param  Params          Additional parameters used for the trace
 	 *  @return TRUE if a blocking hit is found
 	 */
-	bool SweepTestByProfile(const FVector& Start, const FVector& End, const FQuat& Rot, FName ProfileName, const struct FCollisionShape& CollisionShape, const struct FCollisionQueryParams& Params) const;
+	bool SweepTestByProfile(const FVector& Start, const FVector& End, const FQuat& Rot, FName ProfileName, const FCollisionShape& CollisionShape, const FCollisionQueryParams& Params) const;
+
+	/**
+	 *  Sweep a sphere against the world and return the first blocking hit using a specific channel
+	 *  @param  OutHit          First blocking hit found
+	 *  @param  Start           Start location of the sphere
+	 *  @param  End             End location of the sphere
+	 *  @param  TraceChannel    The 'channel' that this trace is in, used to determine which components to hit
+	 *  @param	CollisionShape	CollisionShape - supports Box, Sphere, Capsule
+	 *  @param  Params          Additional parameters used for the trace
+	 * 	@param 	ResponseParam	ResponseContainer to be used for this trace	 
+	 *  @return TRUE if OutHits contains any blocking hit entries
+	 */
+	DEPRECATED(4.8, "Use SweepSingleByChannel instead.")
+	bool SweepSingle(struct FHitResult& OutHit, const FVector& Start, const FVector& End, const FQuat& Rot, ECollisionChannel TraceChannel, const FCollisionShape& CollisionShape, const FCollisionQueryParams& Params, const FCollisionResponseParams& ResponseParam = FCollisionResponseParams::DefaultResponseParam) const
+	{
+		return SweepSingleByChannel(OutHit, Start, End, Rot, TraceChannel, CollisionShape, Params, ResponseParam);
+	}
 
 	/**
 	 *  Sweep a sphere against the world and return the first blocking hit using a specific channel
@@ -1163,7 +1308,7 @@ public:
 	 * 	@param 	ResponseParam	ResponseContainer to be used for this trace	 
 	 *  @return TRUE if OutHits contains any blocking hit entries
 	 */
-	bool SweepSingle(struct FHitResult& OutHit, const FVector& Start, const FVector& End, const FQuat& Rot, ECollisionChannel TraceChannel, const struct FCollisionShape& CollisionShape, const struct FCollisionQueryParams& Params, const struct FCollisionResponseParams& ResponseParam = FCollisionResponseParams::DefaultResponseParam) const;
+	bool SweepSingleByChannel(struct FHitResult& OutHit, const FVector& Start, const FVector& End, const FQuat& Rot, ECollisionChannel TraceChannel, const FCollisionShape& CollisionShape, const FCollisionQueryParams& Params = FCollisionQueryParams::DefaultQueryParam, const FCollisionResponseParams& ResponseParam = FCollisionResponseParams::DefaultResponseParam) const;
 
 	/**
 	 *  Sweep a sphere against the world and return the first blocking hit using object types
@@ -1175,7 +1320,23 @@ public:
 	 *	@param	ObjectQueryParams	List of object types it's looking for
 	 *  @return TRUE if any hit is found
 	 */
-	bool SweepSingle(struct FHitResult& OutHit, const FVector& Start, const FVector& End, const FQuat& Rot, const struct FCollisionShape& CollisionShape, const struct FCollisionQueryParams& Params, const struct FCollisionObjectQueryParams& ObjectQueryParams) const;
+	DEPRECATED(4.8, "Use SweepSingleByObjectType instead.")
+	bool SweepSingle(struct FHitResult& OutHit, const FVector& Start, const FVector& End, const FQuat& Rot, const FCollisionShape& CollisionShape, const FCollisionQueryParams& Params, const FCollisionObjectQueryParams& ObjectQueryParams) const
+	{
+		return SweepSingleByObjectType(OutHit, Start, End, Rot, ObjectQueryParams, CollisionShape, Params);
+	}
+
+	/**
+	 *  Sweep a sphere against the world and return the first blocking hit using object types
+	 *  @param  OutHit          First blocking hit found
+	 *  @param  Start           Start location of the sphere
+	 *  @param  End             End location of the sphere
+	 *	@param	ObjectQueryParams	List of object types it's looking for
+	 *  @param	CollisionShape	CollisionShape - supports Box, Sphere, Capsule
+	 *  @param  Params          Additional parameters used for the trace
+	 *  @return TRUE if any hit is found
+	 */
+	bool SweepSingleByObjectType(struct FHitResult& OutHit, const FVector& Start, const FVector& End, const FQuat& Rot, const FCollisionObjectQueryParams& ObjectQueryParams, const FCollisionShape& CollisionShape, const FCollisionQueryParams& Params = FCollisionQueryParams::DefaultQueryParam) const;
 
 	/**
 	 *  Sweep a sphere against the world and return the first blocking hit using a specific profile
@@ -1187,7 +1348,7 @@ public:
 	 *  @param  Params          Additional parameters used for the trace
 	 *  @return TRUE if OutHits contains any blocking hit entries
 	 */
-	bool SweepSingleByProfile(struct FHitResult& OutHit, const FVector& Start, const FVector& End, const FQuat& Rot, FName ProfileName, const struct FCollisionShape& CollisionShape, const struct FCollisionQueryParams& Params) const;
+	bool SweepSingleByProfile(struct FHitResult& OutHit, const FVector& Start, const FVector& End, const FQuat& Rot, FName ProfileName, const FCollisionShape& CollisionShape, const FCollisionQueryParams& Params = FCollisionQueryParams::DefaultQueryParam) const;
 
 	/**
 	 *  Sweep a sphere against the world and return all initial overlaps using a specific channel (including blocking) if requested, then overlapping hits and then first blocking hit
@@ -1202,7 +1363,26 @@ public:
 	 * 	@param 	ResponseParam	ResponseContainer to be used for this trace	 
 	 *  @return TRUE if OutHits contains any blocking hit entries
 	 */
-	bool SweepMulti(TArray<struct FHitResult>& OutHits, const FVector& Start, const FVector& End, const FQuat& Rot, ECollisionChannel TraceChannel, const struct FCollisionShape& CollisionShape, const struct FCollisionQueryParams& Params, const struct FCollisionResponseParams& ResponseParam = FCollisionResponseParams::DefaultResponseParam) const;
+	DEPRECATED(4.8, "Use SweepMultiByChannel instead.")
+	bool SweepMulti(TArray<struct FHitResult>& OutHits, const FVector& Start, const FVector& End, const FQuat& Rot, ECollisionChannel TraceChannel, const FCollisionShape& CollisionShape, const FCollisionQueryParams& Params, const FCollisionResponseParams& ResponseParam = FCollisionResponseParams::DefaultResponseParam) const
+	{
+		return SweepMultiByChannel(OutHits, Start, End, Rot, TraceChannel, CollisionShape, Params, ResponseParam);
+	}
+
+	/**
+	 *  Sweep a sphere against the world and return all initial overlaps using a specific channel (including blocking) if requested, then overlapping hits and then first blocking hit
+	 *  Results are sorted, so a blocking hit (if found) will be the last element of the array
+	 *  Only the single closest blocking result will be generated, no tests will be done after that
+	 *  @param  OutHits         Array of hits found between ray and the world
+	 *  @param  Start           Start location of the sphere
+	 *  @param  End             End location of the sphere
+	 *  @param  TraceChannel    The 'channel' that this ray is in, used to determine which components to hit
+	 *  @param	CollisionShape	CollisionShape - supports Box, Sphere, Capsule
+	 *  @param  Params          Additional parameters used for the trace
+	 * 	@param 	ResponseParam	ResponseContainer to be used for this trace	 
+	 *  @return TRUE if OutHits contains any blocking hit entries
+	 */
+	bool SweepMultiByChannel(TArray<struct FHitResult>& OutHits, const FVector& Start, const FVector& End, const FQuat& Rot, ECollisionChannel TraceChannel, const FCollisionShape& CollisionShape, const FCollisionQueryParams& Params = FCollisionQueryParams::DefaultQueryParam, const FCollisionResponseParams& ResponseParam = FCollisionResponseParams::DefaultResponseParam) const;
 
 	/**
 	 *  Sweep a sphere against the world and return all initial overlaps using object types (including blocking) if requested, then overlapping hits and then first blocking hit
@@ -1216,7 +1396,25 @@ public:
 	 *	@param	ObjectQueryParams	List of object types it's looking for
 	 *  @return TRUE if any hit is found
 	 */
-	bool SweepMulti(TArray<struct FHitResult>& OutHits, const FVector& Start, const FVector& End, const FQuat& Rot, const struct FCollisionShape& CollisionShape, const struct FCollisionQueryParams& Params, const struct FCollisionObjectQueryParams& ObjectQueryParams) const;
+	DEPRECATED(4.8, "Use SweepMultiByObjectType instead.")
+	bool SweepMulti(TArray<struct FHitResult>& OutHits, const FVector& Start, const FVector& End, const FQuat& Rot, const FCollisionShape& CollisionShape, const FCollisionQueryParams& Params, const FCollisionObjectQueryParams& ObjectQueryParams) const
+	{
+		return SweepMultiByObjectType(OutHits, Start, End, Rot, ObjectQueryParams, CollisionShape, Params);
+	}
+
+	/**
+	 *  Sweep a sphere against the world and return all initial overlaps using object types (including blocking) if requested, then overlapping hits and then first blocking hit
+	 *  Results are sorted, so a blocking hit (if found) will be the last element of the array
+	 *  Only the single closest blocking result will be generated, no tests will be done after that
+	 *  @param  OutHits         Array of hits found between ray and the world
+	 *  @param  Start           Start location of the sphere
+	 *  @param  End             End location of the sphere
+	 *	@param	ObjectQueryParams	List of object types it's looking for
+	 *  @param	CollisionShape	CollisionShape - supports Box, Sphere, Capsule
+	 *  @param  Params          Additional parameters used for the trace
+	 *  @return TRUE if any hit is found
+	 */
+	bool SweepMultiByObjectType(TArray<struct FHitResult>& OutHits, const FVector& Start, const FVector& End, const FQuat& Rot, const FCollisionObjectQueryParams& ObjectQueryParams, const FCollisionShape& CollisionShape, const FCollisionQueryParams& Params = FCollisionQueryParams::DefaultQueryParam) const;
 
 	/**
 	 *  Sweep a sphere against the world and return all initial overlaps using a specific profile, then overlapping hits and then first blocking hit
@@ -1230,7 +1428,7 @@ public:
 	 *  @param  Params          Additional parameters used for the trace
 	 *  @return TRUE if OutHits contains any blocking hit entries
 	 */
-	bool SweepMultiByProfile(TArray<FHitResult>& OutHits, const FVector& Start, const FVector& End, const FQuat& Rot, FName ProfileName, const struct FCollisionShape& CollisionShape, const struct FCollisionQueryParams& Params) const;
+	bool SweepMultiByProfile(TArray<FHitResult>& OutHits, const FVector& Start, const FVector& End, const FQuat& Rot, FName ProfileName, const FCollisionShape& CollisionShape, const FCollisionQueryParams& Params = FCollisionQueryParams::DefaultQueryParam) const;
 
 	/**
 	 *  Test the collision of a shape at the supplied location using a specific channel, and return if any blocking overlap is found
@@ -1241,8 +1439,11 @@ public:
 	 * 	@param 	ResponseParam	ResponseContainer to be used for this trace	 
 	 *  @return TRUE if any blocking results are found
 	 */
-	DEPRECATED(4.8, "Use OverlapBlockingTest instead.")
-	bool OverlapTest(const FVector& Pos, const FQuat& Rot, ECollisionChannel TraceChannel, const struct FCollisionShape& CollisionShape, const struct FCollisionQueryParams& Params, const struct FCollisionResponseParams& ResponseParam = FCollisionResponseParams::DefaultResponseParam) const;
+	DEPRECATED(4.8, "Use OverlapBlockingTestByChannel instead.")
+	bool OverlapTest(const FVector& Pos, const FQuat& Rot, ECollisionChannel TraceChannel, const FCollisionShape& CollisionShape, const FCollisionQueryParams& Params, const FCollisionResponseParams& ResponseParam = FCollisionResponseParams::DefaultResponseParam) const
+	{
+		return OverlapBlockingTestByChannel(Pos, Rot, TraceChannel, CollisionShape, Params, ResponseParam);
+	}
 
 	/**
 	*  Test the collision of a shape at the supplied location using a specific channel, and return if any blocking overlap is found
@@ -1253,7 +1454,7 @@ public:
 	*  @param  ResponseParam	ResponseContainer to be used for this trace
 	*  @return TRUE if any blocking results are found
 	*/
-	bool OverlapBlockingTest(const FVector& Pos, const FQuat& Rot, ECollisionChannel TraceChannel, const struct FCollisionShape& CollisionShape, const struct FCollisionQueryParams& Params, const struct FCollisionResponseParams& ResponseParam = FCollisionResponseParams::DefaultResponseParam) const;
+	bool OverlapBlockingTestByChannel(const FVector& Pos, const FQuat& Rot, ECollisionChannel TraceChannel, const FCollisionShape& CollisionShape, const FCollisionQueryParams& Params = FCollisionQueryParams::DefaultQueryParam, const FCollisionResponseParams& ResponseParam = FCollisionResponseParams::DefaultResponseParam) const;
 
 	/**
 	*  Test the collision of a shape at the supplied location using a specific channel, and return if any blocking or overlapping shape is found
@@ -1264,7 +1465,7 @@ public:
 	*  @param  ResponseParam	ResponseContainer to be used for this trace
 	*  @return TRUE if any blocking or overlapping results are found
 	*/
-	bool OverlapAnyTest(const FVector& Pos, const FQuat& Rot, ECollisionChannel TraceChannel, const struct FCollisionShape& CollisionShape, const struct FCollisionQueryParams& Params, const struct FCollisionResponseParams& ResponseParam = FCollisionResponseParams::DefaultResponseParam) const;
+	bool OverlapAnyTestByChannel(const FVector& Pos, const FQuat& Rot, ECollisionChannel TraceChannel, const FCollisionShape& CollisionShape, const FCollisionQueryParams& Params = FCollisionQueryParams::DefaultQueryParam, const FCollisionResponseParams& ResponseParam = FCollisionResponseParams::DefaultResponseParam) const;
 
 	/**
 	 *  Test the collision of a shape at the supplied location using object types, and return if any blocking overlap is found
@@ -1274,18 +1475,21 @@ public:
 	 *	@param	ObjectQueryParams	List of object types it's looking for
 	 *  @return TRUE if any blocking results are found
 	 */
-	DEPRECATED(4.8, "Use OverlapAnyTest instead.")
-	bool OverlapTest(const FVector& Pos, const FQuat& Rot, const struct FCollisionShape& CollisionShape, const struct FCollisionQueryParams& Params, const struct FCollisionObjectQueryParams& ObjectQueryParams) const;
+	DEPRECATED(4.8, "Use OverlapAnyTestByObjectType instead.")
+	bool OverlapTest(const FVector& Pos, const FQuat& Rot, const FCollisionShape& CollisionShape, const FCollisionQueryParams& Params, const FCollisionObjectQueryParams& ObjectQueryParams) const
+	{
+		return OverlapAnyTestByObjectType(Pos, Rot, ObjectQueryParams, CollisionShape, Params);
+	}
 
 	/**
 	*  Test the collision of a shape at the supplied location using object types, and return if any overlap is found
 	*  @param  Pos             Location of center of box to test against the world
+	*  @param  ObjectQueryParams	List of object types it's looking for
 	*  @param  CollisionShape	CollisionShape - supports Box, Sphere, Capsule, Convex
 	*  @param  Params          Additional parameters used for the trace
-	*  @param  ObjectQueryParams	List of object types it's looking for
 	*  @return TRUE if any blocking results are found
 	*/
-	bool OverlapAnyTest(const FVector& Pos, const FQuat& Rot, const struct FCollisionShape& CollisionShape, const struct FCollisionQueryParams& Params, const struct FCollisionObjectQueryParams& ObjectQueryParams) const;
+	bool OverlapAnyTestByObjectType(const FVector& Pos, const FQuat& Rot, const FCollisionObjectQueryParams& ObjectQueryParams, const FCollisionShape& CollisionShape, const FCollisionQueryParams& Params = FCollisionQueryParams::DefaultQueryParam) const;
 
 	/**
 	 *  Test the collision of a shape at the supplied location using a specific profile, and return if any blocking overlap is found
@@ -1296,7 +1500,10 @@ public:
 	 *  @return TRUE if any blocking results are found
 	 */
 	DEPRECATED(4.8, "Use OverlapBlockingTestByProfile instead.")
-	bool OverlapTestByProfile(const FVector& Pos, const FQuat& Rot, FName ProfileName, const struct FCollisionShape& CollisionShape, const struct FCollisionQueryParams& Params) const;
+	bool OverlapTestByProfile(const FVector& Pos, const FQuat& Rot, FName ProfileName, const FCollisionShape& CollisionShape, const FCollisionQueryParams& Params) const
+	{
+		return OverlapBlockingTestByProfile(Pos, Rot, ProfileName, CollisionShape, Params);
+	}
 
 	/**
 	*  Test the collision of a shape at the supplied location using a specific profile, and return if any blocking overlap is found
@@ -1306,7 +1513,7 @@ public:
 	*  @param  Params          Additional parameters used for the trace
 	*  @return TRUE if any blocking results are found
 	*/
-	bool OverlapBlockingTestByProfile(const FVector& Pos, const FQuat& Rot, FName ProfileName, const struct FCollisionShape& CollisionShape, const struct FCollisionQueryParams& Params) const;
+	bool OverlapBlockingTestByProfile(const FVector& Pos, const FQuat& Rot, FName ProfileName, const FCollisionShape& CollisionShape, const FCollisionQueryParams& Params = FCollisionQueryParams::DefaultQueryParam) const;
 
 	/**
 	*  Test the collision of a shape at the supplied location using a specific profile, and return if any blocking or overlap is found
@@ -1316,7 +1523,7 @@ public:
 	*  @param  Params          Additional parameters used for the trace
 	*  @return TRUE if any blocking or overlapping results are found
 	*/
-	bool OverlapAnyTestByProfile(const FVector& Pos, const FQuat& Rot, FName ProfileName, const struct FCollisionShape& CollisionShape, const struct FCollisionQueryParams& Params) const;
+	bool OverlapAnyTestByProfile(const FVector& Pos, const FQuat& Rot, FName ProfileName, const FCollisionShape& CollisionShape, const FCollisionQueryParams& Params = FCollisionQueryParams::DefaultQueryParam) const;
 
 	/**
 	 *  Test the collision of a shape at the supplied location using a specific channel, and determine the set of components that it overlaps
@@ -1328,8 +1535,8 @@ public:
 	 * 	@param 	ResponseParam	ResponseContainer to be used for this trace
 	 *  @return TRUE if OutOverlap is a blocking result
 	 */
-	DEPRECATED(4.8, "Use OverlapMulti instead. Deciding which overlap to use is gameplay specific.")
-	bool OverlapSingle(struct FOverlapResult& OutOverlap, const FVector& Pos, const FQuat& Rot, ECollisionChannel TraceChannel, const struct FCollisionShape& CollisionShape, const struct FCollisionQueryParams& Params, const struct FCollisionResponseParams& ResponseParam = FCollisionResponseParams::DefaultResponseParam) const;
+	DEPRECATED(4.8, "Use OverlapMultiByChannel instead. Deciding which overlap to use is gameplay specific.")
+	bool OverlapSingle(struct FOverlapResult& OutOverlap, const FVector& Pos, const FQuat& Rot, ECollisionChannel TraceChannel, const FCollisionShape& CollisionShape, const FCollisionQueryParams& Params, const FCollisionResponseParams& ResponseParam = FCollisionResponseParams::DefaultResponseParam) const;
 
 	/**
 	 *  Test the collision of a shape at the supplied location using object types, and determine the set of components that it overlaps
@@ -1340,8 +1547,8 @@ public:
 	 *	@param	ObjectQueryParams	List of object types it's looking for
 	 *  @return TRUE if any overlap is found
 	 */
-	DEPRECATED(4.8, "Use OverlapMulti instead. Deciding which overlap to use is gameplay specific.")
-	bool OverlapSingle(struct FOverlapResult& OutOverlap, const FVector& Pos, const FQuat& Rot, const struct FCollisionShape& CollisionShape, const struct FCollisionQueryParams& Params, const struct FCollisionObjectQueryParams& ObjectQueryParams) const;
+	DEPRECATED(4.8, "Use OverlapMultiByObjectType instead. Deciding which overlap to use is gameplay specific.")
+	bool OverlapSingle(struct FOverlapResult& OutOverlap, const FVector& Pos, const FQuat& Rot, const FCollisionShape& CollisionShape, const FCollisionQueryParams& Params, const FCollisionObjectQueryParams& ObjectQueryParams) const;
 
 	/**
 	 *  Test the collision of a shape at the supplied location using a specific profile, and determine the set of components that it overlaps
@@ -1353,7 +1560,7 @@ public:
 	 *  @return TRUE if OutOverlap is blocking
 	 */
 	DEPRECATED(4.8, "Use OverlapMultiByProfile instead. Deciding which overlap to use is gameplay specific.")
-	bool OverlapSingleByProfile(struct FOverlapResult& OutOverlap, const FVector& Pos, const FQuat& Rot, FName ProfileName, const struct FCollisionShape& CollisionShape, const struct FCollisionQueryParams& Params) const;
+	bool OverlapSingleByProfile(struct FOverlapResult& OutOverlap, const FVector& Pos, const FQuat& Rot, FName ProfileName, const FCollisionShape& CollisionShape, const FCollisionQueryParams& Params) const;
 
 	/**
 	 *  Test the collision of a shape at the supplied location using a specific channel, and determine the set of components that it overlaps
@@ -1365,7 +1572,23 @@ public:
 	 * 	@param 	ResponseParam	ResponseContainer to be used for this trace
 	 *  @return TRUE if OutOverlaps contains any blocking results
 	 */
-	bool OverlapMulti(TArray<struct FOverlapResult>& OutOverlaps, const FVector& Pos, const FQuat& Rot, ECollisionChannel TraceChannel, const struct FCollisionShape& CollisionShape, const struct FCollisionQueryParams& Params, const struct FCollisionResponseParams& ResponseParam = FCollisionResponseParams::DefaultResponseParam) const;
+	DEPRECATED(4.8, "Use OverlapMultiByChannel instead.")
+	bool OverlapMulti(TArray<struct FOverlapResult>& OutOverlaps, const FVector& Pos, const FQuat& Rot, ECollisionChannel TraceChannel, const FCollisionShape& CollisionShape, const FCollisionQueryParams& Params, const FCollisionResponseParams& ResponseParam = FCollisionResponseParams::DefaultResponseParam) const
+	{
+		return OverlapMultiByChannel(OutOverlaps, Pos, Rot, TraceChannel, CollisionShape, Params, ResponseParam);
+	}
+
+	/**
+	 *  Test the collision of a shape at the supplied location using a specific channel, and determine the set of components that it overlaps
+	 *  @param  OutOverlaps     Array of components found to overlap supplied box
+	 *  @param  Pos             Location of center of sphere to test against the world
+	 *  @param  TraceChannel    The 'channel' that this query is in, used to determine which components to hit
+	 *  @param	CollisionShape	CollisionShape - supports Box, Sphere, Capsule
+	 *  @param  Params          Additional parameters used for the trace
+	 * 	@param 	ResponseParam	ResponseContainer to be used for this trace
+	 *  @return TRUE if OutOverlaps contains any blocking results
+	 */
+	bool OverlapMultiByChannel(TArray<struct FOverlapResult>& OutOverlaps, const FVector& Pos, const FQuat& Rot, ECollisionChannel TraceChannel, const FCollisionShape& CollisionShape, const FCollisionQueryParams& Params = FCollisionQueryParams::DefaultQueryParam, const FCollisionResponseParams& ResponseParam = FCollisionResponseParams::DefaultResponseParam) const;
 
 	/**
 	 *  Test the collision of a shape at the supplied location using object types, and determine the set of components that it overlaps
@@ -1376,8 +1599,23 @@ public:
 	 *	@param	ObjectQueryParams	List of object types it's looking for
 	 *  @return TRUE if any overlap is found
 	 */
-	bool OverlapMulti(TArray<struct FOverlapResult>& OutOverlaps, const FVector& Pos, const FQuat& Rot, const struct FCollisionShape& CollisionShape, const struct FCollisionQueryParams& Params, const struct FCollisionObjectQueryParams& ObjectQueryParams) const;
+	DEPRECATED(4.8, "Use OverlapMultiByObjectType instead.")
+	bool OverlapMulti(TArray<struct FOverlapResult>& OutOverlaps, const FVector& Pos, const FQuat& Rot, const FCollisionShape& CollisionShape, const FCollisionQueryParams& Params, const FCollisionObjectQueryParams& ObjectQueryParams) const
+	{
+		return OverlapMultiByObjectType(OutOverlaps, Pos, Rot, ObjectQueryParams, CollisionShape, Params);
+	}
 	
+	/**
+	 *  Test the collision of a shape at the supplied location using object types, and determine the set of components that it overlaps
+	 *  @param  OutOverlaps     Array of components found to overlap supplied box
+	 *  @param  Pos             Location of center of sphere to test against the world
+	 *	@param	ObjectQueryParams	List of object types it's looking for
+	 *  @param	CollisionShape	CollisionShape - supports Box, Sphere, Capsule
+	 *  @param  Params          Additional parameters used for the trace
+	 *  @return TRUE if any overlap is found
+	 */
+	bool OverlapMultiByObjectType(TArray<struct FOverlapResult>& OutOverlaps, const FVector& Pos, const FQuat& Rot, const FCollisionObjectQueryParams& ObjectQueryParams, const FCollisionShape& CollisionShape, const FCollisionQueryParams& Params = FCollisionQueryParams::DefaultQueryParam) const;
+
 	/**
 	 *  Test the collision of a shape at the supplied location using a specific profile, and determine the set of components that it overlaps
 	 *  @param  OutOverlaps     Array of components found to overlap supplied box
@@ -1387,7 +1625,7 @@ public:
 	 *  @param  Params          Additional parameters used for the trace
 	 *  @return TRUE if OutOverlaps contains any blocking results
 	 */
-	bool OverlapMultiByProfile(TArray<struct FOverlapResult>& OutOverlaps, const FVector& Pos, const FQuat& Rot, FName ProfileName, const struct FCollisionShape& CollisionShape, const struct FCollisionQueryParams& Params) const;
+	bool OverlapMultiByProfile(TArray<struct FOverlapResult>& OutOverlaps, const FVector& Pos, const FQuat& Rot, FName ProfileName, const FCollisionShape& CollisionShape, const FCollisionQueryParams& Params = FCollisionQueryParams::DefaultQueryParam) const;
 
 	// COMPONENT SWEEP
 
@@ -1401,7 +1639,7 @@ public:
 	 *  @param  Params          Additional parameters used for the trace
 	 *  @return TRUE if OutHits contains any blocking hit entries
 	 */
-	bool ComponentSweepMulti(TArray<struct FHitResult>& OutHits, class UPrimitiveComponent* PrimComp, const FVector& Start,const FVector& End, const FRotator& Rot, const struct FComponentQueryParams& Params)  const;
+	bool ComponentSweepMulti(TArray<struct FHitResult>& OutHits, class UPrimitiveComponent* PrimComp, const FVector& Start,const FVector& End, const FRotator& Rot, const FComponentQueryParams& Params)  const;
 
 	// COMPONENT OVERLAP
 
@@ -1414,7 +1652,22 @@ public:
 	 *	@param	ObjectQueryParams	List of object types it's looking for. When this enters, we do object query with component shape
 	 *  @return TRUE if any hit is found
 	 */
-	bool ComponentOverlapMulti(TArray<struct FOverlapResult>& OutOverlaps, const class UPrimitiveComponent* PrimComp, const FVector& Pos, const FRotator& Rot, const struct FComponentQueryParams& Params, const struct FCollisionObjectQueryParams& ObjectQueryParams=FCollisionObjectQueryParams::DefaultObjectQueryParam) const;
+	bool ComponentOverlapMulti(TArray<struct FOverlapResult>& OutOverlaps, const class UPrimitiveComponent* PrimComp, const FVector& Pos, const FRotator& Rot, const FComponentQueryParams& Params = FComponentQueryParams::DefaultComponentQueryParams, const FCollisionObjectQueryParams& ObjectQueryParams=FCollisionObjectQueryParams::DefaultObjectQueryParam) const;
+
+	/**
+	 *  Test the collision of the supplied component at the supplied location/rotation using a specific channel, and determine the set of components that it overlaps
+	 *  @param  OutOverlaps     Array of overlaps found between component in specified pose and the world
+	 *  @param  PrimComp        Component to use geometry from to test against the world. Transform of this component is ignored
+	 *  @param  Pos             Location to place PrimComp geometry at to test against the world
+	 *  @param  Rot             Rotation to place PrimComp geometry at to test against the world
+	 *  @param  TraceChannel	The 'channel' that this query is in, used to determine which components to hit
+	 *  @return TRUE if OutOverlaps contains any blocking results
+	 */
+	DEPRECATED(4.8, "Use ComponentOverlapMultiByChannel instead.")
+	bool ComponentOverlapMulti(TArray<struct FOverlapResult>& OutOverlaps, const class UPrimitiveComponent* PrimComp, const FVector& Pos, const FRotator& Rot, ECollisionChannel TraceChannel, const FComponentQueryParams& Params, const FCollisionObjectQueryParams& ObjectQueryParams=FCollisionObjectQueryParams::DefaultObjectQueryParam) const
+	{
+		return ComponentOverlapMultiByChannel(OutOverlaps, PrimComp, Pos, Rot, TraceChannel, Params, ObjectQueryParams);
+	}
 
 	/**
 	 *  Test the collision of the supplied component at the supplied location/rotation using a specific channel, and determine the set of components that it overlaps
@@ -1425,7 +1678,7 @@ public:
 	 *  @param  TraceChannel	The 'channel' that this query is in, used to determine which components to hit
 	 *  @return TRUE if OutOverlaps contains any blocking results
 	 */
-	bool ComponentOverlapMulti(TArray<struct FOverlapResult>& OutOverlaps, const class UPrimitiveComponent* PrimComp, const FVector& Pos, const FRotator& Rot, ECollisionChannel TraceChannel, const struct FComponentQueryParams& Params, const struct FCollisionObjectQueryParams& ObjectQueryParams=FCollisionObjectQueryParams::DefaultObjectQueryParam) const;
+	bool ComponentOverlapMultiByChannel(TArray<struct FOverlapResult>& OutOverlaps, const class UPrimitiveComponent* PrimComp, const FVector& Pos, const FRotator& Rot, ECollisionChannel TraceChannel, const FComponentQueryParams& Params = FComponentQueryParams::DefaultComponentQueryParams, const FCollisionObjectQueryParams& ObjectQueryParams=FCollisionObjectQueryParams::DefaultObjectQueryParam) const;
 
 	/**
 	 * Interface for Async. Pretty much same parameter set except you can optional set delegate to be called when execution is completed and you can set UserData if you'd like
@@ -1447,7 +1700,33 @@ public:
 	 *	@param	UserData		UserData
 	 *	@param bMultiTrace		true if you'd like to get continuous result from the trace, false if you want single
 	 */ 
-	FTraceHandle	AsyncLineTrace(const FVector& Start,const FVector& End, ECollisionChannel TraceChannel, const struct FCollisionQueryParams& Params, const struct FCollisionResponseParams& ResponseParam = FCollisionResponseParams::DefaultResponseParam, FTraceDelegate * InDelegate=NULL, uint32 UserData = 0, bool bMultiTrace=false );
+	DEPRECATED(4.8, "Use AsyncLineTraceByChannel instead.")
+	FTraceHandle	AsyncLineTrace(const FVector& Start,const FVector& End, ECollisionChannel TraceChannel, const FCollisionQueryParams& Params, const FCollisionResponseParams& ResponseParam = FCollisionResponseParams::DefaultResponseParam, FTraceDelegate * InDelegate=NULL, uint32 UserData = 0, bool bMultiTrace=false )
+	{
+		return AsyncLineTraceByChannel(Start, End, TraceChannel, Params, ResponseParam, InDelegate, UserData, bMultiTrace);
+	}
+
+	/**
+	 * Interface for Async. Pretty much same parameter set except you can optional set delegate to be called when execution is completed and you can set UserData if you'd like
+	 * if no delegate, you can query trace data using QueryTraceData or QueryOverlapData
+	 * the data is available only in the next frame after request is made - in other words, if request is made in frame X, you can get the result in frame (X+1)
+	 *
+	 *  @param  Start           Start location of the ray
+	 *  @param  End             End location of the ray
+	 *  @param  TraceChannel    The 'channel' that this ray is in, used to determine which components to hit
+	 *  @param  Params          Additional parameters used for the trace
+	 * 	@param 	ResponseParam	ResponseContainer to be used for this trace
+	 *	@param	InDeleagte		Delegate function to be called - to see example, search FTraceDelegate
+	 *							Example can be void MyActor::TraceDone(const FTraceHandle& TraceHandle, FTraceDatum & TraceData)
+	 *							Before sending to the function, 
+	 *						
+	 *							FTraceDelegate TraceDelegate;
+	 *							TraceDelegate.BindRaw(this, &MyActor::TraceDone);
+	 * 
+	 *	@param	UserData		UserData
+	 *	@param bMultiTrace		true if you'd like to get continuous result from the trace, false if you want single
+	 */ 
+	FTraceHandle	AsyncLineTraceByChannel(const FVector& Start,const FVector& End, ECollisionChannel TraceChannel, const FCollisionQueryParams& Params = FCollisionQueryParams::DefaultQueryParam, const FCollisionResponseParams& ResponseParam = FCollisionResponseParams::DefaultResponseParam, FTraceDelegate * InDelegate=NULL, uint32 UserData = 0, bool bMultiTrace=false );
 
 	/**
 	 * Interface for Async. Pretty much same parameter set except you can optional set delegate to be called when execution is completed and you can set UserData if you'd like
@@ -1468,7 +1747,60 @@ public:
 	 *	@param	UserData		UserData
 	 *	@param bMultiTrace		true if you'd like to get continuous result from the trace, false if you want single
 	 */ 
-	FTraceHandle	AsyncLineTrace(const FVector& Start,const FVector& End, const struct FCollisionQueryParams& Params, const struct FCollisionObjectQueryParams& ObjectQueryParams, FTraceDelegate * InDelegate=NULL, uint32 UserData = 0, bool bMultiTrace=false );
+	DEPRECATED(4.8, "Use AsyncLineTraceByObjectType instead.")
+	FTraceHandle	AsyncLineTrace(const FVector& Start,const FVector& End, const FCollisionQueryParams& Params, const FCollisionObjectQueryParams& ObjectQueryParams, FTraceDelegate * InDelegate=NULL, uint32 UserData = 0, bool bMultiTrace=false )
+	{
+		return AsyncLineTraceByObjectType(Start, End, ObjectQueryParams, Params, InDelegate, UserData, bMultiTrace);
+	}
+
+	/**
+	 * Interface for Async. Pretty much same parameter set except you can optional set delegate to be called when execution is completed and you can set UserData if you'd like
+	 * if no delegate, you can query trace data using QueryTraceData or QueryOverlapData
+	 * the data is available only in the next frame after request is made - in other words, if request is made in frame X, you can get the result in frame (X+1)
+	 *
+	 *  @param  Start           Start location of the ray
+	 *  @param  End             End location of the ray
+	 *	@param	ObjectQueryParams	List of object types it's looking for
+	 *  @param  Params          Additional parameters used for the trace
+	 *	@param	InDeleagte		Delegate function to be called - to see example, search FTraceDelegate
+	 *							Example can be void MyActor::TraceDone(const FTraceHandle& TraceHandle, FTraceDatum & TraceData)
+	 *							Before sending to the function, 
+	 *						
+	 *							FTraceDelegate TraceDelegate;
+	 *							TraceDelegate.BindRaw(this, &MyActor::TraceDone);
+	 * 
+	 *	@param	UserData		UserData
+	 *	@param bMultiTrace		true if you'd like to get continuous result from the trace, false if you want single
+	 */ 
+	FTraceHandle	AsyncLineTraceByObjectType(const FVector& Start,const FVector& End, const FCollisionObjectQueryParams& ObjectQueryParams, const FCollisionQueryParams& Params = FCollisionQueryParams::DefaultQueryParam, FTraceDelegate * InDelegate=NULL, uint32 UserData = 0, bool bMultiTrace=false );
+
+	/**
+	 * Interface for Async trace
+	 * Pretty much same parameter set except you can optional set delegate to be called when execution is completed and you can set UserData if you'd like
+	 * if no delegate, you can query trace data using QueryTraceData or QueryOverlapData
+	 * the data is available only in the next frame after request is made - in other words, if request is made in frame X, you can get the result in frame (X+1)
+	 *
+	 *  @param  Start           Start location of the sphere
+	 *  @param  End             End location of the sphere
+	 *  @param  TraceChannel    The 'channel' that this trace is in, used to determine which components to hit
+	 *  @param	CollisionShape		CollisionShape - supports Box, Sphere, Capsule
+	 *  @param  Params          Additional parameters used for the trace
+	 * 	@param 	ResponseParam	ResponseContainer to be used for this trace	 
+	 *	@param	InDeleagte		Delegate function to be called - to see example, search FTraceDelegate
+	 *							Example can be void MyActor::TraceDone(const FTraceHandle& TraceHandle, FTraceDatum & TraceData)
+	 *							Before sending to the function, 
+	 *						
+	 *							FTraceDelegate TraceDelegate;
+	 *							TraceDelegate.BindRaw(this, &MyActor::TraceDone);
+	 * 
+	 *	@param	UserData		UserData
+	 *	@param bMultiTrace		true if you'd like to get continuous result from the trace, false if you want single
+	 */ 
+	DEPRECATED(4.8, "Use AsyncSweepByChannel instead.")
+	FTraceHandle	AsyncSweep(const FVector& Start, const FVector& End, ECollisionChannel TraceChannel, const FCollisionShape& CollisionShape, const FCollisionQueryParams& Params, const FCollisionResponseParams& ResponseParam = FCollisionResponseParams::DefaultResponseParam, FTraceDelegate * InDelegate = NULL, uint32 UserData = 0, bool bMultiTrace = false)
+	{
+		return AsyncSweepByChannel(Start, End, TraceChannel, CollisionShape, Params, ResponseParam, InDelegate, UserData, bMultiTrace);
+	}
 
 	/**
 	 * Interface for Async trace
@@ -1492,7 +1824,7 @@ public:
 	 *	@param	UserData		UserData
 	 *	@param bMultiTrace		true if you'd like to get continuous result from the trace, false if you want single
 	 */ 
-	FTraceHandle	AsyncSweep(const FVector& Start, const FVector& End, ECollisionChannel TraceChannel, const struct FCollisionShape& CollisionShape, const struct FCollisionQueryParams& Params, const struct FCollisionResponseParams& ResponseParam = FCollisionResponseParams::DefaultResponseParam, FTraceDelegate * InDelegate = NULL, uint32 UserData = 0, bool bMultiTrace = false);
+	FTraceHandle	AsyncSweepByChannel(const FVector& Start, const FVector& End, ECollisionChannel TraceChannel, const FCollisionShape& CollisionShape, const FCollisionQueryParams& Params = FCollisionQueryParams::DefaultQueryParam, const FCollisionResponseParams& ResponseParam = FCollisionResponseParams::DefaultResponseParam, FTraceDelegate * InDelegate = NULL, uint32 UserData = 0, bool bMultiTrace = false);
 
 	/**
 	 * Interface for Async trace
@@ -1515,7 +1847,36 @@ public:
 	 *	@param	UserData		UserData
 	 *	@param bMultiTrace		true if you'd like to get continuous result from the trace, false if you want single
 	 */ 
-	FTraceHandle	AsyncSweep(const FVector& Start, const FVector& End, const struct FCollisionShape& CollisionShape, const struct FCollisionQueryParams& Params, const struct FCollisionObjectQueryParams& ObjectQueryParams, FTraceDelegate * InDelegate = NULL, uint32 UserData = 0, bool bMultiTrace = false);
+	DEPRECATED(4.8, "Use AsyncSweepByObjectType instead.")
+	FTraceHandle	AsyncSweep(const FVector& Start, const FVector& End, const FCollisionShape& CollisionShape, const FCollisionQueryParams& Params, const FCollisionObjectQueryParams& ObjectQueryParams, FTraceDelegate * InDelegate = NULL, uint32 UserData = 0, bool bMultiTrace = false)
+	{
+		return AsyncSweepByObjectType(Start, End, ObjectQueryParams, CollisionShape, Params, InDelegate, UserData, bMultiTrace);
+	}
+	
+	
+	/**
+	 * Interface for Async trace
+	 * Pretty much same parameter set except you can optional set delegate to be called when execution is completed and you can set UserData if you'd like
+	 * if no delegate, you can query trace data using QueryTraceData or QueryOverlapData
+	 * the data is available only in the next frame after request is made - in other words, if request is made in frame X, you can get the result in frame (X+1)
+	 *
+	 *  @param  Start           Start location of the sphere
+	 *  @param  End             End location of the sphere
+	 *	@param	ObjectQueryParams	List of object types it's looking for
+	 *  @param	CollisionShape		CollisionShape - supports Box, Sphere, Capsule
+	 *  @param  Params          Additional parameters used for the trace
+	 *	@param	InDeleagte		Delegate function to be called - to see example, search FTraceDelegate
+	 *							Example can be void MyActor::TraceDone(const FTraceHandle& TraceHandle, FTraceDatum & TraceData)
+	 *							Before sending to the function, 
+	 *						
+	 *							FTraceDelegate TraceDelegate;
+	 *							TraceDelegate.BindRaw(this, &MyActor::TraceDone);
+	 * 
+	 *	@param	UserData		UserData
+	 *	@param bMultiTrace		true if you'd like to get continuous result from the trace, false if you want single
+	 */ 
+	FTraceHandle	AsyncSweepByObjectType(const FVector& Start, const FVector& End, const FCollisionObjectQueryParams& ObjectQueryParams, const FCollisionShape& CollisionShape, const FCollisionQueryParams& Params = FCollisionQueryParams::DefaultQueryParam, FTraceDelegate * InDelegate = NULL, uint32 UserData = 0, bool bMultiTrace = false);
+
 	// overlap functions
 
 	/**
@@ -1539,7 +1900,11 @@ public:
 	 * 
 	 *	@param UserData			UserData
 	 */ 
-	FTraceHandle	AsyncOverlap(const FVector& Pos, const FQuat& Rot, ECollisionChannel TraceChannel, const struct FCollisionShape& CollisionShape, const struct FCollisionQueryParams& Params, const struct FCollisionResponseParams& ResponseParam = FCollisionResponseParams::DefaultResponseParam, FOverlapDelegate * InDelegate = NULL, uint32 UserData = 0);
+	DEPRECATED(4.8, "Use AsyncOverlapByChannel instead.")
+	FTraceHandle	AsyncOverlap(const FVector& Pos, const FQuat& Rot, ECollisionChannel TraceChannel, const FCollisionShape& CollisionShape, const FCollisionQueryParams& Params, const FCollisionResponseParams& ResponseParam = FCollisionResponseParams::DefaultResponseParam, FOverlapDelegate * InDelegate = NULL, uint32 UserData = 0)
+	{
+		return AsyncOverlapByChannel(Pos, Rot, TraceChannel, CollisionShape, Params, ResponseParam, InDelegate, UserData);
+	}
 
 	/**
 	 * Interface for Async trace
@@ -1548,6 +1913,7 @@ public:
 	 * the data is available only in the next frame after request is made - in other words, if request is made in frame X, you can get the result in frame (X+1)
 	 *
 	 *  @param  Pos             Location of center of sphere to test against the world
+	 *	@param	bMultiTrace		true if you'd like to do multi trace, or false otherwise
 	 *  @param  TraceChannel    The 'channel' that this query is in, used to determine which components to hit
 	 *  @param	CollisionShape		CollisionShape - supports Box, Sphere, Capsule
 	 *  @param  Params          Additional parameters used for the trace
@@ -1561,7 +1927,54 @@ public:
 	 * 
 	 *	@param UserData			UserData
 	 */ 
-	FTraceHandle	AsyncOverlap(const FVector& Pos, const FQuat& Rot, const struct FCollisionShape& CollisionShape, const struct FCollisionQueryParams& Params, const struct FCollisionObjectQueryParams& ObjectQueryParams, FOverlapDelegate * InDelegate = NULL, uint32 UserData = 0);
+	FTraceHandle	AsyncOverlapByChannel(const FVector& Pos, const FQuat& Rot, ECollisionChannel TraceChannel, const FCollisionShape& CollisionShape, const FCollisionQueryParams& Params = FCollisionQueryParams::DefaultQueryParam, const FCollisionResponseParams& ResponseParam = FCollisionResponseParams::DefaultResponseParam, FOverlapDelegate * InDelegate = NULL, uint32 UserData = 0);
+
+	/**
+	 * Interface for Async trace
+	 * Pretty much same parameter set except you can optional set delegate to be called when execution is completed and you can set UserData if you'd like
+	 * if no delegate, you can query trace data using QueryTraceData or QueryOverlapData
+	 * the data is available only in the next frame after request is made - in other words, if request is made in frame X, you can get the result in frame (X+1)
+	 *
+	 *  @param  Pos             Location of center of sphere to test against the world
+	 *  @param  TraceChannel    The 'channel' that this query is in, used to determine which components to hit
+	 *  @param	CollisionShape		CollisionShape - supports Box, Sphere, Capsule
+	 *  @param  Params          Additional parameters used for the trace
+	 * 	@param 	ResponseParam	ResponseContainer to be used for this trace
+	 *	@param	InDeleagte		Delegate function to be called - to see example, search FTraceDelegate
+	 *							Example can be void MyActor::TraceDone(const FTraceHandle& TraceHandle, FTraceDatum & TraceData)
+	 *							Before sending to the function, 
+	 *						
+	 *							FTraceDelegate TraceDelegate;
+	 *							TraceDelegate.BindRaw(this, &MyActor::TraceDone);
+	 * 
+	 *	@param UserData			UserData
+	 */ 
+	DEPRECATED(4.8, "Use AsyncOverlapByObjectType instead.")
+	FTraceHandle	AsyncOverlap(const FVector& Pos, const FQuat& Rot, const FCollisionShape& CollisionShape, const FCollisionQueryParams& Params, const FCollisionObjectQueryParams& ObjectQueryParams, FOverlapDelegate * InDelegate = NULL, uint32 UserData = 0)
+	{
+		return AsyncOverlapByObjectType(Pos, Rot, ObjectQueryParams, CollisionShape, Params, InDelegate, UserData);
+	}
+
+	/**
+	 * Interface for Async trace
+	 * Pretty much same parameter set except you can optional set delegate to be called when execution is completed and you can set UserData if you'd like
+	 * if no delegate, you can query trace data using QueryTraceData or QueryOverlapData
+	 * the data is available only in the next frame after request is made - in other words, if request is made in frame X, you can get the result in frame (X+1)
+	 *
+	 *  @param  Pos             Location of center of sphere to test against the world
+	 *	@param	ObjectQueryParams	List of object types it's looking for
+	 *  @param	CollisionShape		CollisionShape - supports Box, Sphere, Capsule
+	 *  @param  Params          Additional parameters used for the trace
+	 *	@param	InDeleagte		Delegate function to be called - to see example, search FTraceDelegate
+	 *							Example can be void MyActor::TraceDone(const FTraceHandle& TraceHandle, FTraceDatum & TraceData)
+	 *							Before sending to the function, 
+	 *						
+	 *							FTraceDelegate TraceDelegate;
+	 *							TraceDelegate.BindRaw(this, &MyActor::TraceDone);
+	 * 
+	 *	@param UserData			UserData
+	 */ 
+	FTraceHandle	AsyncOverlapByObjectType(const FVector& Pos, const FQuat& Rot, const FCollisionObjectQueryParams& ObjectQueryParams, const FCollisionShape& CollisionShape, const FCollisionQueryParams& Params = FCollisionQueryParams::DefaultQueryParam, FOverlapDelegate * InDelegate = NULL, uint32 UserData = 0);
 
 	/**
 	 * Query function 
diff --git a/Engine/Source/Runtime/Engine/Private/AI/Navigation/NavLinkCustomComponent.cpp b/Engine/Source/Runtime/Engine/Private/AI/Navigation/NavLinkCustomComponent.cpp
--- a/Engine/Source/Runtime/Engine/Private/AI/Navigation/NavLinkCustomComponent.cpp
+++ b/Engine/Source/Runtime/Engine/Private/AI/Navigation/NavLinkCustomComponent.cpp
@@ -272,13 +272,13 @@ void UNavLinkCustomComponent::CollectNearbyAgents(TArray<UPathFollowingComponent
 	const float DistThresholdSq = FMath::Square(BroadcastRadius * 0.25f);
 	if (LinkDistSq > DistThresholdSq)
 	{
-		GetWorld()->OverlapMulti(OverlapsL, LocationL, FQuat::Identity, BroadcastChannel, FCollisionShape::MakeSphere(BroadcastRadius), Params);
-		GetWorld()->OverlapMulti(OverlapsR, LocationR, FQuat::Identity, BroadcastChannel, FCollisionShape::MakeSphere(BroadcastRadius), Params);
+		GetWorld()->OverlapMultiByChannel(OverlapsL, LocationL, FQuat::Identity, BroadcastChannel, FCollisionShape::MakeSphere(BroadcastRadius), Params);
+		GetWorld()->OverlapMultiByChannel(OverlapsR, LocationR, FQuat::Identity, BroadcastChannel, FCollisionShape::MakeSphere(BroadcastRadius), Params);
 	}
 	else
 	{
 		const FVector MidPoint = (LocationL + LocationR) * 0.5f;
-		GetWorld()->OverlapMulti(OverlapsL, MidPoint, FQuat::Identity, BroadcastChannel, FCollisionShape::MakeSphere(BroadcastRadius), Params);
+		GetWorld()->OverlapMultiByChannel(OverlapsL, MidPoint, FQuat::Identity, BroadcastChannel, FCollisionShape::MakeSphere(BroadcastRadius), Params);
 	}
 
 	TArray<APawn*> PawnList;
diff --git a/Engine/Source/Runtime/Engine/Private/AI/Navigation/NavigationSystemHelpers.cpp b/Engine/Source/Runtime/Engine/Private/AI/Navigation/NavigationSystemHelpers.cpp
--- a/Engine/Source/Runtime/Engine/Private/AI/Navigation/NavigationSystemHelpers.cpp
+++ b/Engine/Source/Runtime/Engine/Private/AI/Navigation/NavigationSystemHelpers.cpp
@@ -46,7 +46,7 @@ namespace NavigationHelper
 
 		FCollisionQueryParams TraceParams(NAME_None, true, Querier);
 		FHitResult Hit;
-		const bool bHit = Querier->GetWorld()->LineTraceSingle(Hit, FallStart, FallStart+FVector(0,0,-FallLimit), TraceParams, FCollisionObjectQueryParams(ECC_WorldStatic));
+		const bool bHit = Querier->GetWorld()->LineTraceSingleByObjectType(Hit, FallStart, FallStart+FVector(0,0,-FallLimit), FCollisionObjectQueryParams(ECC_WorldStatic), TraceParams);
 		if( bHit )
 		{
 			UE_VLOG_LOCATION(Querier, LogNavigation, Log, Hit.Location, 15, FColor::Red, TEXT("%s")
diff --git a/Engine/Source/Runtime/Engine/Private/ActiveSound.cpp b/Engine/Source/Runtime/Engine/Private/ActiveSound.cpp
--- a/Engine/Source/Runtime/Engine/Private/ActiveSound.cpp
+++ b/Engine/Source/Runtime/Engine/Private/ActiveSound.cpp
@@ -286,7 +286,7 @@ void FActiveSound::CheckOcclusion( const FVector ListenerLocation, const FVector
 	{
 		LastOcclusionCheckTime = WorldPtr->GetTimeSeconds();
 		static FName NAME_SoundOcclusion = FName(TEXT("SoundOcclusion"));
-		const bool bNowOccluded = WorldPtr->LineTraceTest(SoundLocation, ListenerLocation, ECC_Visibility, FCollisionQueryParams(NAME_SoundOcclusion, true));
+		const bool bNowOccluded = WorldPtr->LineTraceTestByChannel(SoundLocation, ListenerLocation, ECC_Visibility, FCollisionQueryParams(NAME_SoundOcclusion, true));
 		if( bNowOccluded != bOccluded )
 		{
 			bOccluded = bNowOccluded;
diff --git a/Engine/Source/Runtime/Engine/Private/CharacterMovementComponent.cpp b/Engine/Source/Runtime/Engine/Private/CharacterMovementComponent.cpp
--- a/Engine/Source/Runtime/Engine/Private/CharacterMovementComponent.cpp
+++ b/Engine/Source/Runtime/Engine/Private/CharacterMovementComponent.cpp
@@ -1716,7 +1716,7 @@ void UCharacterMovementComponent::Crouch(bool bClientSimulation)
 			FCollisionQueryParams CapsuleParams(NAME_CrouchTrace, false, CharacterOwner);
 			FCollisionResponseParams ResponseParam;
 			InitCollisionParams(CapsuleParams, ResponseParam);
-			const bool bEncroached = GetWorld()->OverlapBlockingTest(CharacterOwner->GetActorLocation() - FVector(0.f,0.f,ScaledHalfHeightAdjust), FQuat::Identity,
+			const bool bEncroached = GetWorld()->OverlapBlockingTestByChannel(CharacterOwner->GetActorLocation() - FVector(0.f,0.f,ScaledHalfHeightAdjust), FQuat::Identity,
 				UpdatedComponent->GetCollisionObjectType(), GetPawnCapsuleCollisionShape(SHRINK_None), CapsuleParams, ResponseParam);
 
 			// If encroached, cancel
@@ -1792,7 +1792,7 @@ void UCharacterMovementComponent::UnCrouch(bool bClientSimulation)
 		if (!bCrouchMaintainsBaseLocation)
 		{
 			// Expand in place
-			bEncroached = GetWorld()->OverlapBlockingTest(PawnLocation, FQuat::Identity, CollisionChannel, StandingCapsuleShape, CapsuleParams, ResponseParam);
+			bEncroached = GetWorld()->OverlapBlockingTestByChannel(PawnLocation, FQuat::Identity, CollisionChannel, StandingCapsuleShape, CapsuleParams, ResponseParam);
 		
 			if (bEncroached)
 			{
@@ -1808,7 +1808,7 @@ void UCharacterMovementComponent::UnCrouch(bool bClientSimulation)
 
 					FHitResult Hit(1.f);
 					const FCollisionShape ShortCapsuleShape = GetPawnCapsuleCollisionShape(SHRINK_HeightCustom, ShrinkHalfHeight);
-					const bool bBlockingHit = GetWorld()->SweepSingle(Hit, PawnLocation, PawnLocation + Down, FQuat::Identity, CollisionChannel, ShortCapsuleShape, CapsuleParams);
+					const bool bBlockingHit = GetWorld()->SweepSingleByChannel(Hit, PawnLocation, PawnLocation + Down, FQuat::Identity, CollisionChannel, ShortCapsuleShape, CapsuleParams);
 					if (Hit.bStartPenetrating)
 					{
 						bEncroached = true;
@@ -1818,7 +1818,7 @@ void UCharacterMovementComponent::UnCrouch(bool bClientSimulation)
 						// Compute where the base of the sweep ended up, and see if we can stand there
 						const float DistanceToBase = (Hit.Time * TraceDist) + ShortCapsuleShape.Capsule.HalfHeight;
 						const FVector NewLoc = FVector(PawnLocation.X, PawnLocation.Y, PawnLocation.Z - DistanceToBase + PawnHalfHeight + SweepInflation + MIN_FLOOR_DIST / 2.f);
-						bEncroached = GetWorld()->OverlapBlockingTest(NewLoc, FQuat::Identity, CollisionChannel, StandingCapsuleShape, CapsuleParams, ResponseParam);
+						bEncroached = GetWorld()->OverlapBlockingTestByChannel(NewLoc, FQuat::Identity, CollisionChannel, StandingCapsuleShape, CapsuleParams, ResponseParam);
 						if (!bEncroached)
 						{
 							// Intentionally not using MoveUpdatedComponent, where a horizontal plane constraint would prevent the base of the capsule from staying at the same spot.
@@ -1832,7 +1832,7 @@ void UCharacterMovementComponent::UnCrouch(bool bClientSimulation)
 		{
 			// Expand while keeping base location the same.
 			FVector StandingLocation = PawnLocation + FVector(0.f, 0.f, StandingCapsuleShape.GetCapsuleHalfHeight() - CurrentCrouchedHalfHeight);
-			bEncroached = GetWorld()->OverlapBlockingTest(StandingLocation, FQuat::Identity, CollisionChannel, StandingCapsuleShape, CapsuleParams, ResponseParam);
+			bEncroached = GetWorld()->OverlapBlockingTestByChannel(StandingLocation, FQuat::Identity, CollisionChannel, StandingCapsuleShape, CapsuleParams, ResponseParam);
 
 			if (bEncroached)
 			{
@@ -1843,7 +1843,7 @@ void UCharacterMovementComponent::UnCrouch(bool bClientSimulation)
 					if (CurrentFloor.bBlockingHit && CurrentFloor.FloorDist > MinFloorDist)
 					{
 						StandingLocation.Z -= CurrentFloor.FloorDist - MinFloorDist;
-						bEncroached = GetWorld()->OverlapBlockingTest(StandingLocation, FQuat::Identity, CollisionChannel, StandingCapsuleShape, CapsuleParams, ResponseParam);
+						bEncroached = GetWorld()->OverlapBlockingTestByChannel(StandingLocation, FQuat::Identity, CollisionChannel, StandingCapsuleShape, CapsuleParams, ResponseParam);
 					}
 				}				
 			}
@@ -2892,7 +2892,7 @@ FVector UCharacterMovementComponent::FindWaterLine(FVector InWater, FVector Outo
 
 	TArray<FHitResult> Hits;
 	static const FName NAME_FindWaterLine = FName(TEXT("FindWaterLine"));
-	GetWorld()->LineTraceMulti(Hits, OutofWater, InWater, UpdatedComponent->GetCollisionObjectType(), FCollisionQueryParams(NAME_FindWaterLine, true, CharacterOwner));
+	GetWorld()->LineTraceMultiByChannel(Hits, OutofWater, InWater, UpdatedComponent->GetCollisionObjectType(), FCollisionQueryParams(NAME_FindWaterLine, true, CharacterOwner));
 
 	for( int32 HitIdx = 0; HitIdx < Hits.Num(); HitIdx++ )
 	{
@@ -3235,7 +3235,7 @@ bool UCharacterMovementComponent::FindAirControlImpact(float DeltaTime, float Ad
 		const FVector CapsuleLocation = UpdatedComponent->GetComponentLocation();
 		const FCollisionShape CapsuleShape = GetPawnCapsuleCollisionShape(SHRINK_None);
 		
-		if (GetWorld()->SweepSingle(OutHitResult, CapsuleLocation, CapsuleLocation + TestWalk, FQuat::Identity, UpdatedComponent->GetCollisionObjectType(), CapsuleShape, CapsuleQuery, ResponseParam))
+		if (GetWorld()->SweepSingleByChannel(OutHitResult, CapsuleLocation, CapsuleLocation + TestWalk, FQuat::Identity, UpdatedComponent->GetCollisionObjectType(), CapsuleShape, CapsuleQuery, ResponseParam))
 		{
 			return true;
 		}
@@ -3281,13 +3281,13 @@ bool UCharacterMovementComponent::CheckLedgeDirection(const FVector& OldLocation
 	FCollisionShape CapsuleShape = GetPawnCapsuleCollisionShape(SHRINK_None);
 	const ECollisionChannel CollisionChannel = UpdatedComponent->GetCollisionObjectType();
 	FHitResult Result(1.f);
-	GetWorld()->SweepSingle(Result, OldLocation, SideDest, FQuat::Identity, CollisionChannel, CapsuleShape, CapsuleParams, ResponseParam);
+	GetWorld()->SweepSingleByChannel(Result, OldLocation, SideDest, FQuat::Identity, CollisionChannel, CapsuleShape, CapsuleParams, ResponseParam);
 
 	if ( !Result.bBlockingHit || IsWalkable(Result) )
 	{
 		if ( !Result.bBlockingHit )
 		{
-			GetWorld()->SweepSingle(Result, SideDest, SideDest + GravDir * (MaxStepHeight + LedgeCheckThreshold), FQuat::Identity, CollisionChannel, CapsuleShape, CapsuleParams, ResponseParam);
+			GetWorld()->SweepSingleByChannel(Result, SideDest, SideDest + GravDir * (MaxStepHeight + LedgeCheckThreshold), FQuat::Identity, CollisionChannel, CapsuleShape, CapsuleParams, ResponseParam);
 		}
 		if ( (Result.Time < 1.f) && IsWalkable(Result) )
 		{
@@ -3908,7 +3908,7 @@ void UCharacterMovementComponent::ProjectLocationFromNavMesh(float DeltaSeconds,
 			// influence navmesh generation
 			FCollisionQueryParams Params;
 			FCollisionObjectQueryParams ObjectQueryParams(FCollisionObjectQueryParams::AllStaticObjects);
-			GetWorld()->LineTraceSingle(CachedProjectedNavMeshHitResult, InOutLocation + RayCastOffset, InOutLocation - RayCastOffset, Params, ObjectQueryParams);
+			GetWorld()->LineTraceSingleByObjectType(CachedProjectedNavMeshHitResult, InOutLocation + RayCastOffset, InOutLocation - RayCastOffset, ObjectQueryParams, Params);
 
 			// discard result if we were already inside something
 			if (CachedProjectedNavMeshHitResult.bStartPenetrating)
@@ -4363,7 +4363,7 @@ bool UCharacterMovementComponent::CheckWaterJump(FVector CheckPoint, FVector& Wa
 	InitCollisionParams(CapsuleParams, ResponseParam);
 	FCollisionShape CapsuleShape = GetPawnCapsuleCollisionShape(SHRINK_None);
 	const ECollisionChannel CollisionChannel = UpdatedComponent->GetCollisionObjectType();
-	bool bHit = GetWorld()->SweepSingle( HitInfo, CharacterOwner->GetActorLocation(), CheckPoint, FQuat::Identity, CollisionChannel, CapsuleShape, CapsuleParams, ResponseParam);
+	bool bHit = GetWorld()->SweepSingleByChannel( HitInfo, CharacterOwner->GetActorLocation(), CheckPoint, FQuat::Identity, CollisionChannel, CapsuleShape, CapsuleParams, ResponseParam);
 	
 	if ( bHit && !Cast<APawn>(HitInfo.GetActor()) )
 	{
@@ -4375,7 +4375,7 @@ bool UCharacterMovementComponent::CheckWaterJump(FVector CheckPoint, FVector& Wa
 		FCollisionQueryParams LineParams(CheckWaterJumpName, true, CharacterOwner);
 		FCollisionResponseParams LineResponseParam;
 		InitCollisionParams(LineParams, LineResponseParam);
-		bHit = GetWorld()->LineTraceSingle( HitInfo, Start, CheckPoint, CollisionChannel, LineParams, LineResponseParam );
+		bHit = GetWorld()->LineTraceSingleByChannel( HitInfo, Start, CheckPoint, CollisionChannel, LineParams, LineResponseParam );
 		// if no high obstruction, or it's a valid floor, then pawn can jump out of water
 		return !bHit || IsWalkable(HitInfo);
 	}
@@ -4672,7 +4672,7 @@ void UCharacterMovementComponent::ComputeFloorDist(const FVector& CapsuleLocatio
 		QueryParams.TraceTag = FloorLineTraceName;
 
 		FHitResult Hit(1.f);
-		bBlockingHit = GetWorld()->LineTraceSingle(Hit, LineTraceStart, LineTraceStart + Down, CollisionChannel, QueryParams, ResponseParam);
+		bBlockingHit = GetWorld()->LineTraceSingleByChannel(Hit, LineTraceStart, LineTraceStart + Down, CollisionChannel, QueryParams, ResponseParam);
 
 		if (bBlockingHit)
 		{
@@ -4809,7 +4809,7 @@ bool UCharacterMovementComponent::FloorSweepTest(
 
 	if (!bUseFlatBaseForFloorChecks)
 	{
-		bBlockingHit = GetWorld()->SweepSingle(OutHit, Start, End, FQuat::Identity, TraceChannel, CollisionShape, Params, ResponseParam);
+		bBlockingHit = GetWorld()->SweepSingleByChannel(OutHit, Start, End, FQuat::Identity, TraceChannel, CollisionShape, Params, ResponseParam);
 	}
 	else
 	{
@@ -4819,13 +4819,13 @@ bool UCharacterMovementComponent::FloorSweepTest(
 		const FCollisionShape BoxShape = FCollisionShape::MakeBox(FVector(CapsuleRadius * 0.707f, CapsuleRadius * 0.707f, CapsuleHeight));
 
 		// First test with the box rotated so the corners are along the major axes (ie rotated 45 degrees).
-		bBlockingHit = GetWorld()->SweepSingle(OutHit, Start, End, FQuat(FVector(0.f, 0.f, -1.f), PI * 0.25f), TraceChannel, BoxShape, Params, ResponseParam);
+		bBlockingHit = GetWorld()->SweepSingleByChannel(OutHit, Start, End, FQuat(FVector(0.f, 0.f, -1.f), PI * 0.25f), TraceChannel, BoxShape, Params, ResponseParam);
 
 		if (!bBlockingHit)
 		{
 			// Test again with the same box, not rotated.
 			OutHit.Reset(1.f, false);			
-			bBlockingHit = GetWorld()->SweepSingle(OutHit, Start, End, FQuat::Identity, TraceChannel, BoxShape, Params, ResponseParam);
+			bBlockingHit = GetWorld()->SweepSingleByChannel(OutHit, Start, End, FQuat::Identity, TraceChannel, BoxShape, Params, ResponseParam);
 		}
 	}
 
diff --git a/Engine/Source/Runtime/Engine/Private/CheatManager.cpp b/Engine/Source/Runtime/Engine/Private/CheatManager.cpp
--- a/Engine/Source/Runtime/Engine/Private/CheatManager.cpp
+++ b/Engine/Source/Runtime/Engine/Private/CheatManager.cpp
@@ -57,7 +57,7 @@ void UCheatManager::Teleport()
 	static FName NAME_TeleportTrace = FName(TEXT("TeleportTrace"));
 	FCollisionQueryParams TraceParams(NAME_TeleportTrace, true, AssociatedPawn);
 
-	bool bHit = GetWorld()->LineTraceSingle(Hit, ViewLocation, ViewLocation + 1000000.f * ViewRotation.Vector(), ECC_Pawn, TraceParams);
+	bool bHit = GetWorld()->LineTraceSingleByChannel(Hit, ViewLocation, ViewLocation + 1000000.f * ViewRotation.Vector(), ECC_Pawn, TraceParams);
 	if ( bHit )
 	{
 		Hit.Location += Hit.Normal * 4.0f;
@@ -197,7 +197,7 @@ void UCheatManager::DamageTarget(float DamageAmount)
 
 	FCollisionQueryParams TraceParams(NAME_None, true, MyPC->GetPawn());
 	FHitResult Hit;
-	bool bHit = GetWorld()->LineTraceSingle(Hit, CamLoc, CamRot.Vector() * 100000.f + CamLoc, ECC_Pawn, TraceParams);
+	bool bHit = GetWorld()->LineTraceSingleByChannel(Hit, CamLoc, CamRot.Vector() * 100000.f + CamLoc, ECC_Pawn, TraceParams);
 	if (bHit)
 	{
 		check(Hit.GetActor() != NULL);
@@ -223,7 +223,7 @@ void UCheatManager::DestroyTarget()
 
 	FCollisionQueryParams TraceParams(NAME_None, true, MyPC->GetPawn());
 	FHitResult Hit;
-	bool bHit = GetWorld()->LineTraceSingle(Hit, CamLoc, CamRot.Vector() * 100000.f + CamLoc, ECC_Pawn, TraceParams);
+	bool bHit = GetWorld()->LineTraceSingleByChannel(Hit, CamLoc, CamRot.Vector() * 100000.f + CamLoc, ECC_Pawn, TraceParams);
 	if (bHit)
 	{
 		check(Hit.GetActor() != NULL);
@@ -278,7 +278,7 @@ void UCheatManager::DestroyAllPawnsExceptTarget()
 	FCollisionQueryParams TraceParams(NAME_None, true, MyPC->GetPawn());
 	FHitResult Hit;
 	APawn* HitPawnTarget = NULL;
-	bool bHit = GetWorld()->LineTraceSingle(Hit, CamLoc, CamRot.Vector() * 100000.f + CamLoc, ECC_Pawn, TraceParams);
+	bool bHit = GetWorld()->LineTraceSingleByChannel(Hit, CamLoc, CamRot.Vector() * 100000.f + CamLoc, ECC_Pawn, TraceParams);
 	if (bHit)
 	{
 		check(Hit.GetActor() != NULL);
@@ -713,7 +713,7 @@ void UCheatManager::TickCollisionDebug()
 			{
 				// If we get a hit, draw the capsule
 				FHitResult Result;
-				bool bHit = GetWorld()->SweepSingle(Result, ViewLoc, End, FQuat::Identity, DebugTraceChannel, FCollisionShape::MakeCapsule(DebugCapsuleRadius, DebugCapsuleHalfHeight), CapsuleParams);
+				bool bHit = GetWorld()->SweepSingleByChannel(Result, ViewLoc, End, FQuat::Identity, DebugTraceChannel, FCollisionShape::MakeCapsule(DebugCapsuleRadius, DebugCapsuleHalfHeight), CapsuleParams);
 				if(bHit)
 				{
 					AddCapsuleSweepDebugInfo(ViewLoc, End, Result.ImpactPoint, Result.Normal, Result.ImpactNormal, Result.Location, DebugCapsuleHalfHeight, DebugCapsuleRadius, false, (Result.bStartPenetrating && Result.bBlockingHit)? true: false);
diff --git a/Engine/Source/Runtime/Engine/Private/Collision/WorldCollision.cpp b/Engine/Source/Runtime/Engine/Private/Collision/WorldCollision.cpp
--- a/Engine/Source/Runtime/Engine/Private/Collision/WorldCollision.cpp
+++ b/Engine/Source/Runtime/Engine/Private/Collision/WorldCollision.cpp
@@ -46,6 +46,8 @@ FCollisionResponseContainer FCollisionResponseContainer::DefaultResponseContaine
 /* This is default response param that's used by trace query **/
 FCollisionResponseParams		FCollisionResponseParams::DefaultResponseParam;
 FCollisionObjectQueryParams		FCollisionObjectQueryParams::DefaultObjectQueryParam;
+FCollisionQueryParams			FCollisionQueryParams::DefaultQueryParam(TEXT("DefaultQueryParam"));
+FComponentQueryParams			FComponentQueryParams::DefaultComponentQueryParams(TEXT("DefaultComponentQueryParam"));
 FCollisionShape					FCollisionShape::LineShape;
 
 // default being the 0. That isn't invalid, but ObjectQuery param overrides this 
@@ -70,16 +72,17 @@ void FBaseTraceDatum::Set(UWorld * World, const FCollisionShape& InCollisionShap
 
 //////////////////////////////////////////////////////////////////////////
 
-bool UWorld::LineTraceTest(const FVector& Start, const FVector& End, ECollisionChannel TraceChannel, const struct FCollisionQueryParams& Params, const struct FCollisionResponseParams& ResponseParam) const
+bool UWorld::LineTraceTestByChannel(const FVector& Start,const FVector& End,ECollisionChannel TraceChannel, const FCollisionQueryParams& Params /* = FCollisionQueryParams::DefaultQueryParam */, const FCollisionResponseParams& ResponseParam /* = FCollisionResponseParams::DefaultResponseParam */) const
 {
 #if UE_WITH_PHYSICS
 	return RaycastTest(this, Start, End, TraceChannel, Params, ResponseParam, FCollisionObjectQueryParams::DefaultObjectQueryParam);
 #else
 	return false;
 #endif
+
 }
 
-bool UWorld::LineTraceSingle(struct FHitResult& OutHit, const FVector& Start, const FVector& End, ECollisionChannel TraceChannel, const struct FCollisionQueryParams& Params, const struct FCollisionResponseParams& ResponseParam) const
+bool UWorld::LineTraceSingleByChannel(struct FHitResult& OutHit,const FVector& Start,const FVector& End,ECollisionChannel TraceChannel,const FCollisionQueryParams& Params /* = FCollisionQueryParams::DefaultQueryParam */, const FCollisionResponseParams& ResponseParam /* = FCollisionResponseParams::DefaultResponseParam */) const
 {
 #if UE_WITH_PHYSICS
 	return RaycastSingle(this, OutHit, Start, End, TraceChannel, Params, ResponseParam, FCollisionObjectQueryParams::DefaultObjectQueryParam);
@@ -90,7 +93,7 @@ bool UWorld::LineTraceSingle(struct FHitResult& OutHit, const FVector& Start, co
 #endif
 }
 
-bool UWorld::LineTraceMulti(TArray<struct FHitResult>& OutHits, const FVector& Start, const FVector& End, ECollisionChannel TraceChannel, const struct FCollisionQueryParams& Params, const struct FCollisionResponseParams& ResponseParam) const
+bool UWorld::LineTraceMultiByChannel(TArray<struct FHitResult>& OutHits,const FVector& Start,const FVector& End,ECollisionChannel TraceChannel,const FCollisionQueryParams& Params /* = FCollisionQueryParams::DefaultQueryParam */, const FCollisionResponseParams& ResponseParam /* = FCollisionResponseParams::DefaultResponseParam */) const
 {
 #if UE_WITH_PHYSICS
 	return RaycastMulti(this, OutHits, Start, End, TraceChannel, Params, ResponseParam, FCollisionObjectQueryParams::DefaultObjectQueryParam);
@@ -99,12 +102,12 @@ bool UWorld::LineTraceMulti(TArray<struct FHitResult>& OutHits, const FVector& S
 #endif
 }
 
-bool UWorld::SweepTest(const FVector& Start, const FVector& End, const FQuat& Rot, ECollisionChannel TraceChannel, const struct FCollisionShape& CollisionShape, const struct FCollisionQueryParams& Params, const struct FCollisionResponseParams& ResponseParam) const
+bool UWorld::SweepTestByChannel(const FVector& Start, const FVector& End, const FQuat& Rot, ECollisionChannel TraceChannel, const FCollisionShape& CollisionShape, const FCollisionQueryParams& Params /* = FCollisionQueryParams::DefaultQueryParam */, const FCollisionResponseParams& ResponseParam /* = FCollisionResponseParams::DefaultResponseParam */) const
 {
 	if (CollisionShape.IsNearlyZero())
 	{
 		// if extent is 0, we'll just do linetrace instead
-		return LineTraceTest(Start, End, TraceChannel, Params, ResponseParam);
+		return LineTraceTestByChannel(Start, End, TraceChannel, Params, ResponseParam);
 	}
 	else
 	{
@@ -116,11 +119,11 @@ bool UWorld::SweepTest(const FVector& Start, const FVector& End, const FQuat& Ro
 	}
 }
 
-bool UWorld::SweepSingle(struct FHitResult& OutHit, const FVector& Start, const FVector& End, const FQuat& Rot, ECollisionChannel TraceChannel, const struct FCollisionShape& CollisionShape, const struct FCollisionQueryParams& Params, const struct FCollisionResponseParams& ResponseParam) const
+bool UWorld::SweepSingleByChannel(struct FHitResult& OutHit, const FVector& Start, const FVector& End, const FQuat& Rot, ECollisionChannel TraceChannel, const FCollisionShape& CollisionShape, const FCollisionQueryParams& Params /* = FCollisionQueryParams::DefaultQueryParam */, const FCollisionResponseParams& ResponseParam /* = FCollisionResponseParams::DefaultResponseParam */) const
 {
 	if (CollisionShape.IsNearlyZero())
 	{
-		return LineTraceSingle(OutHit, Start, End, TraceChannel, Params, ResponseParam);
+		return LineTraceSingleByChannel(OutHit, Start, End, TraceChannel, Params, ResponseParam);
 	}
 	else
 	{
@@ -134,11 +137,11 @@ bool UWorld::SweepSingle(struct FHitResult& OutHit, const FVector& Start, const
 	}
 }
 
-bool UWorld::SweepMulti(TArray<FHitResult>& OutHits, const FVector& Start, const FVector& End, const FQuat& Rot, ECollisionChannel TraceChannel, const struct FCollisionShape& CollisionShape, const struct FCollisionQueryParams& Params, const struct FCollisionResponseParams& ResponseParam) const
+bool UWorld::SweepMultiByChannel(TArray<struct FHitResult>& OutHits, const FVector& Start, const FVector& End, const FQuat& Rot, ECollisionChannel TraceChannel, const FCollisionShape& CollisionShape, const FCollisionQueryParams& Params /* = FCollisionQueryParams::DefaultQueryParam */, const FCollisionResponseParams& ResponseParam /* = FCollisionResponseParams::DefaultResponseParam */) const
 {
 	if (CollisionShape.IsNearlyZero())
 	{
-		return LineTraceMulti(OutHits, Start, End, TraceChannel, Params, ResponseParam);
+		return LineTraceMultiByChannel(OutHits, Start, End, TraceChannel, Params, ResponseParam);
 	}
 	else
 	{
@@ -150,27 +153,24 @@ bool UWorld::SweepMulti(TArray<FHitResult>& OutHits, const FVector& Start, const
 	}
 }
 
-bool UWorld::OverlapTest(const FVector& Pos, const FQuat& Rot, ECollisionChannel TraceChannel, const struct FCollisionShape& CollisionShape, const struct FCollisionQueryParams& Params, const struct FCollisionResponseParams& ResponseParam) const
-{
-	return OverlapBlockingTest(Pos, Rot, TraceChannel, CollisionShape, Params, ResponseParam);
-}
-
-bool UWorld::OverlapBlockingTest(const FVector& Pos, const FQuat& Rot, ECollisionChannel TraceChannel, const struct FCollisionShape& CollisionShape, const struct FCollisionQueryParams& Params, const struct FCollisionResponseParams& ResponseParam) const
+bool UWorld::OverlapBlockingTestByChannel(const FVector& Pos, const FQuat& Rot, ECollisionChannel TraceChannel, const FCollisionShape& CollisionShape, const FCollisionQueryParams& Params /* = FCollisionQueryParams::DefaultQueryParam */, const FCollisionResponseParams& ResponseParam /* = FCollisionResponseParams::DefaultResponseParam */) const
 {
 	bool bBlocking = false;
 #if UE_WITH_PHYSICS
 	bBlocking = GeomOverlapBlockingTest(this, CollisionShape, Pos, Rot, TraceChannel, Params, ResponseParam, FCollisionObjectQueryParams::DefaultObjectQueryParam);
 #endif
 	return bBlocking;
+
 }
 
-bool UWorld::OverlapAnyTest(const FVector& Pos, const FQuat& Rot, ECollisionChannel TraceChannel, const struct FCollisionShape& CollisionShape, const struct FCollisionQueryParams& Params, const struct FCollisionResponseParams& ResponseParam) const
+bool UWorld::OverlapAnyTestByChannel(const FVector& Pos, const FQuat& Rot, ECollisionChannel TraceChannel, const FCollisionShape& CollisionShape, const FCollisionQueryParams& Params /* = FCollisionQueryParams::DefaultQueryParam */, const FCollisionResponseParams& ResponseParam /* = FCollisionResponseParams::DefaultResponseParam */) const
 {
 	bool bBlocking = false;
 #if UE_WITH_PHYSICS
 	bBlocking = GeomOverlapAnyTest(this, CollisionShape, Pos, Rot, TraceChannel, Params, ResponseParam, FCollisionObjectQueryParams::DefaultObjectQueryParam);
 #endif
 	return bBlocking;
+
 }
 
 bool UWorld::OverlapSingle(struct FOverlapResult& OutOverlap, const FVector& Pos, const FQuat& Rot, ECollisionChannel TraceChannel, const struct FCollisionShape& CollisionShape, const struct FCollisionQueryParams& Params, const struct FCollisionResponseParams& ResponseParam) const
@@ -189,7 +189,7 @@ bool UWorld::OverlapSingle(struct FOverlapResult& OutOverlap, const FVector& Pos
 	return bBlocking;
 }
 
-bool UWorld::OverlapMulti(TArray<struct FOverlapResult>& OutOverlaps, const FVector& Pos, const FQuat& Rot, ECollisionChannel TraceChannel, const struct FCollisionShape& CollisionShape, const struct FCollisionQueryParams& Params, const struct FCollisionResponseParams& ResponseParam) const
+bool UWorld::OverlapMultiByChannel(TArray<struct FOverlapResult>& OutOverlaps, const FVector& Pos, const FQuat& Rot, ECollisionChannel TraceChannel, const FCollisionShape& CollisionShape, const FCollisionQueryParams& Params /* = FCollisionQueryParams::DefaultQueryParam */, const FCollisionResponseParams& ResponseParam /* = FCollisionResponseParams::DefaultResponseParam */) const
 {
 #if UE_WITH_PHYSICS
 	return GeomOverlapMulti(this, CollisionShape, Pos, Rot, OutOverlaps, TraceChannel, Params, ResponseParam, FCollisionObjectQueryParams::DefaultObjectQueryParam);
@@ -198,7 +198,9 @@ bool UWorld::OverlapMulti(TArray<struct FOverlapResult>& OutOverlaps, const FVec
 #endif
 }
 
-bool UWorld::OverlapMulti(TArray<struct FOverlapResult>& OutOverlaps, const FVector& Pos, const FQuat& Rot, const struct FCollisionShape& CollisionShape, const struct FCollisionQueryParams& Params, const struct FCollisionObjectQueryParams& ObjectQueryParams) const
+// object query interfaces
+
+bool UWorld::OverlapMultiByObjectType(TArray<struct FOverlapResult>& OutOverlaps, const FVector& Pos, const FQuat& Rot, const FCollisionObjectQueryParams& ObjectQueryParams, const FCollisionShape& CollisionShape, const FCollisionQueryParams& Params /* = FCollisionQueryParams::DefaultQueryParam */) const
 {
 #if UE_WITH_PHYSICS
 	GeomOverlapMulti(this, CollisionShape, Pos, Rot, OutOverlaps, DefaultCollisionChannel, Params, FCollisionResponseParams::DefaultResponseParam, ObjectQueryParams);
@@ -210,9 +212,7 @@ bool UWorld::OverlapMulti(TArray<struct FOverlapResult>& OutOverlaps, const FVec
 #endif
 }
 
-// object query interfaces
-
-bool UWorld::LineTraceTest(const FVector& Start,const FVector& End,const struct FCollisionQueryParams& Params, const struct FCollisionObjectQueryParams& ObjectQueryParams) const
+bool UWorld::LineTraceTestByObjectType(const FVector& Start,const FVector& End,const FCollisionObjectQueryParams& ObjectQueryParams, const FCollisionQueryParams& Params /* = FCollisionQueryParams::DefaultQueryParam */) const
 {
 #if UE_WITH_PHYSICS
 	return RaycastTest(this, Start, End, DefaultCollisionChannel, Params, FCollisionResponseParams::DefaultResponseParam, ObjectQueryParams);
@@ -221,7 +221,7 @@ bool UWorld::LineTraceTest(const FVector& Start,const FVector& End,const struct
 #endif
 }
 
-bool UWorld::LineTraceSingle(struct FHitResult& OutHit,const FVector& Start,const FVector& End,const struct FCollisionQueryParams& Params, const struct FCollisionObjectQueryParams& ObjectQueryParams) const
+bool UWorld::LineTraceSingleByObjectType(struct FHitResult& OutHit,const FVector& Start,const FVector& End,const FCollisionObjectQueryParams& ObjectQueryParams, const FCollisionQueryParams& Params /* = FCollisionQueryParams::DefaultQueryParam */) const
 {
 #if UE_WITH_PHYSICS
 	return RaycastSingle(this, OutHit, Start, End, DefaultCollisionChannel, Params, FCollisionResponseParams::DefaultResponseParam, ObjectQueryParams);
@@ -232,7 +232,7 @@ bool UWorld::LineTraceSingle(struct FHitResult& OutHit,const FVector& Start,cons
 #endif
 }
 
-bool UWorld::LineTraceMulti(TArray<struct FHitResult>& OutHits, const FVector& Start, const FVector& End, const struct FCollisionQueryParams& Params, const struct FCollisionObjectQueryParams& ObjectQueryParams) const
+bool UWorld::LineTraceMultiByObjectType(TArray<struct FHitResult>& OutHits,const FVector& Start,const FVector& End,const FCollisionObjectQueryParams& ObjectQueryParams, const FCollisionQueryParams& Params /* = FCollisionQueryParams::DefaultQueryParam */) const
 {
 #if UE_WITH_PHYSICS
 	RaycastMulti(this, OutHits, Start, End, DefaultCollisionChannel, Params, FCollisionResponseParams::DefaultResponseParam, ObjectQueryParams);
@@ -244,12 +244,12 @@ bool UWorld::LineTraceMulti(TArray<struct FHitResult>& OutHits, const FVector& S
 #endif
 }
 
-bool UWorld::SweepTest(const FVector& Start, const FVector& End, const FQuat& Rot, const struct FCollisionShape& CollisionShape, const struct FCollisionQueryParams& Params, const struct FCollisionObjectQueryParams& ObjectQueryParams) const
+bool UWorld::SweepTestByObjectType(const FVector& Start, const FVector& End, const FQuat& Rot, const FCollisionObjectQueryParams& ObjectQueryParams, const FCollisionShape& CollisionShape, const FCollisionQueryParams& Params /* = FCollisionQueryParams::DefaultQueryParam */) const
 {
 	if (CollisionShape.IsNearlyZero())
 	{
 		// if extent is 0, we'll just do linetrace instead
-		return LineTraceTest(Start, End, Params, ObjectQueryParams);
+		return LineTraceTestByObjectType(Start, End, ObjectQueryParams, Params);
 	}
 	else
 	{
@@ -259,13 +259,14 @@ bool UWorld::SweepTest(const FVector& Start, const FVector& End, const FQuat& Ro
 		return false;
 #endif
 	}
+
 }
 
-bool UWorld::SweepSingle(struct FHitResult& OutHit, const FVector& Start, const FVector& End, const FQuat& Rot, const struct FCollisionShape& CollisionShape, const struct FCollisionQueryParams& Params, const struct FCollisionObjectQueryParams& ObjectQueryParams) const
+bool UWorld::SweepSingleByObjectType(struct FHitResult& OutHit, const FVector& Start, const FVector& End, const FQuat& Rot, const FCollisionObjectQueryParams& ObjectQueryParams, const FCollisionShape& CollisionShape, const FCollisionQueryParams& Params /* = FCollisionQueryParams::DefaultQueryParam */) const
 {
 	if (CollisionShape.IsNearlyZero())
 	{
-		return LineTraceSingle(OutHit, Start, End, Params, ObjectQueryParams);
+		return LineTraceSingleByObjectType(OutHit, Start, End, ObjectQueryParams, Params);
 	}
 	else
 	{
@@ -277,13 +278,14 @@ bool UWorld::SweepSingle(struct FHitResult& OutHit, const FVector& Start, const
 		return false;
 #endif
 	}
+
 }
 
-bool UWorld::SweepMulti(TArray<struct FHitResult>& OutHits, const FVector& Start, const FVector& End, const FQuat& Rot, const struct FCollisionShape& CollisionShape, const struct FCollisionQueryParams& Params, const struct FCollisionObjectQueryParams& ObjectQueryParams) const
+bool UWorld::SweepMultiByObjectType(TArray<struct FHitResult>& OutHits, const FVector& Start, const FVector& End, const FQuat& Rot, const FCollisionObjectQueryParams& ObjectQueryParams, const FCollisionShape& CollisionShape, const FCollisionQueryParams& Params /* = FCollisionQueryParams::DefaultQueryParam */) const
 {
 	if (CollisionShape.IsNearlyZero())
 	{
-		return LineTraceMulti(OutHits, Start, End, Params, ObjectQueryParams);
+		return LineTraceMultiByObjectType(OutHits, Start, End, ObjectQueryParams, Params);
 	}
 	else
 	{
@@ -298,12 +300,7 @@ bool UWorld::SweepMulti(TArray<struct FHitResult>& OutHits, const FVector& Start
 	}
 }
 
-bool UWorld::OverlapTest(const FVector& Pos, const FQuat& Rot, const struct FCollisionShape& CollisionShape, const struct FCollisionQueryParams& Params, const struct FCollisionObjectQueryParams& ObjectQueryParams) const
-{
-	return OverlapAnyTest(Pos, Rot, CollisionShape, Params, ObjectQueryParams);
-}
-
-bool UWorld::OverlapAnyTest(const FVector& Pos, const FQuat& Rot, const struct FCollisionShape& CollisionShape, const struct FCollisionQueryParams& Params, const struct FCollisionObjectQueryParams& ObjectQueryParams) const
+bool UWorld::OverlapAnyTestByObjectType(const FVector& Pos, const FQuat& Rot, const FCollisionObjectQueryParams& ObjectQueryParams, const FCollisionShape& CollisionShape, const FCollisionQueryParams& Params /* = FCollisionQueryParams::DefaultQueryParam */) const
 {
 	bool bBlocking = false;
 #if UE_WITH_PHYSICS
@@ -312,6 +309,7 @@ bool UWorld::OverlapAnyTest(const FVector& Pos, const FQuat& Rot, const struct F
 	return bBlocking;
 }
 
+
 bool UWorld::OverlapSingle(struct FOverlapResult& OutOverlap, const FVector& Pos, const FQuat& Rot, const struct FCollisionShape& CollisionShape, const struct FCollisionQueryParams& Params, const struct FCollisionObjectQueryParams& ObjectQueryParams) const
 {
 	bool bBlocking = false;
@@ -328,6 +326,7 @@ bool UWorld::OverlapSingle(struct FOverlapResult& OutOverlap, const FVector& Pos
 	return bBlocking;
 }
 
+
 // profile interfaces
 static void GetCollisionProfileChannelAndResponseParams(FName ProfileName, ECollisionChannel &CollisionChannel, FCollisionResponseParams &ResponseParams)
 {
@@ -349,7 +348,7 @@ bool UWorld::LineTraceTestByProfile(const FVector& Start, const FVector& End, FN
 	FCollisionResponseParams ResponseParam;
 	GetCollisionProfileChannelAndResponseParams(ProfileName, TraceChannel, ResponseParam);
 
-	return LineTraceTest(Start, End, TraceChannel, Params, ResponseParam);
+	return LineTraceTestByChannel(Start, End, TraceChannel, Params, ResponseParam);
 }
 
 bool UWorld::LineTraceSingleByProfile(struct FHitResult& OutHit, const FVector& Start, const FVector& End, FName ProfileName, const struct FCollisionQueryParams& Params) const
@@ -358,7 +357,7 @@ bool UWorld::LineTraceSingleByProfile(struct FHitResult& OutHit, const FVector&
 	FCollisionResponseParams ResponseParam;
 	GetCollisionProfileChannelAndResponseParams(ProfileName, TraceChannel, ResponseParam);
 
-	return LineTraceSingle(OutHit, Start, End, TraceChannel, Params, ResponseParam);
+	return LineTraceSingleByChannel(OutHit, Start, End, TraceChannel, Params, ResponseParam);
 }
 
 bool UWorld::LineTraceMultiByProfile(TArray<struct FHitResult>& OutHits, const FVector& Start, const FVector& End, FName ProfileName, const struct FCollisionQueryParams& Params) const
@@ -367,7 +366,7 @@ bool UWorld::LineTraceMultiByProfile(TArray<struct FHitResult>& OutHits, const F
 	FCollisionResponseParams ResponseParam;
 	GetCollisionProfileChannelAndResponseParams(ProfileName, TraceChannel, ResponseParam);
 
-	return LineTraceMulti(OutHits, Start, End, TraceChannel, Params, ResponseParam);
+	return LineTraceMultiByChannel(OutHits, Start, End, TraceChannel, Params, ResponseParam);
 }
 
 bool UWorld::SweepTestByProfile(const FVector& Start, const FVector& End, const FQuat& Rot, FName ProfileName, const struct FCollisionShape& CollisionShape, const struct FCollisionQueryParams& Params) const
@@ -376,7 +375,7 @@ bool UWorld::SweepTestByProfile(const FVector& Start, const FVector& End, const
 	FCollisionResponseParams ResponseParam;
 	GetCollisionProfileChannelAndResponseParams(ProfileName, TraceChannel, ResponseParam);
 
-	return SweepTest(Start, End, Rot, TraceChannel, CollisionShape, Params, ResponseParam);
+	return SweepTestByChannel(Start, End, Rot, TraceChannel, CollisionShape, Params, ResponseParam);
 }
 
 bool UWorld::SweepSingleByProfile(struct FHitResult& OutHit, const FVector& Start, const FVector& End, const FQuat& Rot, FName ProfileName, const struct FCollisionShape& CollisionShape, const struct FCollisionQueryParams& Params) const
@@ -385,7 +384,7 @@ bool UWorld::SweepSingleByProfile(struct FHitResult& OutHit, const FVector& Star
 	FCollisionResponseParams ResponseParam;
 	GetCollisionProfileChannelAndResponseParams(ProfileName, TraceChannel, ResponseParam);
 
-	return SweepSingle(OutHit, Start, End, Rot, TraceChannel, CollisionShape, Params, ResponseParam);
+	return SweepSingleByChannel(OutHit, Start, End, Rot, TraceChannel, CollisionShape, Params, ResponseParam);
 }
 
 bool UWorld::SweepMultiByProfile(TArray<FHitResult>& OutHits, const FVector& Start, const FVector& End, const FQuat& Rot, FName ProfileName, const struct FCollisionShape& CollisionShape, const struct FCollisionQueryParams& Params) const
@@ -394,12 +393,7 @@ bool UWorld::SweepMultiByProfile(TArray<FHitResult>& OutHits, const FVector& Sta
 	FCollisionResponseParams ResponseParam;
 	GetCollisionProfileChannelAndResponseParams(ProfileName, TraceChannel, ResponseParam);
 
-	return SweepMulti(OutHits, Start, End, Rot, TraceChannel, CollisionShape, Params, ResponseParam);
-}
-
-bool UWorld::OverlapTestByProfile(const FVector& Pos, const FQuat& Rot, FName ProfileName, const struct FCollisionShape& CollisionShape, const struct FCollisionQueryParams& Params) const
-{
-	return OverlapBlockingTestByProfile(Pos, Rot, ProfileName, CollisionShape, Params);
+	return SweepMultiByChannel(OutHits, Start, End, Rot, TraceChannel, CollisionShape, Params, ResponseParam);
 }
 
 bool UWorld::OverlapBlockingTestByProfile(const FVector& Pos, const FQuat& Rot, FName ProfileName, const struct FCollisionShape& CollisionShape, const struct FCollisionQueryParams& Params) const
@@ -408,7 +402,7 @@ bool UWorld::OverlapBlockingTestByProfile(const FVector& Pos, const FQuat& Rot,
 	FCollisionResponseParams ResponseParam;
 	GetCollisionProfileChannelAndResponseParams(ProfileName, TraceChannel, ResponseParam);
 
-	return OverlapBlockingTest(Pos, Rot, TraceChannel, CollisionShape, Params, ResponseParam);
+	return OverlapBlockingTestByChannel(Pos, Rot, TraceChannel, CollisionShape, Params, ResponseParam);
 }
 
 bool UWorld::OverlapAnyTestByProfile(const FVector& Pos, const FQuat& Rot, FName ProfileName, const struct FCollisionShape& CollisionShape, const struct FCollisionQueryParams& Params) const
@@ -417,7 +411,7 @@ bool UWorld::OverlapAnyTestByProfile(const FVector& Pos, const FQuat& Rot, FName
 	FCollisionResponseParams ResponseParam;
 	GetCollisionProfileChannelAndResponseParams(ProfileName, TraceChannel, ResponseParam);
 
-	return OverlapAnyTest(Pos, Rot, TraceChannel, CollisionShape, Params, ResponseParam);
+	return OverlapAnyTestByChannel(Pos, Rot, TraceChannel, CollisionShape, Params, ResponseParam);
 }
 
 bool UWorld::OverlapSingleByProfile(struct FOverlapResult& OutOverlap, const FVector& Pos, const FQuat& Rot, FName ProfileName, const struct FCollisionShape& CollisionShape, const struct FCollisionQueryParams& Params) const
@@ -441,7 +435,7 @@ bool UWorld::OverlapMultiByProfile(TArray<struct FOverlapResult>& OutOverlaps, c
 	FCollisionResponseParams ResponseParam;
 	GetCollisionProfileChannelAndResponseParams(ProfileName, TraceChannel, ResponseParam);
 
-	return OverlapMulti(OutOverlaps, Pos, Rot, TraceChannel, CollisionShape, Params, ResponseParam);
+	return OverlapMultiByChannel(OutOverlaps, Pos, Rot, TraceChannel, CollisionShape, Params, ResponseParam);
 }
 
 
@@ -449,7 +443,7 @@ bool UWorld::ComponentOverlapMulti(TArray<struct FOverlapResult>& OutOverlaps, c
 {
 	if (PrimComp)
 	{
-		ComponentOverlapMulti(OutOverlaps, PrimComp, Pos, Rot, PrimComp->GetCollisionObjectType(), Params, ObjectQueryParams);
+		ComponentOverlapMultiByChannel(OutOverlaps, PrimComp, Pos, Rot, PrimComp->GetCollisionObjectType(), Params, ObjectQueryParams);
 		
 		// object query returns true if any hit is found, not only blocking hit
 		return (OutOverlaps.Num() > 0);
@@ -461,7 +455,7 @@ bool UWorld::ComponentOverlapMulti(TArray<struct FOverlapResult>& OutOverlaps, c
 	}
 }
 
-bool UWorld::ComponentOverlapMulti(TArray<struct FOverlapResult>& OutOverlaps, const class UPrimitiveComponent* PrimComp, const FVector& Pos, const FRotator& Rot, ECollisionChannel TraceChannel, const struct FComponentQueryParams& Params, const struct FCollisionObjectQueryParams& ObjectQueryParams) const
+bool UWorld::ComponentOverlapMultiByChannel(TArray<struct FOverlapResult>& OutOverlaps, const class UPrimitiveComponent* PrimComp, const FVector& Pos, const FRotator& Rot, ECollisionChannel TraceChannel, const FComponentQueryParams& Params /* = FComponentQueryParams::DefaultComponentQueryParams */, const FCollisionObjectQueryParams& ObjectQueryParams/* =FCollisionObjectQueryParams::DefaultObjectQueryParam */) const
 {
 	if (PrimComp)
 	{
diff --git a/Engine/Source/Runtime/Engine/Private/Collision/WorldCollisionAsync.cpp b/Engine/Source/Runtime/Engine/Private/Collision/WorldCollisionAsync.cpp
--- a/Engine/Source/Runtime/Engine/Private/Collision/WorldCollisionAsync.cpp
+++ b/Engine/Source/Runtime/Engine/Private/Collision/WorldCollisionAsync.cpp
@@ -294,33 +294,33 @@ FWorldAsyncTraceState::FWorldAsyncTraceState()
 {
 }
 
-FTraceHandle UWorld::AsyncLineTrace(const FVector& Start,const FVector& End,ECollisionChannel TraceChannel,const struct FCollisionQueryParams& Params, const struct FCollisionResponseParams& ResponseParam, FTraceDelegate * InDelegate, uint32 UserData, bool bMultiTrace)
+FTraceHandle UWorld::AsyncLineTraceByChannel(const FVector& Start,const FVector& End, ECollisionChannel TraceChannel, const FCollisionQueryParams& Params /* = FCollisionQueryParams::DefaultQueryParam */, const FCollisionResponseParams& ResponseParam /* = FCollisionResponseParams::DefaultResponseParam */, FTraceDelegate * InDelegate/* =NULL */, uint32 UserData /* = 0 */, bool bMultiTrace/* =false */ )
 {
 	return StartNewTrace(AsyncTraceState, FTraceDatum(this, FCollisionShape::LineShape, Params, ResponseParam, FCollisionObjectQueryParams::DefaultObjectQueryParam, TraceChannel, UserData, bMultiTrace, Start, End, InDelegate, AsyncTraceState.CurrentFrame));
 }
 
-FTraceHandle UWorld::AsyncLineTrace(const FVector& Start,const FVector& End, const struct FCollisionQueryParams& Params, const struct FCollisionObjectQueryParams& ObjectQueryParams, FTraceDelegate * InDelegate, uint32 UserData, bool bMultiTrace)
+FTraceHandle UWorld::AsyncLineTraceByObjectType(const FVector& Start,const FVector& End, const FCollisionObjectQueryParams& ObjectQueryParams, const FCollisionQueryParams& Params /* = FCollisionQueryParams::DefaultQueryParam */, FTraceDelegate * InDelegate/* =NULL */, uint32 UserData /* = 0 */, bool bMultiTrace/* =false */ )
 {
 	return StartNewTrace(AsyncTraceState, FTraceDatum(this, FCollisionShape::LineShape, Params, FCollisionResponseParams::DefaultResponseParam, ObjectQueryParams, DefaultCollisionChannel, UserData, bMultiTrace, Start, End, InDelegate, AsyncTraceState.CurrentFrame));
 }
 
-FTraceHandle UWorld::AsyncSweep(const FVector& Start,const FVector& End,ECollisionChannel TraceChannel, const struct FCollisionShape & CollisionShape, const struct FCollisionQueryParams& Params, const struct FCollisionResponseParams& ResponseParam, FTraceDelegate * InDelegate, uint32 UserData, bool bMultiTrace)
+FTraceHandle UWorld::AsyncSweepByChannel(const FVector& Start, const FVector& End, ECollisionChannel TraceChannel, const FCollisionShape& CollisionShape, const FCollisionQueryParams& Params /* = FCollisionQueryParams::DefaultQueryParam */, const FCollisionResponseParams& ResponseParam /* = FCollisionResponseParams::DefaultResponseParam */, FTraceDelegate * InDelegate /* = NULL */, uint32 UserData /* = 0 */, bool bMultiTrace /* = false */)
 {
 	return StartNewTrace(AsyncTraceState, FTraceDatum(this, CollisionShape, Params, ResponseParam, FCollisionObjectQueryParams::DefaultObjectQueryParam, TraceChannel, UserData, bMultiTrace, Start, End, InDelegate, AsyncTraceState.CurrentFrame));
 }
 
-FTraceHandle UWorld::AsyncSweep(const FVector& Start,const FVector& End,const struct FCollisionShape & CollisionShape, const struct FCollisionQueryParams& Params, const struct FCollisionObjectQueryParams& ObjectQueryParams, FTraceDelegate * InDelegate, uint32 UserData, bool bMultiTrace)
+FTraceHandle UWorld::AsyncSweepByObjectType(const FVector& Start, const FVector& End, const FCollisionObjectQueryParams& ObjectQueryParams, const FCollisionShape& CollisionShape, const FCollisionQueryParams& Params /* = FCollisionQueryParams::DefaultQueryParam */, FTraceDelegate * InDelegate /* = NULL */, uint32 UserData /* = 0 */, bool bMultiTrace /* = false */)
 {
 	return StartNewTrace(AsyncTraceState, FTraceDatum(this, CollisionShape, Params, FCollisionResponseParams::DefaultResponseParam, ObjectQueryParams, DefaultCollisionChannel, UserData, bMultiTrace, Start, End, InDelegate, AsyncTraceState.CurrentFrame));
 }
 
 // overlap functions
-FTraceHandle UWorld::AsyncOverlap(const FVector& Pos, const FQuat& Rot, ECollisionChannel TraceChannel,const struct FCollisionShape & CollisionShape, const struct FCollisionQueryParams& Params, const struct FCollisionResponseParams& ResponseParam, FOverlapDelegate * InDelegate, uint32 UserData)
+FTraceHandle UWorld::AsyncOverlapByChannel(const FVector& Pos, const FQuat& Rot, ECollisionChannel TraceChannel, const FCollisionShape& CollisionShape, const FCollisionQueryParams& Params /* = FCollisionQueryParams::DefaultQueryParam */, const FCollisionResponseParams& ResponseParam /* = FCollisionResponseParams::DefaultResponseParam */, FOverlapDelegate * InDelegate /* = NULL */, uint32 UserData /* = 0 */)
 {
 	return StartNewTrace(AsyncTraceState, FOverlapDatum(this, CollisionShape, Params, ResponseParam, FCollisionObjectQueryParams::DefaultObjectQueryParam, TraceChannel, UserData, Pos, Rot, InDelegate, AsyncTraceState.CurrentFrame));
 }
 
-FTraceHandle UWorld::AsyncOverlap(const FVector& Pos, const FQuat& Rot, const struct FCollisionShape & CollisionShape, const struct FCollisionQueryParams& Params, const struct FCollisionObjectQueryParams& ObjectQueryParams, FOverlapDelegate * InDelegate, uint32 UserData)
+FTraceHandle UWorld::AsyncOverlapByObjectType(const FVector& Pos, const FQuat& Rot, const FCollisionObjectQueryParams& ObjectQueryParams, const FCollisionShape& CollisionShape, const FCollisionQueryParams& Params /* = FCollisionQueryParams::DefaultQueryParam */, FOverlapDelegate * InDelegate /* = NULL */, uint32 UserData /* = 0 */)
 {
 	return StartNewTrace(AsyncTraceState, FOverlapDatum(this, CollisionShape, Params, FCollisionResponseParams::DefaultResponseParam, ObjectQueryParams, DefaultCollisionChannel, UserData, Pos, Rot, InDelegate, AsyncTraceState.CurrentFrame));
 }
diff --git a/Engine/Source/Runtime/Engine/Private/Controller.cpp b/Engine/Source/Runtime/Engine/Private/Controller.cpp
--- a/Engine/Source/Runtime/Engine/Private/Controller.cpp
+++ b/Engine/Source/Runtime/Engine/Private/Controller.cpp
@@ -172,7 +172,7 @@ bool AController::LineOfSightTo(const AActor* Other, FVector ViewPoint, bool bAl
 	FCollisionQueryParams CollisionParms(NAME_LineOfSight, true, Other);
 	CollisionParms.AddIgnoredActor(this->GetPawn());
 	FVector TargetLocation = Other->GetTargetLocation(Pawn);
-	bool bHit = GetWorld()->LineTraceTest(ViewPoint, TargetLocation, ECC_Visibility, CollisionParms);
+	bool bHit = GetWorld()->LineTraceTestByChannel(ViewPoint, TargetLocation, ECC_Visibility, CollisionParms);
 	if( !bHit )
 	{
 		return true;
@@ -198,7 +198,7 @@ bool AController::LineOfSightTo(const AActor* Other, FVector ViewPoint, bool bAl
 	Other->GetSimpleCollisionCylinder(OtherRadius, OtherHeight);
 	
 	//try viewpoint to head
-	bHit = GetWorld()->LineTraceTest(ViewPoint,  Other->GetActorLocation() + FVector(0.f,0.f,OtherHeight), ECC_Visibility, CollisionParms);
+	bHit = GetWorld()->LineTraceTestByChannel(ViewPoint,  Other->GetActorLocation() + FVector(0.f,0.f,OtherHeight), ECC_Visibility, CollisionParms);
 	return !bHit;
 }
 
diff --git a/Engine/Source/Runtime/Engine/Private/Debug/LogVisualizerHUD.cpp b/Engine/Source/Runtime/Engine/Private/Debug/LogVisualizerHUD.cpp
--- a/Engine/Source/Runtime/Engine/Private/Debug/LogVisualizerHUD.cpp
+++ b/Engine/Source/Runtime/Engine/Private/Debug/LogVisualizerHUD.cpp
@@ -44,7 +44,7 @@ void ALogVisualizerHUD::PostRender()
 
 			FCollisionQueryParams TraceParams(NAME_None, true, this);
 			FHitResult Hit;
-			bool bHit = GetWorld()->LineTraceSingle(Hit, CamLoc, CamRot.Vector() * 100000.f + CamLoc, ECC_Pawn, TraceParams);
+			bool bHit = GetWorld()->LineTraceSingleByChannel(Hit, CamLoc, CamRot.Vector() * 100000.f + CamLoc, ECC_Pawn, TraceParams);
 			if (bHit && Hit.GetActor() != nullptr)
 			{
 				TextItem.Text = FText::FromString(FString::Printf(TEXT("Under cursor: '%s'"), *Hit.GetActor()->GetName()));
diff --git a/Engine/Source/Runtime/Engine/Private/DebugCameraController.cpp b/Engine/Source/Runtime/Engine/Private/DebugCameraController.cpp
--- a/Engine/Source/Runtime/Engine/Private/DebugCameraController.cpp
+++ b/Engine/Source/Runtime/Engine/Private/DebugCameraController.cpp
@@ -422,7 +422,7 @@ void ADebugCameraController::SelectTargetedObject()
 
 	FHitResult Hit;
 	FCollisionQueryParams TraceParams(NAME_None, true, this);
-	bool const bHit = GetWorld()->LineTraceSingle(Hit, CamLoc, CamRot.Vector() * 5000.f * 20.f + CamLoc, ECC_Pawn, TraceParams);
+	bool const bHit = GetWorld()->LineTraceSingleByChannel(Hit, CamLoc, CamRot.Vector() * 5000.f * 20.f + CamLoc, ECC_Pawn, TraceParams);
 	if( bHit)
 	{
 		Select(Hit);
diff --git a/Engine/Source/Runtime/Engine/Private/DebugCameraHUD.cpp b/Engine/Source/Runtime/Engine/Private/DebugCameraHUD.cpp
--- a/Engine/Source/Runtime/Engine/Private/DebugCameraHUD.cpp
+++ b/Engine/Source/Runtime/Engine/Private/DebugCameraHUD.cpp
@@ -106,7 +106,7 @@ void ADebugCameraHUD::PostRender()
 
 			FCollisionQueryParams TraceParams(NAME_None, bTraceComplex, this);
 			FHitResult Hit;
-			bool bHit = GetWorld()->LineTraceSingle(Hit, CamLoc, CamRot.Vector() * 100000.f + CamLoc, ECC_Pawn, TraceParams);
+			bool bHit = GetWorld()->LineTraceSingleByChannel(Hit, CamLoc, CamRot.Vector() * 100000.f + CamLoc, ECC_Pawn, TraceParams);
 
 			yl += Y;
 			Canvas->DrawText(RenderFont, FString::Printf(TEXT("Trace info (%s = %d):"), CVarComplexName, bTraceComplex ? 1 : 0), X, yl, 1.f, 1.f, FontRenderInfo);
diff --git a/Engine/Source/Runtime/Engine/Private/EngineUtils.cpp b/Engine/Source/Runtime/Engine/Private/EngineUtils.cpp
--- a/Engine/Source/Runtime/Engine/Private/EngineUtils.cpp
+++ b/Engine/Source/Runtime/Engine/Private/EngineUtils.cpp
@@ -409,7 +409,7 @@ TArray<FSubLevelStatus> GetSubLevelsStatus( UWorld* World )
 
 			// this will not work for flying around :-(
 			static FName NAME_FindLevel = FName(TEXT("FindLevel"), true);			
-			PlayerController->GetWorld()->LineTraceSingle(Hit,PlayerController->GetPawn()->GetActorLocation(), (PlayerController->GetPawn()->GetActorLocation()-FVector(0.f, 0.f, 256.f)), FCollisionQueryParams(NAME_FindLevel, true, PlayerController->GetPawn()), FCollisionObjectQueryParams(ECC_WorldStatic));
+			PlayerController->GetWorld()->LineTraceSingleByObjectType(Hit,PlayerController->GetPawn()->GetActorLocation(), (PlayerController->GetPawn()->GetActorLocation()-FVector(0.f, 0.f, 256.f)), FCollisionObjectQueryParams(ECC_WorldStatic), FCollisionQueryParams(NAME_FindLevel, true, PlayerController->GetPawn()));
 
 			ULevel* LevelPlayerIsIn = NULL;
 
diff --git a/Engine/Source/Runtime/Engine/Private/GameFramework/SpringArmComponent.cpp b/Engine/Source/Runtime/Engine/Private/GameFramework/SpringArmComponent.cpp
--- a/Engine/Source/Runtime/Engine/Private/GameFramework/SpringArmComponent.cpp
+++ b/Engine/Source/Runtime/Engine/Private/GameFramework/SpringArmComponent.cpp
@@ -157,7 +157,7 @@ void USpringArmComponent::UpdateDesiredArmLocation(bool bDoTrace, bool bDoLocati
 		FCollisionQueryParams QueryParams(TraceTagName, false, GetOwner());
 
 		FHitResult Result;
-		GetWorld()->SweepSingle(Result, ArmOrigin, DesiredLoc, FQuat::Identity, ProbeChannel, FCollisionShape::MakeSphere(ProbeSize), QueryParams);
+		GetWorld()->SweepSingleByChannel(Result, ArmOrigin, DesiredLoc, FQuat::Identity, ProbeChannel, FCollisionShape::MakeSphere(ProbeSize), QueryParams);
 
 		ResultLoc = BlendLocations(DesiredLoc, Result.Location, Result.bBlockingHit, DeltaTime);
 	}
diff --git a/Engine/Source/Runtime/Engine/Private/GameplayStatics.cpp b/Engine/Source/Runtime/Engine/Private/GameplayStatics.cpp
--- a/Engine/Source/Runtime/Engine/Private/GameplayStatics.cpp
+++ b/Engine/Source/Runtime/Engine/Private/GameplayStatics.cpp
@@ -172,7 +172,7 @@ static bool ComponentIsDamageableFrom(UPrimitiveComponent* VictimComp, FVector c
 		// tiny nudge so LineTraceSingle doesn't early out with no hits
 		TraceStart.Z += 0.01f;
 	}
-	bool const bHadBlockingHit = World->LineTraceSingle(OutHitResult, TraceStart, TraceEnd, TraceChannel, LineParams);
+	bool const bHadBlockingHit = World->LineTraceSingleByChannel(OutHitResult, TraceStart, TraceEnd, TraceChannel, LineParams);
 	//::DrawDebugLine(World, TraceStart, TraceEnd, FLinearColor::Red, true);
 
 	// If there was a blocking hit, it will be the last one
@@ -215,7 +215,7 @@ bool UGameplayStatics::ApplyRadialDamageWithFalloff(UObject* WorldContextObject,
 	// query scene to see what we hit
 	TArray<FOverlapResult> Overlaps;
 	UWorld* World = GEngine->GetWorldFromContextObject(WorldContextObject);
-	World->OverlapMulti(Overlaps, Origin, FQuat::Identity, FCollisionShape::MakeSphere(DamageOuterRadius), SphereParams, FCollisionObjectQueryParams(FCollisionObjectQueryParams::InitType::AllDynamicObjects));
+	World->OverlapMultiByObjectType(Overlaps, Origin, FQuat::Identity, FCollisionObjectQueryParams(FCollisionObjectQueryParams::InitType::AllDynamicObjects), FCollisionShape::MakeSphere(DamageOuterRadius), SphereParams);
 
 	// collate into per-actor list of hit components
 	TMap<AActor*, TArray<FHitResult> > OverlapComponentMap;
@@ -1301,7 +1301,7 @@ bool UGameplayStatics::SuggestProjectileVelocity(UObject* WorldContextObject, FV
 
 					FCollisionQueryParams QueryParams(NAME_SuggestProjVelTrace, true);
 					QueryParams.AddIgnoredActors(ActorsToIgnore);
-					if (World->SweepTest(TraceStart, TraceEnd, FQuat::Identity, ECC_WorldDynamic, FCollisionShape::MakeSphere(CollisionRadius), QueryParams, ResponseParam))
+					if (World->SweepTestByChannel(TraceStart, TraceEnd, FQuat::Identity, ECC_WorldDynamic, FCollisionShape::MakeSphere(CollisionRadius), QueryParams, ResponseParam))
 					{
 						// hit something, failed
 						bFailedTrace = true;
diff --git a/Engine/Source/Runtime/Engine/Private/HUD.cpp b/Engine/Source/Runtime/Engine/Private/HUD.cpp
--- a/Engine/Source/Runtime/Engine/Private/HUD.cpp
+++ b/Engine/Source/Runtime/Engine/Private/HUD.cpp
@@ -330,7 +330,7 @@ void AHUD::ShowDebugInfo(float& YL, float& YPos)
 
 			FCollisionQueryParams TraceParams(NAME_None, true, PlayerOwner->PlayerCameraManager->ViewTarget.Target);
 			FHitResult Hit;
-			bool bHit = GetWorld()->LineTraceSingle(Hit, CamLoc, CamRot.Vector() * 100000.f + CamLoc, ECC_WorldDynamic, TraceParams);
+			bool bHit = GetWorld()->LineTraceSingleByChannel(Hit, CamLoc, CamRot.Vector() * 100000.f + CamLoc, ECC_WorldDynamic, TraceParams);
 			if (bHit)
 			{
 				AActor* HitActor = Hit.Actor.Get();
diff --git a/Engine/Source/Runtime/Engine/Private/KismetSystemLibrary.cpp b/Engine/Source/Runtime/Engine/Private/KismetSystemLibrary.cpp
--- a/Engine/Source/Runtime/Engine/Private/KismetSystemLibrary.cpp
+++ b/Engine/Source/Runtime/Engine/Private/KismetSystemLibrary.cpp
@@ -790,7 +790,7 @@ bool UKismetSystemLibrary::SphereOverlapComponents_DEPRECATED(UObject* WorldCont
 	FCollisionObjectQueryParams::InitType InitType = FCollisionObjectQueryParams::GetCollisionChannelFromOverlapFilter(Filter);
 	
 	UWorld* World = GEngine->GetWorldFromContextObject( WorldContextObject );
-	World->OverlapMulti(Overlaps, SpherePos, FQuat::Identity, FCollisionShape::MakeSphere(SphereRadius), Params, FCollisionObjectQueryParams(InitType));
+	World->OverlapMultiByObjectType(Overlaps, SpherePos, FQuat::Identity, FCollisionObjectQueryParams(InitType), FCollisionShape::MakeSphere(SphereRadius), Params);
 	
 	for (int32 OverlapIdx=0; OverlapIdx<Overlaps.Num(); ++OverlapIdx)
 	{
@@ -837,7 +837,7 @@ bool UKismetSystemLibrary::BoxOverlapComponents_DEPRECATED(UObject* WorldContext
 	FCollisionObjectQueryParams::InitType InitType = FCollisionObjectQueryParams::GetCollisionChannelFromOverlapFilter(Filter);
 	
 	UWorld* World = GEngine->GetWorldFromContextObject( WorldContextObject );
-	World->OverlapMulti(Overlaps, BoxPos, FQuat::Identity, FCollisionShape::MakeBox(BoxExtent), Params, FCollisionObjectQueryParams(InitType));
+	World->OverlapMultiByObjectType(Overlaps, BoxPos, FQuat::Identity, FCollisionObjectQueryParams(InitType), FCollisionShape::MakeBox(BoxExtent), Params);
 
 	for (int32 OverlapIdx=0; OverlapIdx<Overlaps.Num(); ++OverlapIdx)
 	{
@@ -883,7 +883,7 @@ bool UKismetSystemLibrary::CapsuleOverlapComponents_DEPRECATED(UObject* WorldCon
 	FCollisionObjectQueryParams::InitType InitType = FCollisionObjectQueryParams::GetCollisionChannelFromOverlapFilter(Filter);
 
 	UWorld* World = GEngine->GetWorldFromContextObject( WorldContextObject );	
-	World->OverlapMulti(Overlaps, CapsulePos, FQuat::Identity, FCollisionShape::MakeCapsule(Radius, HalfHeight), Params, FCollisionObjectQueryParams(InitType));
+	World->OverlapMultiByObjectType(Overlaps, CapsulePos, FQuat::Identity, FCollisionObjectQueryParams(InitType), FCollisionShape::MakeCapsule(Radius, HalfHeight), Params);
 
 	for (int32 OverlapIdx=0; OverlapIdx<Overlaps.Num(); ++OverlapIdx)
 	{
@@ -980,7 +980,7 @@ bool UKismetSystemLibrary::SphereOverlapComponents_NEW(UObject* WorldContextObje
 	UWorld* World = GEngine->GetWorldFromContextObject( WorldContextObject );
 	if(World != nullptr)
 	{
-		World->OverlapMulti(Overlaps, SpherePos, FQuat::Identity, FCollisionShape::MakeSphere(SphereRadius), Params, ObjectParams);
+		World->OverlapMultiByObjectType(Overlaps, SpherePos, FQuat::Identity, ObjectParams, FCollisionShape::MakeSphere(SphereRadius), Params);
 	}
 
 	for (int32 OverlapIdx=0; OverlapIdx<Overlaps.Num(); ++OverlapIdx)
@@ -1035,7 +1035,7 @@ bool UKismetSystemLibrary::BoxOverlapComponents_NEW(UObject* WorldContextObject,
 	UWorld* World = GEngine->GetWorldFromContextObject( WorldContextObject );
 	if (World != nullptr)
 	{
-		World->OverlapMulti(Overlaps, BoxPos, FQuat::Identity, FCollisionShape::MakeBox(BoxExtent), Params, ObjectParams);
+		World->OverlapMultiByObjectType(Overlaps, BoxPos, FQuat::Identity, ObjectParams, FCollisionShape::MakeBox(BoxExtent), Params);
 	}
 
 	for (int32 OverlapIdx=0; OverlapIdx<Overlaps.Num(); ++OverlapIdx)
@@ -1089,7 +1089,7 @@ bool UKismetSystemLibrary::CapsuleOverlapComponents_NEW(UObject* WorldContextObj
 	UWorld* World = GEngine->GetWorldFromContextObject( WorldContextObject );	
 	if (World != nullptr)
 	{
-		World->OverlapMulti(Overlaps, CapsulePos, FQuat::Identity, FCollisionShape::MakeCapsule(Radius, HalfHeight), Params, ObjectParams);
+		World->OverlapMultiByObjectType(Overlaps, CapsulePos, FQuat::Identity, ObjectParams, FCollisionShape::MakeCapsule(Radius, HalfHeight), Params);
 	}
 
 	for (int32 OverlapIdx=0; OverlapIdx<Overlaps.Num(); ++OverlapIdx)
@@ -1197,7 +1197,7 @@ bool UKismetSystemLibrary::LineTraceSingle_DEPRECATED(UObject* WorldContextObjec
 	}
 
 	UWorld* World = GEngine->GetWorldFromContextObject( WorldContextObject );
-	bool const bHit = World->LineTraceSingle(OutHit, Start, End, TraceChannel, Params);
+	bool const bHit = World->LineTraceSingleByChannel(OutHit, Start, End, TraceChannel, Params);
 
 	if (DrawDebugType != EDrawDebugTrace::None)
 	{
@@ -1260,7 +1260,7 @@ bool UKismetSystemLibrary::LineTraceMulti_DEPRECATED(UObject* WorldContextObject
 	}
 
 	UWorld* World = GEngine->GetWorldFromContextObject( WorldContextObject );	
-	bool const bHit = World->LineTraceMulti(OutHits, Start, End, TraceChannel, Params);
+	bool const bHit = World->LineTraceMultiByChannel(OutHits, Start, End, TraceChannel, Params);
 
 	if (DrawDebugType != EDrawDebugTrace::None)
 	{
@@ -1364,7 +1364,7 @@ bool UKismetSystemLibrary::BoxTraceSingle(UObject* WorldContextObject, const FVe
 	}
 
 	UWorld* World = GEngine->GetWorldFromContextObject(WorldContextObject);
-	bool const bHit = World ? World->SweepSingle(OutHit, Start, End, Orientation.Quaternion(), UEngineTypes::ConvertToCollisionChannel(TraceChannel), FCollisionShape::MakeBox(HalfSize), Params) : false;
+	bool const bHit = World ? World->SweepSingleByChannel(OutHit, Start, End, Orientation.Quaternion(), UEngineTypes::ConvertToCollisionChannel(TraceChannel), FCollisionShape::MakeBox(HalfSize), Params) : false;
 
 	if (DrawDebugType != EDrawDebugTrace::None && (World != nullptr))
 	{
@@ -1419,7 +1419,7 @@ bool UKismetSystemLibrary::BoxTraceMulti(UObject* WorldContextObject, const FVec
 	}
 
 	UWorld* World = GEngine->GetWorldFromContextObject(WorldContextObject);
-	bool const bHit = World ? World->SweepMulti(OutHits, Start, End, Orientation.Quaternion(), UEngineTypes::ConvertToCollisionChannel(TraceChannel), FCollisionShape::MakeBox(HalfSize), Params) : false;
+	bool const bHit = World ? World->SweepMultiByChannel(OutHits, Start, End, Orientation.Quaternion(), UEngineTypes::ConvertToCollisionChannel(TraceChannel), FCollisionShape::MakeBox(HalfSize), Params) : false;
 
 	if (DrawDebugType != EDrawDebugTrace::None && (World != nullptr))
 	{
@@ -1488,7 +1488,7 @@ bool UKismetSystemLibrary::SphereTraceSingle_DEPRECATED(UObject* WorldContextObj
 	}
 
 	UWorld* World = GEngine->GetWorldFromContextObject( WorldContextObject );
-	bool const bHit = World->SweepSingle(OutHit, Start, End, FQuat::Identity, TraceChannel, FCollisionShape::MakeSphere(Radius), Params);
+	bool const bHit = World->SweepSingleByChannel(OutHit, Start, End, FQuat::Identity, TraceChannel, FCollisionShape::MakeSphere(Radius), Params);
 
 	if (DrawDebugType != EDrawDebugTrace::None)
 	{
@@ -1550,7 +1550,7 @@ bool UKismetSystemLibrary::SphereTraceMulti_DEPRECATED(UObject* WorldContextObje
 	}
 
 	UWorld* World = GEngine->GetWorldFromContextObject( WorldContextObject );
-	bool const bHit = World->SweepMulti(OutHits, Start, End, FQuat::Identity, TraceChannel, FCollisionShape::MakeSphere(Radius), Params);
+	bool const bHit = World->SweepMultiByChannel(OutHits, Start, End, FQuat::Identity, TraceChannel, FCollisionShape::MakeSphere(Radius), Params);
 
 	if (DrawDebugType != EDrawDebugTrace::None)
 	{
@@ -1619,7 +1619,7 @@ bool UKismetSystemLibrary::CapsuleTraceSingle_DEPRECATED(UObject* WorldContextOb
 	}
 
 	UWorld* World = GEngine->GetWorldFromContextObject( WorldContextObject );
-	bool const bHit = World->SweepSingle(OutHit, Start, End, FQuat::Identity, TraceChannel, FCollisionShape::MakeCapsule(Radius, HalfHeight), Params);
+	bool const bHit = World->SweepSingleByChannel(OutHit, Start, End, FQuat::Identity, TraceChannel, FCollisionShape::MakeCapsule(Radius, HalfHeight), Params);
 
 	if (DrawDebugType != EDrawDebugTrace::None)
 	{
@@ -1687,7 +1687,7 @@ bool UKismetSystemLibrary::CapsuleTraceMulti_DEPRECATED(UObject* WorldContextObj
 	}
 
 	UWorld* World = GEngine->GetWorldFromContextObject( WorldContextObject );
-	bool const bHit = World->SweepMulti(OutHits, Start, End, FQuat::Identity, TraceChannel, FCollisionShape::MakeCapsule(Radius, HalfHeight), Params);
+	bool const bHit = World->SweepMultiByChannel(OutHits, Start, End, FQuat::Identity, TraceChannel, FCollisionShape::MakeCapsule(Radius, HalfHeight), Params);
 
 	if (DrawDebugType != EDrawDebugTrace::None)
 	{
@@ -1792,7 +1792,7 @@ bool UKismetSystemLibrary::LineTraceSingleByObject_DEPRECATED(UObject* WorldCont
 	}
 
 	UWorld* World = GEngine->GetWorldFromContextObject( WorldContextObject );
-	bool const bHit = World->LineTraceSingle(OutHit, Start, End, Params, ObjectParams);
+	bool const bHit = World->LineTraceSingleByObjectType(OutHit, Start, End, ObjectParams, Params);
 
 	if (DrawDebugType != EDrawDebugTrace::None)
 	{
@@ -1883,7 +1883,7 @@ bool UKismetSystemLibrary::LineTraceMultiByObject_DEPRECATED(UObject* WorldConte
 	}
 
 	UWorld* World = GEngine->GetWorldFromContextObject( WorldContextObject );	
-	bool const bHit = World->LineTraceMulti(OutHits, Start, End, Params, ObjectParams);
+	bool const bHit = World->LineTraceMultiByObjectType(OutHits, Start, End, ObjectParams, Params);
 
 	if (DrawDebugType != EDrawDebugTrace::None)
 	{
@@ -1985,7 +1985,7 @@ bool UKismetSystemLibrary::SphereTraceSingleByObject_DEPRECATED(UObject* WorldCo
 	}
 
 	UWorld* World = GEngine->GetWorldFromContextObject( WorldContextObject );
-	bool const bHit = World->SweepSingle(OutHit, Start, End, FQuat::Identity, FCollisionShape::MakeSphere(Radius), Params, ObjectParams);
+	bool const bHit = World->SweepSingleByObjectType(OutHit, Start, End, FQuat::Identity, ObjectParams, FCollisionShape::MakeSphere(Radius), Params);
 
 	if (DrawDebugType != EDrawDebugTrace::None)
 	{
@@ -2075,7 +2075,7 @@ bool UKismetSystemLibrary::SphereTraceMultiByObject_DEPRECATED(UObject* WorldCon
 	}
 
 	UWorld* World = GEngine->GetWorldFromContextObject( WorldContextObject );
-	bool const bHit = World->SweepMulti(OutHits, Start, End, FQuat::Identity, FCollisionShape::MakeSphere(Radius), Params, ObjectParams);
+	bool const bHit = World->SweepMultiByObjectType(OutHits, Start, End, FQuat::Identity, ObjectParams, FCollisionShape::MakeSphere(Radius), Params);
 
 	if (DrawDebugType != EDrawDebugTrace::None)
 	{
@@ -2165,7 +2165,7 @@ bool UKismetSystemLibrary::BoxTraceSingleForObjects(UObject* WorldContextObject,
 	}
 
 	UWorld* World = GEngine->GetWorldFromContextObject(WorldContextObject);
-	bool const bHit = World ? World->SweepSingle(OutHit, Start, End, Orientation.Quaternion(), FCollisionShape::MakeBox(HalfSize), Params, ObjectParams) : false;
+	bool const bHit = World ? World->SweepSingleByObjectType(OutHit, Start, End, Orientation.Quaternion(), ObjectParams, FCollisionShape::MakeBox(HalfSize), Params) : false;
 
 	if (DrawDebugType != EDrawDebugTrace::None && (World != nullptr))
 	{
@@ -2248,7 +2248,7 @@ bool UKismetSystemLibrary::BoxTraceMultiForObjects(UObject* WorldContextObject,
 	}
 
 	UWorld* World = GEngine->GetWorldFromContextObject(WorldContextObject);
-	bool const bHit = World ? World->SweepMulti(OutHits, Start, End, Orientation.Quaternion(), FCollisionShape::MakeBox(HalfSize), Params, ObjectParams) : false;
+	bool const bHit = World ? World->SweepMultiByObjectType(OutHits, Start, End, Orientation.Quaternion(), ObjectParams, FCollisionShape::MakeBox(HalfSize), Params) : false;
 
 	if (DrawDebugType != EDrawDebugTrace::None && (World != nullptr))
 	{
@@ -2345,7 +2345,7 @@ bool UKismetSystemLibrary::CapsuleTraceSingleByObject_DEPRECATED(UObject* WorldC
 	}
 
 	UWorld* World = GEngine->GetWorldFromContextObject( WorldContextObject );
-	bool const bHit = World->SweepSingle(OutHit, Start, End, FQuat::Identity, FCollisionShape::MakeCapsule(Radius, HalfHeight), Params, ObjectParams);
+	bool const bHit = World->SweepSingleByObjectType(OutHit, Start, End, FQuat::Identity, ObjectParams, FCollisionShape::MakeCapsule(Radius, HalfHeight), Params);
 
 	if (DrawDebugType != EDrawDebugTrace::None)
 	{
@@ -2440,7 +2440,7 @@ bool UKismetSystemLibrary::CapsuleTraceMultiByObject_DEPRECATED(UObject* WorldCo
 	}
 
 	UWorld* World = GEngine->GetWorldFromContextObject( WorldContextObject );
-	bool const bHit = World->SweepMulti(OutHits, Start, End, FQuat::Identity, FCollisionShape::MakeCapsule(Radius, HalfHeight), Params, ObjectParams);
+	bool const bHit = World->SweepMultiByObjectType(OutHits, Start, End, FQuat::Identity, ObjectParams, FCollisionShape::MakeCapsule(Radius, HalfHeight), Params);
 
 	if (DrawDebugType != EDrawDebugTrace::None)
 	{
diff --git a/Engine/Source/Runtime/Engine/Private/LevelActor.cpp b/Engine/Source/Runtime/Engine/Private/LevelActor.cpp
--- a/Engine/Source/Runtime/Engine/Private/LevelActor.cpp
+++ b/Engine/Source/Runtime/Engine/Private/LevelActor.cpp
@@ -688,7 +688,7 @@ bool UWorld::EncroachingBlockingGeometry(AActor* TestActor, FVector TestLocation
 		if (Capsule)
 		{
 			FCollisionQueryParams Params(NAME_EncroachingBlockingGeometry, false, TestActor);
-			bFoundBlockingHit = OverlapMulti(Overlaps, TestLocation, FQuat::Identity, BlockingChannel, FCollisionShape::MakeCapsule(FMath::Max(Capsule->GetScaledCapsuleRadius() - Epsilon, 0.1f), FMath::Max(Capsule->GetScaledCapsuleHalfHeight() - Epsilon, 0.1f)), Params);
+			bFoundBlockingHit = OverlapMultiByChannel(Overlaps, TestLocation, FQuat::Identity, BlockingChannel, FCollisionShape::MakeCapsule(FMath::Max(Capsule->GetScaledCapsuleRadius() - Epsilon, 0.1f), FMath::Max(Capsule->GetScaledCapsuleHalfHeight() - Epsilon, 0.1f)), Params);
 		}
 		else
 		{
@@ -696,13 +696,13 @@ bool UWorld::EncroachingBlockingGeometry(AActor* TestActor, FVector TestLocation
 			if (Sphere)
 			{
 				FCollisionQueryParams Params(NAME_EncroachingBlockingGeometry, false, TestActor);
-				bFoundBlockingHit = OverlapMulti(Overlaps, TestLocation, FQuat::Identity, BlockingChannel, FCollisionShape::MakeSphere(FMath::Max(Sphere->GetScaledSphereRadius() - Epsilon, 0.1f)), Params);
+				bFoundBlockingHit = OverlapMultiByChannel(Overlaps, TestLocation, FQuat::Identity, BlockingChannel, FCollisionShape::MakeSphere(FMath::Max(Sphere->GetScaledSphereRadius() - Epsilon, 0.1f)), Params);
 			}
 			else if (PrimComp->IsRegistered())
 			{
 				// must be registered
 				FComponentQueryParams Params(NAME_EncroachingBlockingGeometry, TestActor);
-				bFoundBlockingHit = ComponentOverlapMulti(Overlaps, PrimComp, TestLocation, TestActor->GetActorRotation(), BlockingChannel, Params);
+				bFoundBlockingHit = ComponentOverlapMultiByChannel(Overlaps, PrimComp, TestLocation, TestActor->GetActorRotation(), BlockingChannel, Params);
 			}
 			else
 			{
diff --git a/Engine/Source/Runtime/Engine/Private/MovementComponent.cpp b/Engine/Source/Runtime/Engine/Private/MovementComponent.cpp
--- a/Engine/Source/Runtime/Engine/Private/MovementComponent.cpp
+++ b/Engine/Source/Runtime/Engine/Private/MovementComponent.cpp
@@ -310,7 +310,7 @@ bool UMovementComponent::OverlapTest(const FVector& Location, const FQuat& Rotat
 	FCollisionQueryParams QueryParams(NAME_TestOverlap, false, IgnoreActor);
 	FCollisionResponseParams ResponseParam;
 	InitCollisionParams(QueryParams, ResponseParam);
-	return GetWorld()->OverlapBlockingTest(Location, RotationQuat, CollisionChannel, CollisionShape, QueryParams, ResponseParam);
+	return GetWorld()->OverlapBlockingTestByChannel(Location, RotationQuat, CollisionChannel, CollisionShape, QueryParams, ResponseParam);
 }
 
 bool UMovementComponent::IsExceedingMaxSpeed(float MaxSpeed) const
diff --git a/Engine/Source/Runtime/Engine/Private/NavigationObjectBase.cpp b/Engine/Source/Runtime/Engine/Private/NavigationObjectBase.cpp
--- a/Engine/Source/Runtime/Engine/Private/NavigationObjectBase.cpp
+++ b/Engine/Source/Runtime/Engine/Private/NavigationObjectBase.cpp
@@ -149,7 +149,7 @@ void ANavigationObjectBase::FindBase()
 
 		static FName NAME_NavFindBase = FName(TEXT("NavFindBase"));
 
-		GetWorld()->SweepSingle( Hit, TraceStart, TraceEnd, FQuat::Identity, FCollisionShape::MakeBox(CollisionSlice), FCollisionQueryParams(NAME_NavFindBase, false), FCollisionObjectQueryParams(ECC_WorldStatic));
+		GetWorld()->SweepSingleByObjectType( Hit, TraceStart, TraceEnd, FQuat::Identity, FCollisionObjectQueryParams(ECC_WorldStatic), FCollisionShape::MakeBox(CollisionSlice), FCollisionQueryParams(NAME_NavFindBase, false));
 
 		// @fixme, ensure object is on the navmesh?
 // 		if( Hit.Actor != NULL )
@@ -191,7 +191,7 @@ void ANavigationObjectBase::Validate()
 		FHitResult Hit(ForceInit);
 		const FVector TraceStart = GetActorLocation();
 		const FVector TraceEnd = GetActorLocation() - FVector(0.f,0.f, 4.f * CapsuleComponent->GetScaledCapsuleHalfHeight());
-		GetWorld()->SweepSingle(Hit, TraceStart, TraceEnd, FQuat::Identity, ECC_Pawn, FCollisionShape::MakeBox(Slice), FCollisionQueryParams(NAME_None, false, this));
+		GetWorld()->SweepSingleByChannel(Hit, TraceStart, TraceEnd, FQuat::Identity, ECC_Pawn, FCollisionShape::MakeBox(Slice), FCollisionQueryParams(NAME_None, false, this));
 		if( Hit.bBlockingHit )
 		{
 			const FVector HitLocation = TraceStart + (TraceEnd - TraceStart) * Hit.Time;
diff --git a/Engine/Source/Runtime/Engine/Private/NetworkDriver.cpp b/Engine/Source/Runtime/Engine/Private/NetworkDriver.cpp
--- a/Engine/Source/Runtime/Engine/Private/NetworkDriver.cpp
+++ b/Engine/Source/Runtime/Engine/Private/NetworkDriver.cpp
@@ -1787,7 +1787,7 @@ FNetViewer::FNetViewer(UNetConnection* InConnection, float DeltaSeconds) :
 				World = ViewerPawn->GetWorld();
 			}
 			check( World );
-			World->LineTraceSingle(Hit, ViewLocation, Hit.Location, FCollisionQueryParams(NAME_ServerForwardView, true, ViewTarget), FCollisionObjectQueryParams(ECC_WorldStatic));
+			World->LineTraceSingleByObjectType(Hit, ViewLocation, Hit.Location, FCollisionObjectQueryParams(ECC_WorldStatic), FCollisionQueryParams(NAME_ServerForwardView, true, ViewTarget));
 			ViewLocation = Hit.Location;
 		}
 	}
diff --git a/Engine/Source/Runtime/Engine/Private/Particles/ParticleComponents.cpp b/Engine/Source/Runtime/Engine/Private/Particles/ParticleComponents.cpp
--- a/Engine/Source/Runtime/Engine/Private/Particles/ParticleComponents.cpp
+++ b/Engine/Source/Runtime/Engine/Private/Particles/ParticleComponents.cpp
@@ -3055,14 +3055,14 @@ bool UParticleSystemComponent::ParticleLineCheck(FHitResult& Hit, AActor* Source
 
 	if ( HalfExtent.IsZero() )
 	{
-		return GetWorld()->LineTraceSingle(Hit, Start, End, FCollisionQueryParams(NAME_ParticleCollision, true, SourceActor), ObjectParams);
+		return GetWorld()->LineTraceSingleByObjectType(Hit, Start, End, ObjectParams, FCollisionQueryParams(NAME_ParticleCollision, true, SourceActor));
 	}
 	else
 	{
 		FCollisionQueryParams BoxParams(false);
 		BoxParams.TraceTag = NAME_ParticleCollision;
 		BoxParams.AddIgnoredActor(SourceActor);
-		return GetWorld()->SweepSingle(Hit, Start, End, FQuat::Identity, FCollisionShape::MakeBox(HalfExtent), BoxParams, ObjectParams);
+		return GetWorld()->SweepSingleByObjectType(Hit, Start, End, FQuat::Identity, ObjectParams, FCollisionShape::MakeBox(HalfExtent), BoxParams);
 	}
 }
 
diff --git a/Engine/Source/Runtime/Engine/Private/PhysicsEngine/PhysicsSpring.cpp b/Engine/Source/Runtime/Engine/Private/PhysicsEngine/PhysicsSpring.cpp
--- a/Engine/Source/Runtime/Engine/Private/PhysicsEngine/PhysicsSpring.cpp
+++ b/Engine/Source/Runtime/Engine/Private/PhysicsEngine/PhysicsSpring.cpp
@@ -63,7 +63,7 @@ UPrimitiveComponent* UPhysicsSpringComponent::GetSpringCollision(const FVector&
 	const float DeltaSize = Delta.Size();
 	if (DeltaSize > SMALL_NUMBER)
 	{
-		if (bool bBlockingHit = World->SweepSingle(Hit, Start, End, FQuat::Identity, SpringChannel, FCollisionShape::MakeSphere(SpringRadius), QueryParams))
+		if (bool bBlockingHit = World->SweepSingleByChannel(Hit, Start, End, FQuat::Identity, SpringChannel, FCollisionShape::MakeSphere(SpringRadius), QueryParams))
 		{
 			CollidedComponent = Hit.GetComponent();
 			Time = CollidedComponent ? Hit.Time : 1.f;
diff --git a/Engine/Source/Runtime/Engine/Private/PhysicsEngine/RadialForceComponent.cpp b/Engine/Source/Runtime/Engine/Private/PhysicsEngine/RadialForceComponent.cpp
--- a/Engine/Source/Runtime/Engine/Private/PhysicsEngine/RadialForceComponent.cpp
+++ b/Engine/Source/Runtime/Engine/Private/PhysicsEngine/RadialForceComponent.cpp
@@ -43,7 +43,7 @@ void URadialForceComponent::TickComponent(float DeltaTime, enum ELevelTick TickT
 
 		FCollisionQueryParams Params(AddForceOverlapName, false);
 		Params.bTraceAsyncScene = true; // want to hurt stuff in async scene
-		GetWorld()->OverlapMulti(Overlaps, Origin, FQuat::Identity, FCollisionShape::MakeSphere(Radius), Params, CollisionObjectQueryParams);
+		GetWorld()->OverlapMultiByObjectType(Overlaps, Origin, FQuat::Identity, CollisionObjectQueryParams, FCollisionShape::MakeSphere(Radius), Params);
 
 		// Iterate over each and apply force
 		for ( int32 OverlapIdx=0; OverlapIdx<Overlaps.Num(); OverlapIdx++ )
@@ -97,7 +97,7 @@ void URadialForceComponent::FireImpulse()
 		Params.AddIgnoredActor(GetOwner());
 	}
 
-	GetWorld()->OverlapMulti(Overlaps, Origin, FQuat::Identity, FCollisionShape::MakeSphere(Radius), Params, CollisionObjectQueryParams);
+	GetWorld()->OverlapMultiByObjectType(Overlaps, Origin, FQuat::Identity, CollisionObjectQueryParams, FCollisionShape::MakeSphere(Radius), Params);
 
 	// Iterate over each and apply an impulse
 	for ( int32 OverlapIdx=0; OverlapIdx<Overlaps.Num(); OverlapIdx++ )
diff --git a/Engine/Source/Runtime/Engine/Private/PlayerCameraManager.cpp b/Engine/Source/Runtime/Engine/Private/PlayerCameraManager.cpp
--- a/Engine/Source/Runtime/Engine/Private/PlayerCameraManager.cpp
+++ b/Engine/Source/Runtime/Engine/Private/PlayerCameraManager.cpp
@@ -540,7 +540,7 @@ void APlayerCameraManager::UpdateViewTarget(FTViewTarget& OutVT, float DeltaTime
 			BoxParams.AddIgnoredActor(OutVT.Target);
 			FHitResult Result;
 
-			GetWorld()->SweepSingle(Result, Loc, Pos, FQuat::Identity, ECC_Camera, FCollisionShape::MakeBox(FVector(12.f)), BoxParams);
+			GetWorld()->SweepSingleByChannel(Result, Loc, Pos, FQuat::Identity, ECC_Camera, FCollisionShape::MakeBox(FVector(12.f)), BoxParams);
 			OutVT.POV.Location = !Result.bBlockingHit ? Pos : Result.Location;
 			OutVT.POV.Rotation = Rotator;
 
diff --git a/Engine/Source/Runtime/Engine/Private/PlayerController.cpp b/Engine/Source/Runtime/Engine/Private/PlayerController.cpp
--- a/Engine/Source/Runtime/Engine/Private/PlayerController.cpp
+++ b/Engine/Source/Runtime/Engine/Private/PlayerController.cpp
@@ -1948,7 +1948,7 @@ bool APlayerController::GetHitResultAtScreenPosition(const FVector2D ScreenPosit
 			FVector WorldDirection;
 			SceneView->DeprojectFVector2D(ScreenPosition, WorldOrigin, WorldDirection);
 
-			return GetWorld()->LineTraceSingle(HitResult, WorldOrigin, WorldOrigin + WorldDirection * 100000.f, TraceChannel, CollisionQueryParams);
+			return GetWorld()->LineTraceSingleByChannel(HitResult, WorldOrigin, WorldOrigin + WorldDirection * 100000.f, TraceChannel, CollisionQueryParams);
 		}
 	}
 
@@ -2000,7 +2000,7 @@ bool APlayerController::GetHitResultAtScreenPosition(const FVector2D ScreenPosit
 
 			FCollisionObjectQueryParams ObjParam(ObjectTypes);
 			
-			return GetWorld()->LineTraceSingle(HitResult, WorldOrigin, WorldOrigin + WorldDirection * 100000.f, FCollisionQueryParams("ClickableTrace", bTraceComplex), ObjParam);
+			return GetWorld()->LineTraceSingleByObjectType(HitResult, WorldOrigin, WorldOrigin + WorldDirection * 100000.f, ObjParam, FCollisionQueryParams("ClickableTrace", bTraceComplex));
 		}
 	}
 
diff --git a/Engine/Source/Runtime/Engine/Private/SceneComponent.cpp b/Engine/Source/Runtime/Engine/Private/SceneComponent.cpp
--- a/Engine/Source/Runtime/Engine/Private/SceneComponent.cpp
+++ b/Engine/Source/Runtime/Engine/Private/SceneComponent.cpp
@@ -1488,12 +1488,12 @@ void USceneComponent::UpdatePhysicsVolume( bool bTriggerNotifiers )
 				const UPrimitiveComponent* SelfAsPrimitive = Cast<UPrimitiveComponent>(this);
 				if (SelfAsPrimitive)
 				{
-					MyWorld->ComponentOverlapMulti(Hits, SelfAsPrimitive, GetComponentLocation(), GetComponentRotation(), GetCollisionObjectType(), Params);
+					MyWorld->ComponentOverlapMultiByChannel(Hits, SelfAsPrimitive, GetComponentLocation(), GetComponentRotation(), GetCollisionObjectType(), Params);
 				}
 				else
 				{
 					bOverlappedOrigin = true;
-					MyWorld->OverlapMulti(Hits, GetComponentLocation(), FQuat::Identity, GetCollisionObjectType(), FCollisionShape::MakeSphere(0.f), Params);
+					MyWorld->OverlapMultiByChannel(Hits, GetComponentLocation(), FQuat::Identity, GetCollisionObjectType(), FCollisionShape::MakeSphere(0.f), Params);
 				}				
 				
 				for (int32 HitIdx = 0; HitIdx < Hits.Num(); HitIdx++)
diff --git a/Engine/Source/Runtime/Engine/Private/SkeletalMeshComponentPhysics.cpp b/Engine/Source/Runtime/Engine/Private/SkeletalMeshComponentPhysics.cpp
--- a/Engine/Source/Runtime/Engine/Private/SkeletalMeshComponentPhysics.cpp
+++ b/Engine/Source/Runtime/Engine/Private/SkeletalMeshComponentPhysics.cpp
@@ -3073,7 +3073,7 @@ void USkeletalMeshComponent::ProcessClothCollisionWithEnvironment()
 	static FName ClothOverlapComponentsName(TEXT("ClothOverlapComponents"));
 	FCollisionQueryParams Params(ClothOverlapComponentsName, false);
 
-	GetWorld()->OverlapMulti(Overlaps, Bounds.Origin, FQuat::Identity, FCollisionShape::MakeBox(Bounds.BoxExtent), Params, ObjectParams);
+	GetWorld()->OverlapMultiByObjectType(Overlaps, Bounds.Origin, FQuat::Identity, ObjectParams, FCollisionShape::MakeBox(Bounds.BoxExtent), Params);
 
 	for (int32 OverlapIdx=0; OverlapIdx<Overlaps.Num(); ++OverlapIdx)
 	{
diff --git a/Engine/Source/Runtime/Engine/Private/SoundBase.cpp b/Engine/Source/Runtime/Engine/Private/SoundBase.cpp
--- a/Engine/Source/Runtime/Engine/Private/SoundBase.cpp
+++ b/Engine/Source/Runtime/Engine/Private/SoundBase.cpp
@@ -77,7 +77,7 @@ bool USoundBase::IsAudible( const FVector &SourceLocation, const FVector &Listen
 			static FName NAME_IsAudible(TEXT("IsAudible"));
 
 			// simple trace occlusion check - reduce max audible distance if occluded
-			bIsOccluded = SourceActor->GetWorld()->LineTraceTest(ModifiedSourceLocation, ListenerLocation, ECC_Visibility, FCollisionQueryParams(NAME_IsAudible, true, SourceActor));
+			bIsOccluded = SourceActor->GetWorld()->LineTraceTestByChannel(ModifiedSourceLocation, ListenerLocation, ECC_Visibility, FCollisionQueryParams(NAME_IsAudible, true, SourceActor));
 		}
 		return true;
 	}
diff --git a/Engine/Source/Runtime/Engine/Private/StaticMeshActor.cpp b/Engine/Source/Runtime/Engine/Private/StaticMeshActor.cpp
--- a/Engine/Source/Runtime/Engine/Private/StaticMeshActor.cpp
+++ b/Engine/Source/Runtime/Engine/Private/StaticMeshActor.cpp
@@ -145,7 +145,7 @@ void AStaticMeshActor::CheckForErrors()
 		FCollisionQueryParams SphereParams(FName(TEXT("CheckForErrors")), false, this);
 
 		TArray<FOverlapResult> Overlaps;
-		GetWorld()->OverlapMulti(Overlaps, GetActorLocation(), FQuat::Identity, ECC_Pawn, FCollisionShape::MakeSphere(1.f), SphereParams);
+		GetWorld()->OverlapMultiByChannel(Overlaps, GetActorLocation(), FQuat::Identity, ECC_Pawn, FCollisionShape::MakeSphere(1.f), SphereParams);
 
 		for ( int32 OverlapIdx=0; OverlapIdx<Overlaps.Num(); OverlapIdx++ )
 		{
diff --git a/Engine/Source/Runtime/Engine/Public/CollisionQueryParams.h b/Engine/Source/Runtime/Engine/Public/CollisionQueryParams.h
--- a/Engine/Source/Runtime/Engine/Public/CollisionQueryParams.h
+++ b/Engine/Source/Runtime/Engine/Public/CollisionQueryParams.h
@@ -79,6 +79,9 @@ struct ENGINE_API FCollisionQueryParams
 	{
 		return FString::Printf(TEXT("[%s:%s] TraceAsync(%d), TraceComplex(%d)"), *OwnerTag.ToString(), *TraceTag.ToString(), bTraceAsyncScene, bTraceComplex );
 	}
+
+	/** static variable for default data to be used without reconstructing everytime **/
+	static FCollisionQueryParams DefaultQueryParam;
 };
 
 /** Structure when performing a collision query using a component's geometry */
@@ -93,6 +96,9 @@ struct ENGINE_API FComponentQueryParams : public FCollisionQueryParams
 	: FCollisionQueryParams(InTraceTag, false, InIgnoreActor)
 	{
 	}
+
+	/** static variable for default data to be used without reconstructing everytime **/
+	static FComponentQueryParams DefaultComponentQueryParams;
 };
 
 /** Structure that defines response container for the query. Advanced option. */
diff --git a/Engine/Source/Runtime/Foliage/Private/InstancedFoliage.cpp b/Engine/Source/Runtime/Foliage/Private/InstancedFoliage.cpp
--- a/Engine/Source/Runtime/Foliage/Private/InstancedFoliage.cpp
+++ b/Engine/Source/Runtime/Foliage/Private/InstancedFoliage.cpp
@@ -1740,7 +1740,7 @@ void AInstancedFoliageActor::MapRebuild()
 					FVector End(InstanceToWorld.TransformPosition(Down));
 
 					FHitResult Result;
-					bool bHit = World->LineTraceSingle(Result, Start, End, FCollisionQueryParams(true), FCollisionObjectQueryParams(ECC_WorldStatic));
+					bool bHit = World->LineTraceSingleByObjectType(Result, Start, End, FCollisionObjectQueryParams(ECC_WorldStatic), FCollisionQueryParams(true));
 
 					if (bHit && Result.Component.IsValid() && Result.Component->IsA(UModelComponent::StaticClass()))
 					{
@@ -2104,7 +2104,7 @@ bool AInstancedFoliageActor::FoliageTrace(const UWorld* InWorld, FHitResult& Out
 	bool bInsideProceduralVolume = false;
 	FCollisionShape SphereShape;
 	SphereShape.SetSphere(DesiredInstance.TraceRadius);
-	InWorld->SweepMulti(Hits, StartTrace, DesiredInstance.EndTrace, FQuat::Identity, SphereShape, QueryParams, FCollisionObjectQueryParams(ECC_WorldStatic));
+	InWorld->SweepMultiByObjectType(Hits, StartTrace, DesiredInstance.EndTrace, FQuat::Identity, FCollisionObjectQueryParams(ECC_WorldStatic), SphereShape, QueryParams);
 
 
 	for (const FHitResult& Hit : Hits)
@@ -2196,7 +2196,7 @@ bool AInstancedFoliageActor::CheckCollisionWithWorld(const UWorld* InWorld, cons
 			FBoxSphereBounds WorldBound = ShrinkBound.TransformBy(InstTransform);
 			//::DrawDebugBox(World, WorldBound.Origin, WorldBound.BoxExtent, FColor::Red, true, 10.f);
 			static FName NAME_FoliageCollisionWithWorld = FName(TEXT("FoliageCollisionWithWorld"));
-			if (InWorld->OverlapBlockingTest(WorldBound.Origin, FQuat(Inst.Rotation), ECC_WorldStatic, FCollisionShape::MakeBox(ShrinkBound.BoxExtent * Inst.DrawScale3D * Settings->CollisionScale), FCollisionQueryParams(NAME_FoliageCollisionWithWorld, false)))
+			if (InWorld->OverlapBlockingTestByChannel(WorldBound.Origin, FQuat(Inst.Rotation), ECC_WorldStatic, FCollisionShape::MakeBox(ShrinkBound.BoxExtent * Inst.DrawScale3D * Settings->CollisionScale), FCollisionQueryParams(NAME_FoliageCollisionWithWorld, false)))
 			{
 				return false;
 			}
diff --git a/Engine/Source/Runtime/GameplayAbilities/Private/Abilities/GameplayAbilityTargetActor_Radius.cpp b/Engine/Source/Runtime/GameplayAbilities/Private/Abilities/GameplayAbilityTargetActor_Radius.cpp
--- a/Engine/Source/Runtime/GameplayAbilities/Private/Abilities/GameplayAbilityTargetActor_Radius.cpp
+++ b/Engine/Source/Runtime/GameplayAbilities/Private/Abilities/GameplayAbilityTargetActor_Radius.cpp
@@ -60,7 +60,7 @@ TArray<TWeakObjectPtr<AActor> >	AGameplayAbilityTargetActor_Radius::PerformOverl
 
 	TArray<FOverlapResult> Overlaps;
 
-	SourceActor->GetWorld()->OverlapMulti(Overlaps, Origin, FQuat::Identity, ECC_Pawn, FCollisionShape::MakeSphere(Radius), Params);
+	SourceActor->GetWorld()->OverlapMultiByChannel(Overlaps, Origin, FQuat::Identity, ECC_Pawn, FCollisionShape::MakeSphere(Radius), Params);
 
 	TArray<TWeakObjectPtr<AActor>>	HitActors;
 
diff --git a/Engine/Source/Runtime/Landscape/Private/LandscapeCollision.cpp b/Engine/Source/Runtime/Landscape/Private/LandscapeCollision.cpp
--- a/Engine/Source/Runtime/Landscape/Private/LandscapeCollision.cpp
+++ b/Engine/Source/Runtime/Landscape/Private/LandscapeCollision.cpp
@@ -1054,7 +1054,7 @@ void ULandscapeHeightfieldCollisionComponent::SnapFoliageInstances(AInstancedFol
 					UWorld* World = GetWorld();
 					check(World);
 					// Editor specific landscape heightfield uses ECC_Visibility collision channel
-					World->LineTraceMulti(Results, Start, End, FCollisionQueryParams(TraceTag, true), FCollisionObjectQueryParams(ECollisionChannel::ECC_Visibility));
+					World->LineTraceMultiByObjectType(Results, Start, End, FCollisionObjectQueryParams(ECollisionChannel::ECC_Visibility), FCollisionQueryParams(TraceTag, true));
 
 					bool bFoundHit = false;
 					for (const FHitResult& Hit : Results)
