Framework: Implemented USceneComponent::GetAllSocketNames() in terms of USceneComponent::QuerySupportedSockets() and removed pointless virtual + overrides

diff --git a/Engine/Source/Runtime/Engine/Classes/Components/SceneComponent.h b/Engine/Source/Runtime/Engine/Classes/Components/SceneComponent.h
--- a/Engine/Source/Runtime/Engine/Classes/Components/SceneComponent.h
+++ b/Engine/Source/Runtime/Engine/Classes/Components/SceneComponent.h
@@ -425,7 +425,7 @@ public:
 	 * @return Get the names of all the sockets on the component.
 	 */
 	UFUNCTION(BlueprintCallable, Category="Utilities|Transformation", meta=(Keywords="Bone"))
-	virtual TArray<FName> GetAllSocketNames() const;
+	TArray<FName> GetAllSocketNames() const;
 
 	/** 
 	 * Get world-space socket transform.
diff --git a/Engine/Source/Runtime/Engine/Classes/Components/SkinnedMeshComponent.h b/Engine/Source/Runtime/Engine/Classes/Components/SkinnedMeshComponent.h
--- a/Engine/Source/Runtime/Engine/Classes/Components/SkinnedMeshComponent.h
+++ b/Engine/Source/Runtime/Engine/Classes/Components/SkinnedMeshComponent.h
@@ -377,7 +377,7 @@ public:
 	UFUNCTION(BlueprintCallable, Category="Components|SkinnedMesh")
 	FName GetBoneName(int32 BoneIndex) const;
 
-		/**
+	/**
 	 * Returns bone name linked to a given named socket on the skeletal mesh component.
 	 * If you're unsure to deal with sockets or bones names, you can use this function to filter through, and always return the bone name.
 	 *
@@ -445,6 +445,8 @@ protected:
 public:
 	// Begin USceneComponent interface
 	virtual FBoxSphereBounds CalcBounds(const FTransform& LocalToWorld) const override;
+	virtual FTransform GetSocketTransform(FName InSocketName, ERelativeTransformSpace TransformSpace = RTS_World) const override;
+	virtual bool DoesSocketExist(FName InSocketName) const override;
 	virtual bool HasAnySockets() const override;
 	virtual void QuerySupportedSockets(TArray<FComponentSocketDescription>& OutSockets) const override;
 	virtual void UpdateOverlaps(TArray<FOverlapInfo> const* PendingOverlaps=NULL, bool bDoNotifies=true, const TArray<FOverlapInfo>* OverlapsAtEndLocation=NULL) override;
@@ -668,24 +670,11 @@ public:
 	 */
 	void UpdateMasterBoneMap();
 
-	//
-	// Get all socket names.
-	//
-
-	virtual TArray<FName> GetAllSocketNames() const override;
-
-	//
-	// Bone Transform.
-	//
-
-	virtual FTransform GetSocketTransform(FName InSocketName, ERelativeTransformSpace TransformSpace = RTS_World) const override;
-
 	/**
 	 * @return SkeletalMeshSocket of named socket on the skeletal mesh component, or NULL if not found.
 	 */
 	class USkeletalMeshSocket const* GetSocketByName( FName InSocketName ) const;
 
-	virtual bool DoesSocketExist(FName InSocketName) const override;
 
 	/** 
 	 * Get Bone Matrix from index
diff --git a/Engine/Source/Runtime/Engine/Classes/Components/StaticMeshComponent.h b/Engine/Source/Runtime/Engine/Classes/Components/StaticMeshComponent.h
--- a/Engine/Source/Runtime/Engine/Classes/Components/StaticMeshComponent.h
+++ b/Engine/Source/Runtime/Engine/Classes/Components/StaticMeshComponent.h
@@ -246,7 +246,6 @@ public:
 		// return IsCollisionEnabled() && (StaticMesh != NULL);
 		return false;
 	}
-	virtual TArray<FName> GetAllSocketNames() const override;
 	// End USceneComponent Interface
 
 	// Begin UActorComponent interface.
@@ -439,12 +438,6 @@ public:
 	 */
 	class UStaticMeshSocket const* GetSocketByName( FName InSocketName ) const;
 
-	/**
-	 * Returns true if component is attached to the static mesh.
-	 * @return	true if Component is attached to StaticMesh.
-	 */
-	bool IsComponentAttached( FName SocketName = NAME_None );
-
 	/** Returns the wireframe color to use for this component. */
 	FColor GetWireframeColor() const;
 
diff --git a/Engine/Source/Runtime/Engine/Private/SceneComponent.cpp b/Engine/Source/Runtime/Engine/Private/SceneComponent.cpp
--- a/Engine/Source/Runtime/Engine/Private/SceneComponent.cpp
+++ b/Engine/Source/Runtime/Engine/Private/SceneComponent.cpp
@@ -1263,9 +1263,9 @@ void USceneComponent::PostInterpChange(UProperty* PropertyThatChanged)
 {
 	Super::PostInterpChange(PropertyThatChanged);
 
-	static FName RelativeScale3D(TEXT("RelativeScale3D"));
+	static FName RelativeScale3D(GET_MEMBER_NAME_CHECKED(USceneComponent, RelativeScale3D));
 
-	if(PropertyThatChanged->GetFName() == RelativeScale3D)
+	if (PropertyThatChanged->GetFName() == RelativeScale3D)
 	{
 		UpdateComponentToWorld();
 	}
@@ -1273,7 +1273,18 @@ void USceneComponent::PostInterpChange(UProperty* PropertyThatChanged)
 
 TArray<FName> USceneComponent::GetAllSocketNames() const
 {
-	return TArray<FName>();
+	TArray<FComponentSocketDescription> SocketList;
+	QuerySupportedSockets(/*out*/ SocketList);
+
+	TArray<FName> ResultList;
+	ResultList.Reserve(SocketList.Num());
+
+	for (const FComponentSocketDescription& SocketDesc : SocketList)
+	{
+		ResultList.Add(SocketDesc.Name);
+	}
+
+	return ResultList;
 }
 
 FTransform USceneComponent::GetSocketTransform(FName SocketName, ERelativeTransformSpace TransformSpace) const
diff --git a/Engine/Source/Runtime/Engine/Private/SkinnedMeshComponent.cpp b/Engine/Source/Runtime/Engine/Private/SkinnedMeshComponent.cpp
--- a/Engine/Source/Runtime/Engine/Private/SkinnedMeshComponent.cpp
+++ b/Engine/Source/Runtime/Engine/Private/SkinnedMeshComponent.cpp
@@ -1100,19 +1100,6 @@ void USkinnedMeshComponent::UpdateMasterBoneMap()
 	}
 }
 
-TArray<FName> USkinnedMeshComponent::GetAllSocketNames() const
-{
-	TArray<FName> SocketNames;
-	if( SkeletalMesh )
-	{
-		for( auto It=SkeletalMesh->Skeleton->Sockets.CreateConstIterator(); It; ++It )
-		{
-			SocketNames.Add( (*It)->SocketName ) ;
-		}
-	}
-	return SocketNames;
-}
-
 FTransform USkinnedMeshComponent::GetSocketTransform(FName InSocketName, ERelativeTransformSpace TransformSpace) const
 {
 	FTransform OutSocketTransform = ComponentToWorld;
diff --git a/Engine/Source/Runtime/Engine/Private/StaticMeshComponent.cpp b/Engine/Source/Runtime/Engine/Private/StaticMeshComponent.cpp
--- a/Engine/Source/Runtime/Engine/Private/StaticMeshComponent.cpp
+++ b/Engine/Source/Runtime/Engine/Private/StaticMeshComponent.cpp
@@ -203,19 +203,6 @@ void UStaticMeshComponent::QuerySupportedSockets(TArray<FComponentSocketDescript
 	}
 }
 
-TArray<FName> UStaticMeshComponent::GetAllSocketNames() const
-{
-	TArray<FName> SocketNames;
-	if( StaticMesh )
-	{
-		for( auto It=StaticMesh->Sockets.CreateConstIterator(); It; ++It )
-		{
-			SocketNames.Add( (*It)->SocketName ) ;
-		}
-	}
-	return SocketNames;
-}
-
 FString UStaticMeshComponent::GetDetailedInfoInternal() const
 {
 	FString Result;  
