fixed postponed path support in navigation path update (invalidated/goal moved)
#ue4

diff --git a/Engine/Source/Runtime/Engine/Private/AI/Navigation/NavigationData.cpp b/Engine/Source/Runtime/Engine/Private/AI/Navigation/NavigationData.cpp
--- a/Engine/Source/Runtime/Engine/Private/AI/Navigation/NavigationData.cpp
+++ b/Engine/Source/Runtime/Engine/Private/AI/Navigation/NavigationData.cpp
@@ -251,6 +251,7 @@ void ANavigationData::TickActor(float DeltaTime, enum ELevelTick TickType, FActo
 			if (PathNavAgent && PathNavAgent->ShouldPostponePathUpdates())
 			{
 				PostponedRequests.Add(RecalcRequest);
+				continue;
 			}
 
 			FPathFindingQuery Query(RecalcRequest.Path);
