Fixed incorrect comments on MoveComponent in SceneComponent, PrimitiveComponent and Actor.

diff --git a/Engine/Source/Runtime/Engine/Classes/Components/PrimitiveComponent.h b/Engine/Source/Runtime/Engine/Classes/Components/PrimitiveComponent.h
--- a/Engine/Source/Runtime/Engine/Classes/Components/PrimitiveComponent.h
+++ b/Engine/Source/Runtime/Engine/Classes/Components/PrimitiveComponent.h
@@ -1305,6 +1305,18 @@ public:
 	// End UObject interface.
 
 	//Begin USceneComponent Interface
+	/**
+	 * Tries to move the component by a movement vector (Delta) and sets rotation to NewRotation.
+	 * Assumes that the component's current location is valid and that the component does fit in its current Location.
+	 * Dispatches blocking hit notifications (if bSweep is true), and calls UpdateOverlaps() after movement to update overlap state.
+	 * 
+	 * @param Delta			The desired location change in world space.
+	 * @param NewRotation	The new desired rotation in world space.
+	 * @param bSweep		Should we sweep to the destination location, stopping short of the target if blocked by something. Note:If the component has no collision this will have no effect.
+	 * @param Hit			Optional output describing the blocking hit that stopped the move, if any.
+	 * @param MoveFlags		Flags controlling behavior of the move. @see EMoveComponentFlags
+	 * @return				True if some movement occurred, false if no movement occurred.
+	 */
 	virtual bool MoveComponent(const FVector& Delta, const FRotator& NewRotation, bool bSweep, FHitResult* OutHit = NULL, EMoveComponentFlags MoveFlags = MOVECOMP_NoFlags) override;
 	virtual bool IsWorldGeometry() const override;
 	virtual ECollisionEnabled::Type GetCollisionEnabled() const override;
diff --git a/Engine/Source/Runtime/Engine/Classes/GameFramework/Actor.h b/Engine/Source/Runtime/Engine/Classes/GameFramework/Actor.h
--- a/Engine/Source/Runtime/Engine/Classes/GameFramework/Actor.h
+++ b/Engine/Source/Runtime/Engine/Classes/GameFramework/Actor.h
@@ -620,7 +620,7 @@ public:
 	/** 
 	 *	Move the Actor to the specified location.
 	 *	@param NewLocation	The new location to move the Actor to.
-	 *	@param bSweep		Should we sweep to the destination location, stopping short of the target if blocked by something.
+	 *	@param bSweep		Should we sweep to the destination location, stopping short of the target if blocked by something. Note:If the root component has no collision this will have no effect.
 	 *  @param SweepHitResult	The hit result from the move if swept.
 	 *	@return	Whether the location was successfully set (if not swept), or whether movement occurred at all (if swept).
 	 */
