Absolute (int) and Sign (int) functions updated to normal naming conventions in the Kismet Math Library.

PR #853: Update KismetMathLibrary.h equal int ABS (Contributed by [REDACTED_PERSON])

diff --git a/Engine/Source/Runtime/Engine/Classes/Kismet/KismetMathLibrary.h b/Engine/Source/Runtime/Engine/Classes/Kismet/KismetMathLibrary.h
--- a/Engine/Source/Runtime/Engine/Classes/Kismet/KismetMathLibrary.h
+++ b/Engine/Source/Runtime/Engine/Classes/Kismet/KismetMathLibrary.h
@@ -201,7 +201,7 @@ class UKismetMathLibrary : public UBlueprintFunctionLibrary
 	static int32 Or_IntInt(int32 A, int32 B);
 
 	/* Sign (integer, returns -1 if A < 0, 0 if A is zero, and +1 if A > 0) */
-	UFUNCTION(BlueprintPure, meta=(FriendlyName = "Sign (integer)"), Category="Math|Integer")
+	UFUNCTION(BlueprintPure, meta=(FriendlyName = "Sign (int)"), Category="Math|Integer")
 	static int32 SignOfInteger(int32 A);
 
 	/* Returns a uniformly distributed random number between 0 and Max - 1 */
@@ -225,7 +225,7 @@ class UKismetMathLibrary : public UBlueprintFunctionLibrary
 	static int32 Clamp(int32 Value, int32 Min, int32 Max);
 
 	/* Returns the absolute (positive) value of A */
-	UFUNCTION(BlueprintPure, meta=(FriendlyName = "Absolute (int)"), Category="Math|Integer")
+	UFUNCTION(BlueprintPure, meta=(FriendlyName = "Absolute (int)", CompactNodeTitle = "ABS"), Category="Math|Integer")
 	static int32 Abs_Int(int32 A);
 
 
