Exposed constraint projection so it can be altered by the runtime as well as the flag in the physics asset

diff --git a/Engine/Source/Runtime/Engine/Classes/PhysicsEngine/ConstraintInstance.h b/Engine/Source/Runtime/Engine/Classes/PhysicsEngine/ConstraintInstance.h
--- a/Engine/Source/Runtime/Engine/Classes/PhysicsEngine/ConstraintInstance.h
+++ b/Engine/Source/Runtime/Engine/Classes/PhysicsEngine/ConstraintInstance.h
@@ -488,6 +488,10 @@ public:
 
 	void OnConstraintBroken();
 
+	void EnableProjection();
+
+	void DisableProjection();
+
 	//Hacks to easily get zeroed memory for special case when we don't use GC
 	static void Free(FConstraintInstance * Ptr);
 	static FConstraintInstance * Alloc();
diff --git a/Engine/Source/Runtime/Engine/Private/PhysicsEngine/ConstraintInstance.cpp b/Engine/Source/Runtime/Engine/Private/PhysicsEngine/ConstraintInstance.cpp
--- a/Engine/Source/Runtime/Engine/Private/PhysicsEngine/ConstraintInstance.cpp
+++ b/Engine/Source/Runtime/Engine/Private/PhysicsEngine/ConstraintInstance.cpp
@@ -1304,4 +1304,18 @@ FConstraintInstance * FConstraintInstance::Alloc()
 	return new (Memory)FConstraintInstance();
 }
 
+void FConstraintInstance::EnableProjection()
+{
+	SCOPED_SCENE_WRITE_LOCK(ConstraintData->getScene());
+	ConstraintData->setProjectionLinearTolerance(ProjectionLinearTolerance);
+	ConstraintData->setProjectionAngularTolerance(ProjectionAngularTolerance);
+	ConstraintData->setConstraintFlag(PxConstraintFlag::ePROJECTION, true);
+}
+
+void FConstraintInstance::DisableProjection()
+{
+	SCOPED_SCENE_WRITE_LOCK(ConstraintData->getScene());
+	ConstraintData->setConstraintFlag(PxConstraintFlag::ePROJECTION, false);
+}
+
 #undef LOCTEXT_NAMESPACE
\ No newline at end of file
