Make AudioComponent IsPlaying a const/pure function

diff --git a/Engine/Source/Runtime/Engine/Classes/Components/AudioComponent.h b/Engine/Source/Runtime/Engine/Classes/Components/AudioComponent.h
--- a/Engine/Source/Runtime/Engine/Classes/Components/AudioComponent.h
+++ b/Engine/Source/Runtime/Engine/Classes/Components/AudioComponent.h
@@ -190,7 +190,7 @@ class ENGINE_API UAudioComponent : public USceneComponent
 
 	/** @return true if this component is currently playing a SoundCue. */
 	UFUNCTION(BlueprintCallable, Category="Audio|Components|Audio")
-	bool IsPlaying();
+	bool IsPlaying() const;
 
 	/** This will allow one to adjust the volume of an AudioComponent on the fly */
 	UFUNCTION(BlueprintCallable, Category="Audio|Components|Audio")
diff --git a/Engine/Source/Runtime/Engine/Private/AudioComponent.cpp b/Engine/Source/Runtime/Engine/Private/AudioComponent.cpp
--- a/Engine/Source/Runtime/Engine/Private/AudioComponent.cpp
+++ b/Engine/Source/Runtime/Engine/Private/AudioComponent.cpp
@@ -287,7 +287,7 @@ void UAudioComponent::PlaybackCompleted(bool bFailedToStart)
 	}
 }
 
-bool UAudioComponent::IsPlaying( void )
+bool UAudioComponent::IsPlaying() const
 {
 	return bIsActive;
 }
