Fixed cross product's compact display name for 3D vectors.

diff --git a/Engine/Source/Runtime/Engine/Classes/Kismet/KismetMathLibrary.h b/Engine/Source/Runtime/Engine/Classes/Kismet/KismetMathLibrary.h
--- a/Engine/Source/Runtime/Engine/Classes/Kismet/KismetMathLibrary.h
+++ b/Engine/Source/Runtime/Engine/Classes/Kismet/KismetMathLibrary.h
@@ -571,7 +571,7 @@ class UKismetMathLibrary : public UBlueprintFunctionLibrary
 	static float Dot_VectorVector(FVector A, FVector B);
 
 	/* Returns the cross product of two 3d vectors */
-	UFUNCTION(BlueprintPure, meta=(FriendlyName = "Cross Product", CompactNodeTitle = ""), Category="Math|Vector" )
+	UFUNCTION(BlueprintPure, meta=(FriendlyName = "Cross Product", CompactNodeTitle = "x"), Category="Math|Vector" )
 	static FVector Cross_VectorVector(FVector A, FVector B);
 
 	/* Returns the dot product of two 2d vectors */
