Don't register (non-transactional) blueprint constructed components that are about to be destroyed after an undo

diff --git a/Engine/Source/Runtime/Engine/Private/ActorEditor.cpp b/Engine/Source/Runtime/Engine/Private/ActorEditor.cpp
--- a/Engine/Source/Runtime/Engine/Private/ActorEditor.cpp
+++ b/Engine/Source/Runtime/Engine/Private/ActorEditor.cpp
@@ -37,9 +37,69 @@ void AActor::PostEditChangeProperty(FPropertyChangedEvent& PropertyChangedEvent)
 
 	if ( ReregisterComponentsWhenModified() )
 	{
-		ReregisterAllComponents();
+		// In the Undo case we have an annotation storing information about constructed components and we do not want
+		// to improperly apply out of date changes so we need to skip registration of all blueprint created components
+		// and defer instance components attached to them until after rerun
+		if (CurrentTransactionAnnotation.IsValid())
+		{
+			UnregisterAllComponents();
+
+			TInlineComponentArray<UActorComponent*> Components;
+			GetComponents(Components);
+
+			Components.Sort([](UActorComponent& A, UActorComponent& B)
+			{
+				if (&B == B.GetOwner()->GetRootComponent())
+				{
+					return false;
+				}
+				if (USceneComponent* ASC = Cast<USceneComponent>(&A))
+				{
+					if (ASC->AttachParent == &B)
+					{
+						return false;
+					}
+				}
+				return true;
+			});
 
-		RerunConstructionScripts();
+			bool bRequiresReregister = false;
+			for (UActorComponent* Component : Components)
+			{
+				if (Component->CreationMethod == EComponentCreationMethod::Native)
+				{
+					Component->RegisterComponent();
+				}
+				else if (Component->CreationMethod == EComponentCreationMethod::Instance)
+				{
+					USceneComponent* SC = Cast<USceneComponent>(Component);
+					if (SC == nullptr || SC == RootComponent || (SC->AttachParent && SC->AttachParent->IsRegistered()))
+					{
+						Component->RegisterComponent();
+					}
+					else
+					{
+						bRequiresReregister = true;
+					}
+				}
+				else
+				{
+					bRequiresReregister = true;
+				}
+			}
+
+			RerunConstructionScripts();
+
+			if (bRequiresReregister)
+			{
+				ReregisterAllComponents();
+			}
+		}
+		else
+		{
+			ReregisterAllComponents();
+			RerunConstructionScripts();
+		}
 	}
 
 	// Let other systems know that an actor was moved
