Engine: Made FFilePath and FDirectoryPath blueprint types

diff --git a/Engine/Source/Runtime/Engine/Classes/Engine/EngineTypes.h b/Engine/Source/Runtime/Engine/Classes/Engine/EngineTypes.h
--- a/Engine/Source/Runtime/Engine/Classes/Engine/EngineTypes.h
+++ b/Engine/Source/Runtime/Engine/Classes/Engine/EngineTypes.h
@@ -2894,7 +2894,7 @@ enum EAngularConstraintMotion
 /**
  * Structure for file paths that are displayed in the UI.
  */
-USTRUCT()
+USTRUCT(BlueprintType)
 struct FFilePath
 {
 	GENERATED_USTRUCT_BODY()
@@ -2902,7 +2902,7 @@ struct FFilePath
 	/**
 	 * The path to the file.
 	 */
-	UPROPERTY(EditAnywhere, Category=FilePath)
+	UPROPERTY(EditAnywhere, BlueprintReadWrite, Category=FilePath)
 	FString FilePath;
 };
 
@@ -2910,7 +2910,7 @@ struct FFilePath
 /**
  * Structure for directory paths that are displayed in the UI.
  */
-USTRUCT()
+USTRUCT(BlueprintType)
 struct FDirectoryPath
 {
 	GENERATED_USTRUCT_BODY()
@@ -2918,7 +2918,7 @@ struct FDirectoryPath
 	/**
 	 * The path to the directory.
 	 */
-	UPROPERTY(EditAnywhere, Category=Path)
+	UPROPERTY(EditAnywhere, BlueprintReadWrite, Category=Path)
 	FString Path;
 };
 
