Animation: Added stats for UAnimInstance construction and init time

diff --git a/Engine/Source/Runtime/Engine/Private/Animation/AnimInstance.cpp b/Engine/Source/Runtime/Engine/Private/Animation/AnimInstance.cpp
--- a/Engine/Source/Runtime/Engine/Private/Animation/AnimInstance.cpp
+++ b/Engine/Source/Runtime/Engine/Private/Animation/AnimInstance.cpp
@@ -46,6 +46,9 @@ DEFINE_STAT(STAT_TickUpdateRate);
 DEFINE_STAT(STAT_PerformAnimEvaluation);
 DEFINE_STAT(STAT_PostAnimEvaluation);
 
+DECLARE_CYCLE_STAT_EXTERN(TEXT("Anim Init Time"), STAT_AnimInitTime, STATGROUP_Anim, );
+DEFINE_STAT(STAT_AnimInitTime);
+
 DEFINE_STAT(STAT_AnimStateMachineUpdate);
 DEFINE_STAT(STAT_AnimStateMachineFindTransition);
 DEFINE_STAT(STAT_AnimStateMachineEvaluate);
@@ -227,6 +230,8 @@ UWorld* UAnimInstance::GetWorld() const
 
 void UAnimInstance::InitializeAnimation()
 {
+	SCOPE_CYCLE_COUNTER(STAT_AnimInitTime);
+
 	// make sure your skeleton is initialized
 	// you can overwrite different skeleton
 	USkeletalMeshComponent* OwnerComponent = GetSkelMeshComponent();
diff --git a/Engine/Source/Runtime/Engine/Private/SkeletalMeshComponent.cpp b/Engine/Source/Runtime/Engine/Private/SkeletalMeshComponent.cpp
--- a/Engine/Source/Runtime/Engine/Private/SkeletalMeshComponent.cpp
+++ b/Engine/Source/Runtime/Engine/Private/SkeletalMeshComponent.cpp
@@ -28,6 +28,9 @@
 
 TAutoConsoleVariable<int32> CVarUseParallelAnimationEvaluation(TEXT("a.ParallelAnimEvaluation"), 1, TEXT("If 1, animation evaluation will be run across the task graph system. If 0, evaluation will run purely on the game thread"));
 
+DECLARE_CYCLE_STAT_EXTERN(TEXT("Anim Instance Spawn Time"), STAT_AnimSpawnTime, STATGROUP_Anim, );
+DEFINE_STAT(STAT_AnimSpawnTime);
+
 class FParallelAnimationEvaluationTask
 {
 	TWeakObjectPtr<USkeletalMeshComponent> SkeletalMeshComponent;
@@ -244,6 +247,7 @@ void USkeletalMeshComponent::InitializeAnimScriptInstance(bool bForceReinit)
 	{
 		if (NeedToSpawnAnimScriptInstance(bForceReinit))
 		{
+			SCOPE_CYCLE_COUNTER(STAT_AnimSpawnTime);
 			AnimScriptInstance = NewObject<UAnimInstance>(this, AnimBlueprintGeneratedClass);
 
 			if (AnimScriptInstance)
@@ -253,6 +257,8 @@ void USkeletalMeshComponent::InitializeAnimScriptInstance(bool bForceReinit)
 		}
 		else if (AnimationMode == EAnimationMode::AnimationSingleNode)
 		{
+			SCOPE_CYCLE_COUNTER(STAT_AnimSpawnTime);
+
 			UAnimSingleNodeInstance* OldInstance = NULL;
 			if (!bForceReinit)
 			{
@@ -363,6 +369,7 @@ void USkeletalMeshComponent::PostEditChangeProperty(FPropertyChangedEvent& Prope
 				{
 					if (NeedToSpawnAnimScriptInstance(false))
 					{
+						SCOPE_CYCLE_COUNTER(STAT_AnimSpawnTime);
 						AnimScriptInstance = NewObject<UAnimInstance>(this, AnimBlueprintGeneratedClass);
 						AnimScriptInstance->InitializeAnimation();
 					}
