Fix race conditions when updating textures from an SWebBrowser.  Introduced a new thread safe way to update a texture (this is expensive as it requires an extra memcopy).

diff --git a/Engine/Source/Developer/StandaloneRenderer/Private/OpenGL/SlateOpenGLTextures.h b/Engine/Source/Developer/StandaloneRenderer/Private/OpenGL/SlateOpenGLTextures.h
--- a/Engine/Source/Developer/StandaloneRenderer/Private/OpenGL/SlateOpenGLTextures.h
+++ b/Engine/Source/Developer/StandaloneRenderer/Private/OpenGL/SlateOpenGLTextures.h
@@ -32,7 +32,7 @@ public:
 	virtual FSlateShaderResource* GetSlateResource() override {return this;}
 	virtual void ResizeTexture( uint32 Width, uint32 Height ) override;
 	virtual void UpdateTexture(const TArray<uint8>& Bytes) override;
-
+	virtual void UpdateTextureThreadSafe(const TArray<uint8>& Bytes) {}
 private:
 	static GLuint NullTexture;
 
diff --git a/Engine/Source/Developer/StandaloneRenderer/Private/Windows/D3D/SlateD3DTextures.h b/Engine/Source/Developer/StandaloneRenderer/Private/Windows/D3D/SlateD3DTextures.h
--- a/Engine/Source/Developer/StandaloneRenderer/Private/Windows/D3D/SlateD3DTextures.h
+++ b/Engine/Source/Developer/StandaloneRenderer/Private/Windows/D3D/SlateD3DTextures.h
@@ -33,6 +33,7 @@ public:
 	virtual FSlateShaderResource* GetSlateResource() override {return this;}
 	virtual void ResizeTexture( uint32 Width, uint32 Height ) override;
 	virtual void UpdateTexture(const TArray<uint8>& Bytes) override;
+	virtual void UpdateTextureThreadSafe(const TArray<uint8>& Bytes) {}
 private:
 	/** Actual texture.  In D3D the SRV is used by the shader so our parent class holds that */
 	TRefCountPtr<ID3D11Texture2D> D3DTexture;
diff --git a/Engine/Source/Runtime/Engine/Private/Slate/SlateTextures.cpp b/Engine/Source/Runtime/Engine/Private/Slate/SlateTextures.cpp
--- a/Engine/Source/Runtime/Engine/Private/Slate/SlateTextures.cpp
+++ b/Engine/Source/Runtime/Engine/Private/Slate/SlateTextures.cpp
@@ -162,6 +162,26 @@ void FSlateTexture2DRHIRef::UpdateTexture(const TArray<uint8>& Bytes)
 	}
 }
 
+void FSlateTexture2DRHIRef::UpdateTextureThreadSafe(const TArray<uint8>& Bytes)
+{
+	if(IsInGameThread())
+	{
+		// Make bulk data for updating the texture memory later
+		FSlateTextureData* BulkData = new FSlateTextureData( Bytes.Num(), 0, 1, Bytes );
+
+		// Update the texture RHI
+		ENQUEUE_UNIQUE_RENDER_COMMAND_TWOPARAMETER(
+			FSlateTexture2DRHIRef_UpdateTextureThreadSafe,
+			FSlateTexture2DRHIRef*, ThisTexture, this,
+			FSlateTextureData*, BulkData, BulkData,
+			{
+				ThisTexture->UpdateTexture( BulkData->GetRawBytes() );
+				delete BulkData;
+			});
+	}
+}
+
+
 void FSlateRenderTargetRHI::SetRHIRef( FTexture2DRHIRef InRenderTargetTexture, uint32 InWidth, uint32 InHeight )
 {
 	check( IsInRenderingThread() );
diff --git a/Engine/Source/Runtime/Engine/Public/Slate/SlateTextures.h b/Engine/Source/Runtime/Engine/Public/Slate/SlateTextures.h
--- a/Engine/Source/Runtime/Engine/Public/Slate/SlateTextures.h
+++ b/Engine/Source/Runtime/Engine/Public/Slate/SlateTextures.h
@@ -73,7 +73,7 @@ public:
 	virtual FRenderResource* GetRenderResource() override {return this;}
 	virtual void ResizeTexture( uint32 Width, uint32 Height ) override;
 	virtual void UpdateTexture(const TArray<uint8>& Bytes) override;
-
+	virtual void UpdateTextureThreadSafe(const TArray<uint8>& Bytes) override;
 protected:
 	/** Width of this texture */
 	uint32 Width;
diff --git a/Engine/Source/Runtime/SlateCore/Public/Textures/SlateUpdatableTexture.h b/Engine/Source/Runtime/SlateCore/Public/Textures/SlateUpdatableTexture.h
--- a/Engine/Source/Runtime/SlateCore/Public/Textures/SlateUpdatableTexture.h
+++ b/Engine/Source/Runtime/SlateCore/Public/Textures/SlateUpdatableTexture.h
@@ -42,4 +42,11 @@ public:
 	 * @param Bytes Array of texture data
 	 */
 	virtual void UpdateTexture(const TArray<uint8>& Bytes) = 0;
+
+	/**
+	 * Updates the texture contents via a byte array making a copy first for thread safety
+	 *
+	 * @param Bytes Array of texture data
+	 */
+	virtual void UpdateTextureThreadSafe(const TArray<uint8>& Bytes) = 0;
 };
diff --git a/Engine/Source/Runtime/WebBrowser/Private/WebBrowserWindow.cpp b/Engine/Source/Runtime/WebBrowser/Private/WebBrowserWindow.cpp
--- a/Engine/Source/Runtime/WebBrowser/Private/WebBrowserWindow.cpp
+++ b/Engine/Source/Runtime/WebBrowser/Private/WebBrowserWindow.cpp
@@ -87,7 +87,7 @@ void FWebBrowserWindow::SetViewportSize(FIntPoint WindowSize)
 		if (UpdatableTexture != nullptr)
 		{
 			UpdatableTexture->ResizeTexture(ViewportSize.X, ViewportSize.Y);
-			UpdatableTexture->UpdateTexture(TextureData);
+			UpdatableTexture->UpdateTextureThreadSafe(TextureData);
 		}
 		if (IsValid())
 		{
@@ -424,7 +424,7 @@ void FWebBrowserWindow::OnPaint(CefRenderHandler::PaintElementType Type, const C
 
 	if (UpdatableTexture != nullptr)
 	{
-		UpdatableTexture->UpdateTexture(TextureData);
+		UpdatableTexture->UpdateTextureThreadSafe(TextureData);
 	}
 
 	bHasBeenPainted = true;
