Pause demo when at the end of the stream + don't allow CurrentDemoTime to go past TotalDemoTime

diff --git a/Engine/Source/Runtime/Engine/Classes/Engine/DemoNetDriver.h b/Engine/Source/Runtime/Engine/Classes/Engine/DemoNetDriver.h
--- a/Engine/Source/Runtime/Engine/Classes/Engine/DemoNetDriver.h
+++ b/Engine/Source/Runtime/Engine/Classes/Engine/DemoNetDriver.h
@@ -40,6 +40,9 @@ class UDemoNetDriver
 	/** True if we are at the end of playing a demo */
 	bool bDemoPlaybackDone;
 
+	/** True if as have paused all of the channels */
+	bool bChannelsArePaused;
+
 	/** This is our spectator controller that is used to view the demo world from */
 	APlayerController* SpectatorController;
 
@@ -82,6 +85,7 @@ public:
 	virtual bool IsNetResourceValid(void) { return true; }
 
 	void TickDemoRecord( float DeltaSeconds );
+	void PauseChannels( const bool bPause );
 	bool ReadDemoFrame();
 	void TickDemoPlayback( float DeltaSeconds );
 	void SpawnDemoRecSpectator( UNetConnection* Connection );
diff --git a/Engine/Source/Runtime/Engine/Private/DemoNetDriver.cpp b/Engine/Source/Runtime/Engine/Private/DemoNetDriver.cpp
--- a/Engine/Source/Runtime/Engine/Private/DemoNetDriver.cpp
+++ b/Engine/Source/Runtime/Engine/Private/DemoNetDriver.cpp
@@ -45,6 +45,7 @@ bool UDemoNetDriver::InitBase( bool bInitAsClient, FNetworkNotify* InNotify, con
 		Time					= 0;
 		bIsRecordingDemoFrame	= false;
 		bDemoPlaybackDone		= false;
+		bChannelsArePaused		= false;
 
 		ResetDemoState();
 
@@ -687,20 +688,44 @@ void UDemoNetDriver::TickDemoRecord( float DeltaSeconds )
 	*FileAr << EndCount;
 }
 
-bool UDemoNetDriver::ReadDemoFrame()
+void UDemoNetDriver::PauseChannels( const bool bPause )
 {
-	if ( !ReplayStreamer->IsDataAvailable() )
+	if ( bPause == bChannelsArePaused )
 	{
-		return false;
+		return;
 	}
 
-	const uint32 TotalDemoTimeInMS = ReplayStreamer->GetTotalDemoTime();
-
-	if ( TotalDemoTimeInMS > 0 )
+	// Pause all non player controller actors
+	for ( int32 i = ServerConnection->OpenChannels.Num() - 1; i >= 0; i-- )
 	{
-		DemoTotalTime = (float)TotalDemoTimeInMS / 1000.0f;
+		UChannel* OpenChannel = ServerConnection->OpenChannels[i];
+		if ( OpenChannel != NULL )
+		{
+			UActorChannel* ActorChannel = Cast< UActorChannel >( OpenChannel );
+			if ( ActorChannel != NULL && ActorChannel->GetActor() != NULL )
+			{
+				if ( Cast< APlayerController >( ActorChannel->GetActor() ) == NULL )
+				{
+					// Better way to pause each actor?
+					ActorChannel->GetActor()->CustomTimeDilation = bPause ? 0.0f : 1.0f;
+				}
+			}
+		}
 	}
 
+	bChannelsArePaused = bPause;
+}
+
+bool UDemoNetDriver::ReadDemoFrame()
+{
+	if ( !ReplayStreamer->IsDataAvailable() )
+	{
+		PauseChannels( true );
+		return false;
+	} 
+
+	PauseChannels( false );
+
 	FArchive* FileAr = ReplayStreamer->GetStreamingArchive();
 
 	if ( FileAr->IsError() )
@@ -709,28 +734,10 @@ bool UDemoNetDriver::ReadDemoFrame()
 		return false;
 	}
 
-	if ( FileAr->AtEnd() )
+	if ( FileAr->AtEnd() )		// FIXME: Resolve how we're going to really do this with the streamers
 	{
 		bDemoPlaybackDone = true;
-
-		// Pause all non player controller actors
-		for ( int32 i = ServerConnection->OpenChannels.Num() - 1; i >= 0; i-- )
-		{
-			UChannel* OpenChannel = ServerConnection->OpenChannels[i];
-			if ( OpenChannel != NULL )
-			{
-				UActorChannel* ActorChannel = Cast< UActorChannel >( OpenChannel );
-				if ( ActorChannel != NULL && ActorChannel->GetActor() != NULL )
-				{
-					if ( Cast< APlayerController >( ActorChannel->GetActor() ) == NULL )
-					{
-						// Better way to pause each actor?
-						ActorChannel->GetActor()->CustomTimeDilation = 0.0f;
-					}
-				}
-			}
-		}
-
+		PauseChannels( true );
 		return false;
 	}
 
@@ -856,9 +863,21 @@ void UDemoNetDriver::TickDemoPlayback( float DeltaSeconds )
 		return;
 	}
 
+	const uint32 TotalDemoTimeInMS = ReplayStreamer->GetTotalDemoTime();
+
+	if ( TotalDemoTimeInMS > 0 )
+	{
+		DemoTotalTime = (float)TotalDemoTimeInMS / 1000.0f;
+	}
+
 	DemoDeltaTime += DeltaSeconds;
 	DemoCurrentTime += DeltaSeconds;
 
+	if ( DemoCurrentTime >= DemoTotalTime )
+	{
+		DemoCurrentTime = DemoTotalTime;
+	}
+
 	while ( true )
 	{
 		// Read demo frames until we are caught up
