Core - Added stat for FAsyncAudioDecompressWorker for Wave

diff --git a/Engine/Source/Runtime/Engine/Private/AudioDecompress.cpp b/Engine/Source/Runtime/Engine/Private/AudioDecompress.cpp
--- a/Engine/Source/Runtime/Engine/Private/AudioDecompress.cpp
+++ b/Engine/Source/Runtime/Engine/Private/AudioDecompress.cpp
@@ -299,6 +299,7 @@ void FAsyncAudioDecompressWorker::DoWork( void )
 		// Parse the audio header for the relevant information
 		if (AudioInfo->ReadCompressedInfo(Wave->ResourceData, Wave->ResourceSize, &QualityInfo))
 		{
+			FScopeCycleCounterUObject WaveObject( Wave );
 
 #if PLATFORM_ANDROID
 			// Handle resampling
