fix for compiling without recast

diff --git a/Engine/Source/Runtime/Engine/Private/AI/Navigation/NavigationSystem.cpp b/Engine/Source/Runtime/Engine/Private/AI/Navigation/NavigationSystem.cpp
--- a/Engine/Source/Runtime/Engine/Private/AI/Navigation/NavigationSystem.cpp
+++ b/Engine/Source/Runtime/Engine/Private/AI/Navigation/NavigationSystem.cpp
@@ -3651,6 +3651,7 @@ void UNavigationSystem::UpdateInvokers()
 #endif // ENABLE_VISUAL_LOG
 		}
 
+#if WITH_RECAST
 		const double UpdateStartTime = FPlatformTime::Seconds();
 		for (TActorIterator<ARecastNavMesh> It(GetWorld()); It; ++It)
 		{
@@ -3658,7 +3659,8 @@ void UNavigationSystem::UpdateInvokers()
 		}
 		const double UpdateEndTime = FPlatformTime::Seconds();
 		UE_VLOG(this, LogNavigation, Log, TEXT("Marking tiles to update %fms (%d invokers)"), (UpdateEndTime - UpdateStartTime) * 1000, InvokerLocations.Num());
-		
+#endif
+
 		// once per second
 		NextInvokersUpdateTime = CurrentTime + ActiveTilesUpdateInterval;
 	}
diff --git a/Engine/Source/Runtime/Engine/Private/AI/Navigation/RecastNavMesh.cpp b/Engine/Source/Runtime/Engine/Private/AI/Navigation/RecastNavMesh.cpp
--- a/Engine/Source/Runtime/Engine/Private/AI/Navigation/RecastNavMesh.cpp
+++ b/Engine/Source/Runtime/Engine/Private/AI/Navigation/RecastNavMesh.cpp
@@ -1940,6 +1940,7 @@ bool ARecastNavMesh::HasValidNavmesh() const
 //----------------------------------------------------------------------//
 // RecastNavMesh: Active Tiles 
 //----------------------------------------------------------------------//
+#if WITH_RECAST
 void ARecastNavMesh::UpdateActiveTiles(const TArray<FNavigationInvokerRaw>& InvokerLocations)
 {
 	if (HasValidNavmesh() == false)
@@ -2062,6 +2063,7 @@ void ARecastNavMesh::RebuildTile(const TArray<FIntPoint>& Tiles)
 		}
 	}
 }
+#endif
 
 //----------------------------------------------------------------------//
 // FRecastNavMeshCachedData
