Don't call BeginPlay twice for actors that arrive from the server in the same batch as the game mode transition that causes world settings notify begin play.

diff --git a/Engine/Source/Runtime/Engine/Private/Actor.cpp b/Engine/Source/Runtime/Engine/Private/Actor.cpp
--- a/Engine/Source/Runtime/Engine/Private/Actor.cpp
+++ b/Engine/Source/Runtime/Engine/Private/Actor.cpp
@@ -2424,7 +2424,7 @@ void AActor::PostNetInit()
 	}
 	check(RemoteRole == ROLE_Authority);
 
-	if (GetWorld() && GetWorld()->HasBegunPlay())
+	if (!HasActorBegunPlay() && GetWorld() && GetWorld()->HasBegunPlay())
 	{
 		BeginPlay();
 	}
