Added additive flag to rotation multiplier skeletal control node.
#github https://github.com/EpicGames/UnrealEngine/pull/473
user:eelkin2000

diff --git a/Engine/Source/Runtime/Engine/Classes/Animation/BoneControllers/AnimNode_RotationMultiplier.h b/Engine/Source/Runtime/Engine/Classes/Animation/BoneControllers/AnimNode_RotationMultiplier.h
--- a/Engine/Source/Runtime/Engine/Classes/Animation/BoneControllers/AnimNode_RotationMultiplier.h
+++ b/Engine/Source/Runtime/Engine/Classes/Animation/BoneControllers/AnimNode_RotationMultiplier.h
@@ -37,6 +37,9 @@ struct ENGINE_API FAnimNode_RotationMultiplier : public FAnimNode_SkeletalContro
 
 	UPROPERTY(EditAnywhere, BlueprintReadWrite, Category=Multiplier)
 	TEnumAsByte<EBoneAxis> RotationAxisToRefer;
+	
+	UPROPERTY(EditAnywhere, BlueprintReadWrite, Category = Multiplier)
+	bool bIsAdditive;	
 
 	FAnimNode_RotationMultiplier();
 
diff --git a/Engine/Source/Runtime/Engine/Private/Animation/AnimNodeRotationMultiplier.cpp b/Engine/Source/Runtime/Engine/Private/Animation/AnimNodeRotationMultiplier.cpp
--- a/Engine/Source/Runtime/Engine/Private/Animation/AnimNodeRotationMultiplier.cpp
+++ b/Engine/Source/Runtime/Engine/Private/Animation/AnimNodeRotationMultiplier.cpp
@@ -133,6 +133,12 @@ void FAnimNode_RotationMultiplier::EvaluateBoneTransforms(USkeletalMeshComponent
 		MultiplyQuatBasedOnSourceIndex(LocalRefPose, MeshBases, RotationAxisToRefer, SourceBone.BoneIndex, Multiplier, RefQuat, NewQuat);
 
 		FTransform NewLocalTransform = MeshBases.GetLocalSpaceTransform(TargetBone.BoneIndex);
+		
+		if (bIsAdditive)
+		{
+			NewQuat = NewLocalTransform.GetRotation() * NewQuat;
+		}
+		
 		NewLocalTransform.SetRotation(NewQuat);
 
 		const int32 ParentIndex = RequiredBones.GetParentBoneIndex(TargetBone.BoneIndex);
@@ -159,4 +165,4 @@ void FAnimNode_RotationMultiplier::InitializeBoneReferences(const FBoneContainer
 {
 	SourceBone.Initialize(RequiredBones);
 	TargetBone.Initialize(RequiredBones);
-}
\ No newline at end of file
+}
