Consider all components not just this one when evaluating whether there is already an Actor overlap

diff --git a/Engine/Source/Runtime/Engine/Private/PrimitiveComponent.cpp b/Engine/Source/Runtime/Engine/Private/PrimitiveComponent.cpp
--- a/Engine/Source/Runtime/Engine/Private/PrimitiveComponent.cpp
+++ b/Engine/Source/Runtime/Engine/Private/PrimitiveComponent.cpp
@@ -1926,8 +1926,9 @@ void UPrimitiveComponent::BeginComponentOverlap(const FOverlapInfo& OtherOverlap
 		if (!bComponentsAlreadyTouching)
 		{
 			AActor* const OtherActor = OtherComp->GetOwner();
+			AActor* const MyActor = GetOwner();
 
-			const bool bNotifyActorTouch = bDoNotifies && !IsOverlappingActor(OtherActor);
+			const bool bNotifyActorTouch = bDoNotifies && !MyActor->IsOverlappingActor(OtherActor);
 
 			// Perform reflexive touch.
 			OverlappingComponents.Add(OtherOverlap);										// already verified uniqueness above
@@ -1935,8 +1936,6 @@ void UPrimitiveComponent::BeginComponentOverlap(const FOverlapInfo& OtherOverlap
 			
 			if (bDoNotifies)
 			{
-				AActor* const MyActor = GetOwner();
-
 				// first execute component delegates
 				if (!IsPendingKill())
 				{
@@ -2016,7 +2015,7 @@ void UPrimitiveComponent::EndComponentOverlap(const FOverlapInfo& OtherOverlap,
 		}
 		
 		// if this was the last touch on the other actor, notify that we've untouched the actor as well
-		if (bDoNotifies && !IsOverlappingActor(OtherActor) && MyActor && OtherActor)
+		if (bDoNotifies && MyActor && OtherActor && !MyActor->IsOverlappingActor(OtherActor) )
 		{			
 			if (IsActorValidToNotify(MyActor))
 			{
