Merging 2464332 animation fixes

//depot/UE4-UT/...

to //depot/UE4/...

diff --git a/Engine/Source/Runtime/Engine/Private/SkeletalMeshComponent.cpp b/Engine/Source/Runtime/Engine/Private/SkeletalMeshComponent.cpp
--- a/Engine/Source/Runtime/Engine/Private/SkeletalMeshComponent.cpp
+++ b/Engine/Source/Runtime/Engine/Private/SkeletalMeshComponent.cpp
@@ -987,7 +987,7 @@ void USkeletalMeshComponent::RefreshBoneTransforms(FActorComponentTickFunction*
 		CachedSpaceBases.Empty();
 	}
 
-	const bool bDoPAE = !!CVarUseParallelAnimationEvaluation.GetValueOnGameThread();
+	const bool bDoPAE = !!CVarUseParallelAnimationEvaluation.GetValueOnGameThread() && FApp::ShouldUseThreadingForPerformance();
 	if (AnimEvaluationContext.bDoEvaluation && TickFunction && bDoPAE)
 	{
 		if (SkeletalMesh->RefSkeleton.GetNum() != AnimEvaluationContext.LocalAtoms.Num())
