#ue4 - Perf: fix perf regression in UPrimitiveComponent::MoveComponent() caused by continual FText creation.

- Also don't call GetOwner() unless necessary in UPrimitiveComponent::CheckStaticMobilityAndWarn().

diff --git a/Engine/Source/Runtime/Engine/Private/PrimitiveComponent.cpp b/Engine/Source/Runtime/Engine/Private/PrimitiveComponent.cpp
--- a/Engine/Source/Runtime/Engine/Private/PrimitiveComponent.cpp
+++ b/Engine/Source/Runtime/Engine/Private/PrimitiveComponent.cpp
@@ -1290,14 +1290,18 @@ FCollisionShape UPrimitiveComponent::GetCollisionShape(float Inflation) const
 
 bool UPrimitiveComponent::CheckStaticMobilityAndWarn(const FText& ActionText) const
 {
-	AActor* Actor = GetOwner();
 	// static things can move before they are registered (e.g. immediately after streaming), but not after.
-	if (Mobility == EComponentMobility::Static && Actor && Actor->IsActorInitialized())
+	if (Mobility == EComponentMobility::Static)
 	{
-		FMessageLog("PIE").Warning(FText::Format(LOCTEXT("InvalidStaticMove", "Mobility of {0} : {1} has to be 'Movable' if you'd like to {2}. "),
-			FText::FromString(GetNameSafe(GetOwner())), FText::FromString(GetName()), ActionText));
+		AActor* Actor = GetOwner();
+		if (Actor && Actor->IsActorInitialized())
+		{
+			static const FText WarnText = LOCTEXT("InvalidStaticMove", "Mobility of {0} : {1} has to be 'Movable' if you'd like to {2}. ");
+			FMessageLog("PIE").Warning(FText::Format(WarnText,
+				FText::FromString(GetNameSafe(GetOwner())), FText::FromString(GetName()), ActionText));
 
-		return true;
+			return true;
+		}
 	}
 
 	return false;
@@ -1330,7 +1334,8 @@ bool UPrimitiveComponent::MoveComponent( const FVector& Delta, const FRotator& N
 
 	// static things can move before they are registered (e.g. immediately after streaming), but not after.
 	// TODO: Static components without an owner can move, should they be able to?
-	if (CheckStaticMobilityAndWarn(LOCTEXT("InvalidMove", "move")))
+	static const FText WarnText = LOCTEXT("InvalidMove", "move");
+	if (CheckStaticMobilityAndWarn(WarnText))
 	{
 		if (OutHit)
 		{
