Do not consider timers with unbound delegates for re-adding to active list.

(Merging CL 2465792 from my branch)

diff --git a/Engine/Source/Runtime/Engine/Private/TimerManager.cpp b/Engine/Source/Runtime/Engine/Private/TimerManager.cpp
--- a/Engine/Source/Runtime/Engine/Private/TimerManager.cpp
+++ b/Engine/Source/Runtime/Engine/Private/TimerManager.cpp
@@ -518,6 +518,11 @@ void FTimerManager::Tick(float DeltaTime)
 				}
 			}
 
+			if (!CurrentlyExecutingTimer.TimerDelegate.IsBound())
+			{
+				continue;
+			}
+
 			if( CurrentlyExecutingTimer.bLoop && 
 				(CurrentlyExecutingTimer.TimerHandle.IsValid() ||
 				 CurrentlyExecutingTimer.TimerDelegate.IsBound()) && 							// did not get cleared during execution
