Rollback //depot/Ue4/Engine/Source/Runtime/Engine/Private/TimerManager.cpp to revision 24

(Undoing my previous change).

diff --git a/Engine/Source/Runtime/Engine/Private/TimerManager.cpp b/Engine/Source/Runtime/Engine/Private/TimerManager.cpp
--- a/Engine/Source/Runtime/Engine/Private/TimerManager.cpp
+++ b/Engine/Source/Runtime/Engine/Private/TimerManager.cpp
@@ -518,11 +518,6 @@ void FTimerManager::Tick(float DeltaTime)
 				}
 			}
 
-			if (!CurrentlyExecutingTimer.TimerDelegate.IsBound())
-			{
-				continue;
-			}
-
 			if( CurrentlyExecutingTimer.bLoop && 
 				(CurrentlyExecutingTimer.TimerHandle.IsValid() ||
 				 CurrentlyExecutingTimer.TimerDelegate.IsBound()) && 							// did not get cleared during execution
