Stats - Fixed bad memory accounting for static meshes, added distance field memory usage

diff --git a/Engine/Source/Runtime/Core/Private/Misc/Core.cpp b/Engine/Source/Runtime/Core/Private/Misc/Core.cpp
--- a/Engine/Source/Runtime/Core/Private/Misc/Core.cpp
+++ b/Engine/Source/Runtime/Core/Private/Misc/Core.cpp
@@ -266,7 +266,7 @@ DEFINE_STAT(STAT_StaticMeshVertexMemory);
 DEFINE_STAT(STAT_ResourceVertexColorMemory);
 DEFINE_STAT(STAT_InstVertexColorMemory);
 DEFINE_STAT(STAT_StaticMeshIndexMemory);
-
+DEFINE_STAT(STAT_StaticMeshDistanceFieldMemory);
 
 
 /** Threading stats objects */
diff --git a/Engine/Source/Runtime/Core/Public/Misc/CoreStats.h b/Engine/Source/Runtime/Core/Public/Misc/CoreStats.h
--- a/Engine/Source/Runtime/Core/Public/Misc/CoreStats.h
+++ b/Engine/Source/Runtime/Core/Public/Misc/CoreStats.h
@@ -42,6 +42,7 @@ DECLARE_MEMORY_STAT_EXTERN(TEXT("StaticMesh Vertex Memory"),STAT_StaticMeshVerte
 DECLARE_MEMORY_STAT_EXTERN(TEXT("StaticMesh VxColor Resource Mem"),STAT_ResourceVertexColorMemory,STATGROUP_MemoryStaticMesh, CORE_API);
 DECLARE_MEMORY_STAT_EXTERN(TEXT("StaticMesh VxColor Inst Mem"),STAT_InstVertexColorMemory,STATGROUP_MemoryStaticMesh, CORE_API);
 DECLARE_MEMORY_STAT_EXTERN(TEXT("StaticMesh Index Memory"),STAT_StaticMeshIndexMemory,STATGROUP_MemoryStaticMesh, CORE_API);
+DECLARE_MEMORY_STAT_EXTERN(TEXT("StaticMesh Distance Field Memory"),STAT_StaticMeshDistanceFieldMemory,STATGROUP_MemoryStaticMesh, CORE_API);
 
 DECLARE_MEMORY_STAT_EXTERN(TEXT("Texture Memory Used"),STAT_TextureMemory,STATGROUP_Memory, CORE_API);
 DECLARE_MEMORY_STAT_EXTERN(TEXT("Reflection Capture Texture Memory"),STAT_ReflectionCaptureTextureMemory,STATGROUP_Memory, CORE_API);
diff --git a/Engine/Source/Runtime/Engine/Private/StaticMesh.cpp b/Engine/Source/Runtime/Engine/Private/StaticMesh.cpp
--- a/Engine/Source/Runtime/Engine/Private/StaticMesh.cpp
+++ b/Engine/Source/Runtime/Engine/Private/StaticMesh.cpp
@@ -530,30 +530,35 @@ void FStaticMeshLODResources::InitResources(UStaticMesh* Parent)
 	if (DistanceFieldData)
 	{
 		DistanceFieldData->VolumeTexture.Initialize();
+		INC_DWORD_STAT_BY( STAT_StaticMeshDistanceFieldMemory, DistanceFieldData->GetResourceSize() );
 	}
 
-	const uint32 StaticMeshVertexMemory = 
-		VertexBuffer.GetStride() * VertexBuffer.GetNumVertices() + 
-		PositionVertexBuffer.GetStride() * PositionVertexBuffer.GetNumVertices();
-	const uint32 StaticMeshIndexMemory = IndexBuffer.GetAllocatedSize()
-		+ WireframeIndexBuffer.GetAllocatedSize()
-		+ (AdjacencyIndexBuffer.IsInitialized() ? AdjacencyIndexBuffer.GetAllocatedSize() : 0);
-	const uint32 ResourceVertexColorMemory = ColorVertexBuffer.GetStride() * ColorVertexBuffer.GetNumVertices();
-
-	INC_DWORD_STAT_BY( STAT_StaticMeshVertexMemory, StaticMeshVertexMemory );
-	INC_DWORD_STAT_BY( STAT_ResourceVertexColorMemory, ResourceVertexColorMemory );
-	INC_DWORD_STAT_BY( STAT_StaticMeshIndexMemory, StaticMeshIndexMemory );
+	ENQUEUE_UNIQUE_RENDER_COMMAND_ONEPARAMETER(
+		UpdateMemoryStats,
+		FStaticMeshLODResources*, This, this,
+		{		
+			const uint32 StaticMeshVertexMemory =
+			This->VertexBuffer.GetStride() * This->VertexBuffer.GetNumVertices() +
+			This->PositionVertexBuffer.GetStride() * This->PositionVertexBuffer.GetNumVertices();
+			const uint32 StaticMeshIndexMemory = This->IndexBuffer.GetAllocatedSize()
+				+ This->WireframeIndexBuffer.GetAllocatedSize()
+				+ (RHISupportsTessellation( GShaderPlatformForFeatureLevel[GMaxRHIFeatureLevel] ) ? This->AdjacencyIndexBuffer.GetAllocatedSize() : 0);
+			const uint32 ResourceVertexColorMemory = This->ColorVertexBuffer.GetStride() * This->ColorVertexBuffer.GetNumVertices();
+
+			INC_DWORD_STAT_BY( STAT_StaticMeshVertexMemory, StaticMeshVertexMemory );
+			INC_DWORD_STAT_BY( STAT_ResourceVertexColorMemory, ResourceVertexColorMemory );
+			INC_DWORD_STAT_BY( STAT_StaticMeshIndexMemory, StaticMeshIndexMemory );
+		});
 }
 
 void FStaticMeshLODResources::ReleaseResources()
 {
-	// TODO: The sizes for index buffers will be incorrect outside of the editor because we will have freed the CPU arrays.
 	const uint32 StaticMeshVertexMemory = 
 		VertexBuffer.GetStride() * VertexBuffer.GetNumVertices() + 
 		PositionVertexBuffer.GetStride() * PositionVertexBuffer.GetNumVertices();
 	const uint32 StaticMeshIndexMemory = IndexBuffer.GetAllocatedSize()
 		+ WireframeIndexBuffer.GetAllocatedSize()
-		+ (AdjacencyIndexBuffer.IsInitialized() ? AdjacencyIndexBuffer.GetAllocatedSize() : 0);
+		+ (RHISupportsTessellation( GShaderPlatformForFeatureLevel[GMaxRHIFeatureLevel] ) ? AdjacencyIndexBuffer.GetAllocatedSize() : 0);
 	const uint32 ResourceVertexColorMemory = ColorVertexBuffer.GetStride() * ColorVertexBuffer.GetNumVertices();
 
 	DEC_DWORD_STAT_BY( STAT_StaticMeshVertexMemory, StaticMeshVertexMemory );
@@ -579,6 +584,7 @@ void FStaticMeshLODResources::ReleaseResources()
 
 	if (DistanceFieldData)
 	{
+		DEC_DWORD_STAT_BY( STAT_StaticMeshDistanceFieldMemory, DistanceFieldData->GetResourceSize() );
 		DistanceFieldData->VolumeTexture.Release();
 	}
 }
@@ -1229,11 +1235,16 @@ void UStaticMesh::InitResources()
 		RenderData->InitResources(this);
 	}
 
-#if STATS
-	uint32 StaticMeshResourceSize = GetResourceSize(EResourceSizeMode::Exclusive);
-	INC_DWORD_STAT_BY( STAT_StaticMeshTotalMemory, StaticMeshResourceSize );
-	INC_DWORD_STAT_BY( STAT_StaticMeshTotalMemory2, StaticMeshResourceSize );
-#endif
+#if	STATS
+	ENQUEUE_UNIQUE_RENDER_COMMAND_ONEPARAMETER(
+		UpdateMemoryStats,
+		UStaticMesh*, This, this,
+		{
+ 			const uint32 StaticMeshResourceSize = This->GetResourceSize( EResourceSizeMode::Exclusive );
+ 			INC_DWORD_STAT_BY( STAT_StaticMeshTotalMemory, StaticMeshResourceSize );
+ 			INC_DWORD_STAT_BY( STAT_StaticMeshTotalMemory2, StaticMeshResourceSize );
+		} );
+#endif // STATS
 }
 
 /**
