Removed UPolys from non editor only data builds.

diff --git a/Engine/Source/Runtime/Core/Public/UObject/ObjectVersion.h b/Engine/Source/Runtime/Core/Public/UObject/ObjectVersion.h
--- a/Engine/Source/Runtime/Core/Public/UObject/ObjectVersion.h
+++ b/Engine/Source/Runtime/Core/Public/UObject/ObjectVersion.h
@@ -479,6 +479,8 @@ enum EUnrealEngineObjectUE4Version
 	// FPropertyTag stores GUID of struct
 	VER_UE4_STRUCT_GUID_IN_PROPERTY_TAG,
 
+	// Remove unused UPolys from UModel cooked content
+	VER_UE4_REMOVE_UNUSED_UPOLYS_FROM_UMODEL,
 	// -----<new versions can be added before this line>-------------------------------------------------
 	// - this needs to be the last line (see note below)
 	VER_UE4_AUTOMATIC_VERSION_PLUS_ONE,
diff --git a/Engine/Source/Runtime/Engine/Private/Model.cpp b/Engine/Source/Runtime/Engine/Private/Model.cpp
--- a/Engine/Source/Runtime/Engine/Private/Model.cpp
+++ b/Engine/Source/Runtime/Engine/Private/Model.cpp
@@ -191,13 +191,31 @@ void UModel::Serialize( FArchive& Ar )
 	}
 
 #if WITH_EDITOR
-	Ar << Polys;
-	LeafHulls.BulkSerialize( Ar );
-	Leaves.BulkSerialize( Ar );
+	bool bHasEditorOnlyData = !Ar.IsFilterEditorOnly();
+	
+	if ( Ar.UE4Ver() < VER_UE4_REMOVE_UNUSED_UPOLYS_FROM_UMODEL )
+	{
+		bHasEditorOnlyData = true;
+	}
+
+	// if we are cooking then don't save this stuff out
+	if ( bHasEditorOnlyData )
+	{
+		Ar << Polys;
+		LeafHulls.BulkSerialize( Ar );
+		Leaves.BulkSerialize( Ar );
+	}
 #else
-	if(Ar.IsLoading())
+	bool bHasEditorOnlyData = !Ar.IsFilterEditorOnly();
+	
+	if ( Ar.UE4Ver() < VER_UE4_REMOVE_UNUSED_UPOLYS_FROM_UMODEL )
+	{
+		bHasEditorOnlyData = true;
+	}
+
+	if((Ar.IsLoading() || Ar.IsSaving()) && bHasEditorOnlyData)
 	{
-		UPolys* DummyPolys;
+		UPolys* DummyPolys = NULL;
 		Ar << DummyPolys;
 
 		TArray<int32> DummyLeafHulls;
