Add scene lock when adding an aggregate

diff --git a/Engine/Source/Runtime/Engine/Private/SkeletalMeshComponentPhysics.cpp b/Engine/Source/Runtime/Engine/Private/SkeletalMeshComponentPhysics.cpp
--- a/Engine/Source/Runtime/Engine/Private/SkeletalMeshComponentPhysics.cpp
+++ b/Engine/Source/Runtime/Engine/Private/SkeletalMeshComponentPhysics.cpp
@@ -828,8 +828,12 @@ void USkeletalMeshComponent::InitArticulated(FPhysScene* PhysScene)
 	{
 		// Get the scene type from the SkeletalMeshComponent's BodyInstance
 		const uint32 SceneType = BodyInstance.UseAsyncScene() ? PST_Async : PST_Sync;
-		PhysScene->GetPhysXScene(SceneType)->addAggregate(*Aggregate);
-
+		{
+			PxScene* PScene = PhysScene->GetPhysXScene(SceneType);
+			SCOPED_SCENE_WRITE_LOCK(PScene);
+			PScene->addAggregate(*Aggregate);
+		}
+		
 		// If we've used an aggregate, InitBody would not be able to set awake status as we *must* have a scene
 		// to do that, so we reconcile this here.
 		AActor* Owner = GetOwner();
