fix for UE-10554 BSP Brush settings not updating in real time, UE-10560 BSP does not update it's location, rotation, or scale when changed in the details panel

diff --git a/Engine/Source/Runtime/Engine/Private/Brush.cpp b/Engine/Source/Runtime/Engine/Private/Brush.cpp
--- a/Engine/Source/Runtime/Engine/Private/Brush.cpp
+++ b/Engine/Source/Runtime/Engine/Private/Brush.cpp
@@ -58,6 +58,11 @@ void ABrush::PostEditChangeProperty(FPropertyChangedEvent& PropertyChangedEvent)
 
 	if(IsStaticBrush() && !GIsTransacting)
 	{
+		if(PropertyChangedEvent.ChangeType != EPropertyChangeType::Interactive)
+		{
+			GEditor->RebuildAlteredBSP();
+		}
+
 		// Trigger a csg rebuild if we're in the editor.
 		SetNeedRebuild(GetLevel());
 	}
