Integrate - Added format to ListTextures

diff --git a/Engine/Source/Runtime/Engine/Private/UnrealEngine.cpp b/Engine/Source/Runtime/Engine/Private/UnrealEngine.cpp
--- a/Engine/Source/Runtime/Engine/Private/UnrealEngine.cpp
+++ b/Engine/Source/Runtime/Engine/Private/UnrealEngine.cpp
@@ -1908,6 +1908,7 @@ struct FSortedTexture
 	int32		OrigSizeY;
 	int32		CookedSizeX;
 	int32		CookedSizeY;
+	EPixelFormat Format;
 	int32		CurSizeX;
 	int32		CurSizeY;
 	int32		LODBias;
@@ -1919,11 +1920,12 @@ struct FSortedTexture
 	int32		UsageCount;
 
 	/** Constructor, initializing every member variable with passed in values. */
-	FSortedTexture(	int32 InOrigSizeX, int32 InOrigSizeY, int32 InCookedSizeX, int32 InCookedSizeY, int32 InCurSizeX, int32 InCurSizeY, int32 InLODBias, int32 InMaxSize, int32 InCurrentSize, const FString& InName, int32 InLODGroup, bool bInIsStreaming, int32 InUsageCount )
+	FSortedTexture(	int32 InOrigSizeX, int32 InOrigSizeY, int32 InCookedSizeX, int32 InCookedSizeY, EPixelFormat InFormat, int32 InCurSizeX, int32 InCurSizeY, int32 InLODBias, int32 InMaxSize, int32 InCurrentSize, const FString& InName, int32 InLODGroup, bool bInIsStreaming, int32 InUsageCount )
 	:	OrigSizeX( InOrigSizeX )
 	,	OrigSizeY( InOrigSizeY )
 	,	CookedSizeX( InCookedSizeX )
 	,	CookedSizeY( InCookedSizeY )
+	,	Format( InFormat )
 	,	CurSizeX( InCurSizeX )
 	,	CurSizeY( InCurSizeY )
 	,	LODBias( InLODBias )
@@ -1943,7 +1945,7 @@ struct FCompareFSortedTexture
 	{}
 	FORCEINLINE bool operator()( const FSortedTexture& A, const FSortedTexture& B ) const
 	{
-		return bAlphaSort ? ( A.Name < B.Name ) : ( B.MaxSize < A.MaxSize );
+		return bAlphaSort ? ( A.Name < B.Name ) : ( B.CurrentSize < A.CurrentSize );
 	}
 };
 
@@ -3246,6 +3248,7 @@ bool UEngine::HandleListTexturesCommand( const TCHAR* Cmd, FOutputDevice& Ar )
 		int32				OrigSizeY			= Texture->GetSizeY();
 		int32				CookedSizeX			= Texture->GetSizeX() >> LODBias;
 		int32				CookedSizeY			= Texture->GetSizeY() >> LODBias;
+		EPixelFormat		Format				= Texture->GetPixelFormat();
 		int32				DroppedMips			= Texture->GetNumMips() - Texture->ResidentMips;
 		int32				CurSizeX			= Texture->GetSizeX() >> DroppedMips;
 		int32				CurSizeY			= Texture->GetSizeY() >> DroppedMips;
@@ -3263,6 +3266,7 @@ bool UEngine::HandleListTexturesCommand( const TCHAR* Cmd, FOutputDevice& Ar )
 				OrigSizeY, 
 				CookedSizeX,
 				CookedSizeY,
+				Format,
 				CurSizeX,
 				CurSizeY,
 				LODBias, 
@@ -3284,15 +3288,16 @@ bool UEngine::HandleListTexturesCommand( const TCHAR* Cmd, FOutputDevice& Ar )
 	// Display.
 	int32 TotalMaxSize		= 0;
 	int32 TotalCurrentSize	= 0;
-	Ar.Logf( TEXT(",Authored Width,Authored Height,Cooked Width,Cooked Height,Current Width,Current Height,Max Size,Current Size,LODBias,LODGroup,Name,Streaming,Usage Count") );
+	Ar.Logf( TEXT(",Authored Width,Authored Height,Cooked Width,Cooked Height,Format,Current Width,Current Height,Max Size,Current Size,LODBias,LODGroup,Name,Streaming,Usage Count") );
 	for( int32 TextureIndex=0; TextureIndex<SortedTextures.Num(); TextureIndex++ )
 	{
 		const FSortedTexture& SortedTexture = SortedTextures[TextureIndex];
-		Ar.Logf( TEXT(",%i,%i,%i,%i,%i,%i,%i,%i,%i,%s,%s,%s,%i"),
+		Ar.Logf( TEXT(",%i,%i,%i,%i,%s,%i,%i,%i,%i,%i,%s,%s,%s,%i"),
 			SortedTexture.OrigSizeX,
 			SortedTexture.OrigSizeY,
 			SortedTexture.CookedSizeX,
 			SortedTexture.CookedSizeY,
+			GetPixelFormatString(SortedTexture.Format),
 			SortedTexture.CurSizeX,
 			SortedTexture.CurSizeY,
 			SortedTexture.MaxSize,
