Integrate - Added SetLightFunctionDisabledBrightness

diff --git a/Engine/Source/Runtime/Engine/Classes/Components/LightComponent.h b/Engine/Source/Runtime/Engine/Classes/Components/LightComponent.h
--- a/Engine/Source/Runtime/Engine/Classes/Components/LightComponent.h
+++ b/Engine/Source/Runtime/Engine/Classes/Components/LightComponent.h
@@ -175,6 +175,9 @@ public:
 	UFUNCTION(BlueprintCallable, Category="Rendering|Components|Light")
 	void SetLightFunctionFadeDistance(float NewLightFunctionFadeDistance);
 
+	UFUNCTION(BlueprintCallable, Category="Rendering|Components|Light")
+	void SetLightFunctionDisabledBrightness(float NewValue);
+
 	UFUNCTION(BlueprintCallable, Category="Rendering|Components|Light")
 	void SetAffectDynamicIndirectLighting(bool bNewValue);
 
diff --git a/Engine/Source/Runtime/Engine/Private/LightComponent.cpp b/Engine/Source/Runtime/Engine/Private/LightComponent.cpp
--- a/Engine/Source/Runtime/Engine/Private/LightComponent.cpp
+++ b/Engine/Source/Runtime/Engine/Private/LightComponent.cpp
@@ -731,6 +731,16 @@ void ULightComponent::SetLightFunctionFadeDistance(float NewLightFunctionFadeDis
 	}
 }
 
+void ULightComponent::SetLightFunctionDisabledBrightness(float NewValue)
+{
+	if (AreDynamicDataChangesAllowed()
+		&& NewValue != DisabledBrightness)
+	{
+		DisabledBrightness = NewValue;
+		MarkRenderStateDirty();
+	}
+}
+
 void ULightComponent::SetAffectDynamicIndirectLighting(bool bNewValue)
 {
 	if (AreDynamicDataChangesAllowed()
