Integrate - Added bOverrideMinLOD to static mesh component instead of relying on a magic value of MinLOD
Renamed MinLOD -> ClampedMinLOD where MinLOD has baen sanitized

diff --git a/Engine/Source/Runtime/Engine/Classes/Components/StaticMeshComponent.h b/Engine/Source/Runtime/Engine/Classes/Components/StaticMeshComponent.h
--- a/Engine/Source/Runtime/Engine/Classes/Components/StaticMeshComponent.h
+++ b/Engine/Source/Runtime/Engine/Classes/Components/StaticMeshComponent.h
@@ -130,12 +130,15 @@ class ENGINE_API UStaticMeshComponent : public UMeshComponent
 	UPROPERTY()
 	int32 PreviousLODLevel;
 
+	/** Whether to override the MinLOD setting of the static mesh asset with the MinLOD of this component. */
+	UPROPERTY(EditAnywhere, AdvancedDisplay, BlueprintReadOnly, Category=LOD)
+	bool bOverrideMinLOD;
+
 	/** 
 	 * Specifies the smallest LOD that will be used for this component.  
-	 * If -1 (default), the MinLOD of the static mesh asset will be used instead.
 	 * This is ignored if ForcedLodModel is enabled.
 	 */
-	UPROPERTY(EditAnywhere, AdvancedDisplay, BlueprintReadOnly, Category=LOD)
+	UPROPERTY(EditAnywhere, AdvancedDisplay, BlueprintReadOnly, Category=LOD, meta=(editcondition = "bOverrideMinLOD"))
 	int32 MinLOD;
 
 	UPROPERTY(EditAnywhere, BlueprintReadOnly, Category=StaticMesh, ReplicatedUsing=OnRep_StaticMesh)
diff --git a/Engine/Source/Runtime/Engine/Private/StaticMeshComponent.cpp b/Engine/Source/Runtime/Engine/Private/StaticMeshComponent.cpp
--- a/Engine/Source/Runtime/Engine/Private/StaticMeshComponent.cpp
+++ b/Engine/Source/Runtime/Engine/Private/StaticMeshComponent.cpp
@@ -144,7 +144,7 @@ UStaticMeshComponent::UStaticMeshComponent(const FObjectInitializer& ObjectIniti
 
 	WireframeColorOverride = FColor(255, 255, 255, 255);
 
-	MinLOD = -1;
+	MinLOD = 0;
 	bOverrideLightMapRes = false;
 	OverriddenLightMapRes = 64;
 	SubDivisionStepSize = 32;
diff --git a/Engine/Source/Runtime/Engine/Private/StaticMeshRender.cpp b/Engine/Source/Runtime/Engine/Private/StaticMeshRender.cpp
--- a/Engine/Source/Runtime/Engine/Private/StaticMeshRender.cpp
+++ b/Engine/Source/Runtime/Engine/Private/StaticMeshRender.cpp
@@ -63,7 +63,8 @@ FStaticMeshSceneProxy::FStaticMeshSceneProxy(UStaticMeshComponent* InComponent):
 {
 	check(RenderData);
 
-	MinLOD = FMath::Clamp(InComponent->MinLOD < 0 ? InComponent->StaticMesh->MinLOD : InComponent->MinLOD, 0, RenderData->LODResources.Num() - 1);
+	const int32 EffectiveMinLOD = InComponent->bOverrideMinLOD ? InComponent->MinLOD : InComponent->StaticMesh->MinLOD;
+	ClampedMinLOD = FMath::Clamp(EffectiveMinLOD, 0, RenderData->LODResources.Num() - 1);
 
 	WireframeColor = InComponent->GetWireframeColor();
 	LevelColor = FLinearColor(1,1,1);
@@ -466,7 +467,7 @@ void FStaticMeshSceneProxy::DrawStaticElements(FStaticPrimitiveDrawInterface* PD
 		} 
 		else //no LOD is being forced, submit them all with appropriate cull distances
 		{
-			for(int32 LODIndex = MinLOD; LODIndex < NumLODs; LODIndex++)
+			for(int32 LODIndex = ClampedMinLOD; LODIndex < NumLODs; LODIndex++)
 			{
 				const FStaticMeshLODResources& LODModel = RenderData->LODResources[LODIndex];
 				float ScreenSize = GetScreenSize(LODIndex);
@@ -1177,7 +1178,7 @@ int32 FStaticMeshSceneProxy::GetLOD(const FSceneView* View) const
 #endif
 
 	const FBoxSphereBounds& Bounds = GetBounds();
-	return ComputeStaticMeshLOD(RenderData, Bounds.Origin, Bounds.SphereRadius, *View, MinLOD);
+	return ComputeStaticMeshLOD(RenderData, Bounds.Origin, Bounds.SphereRadius, *View, ClampedMinLOD);
 }
 
 FPrimitiveSceneProxy* UStaticMeshComponent::CreateSceneProxy()
diff --git a/Engine/Source/Runtime/Engine/Public/StaticMeshResources.h b/Engine/Source/Runtime/Engine/Public/StaticMeshResources.h
--- a/Engine/Source/Runtime/Engine/Public/StaticMeshResources.h
+++ b/Engine/Source/Runtime/Engine/Public/StaticMeshResources.h
@@ -874,7 +874,8 @@ protected:
 	 */
 	int32 ForcedLodModel;
 
-	int32 MinLOD;
+	/** Minimum LOD index to use.  Clamped to valid range [0, NumLODs - 1]. */
+	int32 ClampedMinLOD;
 
 	FVector TotalScale3D;
 
