Integrate - New cinematic shadowing controls
* bCastShadowsFromCinematicObjectsOnly on the light
* bCastCinematicShadow on components
* Useful for limiting cinematic spotlight shadow cost
* These do not affect static shadowing or Ray Traced Distance Field shadows

diff --git a/Engine/Source/Runtime/Engine/Classes/Components/LightComponent.h b/Engine/Source/Runtime/Engine/Classes/Components/LightComponent.h
--- a/Engine/Source/Runtime/Engine/Classes/Components/LightComponent.h
+++ b/Engine/Source/Runtime/Engine/Classes/Components/LightComponent.h
@@ -82,6 +82,14 @@ class ENGINE_API ULightComponent : public ULightComponentBase
 	UPROPERTY(EditAnywhere, BlueprintReadOnly, Category=Light, AdvancedDisplay)
 	uint32 CastTranslucentShadows:1;
 
+	/** 
+	 * Whether the light should only cast shadows from components marked as bCastCinematicShadows. 
+	 * This is useful for setting up cinematic Movable spotlights aimed at characters and avoiding the shadow depth rendering costs of the background.
+	 * Note: this only works with dynamic shadow maps, not with static shadowing or Ray Traced Distance Field shadows.
+	 */
+	UPROPERTY(EditAnywhere, BlueprintReadOnly, Category=Light, AdvancedDisplay)
+	uint32 bCastShadowsFromCinematicObjectsOnly:1;
+
 	/**
 	 * Whether the light should be injected into the Light Propagation Volume
 	 **/
diff --git a/Engine/Source/Runtime/Engine/Classes/Components/PrimitiveComponent.h b/Engine/Source/Runtime/Engine/Classes/Components/PrimitiveComponent.h
--- a/Engine/Source/Runtime/Engine/Classes/Components/PrimitiveComponent.h
+++ b/Engine/Source/Runtime/Engine/Classes/Components/PrimitiveComponent.h
@@ -298,6 +298,13 @@ public:
 	UPROPERTY(EditAnywhere, AdvancedDisplay, BlueprintReadOnly, Category=Lighting, meta=(EditCondition="CastShadow", DisplayName = "Dynamic Inset Shadow"))
 	uint32 bCastInsetShadow:1;
 
+	/** 
+	 * Whether this component should cast shadows from lights that have bCastShadowsFromCinematicObjectsOnly enabled.
+	 * This is useful for characters in a cinematic with special cinematic lights, where the cost of shadowmap rendering of the environment is undesired.
+	 */
+	UPROPERTY(EditAnywhere, AdvancedDisplay, BlueprintReadOnly, Category=Lighting, meta=(EditCondition="CastShadow"))
+	uint32 bCastCinematicShadow:1;
+
 	/** 
 	 *	If true, the primitive will cast shadows even if bHidden is true.
 	 *	Controls whether the primitive should cast shadows when hidden.
diff --git a/Engine/Source/Runtime/Engine/Private/LightComponent.cpp b/Engine/Source/Runtime/Engine/Private/LightComponent.cpp
--- a/Engine/Source/Runtime/Engine/Private/LightComponent.cpp
+++ b/Engine/Source/Runtime/Engine/Private/LightComponent.cpp
@@ -213,6 +213,7 @@ FLightSceneProxy::FLightSceneProxy(const ULightComponent* InLightComponent)
 	, bCastDynamicShadow(InLightComponent->CastShadows && InLightComponent->CastDynamicShadows)
 	, bCastStaticShadow(InLightComponent->CastShadows && InLightComponent->CastStaticShadows)
 	, bCastTranslucentShadows(InLightComponent->CastTranslucentShadows)
+	, bCastShadowsFromCinematicObjectsOnly(InLightComponent->bCastShadowsFromCinematicObjectsOnly)
 	, bAffectTranslucentLighting(InLightComponent->bAffectTranslucentLighting)
 	, bUsedAsAtmosphereSunLight(InLightComponent->IsUsedAsAtmosphereSunLight())
 	, bAffectDynamicIndirectLighting(InLightComponent->bAffectDynamicIndirectLighting)
@@ -455,6 +456,11 @@ bool ULightComponent::CanEditChange(const UProperty* InProperty) const
 	{
 		FString PropertyName = InProperty->GetName();
 
+		if (PropertyName == GET_MEMBER_NAME_STRING_CHECKED(ULightComponent, bCastShadowsFromCinematicObjectsOnly))
+		{
+			return Mobility == EComponentMobility::Movable;
+		}
+
 		if (PropertyName == GET_MEMBER_NAME_STRING_CHECKED(ULightComponent, LightFunctionMaterial)
 			|| PropertyName == GET_MEMBER_NAME_STRING_CHECKED(ULightComponent, LightFunctionScale)
 			|| PropertyName == GET_MEMBER_NAME_STRING_CHECKED(ULightComponent, LightFunctionFadeDistance)
@@ -517,6 +523,7 @@ void ULightComponent::PostEditChangeProperty(FPropertyChangedEvent& PropertyChan
 	// Exclude properties that don't affect built lighting
 	//@todo - make this inclusive instead of exclusive?
 	if( PropertyName != GET_MEMBER_NAME_STRING_CHECKED(ULightComponent, CastTranslucentShadows) &&
+		PropertyName != GET_MEMBER_NAME_STRING_CHECKED(ULightComponent, bCastShadowsFromCinematicObjectsOnly) &&
 		PropertyName != GET_MEMBER_NAME_STRING_CHECKED(ULightComponent, CastDynamicShadows) &&
 		PropertyName != GET_MEMBER_NAME_STRING_CHECKED(ULightComponent, bAffectTranslucentLighting) &&
 		PropertyName != GET_MEMBER_NAME_STRING_CHECKED(ULightComponent, MinRoughness) &&
diff --git a/Engine/Source/Runtime/Engine/Private/PointLightComponent.cpp b/Engine/Source/Runtime/Engine/Private/PointLightComponent.cpp
--- a/Engine/Source/Runtime/Engine/Private/PointLightComponent.cpp
+++ b/Engine/Source/Runtime/Engine/Private/PointLightComponent.cpp
@@ -235,6 +235,11 @@ bool UPointLightComponent::CanEditChange(const UProperty* InProperty) const
 	{
 		FString PropertyName = InProperty->GetName();
 
+		if (PropertyName == GET_MEMBER_NAME_STRING_CHECKED(ULightComponent, bCastShadowsFromCinematicObjectsOnly) && bUseRayTracedDistanceFieldShadows)
+		{
+			return false;
+		}
+
 		if (FCString::Strcmp(*PropertyName, TEXT("LightFalloffExponent")) == 0)
 		{
 			return !bUseInverseSquaredFalloff;
diff --git a/Engine/Source/Runtime/Engine/Private/PrimitiveComponent.cpp b/Engine/Source/Runtime/Engine/Private/PrimitiveComponent.cpp
--- a/Engine/Source/Runtime/Engine/Private/PrimitiveComponent.cpp
+++ b/Engine/Source/Runtime/Engine/Private/PrimitiveComponent.cpp
@@ -591,7 +591,8 @@ bool UPrimitiveComponent::CanEditChange(const UProperty* InProperty) const
 			return Mobility != EComponentMobility::Movable || bLightAsIfStatic;
 		}
 
-		if (PropertyName == GET_MEMBER_NAME_CHECKED(UPrimitiveComponent, IndirectLightingCacheQuality))
+		if (PropertyName == GET_MEMBER_NAME_CHECKED(UPrimitiveComponent, IndirectLightingCacheQuality)
+			|| PropertyName == GET_MEMBER_NAME_CHECKED(UPrimitiveComponent, bCastCinematicShadow))
 		{
 			return Mobility == EComponentMobility::Movable;
 		}
diff --git a/Engine/Source/Runtime/Engine/Private/PrimitiveSceneProxy.cpp b/Engine/Source/Runtime/Engine/Private/PrimitiveSceneProxy.cpp
--- a/Engine/Source/Runtime/Engine/Private/PrimitiveSceneProxy.cpp
+++ b/Engine/Source/Runtime/Engine/Private/PrimitiveSceneProxy.cpp
@@ -46,6 +46,7 @@ FPrimitiveSceneProxy::FPrimitiveSceneProxy(const UPrimitiveComponent* InComponen
 ,	bCastShadowAsTwoSided(InComponent->bCastShadowAsTwoSided)
 ,	bSelfShadowOnly(InComponent->bSelfShadowOnly)
 ,	bCastInsetShadow(InComponent->bSelfShadowOnly ? true : InComponent->bCastInsetShadow)	// Assumed to be enabled if bSelfShadowOnly is enabled.
+,	bCastCinematicShadow(InComponent->bCastCinematicShadow)
 ,	bCastFarShadow(InComponent->bCastFarShadow)
 ,	bLightAttachmentsAsGroup(InComponent->bLightAttachmentsAsGroup)
 ,	bStaticElementsAlwaysUseProxyPrimitiveUniformBuffer(false)
diff --git a/Engine/Source/Runtime/Engine/Public/PrimitiveSceneProxy.h b/Engine/Source/Runtime/Engine/Public/PrimitiveSceneProxy.h
--- a/Engine/Source/Runtime/Engine/Public/PrimitiveSceneProxy.h
+++ b/Engine/Source/Runtime/Engine/Public/PrimitiveSceneProxy.h
@@ -371,6 +371,7 @@ public:
 	inline bool CastsShadowAsTwoSided() const { return bCastShadowAsTwoSided; }
 	inline bool CastsSelfShadowOnly() const { return bSelfShadowOnly; }
 	inline bool CastsInsetShadow() const { return bCastInsetShadow; }
+	inline bool CastsCinematicShadow() const { return bCastCinematicShadow; }
 	inline bool CastsFarShadow() const { return bCastFarShadow; }
 	inline bool LightAttachmentsAsGroup() const { return bLightAttachmentsAsGroup; }
 	inline bool StaticElementsAlwaysUseProxyPrimitiveUniformBuffer() const { return bStaticElementsAlwaysUseProxyPrimitiveUniformBuffer; }
@@ -560,6 +561,12 @@ protected:
 	/** Whether this component should create a per-object shadow that gives higher effective shadow resolution. true if bSelfShadowOnly is true. */
 	uint32 bCastInsetShadow : 1;
 
+	/** 
+	 * Whether this component should create a per-object shadow that gives higher effective shadow resolution. 
+	 * Useful for cinematic character shadowing. Assumed to be enabled if bSelfShadowOnly is enabled.
+	 */
+	uint32 bCastCinematicShadow : 1;
+
 	/* When enabled, the component will be rendering into the distant shadow cascades (only for directional lights). */
 	uint32 bCastFarShadow : 1;
 
diff --git a/Engine/Source/Runtime/Engine/Public/SceneManagement.h b/Engine/Source/Runtime/Engine/Public/SceneManagement.h
--- a/Engine/Source/Runtime/Engine/Public/SceneManagement.h
+++ b/Engine/Source/Runtime/Engine/Public/SceneManagement.h
@@ -784,6 +784,7 @@ public:
 	inline bool CastsDynamicShadow() const { return bCastDynamicShadow; }
 	inline bool CastsStaticShadow() const { return bCastStaticShadow; }
 	inline bool CastsTranslucentShadows() const { return bCastTranslucentShadows; }
+	inline bool CastsShadowsFromCinematicObjectsOnly() const { return bCastShadowsFromCinematicObjectsOnly; }
 	inline bool AffectsTranslucentLighting() const { return bAffectTranslucentLighting; }
 	inline bool UseRayTracedDistanceFieldShadows() const { return bUseRayTracedDistanceFieldShadows; }
 	inline float GetRayStartOffsetDepthScale() const { return RayStartOffsetDepthScale; }
@@ -891,6 +892,8 @@ protected:
 	/** Whether the light is allowed to cast dynamic shadows from translucency. */
 	const uint32 bCastTranslucentShadows : 1;
 
+	const uint32 bCastShadowsFromCinematicObjectsOnly : 1;
+
 	/** Whether the light affects translucency or not.  Disabling this can save GPU time when there are many small lights. */
 	const uint32 bAffectTranslucentLighting : 1;
 
diff --git a/Engine/Source/Runtime/Renderer/Private/LightSceneInfo.cpp b/Engine/Source/Runtime/Renderer/Private/LightSceneInfo.cpp
--- a/Engine/Source/Runtime/Renderer/Private/LightSceneInfo.cpp
+++ b/Engine/Source/Runtime/Renderer/Private/LightSceneInfo.cpp
@@ -241,5 +241,10 @@ bool FLightSceneInfoCompact::AffectsPrimitive(const FPrimitiveSceneInfoCompact&
 		return false;
 	}
 
+	if (LightSceneInfo->Proxy->CastsShadowsFromCinematicObjectsOnly() && !CompactPrimitiveSceneInfo.Proxy->CastsCinematicShadow())
+	{
+		return false;
+	}
+
 	return true;
 }
