Integrate - Static meshes now block on outstanding async distance field builds when necessary.  This currently happens during static mesh importing as the mesh is rebuilt redundantly.

diff --git a/Engine/Source/Runtime/Engine/Private/DistanceFieldAtlas.cpp b/Engine/Source/Runtime/Engine/Private/DistanceFieldAtlas.cpp
--- a/Engine/Source/Runtime/Engine/Private/DistanceFieldAtlas.cpp
+++ b/Engine/Source/Runtime/Engine/Private/DistanceFieldAtlas.cpp
@@ -382,6 +382,55 @@ void FDistanceFieldAsyncQueue::AddTask(FAsyncDistanceFieldTask* Task)
 #endif
 }
 
+void FDistanceFieldAsyncQueue::BlockUntilBuildCompleteAndReleaseData(UStaticMesh* StaticMesh)
+{
+	bool bReferenced = false;
+	bool bHadToBlock = false;
+	double StartTime = 0;
+
+	do 
+	{
+		ProcessAsyncTasks();
+
+		bReferenced = false;
+
+		for (int TaskIndex = 0; TaskIndex < ReferencedTasks.Num(); TaskIndex++)
+		{
+			bReferenced = bReferenced || ReferencedTasks[TaskIndex]->StaticMesh == StaticMesh;
+			bReferenced = bReferenced || ReferencedTasks[TaskIndex]->GenerateSource == StaticMesh;
+		}
+
+		if (bReferenced)
+		{
+			if (!bHadToBlock)
+			{
+				StartTime = FPlatformTime::Seconds();
+			}
+
+			bHadToBlock = true;
+			FPlatformProcess::Sleep(.01f);
+		}
+	} 
+	while (bReferenced);
+
+	for (int32 LODIndex = 0; LODIndex < StaticMesh->RenderData->LODResources.Num(); ++LODIndex)
+	{
+		FDistanceFieldVolumeData* DistanceFieldData = StaticMesh->RenderData->LODResources[LODIndex].DistanceFieldData;
+
+		if (DistanceFieldData)
+		{
+			DistanceFieldData->VolumeTexture.Release();
+		}
+	}
+
+	if (bHadToBlock)
+	{
+		UE_LOG(LogStaticMesh, Warning, TEXT("Main thread blocked for %.3fs for async distance field build of %s to complete!  This can happen if the mesh is rebuilt excessively."),
+			(float)(FPlatformTime::Seconds() - StartTime), 
+			*StaticMesh->GetName());
+	}
+}
+
 void FDistanceFieldAsyncQueue::Build(FAsyncDistanceFieldTask* Task, FQueuedThreadPool& ThreadPool)
 {
 #if WITH_EDITOR
diff --git a/Engine/Source/Runtime/Engine/Private/StaticMesh.cpp b/Engine/Source/Runtime/Engine/Private/StaticMesh.cpp
--- a/Engine/Source/Runtime/Engine/Private/StaticMesh.cpp
+++ b/Engine/Source/Runtime/Engine/Private/StaticMesh.cpp
@@ -1700,6 +1700,14 @@ void UStaticMesh::CacheDerivedData()
 	ITargetPlatform* RunningPlatform = TargetPlatformManager.GetRunningTargetPlatform();
 	check(RunningPlatform);
 	const FStaticMeshLODSettings& LODSettings = RunningPlatform->GetStaticMeshLODSettings();
+
+	if (RenderData)
+	{
+		// Finish any previous async builds before modifying RenderData
+		// This can happen during import as the mesh is rebuilt redundantly
+		GDistanceFieldAsyncQueue->BlockUntilBuildCompleteAndReleaseData(this);
+	}
+
 	RenderData = new FStaticMeshRenderData();
 	RenderData->Cache(this, LODSettings);
 
diff --git a/Engine/Source/Runtime/Engine/Public/DistanceFieldAtlas.h b/Engine/Source/Runtime/Engine/Public/DistanceFieldAtlas.h
--- a/Engine/Source/Runtime/Engine/Public/DistanceFieldAtlas.h
+++ b/Engine/Source/Runtime/Engine/Public/DistanceFieldAtlas.h
@@ -187,6 +187,9 @@ public:
 	/** Adds a new build task. */
 	ENGINE_API void AddTask(FAsyncDistanceFieldTask* Task);
 
+	/** Blocks the main thread until the async build of the specified mesh is complete. */
+	ENGINE_API void BlockUntilBuildCompleteAndReleaseData(UStaticMesh* StaticMesh);
+
 	/** Called once per frame, fetches completed tasks and applies them to the scene. */
 	ENGINE_API void ProcessAsyncTasks();
 
