Integrate - Fixed texture coordinates used by [REDACTED_PERSON] not getting picked up by customized UVs

diff --git a/Engine/Source/Runtime/Engine/Private/Materials/HLSLMaterialTranslator.h b/Engine/Source/Runtime/Engine/Private/Materials/HLSLMaterialTranslator.h
--- a/Engine/Source/Runtime/Engine/Private/Materials/HLSLMaterialTranslator.h
+++ b/Engine/Source/Runtime/Engine/Private/Materials/HLSLMaterialTranslator.h
@@ -337,19 +337,22 @@ public:
 				Chunk[CompiledMP_PrevWorldPositionOffset] = Material->CompilePropertyAndSetMaterialProperty(MP_WorldPositionOffset, this, SF_Vertex, true);
 			}
 
+			Chunk[MP_PixelDepthOffset] = Material->CompilePropertyAndSetMaterialProperty(MP_PixelDepthOffset,this);
+
+			// No more calls to non-vertex shader CompilePropertyAndSetMaterialProperty beyond this point
+			const uint32 SavedNumUserTexCoords = NumUserTexCoords;
+
 			for (uint32 CustomUVIndex = MP_CustomizedUVs0; CustomUVIndex <= MP_CustomizedUVs7; CustomUVIndex++)
 			{
 				// Only compile custom UV inputs for UV channels requested by the pixel shader inputs
 				// Any unconnected inputs will have a texcoord generated for them in Material->CompileProperty, which will pass through the vertex (uncustomized) texture coordinates
 				// Note: this is using NumUserTexCoords, which is set by translating all the pixel properties above
-				if (CustomUVIndex - MP_CustomizedUVs0 < NumUserTexCoords)
+				if (CustomUVIndex - MP_CustomizedUVs0 < SavedNumUserTexCoords)
 				{
 					Chunk[CustomUVIndex] = Material->CompilePropertyAndSetMaterialProperty((EMaterialProperty)CustomUVIndex, this);
 				}
 			}
 
-			Chunk[MP_PixelDepthOffset] = Material->CompilePropertyAndSetMaterialProperty(MP_PixelDepthOffset,this);
-
 			bUsesPixelDepthOffset = IsMaterialPropertyUsed(MP_PixelDepthOffset, Chunk[MP_PixelDepthOffset], FLinearColor(0, 0, 0, 0), 1);
 			const bool bUsesWorldPositionOffset = IsMaterialPropertyUsed(MP_WorldPositionOffset, Chunk[MP_WorldPositionOffset], FLinearColor(0, 0, 0, 0), 3);
 			MaterialCompilationOutput.bModifiesMeshPosition = bUsesPixelDepthOffset || bUsesWorldPositionOffset;
@@ -416,6 +419,9 @@ public:
 				}
 			}
 
+			// Catch any modifications to NumUserTexCoords that will not seen by customized UVs
+			check(SavedNumUserTexCoords == NumUserTexCoords);
+
 			ResourcesString = TEXT("");
 
 			// Output the implementation for any custom expressions we will call below.
