Integrate - Added SetIndirectLightingIntensity

diff --git a/Engine/Source/Runtime/Engine/Classes/Components/LightComponent.h b/Engine/Source/Runtime/Engine/Classes/Components/LightComponent.h
--- a/Engine/Source/Runtime/Engine/Classes/Components/LightComponent.h
+++ b/Engine/Source/Runtime/Engine/Classes/Components/LightComponent.h
@@ -170,6 +170,9 @@ public:
 	UFUNCTION(BlueprintCallable, Category="Rendering|Components|Light")
 	void SetIntensity(float NewIntensity);
 
+	UFUNCTION(BlueprintCallable, Category="Rendering|Components|Light")
+	void SetIndirectLightingIntensity(float NewIntensity);
+
 	/** Set color of the light */
 	UFUNCTION(BlueprintCallable, Category="Rendering|Components|Light")
 	void SetLightColor(FLinearColor NewLightColor);
diff --git a/Engine/Source/Runtime/Engine/Classes/Components/SkyLightComponent.h b/Engine/Source/Runtime/Engine/Classes/Components/SkyLightComponent.h
--- a/Engine/Source/Runtime/Engine/Classes/Components/SkyLightComponent.h
+++ b/Engine/Source/Runtime/Engine/Classes/Components/SkyLightComponent.h
@@ -161,6 +161,9 @@ class ENGINE_API USkyLightComponent : public ULightComponentBase
 	UFUNCTION(BlueprintCallable, Category="Rendering|Components|SkyLight")
 	void SetIntensity(float NewIntensity);
 
+	UFUNCTION(BlueprintCallable, Category="Rendering|Components|Light")
+	void SetIndirectLightingIntensity(float NewIntensity);
+
 	/** Set color of the light */
 	UFUNCTION(BlueprintCallable, Category="Rendering|Components|SkyLight")
 	void SetLightColor(FLinearColor NewLightColor);
diff --git a/Engine/Source/Runtime/Engine/Private/LightComponent.cpp b/Engine/Source/Runtime/Engine/Private/LightComponent.cpp
--- a/Engine/Source/Runtime/Engine/Private/LightComponent.cpp
+++ b/Engine/Source/Runtime/Engine/Private/LightComponent.cpp
@@ -687,6 +687,23 @@ void ULightComponent::SetIntensity(float NewIntensity)
 	}
 }
 
+void ULightComponent::SetIndirectLightingIntensity(float NewIntensity)
+{
+	// Can't set brightness on a static light
+	if ((IsRunningUserConstructionScript() || !(IsRegistered() && Mobility == EComponentMobility::Static))
+		&& IndirectLightingIntensity != NewIntensity)
+	{
+		IndirectLightingIntensity = NewIntensity;
+
+		// Use lightweight color and brightness update 
+		if( World && World->Scene )
+		{
+			//@todo - remove from scene if brightness or color becomes 0
+			World->Scene->UpdateLightColorAndBrightness( this );
+		}
+	}
+}
+
 /** Set color of the light */
 void ULightComponent::SetLightColor(FLinearColor NewLightColor)
 {
diff --git a/Engine/Source/Runtime/Engine/Private/SkyLightComponent.cpp b/Engine/Source/Runtime/Engine/Private/SkyLightComponent.cpp
--- a/Engine/Source/Runtime/Engine/Private/SkyLightComponent.cpp
+++ b/Engine/Source/Runtime/Engine/Private/SkyLightComponent.cpp
@@ -476,6 +476,17 @@ void USkyLightComponent::SetIntensity(float NewIntensity)
 	}
 }
 
+void USkyLightComponent::SetIndirectLightingIntensity(float NewIntensity)
+{
+	// Can't set brightness on a static light
+	if ((IsRunningUserConstructionScript() || !(IsRegistered() && Mobility == EComponentMobility::Static))
+		&& IndirectLightingIntensity != NewIntensity)
+	{
+		IndirectLightingIntensity = NewIntensity;
+		UpdateLimitedRenderingStateFast();
+	}
+}
+
 /** Set color of the light */
 void USkyLightComponent::SetLightColor(FLinearColor NewLightColor)
 {
