Compile fix

diff --git a/Engine/Source/Runtime/Engine/Private/SkyLightComponent.cpp b/Engine/Source/Runtime/Engine/Private/SkyLightComponent.cpp
--- a/Engine/Source/Runtime/Engine/Private/SkyLightComponent.cpp
+++ b/Engine/Source/Runtime/Engine/Private/SkyLightComponent.cpp
@@ -479,7 +479,7 @@ void USkyLightComponent::SetIntensity(float NewIntensity)
 void USkyLightComponent::SetIndirectLightingIntensity(float NewIntensity)
 {
 	// Can't set brightness on a static light
-	if ((IsRunningUserConstructionScript() || !(IsRegistered() && Mobility == EComponentMobility::Static))
+	if (AreDynamicDataChangesAllowed()
 		&& IndirectLightingIntensity != NewIntensity)
 	{
 		IndirectLightingIntensity = NewIntensity;
