Added ability to toggle an arbitrary CVar.

(Merging CL 2457699 to main.)

diff --git a/Engine/Source/Runtime/Engine/Private/GameEngine.cpp b/Engine/Source/Runtime/Engine/Private/GameEngine.cpp
--- a/Engine/Source/Runtime/Engine/Private/GameEngine.cpp
+++ b/Engine/Source/Runtime/Engine/Private/GameEngine.cpp
@@ -583,6 +583,49 @@ bool UGameEngine::Exec( UWorld* InWorld, const TCHAR* Cmd, FOutputDevice& Ar )
 	{
 		return HandleCancelCommand( Cmd, Ar, InWorld );	
 	}
+	else if ( FParse::Command( &Cmd, TEXT("TOGGLECVAR") ) )
+	{
+		FString CVarName;
+		FParse::Token(Cmd, CVarName, false);
+
+		bool bEnoughParamsSupplied = false;
+		IConsoleVariable * CVar = nullptr;
+
+		if (CVarName.Len() > 0)
+		{
+			CVar = IConsoleManager::Get().FindConsoleVariable(*CVarName);
+		}
+
+		if (CVar)
+		{
+			// values to toggle between
+			FString StringVal1, StringVal2;
+			
+			if (FParse::Token(Cmd, StringVal1, false))
+			{
+				if (FParse::Token(Cmd, StringVal2, false))
+				{
+					bEnoughParamsSupplied = true;
+					FString CurrentValue = CVar->GetString();
+
+					FString Command(FString::Printf(TEXT("%s %s"), *CVarName, (CurrentValue == StringVal1) ? *StringVal2 : *StringVal1));
+					GEngine->Exec(InWorld, *Command);
+				}
+			}
+		}
+		else
+		{
+			Ar.Log(*FString::Printf(TEXT("TOGGLECVAR: cvar '%s' was not found"), *CVarName));
+			bEnoughParamsSupplied = true;	// cannot say anything about the rest of parameters
+		}
+		
+		if (!bEnoughParamsSupplied)
+		{
+			Ar.Log(TEXT("Usage: TOGGLECVAR CVarName Value1 Value2"));
+		}
+
+		return true;
+	}
 #if !UE_BUILD_SHIPPING
 	else if( FParse::Command( &Cmd, TEXT("ApplyUserSettings") ) )
 	{
