Compile fix

diff --git a/Engine/Source/Runtime/Engine/Private/LightComponent.cpp b/Engine/Source/Runtime/Engine/Private/LightComponent.cpp
--- a/Engine/Source/Runtime/Engine/Private/LightComponent.cpp
+++ b/Engine/Source/Runtime/Engine/Private/LightComponent.cpp
@@ -690,7 +690,7 @@ void ULightComponent::SetIntensity(float NewIntensity)
 void ULightComponent::SetIndirectLightingIntensity(float NewIntensity)
 {
 	// Can't set brightness on a static light
-	if ((IsRunningUserConstructionScript() || !(IsRegistered() && Mobility == EComponentMobility::Static))
+	if (AreDynamicDataChangesAllowed()
 		&& IndirectLightingIntensity != NewIntensity)
 	{
 		IndirectLightingIntensity = NewIntensity;
