fixed PostProcessMaterial SwitchParameter as suggested here:
https://udn.unrealengine.com/questions/238114/issue-with-switchparameter-of-postprocessmaterial.html

diff --git a/Engine/Source/Runtime/Engine/Private/Materials/Material.cpp b/Engine/Source/Runtime/Engine/Private/Materials/Material.cpp
--- a/Engine/Source/Runtime/Engine/Private/Materials/Material.cpp
+++ b/Engine/Source/Runtime/Engine/Private/Materials/Material.cpp
@@ -1129,7 +1129,7 @@ void UMaterialInterface::OverrideBlendableSettings(class FSceneView& View, float
 	}
 	else
 	{
-		UMaterialInstanceDynamic* MID = View.State->GetReusableMID((UMaterialInterface*)Material);//, (UMaterialInterface*)this);
+		UMaterialInstanceDynamic* MID = View.State->GetReusableMID((UMaterialInterface*)this);
 
 		if(MID)
 		{
diff --git a/Engine/Source/Runtime/Engine/Private/Materials/MaterialInstance.cpp b/Engine/Source/Runtime/Engine/Private/Materials/MaterialInstance.cpp
--- a/Engine/Source/Runtime/Engine/Private/Materials/MaterialInstance.cpp
+++ b/Engine/Source/Runtime/Engine/Private/Materials/MaterialInstance.cpp
@@ -2308,21 +2308,19 @@ void UMaterialInstance::OverrideBlendableSettings(class FSceneView& View, float
 		// do we partly want to fade this one in.
 		if(Weight < 1.0f)
 		{
-			UMaterial* Base = Material->GetBaseMaterial();
-
-			UMaterialInstanceDynamic* MID = View.State->GetReusableMID((UMaterialInterface*)Base);//, (UMaterialInterface*)this);
+			UMaterialInstanceDynamic* MID = View.State->GetReusableMID((UMaterialInterface*)this);
 
 			if(MID)
 			{
-				MID->K2_CopyMaterialInstanceParameters((UMaterialInterface*)Base);
+				MID->K2_CopyMaterialInstanceParameters((UMaterialInterface*)Material);
 
-				FPostProcessMaterialNode NewNode(MID, Base->BlendableLocation, Base->BlendablePriority);
+				FPostProcessMaterialNode NewNode(MID, Material->BlendableLocation, Material->BlendablePriority);
 
 				// it's the first material, no blending needed
 				Dest.BlendableManager.PushBlendableData(1.0f, NewNode);
 
 				// can be optimized
-				PostProcessMaterialNode = IteratePostProcessMaterialNodes(Dest, Base, Iterator);
+				PostProcessMaterialNode = IteratePostProcessMaterialNodes(Dest, Material, Iterator);
 			}
 		}
 	}
