Turn off start awake flag for destruction temporarily. Waiting on flag from APEX

diff --git a/Engine/Source/Runtime/Engine/Private/DestructibleComponent.cpp b/Engine/Source/Runtime/Engine/Private/DestructibleComponent.cpp
--- a/Engine/Source/Runtime/Engine/Private/DestructibleComponent.cpp
+++ b/Engine/Source/Runtime/Engine/Private/DestructibleComponent.cpp
@@ -342,7 +342,8 @@ void UDestructibleComponent::CreatePhysicsState()
 
 
 	//  Put to sleep or wake up only if the component is physics-simulated
-	if (PRootActor != NULL && BodyInstance.bSimulatePhysics)
+	//TODO: APEX 1.3.3 defers adding actors to scene which means this doesn't work. Waiting on flag from them so we can do this properly.
+	/*if (PRootActor != NULL && BodyInstance.bSimulatePhysics)
 	{
 		SCOPED_SCENE_WRITE_LOCK(PRootActor->getScene());
 
@@ -357,7 +358,7 @@ void UDestructibleComponent::CreatePhysicsState()
 		{
 			PRootActor->putToSleep();
 		}
-	}
+	}*/
 
 	UpdateBounds();
 #endif	// #if WITH_APEX
