Don't call BeginPlay from AWorldSettings::NotifyBeginPlay on Actor's that have already begun play (presumably because they seamlessly travelled)

diff --git a/Engine/Source/Runtime/Engine/Private/WorldSettings.cpp b/Engine/Source/Runtime/Engine/Private/WorldSettings.cpp
--- a/Engine/Source/Runtime/Engine/Private/WorldSettings.cpp
+++ b/Engine/Source/Runtime/Engine/Private/WorldSettings.cpp
@@ -125,7 +125,11 @@ void AWorldSettings::NotifyBeginPlay()
 	{
 		for (FActorIterator It(World); It; ++It)
 		{
-			It->BeginPlay();
+			// Actors that have traveled seamlessly from other levels already had BeginPlay called in that level
+			if (!It->HasActorBegunPlay())
+			{
+				It->BeginPlay();
+			}
 		}
 		World->bBegunPlay = true;
 	}
