Get old behavior working for bFirstPartitionIsDepthZero

diff --git a/Engine/Source/Runtime/Engine/Private/DestructibleFractureSettings.cpp b/Engine/Source/Runtime/Engine/Private/DestructibleFractureSettings.cpp
--- a/Engine/Source/Runtime/Engine/Private/DestructibleFractureSettings.cpp
+++ b/Engine/Source/Runtime/Engine/Private/DestructibleFractureSettings.cpp
@@ -369,11 +369,12 @@ bool UDestructibleFractureSettings::SetRootMesh(const TArray<NxExplicitRenderTri
 												const TArray<uint32>& MeshPartition, bool bFirstPartitionIsDepthZero)
 {
 	bool Success = false;
+	physx::PxI32 NegativeOne = -1;
 
 	if (ApexDestructibleAssetAuthoring != NULL)
 	{
 		Success = ApexDestructibleAssetAuthoring->setRootMesh(MeshTriangles.GetData(), MeshTriangles.Num(), SubmeshData.GetData(), 
-															  SubmeshData.Num(), (uint32*)MeshPartition.GetData(), MeshPartition.Num(), nullptr);	//todo: this is ignoring bFirstPartitionIsDepthZero. Need to find out more about the new parameters to this function
+															  SubmeshData.Num(), (uint32*)MeshPartition.GetData(), MeshPartition.Num(), &NegativeOne, bFirstPartitionIsDepthZero ? 1 : 0);
 		if (Success)
 		{
 			NxCollisionDesc CollisionDesc;
