removed error/verify if the ini setting MinTextureResidentMipCount is missing in the ini (could be made a cvar in which case the c++ code has the default anyway)

diff --git a/Engine/Source/Runtime/Engine/Private/SystemSettings.cpp b/Engine/Source/Runtime/Engine/Private/SystemSettings.cpp
--- a/Engine/Source/Runtime/Engine/Private/SystemSettings.cpp
+++ b/Engine/Source/Runtime/Engine/Private/SystemSettings.cpp
@@ -344,8 +344,8 @@ void FSystemSettings::Initialize( bool bSetupForEditor )
 	IConsoleManager::Get().RegisterConsoleVariableSink_Handle(FConsoleCommandDelegate::CreateRaw(this, &FSystemSettings::CVarSink));
 
 	// intialize a critical texture streaming value used by texture loading, etc
-	int32 MinTextureResidentMipCount = 1;
-	verify(GConfig->GetInt(TEXT("TextureStreaming"), TEXT("MinTextureResidentMipCount"), MinTextureResidentMipCount, GEngineIni));
+	int32 MinTextureResidentMipCount = 7;
+	GConfig->GetInt(TEXT("TextureStreaming"), TEXT("MinTextureResidentMipCount"), MinTextureResidentMipCount, GEngineIni);
 	UTexture2D::SetMinTextureResidentMipCount(MinTextureResidentMipCount);
 }
 
