GitHub pull request: FSceneViewport::OnMouseLeave not working for non PIE game launches
https://github.com/EpicGames/UnrealEngine/pull/862/

diff --git a/Engine/Source/Runtime/Engine/Private/Slate/SceneViewport.cpp b/Engine/Source/Runtime/Engine/Private/Slate/SceneViewport.cpp
--- a/Engine/Source/Runtime/Engine/Private/Slate/SceneViewport.cpp
+++ b/Engine/Source/Runtime/Engine/Private/Slate/SceneViewport.cpp
@@ -475,7 +475,7 @@ void FSceneViewport::OnMouseLeave( const FPointerEvent& MouseEvent )
 	{
 		ViewportClient->MouseLeave( this );
 	
-		if ( IsPlayInEditorViewport() )
+		if ( (FApp::IsGame() && !GIsEditor) || IsPlayInEditorViewport() )
 		{
 			CachedMousePos = FIntPoint(-1, -1);
 		}
