fixed Subsurface_Profile causing wrong colors when using extreme inputs (caused by floating point precision)

diff --git a/Engine/Source/Runtime/Engine/Private/Rendering/SubsurfaceProfile.cpp b/Engine/Source/Runtime/Engine/Private/Rendering/SubsurfaceProfile.cpp
--- a/Engine/Source/Runtime/Engine/Private/Rendering/SubsurfaceProfile.cpp
+++ b/Engine/Source/Runtime/Engine/Private/Rendering/SubsurfaceProfile.cpp
@@ -122,6 +122,17 @@ void FSubsurfaceProfileTexture::ReleaseDynamicRHI()
 	GSSProfiles.SafeRelease();
 }
 
+static float GetNextSmallerPositiveFloat(float x)
+{
+	check(x > 0);
+	uint32 bx = *(uint32 *)&x;
+
+	// float are ordered like int, at least for the positive part
+	uint32 ax = bx - 1;
+
+	return *(float *)&ax;
+}
+
 void FSubsurfaceProfileTexture::CreateTexture(FRHICommandListImmediate& RHICmdList)
 {
 	// call SetRendererModule() is missing
@@ -160,6 +171,9 @@ void FSubsurfaceProfileTexture::CreateTexture(FRHICommandListImmediate& RHICmdLi
 
 	FLinearColor kernel[Width];
 
+	const float FloatScale = GetNextSmallerPositiveFloat(0x10000);
+	check((int32)GetNextSmallerPositiveFloat(0x10000) == 0xffff);
+
 	for (uint32 y = 0; y < Height; ++y)
 	{
 		FSubsurfaceProfileStruct Data = SubsurfaceProfileEntries[y].Settings;
@@ -190,12 +204,15 @@ void FSubsurfaceProfileTexture::CreateTexture(FRHICommandListImmediate& RHICmdLi
 
 			if (b16Bit)
 			{
+				// scale from 0..1 to 0..0xffff
+				// scale with 0x10000 and round down to evenly distribute, avoid 0x10000
+
 				uint16* Dest = (uint16*)(DestBuffer + DestStride * y);
 
-				Dest[Pos * 4 + 0] = (uint16)(C.X * (256 * 256 - 0.0001f));
-				Dest[Pos * 4 + 1] = (uint16)(C.Y * (256 * 256 - 0.0001f));
-				Dest[Pos * 4 + 2] = (uint16)(C.Z * (256 * 256 - 0.0001f));
-				Dest[Pos * 4 + 3] = (uint16)(C.W * (256 * 256 - 0.0001f));
+				Dest[Pos * 4 + 0] = (uint16)(C.X * FloatScale);
+				Dest[Pos * 4 + 1] = (uint16)(C.Y * FloatScale);
+				Dest[Pos * 4 + 2] = (uint16)(C.Z * FloatScale);
+				Dest[Pos * 4 + 3] = (uint16)(C.W * FloatScale);
 			}
 			else
 			{
