Fixed a call to FString's ParseIntoArray family of functions that was passing a pointer rather than reference. Perforce failed to properly change which changelist this change was in, so it was separated from the rest of its kind.

diff --git a/Engine/Source/Runtime/Engine/Private/StaticMeshComponent.cpp b/Engine/Source/Runtime/Engine/Private/StaticMeshComponent.cpp
--- a/Engine/Source/Runtime/Engine/Private/StaticMeshComponent.cpp
+++ b/Engine/Source/Runtime/Engine/Private/StaticMeshComponent.cpp
@@ -1852,7 +1852,7 @@ void FStaticMeshComponentLODInfo::ImportText(const TCHAR** SourceText)
 		bool bValidInput = true;
 
 		// Tokenize the text
-		SourceTextStr.ParseIntoArray(&Tokens, TEXT(","), false);
+		SourceTextStr.ParseIntoArray(Tokens, TEXT(","), false);
 
 		// There should be 11 tokens per vertex
 		check(Tokens.Num() * 11 >= VertCount);
