Async distance field builds are now forced to be complete before cooking distance field data

diff --git a/Engine/Source/Runtime/Engine/Private/DistanceFieldAtlas.cpp b/Engine/Source/Runtime/Engine/Private/DistanceFieldAtlas.cpp
--- a/Engine/Source/Runtime/Engine/Private/DistanceFieldAtlas.cpp
+++ b/Engine/Source/Runtime/Engine/Private/DistanceFieldAtlas.cpp
@@ -382,7 +382,7 @@ void FDistanceFieldAsyncQueue::AddTask(FAsyncDistanceFieldTask* Task)
 #endif
 }
 
-void FDistanceFieldAsyncQueue::BlockUntilBuildCompleteAndReleaseData(UStaticMesh* StaticMesh)
+void FDistanceFieldAsyncQueue::BlockUntilBuildComplete(UStaticMesh* StaticMesh, bool bWarnIfBlocked)
 {
 	bool bReferenced = false;
 	bool bHadToBlock = false;
@@ -413,17 +413,7 @@ void FDistanceFieldAsyncQueue::BlockUntilBuildCompleteAndReleaseData(UStaticMesh
 	} 
 	while (bReferenced);
 
-	for (int32 LODIndex = 0; LODIndex < StaticMesh->RenderData->LODResources.Num(); ++LODIndex)
-	{
-		FDistanceFieldVolumeData* DistanceFieldData = StaticMesh->RenderData->LODResources[LODIndex].DistanceFieldData;
-
-		if (DistanceFieldData)
-		{
-			DistanceFieldData->VolumeTexture.Release();
-		}
-	}
-
-	if (bHadToBlock)
+	if (bHadToBlock && bWarnIfBlocked)
 	{
 		UE_LOG(LogStaticMesh, Warning, TEXT("Main thread blocked for %.3fs for async distance field build of %s to complete!  This can happen if the mesh is rebuilt excessively."),
 			(float)(FPlatformTime::Seconds() - StartTime), 
diff --git a/Engine/Source/Runtime/Engine/Private/StaticMesh.cpp b/Engine/Source/Runtime/Engine/Private/StaticMesh.cpp
--- a/Engine/Source/Runtime/Engine/Private/StaticMesh.cpp
+++ b/Engine/Source/Runtime/Engine/Private/StaticMesh.cpp
@@ -636,6 +636,11 @@ void FStaticMeshRenderData::Serialize(FArchive& Ar, UStaticMesh* Owner, bool bCo
 		FStripDataFlags StripFlags( Ar );
 		if ( !StripFlags.IsDataStrippedForServer() )
 		{
+			if (Ar.IsSaving())
+			{
+				GDistanceFieldAsyncQueue->BlockUntilBuildComplete(Owner, false);
+			}
+
 			for (int32 ResourceIndex = 0; ResourceIndex < LODResources.Num(); ResourceIndex++)
 			{
 				FStaticMeshLODResources& LOD = LODResources[ResourceIndex];
@@ -1705,7 +1710,18 @@ void UStaticMesh::CacheDerivedData()
 	{
 		// Finish any previous async builds before modifying RenderData
 		// This can happen during import as the mesh is rebuilt redundantly
-		GDistanceFieldAsyncQueue->BlockUntilBuildCompleteAndReleaseData(this);
+		GDistanceFieldAsyncQueue->BlockUntilBuildComplete(this, true);
+
+		for (int32 LODIndex = 0; LODIndex < RenderData->LODResources.Num(); ++LODIndex)
+		{
+			FDistanceFieldVolumeData* DistanceFieldData = RenderData->LODResources[LODIndex].DistanceFieldData;
+
+			if (DistanceFieldData)
+			{
+				// Release before destroying RenderData
+				DistanceFieldData->VolumeTexture.Release();
+			}
+		}
 	}
 
 	RenderData = new FStaticMeshRenderData();
diff --git a/Engine/Source/Runtime/Engine/Public/DistanceFieldAtlas.h b/Engine/Source/Runtime/Engine/Public/DistanceFieldAtlas.h
--- a/Engine/Source/Runtime/Engine/Public/DistanceFieldAtlas.h
+++ b/Engine/Source/Runtime/Engine/Public/DistanceFieldAtlas.h
@@ -188,7 +188,7 @@ public:
 	ENGINE_API void AddTask(FAsyncDistanceFieldTask* Task);
 
 	/** Blocks the main thread until the async build of the specified mesh is complete. */
-	ENGINE_API void BlockUntilBuildCompleteAndReleaseData(UStaticMesh* StaticMesh);
+	ENGINE_API void BlockUntilBuildComplete(UStaticMesh* StaticMesh, bool bWarnIfBlocked);
 
 	/** Called once per frame, fetches completed tasks and applies them to the scene. */
 	ENGINE_API void ProcessAsyncTasks();
