Add Begin/EndPlay virtuals to ActorComponent.  BeginPlay called if bWantsBeginPlay being true, EndPlay called if BeginPlay called
Removed blueprint exposure of Initialize/UninitializeComponent and replaced with Begin/EndPlay
EEndPlayReason::ActorDestroyed renamed EEndPlayReason::Destroyed
UActorComponent::bHasBeenCreated and bHasBeenInitialized made private and exposed via Get accessors

diff --git a/Engine/Config/BaseEngine.ini b/Engine/Config/BaseEngine.ini
--- a/Engine/Config/BaseEngine.ini
+++ b/Engine/Config/BaseEngine.ini
@@ -756,6 +756,11 @@ bUseStreamingPause=false
 +K2FieldRedirects=(OldFieldName="PrimitiveComponent.GetTouchingComponents",NewFieldName="PrimitiveComponent.GetOverlappingComponents")
 +K2ParamRedirects=(NodeName="/Script/BlueprintGraph.K2Node_CallFunction", OldParamName="PrimitiveComponent.GetOverlappingComponents.TouchingComponents", NewParamName="OverlappingComponents")
 
++K2FieldRedirects=(OldFieldName="ActorComponent.ReceiveInitializeComponent",NewFieldName="ActorComponent.ReceiveBeginPlay")
++K2FieldRedirects=(OldFieldName="ActorComponent.ReceiveUninitializeComponent",NewFieldName="ActorComponent.ReceiveEndPlay")
++TaggedPropertyRedirects=(ClassName="ActorComponent",OldPropertyName="bWantsInitializeComponent",NewPropertyName="bWantsBeginPlay")
++EnumRedirects=(EnumName="EEndPlayReason",OldEnumEntry="EEndPlayReason::ActorDestroyed",NewEnumEntry="EEndPlayReason::Destroyed")
+
 +K2FieldRedirects=(OldFieldName="Actor.AddActorLocalTranslation",NewFieldName="Actor.K2_AddActorLocalOffset")
 +K2ParamRedirects=(NodeName="/Script/BlueprintGraph.K2Node_CallFunction", OldParamName="Actor.K2_AddActorLocalOffset.DeltaTranslation", NewParamName="DeltaLocation")
 
diff --git a/Engine/Source/Editor/KismetCompiler/Private/KismetCompiler.cpp b/Engine/Source/Editor/KismetCompiler/Private/KismetCompiler.cpp
--- a/Engine/Source/Editor/KismetCompiler/Private/KismetCompiler.cpp
+++ b/Engine/Source/Editor/KismetCompiler/Private/KismetCompiler.cpp
@@ -3579,7 +3579,7 @@ void FKismetCompilerContext::Compile()
 
 		CopyTermDefaultsToDefaultObject(NewCDO);
 		SetCanEverTick();
-		SetWantsInitialize();
+		SetWantsBeginPlay();
 		FKismetCompilerUtilities::ValidateEnumProperties(NewCDO, MessageLog);
 	}
 
@@ -3851,7 +3851,7 @@ void FKismetCompilerContext::SetCanEverTick() const
 	}
 }
 
-void FKismetCompilerContext::SetWantsInitialize() const
+void FKismetCompilerContext::SetWantsBeginPlay() const
 {
 	UActorComponent* CDComponent = Cast<UActorComponent>(NewClass->GetDefaultObject());
 
@@ -3860,26 +3860,26 @@ void FKismetCompilerContext::SetWantsInitialize() const
 		return;
 	}
 
-	const bool bOldFlag = CDComponent->bWantsInitializeComponent;
+	const bool bOldFlag = CDComponent->bWantsBeginPlay;
 
-	CDComponent->bWantsInitializeComponent = NewClass->GetSuperClass()->GetDefaultObject<UActorComponent>()->bWantsInitializeComponent;
+	CDComponent->bWantsBeginPlay = NewClass->GetSuperClass()->GetDefaultObject<UActorComponent>()->bWantsBeginPlay;
 
-	if (!CDComponent->bWantsInitializeComponent)
+	if (!CDComponent->bWantsBeginPlay)
 	{
-		static FName ReceiveInitializeComponentName(GET_FUNCTION_NAME_CHECKED(UActorComponent, ReceiveInitializeComponent));
-		static FName ReceiveUninitializeComponentName(GET_FUNCTION_NAME_CHECKED(UActorComponent, ReceiveUninitializeComponent));
-		const UFunction* ReciveInitializeComponentEvent = FKismetCompilerUtilities::FindOverriddenImplementableEvent(ReceiveInitializeComponentName, NewClass);
-		const UFunction* ReciveUninitializeComponentEvent = FKismetCompilerUtilities::FindOverriddenImplementableEvent(ReceiveUninitializeComponentName, NewClass);
-		if (ReciveInitializeComponentEvent || ReciveUninitializeComponentEvent)
+		static FName ReceiveBeginPlayName(GET_FUNCTION_NAME_CHECKED(UActorComponent, ReceiveBeginPlay));
+		static FName ReceiveEndPlayName(GET_FUNCTION_NAME_CHECKED(UActorComponent, ReceiveEndPlay));
+		const UFunction* ReciveBeginPlayEvent = FKismetCompilerUtilities::FindOverriddenImplementableEvent(ReceiveBeginPlayName, NewClass);
+		const UFunction* ReciveEndPlayEvent = FKismetCompilerUtilities::FindOverriddenImplementableEvent(ReceiveEndPlayName, NewClass);
+		if (ReciveBeginPlayEvent || ReciveEndPlayEvent)
 		{
-			CDComponent->bWantsInitializeComponent = true;
+			CDComponent->bWantsBeginPlay = true;
 		}
 	}
 
-	if(CDComponent->bWantsInitializeComponent != bOldFlag)
+	if(CDComponent->bWantsBeginPlay != bOldFlag)
 	{
-		UE_LOG(LogK2Compiler, Verbose, TEXT("Overridden flag for class '%s': bWantsInitializeComponent %s "), *NewClass->GetName(),
-			CDComponent->bWantsInitializeComponent ? *(GTrue.ToString()) : *(GFalse.ToString()) );
+		UE_LOG(LogK2Compiler, Verbose, TEXT("Overridden flag for class '%s': bWantsBeginPlay %s "), *NewClass->GetName(),
+			CDComponent->bWantsBeginPlay ? *(GTrue.ToString()) : *(GFalse.ToString()) );
 	}
 }
 
diff --git a/Engine/Source/Editor/KismetCompiler/Public/KismetCompiler.h b/Engine/Source/Editor/KismetCompiler/Public/KismetCompiler.h
--- a/Engine/Source/Editor/KismetCompiler/Public/KismetCompiler.h
+++ b/Engine/Source/Editor/KismetCompiler/Public/KismetCompiler.h
@@ -243,9 +243,9 @@ protected:
 
 	/** 
 	 * Function works only if subclass UActorComponent.
-	 * If InitializeComponent or UninitializeComponent event is defined, force bWantsInitializeComponent.
+	 * If BeginPlay or EndPlay event is defined on the Component, force bWantsBeginPlay.
 	 */
-	void SetWantsInitialize() const;
+	void SetWantsBeginPlay() const;
 
 	/** Scan FunctionList and return Entry point, for matching one  */
 	const UK2Node_FunctionEntry* FindLocalEntryPoint(const UFunction* Function) const;
diff --git a/Engine/Source/Editor/UnrealEd/Private/Kismet2/Kismet2.cpp b/Engine/Source/Editor/UnrealEd/Private/Kismet2/Kismet2.cpp
--- a/Engine/Source/Editor/UnrealEd/Private/Kismet2/Kismet2.cpp
+++ b/Engine/Source/Editor/UnrealEd/Private/Kismet2/Kismet2.cpp
@@ -513,7 +513,7 @@ UBlueprint* FKismetEditorUtilities::CreateBlueprint(UClass* ParentClass, UObject
 				else if(NewBP->GeneratedClass->IsChildOf(UActorComponent::StaticClass()))
 				{
 					FKismetEditorUtilities::AddDefaultEventNode(NewBP, NewBP->UbergraphPages[0], FName(TEXT("ReceiveTick")), UActorComponent::StaticClass(), NodePositionY);
-					FKismetEditorUtilities::AddDefaultEventNode(NewBP, NewBP->UbergraphPages[0], FName(TEXT("ReceiveInitializeComponent")), UActorComponent::StaticClass(), NodePositionY);
+					FKismetEditorUtilities::AddDefaultEventNode(NewBP, NewBP->UbergraphPages[0], FName(TEXT("ReceiveBeginPlay")), UActorComponent::StaticClass(), NodePositionY);
 				}
 				else if(NewBP->GeneratedClass->IsChildOf(WidgetClass))
 				{
diff --git a/Engine/Source/Runtime/Engine/Classes/Components/ActorComponent.h b/Engine/Source/Runtime/Engine/Classes/Components/ActorComponent.h
--- a/Engine/Source/Runtime/Engine/Classes/Components/ActorComponent.h
+++ b/Engine/Source/Runtime/Engine/Classes/Components/ActorComponent.h
@@ -115,16 +115,29 @@ public:
 	UPROPERTY(transient, ReplicatedUsing=OnRep_IsActive)
 	uint32 bIsActive:1;
 
-	/** If TRUE, we call the virtual InitializeComponent */
-	UPROPERTY()
+	/** If true, we call the virtual InitializeComponent */
 	uint32 bWantsInitializeComponent:1;
 
+	/** If true, we call the virtual BeginPlay */
+	UPROPERTY()
+	uint32 bWantsBeginPlay:1;
+
+private:
 	/** Indicates that OnCreatedComponent has been called, but OnDestroyedComponent has not yet */
 	uint32 bHasBeenCreated:1;
 
 	/** Indicates that InitializeComponent has been called, but UninitializeComponent has not yet */
 	uint32 bHasBeenInitialized:1;
 
+	/** Indicates that BeginPlay has been called, but EndPlay has not yet */
+	uint32 bHasBegunPlay:1;
+
+public:
+
+	bool HasBeenCreated() const { return bHasBeenCreated; }
+	bool HasBeenInitialized() const { return bHasBeenInitialized; }
+	bool HasBegunPlay() const { return bHasBegunPlay; }
+
 	UPROPERTY()
 	EComponentCreationMethod CreationMethod;
 
@@ -291,24 +304,42 @@ protected:
 
 public:
 	/**
-	 * Starts gameplay for this component.
-	 * Requires component to be registered, and bWantsInitializeComponent to be TRUE.
+	 * Initializes the component.  Occurs at level startup. This is before BeginPlay (Actor or Component).  
+	 * All Components in the level will be Initialized on load before any Actor/Component gets BeginPlay
+	 * Requires component to be registered, and bWantsInitializeComponent to be true.
 	 */
 	virtual void InitializeComponent();
 
-	/** Event when the component is initialized, either via creation or its Actor's BeginPlay. */
-	UFUNCTION(BlueprintImplementableEvent, meta=(Keywords = "Begin", FriendlyName = "Initialize Component"))
-	virtual void ReceiveInitializeComponent();
+	/**
+	 * BeginsPlay for the component.  Occurs at level startup. This is before BeginPlay (Actor or Component).  
+	 * All Components (that want initialization) in the level will be Initialized on load before any 
+	 * Actor/Component gets BeginPlay.
+	 * Requires component to be registered and initialized.
+	 */
+	virtual void BeginPlay();
+
+	/** 
+	 * Blueprint implementable event for when the component is beginning play, called before its Owner's BeginPlay on Actor BeginPlay 
+	 * or when the component is dynamically created if the Actor has already BegunPlay. 
+	 */
+	UFUNCTION(BlueprintImplementableEvent, meta=(FriendlyName = "Begin Play"))
+	virtual void ReceiveBeginPlay();
 
 	/**
 	 * Ends gameplay for this component.
+	 * Called from AActor::EndPlay only if bHasBegunPlay is true
+	 */
+	virtual void EndPlay(const EEndPlayReason::Type EndPlayReason);
+
+	/**
+	 * Handle this component being Uninitialized.
 	 * Called from AActor::EndPlay only if bHasBeenInitialized is true
 	 */
 	virtual void UninitializeComponent();
 
-	/** Event when the component is uninitialized, generally via descruction or its Actor's EndPlay. */
-	UFUNCTION(BlueprintImplementableEvent, meta=(Keywords = "End Delete", FriendlyName = "Uninitialize Component"))
-	virtual void ReceiveUninitializeComponent();
+	/** Blueprint implementable event for when the component ends play, generally via destruction or its Actor's EndPlay. */
+	UFUNCTION(BlueprintImplementableEvent, meta=(Keywords = "delete", FriendlyName = "End Play"))
+	virtual void ReceiveEndPlay(EEndPlayReason::Type EndPlayReason);
 	
 	/**
 	 * When called, will call the virtual call chain to register all of the tick functions
diff --git a/Engine/Source/Runtime/Engine/Classes/Engine/EngineTypes.h b/Engine/Source/Runtime/Engine/Classes/Engine/EngineTypes.h
--- a/Engine/Source/Runtime/Engine/Classes/Engine/EngineTypes.h
+++ b/Engine/Source/Runtime/Engine/Classes/Engine/EngineTypes.h
@@ -2361,11 +2361,11 @@ namespace EEndPlayReason
 {
 	enum Type
 	{
-		ActorDestroyed,		// When the Actor is explicitly destroyed
+		Destroyed,			// When the Actor or Component is explicitly destroyed
 		LevelTransition,	// When the world is being unloaded for a level transition
 		EndPlayInEditor,	// When the world is being unloaded because PIE is ending
 		RemovedFromWorld,	// When the level it is a member of is streamed out
-		Quit,			// When the application is being exited
+		Quit,				// When the application is being exited
 	};
 
 }
diff --git a/Engine/Source/Runtime/Engine/Private/Actor.cpp b/Engine/Source/Runtime/Engine/Private/Actor.cpp
--- a/Engine/Source/Runtime/Engine/Private/Actor.cpp
+++ b/Engine/Source/Runtime/Engine/Private/Actor.cpp
@@ -1531,6 +1531,17 @@ void AActor::EndPlay(const EEndPlayReason::Type EndPlayReason)
 		// Dispatch the blueprint events
 		ReceiveEndPlay(EndPlayReason);
 		OnEndPlay.Broadcast(EndPlayReason);
+
+		TInlineComponentArray<UActorComponent*> Components;
+		GetComponents(Components);
+
+		for (UActorComponent* Component : Components)
+		{
+			if (Component->HasBegunPlay())
+			{
+				Component->EndPlay(EndPlayReason);
+			}
+		}
 	}
 
 	// Behaviors specific to an actor being unloaded due to a streaming level removal
@@ -1603,7 +1614,7 @@ FTransform AActor::GetTransform() const
 
 void AActor::Destroyed()
 {
-	RouteEndPlay(EEndPlayReason::ActorDestroyed);
+	RouteEndPlay(EEndPlayReason::Destroyed);
 
 	ReceiveDestroyed();
 	OnDestroyed.Broadcast();
@@ -2228,7 +2239,7 @@ static void DispatchOnComponentsCreated(AActor* NewActor)
 
 	for (UActorComponent* ActorComp : Components)
 	{
-		if (ActorComp && !ActorComp->bHasBeenCreated)
+		if (ActorComp && !ActorComp->HasBeenCreated())
 		{
 			ActorComp->OnComponentCreated();
 		}
@@ -2449,6 +2460,14 @@ void AActor::BeginPlay()
 	ensure(!bActorHasBegunPlay);
 	SetLifeSpan( InitialLifeSpan );
 
+	TInlineComponentArray<UActorComponent*> Components;
+	GetComponents(Components);
+
+	for (UActorComponent* Component : Components)
+	{
+		Component->BeginPlay();
+	}
+
 	ReceiveBeginPlay();
 
 	bActorHasBegunPlay = true;
@@ -3378,7 +3397,7 @@ void AActor::UninitializeComponents()
 
 	for (UActorComponent* ActorComp : Components)
 	{
-		if (ActorComp->bHasBeenInitialized)
+		if (ActorComp->HasBeenInitialized())
 		{
 			ActorComp->UninitializeComponent();
 		}
diff --git a/Engine/Source/Runtime/Engine/Private/ActorComponent.cpp b/Engine/Source/Runtime/Engine/Private/ActorComponent.cpp
--- a/Engine/Source/Runtime/Engine/Private/ActorComponent.cpp
+++ b/Engine/Source/Runtime/Engine/Private/ActorComponent.cpp
@@ -353,6 +353,11 @@ FString UActorComponent::GetReadableName() const
 
 void UActorComponent::BeginDestroy()
 {
+	if (bHasBegunPlay)
+	{
+		EndPlay(EEndPlayReason::Destroyed);
+	}
+
 	// Ensure that we call UninitializeComponent before we destroy this component
 	if (bHasBeenInitialized)
 	{
@@ -362,7 +367,7 @@ void UActorComponent::BeginDestroy()
 	ExecuteUnregisterEvents();
 
 	// Ensure that we call OnComponentDestroyed before we destroy this component
-	if(bHasBeenCreated)
+	if (bHasBeenCreated)
 	{
 		OnComponentDestroyed();
 	}
@@ -560,8 +565,6 @@ void UActorComponent::InitializeComponent()
 	check(bRegistered);
 	check(!bHasBeenInitialized);
 
-	ReceiveInitializeComponent();
-
 	bHasBeenInitialized = true;
 }
 
@@ -569,13 +572,30 @@ void UActorComponent::UninitializeComponent()
 {
 	check(bHasBeenInitialized);
 
+	bHasBeenInitialized = false;
+}
+
+void UActorComponent::BeginPlay()
+{
+	check(bRegistered);
+	check(!bHasBegunPlay);
+
+	ReceiveBeginPlay();
+
+	bHasBegunPlay = true;
+}
+
+void UActorComponent::EndPlay(const EEndPlayReason::Type EndPlayReason)
+{
+	check(bHasBegunPlay);
+
 	// If we're already pending kill blueprints don't get to be notified
 	if (!HasAnyFlags(RF_BeginDestroyed))
 	{
-		ReceiveUninitializeComponent();
+		ReceiveEndPlay(EndPlayReason);
 	}
 
-	bHasBeenInitialized = false;
+	bHasBegunPlay = false;
 }
 
 FActorComponentInstanceData* UActorComponent::GetComponentInstanceData() const
@@ -761,6 +781,14 @@ void UActorComponent::RegisterComponentWithWorld(UWorld* InWorld)
 		}
 	}
 
+	if (Owner && Owner->HasActorBegunPlay())
+	{
+		if (bWantsBeginPlay)
+		{
+			BeginPlay();
+		}
+	}
+
 	// If this is a blueprint created component and it has component children they can miss getting registered in some scenarios
 	if (IsCreatedByConstructionScript())
 	{
@@ -816,6 +844,11 @@ void UActorComponent::UnregisterComponent()
 
 void UActorComponent::DestroyComponent(bool bPromoteChildren/*= false*/)
 {
+	if (bHasBegunPlay)
+	{
+		EndPlay(EEndPlayReason::Destroyed);
+	}
+
 	// Ensure that we call UninitializeComponent before we destroy this component
 	if (bHasBeenInitialized)
 	{
