#ue4 - Various improvements to NavWalking, especially when bProjectNavMeshWalking is enabled.

- Limit NavWalking projected offset to be not too far above or below the navmesh.
- Optimize to avoid reprojecting to navmesh when bProjectNavMeshWalking is true and XY location has not changed.
- Project for navmesh from the last navmesh height, not the location where the agent was snapped to the terrain.

diff --git a/Engine/Source/Runtime/Engine/Classes/GameFramework/CharacterMovementComponent.h b/Engine/Source/Runtime/Engine/Classes/GameFramework/CharacterMovementComponent.h
--- a/Engine/Source/Runtime/Engine/Classes/GameFramework/CharacterMovementComponent.h
+++ b/Engine/Source/Runtime/Engine/Classes/GameFramework/CharacterMovementComponent.h
@@ -740,6 +740,20 @@ public:
 	UPROPERTY(Transient)
 	float NavMeshProjectionTimer;
 
+	/**
+	 * Scale of the total capsule height to use for projection from navmesh to underlying geometry in the upward direction.
+	 * In other words, start the trace at [CapsuleHeight * NavMeshProjectionCapsuleHeightScaleUp] above nav mesh.
+	 */
+	UPROPERTY(Category = "NavMesh Movement", EditAnywhere, BlueprintReadOnly, meta=(editcondition = "bProjectNavMeshWalking"))
+	float NavMeshProjectionCapsuleHeightScaleUp;
+
+	/**
+	 * Scale of the total capsule height to use for projection from navmesh to underlying geometry in the downward direction.
+	 * In other words, trace down to [CapsuleHeight * NavMeshProjectionCapsuleHeightScaleDown] below nav mesh.
+	 */
+	UPROPERTY(Category = "NavMesh Movement", EditAnywhere, BlueprintReadOnly, meta=(editcondition = "bProjectNavMeshWalking"))
+	float NavMeshProjectionCapsuleHeightScaleDown;
+
 	/** Change avoidance state and registers in RVO manager if needed */
 	UFUNCTION(BlueprintCallable, Category = "Pawn|Components|CharacterMovement", meta = (UnsafeDuringActorConstruction = "true"))
 	void SetAvoidanceEnabled(bool bEnable);
@@ -1129,7 +1143,7 @@ protected:
 	 * navmesh can differ quite significantly from geometry).
 	 * Updates CachedProjectedNavMeshHitResult, access this for more info about hits.
 	 */
-	virtual void ProjectLocationFromNavMesh(float DeltaSeconds, FVector& InOutLocation);
+	virtual void ProjectLocationFromNavMesh(float DeltaSeconds, FVector& InOutLocation, float UpOffset, float DownOffset);
 
 public:
 
@@ -1522,7 +1536,8 @@ protected:
 
 public:
 
-	/** Project a location to navmesh to find adjusted height
+	/**
+	 * Project a location to navmesh to find adjusted height.
 	 * @param TestLocation		Location to project
 	 * @param NavFloorLocation	Location on navmesh
 	 * @return True if projection was performed (successfully or not)
diff --git a/Engine/Source/Runtime/Engine/Private/CharacterMovementComponent.cpp b/Engine/Source/Runtime/Engine/Private/CharacterMovementComponent.cpp
--- a/Engine/Source/Runtime/Engine/Private/CharacterMovementComponent.cpp
+++ b/Engine/Source/Runtime/Engine/Private/CharacterMovementComponent.cpp
@@ -23,6 +23,7 @@
 #include "Components/DestructibleComponent.h"
 
 DEFINE_LOG_CATEGORY_STATIC(LogCharacterMovement, Log, All);
+DEFINE_LOG_CATEGORY_STATIC(LogNavMeshMovement, Log, All);
 
 /**
  * Character stats
@@ -284,6 +285,8 @@ UCharacterMovementComponent::UCharacterMovementComponent(const FObjectInitialize
 	OldBaseLocation = FVector::ZeroVector;
 
 	NavMeshProjectionInterval = 0.1f;
+	NavMeshProjectionCapsuleHeightScaleUp = 1.0f;
+	NavMeshProjectionCapsuleHeightScaleDown = 1.0f;
 }
 
 void UCharacterMovementComponent::PostLoad()
@@ -3799,15 +3802,38 @@ void UCharacterMovementComponent::PhysNavWalking(float deltaTime, int32 Iteratio
 	FVector AdjustedDest = OldLocation + DeltaMove;
 	FNavLocation DestNavLocation;
 
-	if (DeltaMove.IsNearlyZero() &&
-		CachedNavLocation.NodeRef != INVALID_NAVNODEREF &&
-		CachedNavLocation.Location.Equals(OldLocation))
+	bool bSameNavLocation = false;
+	if (CachedNavLocation.NodeRef != INVALID_NAVNODEREF)
+	{
+		if (bProjectNavMeshWalking)
+		{
+			bSameNavLocation = (OldLocation - CachedNavLocation.Location).SizeSquared2D() <= KINDA_SMALL_NUMBER;
+		}
+		else
+		{
+			bSameNavLocation = CachedNavLocation.Location.Equals(OldLocation);
+		}
+	}
+		
+
+	if (DeltaMove.IsNearlyZero() && bSameNavLocation)
 	{
 		DestNavLocation = CachedNavLocation;
+		UE_LOG(LogNavMeshMovement, VeryVerbose, TEXT("%s using cached navmesh location! (bProjectNavMeshWalking = %d)"), *GetNameSafe(CharacterOwner), bProjectNavMeshWalking);
 	}
 	else
 	{
 		SCOPE_CYCLE_COUNTER(STAT_CharNavProjectPoint);
+
+		// Start the trace from the Z location of the last valid trace.
+		// Otherwise if we are projecting our location to the underlying geometry and it's far above or below the navmesh,
+		// we'll follow that geometry's plane out of range of valid navigation.
+		if (bSameNavLocation && bProjectNavMeshWalking)
+		{
+			AdjustedDest.Z = CachedNavLocation.Location.Z;
+		}
+
+		// Find the point on the NavMesh
 		const bool bHasNavigationData = FindNavFloor(AdjustedDest, DestNavLocation);
 		if (!bHasNavigationData)
 		{
@@ -3824,7 +3850,10 @@ void UCharacterMovementComponent::PhysNavWalking(float deltaTime, int32 Iteratio
 		if (bProjectNavMeshWalking)
 		{
 			SCOPE_CYCLE_COUNTER(STAT_CharNavProjectLocation);
-			ProjectLocationFromNavMesh(deltaTime, NewLocation);
+			const float TotalCapsuleHeight = CharacterOwner->GetCapsuleComponent()->GetScaledCapsuleHalfHeight() * 2.0f;
+			const float UpOffset = TotalCapsuleHeight * NavMeshProjectionCapsuleHeightScaleUp;
+			const float DownOffset = TotalCapsuleHeight * NavMeshProjectionCapsuleHeightScaleDown;
+			ProjectLocationFromNavMesh(deltaTime, NewLocation, UpOffset, DownOffset);
 		}
 
 		FVector AdjustedDelta = NewLocation - OldLocation;
@@ -3885,11 +3914,12 @@ bool UCharacterMovementComponent::FindNavFloor(const FVector& TestLocation, FNav
 	return true;
 }
 
-void UCharacterMovementComponent::ProjectLocationFromNavMesh(float DeltaSeconds, FVector& InOutLocation)
+void UCharacterMovementComponent::ProjectLocationFromNavMesh(float DeltaSeconds, FVector& InOutLocation, float UpOffset, float DownOffset)
 {
 	SCOPE_CYCLE_COUNTER(STAT_CharNavProjectLocation);
 
-	const FVector RayCastOffset(0.0f, 0.0f, CharacterOwner->GetCapsuleComponent()->GetScaledCapsuleHalfHeight() * 2.0f);
+	const FVector RayCastOffsetUp(0.0f, 0.0f, UpOffset);
+	const FVector RayCastOffsetDown(0.0f, 0.0f, DownOffset);
 
 	NavMeshProjectionTimer -= DeltaSeconds;
 	if (NavMeshProjectionTimer <= 0.0f)
@@ -3897,18 +3927,19 @@ void UCharacterMovementComponent::ProjectLocationFromNavMesh(float DeltaSeconds,
 		// We can skip this trace if we are at the same location as the last trace (ie, we haven't moved).
 		const bool bCachedLocationStillValid = (!bAlwaysCheckFloor &&
 												CachedProjectedNavMeshHitResult.bBlockingHit &&
-												(CachedProjectedNavMeshHitResult.TraceStart - (InOutLocation + RayCastOffset)).SizeSquared() < KINDA_SMALL_NUMBER);
+												CachedProjectedNavMeshHitResult.TraceStart == (InOutLocation + RayCastOffsetUp) &&
+												CachedProjectedNavMeshHitResult.TraceEnd == (InOutLocation - RayCastOffsetDown));
 
 		if (!bCachedLocationStillValid)
 		{
-			UE_LOG(LogCharacterMovement, Verbose, TEXT("ProjectLocationFromNavMesh(): %s interval: %.3f velocity: %s"), *GetNameSafe(CharacterOwner), NavMeshProjectionInterval, *Velocity.ToString());
+			UE_LOG(LogNavMeshMovement, VeryVerbose, TEXT("ProjectLocationFromNavMesh(): %s interval: %.3f velocity: %s"), *GetNameSafe(CharacterOwner), NavMeshProjectionInterval, *Velocity.ToString());
 
 			// raycast to underlying mesh to allow us to more closely follow geometry
 			// we use static objects here as a best approximation to accept only objects that
 			// influence navmesh generation
 			FCollisionQueryParams Params;
 			FCollisionObjectQueryParams ObjectQueryParams(FCollisionObjectQueryParams::AllStaticObjects);
-			GetWorld()->LineTraceSingleByObjectType(CachedProjectedNavMeshHitResult, InOutLocation + RayCastOffset, InOutLocation - RayCastOffset, ObjectQueryParams, Params);
+			GetWorld()->LineTraceSingleByObjectType(CachedProjectedNavMeshHitResult, InOutLocation + RayCastOffsetUp, InOutLocation - RayCastOffsetDown, ObjectQueryParams, Params);
 
 			// discard result if we were already inside something
 			if (CachedProjectedNavMeshHitResult.bStartPenetrating)
@@ -3918,7 +3949,7 @@ void UCharacterMovementComponent::ProjectLocationFromNavMesh(float DeltaSeconds,
 		}
 		else
 		{
-			UE_LOG(LogCharacterMovement, Verbose, TEXT("ProjectLocationFromNavMesh(): %s interval: %.3f velocity: %s [SKIP TRACE]"), *GetNameSafe(CharacterOwner), NavMeshProjectionInterval, *Velocity.ToString());
+			UE_LOG(LogNavMeshMovement, VeryVerbose, TEXT("ProjectLocationFromNavMesh(): %s interval: %.3f velocity: %s [SKIP TRACE]"), *GetNameSafe(CharacterOwner), NavMeshProjectionInterval, *Velocity.ToString());
 		}
 
 		// Wrap around to maintain same relative offset to tick time changes.
@@ -3935,7 +3966,11 @@ void UCharacterMovementComponent::ProjectLocationFromNavMesh(float DeltaSeconds,
 	// project to last plane we found
 	if (CachedProjectedNavMeshHitResult.bBlockingHit)
 	{
-		FVector ProjectedPoint = FMath::LinePlaneIntersection(InOutLocation, InOutLocation - RayCastOffset, CachedProjectedNavMeshHitResult.Location, CachedProjectedNavMeshHitResult.Normal);
+		FVector ProjectedPoint = FMath::LinePlaneIntersection(InOutLocation + RayCastOffsetUp, InOutLocation - RayCastOffsetDown, CachedProjectedNavMeshHitResult.Location, CachedProjectedNavMeshHitResult.Normal);
+
+		// Limit to not be too far above or below NavMesh location
+		ProjectedPoint.Z = FMath::Clamp(ProjectedPoint.Z, InOutLocation.Z - DownOffset, InOutLocation.Z + UpOffset);
+		
 		InOutLocation.Z = ProjectedPoint.Z;
 	}
 }
@@ -4036,6 +4071,7 @@ void UCharacterMovementComponent::ProcessLanded(const FHitResult& Hit, float rem
 			if (!bHasNavigationData || NavLocation.NodeRef == INVALID_NAVNODEREF)
 			{
 				GroundMovementMode = MOVE_Walking;
+				UE_LOG(LogNavMeshMovement, Verbose, TEXT("ProcessLanded(): %s tried to go to NavWalking but couldn't find NavMesh! Using Walking instead."), *GetNameSafe(CharacterOwner));
 			}
 		}
 
