#ue4 - Rename the nav mesh projected height up/down limits. Add edit limits, min of zero.

diff --git a/Engine/Config/BaseEngine.ini b/Engine/Config/BaseEngine.ini
--- a/Engine/Config/BaseEngine.ini
+++ b/Engine/Config/BaseEngine.ini
@@ -512,6 +512,8 @@ bUseStreamingPause=false
 +TaggedPropertyRedirects=(ClassName="CharacterMovementComponent",OldPropertyName="CrouchedSpeedPercent",NewPropertyName="CrouchedSpeedMultiplier")
 +TaggedPropertyRedirects=(ClassName="CharacterMovementComponent",OldPropertyName="bOrientToMovement",NewPropertyName="bOrientRotationToMovement")
 +TaggedPropertyRedirects=(ClassName="CharacterMovementComponent",OldPropertyName="bCrouchMovesCharacterDown",NewPropertyName="bCrouchMaintainsBaseLocation")
++TaggedPropertyRedirects=(ClassName="CharacterMovementComponent",OldPropertyName="NavMeshProjectionCapsuleHeightScaleUp",NewPropertyName="NavMeshProjectionHeightScaleUp")
++TaggedPropertyRedirects=(ClassName="CharacterMovementComponent",OldPropertyName="NavMeshProjectionCapsuleHeightScaleDown",NewPropertyName="NavMeshProjectionHeightScaleDown")
 +TaggedPropertyRedirects=(ClassName="PawnSensingComponent",OldPropertyName="SightCounterInterval",NewPropertyName="SensingInterval")
 +TaggedPropertyRedirects=(ClassName="PawnSensingComponent",OldPropertyName="bWantsSeePlayerNotify",NewPropertyName="bSeePawns")
 +TaggedPropertyRedirects=(ClassName="FloatingPawnMovement",OldPropertyName="AccelRate",NewPropertyName="Acceleration")
diff --git a/Engine/Source/Runtime/Engine/Classes/GameFramework/CharacterMovementComponent.h b/Engine/Source/Runtime/Engine/Classes/GameFramework/CharacterMovementComponent.h
--- a/Engine/Source/Runtime/Engine/Classes/GameFramework/CharacterMovementComponent.h
+++ b/Engine/Source/Runtime/Engine/Classes/GameFramework/CharacterMovementComponent.h
@@ -742,17 +742,17 @@ public:
 
 	/**
 	 * Scale of the total capsule height to use for projection from navmesh to underlying geometry in the upward direction.
-	 * In other words, start the trace at [CapsuleHeight * NavMeshProjectionCapsuleHeightScaleUp] above nav mesh.
+	 * In other words, start the trace at [CapsuleHeight * NavMeshProjectionHeightScaleUp] above nav mesh.
 	 */
-	UPROPERTY(Category = "NavMesh Movement", EditAnywhere, BlueprintReadOnly, meta=(editcondition = "bProjectNavMeshWalking"))
-	float NavMeshProjectionCapsuleHeightScaleUp;
+	UPROPERTY(Category = "NavMesh Movement", EditAnywhere, BlueprintReadOnly, meta=(editcondition = "bProjectNavMeshWalking", ClampMin="0", UIMin="0"))
+	float NavMeshProjectionHeightScaleUp;
 
 	/**
 	 * Scale of the total capsule height to use for projection from navmesh to underlying geometry in the downward direction.
-	 * In other words, trace down to [CapsuleHeight * NavMeshProjectionCapsuleHeightScaleDown] below nav mesh.
+	 * In other words, trace down to [CapsuleHeight * NavMeshProjectionHeightScaleDown] below nav mesh.
 	 */
-	UPROPERTY(Category = "NavMesh Movement", EditAnywhere, BlueprintReadOnly, meta=(editcondition = "bProjectNavMeshWalking"))
-	float NavMeshProjectionCapsuleHeightScaleDown;
+	UPROPERTY(Category = "NavMesh Movement", EditAnywhere, BlueprintReadOnly, meta=(editcondition = "bProjectNavMeshWalking", ClampMin="0", UIMin="0"))
+	float NavMeshProjectionHeightScaleDown;
 
 	/** Change avoidance state and registers in RVO manager if needed */
 	UFUNCTION(BlueprintCallable, Category = "Pawn|Components|CharacterMovement", meta = (UnsafeDuringActorConstruction = "true"))
diff --git a/Engine/Source/Runtime/Engine/Private/CharacterMovementComponent.cpp b/Engine/Source/Runtime/Engine/Private/CharacterMovementComponent.cpp
--- a/Engine/Source/Runtime/Engine/Private/CharacterMovementComponent.cpp
+++ b/Engine/Source/Runtime/Engine/Private/CharacterMovementComponent.cpp
@@ -285,8 +285,8 @@ UCharacterMovementComponent::UCharacterMovementComponent(const FObjectInitialize
 	OldBaseLocation = FVector::ZeroVector;
 
 	NavMeshProjectionInterval = 0.1f;
-	NavMeshProjectionCapsuleHeightScaleUp = 1.0f;
-	NavMeshProjectionCapsuleHeightScaleDown = 1.0f;
+	NavMeshProjectionHeightScaleUp = 1.0f;
+	NavMeshProjectionHeightScaleDown = 1.0f;
 }
 
 void UCharacterMovementComponent::PostLoad()
@@ -3851,8 +3851,8 @@ void UCharacterMovementComponent::PhysNavWalking(float deltaTime, int32 Iteratio
 		{
 			SCOPE_CYCLE_COUNTER(STAT_CharNavProjectLocation);
 			const float TotalCapsuleHeight = CharacterOwner->GetCapsuleComponent()->GetScaledCapsuleHalfHeight() * 2.0f;
-			const float UpOffset = TotalCapsuleHeight * NavMeshProjectionCapsuleHeightScaleUp;
-			const float DownOffset = TotalCapsuleHeight * NavMeshProjectionCapsuleHeightScaleDown;
+			const float UpOffset = TotalCapsuleHeight * FMath::Max(0.f, NavMeshProjectionHeightScaleUp);
+			const float DownOffset = TotalCapsuleHeight * FMath::Max(0.f, NavMeshProjectionHeightScaleDown);
 			ProjectLocationFromNavMesh(deltaTime, NewLocation, UpOffset, DownOffset);
 		}
 
