Ability to disable replays entirely on command line

diff --git a/Engine/Source/Runtime/Engine/Private/World.cpp b/Engine/Source/Runtime/Engine/Private/World.cpp
--- a/Engine/Source/Runtime/Engine/Private/World.cpp
+++ b/Engine/Source/Runtime/Engine/Private/World.cpp
@@ -2528,6 +2528,12 @@ bool UWorld::HandleLogActorCountsCommand( const TCHAR* Cmd, FOutputDevice& Ar, U
 
 bool UWorld::HandleDemoRecordCommand( const TCHAR* Cmd, FOutputDevice& Ar, UWorld* InWorld )
 {
+	if ( FParse::Param( FCommandLine::Get(),TEXT( "NOREPLAYS" ) ) )
+	{
+		UE_LOG( LogWorld, Warning, TEXT( "HandleDemoRecordCommand: Rejected due to -noreplays option" ) );
+		return true;
+	}
+
 	FURL DemoURL;
 	FString DemoName;
 
