Add support for gamemode to control recording replays

diff --git a/Engine/Source/Runtime/Engine/Classes/GameFramework/GameMode.h b/Engine/Source/Runtime/Engine/Classes/GameFramework/GameMode.h
--- a/Engine/Source/Runtime/Engine/Classes/GameFramework/GameMode.h
+++ b/Engine/Source/Runtime/Engine/Classes/GameFramework/GameMode.h
@@ -254,6 +254,10 @@ protected:
 	UPROPERTY(config)
 	float InactivePlayerStateLifeSpan;
 
+	/** If true, dedicated servers will record replays when HandleMatchHasStarted/HandleMatchHasStopped is called */
+	UPROPERTY(config)
+	bool bHandleDedicatedServerReplays;
+
 public:
 
 	/** Alters the synthetic bandwidth limit for a running game. */
@@ -616,6 +620,9 @@ public:
 	/** Called periodically, overridden by subclasses */
 	virtual void DefaultTimer();	
 
+	/** Returns true if replays will start/stop during gameplay starting/stopping */
+	virtual bool IsHandlingReplays();
+
 protected:
 
 	/** Handle for efficient management of DefaultTimer timer */
diff --git a/Engine/Source/Runtime/Engine/Private/GameMode.cpp b/Engine/Source/Runtime/Engine/Private/GameMode.cpp
--- a/Engine/Source/Runtime/Engine/Private/GameMode.cpp
+++ b/Engine/Source/Runtime/Engine/Private/GameMode.cpp
@@ -620,6 +620,11 @@ void AGameMode::HandleMatchHasStarted()
 			}
 		}
 	}
+
+	if (IsHandlingReplays())
+	{
+		GEngine->Exec(GetWorld(), TEXT("DEMOREC %m"));
+	}
 }
 
 bool AGameMode::ReadyToEndMatch()
@@ -641,6 +646,11 @@ void AGameMode::EndMatch()
 void AGameMode::HandleMatchHasEnded()
 {
 	GameSession->HandleMatchHasEnded();
+
+	if (IsHandlingReplays())
+	{
+		GEngine->Exec(GetWorld(), TEXT("DEMOSTOP"));
+	}
 }
 
 void AGameMode::StartToLeaveMap()
@@ -1737,3 +1747,8 @@ FString AGameMode::GetRedirectURL(const FString& MapName) const
 {
 	return FString();
 }
+
+bool AGameMode::IsHandlingReplays()
+{
+	return bHandleDedicatedServerReplays && GetNetMode() == ENetMode::NM_DedicatedServer;
+}
diff --git a/Engine/Source/Runtime/Engine/Private/UnrealEngine.cpp b/Engine/Source/Runtime/Engine/Private/UnrealEngine.cpp
--- a/Engine/Source/Runtime/Engine/Private/UnrealEngine.cpp
+++ b/Engine/Source/Runtime/Engine/Private/UnrealEngine.cpp
@@ -9180,12 +9180,15 @@ bool UEngine::LoadMap( FWorldContext& WorldContext, FURL URL, class UPendingNetG
 	IStreamingManager::Get().RemoveStreamingViews( RemoveStreamingViews_All );
 	
 	// See if we need to record network demos
-	const TCHAR* DemoRecName = URL.GetOption( TEXT( "DemoRec=" ), NULL );
-
-	if ( DemoRecName != NULL )
+	if ( WorldContext.World()->GetAuthGameMode() == NULL || !WorldContext.World()->GetAuthGameMode()->IsHandlingReplays() )
 	{
-		// Play the demo
-		GEngine->Exec( WorldContext.World(), *FString::Printf( TEXT("DEMOREC %s"), DemoRecName ) );
+		const TCHAR* DemoRecName = URL.GetOption( TEXT( "DemoRec=" ), NULL );
+
+		if ( DemoRecName != NULL )
+		{
+			// Record the demo
+			GEngine->Exec( WorldContext.World(), *FString::Printf( TEXT("DEMOREC %s"), DemoRecName ) );
+		}
 	}
 
 	MALLOC_PROFILER( FMallocProfiler::SnapshotMemoryLoadMapEnd( URL.Map ); )
