#ue4 - Add interp speed option to nav movement projection, and don't let agents pop so high above navmesh.

- Dramatically reduces "pop" of height adjustments of agents.

diff --git a/Engine/Source/Runtime/Engine/Classes/GameFramework/CharacterMovementComponent.h b/Engine/Source/Runtime/Engine/Classes/GameFramework/CharacterMovementComponent.h
--- a/Engine/Source/Runtime/Engine/Classes/GameFramework/CharacterMovementComponent.h
+++ b/Engine/Source/Runtime/Engine/Classes/GameFramework/CharacterMovementComponent.h
@@ -734,24 +734,28 @@ public:
 	FHitResult CachedProjectedNavMeshHitResult;
 
 	/** How often we should raycast to project from navmesh to underlying geometry */
-	UPROPERTY(Category = "NavMesh Movement", EditAnywhere, BlueprintReadOnly, meta=(editcondition = "bProjectNavMeshWalking"))
+	UPROPERTY(Category = "NavMesh Movement", EditAnywhere, BlueprintReadWrite, meta=(editcondition = "bProjectNavMeshWalking"))
 	float NavMeshProjectionInterval;
 
 	UPROPERTY(Transient)
 	float NavMeshProjectionTimer;
 
+	/** Speed at which to interpolate agent navmesh offset between traces. 0: Instant (no interp) > 0: Interp speed") */
+	UPROPERTY(Category = "NavMesh Movement", EditAnywhere, BlueprintReadWrite, meta=(editcondition = "bProjectNavMeshWalking", ClampMin="0", UIMin="0"))
+	float NavMeshProjectionInterpSpeed;
+
 	/**
 	 * Scale of the total capsule height to use for projection from navmesh to underlying geometry in the upward direction.
 	 * In other words, start the trace at [CapsuleHeight * NavMeshProjectionHeightScaleUp] above nav mesh.
 	 */
-	UPROPERTY(Category = "NavMesh Movement", EditAnywhere, BlueprintReadOnly, meta=(editcondition = "bProjectNavMeshWalking", ClampMin="0", UIMin="0"))
+	UPROPERTY(Category = "NavMesh Movement", EditAnywhere, BlueprintReadWrite, meta=(editcondition = "bProjectNavMeshWalking", ClampMin="0", UIMin="0"))
 	float NavMeshProjectionHeightScaleUp;
 
 	/**
 	 * Scale of the total capsule height to use for projection from navmesh to underlying geometry in the downward direction.
 	 * In other words, trace down to [CapsuleHeight * NavMeshProjectionHeightScaleDown] below nav mesh.
 	 */
-	UPROPERTY(Category = "NavMesh Movement", EditAnywhere, BlueprintReadOnly, meta=(editcondition = "bProjectNavMeshWalking", ClampMin="0", UIMin="0"))
+	UPROPERTY(Category = "NavMesh Movement", EditAnywhere, BlueprintReadWrite, meta=(editcondition = "bProjectNavMeshWalking", ClampMin="0", UIMin="0"))
 	float NavMeshProjectionHeightScaleDown;
 
 	/** Change avoidance state and registers in RVO manager if needed */
@@ -1143,7 +1147,7 @@ protected:
 	 * navmesh can differ quite significantly from geometry).
 	 * Updates CachedProjectedNavMeshHitResult, access this for more info about hits.
 	 */
-	virtual void ProjectLocationFromNavMesh(float DeltaSeconds, FVector& InOutLocation, float UpOffset, float DownOffset);
+	virtual FVector ProjectLocationFromNavMesh(float DeltaSeconds, const FVector& CurrentFeetLocation, const FVector& TargetNavLocation, float UpOffset, float DownOffset);
 
 public:
 
diff --git a/Engine/Source/Runtime/Engine/Private/CharacterMovementComponent.cpp b/Engine/Source/Runtime/Engine/Private/CharacterMovementComponent.cpp
--- a/Engine/Source/Runtime/Engine/Private/CharacterMovementComponent.cpp
+++ b/Engine/Source/Runtime/Engine/Private/CharacterMovementComponent.cpp
@@ -285,7 +285,8 @@ UCharacterMovementComponent::UCharacterMovementComponent(const FObjectInitialize
 	OldBaseLocation = FVector::ZeroVector;
 
 	NavMeshProjectionInterval = 0.1f;
-	NavMeshProjectionHeightScaleUp = 1.0f;
+	NavMeshProjectionInterpSpeed = 12.f;
+	NavMeshProjectionHeightScaleUp = 0.67f;
 	NavMeshProjectionHeightScaleDown = 1.0f;
 }
 
@@ -3853,7 +3854,7 @@ void UCharacterMovementComponent::PhysNavWalking(float deltaTime, int32 Iteratio
 			const float TotalCapsuleHeight = CharacterOwner->GetCapsuleComponent()->GetScaledCapsuleHalfHeight() * 2.0f;
 			const float UpOffset = TotalCapsuleHeight * FMath::Max(0.f, NavMeshProjectionHeightScaleUp);
 			const float DownOffset = TotalCapsuleHeight * FMath::Max(0.f, NavMeshProjectionHeightScaleDown);
-			ProjectLocationFromNavMesh(deltaTime, NewLocation, UpOffset, DownOffset);
+			NewLocation = ProjectLocationFromNavMesh(deltaTime, OldLocation, NewLocation, UpOffset, DownOffset);
 		}
 
 		FVector AdjustedDelta = NewLocation - OldLocation;
@@ -3914,21 +3915,23 @@ bool UCharacterMovementComponent::FindNavFloor(const FVector& TestLocation, FNav
 	return true;
 }
 
-void UCharacterMovementComponent::ProjectLocationFromNavMesh(float DeltaSeconds, FVector& InOutLocation, float UpOffset, float DownOffset)
+FVector UCharacterMovementComponent::ProjectLocationFromNavMesh(float DeltaSeconds, const FVector& CurrentFeetLocation, const FVector& TargetNavLocation, float UpOffset, float DownOffset)
 {
 	SCOPE_CYCLE_COUNTER(STAT_CharNavProjectLocation);
 
+	FVector NewLocation = TargetNavLocation;
+
 	const FVector RayCastOffsetUp(0.0f, 0.0f, UpOffset);
 	const FVector RayCastOffsetDown(0.0f, 0.0f, DownOffset);
 
 	NavMeshProjectionTimer -= DeltaSeconds;
 	if (NavMeshProjectionTimer <= 0.0f)
 	{
-		// We can skip this trace if we are at the same location as the last trace (ie, we haven't moved).
+		// We can skip this trace if we are checking at the same location as the last trace (ie, we haven't moved).
 		const bool bCachedLocationStillValid = (!bAlwaysCheckFloor &&
 												CachedProjectedNavMeshHitResult.bBlockingHit &&
-												CachedProjectedNavMeshHitResult.TraceStart == (InOutLocation + RayCastOffsetUp) &&
-												CachedProjectedNavMeshHitResult.TraceEnd == (InOutLocation - RayCastOffsetDown));
+												CachedProjectedNavMeshHitResult.TraceStart == (TargetNavLocation + RayCastOffsetUp) &&
+												CachedProjectedNavMeshHitResult.TraceEnd == (TargetNavLocation - RayCastOffsetDown));
 
 		if (!bCachedLocationStillValid)
 		{
@@ -3939,7 +3942,7 @@ void UCharacterMovementComponent::ProjectLocationFromNavMesh(float DeltaSeconds,
 			// influence navmesh generation
 			FCollisionQueryParams Params;
 			FCollisionObjectQueryParams ObjectQueryParams(FCollisionObjectQueryParams::AllStaticObjects);
-			GetWorld()->LineTraceSingleByObjectType(CachedProjectedNavMeshHitResult, InOutLocation + RayCastOffsetUp, InOutLocation - RayCastOffsetDown, ObjectQueryParams, Params);
+			GetWorld()->LineTraceSingleByObjectType(CachedProjectedNavMeshHitResult, TargetNavLocation + RayCastOffsetUp, TargetNavLocation - RayCastOffsetDown, ObjectQueryParams, Params);
 
 			// discard result if we were already inside something
 			if (CachedProjectedNavMeshHitResult.bStartPenetrating)
@@ -3966,13 +3969,21 @@ void UCharacterMovementComponent::ProjectLocationFromNavMesh(float DeltaSeconds,
 	// project to last plane we found
 	if (CachedProjectedNavMeshHitResult.bBlockingHit)
 	{
-		FVector ProjectedPoint = FMath::LinePlaneIntersection(InOutLocation + RayCastOffsetUp, InOutLocation - RayCastOffsetDown, CachedProjectedNavMeshHitResult.Location, CachedProjectedNavMeshHitResult.Normal);
+		FVector ProjectedPoint = FMath::LinePlaneIntersection(TargetNavLocation + RayCastOffsetUp, TargetNavLocation - RayCastOffsetDown, CachedProjectedNavMeshHitResult.Location, CachedProjectedNavMeshHitResult.Normal);
 
 		// Limit to not be too far above or below NavMesh location
-		ProjectedPoint.Z = FMath::Clamp(ProjectedPoint.Z, InOutLocation.Z - DownOffset, InOutLocation.Z + UpOffset);
-		
-		InOutLocation.Z = ProjectedPoint.Z;
+		ProjectedPoint.Z = FMath::Clamp(ProjectedPoint.Z, TargetNavLocation.Z - DownOffset, TargetNavLocation.Z + UpOffset);
+
+		// Interp for smoother updates (less "pop" when trace hits something new). 0 interp speed is instant.
+		const float InterpSpeed = FMath::Max(0.f, NavMeshProjectionInterpSpeed);
+		ProjectedPoint.Z = FMath::FInterpTo(CurrentFeetLocation.Z, ProjectedPoint.Z, DeltaSeconds, InterpSpeed);
+		ProjectedPoint.Z = FMath::Clamp(ProjectedPoint.Z, TargetNavLocation.Z - DownOffset, TargetNavLocation.Z + UpOffset);
+
+		// Final result
+		NewLocation.Z = ProjectedPoint.Z;
 	}
+
+	return NewLocation;
 }
 
 void UCharacterMovementComponent::PhysCustom(float deltaTime, int32 Iterations)
