#ue4 - NavWalking location projection uses ECC_WorldStatic channel rather than StaticObject types.

diff --git a/Engine/Source/Runtime/Engine/Private/CharacterMovementComponent.cpp b/Engine/Source/Runtime/Engine/Private/CharacterMovementComponent.cpp
--- a/Engine/Source/Runtime/Engine/Private/CharacterMovementComponent.cpp
+++ b/Engine/Source/Runtime/Engine/Private/CharacterMovementComponent.cpp
@@ -3940,9 +3940,11 @@ FVector UCharacterMovementComponent::ProjectLocationFromNavMesh(float DeltaSecon
 			// raycast to underlying mesh to allow us to more closely follow geometry
 			// we use static objects here as a best approximation to accept only objects that
 			// influence navmesh generation
-			FCollisionQueryParams Params;
-			FCollisionObjectQueryParams ObjectQueryParams(FCollisionObjectQueryParams::AllStaticObjects);
-			GetWorld()->LineTraceSingleByObjectType(CachedProjectedNavMeshHitResult, TargetNavLocation + RayCastOffsetUp, TargetNavLocation - RayCastOffsetDown, ObjectQueryParams, Params);
+			static const FName ProjectLocationName(TEXT("NavProjectLocation"));
+			FCollisionQueryParams Params(ProjectLocationName, false);
+			FCollisionResponseParams ResponseParams(ECR_Ignore); // ignore everything
+			ResponseParams.CollisionResponse.SetResponse(ECC_WorldStatic, ECR_Block); // get blocked only by WorldStatic
+			GetWorld()->LineTraceSingleByChannel(CachedProjectedNavMeshHitResult, TargetNavLocation + RayCastOffsetUp, TargetNavLocation - RayCastOffsetDown, ECC_WorldStatic, Params, ResponseParams);
 
 			// discard result if we were already inside something
 			if (CachedProjectedNavMeshHitResult.bStartPenetrating)
