Fixed non unity build error

diff --git a/Engine/Source/Runtime/Engine/Private/Brush.cpp b/Engine/Source/Runtime/Engine/Private/Brush.cpp
--- a/Engine/Source/Runtime/Engine/Private/Brush.cpp
+++ b/Engine/Source/Runtime/Engine/Private/Brush.cpp
@@ -11,6 +11,7 @@
 #include "Engine/BrushBuilder.h"
 #include "ActorEditorUtils.h"
 #include "Components/BrushComponent.h"
+#include "Editor.h"
 
 #if WITH_EDITOR
 /** Define static delegate */
