Fixed trivial navigation paths within bounds of a single nav poly claiming their cost is 0 #UE4

diff --git a/Engine/Source/Runtime/Engine/Classes/AI/Navigation/RecastNavMesh.h b/Engine/Source/Runtime/Engine/Classes/AI/Navigation/RecastNavMesh.h
--- a/Engine/Source/Runtime/Engine/Classes/AI/Navigation/RecastNavMesh.h
+++ b/Engine/Source/Runtime/Engine/Classes/AI/Navigation/RecastNavMesh.h
@@ -115,11 +115,6 @@ struct ENGINE_API FNavMeshPath : public FNavigationPath
 	/** get cost of path, starting from given point */
 	virtual float GetCostFromIndex(int32 PathPointIndex) const
 	{
-		if (PathPointIndex >= PathCorridorCost.Num() - 1)
-		{
-			return 0.f;
-		}
-
 		float TotalCost = 0.f;
 		const float* Cost = PathCorridorCost.GetData();
 		for (int32 PolyIndex = PathPointIndex; PolyIndex < PathCorridorCost.Num(); ++PolyIndex, ++Cost)
