Merge to UE4: Warn when demo record frames are slow

diff --git a/Engine/Source/Runtime/Engine/Classes/Engine/DemoNetDriver.h b/Engine/Source/Runtime/Engine/Classes/Engine/DemoNetDriver.h
--- a/Engine/Source/Runtime/Engine/Classes/Engine/DemoNetDriver.h
+++ b/Engine/Source/Runtime/Engine/Classes/Engine/DemoNetDriver.h
@@ -49,8 +49,15 @@ class UDemoNetDriver
 	UPROPERTY(config)
 	FString DemoSpectatorClass;
 
+	/** Our network replay streamer */
 	TSharedPtr< class INetworkReplayStreamer > ReplayStreamer;
 
+	/** Internal debug timing/tracking */
+	double		AccumulatedRecordTime;
+	double		LastRecordAvgFlush;
+	double		MaxRecordTime;
+	int32		RecordCountSinceFlush;
+
 public:
 
 	// UNetDriver interface.
diff --git a/Engine/Source/Runtime/Engine/Private/DemoNetDriver.cpp b/Engine/Source/Runtime/Engine/Private/DemoNetDriver.cpp
--- a/Engine/Source/Runtime/Engine/Private/DemoNetDriver.cpp
+++ b/Engine/Source/Runtime/Engine/Private/DemoNetDriver.cpp
@@ -388,7 +388,39 @@ void UDemoNetDriver::TickFlush( float DeltaSeconds )
 	{
 		if ( ClientConnections.Num() > 0 )
 		{
+			const double StartTime = FPlatformTime::Seconds();
+
 			TickDemoRecord( DeltaSeconds );
+
+			const double EndTime = FPlatformTime::Seconds();
+
+			const double RecordTotalTime = ( EndTime - StartTime );
+
+			MaxRecordTime = FMath::Max( MaxRecordTime, RecordTotalTime );
+
+			AccumulatedRecordTime += RecordTotalTime;
+
+			RecordCountSinceFlush++;
+
+			const double ElapsedTime = EndTime - LastRecordAvgFlush;
+
+			const double AVG_FLUSH_TIME_IN_SECONDS = 2;
+
+			if ( ElapsedTime > AVG_FLUSH_TIME_IN_SECONDS && RecordCountSinceFlush > 0 )
+			{
+				const float AvgTimeMS = ( AccumulatedRecordTime / RecordCountSinceFlush ) * 1000;
+				const float MaxRecordTimeMS = MaxRecordTime * 1000;
+
+				if ( AvgTimeMS > 3.0f || MaxRecordTimeMS > 6.0f )
+				{
+					UE_LOG( LogDemo, Warning, TEXT( "UDemoNetDriver::TickFlush: SLOW FRAME. Avg: %2.2f, Max: %2.2f, Actors: %i" ), AvgTimeMS, MaxRecordTimeMS, World->NetworkActors.Num() );
+				}
+
+				LastRecordAvgFlush		= EndTime;
+				AccumulatedRecordTime	= 0;
+				MaxRecordTime			= 0;
+				RecordCountSinceFlush	= 0;
+			}
 		}
 		else if ( ServerConnection != NULL )
 		{
