Merge: Handle replay streamer error conditions

diff --git a/Engine/Source/Runtime/Engine/Private/DemoNetDriver.cpp b/Engine/Source/Runtime/Engine/Private/DemoNetDriver.cpp
--- a/Engine/Source/Runtime/Engine/Private/DemoNetDriver.cpp
+++ b/Engine/Source/Runtime/Engine/Private/DemoNetDriver.cpp
@@ -382,102 +382,117 @@ void UDemoNetDriver::TickFlush( float DeltaSeconds )
 {
 	Super::TickFlush( DeltaSeconds );
 
+	if ( ClientConnections.Num() == 0 && ServerConnection == nullptr )
+	{
+		// Nothing to do
+		return;
+	}
+
+	if ( ReplayStreamer->GetLastError() != ENetworkReplayError::None )
+	{
+		UE_LOG( LogDemo, Error, TEXT( "UDemoNetDriver::TickFlush: ReplayStreamer ERROR: %s" ), ENetworkReplayError::ToString( ReplayStreamer->GetLastError() ) );
+		StopDemo();
+		return;
+	}
+
 	FArchive* FileAr = ReplayStreamer->GetStreamingArchive();
 
-	if ( FileAr )
+	if ( FileAr == nullptr )
 	{
-		if ( ClientConnections.Num() > 0 )
-		{
-			const double StartTime = FPlatformTime::Seconds();
+		return;
+	}
 
-			TickDemoRecord( DeltaSeconds );
+	if ( ClientConnections.Num() > 0 )
+	{
+		const double StartTime = FPlatformTime::Seconds();
 
-			const double EndTime = FPlatformTime::Seconds();
+		TickDemoRecord( DeltaSeconds );
 
-			const double RecordTotalTime = ( EndTime - StartTime );
+		const double EndTime = FPlatformTime::Seconds();
 
-			MaxRecordTime = FMath::Max( MaxRecordTime, RecordTotalTime );
+		const double RecordTotalTime = ( EndTime - StartTime );
 
-			AccumulatedRecordTime += RecordTotalTime;
+		MaxRecordTime = FMath::Max( MaxRecordTime, RecordTotalTime );
 
-			RecordCountSinceFlush++;
+		AccumulatedRecordTime += RecordTotalTime;
 
-			const double ElapsedTime = EndTime - LastRecordAvgFlush;
+		RecordCountSinceFlush++;
 
-			const double AVG_FLUSH_TIME_IN_SECONDS = 2;
+		const double ElapsedTime = EndTime - LastRecordAvgFlush;
 
-			if ( ElapsedTime > AVG_FLUSH_TIME_IN_SECONDS && RecordCountSinceFlush > 0 )
-			{
-				const float AvgTimeMS = ( AccumulatedRecordTime / RecordCountSinceFlush ) * 1000;
-				const float MaxRecordTimeMS = MaxRecordTime * 1000;
+		const double AVG_FLUSH_TIME_IN_SECONDS = 2;
 
-				if ( AvgTimeMS > 3.0f || MaxRecordTimeMS > 6.0f )
-				{
-					UE_LOG( LogDemo, Warning, TEXT( "UDemoNetDriver::TickFlush: SLOW FRAME. Avg: %2.2f, Max: %2.2f, Actors: %i" ), AvgTimeMS, MaxRecordTimeMS, World->NetworkActors.Num() );
-				}
+		if ( ElapsedTime > AVG_FLUSH_TIME_IN_SECONDS && RecordCountSinceFlush > 0 )
+		{
+			const float AvgTimeMS = ( AccumulatedRecordTime / RecordCountSinceFlush ) * 1000;
+			const float MaxRecordTimeMS = MaxRecordTime * 1000;
 
-				LastRecordAvgFlush		= EndTime;
-				AccumulatedRecordTime	= 0;
-				MaxRecordTime			= 0;
-				RecordCountSinceFlush	= 0;
+			if ( AvgTimeMS > 3.0f || MaxRecordTimeMS > 6.0f )
+			{
+				UE_LOG( LogDemo, Warning, TEXT( "UDemoNetDriver::TickFlush: SLOW FRAME. Avg: %2.2f, Max: %2.2f, Actors: %i" ), AvgTimeMS, MaxRecordTimeMS, World->NetworkActors.Num() );
 			}
+
+			LastRecordAvgFlush		= EndTime;
+			AccumulatedRecordTime	= 0;
+			MaxRecordTime			= 0;
+			RecordCountSinceFlush	= 0;
 		}
-		else if ( ServerConnection != NULL )
+	}
+	else if ( ServerConnection != NULL )
+	{
+		// Wait until all levels are streamed in
+		for ( int32 i = 0; i < World->StreamingLevels.Num(); ++i )
 		{
-			// Wait until all levels are streamed in
-			for ( int32 i = 0; i < World->StreamingLevels.Num(); ++i )
-			{
-				ULevelStreaming * StreamingLevel = World->StreamingLevels[i];
-
-				if ( StreamingLevel != NULL && StreamingLevel->ShouldBeLoaded() && (!StreamingLevel->IsLevelLoaded() || !StreamingLevel->GetLoadedLevel()->GetOutermost()->IsFullyLoaded() || !StreamingLevel->IsLevelVisible() ) )
-				{
-					// Abort, we have more streaming levels to load
-					return;
-				}
-			}
+			ULevelStreaming * StreamingLevel = World->StreamingLevels[i];
 
-			if ( CVarDemoTimeDilation.GetValueOnGameThread() >= 0.0f )
+			if ( StreamingLevel != NULL && StreamingLevel->ShouldBeLoaded() && (!StreamingLevel->IsLevelLoaded() || !StreamingLevel->GetLoadedLevel()->GetOutermost()->IsFullyLoaded() || !StreamingLevel->IsLevelVisible() ) )
 			{
-				World->GetWorldSettings()->DemoPlayTimeDilation = CVarDemoTimeDilation.GetValueOnGameThread();
+				// Abort, we have more streaming levels to load
+				return;
 			}
+		}
+
+		if ( CVarDemoTimeDilation.GetValueOnGameThread() >= 0.0f )
+		{
+			World->GetWorldSettings()->DemoPlayTimeDilation = CVarDemoTimeDilation.GetValueOnGameThread();
+		}
 
-			// Clamp time between 1000 hz, and 2 hz 
-			// (this is useful when debugging and you set a breakpoint, you don't want all that time to pass in one frame)
-			DeltaSeconds = FMath::Clamp( DeltaSeconds, 1.0f / 1000.0f, 1.0f / 2.0f );
+		// Clamp time between 1000 hz, and 2 hz 
+		// (this is useful when debugging and you set a breakpoint, you don't want all that time to pass in one frame)
+		DeltaSeconds = FMath::Clamp( DeltaSeconds, 1.0f / 1000.0f, 1.0f / 2.0f );
 
-			// We need to compensate for the fact that DeltaSeconds is real-time for net drivers
-			DeltaSeconds *= World->GetWorldSettings()->GetEffectiveTimeDilation();
+		// We need to compensate for the fact that DeltaSeconds is real-time for net drivers
+		DeltaSeconds *= World->GetWorldSettings()->GetEffectiveTimeDilation();
 
-			// Update time dilation on spectator pawn to compensate for any demo dilation 
-			//	(we want to continue to fly around in real-time)
-			if ( SpectatorController != NULL )
+		// Update time dilation on spectator pawn to compensate for any demo dilation 
+		//	(we want to continue to fly around in real-time)
+		if ( SpectatorController != NULL )
+		{
+			if ( SpectatorController->GetSpectatorPawn() != NULL )
 			{
-				if ( SpectatorController->GetSpectatorPawn() != NULL )
-				{
-					// Disable collision on the spectator
-					SpectatorController->GetSpectatorPawn()->SetActorEnableCollision( false );
+				// Disable collision on the spectator
+				SpectatorController->GetSpectatorPawn()->SetActorEnableCollision( false );
 					
-					SpectatorController->GetSpectatorPawn()->PrimaryActorTick.bTickEvenWhenPaused = true;
+				SpectatorController->GetSpectatorPawn()->PrimaryActorTick.bTickEvenWhenPaused = true;
 
-					USpectatorPawnMovement* SpectatorMovement = Cast<USpectatorPawnMovement>(SpectatorController->GetSpectatorPawn()->GetMovementComponent());
+				USpectatorPawnMovement* SpectatorMovement = Cast<USpectatorPawnMovement>(SpectatorController->GetSpectatorPawn()->GetMovementComponent());
 
-					if ( SpectatorMovement )
-					{
-						SpectatorMovement->bIgnoreTimeDilation = true;
-						SpectatorMovement->PrimaryComponentTick.bTickEvenWhenPaused = true;
-					}
+				if ( SpectatorMovement )
+				{
+					SpectatorMovement->bIgnoreTimeDilation = true;
+					SpectatorMovement->PrimaryComponentTick.bTickEvenWhenPaused = true;
 				}
 			}
+		}
 
-			if ( bDemoPlaybackDone )
-			{
-				return;
-			}
+		if ( bDemoPlaybackDone )
+		{
+			return;
+		}
 
-			if ( World->GetWorldSettings()->Pauser == NULL )
-			{
-				TickDemoPlayback( DeltaSeconds );
-			}
+		if ( World->GetWorldSettings()->Pauser == NULL )
+		{
+			TickDemoPlayback( DeltaSeconds );
 		}
 	}
 }
