Fix issue where deferred body removal was only working for rigid bodies (excluding rigidstatic)

diff --git a/Engine/Source/Runtime/Engine/Private/PhysicsEngine/PhysScene.cpp b/Engine/Source/Runtime/Engine/Private/PhysicsEngine/PhysScene.cpp
--- a/Engine/Source/Runtime/Engine/Private/PhysicsEngine/PhysScene.cpp
+++ b/Engine/Source/Runtime/Engine/Private/PhysicsEngine/PhysScene.cpp
@@ -345,22 +345,22 @@ void FPhysScene::RemoveActiveBody(FBodyInstance* BodyInstance, uint32 SceneType)
 #endif
 void FPhysScene::TermBody(FBodyInstance* BodyInstance)
 {
-	if (PxRigidBody * PRigidBody = BodyInstance->GetPxRigidBody())
+	if (PxRigidBody* PRigidBody = BodyInstance->GetPxRigidBody())
 	{
 #if WITH_SUBSTEPPING
-		FPhysSubstepTask * PhysSubStepper = PhysSubSteppers[SceneType(BodyInstance)];
+		FPhysSubstepTask* PhysSubStepper = PhysSubSteppers[SceneType(BodyInstance)];
 		PhysSubStepper->RemoveBodyInstance(BodyInstance);
 #endif
+	}
 
-		// Remove body from any pending deferred addition / removal
-		for(int32 SceneIdx = 0 ; SceneIdx < PST_MAX ; ++SceneIdx)
+	// Remove body from any pending deferred addition / removal
+	for(int32 SceneIdx = 0 ; SceneIdx < PST_MAX ; ++SceneIdx)
+	{
+		int32 FoundIdx = INDEX_NONE;
+		if(DeferredAddInstances[SceneIdx].Find(BodyInstance, FoundIdx))
 		{
-			int32 FoundIdx = INDEX_NONE;
-			if(DeferredAddActors[SceneIdx].Find(PRigidBody, FoundIdx))
-			{
-				DeferredAddActors->RemoveAt(FoundIdx);
-				DeferredAddInstances->RemoveAt(FoundIdx);
-			}
+			DeferredAddActors->RemoveAt(FoundIdx);
+			DeferredAddInstances->RemoveAt(FoundIdx);
 		}
 	}
 
