Added scene lock when adding aggregate for instanced static mesh

diff --git a/Engine/Source/Runtime/Engine/Private/InstancedStaticMesh.cpp b/Engine/Source/Runtime/Engine/Private/InstancedStaticMesh.cpp
--- a/Engine/Source/Runtime/Engine/Private/InstancedStaticMesh.cpp
+++ b/Engine/Source/Runtime/Engine/Private/InstancedStaticMesh.cpp
@@ -1001,7 +1001,9 @@ void UInstancedStaticMeshComponent::CreatePhysicsState()
 	{
 		auto* Aggregate = GPhysXSDK->createAggregate(AggregateMaxSize, false);
 		Aggregates.Add(Aggregate);
-		PhysScene->GetPhysXScene(SceneType)->addAggregate(*Aggregate);
+		PxScene* PScene = PhysScene->GetPhysXScene(SceneType);
+		SCOPED_SCENE_WRITE_LOCK(PScene);
+		PScene->addAggregate(*Aggregate);
 	}
 #endif
 
