fixed typo

diff --git a/Engine/Source/Runtime/Engine/Classes/Materials/MaterialInstance.h b/Engine/Source/Runtime/Engine/Classes/Materials/MaterialInstance.h
--- a/Engine/Source/Runtime/Engine/Classes/Materials/MaterialInstance.h
+++ b/Engine/Source/Runtime/Engine/Classes/Materials/MaterialInstance.h
@@ -136,7 +136,7 @@ class UMaterialInstance : public UMaterialInterface
 	class UMaterialInterface* Parent;
 
 	/**
-	 * Delegate for custom static paramters getter.
+	 * Delegate for custom static parameters getter.
 	 *
 	 * @param OutStaticParameterSet Parameter set to append.
 	 * @param Material Material instance to collect parameters.
