fixed typos

diff --git a/Engine/Source/Runtime/Engine/Classes/Materials/MaterialInstance.h b/Engine/Source/Runtime/Engine/Classes/Materials/MaterialInstance.h
--- a/Engine/Source/Runtime/Engine/Classes/Materials/MaterialInstance.h
+++ b/Engine/Source/Runtime/Engine/Classes/Materials/MaterialInstance.h
@@ -156,7 +156,7 @@ class UMaterialInstance : public UMaterialInterface
 	// Custom static parameters getter delegate.
 	ENGINE_API static FCustomStaticParametersGetterDelegate CustomStaticParametersGetters;
 
-	// An array of custom paramter set updaters.
+	// An array of custom parameter set updaters.
 	ENGINE_API static TArray<FCustomParameterSetUpdaterDelegate> CustomParameterSetUpdaters;
 
 	/**
@@ -348,7 +348,7 @@ public:
 	void GetAllShaderMaps(TArray<FMaterialShaderMap*>& OutShaderMaps);
 
 	/**
-	 * Builds a composited set of static parameters, including inherited and overidden values
+	 * Builds a composited set of static parameters, including inherited and overridden values
 	 */
 	ENGINE_API void GetStaticParameterValues(FStaticParameterSet& OutStaticParameters);
 
