fixed MID usage with SubsurfaceProfile materials

diff --git a/Engine/Source/Runtime/Engine/Classes/Materials/Material.h b/Engine/Source/Runtime/Engine/Classes/Materials/Material.h
--- a/Engine/Source/Runtime/Engine/Classes/Materials/Material.h
+++ b/Engine/Source/Runtime/Engine/Classes/Materials/Material.h
@@ -857,7 +857,7 @@ private:
 	/** Sets up transient properties in MaterialResources. */
 	void UpdateResourceAllocations();
 
-	/** to share code for PostLoad() and PostEditChangeProperty() */
+	/** to share code for PostLoad() and PostEditChangeProperty(), and UMaterialInstance::InitResources(), needs to be refactored */
 	void PropagateDataToMaterialProxy();
 
 public:
diff --git a/Engine/Source/Runtime/Engine/Private/Materials/Material.cpp b/Engine/Source/Runtime/Engine/Private/Materials/Material.cpp
--- a/Engine/Source/Runtime/Engine/Private/Materials/Material.cpp
+++ b/Engine/Source/Runtime/Engine/Private/Materials/Material.cpp
@@ -230,17 +230,20 @@ public:
 			{
 				const USubsurfaceProfilePointer SubsurfaceProfileRT = GetSubsurfaceProfileRT();
 
+				int32 AllocationId = 0;
 				if(SubsurfaceProfileRT)
 				{
 					// can be optimized (cached)
-					*OutValue = GSubsufaceProfileTextureObject.FindAllocationId(SubsurfaceProfileRT) / 255.0f;
+					AllocationId = GSubsufaceProfileTextureObject.FindAllocationId(SubsurfaceProfileRT);
 				}
 				else
 				{
 					// no profile specified means we use the default one stored at [0] which is human skin
-					*OutValue = 0.0f;
+					AllocationId = 0;
 				}
 
+				*OutValue = AllocationId / 255.0f;
+
 				return true;
 			}
 
@@ -2038,6 +2041,7 @@ void UMaterial::PostLoad()
 	}
 #endif // #if WITH_EDITOR
 
+	// needed for UMaterial as it doesn't have the InitResources() override where this is called
 	PropagateDataToMaterialProxy();
 
 	STAT(double MaterialLoadTime = 0);
@@ -2379,6 +2383,7 @@ void UMaterial::PostEditChangeProperty(FPropertyChangedEvent& PropertyChangedEve
 		}
 	}
 
+	// needed for UMaterial as it doesn't have the InitResources() override where this is called
 	PropagateDataToMaterialProxy();
 
 	// many property changes can require rebuild of graph so always mark as changed
diff --git a/Engine/Source/Runtime/Engine/Private/Materials/MaterialInstance.cpp b/Engine/Source/Runtime/Engine/Private/Materials/MaterialInstance.cpp
--- a/Engine/Source/Runtime/Engine/Private/Materials/MaterialInstance.cpp
+++ b/Engine/Source/Runtime/Engine/Private/Materials/MaterialInstance.cpp
@@ -159,16 +159,18 @@ bool FMaterialInstanceResource::GetScalarValue(
 	{
 		const USubsurfaceProfilePointer SubsurfaceProfileRT = GetSubsurfaceProfileRT();
 
+		int32 AllocationId = 0;
 		if (SubsurfaceProfileRT)
 		{
 			// can be optimized (cached)
-			*OutValue = GSubsufaceProfileTextureObject.FindAllocationId(SubsurfaceProfileRT) / 255.0f;
+			AllocationId = GSubsufaceProfileTextureObject.FindAllocationId(SubsurfaceProfileRT);
 		}
 		else
 		{
 			// no profile specified means we use the default one stored at [0] which is human skin
-			*OutValue = 0.0f;
+			AllocationId = 0;
 		}
+		*OutValue = AllocationId / 255.0f;
 
 		return true;
 	}
@@ -475,6 +477,8 @@ void UMaterialInstance::InitResources()
 	GameThread_InitMIParameters(this, VectorParameterValues);
 	GameThread_InitMIParameters(this, TextureParameterValues);
 	GameThread_InitMIParameters(this, FontParameterValues);
+	PropagateDataToMaterialProxy();
+
 	CacheMaterialInstanceUniformExpressions(this);
 }
 
