refactored SubsurfaceProfile rendering module interaction

diff --git a/Engine/Source/Runtime/Engine/Classes/Engine/SubsurfaceProfile.h b/Engine/Source/Runtime/Engine/Classes/Engine/SubsurfaceProfile.h
--- a/Engine/Source/Runtime/Engine/Classes/Engine/SubsurfaceProfile.h
+++ b/Engine/Source/Runtime/Engine/Classes/Engine/SubsurfaceProfile.h
@@ -80,8 +80,6 @@ public:
 	// destructor
 	~FSubsurfaceProfileTexture();
 
-	void SetRendererModule(class IRendererModule* InRendererModule) { RendererModule = InRendererModule; }
-
 	// convenience, can be optimized 
 	int32 AddOrUpdateProfile(const FSubsurfaceProfileStruct Settings, const USubsurfaceProfilePointer Profile)
 	{
@@ -148,8 +146,6 @@ private:
 		USubsurfaceProfilePointer GameThreadObject;
 	};
 
-	// 0 if there if no user of the feature
-	IRendererModule* RendererModule;
 	//
 	TArray<FSubsurfaceProfileEntry> SubsurfaceProfileEntries;
 
diff --git a/Engine/Source/Runtime/Engine/Private/Materials/MaterialInterface.cpp b/Engine/Source/Runtime/Engine/Private/Materials/MaterialInterface.cpp
--- a/Engine/Source/Runtime/Engine/Private/Materials/MaterialInterface.cpp
+++ b/Engine/Source/Runtime/Engine/Private/Materials/MaterialInterface.cpp
@@ -320,43 +320,21 @@ void UMaterialInterface::UpdateMaterialRenderProxy(FMaterialRenderProxy& Proxy)
 			Settings = LocalSubsurfaceProfile->Settings;
 		}
 
-		// this can be improved, it doesn't support Renderer hot reload
+		ENQUEUE_UNIQUE_RENDER_COMMAND_THREEPARAMETER(
+			UpdateMaterialRenderProxySubsurface,
+			const FSubsurfaceProfileStruct, Settings, Settings,
+			USubsurfaceProfilePointer, LocalSubsurfaceProfile, (USubsurfaceProfilePointer)LocalSubsurfaceProfile,
+			FMaterialRenderProxy&, Proxy, Proxy,
 		{
-			static bool bFirst = true;
+			uint32 AllocationId = 0;
 
-			if (bFirst)
+			if (LocalSubsurfaceProfile)
 			{
-				bFirst = false;
+				AllocationId = GSubsufaceProfileTextureObject.AddOrUpdateProfile(Settings, LocalSubsurfaceProfile);
 
-				static const FName RendererModuleName("Renderer");
-				IRendererModule& RendererModule = FModuleManager::GetModuleChecked<IRendererModule>(RendererModuleName);
-
-				ENQUEUE_UNIQUE_RENDER_COMMAND_ONEPARAMETER(
-					UpdateMaterialRenderProxySubsurfaceSetup,
-					IRendererModule&, RendererModule, RendererModule,
-					{
-					GSubsufaceProfileTextureObject.SetRendererModule(&RendererModule);
-				});
+				check(AllocationId >= 0 && AllocationId <= 255);
 			}
-		}
-
-		{
-			ENQUEUE_UNIQUE_RENDER_COMMAND_THREEPARAMETER(
-				UpdateMaterialRenderProxySubsurface,
-				const FSubsurfaceProfileStruct, Settings, Settings,
-				USubsurfaceProfilePointer, LocalSubsurfaceProfile, (USubsurfaceProfilePointer)LocalSubsurfaceProfile,
-				FMaterialRenderProxy&, Proxy, Proxy,
-			{
-				uint32 AllocationId = 0;
-
-				if (LocalSubsurfaceProfile)
-				{
-					AllocationId = GSubsufaceProfileTextureObject.AddOrUpdateProfile(Settings, LocalSubsurfaceProfile);
-
-					check(AllocationId >= 0 && AllocationId <= 255);
-				}
-				Proxy.SetSubsurfaceProfileRT(LocalSubsurfaceProfile);
-			});
-		}
+			Proxy.SetSubsurfaceProfileRT(LocalSubsurfaceProfile);
+		});
 	}
 }
diff --git a/Engine/Source/Runtime/Engine/Private/Rendering/SubsurfaceProfile.cpp b/Engine/Source/Runtime/Engine/Private/Rendering/SubsurfaceProfile.cpp
--- a/Engine/Source/Runtime/Engine/Private/Rendering/SubsurfaceProfile.cpp
+++ b/Engine/Source/Runtime/Engine/Private/Rendering/SubsurfaceProfile.cpp
@@ -4,6 +4,7 @@
 #include "SeparableSSS.h"
 #include "RendererInterface.h"
 #include "Engine/SubsurfaceProfile.h"
+#include "EngineModule.h" // GetRendererModule()
 
 DEFINE_LOG_CATEGORY_STATIC(LogSubsurfaceProfile, Log, All);
 
@@ -15,7 +16,6 @@ static TRefCountPtr<IPooledRenderTarget> GSSProfiles;
 
 
 FSubsurfaceProfileTexture::FSubsurfaceProfileTexture()
-	: RendererModule(0)
 {
 	FSubsurfaceProfileStruct DefaultSkin;
 
@@ -91,9 +91,6 @@ void FSubsurfaceProfileTexture::UpdateProfile(int32 AllocationId, const FSubsurf
 		return;
 	}
 
-	// call SetRendererModule() is missing
-	check(RendererModule);
-
 	check(AllocationId < SubsurfaceProfileEntries.Num());
 
 	SubsurfaceProfileEntries[AllocationId].Settings = Settings;
@@ -103,12 +100,6 @@ void FSubsurfaceProfileTexture::UpdateProfile(int32 AllocationId, const FSubsurf
 
 const IPooledRenderTarget* FSubsurfaceProfileTexture::GetTexture(FRHICommandListImmediate& RHICmdList)
 {
-	if(!RendererModule)
-	{
-		// call SetRendererModule() is missing, thiscan be if no SubsurfaceProfile was used yet but VisualizeSubsurface requests the texture
-		return 0;
-	}
-
 	if (!GSSProfiles)
 	{
 		CreateTexture(RHICmdList);
@@ -135,9 +126,6 @@ static float GetNextSmallerPositiveFloat(float x)
 
 void FSubsurfaceProfileTexture::CreateTexture(FRHICommandListImmediate& RHICmdList)
 {
-	// call SetRendererModule() is missing
-	check(RendererModule);
-
 	uint32 Height = SubsurfaceProfileEntries.Num();
 
 	check(Height);
@@ -155,7 +143,7 @@ void FSubsurfaceProfileTexture::CreateTexture(FRHICommandListImmediate& RHICmdLi
 		Desc.Format = PF_A16B16G16R16;
 	}
 
-	RendererModule->RenderTargetPoolFindFreeElement(Desc, GSSProfiles, TEXT("SSProfiles"));
+	GetRendererModule().RenderTargetPoolFindFreeElement(Desc, GSSProfiles, TEXT("SSProfiles"));
 
 	// Write the contents of the texture.
 	uint32 DestStride;
