fixed typo

diff --git a/Engine/Source/Runtime/Engine/Private/Rendering/SubsurfaceProfile.cpp b/Engine/Source/Runtime/Engine/Private/Rendering/SubsurfaceProfile.cpp
--- a/Engine/Source/Runtime/Engine/Private/Rendering/SubsurfaceProfile.cpp
+++ b/Engine/Source/Runtime/Engine/Private/Rendering/SubsurfaceProfile.cpp
@@ -25,7 +25,7 @@ FSubsurfaceProfileTexture::FSubsurfaceProfileTexture()
 
 FSubsurfaceProfileTexture::~FSubsurfaceProfileTexture()
 {
-	// we assume all ~USubsurfaceProfile() have bene called already
+	// we assume all ~USubsurfaceProfile() have been called already
 	for (int32 i = 0; i < SubsurfaceProfileEntries.Num(); ++i)
 	{
 		check(SubsurfaceProfileEntries[i].GameThreadObject == 0);
