refactored SubsurfaceProfile code to be simpler and cleaner

diff --git a/Engine/Source/Runtime/Engine/Classes/Engine/SubsurfaceProfile.h b/Engine/Source/Runtime/Engine/Classes/Engine/SubsurfaceProfile.h
--- a/Engine/Source/Runtime/Engine/Classes/Engine/SubsurfaceProfile.h
+++ b/Engine/Source/Runtime/Engine/Classes/Engine/SubsurfaceProfile.h
@@ -66,23 +66,22 @@ class USubsurfaceProfile : public UObject
 	// End UObject interface
 };
 
-// For render thread side, to identify the right index in a map, is never dereferenced
-typedef void* USubsurfaceProfilePointer;
-
-
 // render thread
 class FSubsurfaceProfileTexture : public FRenderResource
 {
 public:
-	// constructor runs on main thread (global variable)
+	// constructor
 	FSubsurfaceProfileTexture();
 
 	// destructor
 	~FSubsurfaceProfileTexture();
 
 	// convenience, can be optimized 
-	int32 AddOrUpdateProfile(const FSubsurfaceProfileStruct Settings, const USubsurfaceProfilePointer Profile)
+	// @param Profile must not be 0, game thread pointer, do not dereference, only for comparison
+	int32 AddOrUpdateProfile(const FSubsurfaceProfileStruct Settings, const USubsurfaceProfile* Profile)
 	{
+		check(Profile);
+
 		int32 AllocationId = FindAllocationId(Profile); 
 		
 		if (AllocationId != -1)
@@ -98,17 +97,17 @@ public:
 	}
 
 	// O(n) n is a small number
-	// @param Profile must not be 0
+	// @param InProfile must not be 0, game thread pointer, do not dereference, only for comparison
 	// @return AllocationId -1: no allocation, should be deallocated with DeallocateSubsurfaceProfile()
-	int32 AddProfile(const FSubsurfaceProfileStruct Settings, const USubsurfaceProfilePointer Profile);
+	int32 AddProfile(const FSubsurfaceProfileStruct Settings, const USubsurfaceProfile* InProfile);
 
 	// O(n) to find the element, n is the SSProfile count and usually quite small
-	void RemoveProfile(const USubsurfaceProfilePointer InProfile);
+	void RemoveProfile(const USubsurfaceProfile* InProfile);
 
-	// @param InRendererModule first call needs to have it != 0
-	void UpdateProfile(const FSubsurfaceProfileStruct Settings, const USubsurfaceProfilePointer Profile) { int32 AllocationId = FindAllocationId(Profile); UpdateProfile(AllocationId, Settings); }
+	// @param InProfile must not be 0, game thread pointer, do not dereference, only for comparison
+	void UpdateProfile(const FSubsurfaceProfileStruct Settings, const USubsurfaceProfile* Profile) { int32 AllocationId = FindAllocationId(Profile); UpdateProfile(AllocationId, Settings); }
 
-	// @param InRendererModule first call needs to have it != 0
+	// @param InProfile must not be 0, game thread pointer, do not dereference, only for comparison
 	void UpdateProfile(int32 AllocationId, const FSubsurfaceProfileStruct Settings);
 
 	// @return can be 0 if there is no SubsurfaceProfile
@@ -127,23 +126,26 @@ public:
 	// for debugging / VisualizeSSS
 	ENGINE_API bool GetEntryString(uint32 Index, FString& Out) const;
 
+	// @param InProfile must not be 0, game thread pointer, do not dereference, only for comparison
 	// @return -1 if not found
-	int32 FindAllocationId(const USubsurfaceProfilePointer InProfile) const;
+	int32 FindAllocationId(const USubsurfaceProfile* InProfile) const;
 
 private:
 
 	struct FSubsurfaceProfileEntry
 	{
-		FSubsurfaceProfileEntry(FSubsurfaceProfileStruct InSubsurfaceProfileStruct, const USubsurfaceProfilePointer InGameThreadObject)
+		// @param InProfile game thread pointer, do not dereference, only for comparison, 0 if the entry can be reused or it's [0] which is used as default
+		FSubsurfaceProfileEntry(FSubsurfaceProfileStruct InSubsurfaceProfileStruct, const USubsurfaceProfile* InProfile)
 			: Settings(InSubsurfaceProfileStruct)
-			, GameThreadObject(InGameThreadObject)
+			, Profile(InProfile)
 		{
 		}
 
 		FSubsurfaceProfileStruct Settings;
 
-		// not referenced, void and not USubsurfaceProfile* to prevent accidential access, 0 if the entry can be reused or it's [0] which is used as default
-		USubsurfaceProfilePointer GameThreadObject;
+		// 0 if the entry can be reused or it's [0] which is used as default
+		// game thread pointer, do not dereference, only for comparison
+		const USubsurfaceProfile* Profile;
 	};
 
 	//
diff --git a/Engine/Source/Runtime/Engine/Private/Materials/Material.cpp b/Engine/Source/Runtime/Engine/Private/Materials/Material.cpp
--- a/Engine/Source/Runtime/Engine/Private/Materials/Material.cpp
+++ b/Engine/Source/Runtime/Engine/Private/Materials/Material.cpp
@@ -228,7 +228,7 @@ public:
 			static FName NameSubsurfaceProfile(TEXT("__SubsurfaceProfile"));
 			if (ParameterName == NameSubsurfaceProfile)
 			{
-				const USubsurfaceProfilePointer SubsurfaceProfileRT = GetSubsurfaceProfileRT();
+				const USubsurfaceProfile* SubsurfaceProfileRT = GetSubsurfaceProfileRT();
 
 				int32 AllocationId = 0;
 				if(SubsurfaceProfileRT)
diff --git a/Engine/Source/Runtime/Engine/Private/Materials/MaterialInstance.cpp b/Engine/Source/Runtime/Engine/Private/Materials/MaterialInstance.cpp
--- a/Engine/Source/Runtime/Engine/Private/Materials/MaterialInstance.cpp
+++ b/Engine/Source/Runtime/Engine/Private/Materials/MaterialInstance.cpp
@@ -157,7 +157,7 @@ bool FMaterialInstanceResource::GetScalarValue(
 	static FName NameSubsurfaceProfile(TEXT("__SubsurfaceProfile"));
 	if (ParameterName == NameSubsurfaceProfile)
 	{
-		const USubsurfaceProfilePointer SubsurfaceProfileRT = GetSubsurfaceProfileRT();
+		const USubsurfaceProfile* SubsurfaceProfileRT = GetSubsurfaceProfileRT();
 
 		int32 AllocationId = 0;
 		if (SubsurfaceProfileRT)
diff --git a/Engine/Source/Runtime/Engine/Private/Materials/MaterialInterface.cpp b/Engine/Source/Runtime/Engine/Private/Materials/MaterialInterface.cpp
--- a/Engine/Source/Runtime/Engine/Private/Materials/MaterialInterface.cpp
+++ b/Engine/Source/Runtime/Engine/Private/Materials/MaterialInterface.cpp
@@ -323,7 +323,7 @@ void UMaterialInterface::UpdateMaterialRenderProxy(FMaterialRenderProxy& Proxy)
 		ENQUEUE_UNIQUE_RENDER_COMMAND_THREEPARAMETER(
 			UpdateMaterialRenderProxySubsurface,
 			const FSubsurfaceProfileStruct, Settings, Settings,
-			USubsurfaceProfilePointer, LocalSubsurfaceProfile, (USubsurfaceProfilePointer)LocalSubsurfaceProfile,
+			USubsurfaceProfile*, LocalSubsurfaceProfile, LocalSubsurfaceProfile,
 			FMaterialRenderProxy&, Proxy, Proxy,
 		{
 			uint32 AllocationId = 0;
diff --git a/Engine/Source/Runtime/Engine/Private/Rendering/SubsurfaceProfile.cpp b/Engine/Source/Runtime/Engine/Private/Rendering/SubsurfaceProfile.cpp
--- a/Engine/Source/Runtime/Engine/Private/Rendering/SubsurfaceProfile.cpp
+++ b/Engine/Source/Runtime/Engine/Private/Rendering/SubsurfaceProfile.cpp
@@ -17,6 +17,8 @@ static TRefCountPtr<IPooledRenderTarget> GSSProfiles;
 
 FSubsurfaceProfileTexture::FSubsurfaceProfileTexture()
 {
+	check(IsInGameThread());
+
 	FSubsurfaceProfileStruct DefaultSkin;
 
 	// add element 0, it is used as default profile
@@ -25,14 +27,16 @@ FSubsurfaceProfileTexture::FSubsurfaceProfileTexture()
 
 FSubsurfaceProfileTexture::~FSubsurfaceProfileTexture()
 {
+	check(IsInGameThread());
+
 	// we assume all ~USubsurfaceProfile() have been called already
 	for (int32 i = 0; i < SubsurfaceProfileEntries.Num(); ++i)
 	{
-		check(SubsurfaceProfileEntries[i].GameThreadObject == 0);
+		check(SubsurfaceProfileEntries[i].Profile == 0);
 	}
 }
 
-int32 FSubsurfaceProfileTexture::AddProfile(const FSubsurfaceProfileStruct Settings, const USubsurfaceProfilePointer InProfile)
+int32 FSubsurfaceProfileTexture::AddProfile(const FSubsurfaceProfileStruct Settings, const USubsurfaceProfile* InProfile)
 {
 	check(InProfile);
 	check(FindAllocationId(InProfile) == -1);
@@ -41,7 +45,7 @@ int32 FSubsurfaceProfileTexture::AddProfile(const FSubsurfaceProfileStruct Setti
 	{
 		for (int32 i = 1; i < SubsurfaceProfileEntries.Num(); ++i)
 		{
-			if (SubsurfaceProfileEntries[i].GameThreadObject == 0)
+			if (SubsurfaceProfileEntries[i].Profile == 0)
 			{
 				RetAllocationId = i; break;
 			}
@@ -60,7 +64,7 @@ int32 FSubsurfaceProfileTexture::AddProfile(const FSubsurfaceProfileStruct Setti
 }
 
 
-void FSubsurfaceProfileTexture::RemoveProfile(const USubsurfaceProfilePointer InProfile)
+void FSubsurfaceProfileTexture::RemoveProfile(const USubsurfaceProfile* InProfile)
 {
 	int32 AllocationId = FindAllocationId(InProfile);
 
@@ -73,10 +77,10 @@ void FSubsurfaceProfileTexture::RemoveProfile(const USubsurfaceProfilePointer In
 	// >0 as 0 is used as default profile which should never be removed
 	check(AllocationId > 0);
 
-	check(SubsurfaceProfileEntries[AllocationId].GameThreadObject == InProfile);
+	check(SubsurfaceProfileEntries[AllocationId].Profile == InProfile);
 
 	// make it available for reuse
-	SubsurfaceProfileEntries[AllocationId].GameThreadObject = 0;
+	SubsurfaceProfileEntries[AllocationId].Profile = 0;
 	SubsurfaceProfileEntries[AllocationId].Settings.Invalidate();
 }
 
@@ -243,7 +247,7 @@ bool FSubsurfaceProfileTexture::GetEntryString(uint32 Index, FString& Out) const
 
 	Out = FString::Printf(TEXT(" %c. %p ScatterRadius=%.1f, SubsurfaceColor=%.1f %.1f %.1f, FalloffColor=%.1f %.1f %.1f"), 
 		MiniFontCharFromIndex(Index), 
-		SubsurfaceProfileEntries[Index].GameThreadObject,
+		SubsurfaceProfileEntries[Index].Profile,
 		ref.ScatterRadius,
 		ref.SubsurfaceColor.R, ref.SubsurfaceColor.G, ref.SubsurfaceColor.B,
 		ref.FalloffColor.R, ref.FalloffColor.G, ref.FalloffColor.B);
@@ -251,12 +255,12 @@ bool FSubsurfaceProfileTexture::GetEntryString(uint32 Index, FString& Out) const
 	return true;
 }
 
-int32 FSubsurfaceProfileTexture::FindAllocationId(const USubsurfaceProfilePointer InProfile) const
+int32 FSubsurfaceProfileTexture::FindAllocationId(const USubsurfaceProfile* InProfile) const
 {
-	// we start at 1 because [0] is the default profile and always [0].GameThreadObject = 0 so we don't need to iterate that one
+	// we start at 1 because [0] is the default profile and always [0].Profile = 0 so we don't need to iterate that one
 	for (int32 i = 1; i < SubsurfaceProfileEntries.Num(); ++i)
 	{
-		if (SubsurfaceProfileEntries[i].GameThreadObject == InProfile)
+		if (SubsurfaceProfileEntries[i].Profile == InProfile)
 		{
 			return i;
 		}
@@ -273,7 +277,7 @@ void FSubsurfaceProfileTexture::Dump()
 	for (int32 i = 0; i < SubsurfaceProfileEntries.Num(); ++i)
 	{
 		// + 1 as the Id is one higher than the array index, 0 is used for the default profile (not assigned)
-		UE_LOG(LogSubsurfaceProfile, Log, TEXT("  %d. AllocationId=%d, Pointer=%p"), i, i + 1, SubsurfaceProfileEntries[i].GameThreadObject);
+		UE_LOG(LogSubsurfaceProfile, Log, TEXT("  %d. AllocationId=%d, Pointer=%p"), i, i + 1, SubsurfaceProfileEntries[i].Profile);
 
 		{
 			UE_LOG(LogSubsurfaceProfile, Log, TEXT("     ScatterRadius = %f"),
@@ -309,7 +313,7 @@ void USubsurfaceProfile::BeginDestroy()
 {
 	ENQUEUE_UNIQUE_RENDER_COMMAND_ONEPARAMETER(
 		RemoveSubsurfaceProfile,
-		USubsurfaceProfilePointer, Ref, (USubsurfaceProfilePointer)this,
+		USubsurfaceProfile*, Ref, this,
 		{
 			GSubsufaceProfileTextureObject.RemoveProfile(Ref);
 		});
@@ -322,7 +326,7 @@ void USubsurfaceProfile::PostEditChangeProperty(struct FPropertyChangedEvent& Pr
 	ENQUEUE_UNIQUE_RENDER_COMMAND_TWOPARAMETER(
 		UpdateSubsurfaceProfile,
 		const FSubsurfaceProfileStruct, Settings, this->Settings,
-		USubsurfaceProfilePointer, Profile, (USubsurfaceProfilePointer)this,
+		USubsurfaceProfile*, Profile, this,
 	{
 		// any changes to the setting require an update of the texture
 		GSubsufaceProfileTextureObject.UpdateProfile(Settings, Profile);
diff --git a/Engine/Source/Runtime/Engine/Public/MaterialShared.h b/Engine/Source/Runtime/Engine/Public/MaterialShared.h
--- a/Engine/Source/Runtime/Engine/Public/MaterialShared.h
+++ b/Engine/Source/Runtime/Engine/Public/MaterialShared.h
@@ -1257,7 +1257,6 @@ struct FUniformExpressionCache
 };
 
 class USubsurfaceProfile;
-typedef void* USubsurfaceProfilePointer;
 
 /**
  * A material render proxy used by the renderer.
@@ -1322,8 +1321,8 @@ public:
 		return MaterialRenderProxyMap;
 	}
 
-	void SetSubsurfaceProfileRT(const USubsurfaceProfilePointer Ptr) { SubsurfaceProfileRT = Ptr; }
-	USubsurfaceProfilePointer GetSubsurfaceProfileRT() const { return SubsurfaceProfileRT; }
+	void SetSubsurfaceProfileRT(const USubsurfaceProfile* Ptr) { SubsurfaceProfileRT = Ptr; }
+	const USubsurfaceProfile* GetSubsurfaceProfileRT() const { return SubsurfaceProfileRT; }
 
 private:
 
@@ -1331,8 +1330,8 @@ private:
 	bool bSelected : 1;
 	/** true if the material is hovered. */
 	bool bHovered : 1;
-	/** 0 if not set, for the render thread */
-	USubsurfaceProfilePointer SubsurfaceProfileRT;
+	/** 0 if not set, game thread pointer, do not dereference, only for comparison */
+	const USubsurfaceProfile* SubsurfaceProfileRT;
 
 	/** 
 	 * Tracks all material render proxies in all scenes, can only be accessed on the rendering thread.
