[UE-9868] - IWCE: Adding a component to an empty actor with multiple child actor attachments will ensure
- USceneComponent::DestroyComponent() was attempting to promote child components to the root position but they could be components on attched child actors.
- Added a check to make sure only components with the same owner are selected in this case.

diff --git a/Engine/Source/Runtime/Engine/Private/SceneComponent.cpp b/Engine/Source/Runtime/Engine/Private/SceneComponent.cpp
--- a/Engine/Source/Runtime/Engine/Private/SceneComponent.cpp
+++ b/Engine/Source/Runtime/Engine/Private/SceneComponent.cpp
@@ -295,36 +295,24 @@ void USceneComponent::DestroyComponent(bool bPromoteChildren/*= false*/)
 		if (Owner != NULL)
 		{
 			Owner->Modify();
-
-			// Find an appropriate child node to promote to this node's position in the hierarchy
 			USceneComponent* ChildToPromote = nullptr;
-			if (AttachChildren.Num() > 0)
-			{
-				// Start with the first child node
-				ChildToPromote = AttachChildren[0];
-				check(ChildToPromote != nullptr);
-
-				// Always choose non editor-only child nodes over editor-only child nodes (since we don't want editor-only nodes to end up with non editor-only child nodes)
-				if (ChildToPromote->IsEditorOnly())
-				{
-					for (int32 ChildIndex = 1; ChildIndex < AttachChildren.Num(); ++ChildIndex)
-					{
-						USceneComponent* Child = AttachChildren[ChildIndex];
-						if (Child != nullptr && !Child->IsEditorOnly())
-						{
-							ChildToPromote = Child;
-							break;
-						}
-					}
-				}
-			}
 
 			// Handle removal of the root node
 			if (this == Owner->GetRootComponent())
 			{
-				// We only promote non editor-only components to root in instanced mode
-				if (ChildToPromote == nullptr || ChildToPromote->IsEditorOnly())
+				// Always choose non editor-only child nodes over editor-only child nodes (since we don't want editor-only nodes to end up with non editor-only child nodes)
+				// Exclude scene components owned by attached child actors
+				USceneComponent** FindResult =
+					AttachChildren.FindByPredicate([Owner](USceneComponent* Child){ return Child != nullptr && !Child->IsEditorOnly() && Child->GetOwner() == Owner; });
+
+				if (FindResult != nullptr)
+				{
+					ChildToPromote = *FindResult;
+				}
+				else
 				{
+					// Didn't find a suitable component to promote so create a new default component
+
 					Rename(NULL, GetOuter(), REN_DoNotDirty | REN_DontCreateRedirectors);
 
 					// Construct a new default root component
@@ -360,6 +348,25 @@ void USceneComponent::DestroyComponent(bool bPromoteChildren/*= false*/)
 				// Detach from parent
 				DetachFromParent(true);
 
+				// Find an appropriate child node to promote to this node's position in the hierarchy
+				if (AttachChildren.Num() > 0)
+				{
+					// Always choose non editor-only child nodes over editor-only child nodes (since we don't want editor-only nodes to end up with non editor-only child nodes)
+					USceneComponent** FindResult =
+						AttachChildren.FindByPredicate([Owner](USceneComponent* Child){ return Child != nullptr && !Child->IsEditorOnly(); });
+
+					if (FindResult != nullptr)
+					{
+						ChildToPromote = *FindResult;
+					}
+					else
+					{
+						// Default to first child node
+						check(AttachChildren[0] != nullptr);
+						ChildToPromote = AttachChildren[0];
+					}
+				}
+
 				if (ChildToPromote != nullptr)
 				{
 					// Attach the child node that we're promoting to the parent and move it to the same position as the old node was in the array
