FindTimer can now find the currently-executing timer.

diff --git a/Engine/Source/Runtime/Engine/Private/TimerManager.cpp b/Engine/Source/Runtime/Engine/Private/TimerManager.cpp
--- a/Engine/Source/Runtime/Engine/Private/TimerManager.cpp
+++ b/Engine/Source/Runtime/Engine/Private/TimerManager.cpp
@@ -32,6 +32,16 @@ inline bool DEPRECATED_CompareUnifiedDelegates(const FTimerUnifiedDelegate& Lhs,
 /** Will find and return a timer if it exists, regardless whether it is paused. */ 
 FTimerData const* FTimerManager::DEPRECATED_FindTimer(FTimerUnifiedDelegate const& InDelegate, int32* OutTimerIndex) const
 {
+	if (!CurrentlyExecutingTimer.TimerHandle.IsValid() && DEPRECATED_CompareUnifiedDelegates(InDelegate, CurrentlyExecutingTimer.TimerDelegate))
+	{
+		// found it currently executing
+		if (OutTimerIndex)
+		{
+			*OutTimerIndex = -1;
+		}
+		return &CurrentlyExecutingTimer;
+	}
+
 	int32 ActiveTimerIdx = DEPRECATED_FindTimerInList(ActiveTimerHeap, InDelegate);
 	if (ActiveTimerIdx != INDEX_NONE)
 	{
@@ -70,6 +80,16 @@ FTimerData const* FTimerManager::DEPRECATED_FindTimer(FTimerUnifiedDelegate cons
 
 FTimerData const* FTimerManager::FindTimer(FTimerHandle const& InHandle, int32* OutTimerIndex) const
 {
+	if (CurrentlyExecutingTimer.TimerHandle == InHandle)
+	{
+		// found it currently executing
+		if (OutTimerIndex)
+		{
+			*OutTimerIndex = -1;
+		}
+		return &CurrentlyExecutingTimer;
+	}
+
 	int32 ActiveTimerIdx = FindTimerInList(ActiveTimerHeap, InHandle);
 	if (ActiveTimerIdx != INDEX_NONE)
 	{
@@ -259,15 +279,6 @@ void FTimerManager::DEPRECATED_InternalClearTimer(FTimerUnifiedDelegate const& I
 	{
 		InternalClearTimer(TimerIdx, TimerData->Status);
 	}
-	else
-	{
-		// Edge case. We're currently handling this timer when it got cleared.  Unbind it to prevent it firing again
-		// in case it was scheduled to fire multiple times.
-		if (!CurrentlyExecutingTimer.TimerHandle.IsValid() && DEPRECATED_CompareUnifiedDelegates(CurrentlyExecutingTimer.TimerDelegate, InDelegate))
-		{
-			CurrentlyExecutingTimer.TimerDelegate.Unbind();
-		}
-	}
 }
 
 void FTimerManager::InternalClearTimer(FTimerHandle const& InHandle)
@@ -287,16 +298,6 @@ void FTimerManager::InternalClearTimer(FTimerHandle const& InHandle)
 	{
 		InternalClearTimer(TimerIdx, TimerData->Status);
 	}
-	else
-	{
-		// Edge case. We're currently handling this timer when it got cleared.  Unbind it to prevent it firing again
-		// in case it was scheduled to fire multiple times.
-		if (CurrentlyExecutingTimer.TimerHandle == InHandle)
-		{
-			CurrentlyExecutingTimer.TimerDelegate.Unbind();
-			CurrentlyExecutingTimer.TimerHandle.Invalidate();
-		}
-	}
 }
 
 void FTimerManager::InternalClearTimer(int32 TimerIdx, ETimerStatus::Type TimerStatus)
@@ -315,6 +316,13 @@ void FTimerManager::InternalClearTimer(int32 TimerIdx, ETimerStatus::Type TimerS
 			PausedTimerList.RemoveAtSwap(TimerIdx);
 			break;
 
+		case ETimerStatus::Executing:
+			// Edge case. We're currently handling this timer when it got cleared.  Unbind it to prevent it firing again
+			// in case it was scheduled to fire multiple times.
+			CurrentlyExecutingTimer.TimerDelegate.Unbind();
+			CurrentlyExecutingTimer.TimerHandle.Invalidate();
+			break;
+
 		default:
 			check(false);
 	}
@@ -366,6 +374,7 @@ void FTimerManager::InternalClearAllTimers(void const* Object)
 		if (CurrentlyExecutingTimer.TimerDelegate.IsBoundToObject(Object))
 		{
 			CurrentlyExecutingTimer.TimerDelegate.Unbind();
+			CurrentlyExecutingTimer.TimerHandle.Invalidate();
 		}
 	}
 }
@@ -426,27 +435,41 @@ void FTimerManager::InternalPauseTimer( FTimerData const* TimerToPause, int32 Ti
 	{
 		ETimerStatus::Type PreviousStatus = TimerToPause->Status;
 
-		// Add to Paused list
-		int32 NewIndex = PausedTimerList.Add(*TimerToPause);
+		// Don't pause the timer if it's currently executing and isn't going to loop
+		if( PreviousStatus != ETimerStatus::Executing || TimerToPause->bLoop )
+		{
+			// Add to Paused list
+			int32 NewIndex = PausedTimerList.Add(*TimerToPause);
+
+			// Set new status
+			FTimerData& NewTimer = PausedTimerList[NewIndex];
+			NewTimer.Status = ETimerStatus::Paused;
 
-		// Set new status
-		FTimerData &NewTimer = PausedTimerList[NewIndex];
-		NewTimer.Status = ETimerStatus::Paused;
+			if( PreviousStatus == ETimerStatus::Active )
+			{
+				// Store time remaining in ExpireTime while paused
+				NewTimer.ExpireTime = NewTimer.ExpireTime - InternalTime;
+			}
+		}
 
 		// Remove from previous TArray
 		switch( PreviousStatus )
 		{
-			case ETimerStatus::Active : 
-				// Store time remaining in ExpireTime while paused
-				NewTimer.ExpireTime = NewTimer.ExpireTime - InternalTime;
-				ActiveTimerHeap.HeapRemoveAt(TimerIdx); 
+			case ETimerStatus::Active:
+				ActiveTimerHeap.HeapRemoveAt(TimerIdx);
 				break;
-			
-			case ETimerStatus::Pending : 
-				PendingTimerList.RemoveAtSwap(TimerIdx); 
+
+			case ETimerStatus::Pending:
+				PendingTimerList.RemoveAtSwap(TimerIdx);
 				break;
 
-			default : check(false);
+			case ETimerStatus::Executing:
+				CurrentlyExecutingTimer.TimerDelegate.Unbind();
+				CurrentlyExecutingTimer.TimerHandle.Invalidate();
+				break;
+
+			default:
+				check(false);
 		}
 	}
 }
@@ -499,7 +522,8 @@ void FTimerManager::Tick(float DeltaTime)
 			// Timer has expired! Fire the delegate, then handle potential looping.
 
 			// Remove it from the heap and store it while we're executing
-			ActiveTimerHeap.HeapPop(CurrentlyExecutingTimer); 
+			ActiveTimerHeap.HeapPop(CurrentlyExecutingTimer);
+			CurrentlyExecutingTimer.Status = ETimerStatus::Executing;
 
 			// Determine how many times the timer may have elapsed (e.g. for large DeltaTime on a short looping timer)
 			int32 const CallCount = CurrentlyExecutingTimer.bLoop ? 
@@ -528,10 +552,12 @@ void FTimerManager::Tick(float DeltaTime)
 			{
 				// Put this timer back on the heap
 				CurrentlyExecutingTimer.ExpireTime += CallCount * CurrentlyExecutingTimer.Rate;
+				CurrentlyExecutingTimer.Status = ETimerStatus::Active;
 				ActiveTimerHeap.HeapPush(CurrentlyExecutingTimer);
 			}
 
 			CurrentlyExecutingTimer.TimerDelegate.Unbind();
+			CurrentlyExecutingTimer.TimerHandle.Invalidate();
 		}
 		else
 		{
diff --git a/Engine/Source/Runtime/Engine/Public/TimerManager.h b/Engine/Source/Runtime/Engine/Public/TimerManager.h
--- a/Engine/Source/Runtime/Engine/Public/TimerManager.h
+++ b/Engine/Source/Runtime/Engine/Public/TimerManager.h
@@ -115,7 +115,8 @@ namespace ETimerStatus
 	{
 		Pending,
 		Active,
-		Paused
+		Paused,
+		Executing
 	};
 }
 
