UE4 - merge GDC branch, code [REDACTED_HANDLE] to main

diff --git a/Engine/Source/Runtime/Core/Private/Stats/StatsCommand.cpp b/Engine/Source/Runtime/Core/Private/Stats/StatsCommand.cpp
--- a/Engine/Source/Runtime/Core/Private/Stats/StatsCommand.cpp
+++ b/Engine/Source/Runtime/Core/Private/Stats/StatsCommand.cpp
@@ -329,6 +329,9 @@ void DumpCPUSummary(FStatsThreadState const& StatsData, int64 TargetFrame)
 	}
 }
 
+static int32 HitchIndex = 0;
+static float TotalHitchTime = 0.0f;
+
 static void DumpHitch(int64 Frame)
 {
 	// !!!CAUTION!!! 
@@ -350,7 +353,10 @@ static void DumpHitch(int64 Frame)
 
 	if( GameThreadTime > GHitchThreshold || RenderThreadTime > GHitchThreshold )
 	{
-		UE_LOG(LogStats, Log, TEXT("------------------Thread Hitch, Frame %lld  %6.1fms ---------------"), Frame, FMath::Max<float>(GameThreadTime, RenderThreadTime) * 1000.0f );
+		HitchIndex++;
+		float ThisHitch = FMath::Max<float>(GameThreadTime, RenderThreadTime) * 1000.0f;
+		TotalHitchTime += ThisHitch;
+		UE_LOG(LogStats, Log, TEXT("------------------Thread Hitch %d, Frame %lld  %6.1fms ---------------"), HitchIndex, Frame, ThisHitch);
 		FRawStatStackNode Stack;
 		Stats.UncondenseStackStats(Frame, Stack);
 		Stack.AddNameHierarchy();
@@ -1221,12 +1227,15 @@ static void StatCmd(FString InCmd)
 		if (bToggle)
 		{
 			StatsMasterEnableAdd();
+			HitchIndex = 0;
+			TotalHitchTime = 0.0f;
 			DumpHitchDelegateHandle = Stats.NewFrameDelegate.AddStatic(&DumpHitch);
 		}
 		else
 		{
 			StatsMasterEnableSubtract();
 			Stats.NewFrameDelegate.Remove(DumpHitchDelegateHandle);
+			UE_LOG(LogStats, Log, TEXT( "**************************** %d hitches    %8.0fms total hitch time" ), HitchIndex, TotalHitchTime);
 		}
 	}
 	else if( FParse::Command( &Cmd, TEXT( "STARTFILE" ) ) )
diff --git a/Engine/Source/Runtime/Core/Private/Windows/WindowsApplication.cpp b/Engine/Source/Runtime/Core/Private/Windows/WindowsApplication.cpp
--- a/Engine/Source/Runtime/Core/Private/Windows/WindowsApplication.cpp
+++ b/Engine/Source/Runtime/Core/Private/Windows/WindowsApplication.cpp
@@ -268,15 +268,7 @@ void FWindowsApplication::SetMessageHandler( const TSharedRef< FGenericApplicati
 
 FModifierKeysState FWindowsApplication::GetModifierKeys() const
 {
-	const bool bIsLeftShiftDown = ( ::GetAsyncKeyState( VK_LSHIFT ) & 0x8000 ) != 0;
-	const bool bIsRightShiftDown = ( ::GetAsyncKeyState( VK_RSHIFT ) & 0x8000 ) != 0;
-	const bool bIsLeftControlDown = ( ::GetAsyncKeyState( VK_LCONTROL ) & 0x8000 ) != 0;
-	const bool bIsRightControlDown = ( ::GetAsyncKeyState( VK_RCONTROL ) & 0x8000 ) != 0;
-	const bool bIsLeftAltDown = ( ::GetAsyncKeyState( VK_LMENU ) & 0x8000 ) != 0;
-	const bool bIsRightAltDown = ( ::GetAsyncKeyState (VK_RMENU ) & 0x8000 ) != 0;
-	const bool bAreCapsLocked = ( ::GetKeyState( VK_CAPITAL ) & 0x0001 ) != 0;
-
-	return FModifierKeysState(bIsLeftShiftDown, bIsRightShiftDown, bIsLeftControlDown, bIsRightControlDown, bIsLeftAltDown, bIsRightAltDown, false, false, bAreCapsLocked); // Win key is ignored
+	return CachedModifierKeyState;
 }
 
 
@@ -298,11 +290,15 @@ static TSharedPtr< FWindowsWindow > FindWindowByHWND(const TArray< TSharedRef< F
 bool FWindowsApplication::IsCursorDirectlyOverSlateWindow() const
 {
 	POINT CursorPos;
-	::GetCursorPos(&CursorPos);
-	HWND hWnd = ::WindowFromPoint(CursorPos);
-	
-	TSharedPtr< FWindowsWindow > SlatWindowUnderCursor = FindWindowByHWND(Windows, hWnd);
-	return SlatWindowUnderCursor.IsValid();
+	BOOL bGotPoint = ::GetCursorPos(&CursorPos);
+	if (bGotPoint)
+	{
+		HWND hWnd = ::WindowFromPoint(CursorPos);
+
+		TSharedPtr< FWindowsWindow > SlatWindowUnderCursor = FindWindowByHWND(Windows, hWnd);
+		return SlatWindowUnderCursor.IsValid();
+	}
+	return false;
 }
 
 
@@ -1752,6 +1748,19 @@ void FWindowsApplication::ProcessDeferredEvents( const float TimeDelta )
 	}
 }
 
+void FWindowsApplication::Tick( const float TimeDelta )
+{
+	const bool bIsLeftShiftDown = (::GetAsyncKeyState(VK_LSHIFT) & 0x8000) != 0;
+	const bool bIsRightShiftDown = (::GetAsyncKeyState(VK_RSHIFT) & 0x8000) != 0;
+	const bool bIsLeftControlDown = (::GetAsyncKeyState(VK_LCONTROL) & 0x8000) != 0;
+	const bool bIsRightControlDown = (::GetAsyncKeyState(VK_RCONTROL) & 0x8000) != 0;
+	const bool bIsLeftAltDown = (::GetAsyncKeyState(VK_LMENU) & 0x8000) != 0;
+	const bool bIsRightAltDown = (::GetAsyncKeyState(VK_RMENU) & 0x8000) != 0;
+	const bool bAreCapsLocked = (::GetKeyState(VK_CAPITAL) & 0x0001) != 0;
+
+	CachedModifierKeyState = FModifierKeysState(bIsLeftShiftDown, bIsRightShiftDown, bIsLeftControlDown, bIsRightControlDown, bIsLeftAltDown, bIsRightAltDown, false, false, bAreCapsLocked); // Win key is ignored
+}
+
 void FWindowsApplication::PollGameDeviceState( const float TimeDelta )
 {
 	// initialize any externally-implemented input devices (we delay load initialize the array so any plugins have had time to load)
diff --git a/Engine/Source/Runtime/Core/Private/Windows/WindowsPlatformMemory.cpp b/Engine/Source/Runtime/Core/Private/Windows/WindowsPlatformMemory.cpp
--- a/Engine/Source/Runtime/Core/Private/Windows/WindowsPlatformMemory.cpp
+++ b/Engine/Source/Runtime/Core/Private/Windows/WindowsPlatformMemory.cpp
@@ -123,7 +123,7 @@ void FWindowsPlatformMemory::GetStatsForMallocProfiler( FGenericMemoryStats& out
 	FPlatformMemoryStats Stats = GetStats();
 
 	// Windows specific stats.
-	out_Stats.Add( GET_STATDESCRIPTION( STAT_WindowsSpecificMemoryStat ), Stats.WindowsSpecificMemoryStat );
+	out_Stats.Add(TEXT("Windows Specific Memory Stat"), Stats.WindowsSpecificMemoryStat );
 #endif // STATS
 }
 
@@ -153,11 +153,26 @@ const FPlatformMemoryConstants& FWindowsPlatformMemory::GetConstants()
 
 void FWindowsPlatformMemory::UpdateStats()
 {
-	FGenericPlatformMemory::UpdateStats();
-
-	// Windows specific stats.
-	FPlatformMemoryStats MemoryStats = FPlatformMemory::GetStats();
-	SET_MEMORY_STAT(STAT_WindowsSpecificMemoryStat,MemoryStats.WindowsSpecificMemoryStat);
+#if STATS
+	if (FThreadStats::IsCollectingData(GET_STATID(STAT_TotalPhysical)))
+	{
+		FPlatformMemoryStats MemoryStats = FPlatformMemory::GetStats();
+		SET_MEMORY_STAT(STAT_TotalPhysical,MemoryStats.TotalPhysical);
+		SET_MEMORY_STAT(STAT_TotalVirtual,MemoryStats.TotalVirtual);
+		SET_MEMORY_STAT(STAT_PageSize,MemoryStats.PageSize);
+		SET_MEMORY_STAT(STAT_TotalPhysicalGB,MemoryStats.TotalPhysicalGB);
+
+		SET_MEMORY_STAT(STAT_AvailablePhysical,MemoryStats.AvailablePhysical);
+		SET_MEMORY_STAT(STAT_AvailableVirtual,MemoryStats.AvailableVirtual);
+		SET_MEMORY_STAT(STAT_UsedPhysical,MemoryStats.UsedPhysical);
+		SET_MEMORY_STAT(STAT_PeakUsedPhysical,MemoryStats.PeakUsedPhysical);
+		SET_MEMORY_STAT(STAT_UsedVirtual,MemoryStats.UsedVirtual);
+		SET_MEMORY_STAT(STAT_PeakUsedVirtual,MemoryStats.PeakUsedVirtual);
+
+		// Windows specific stats.
+		SET_MEMORY_STAT(STAT_WindowsSpecificMemoryStat,MemoryStats.WindowsSpecificMemoryStat);
+	}
+#endif
 }
 
 void* FWindowsPlatformMemory::BinnedAllocFromOS( SIZE_T Size )
diff --git a/Engine/Source/Runtime/Core/Public/Containers/Ticker.h b/Engine/Source/Runtime/Core/Public/Containers/Ticker.h
--- a/Engine/Source/Runtime/Core/Public/Containers/Ticker.h
+++ b/Engine/Source/Runtime/Core/Public/Containers/Ticker.h
@@ -71,6 +71,7 @@ public:
 	 */
 	void Tick(float DeltaTime)
 	{
+		QUICK_SCOPE_CYCLE_COUNTER(STAT_FTicker_Tick);
 		if (!PriorityQueue.Num())
 		{
 			return;
diff --git a/Engine/Source/Runtime/Core/Public/GenericPlatform/GenericPlatformMemoryPoolStats.h b/Engine/Source/Runtime/Core/Public/GenericPlatform/GenericPlatformMemoryPoolStats.h
--- a/Engine/Source/Runtime/Core/Public/GenericPlatform/GenericPlatformMemoryPoolStats.h
+++ b/Engine/Source/Runtime/Core/Public/GenericPlatform/GenericPlatformMemoryPoolStats.h
@@ -9,5 +9,18 @@
 DECLARE_MEMORY_STAT_POOL_EXTERN(TEXT("Physical Memory Pool [Physical]"), MCR_Physical, STATGROUP_Memory,  FPlatformMemory::MCR_Physical, CORE_API);
 DECLARE_MEMORY_STAT_POOL_EXTERN(TEXT("GPU Memory Pool [GPU]"), MCR_GPU, STATGROUP_Memory,  FPlatformMemory::MCR_GPU, CORE_API);
 DECLARE_MEMORY_STAT_POOL_EXTERN(TEXT("Texture Memory Pool [Texture]"), MCR_TexturePool, STATGROUP_Memory,  FPlatformMemory::MCR_TexturePool, CORE_API);
+// Must match values in the MemoryProfiler2.FMemoryAllocationStatsV4
+DECLARE_MEMORY_STAT_EXTERN(TEXT("Total Physical"),		STAT_TotalPhysical,STATGROUP_MemoryPlatform, CORE_API);
+DECLARE_MEMORY_STAT_EXTERN(TEXT("Total Virtual"),		STAT_TotalVirtual,STATGROUP_MemoryPlatform, CORE_API);
+DECLARE_MEMORY_STAT_EXTERN(TEXT("Page Size"),			STAT_PageSize,STATGROUP_MemoryPlatform, CORE_API);
+DECLARE_MEMORY_STAT_EXTERN(TEXT("Total Physical GB"),	STAT_TotalPhysicalGB,STATGROUP_MemoryPlatform, CORE_API);
+
+DECLARE_MEMORY_STAT_EXTERN(TEXT("Available Physical"),	STAT_AvailablePhysical,STATGROUP_MemoryPlatform, CORE_API);
+DECLARE_MEMORY_STAT_EXTERN(TEXT("Available Virtual"),	STAT_AvailableVirtual,STATGROUP_MemoryPlatform, CORE_API);
+DECLARE_MEMORY_STAT_EXTERN(TEXT("Used Physical"),		STAT_UsedPhysical,STATGROUP_MemoryPlatform, CORE_API);
+DECLARE_MEMORY_STAT_EXTERN(TEXT("Peak Used Physical"),	STAT_PeakUsedPhysical,STATGROUP_MemoryPlatform, CORE_API);
+DECLARE_MEMORY_STAT_EXTERN(TEXT("Used Virtual"),		STAT_UsedVirtual,STATGROUP_MemoryPlatform, CORE_API);
+DECLARE_MEMORY_STAT_EXTERN(TEXT("Peak Used Virtual"),	STAT_PeakUsedVirtual,STATGROUP_MemoryPlatform, CORE_API);
+
 
 
diff --git a/Engine/Source/Runtime/Core/Public/HAL/MallocBinned.h b/Engine/Source/Runtime/Core/Public/HAL/MallocBinned.h
--- a/Engine/Source/Runtime/Core/Public/HAL/MallocBinned.h
+++ b/Engine/Source/Runtime/Core/Public/HAL/MallocBinned.h
@@ -19,8 +19,12 @@
 #endif
 
 #if defined CACHE_FREED_OS_ALLOCS
-#	define MAX_CACHED_OS_FREES (32)
-#	define MAX_CACHED_OS_FREES_BYTE_LIMIT (4*1024*1024)
+	#define MAX_CACHED_OS_FREES (64)
+	#if PLATFORM_64BITS
+		#define MAX_CACHED_OS_FREES_BYTE_LIMIT (64*1024*1024)
+	#else
+		#define MAX_CACHED_OS_FREES_BYTE_LIMIT (16*1024*1024)
+	#endif
 #endif
 
 #if defined USE_INTERNAL_LOCKS && !defined USE_COARSE_GRAIN_LOCKS
@@ -686,17 +690,21 @@ private:
 #ifdef USE_FINE_GRAIN_LOCKS
 		FScopeLock MainLock(&AccessGuard);
 #endif
-		if ((CachedTotal + Size > MAX_CACHED_OS_FREES_BYTE_LIMIT) || (Size > BINNED_ALLOC_POOL_SIZE))
+		if (Size > MAX_CACHED_OS_FREES_BYTE_LIMIT / 4)
 		{
 			FPlatformMemory::BinnedFreeToOS(Ptr);
 			return;
 		}
-		if (FreedPageBlocksNum >= MAX_CACHED_OS_FREES) 
+		while (FreedPageBlocksNum && (FreedPageBlocksNum >= MAX_CACHED_OS_FREES || CachedTotal + Size > MAX_CACHED_OS_FREES_BYTE_LIMIT)) 
 		{
 			//Remove the oldest one
-			void* FreePtr = FreedPageBlocks[FreedPageBlocksNum-1].Ptr;
-			CachedTotal -= FreedPageBlocks[FreedPageBlocksNum-1].ByteSize;
-			--FreedPageBlocksNum;
+			void* FreePtr = FreedPageBlocks[0].Ptr;
+			CachedTotal -= FreedPageBlocks[0].ByteSize;
+			FreedPageBlocksNum--;
+			if (FreedPageBlocksNum)
+			{
+				FMemory::Memmove(&FreedPageBlocks[0], &FreedPageBlocks[1], sizeof(FFreePageBlock) * FreedPageBlocksNum);
+			}
 			FPlatformMemory::BinnedFreeToOS(FreePtr);
 		}
 		FreedPageBlocks[FreedPageBlocksNum].Ptr = Ptr;
@@ -718,6 +726,22 @@ private:
 			// before the big call to the OS
 			FScopeLock MainLock(&AccessGuard);
 #endif
+			for (uint32 i=0; i < FreedPageBlocksNum; ++i)
+			{
+				// look for exact matches first, these are aligned to the page size, so it should be quite common to hit these on small pages sizes
+				if (FreedPageBlocks[i].ByteSize == NewSize)
+				{
+					void* Ret=FreedPageBlocks[i].Ptr;
+					OutActualSize=FreedPageBlocks[i].ByteSize;
+					CachedTotal-=FreedPageBlocks[i].ByteSize;
+					if (i < FreedPageBlocksNum - 1)
+					{
+						FMemory::Memmove(&FreedPageBlocks[i], &FreedPageBlocks[i + 1], sizeof(FFreePageBlock) * (FreedPageBlocksNum - i - 1));
+					}
+					FreedPageBlocksNum--;
+					return Ret;
+				}
+			};
 			for (uint32 i=0; i < FreedPageBlocksNum; ++i)
 			{
 				// is it possible (and worth i.e. <25% overhead) to use this block
@@ -726,7 +750,11 @@ private:
 					void* Ret=FreedPageBlocks[i].Ptr;
 					OutActualSize=FreedPageBlocks[i].ByteSize;
 					CachedTotal-=FreedPageBlocks[i].ByteSize;
-					FreedPageBlocks[i]=FreedPageBlocks[--FreedPageBlocksNum];
+					if (i < FreedPageBlocksNum - 1)
+					{
+						FMemory::Memmove(&FreedPageBlocks[i], &FreedPageBlocks[i + 1], sizeof(FFreePageBlock) * (FreedPageBlocksNum - i - 1));
+					}
+					FreedPageBlocksNum--;
 					return Ret;
 				}
 			};
diff --git a/Engine/Source/Runtime/Core/Public/Logging/TokenizedMessage.h b/Engine/Source/Runtime/Core/Public/Logging/TokenizedMessage.h
--- a/Engine/Source/Runtime/Core/Public/Logging/TokenizedMessage.h
+++ b/Engine/Source/Runtime/Core/Public/Logging/TokenizedMessage.h
@@ -476,23 +476,30 @@ class FActionToken
 public:
 
 	/** Factory method, tokens can only be constructed as shared refs */
-	CORE_API static TSharedRef<FActionToken> Create( const FText& ActionName, const FText& ActionDescription, const FOnActionTokenExecuted& Action )
+	CORE_API static TSharedRef<FActionToken> Create( const FText& ActionName, const FText& ActionDescription, const FOnActionTokenExecuted& Action, bool bInSingleUse=false )
 	{
-		return MakeShareable(new FActionToken(ActionName, ActionDescription, Action));
+		return MakeShareable(new FActionToken(ActionName, ActionDescription, Action, bInSingleUse));
 	}
 
 	/** Executes the assigned action delegate. */
-	void ExecuteAction( )
+	void ExecuteAction()
 	{
 		Action.ExecuteIfBound();
+		bActionExecuted = true;
 	}
 
 	/** Gets the action's description text. */
-	const FText& GetActionDescription( )
+	const FText& GetActionDescription()
 	{
 		return ActionDescription;
 	}
 
+	/** Returns true if the action can be activated */
+	bool CanExecuteAction() const
+	{
+		return Action.IsBound() && (!bSingleUse || !bActionExecuted);
+	}
+
 public:
 
 	// IMessageToken interface
@@ -504,9 +511,11 @@ public:
 protected:
 
 	/** Hidden constructor. */
-	FActionToken( const FText& InActionName, const FText& InActionDescription, const FOnActionTokenExecuted& InAction )
+	FActionToken(const FText& InActionName, const FText& InActionDescription, const FOnActionTokenExecuted& InAction, bool bInSingleUse)
 		: Action(InAction)
 		, ActionDescription(InActionDescription)
+		, bSingleUse(bInSingleUse)
+		, bActionExecuted(false)
 	{
 		CachedText = InActionName;
 	}
@@ -518,6 +527,12 @@ private:
 
 	/** The action's description text. */
 	const FText ActionDescription;
+	
+	/** If true, the action can only be performed once. */
+	bool bSingleUse;
+
+	/** If true, the action has been executed already. */
+	bool bActionExecuted;
 };
 
 
diff --git a/Engine/Source/Runtime/Core/Public/Misc/CoreStats.h b/Engine/Source/Runtime/Core/Public/Misc/CoreStats.h
--- a/Engine/Source/Runtime/Core/Public/Misc/CoreStats.h
+++ b/Engine/Source/Runtime/Core/Public/Misc/CoreStats.h
@@ -7,7 +7,7 @@
 
 CORE_API DECLARE_LOG_CATEGORY_EXTERN(LogSHA, Warning, All);
 CORE_API DECLARE_LOG_CATEGORY_EXTERN(LogStats, Log, All);
-CORE_API DECLARE_LOG_CATEGORY_EXTERN(LogStreaming, Warning, All);
+CORE_API DECLARE_LOG_CATEGORY_EXTERN(LogStreaming, Display, All);
 CORE_API DECLARE_LOG_CATEGORY_EXTERN(LogInit, Log, All);
 CORE_API DECLARE_LOG_CATEGORY_EXTERN(LogExit, Log, All);
 CORE_API DECLARE_LOG_CATEGORY_EXTERN(LogExec, Warning, All);
diff --git a/Engine/Source/Runtime/Core/Public/Windows/WindowsApplication.h b/Engine/Source/Runtime/Core/Public/Windows/WindowsApplication.h
--- a/Engine/Source/Runtime/Core/Public/Windows/WindowsApplication.h
+++ b/Engine/Source/Runtime/Core/Public/Windows/WindowsApplication.h
@@ -244,6 +244,7 @@ public:
 	virtual void PollGameDeviceState( const float TimeDelta ) override;
 	virtual void PumpMessages( const float TimeDelta ) override;
 	virtual void ProcessDeferredEvents( const float TimeDelta ) override;
+	virtual void Tick( const float TimeDelta ) override;
 	virtual TSharedRef< FGenericWindow > MakeWindow() override;
 	virtual void InitializeWindow( const TSharedRef< FGenericWindow >& Window, const TSharedRef< FGenericWindowDefinition >& InDefinition, const TSharedPtr< FGenericWindow >& InParent, const bool bShowImmediately ) override;
 	virtual void SetCapture( const TSharedPtr< FGenericWindow >& InWindow ) override;
@@ -368,6 +369,8 @@ private:
 	bool bHasLoadedInputPlugins;
 
 	TArray<int32> PressedModifierKeys;
+	
+	FModifierKeysState CachedModifierKeyState;
 
 	FAutoConsoleVariableRef CVarDeferMessageProcessing;
 
diff --git a/Engine/Source/Runtime/CoreUObject/Private/Misc/WorldCompositionUtility.cpp b/Engine/Source/Runtime/CoreUObject/Private/Misc/WorldCompositionUtility.cpp
--- a/Engine/Source/Runtime/CoreUObject/Private/Misc/WorldCompositionUtility.cpp
+++ b/Engine/Source/Runtime/CoreUObject/Private/Misc/WorldCompositionUtility.cpp
@@ -41,7 +41,7 @@ FArchive& operator<<( FArchive& Ar, FWorldTileInfo& D )
 	
 	if (Ar.UE4Ver() >= VER_UE4_WORLD_LEVEL_INFO_UPDATED)
 	{
-		Ar << D.Reserved0 << D.ParentTilePackageName;
+		Ar << D.bHideInTileView << D.ParentTilePackageName;
 	}
 
 	if (Ar.UE4Ver() >= VER_UE4_WORLD_LEVEL_INFO_LOD_LIST)
diff --git a/Engine/Source/Runtime/CoreUObject/Private/Serialization/AsyncLoading.cpp b/Engine/Source/Runtime/CoreUObject/Private/Serialization/AsyncLoading.cpp
--- a/Engine/Source/Runtime/CoreUObject/Private/Serialization/AsyncLoading.cpp
+++ b/Engine/Source/Runtime/CoreUObject/Private/Serialization/AsyncLoading.cpp
@@ -13,7 +13,7 @@
 
 DECLARE_MEMORY_STAT(TEXT("Streaming Memory Used"),STAT_StreamingAllocSize,STATGROUP_Memory);
 
-DECLARE_STATS_GROUP_VERBOSE(TEXT("Async Load"), STATGROUP_AsyncLoad, STATCAT_Advanced);
+DECLARE_STATS_GROUP(TEXT("Async Load"), STATGROUP_AsyncLoad, STATCAT_Advanced);
 
 DECLARE_CYCLE_STAT(TEXT("Tick AsyncPackage"),STAT_FAsyncPackage_Tick,STATGROUP_AsyncLoad);
 DECLARE_FLOAT_ACCUMULATOR_STAT(TEXT("Tick AsyncPackage Time"), STAT_FAsyncPackage_TickTime, STATGROUP_AsyncLoad);
diff --git a/Engine/Source/Runtime/CoreUObject/Private/UObject/GarbageCollection.cpp b/Engine/Source/Runtime/CoreUObject/Private/UObject/GarbageCollection.cpp
--- a/Engine/Source/Runtime/CoreUObject/Private/UObject/GarbageCollection.cpp
+++ b/Engine/Source/Runtime/CoreUObject/Private/UObject/GarbageCollection.cpp
@@ -516,6 +516,7 @@ public:
 					ChunkTasks.Add(TGraphTask<FGCTask>::CreateTask().ConstructAndDispatchWhenReady(this, &ObjectsToSerialize, StartIndex, NumPerChunk));
 					StartIndex += NumPerChunk;
 				}
+				QUICK_SCOPE_CYCLE_COUNTER(STAT_GC_Subtask_Wait);
 				FTaskGraphInterface::Get().WaitUntilTasksComplete(ChunkTasks, ENamedThreads::GameThread_Local);
 				GIsRunningParallelReachability = false;
 			}
diff --git a/Engine/Source/Runtime/CoreUObject/Private/UObject/Obj.cpp b/Engine/Source/Runtime/CoreUObject/Private/UObject/Obj.cpp
--- a/Engine/Source/Runtime/CoreUObject/Private/UObject/Obj.cpp
+++ b/Engine/Source/Runtime/CoreUObject/Private/UObject/Obj.cpp
@@ -2889,13 +2889,7 @@ bool StaticExec( UWorld* InWorld, const TCHAR* Cmd, FOutputDevice& Ar )
 	}
 	else if( FParse::Command(&Str,TEXT("OBJ")) )
 	{
-		if( FParse::Command(&Str,TEXT("GARBAGE")) || FParse::Command(&Str,TEXT("GC")) )
-		{
-			// Purge unclaimed objects.
-			CollectGarbage( GARBAGE_COLLECTION_KEEPFLAGS );
-			return true;
-		}
-		else if( FParse::Command(&Str,TEXT("CYCLES")) )
+		if( FParse::Command(&Str,TEXT("CYCLES")) )
 		{
 			// find all cycles in the reference graph
 
diff --git a/Engine/Source/Runtime/CoreUObject/Private/UObject/UObjectBaseUtility.cpp b/Engine/Source/Runtime/CoreUObject/Private/UObject/UObjectBaseUtility.cpp
--- a/Engine/Source/Runtime/CoreUObject/Private/UObject/UObjectBaseUtility.cpp
+++ b/Engine/Source/Runtime/CoreUObject/Private/UObject/UObjectBaseUtility.cpp
@@ -140,8 +140,8 @@ void UObjectBaseUtility::MarkPackageDirty() const
 			// It is against policy to dirty a map or package during load in the Editor, to enforce this policy
 			// we explicitly disable the ability to dirty a package or map during load.  Commandlets can still
 			// set the dirty state on load.
-			if( (!GIsEditor || IsRunningCommandlet() || 
-				(GIsEditor && !GIsEditorLoadingPackage && !GIsPlayInEditorWorld)))
+			if( IsRunningCommandlet() || 
+				(GIsEditor && !GIsEditorLoadingPackage && !GIsPlayInEditorWorld))
 			{
 				const bool bIsDirty = Package->IsDirty();
 
diff --git a/Engine/Source/Runtime/CoreUObject/Public/Misc/WorldCompositionUtility.h b/Engine/Source/Runtime/CoreUObject/Public/Misc/WorldCompositionUtility.h
--- a/Engine/Source/Runtime/CoreUObject/Public/Misc/WorldCompositionUtility.h
+++ b/Engine/Source/Runtime/CoreUObject/Public/Misc/WorldCompositionUtility.h
@@ -89,7 +89,7 @@ public:
 		: Position(0,0)
 		, AbsolutePosition(0,0)
 		, Bounds(0)
-		, Reserved0(false)
+		, bHideInTileView(false)
 		, ZOrder(0)
 	{
 	}
@@ -103,6 +103,7 @@ public:
 				(Bounds.Max.Equals(OtherInfo.Bounds.Max, 0.5f)) &&
 				(ParentTilePackageName == OtherInfo.ParentTilePackageName) &&
 				(Layer == OtherInfo.Layer) &&
+				(bHideInTileView == OtherInfo.bHideInTileView) &&
 				(ZOrder == OtherInfo.ZOrder) &&
 				(LODList == OtherInfo.LODList);
 	}
@@ -130,8 +131,8 @@ public:
 	FBox				Bounds;
 	/** Tile assigned layer  */
 	FWorldTileLayer		Layer;
-	/** */
-	bool				Reserved0;
+	/** Whether to hide sub-level tile in tile view*/
+	bool				bHideInTileView;
 	/** Parent tile package name */
 	FString				ParentTilePackageName;
 	/** LOD information */
diff --git a/Engine/Source/Runtime/CoreUObject/Public/Serialization/AsyncLoading.h b/Engine/Source/Runtime/CoreUObject/Public/Serialization/AsyncLoading.h
--- a/Engine/Source/Runtime/CoreUObject/Public/Serialization/AsyncLoading.h
+++ b/Engine/Source/Runtime/CoreUObject/Public/Serialization/AsyncLoading.h
@@ -6,6 +6,7 @@
 
 #pragma once
 
+#define PERF_TRACK_DETAILED_ASYNC_STATS (0)
 /**
  * Structure containing intermediate data required for async loading of all imports and exports of a
  * ULinkerLoad.
@@ -203,11 +204,6 @@ public:
 	/** Total time spent in FinishObjects.																*/
 	double						FinishObjectsTime;
 
-	/** Map of each class of object loaded to the total time spent calling CreateExport on those objects */
-	TMap<const UClass*,FMapTimeEntry>		CreateExportTimePerClass;
-	/** Map of each class of object loaded to the total time spent calling PostLoad on those objects */
-	TMap<const UClass*,FMapTimeEntry>		PostLoadTimePerClass;
-
 #endif
 
 private:
diff --git a/Engine/Source/Runtime/Engine/Classes/AI/Navigation/NavCollision.h b/Engine/Source/Runtime/Engine/Classes/AI/Navigation/NavCollision.h
--- a/Engine/Source/Runtime/Engine/Classes/AI/Navigation/NavCollision.h
+++ b/Engine/Source/Runtime/Engine/Classes/AI/Navigation/NavCollision.h
@@ -103,6 +103,8 @@ class ENGINE_API UNavCollision : public UObject
 protected:
 	void ClearCollision();
 
+	FORCEINLINE bool ShouldUseConvexCollision() const { return bGatherConvexGeometry || (CylinderCollision.Num() == 0 && BoxCollision.Num() == 0);  }
+
 #if WITH_EDITOR
 	void InvalidatePhysicsData();
 #endif // WITH_EDITOR
diff --git a/Engine/Source/Runtime/Engine/Classes/AI/Navigation/NavMeshRenderingComponent.h b/Engine/Source/Runtime/Engine/Classes/AI/Navigation/NavMeshRenderingComponent.h
--- a/Engine/Source/Runtime/Engine/Classes/AI/Navigation/NavMeshRenderingComponent.h
+++ b/Engine/Source/Runtime/Engine/Classes/AI/Navigation/NavMeshRenderingComponent.h
@@ -32,10 +32,12 @@ public:
 
 	void GatherData(FNavMeshSceneProxyData& DebugDrawData) const;
 
+	static bool IsNavigationShowFlagSet(const UWorld* World);
+
 protected:
 	void TimerFunction();
 
 protected:
-	int32 bCollectNavigationData : 1;
+	uint32 bCollectNavigationData : 1;
 	FTimerHandle TimerHandle;
 };
diff --git a/Engine/Source/Runtime/Engine/Classes/AI/Navigation/NavRelevantInterface.h b/Engine/Source/Runtime/Engine/Classes/AI/Navigation/NavRelevantInterface.h
--- a/Engine/Source/Runtime/Engine/Classes/AI/Navigation/NavRelevantInterface.h
+++ b/Engine/Source/Runtime/Engine/Classes/AI/Navigation/NavRelevantInterface.h
@@ -2,9 +2,11 @@
 
 #pragma once
 #include "Interface.h"
+#include "AI/Navigation/NavigationTypes.h"
 #include "NavRelevantInterface.generated.h"
 
 struct FNavigationRelevantData;
+struct FNavigableGeometryExport;
 
 UINTERFACE(MinimalAPI, meta=(CannotImplementInterfaceInBlueprint))
 class UNavRelevantInterface : public UInterface
@@ -22,6 +24,18 @@ class INavRelevantInterface
 	/** Get bounds for navigation octree */
 	virtual FBox GetNavigationBounds() const { return FBox(0); }
 
+	/** if this instance knows how to export sub-sections of self */
+	virtual bool SupportsGatheringGeometrySlices() const { return false; }
+
+	/** This function is called "on demand", whenever that specified piece of geometry is needed for navigation generation */
+	virtual void GatherGeometrySlice(FNavigableGeometryExport& GeomExport, const FBox& SliceBox) const {}
+
+	virtual ENavDataGatheringMode GetGeometryGatheringMode() const { return ENavDataGatheringMode::Default; }
+
+	/** Called on Game-thread to give implementer a change to perform actions that require game-thread to run, 
+	 *	for example precaching physics data */
+	virtual void PrepareGeometryExportSync() {}
+
 	/** Update bounds, called after moving owning actor */
 	virtual void UpdateNavigationBounds() {}
 
diff --git a/Engine/Source/Runtime/Engine/Classes/AI/Navigation/NavigationData.h b/Engine/Source/Runtime/Engine/Classes/AI/Navigation/NavigationData.h
--- a/Engine/Source/Runtime/Engine/Classes/AI/Navigation/NavigationData.h
+++ b/Engine/Source/Runtime/Engine/Classes/AI/Navigation/NavigationData.h
@@ -441,7 +441,7 @@ protected:
 	UPROPERTY(EditAnywhere, Category = Runtime, config)
 	float ObservedPathsTickInterval;
 
-	UPROPERTY()
+	UPROPERTY(Transient)
 	UWorld* CachedWorld;
 
 public:
@@ -797,7 +797,7 @@ protected:
 	FNavRaycastPtr RaycastImplementation; 
 
 protected:
-	TUniquePtr<FNavDataGenerator> NavDataGenerator;
+	TSharedPtr<FNavDataGenerator> NavDataGenerator;
 	/** 
 	 *	Container for all path objects generated with this Navigation Data instance. 
 	 *	Is meant to be added to only on GameThread, and in fact should user should never 
diff --git a/Engine/Source/Runtime/Engine/Classes/AI/Navigation/NavigationSystem.h b/Engine/Source/Runtime/Engine/Classes/AI/Navigation/NavigationSystem.h
--- a/Engine/Source/Runtime/Engine/Classes/AI/Navigation/NavigationSystem.h
+++ b/Engine/Source/Runtime/Engine/Classes/AI/Navigation/NavigationSystem.h
@@ -176,6 +176,10 @@ public:
 	uint32 bSkipAgentHeightCheckWhenPickingNavData:1;
 
 protected:
+	
+	UPROPERTY(EditDefaultsOnly, Category = "NavigationSystem", config)
+	ENavDataGatheringModeConfig DataGatheringMode;
+
 	/** If set to true navigation will be generated only around registered "navigation enforcers"
 	*	This has a range of consequences (including how navigation octree operates) so it needs to
 	*	be a conscious decision.
@@ -184,7 +188,7 @@ protected:
 	*/
 	UPROPERTY(EditDefaultsOnly, Category = "Navigation Enforcing", config)
 	uint32 bGenerateNavigationOnlyAroundNavigationInvokers : 1;
-
+	
 	/** Minimal time, in seconds, between active tiles set update */
 	UPROPERTY(EditAnywhere, Category = "Navigation Enforcing", meta = (ClampMin = "0.1", UIMin = "0.1", EditCondition = "bGenerateNavigationOnlyAroundNavigationInvokers"), config)
 	float ActiveTilesUpdateInterval;
@@ -235,7 +239,7 @@ public:
 
 	UFUNCTION(BlueprintPure, Category = "AI|Navigation", meta = (WorldContext = "WorldContext"))
 	static UNavigationSystem* GetNavigationSystem(UObject* WorldContext);
-
+	
 	/** Project a point onto the NavigationData */
 	UFUNCTION(BlueprintPure, Category="AI|Navigation", meta=(WorldContext="WorldContext" ) )
 	static FVector ProjectPointToNavigation(UObject* WorldContext, const FVector& Point, ANavigationData* NavData = NULL, TSubclassOf<UNavigationQueryFilter> FilterClass = NULL);
@@ -281,19 +285,31 @@ public:
 	UFUNCTION(BlueprintCallable, Category="AI|Navigation", meta=(WorldContext="WorldContext" ))
 	static bool NavigationRaycast(UObject* WorldContext, const FVector& RayStart, const FVector& RayEnd, FVector& HitLocation, TSubclassOf<UNavigationQueryFilter> FilterClass = NULL, AController* Querier = NULL);
 
+	/** will limit the number of simultaneously running navmesh tile generation jobs to specified number.
+	 *	@param MaxNumberOfJobs gets trimmed to be at least 1. You cannot use this function to pause navmesh generation */
+	UFUNCTION(BlueprintCallable, Category = "AI|Navigation")
+	void SetMaxSimultaneousTileGenerationJobsCount(int32 MaxNumberOfJobs);
+	
+	/** Brings limit of simultaneous navmesh tile generation jobs back to Project Setting's default value */
+	UFUNCTION(BlueprintCallable, Category = "AI|Navigation")
+	void ResetMaxSimultaneousTileGenerationJobsCount();
+
 	/** Registers given actor as a "navigation enforcer" which means navigation system will
 	 *	make sure navigation is being generated in specified radius around it.
 	 *	@note: you need NavigationSystem's GenerateNavigationOnlyAroundNavigationInvokers to be set to true
 	 *		to take advantage of this feature
 	 */
-	UFUNCTION(BlueprintCallable, Category = Navigation)
+	UFUNCTION(BlueprintCallable, Category = "AI|Navigation")
 	void RegisterNavigationInvoker(AActor* Invoker, float TileGenerationRadius = 3000, float TileRemovalRadius = 5000);
 
 	/** Removes given actor from the list of active navigation enforcers.
 	 *	@see RegisterNavigationInvoker for more details */
-	UFUNCTION(BlueprintCallable, Category = Navigation)
+	UFUNCTION(BlueprintCallable, Category = "AI|Navigation")
 	void UnregisterNavigationInvoker(AActor* Invoker);
 
+	UFUNCTION(BlueprintCallable, Category = "AI|Navigation|Generation")
+	void SetGeometryGatheringMode(ENavDataGatheringModeConfig NewMode);
+
 	FORCEINLINE bool IsActiveTilesGenerationEnabled() const{ return bGenerateNavigationOnlyAroundNavigationInvokers; }
 	
 	/** delegate type for events that dirty the navigation data ( Params: const FBox& DirtyBounds ) */
@@ -554,7 +570,8 @@ public:
 	void AddDirtyAreas(const TArray<FBox>& NewAreas, int32 Flags);
 	bool HasDirtyAreasQueued() const;
 
-	const FNavigationOctree* GetNavOctree() const { return NavOctree; }
+	const FNavigationOctree* GetNavOctree() const { return NavOctree.Get(); }
+	FNavigationOctree* GetMutableNavOctree() { return NavOctree.Get(); }
 
 	FORCEINLINE void SetObjectsNavOctreeId(UObject* Object, FOctreeElementId Id) { ObjectToOctreeId.Add(Object, Id); }
 	FORCEINLINE const FOctreeElementId* GetObjectsNavOctreeId(const UObject* Object) const { return ObjectToOctreeId.Find(Object); }
@@ -693,6 +710,8 @@ public:
 	virtual void OnEditorModeChanged(FEdMode* Mode, bool IsEntering);
 #endif // WITH_EDITOR
 
+	FORCEINLINE bool IsSetUpForLazyGeometryExporting() const { return bGenerateNavigationOnlyAroundNavigationInvokers; }
+
 	/** register actor important for generation (navigation data will be build around them first) */
 	void RegisterGenerationSeed(AActor* SeedActor);
 	void UnregisterGenerationSeed(AActor* SeedActor);
@@ -732,7 +751,7 @@ protected:
 
 	FNavigationSystem::EMode OperationMode;
 
-	FNavigationOctree* NavOctree;
+	TSharedPtr<FNavigationOctree, ESPMode::ThreadSafe> NavOctree;
 
 	TArray<FAsyncPathFindingQuery> AsyncPathFindingQueries;
 
diff --git a/Engine/Source/Runtime/Engine/Classes/AI/Navigation/NavigationTypes.h b/Engine/Source/Runtime/Engine/Classes/AI/Navigation/NavigationTypes.h
--- a/Engine/Source/Runtime/Engine/Classes/AI/Navigation/NavigationTypes.h
+++ b/Engine/Source/Runtime/Engine/Classes/AI/Navigation/NavigationTypes.h
@@ -1,6 +1,8 @@
 // Copyright 1998-2015 Epic Games, Inc. All Rights Reserved.
 
 #pragma once
+
+#include "Containers/BitArray.h"
 #include "NavigationTypes.generated.h"
 
 #define INVALID_NAVNODEREF (0)
@@ -169,6 +171,22 @@ struct FNavigationDirtyElement
 	}
 };
 
+UENUM()
+enum class ENavDataGatheringMode : uint8
+{
+	Default,
+	Instant,
+	Lazy
+};
+
+UENUM()
+enum class ENavDataGatheringModeConfig : uint8
+{
+	Invalid UMETA(Hidden),
+	Instant,
+	Lazy
+};
+
 //
 // Used to gather per instance transforms in a specific area
 //
@@ -665,4 +683,17 @@ struct FNavigationInvoker
 
 	FNavigationInvoker();
 	FNavigationInvoker(AActor& InActor, float InGenerationRadius, float InRemovalRadius);
+};
+
+//----------------------------------------------------------------------//
+// generic "landscape" support
+//----------------------------------------------------------------------//
+struct ENGINE_API FNavHeightfieldSamples
+{
+	TNavStatArray<int16> Heights;
+	TBitArray<> Holes;
+
+	FNavHeightfieldSamples();
+
+	FORCEINLINE bool IsEmpty() const { return Heights.Num() == 0; }
 };
\ No newline at end of file
diff --git a/Engine/Source/Runtime/Engine/Classes/AI/Navigation/RecastNavMesh.h b/Engine/Source/Runtime/Engine/Classes/AI/Navigation/RecastNavMesh.h
--- a/Engine/Source/Runtime/Engine/Classes/AI/Navigation/RecastNavMesh.h
+++ b/Engine/Source/Runtime/Engine/Classes/AI/Navigation/RecastNavMesh.h
@@ -402,7 +402,7 @@ class ENGINE_API ARecastNavMesh : public ANavigationData
 	uint32 bDrawTriangleEdges:1;
 
 	/** should we draw edges of every poly (i.e. not only border-edges)  */
-	UPROPERTY(EditAnywhere, Category=Display)
+	UPROPERTY(EditAnywhere, Category=Display, config)
 	uint32 bDrawPolyEdges:1;
 
 	/** should we draw border-edges */
@@ -506,6 +506,9 @@ class ENGINE_API ARecastNavMesh : public ANavigationData
 	UPROPERTY(EditAnywhere, Category = Generation, config, meta = (ClampMin = "0.0"))
 	float MaxSimplificationError;
 
+	UPROPERTY(EditAnywhere, Category = Generation, config, meta = (ClampMin = "0", UIMin = "0"), AdvancedDisplay)
+	int32 MaxSimultaneousTileGenerationJobsCount;
+	
 	/** navmesh draw distance in game (always visible in editor) */
 	UPROPERTY(config)
 	float DefaultDrawDistance;
@@ -552,6 +555,9 @@ class ENGINE_API ARecastNavMesh : public ANavigationData
 	UPROPERTY(EditAnywhere, Category = Generation, config, AdvancedDisplay)
 	uint32 bMarkLowHeightAreas : 1;
 
+	UPROPERTY(EditAnywhere, Category = Generation, config, AdvancedDisplay)
+	uint32 bDoFullyAsyncNavDataGathering : 1;
+	
 	/** TODO: switch to disable new code from OffsetFromCorners if necessary - remove it later */
 	UPROPERTY(config)
 	uint32 bUseBetterOffsetsFromCorners : 1;
@@ -768,6 +774,9 @@ public:
 	/** Area sort function */
 	virtual void SortAreasForGenerator(TArray<FRecastAreaNavModifierElement>& Areas) const;
 
+	int32 GetMaxSimultaneousTileGenerationJobsCount() const { return MaxSimultaneousTileGenerationJobsCount; }
+	void SetMaxSimultaneousTileGenerationJobsCount(int32 NewJobsCountLimit);
+
 	//----------------------------------------------------------------------//
 	// Custom navigation links
 	//----------------------------------------------------------------------//
@@ -881,6 +890,7 @@ public:
 	virtual bool SupportsRuntimeGeneration() const override;
 	virtual bool SupportsStreaming() const override;
 	virtual void ConditionalConstructGenerator() override;
+	bool ShouldGatherDataOnGameThread() const { return bDoFullyAsyncNavDataGathering == false; }
 
 	void UpdateActiveTiles(const TArray<FNavigationInvokerRaw>& InvokerLocations);
 	void RemoveTiles(const TArray<FIntPoint>& Tiles);
diff --git a/Engine/Source/Runtime/Engine/Classes/Camera/CameraComponent.h b/Engine/Source/Runtime/Engine/Classes/Camera/CameraComponent.h
--- a/Engine/Source/Runtime/Engine/Classes/Camera/CameraComponent.h
+++ b/Engine/Source/Runtime/Engine/Classes/Camera/CameraComponent.h
@@ -38,6 +38,10 @@ class UCameraComponent : public USceneComponent
 	UPROPERTY(Interp, EditAnywhere, BlueprintReadWrite, Category=CameraSettings)
 	uint32 bConstrainAspectRatio:1;
 
+	// If true, account for the field of view angle when computing which level of detail to use for meshes.
+	UPROPERTY(Interp, EditAnywhere, AdvancedDisplay, BlueprintReadWrite, Category=CameraSettings)
+	uint32 bUseFieldOfViewForLOD:1;
+
 	/**
 	 * If this camera component is placed on a pawn, should it use the view/control rotation of the pawn where possible?
 	 * @see APawn::GetViewRotation()
diff --git a/Engine/Source/Runtime/Engine/Classes/Camera/CameraTypes.h b/Engine/Source/Runtime/Engine/Classes/Camera/CameraTypes.h
--- a/Engine/Source/Runtime/Engine/Classes/Camera/CameraTypes.h
+++ b/Engine/Source/Runtime/Engine/Classes/Camera/CameraTypes.h
@@ -68,6 +68,10 @@ struct FMinimalViewInfo
 	UPROPERTY(EditAnywhere, BlueprintReadWrite, Category=Camera)
 	uint32 bConstrainAspectRatio:1; 
 
+	// If true, account for the field of view angle when computing which level of detail to use for meshes.
+	UPROPERTY(EditAnywhere, AdvancedDisplay, BlueprintReadWrite, Category=CameraSettings)
+	uint32 bUseFieldOfViewForLOD:1;
+
 	// The type of camera
 	UPROPERTY(EditAnywhere, BlueprintReadWrite, Category=Camera)
 	TEnumAsByte<ECameraProjectionMode::Type> ProjectionMode;
@@ -89,6 +93,7 @@ struct FMinimalViewInfo
 		, OrthoFarClipPlane(WORLD_MAX)
 		, AspectRatio(1.33333333f)
 		, bConstrainAspectRatio(false)
+		, bUseFieldOfViewForLOD(true)
 		, ProjectionMode(ECameraProjectionMode::Perspective)
 		, PostProcessBlendWeight(0.0f)
 	{
diff --git a/Engine/Source/Runtime/Engine/Classes/Camera/PlayerCameraManager.h b/Engine/Source/Runtime/Engine/Classes/Camera/PlayerCameraManager.h
--- a/Engine/Source/Runtime/Engine/Classes/Camera/PlayerCameraManager.h
+++ b/Engine/Source/Runtime/Engine/Classes/Camera/PlayerCameraManager.h
@@ -206,7 +206,7 @@ public:
 
 	/** Color to fade to (when bEnableFading == true). */
 	UPROPERTY()
-	FColor FadeColor;
+	FLinearColor FadeColor;
 
 	/** Amount of fading to apply (when bEnableFading == true). */
 	UPROPERTY()
@@ -268,7 +268,7 @@ public:
 	UPROPERTY(EditAnywhere, BlueprintReadWrite, Category=Debug)
 	FVector FreeCamOffset;
 
-	/** Current camera fade alpha value (when bEnableFading == true) */
+	/** Current camera fade alpha range, where X = starting alpha and Y = final alpha (when bEnableFading == true) */
 	UPROPERTY()
 	FVector2D FadeAlpha;
 
@@ -358,6 +358,15 @@ public:
 	UPROPERTY()
 	uint32 bFadeAudio:1;
 
+protected:
+	/** True if camera fade hold at it's final value when finished. */
+	UPROPERTY()
+	uint32 bHoldFadeWhenFinished: 1;
+
+	UPROPERTY()
+	uint32 bAutoAnimateFade : 1;
+
+public:
 	/** True to turn on scaling of color channels in final image using ColorScale property. */
 	UPROPERTY()
 	uint32 bEnableColorScaling:1;
@@ -494,6 +503,7 @@ public:
 	virtual void InitializeFor(class APlayerController* PC);
 	
 	/** @return Returns the camera's current full FOV angle, in degrees. */
+	UFUNCTION(BlueprintCallable, Category = "Camera")
 	virtual float GetFOVAngle() const;
 	
 	/** 
@@ -551,6 +561,7 @@ protected:
 	
 	/** Internal. Applies appropriate audio fading to the audio system. */
 	virtual void ApplyAudioFade();
+	virtual void StopAudioFade();
 	
 	/**
 	 * Internal helper to blend two viewtargets.
@@ -638,6 +649,36 @@ public:
 	/** Stops playing all camera shakes. */
 	virtual void ClearAllCameraShakes();
 
+	//
+	//  CameraAnim fades.
+	//
+
+	/** 
+	 * Does a camera fade to/from a solid color.  Animates automatically.
+	 * @param FromAlpha - Alpha at which to begin the fade. Range [0..1], where 0 is fully transparent and 1 is fully opaque solid color.
+	 * @param ToAlpha - Alpha at which to finish the fade.
+	 * @param Duration - How long the fade should take, in seconds.
+	 * @param Color - Color to fade to/from.
+	 * @param bShouldFadeAudio - True to fade audio volume along with the alpha of the solid color.
+	 * @param bHoldWhenFinished - True for fade to hold at the ToAlpha until explicitly stopped (e.g. with StopCameraFade)
+	 */
+	UFUNCTION(BlueprintCallable, Category = "Camera")
+	virtual void StartCameraFade(float FromAlpha, float ToAlpha, float Duration, FLinearColor Color, bool bShouldFadeAudio = false, bool bHoldWhenFinished = false);
+
+	/** 
+	 * Stops camera fading.
+	 */
+	UFUNCTION(BlueprintCallable, Category = "Camera")
+	virtual void StopCameraFade();
+
+	/** 
+	 * Turns on camera fading at the given opacity. Does not auto-animate, allowing user to animate themselves.
+	 * Call StopCameraFade to turn fading back off.
+	*/
+	UFUNCTION(BlueprintCallable, Category = "Camera")
+	virtual void SetManualCameraFade(float InFadeAmount, FLinearColor Color, bool bInFadeAudio);
+
+
 	//
 	//  CameraAnim support.
 	//
diff --git a/Engine/Source/Runtime/Engine/Classes/Components/DestructibleComponent.h b/Engine/Source/Runtime/Engine/Classes/Components/DestructibleComponent.h
--- a/Engine/Source/Runtime/Engine/Classes/Components/DestructibleComponent.h
+++ b/Engine/Source/Runtime/Engine/Classes/Components/DestructibleComponent.h
@@ -202,7 +202,7 @@ public:
 
 	// End DestructibleComponent interface.
 
-	virtual bool DoCustomNavigableGeometryExport(struct FNavigableGeometryExport* GeomExport) const override;
+	virtual bool DoCustomNavigableGeometryExport(FNavigableGeometryExport& GeomExport) const override;
 
 	FORCEINLINE static int32 ChunkIdxToBoneIdx(int32 ChunkIdx) { return ChunkIdx + 1; }
 	FORCEINLINE static int32 BoneIdxToChunkIdx(int32 BoneIdx) { return FMath::Max(BoneIdx - 1, 0); }
diff --git a/Engine/Source/Runtime/Engine/Classes/Components/HierarchicalInstancedStaticMeshComponent.h b/Engine/Source/Runtime/Engine/Classes/Components/HierarchicalInstancedStaticMeshComponent.h
--- a/Engine/Source/Runtime/Engine/Classes/Components/HierarchicalInstancedStaticMeshComponent.h
+++ b/Engine/Source/Runtime/Engine/Classes/Components/HierarchicalInstancedStaticMeshComponent.h
@@ -51,6 +51,27 @@ struct FClusterNode
 	}
 };
 
+class FAsyncBuildInstanceBuffer : public FNonAbandonableTask
+{
+public:
+	class UHierarchicalInstancedStaticMeshComponent* Component;
+	class UWorld* World;
+
+	FAsyncBuildInstanceBuffer(class UHierarchicalInstancedStaticMeshComponent* InComponent, class UWorld* InWorld)
+		: Component(InComponent)
+		, World(InWorld)
+	{
+	}
+	void DoWork();
+	FORCEINLINE TStatId GetStatId() const
+	{
+		RETURN_QUICK_DECLARE_CYCLE_STAT(FAsyncBuildInstanceBuffer, STATGROUP_ThreadPoolAsyncTasks);
+	}
+	static const TCHAR *Name()
+	{
+		return TEXT("FAsyncBuildInstanceBuffer");
+	}
+};
 
 UCLASS(ClassGroup=Rendering, meta=(BlueprintSpawnableComponent))
 class ENGINE_API UHierarchicalInstancedStaticMeshComponent : public UInstancedStaticMeshComponent
@@ -61,6 +82,10 @@ class ENGINE_API UHierarchicalInstancedStaticMeshComponent : public UInstancedSt
 
 	TSharedPtr<TArray<FClusterNode>, ESPMode::ThreadSafe> ClusterTreePtr;
 
+	// Temp hack, long term we will load data in the right format directly
+	FAsyncTask<FAsyncBuildInstanceBuffer>* AsyncBuildInstanceBufferTask;
+
+	// Prebuilt instance buffer, used mostly for grass. Long term instances will be stored directly in the correct format.
 	FStaticMeshInstanceData WriteOncePrebuiltInstanceBuffer;
 	
 	// Table for remaping instances from cluster tree to PerInstanceSMData order
@@ -89,10 +114,9 @@ public:
 	//Begin UObject Interface
 	virtual void Serialize(FArchive& Ar) override;
 	virtual SIZE_T GetResourceSize(EResourceSizeMode::Type Mode) override;
-
 	virtual FPrimitiveSceneProxy* CreateSceneProxy() override;
-	virtual FBoxSphereBounds CalcBounds(const FTransform& BoundTransform) const override;
 	virtual void PostLoad() override;
+	virtual FBoxSphereBounds CalcBounds(const FTransform& BoundTransform) const override;
 
 	// UInstancedStaticMesh interface
 	virtual int32 AddInstance(const FTransform& InstanceTransform) override;
@@ -125,7 +149,19 @@ public:
 	bool IsAsyncBuilding() const { return bIsAsyncBuilding; }
 	bool IsTreeFullyBuilt() const { return NumBuiltInstances == PerInstanceSMData.Num() && RemovedInstances.Num() == 0; }
 
+	void FlushAsyncBuildInstanceBufferTask();
+
+	/** Heuristic for the number of leaves in the tree **/
+	int32 DesiredInstancesPerLeaf();
+
 protected:
+	/** Gets and approximate number of verts for each LOD to generate heuristics **/
+	int32 GetVertsForLOD(int32 LODIndex);
+	/** Average number of instances per leaf **/
+	float ActualInstancesPerLeaf();
+	/** For testing, prints some stats after any kind of build **/
+	void PostBuildStats();
+
 	virtual void GetNavigationPerInstanceTransforms(const FBox& AreaBox, TArray<FTransform>& InstanceData) const override;
 	virtual void PartialNavigationUpdate(int32 InstanceIdx) override;
 	void FlushAccumulatedNavigationUpdates();
diff --git a/Engine/Source/Runtime/Engine/Classes/Components/InstancedStaticMeshComponent.h b/Engine/Source/Runtime/Engine/Classes/Components/InstancedStaticMeshComponent.h
--- a/Engine/Source/Runtime/Engine/Classes/Components/InstancedStaticMeshComponent.h
+++ b/Engine/Source/Runtime/Engine/Classes/Components/InstancedStaticMeshComponent.h
@@ -167,7 +167,7 @@ public:
 #endif
 	virtual void GetLightAndShadowMapMemoryUsage( int32& LightMapMemoryUsage, int32& ShadowMapMemoryUsage ) const override;
 	
-	virtual bool DoCustomNavigableGeometryExport(struct FNavigableGeometryExport* GeomExport) const override;
+	virtual bool DoCustomNavigableGeometryExport(FNavigableGeometryExport& GeomExport) const override;
 	// End UPrimitiveComponent Interface
 	
 	// Begin UNavRelevantInterface interface.
diff --git a/Engine/Source/Runtime/Engine/Classes/Components/PrimitiveComponent.h b/Engine/Source/Runtime/Engine/Classes/Components/PrimitiveComponent.h
--- a/Engine/Source/Runtime/Engine/Classes/Components/PrimitiveComponent.h
+++ b/Engine/Source/Runtime/Engine/Classes/Components/PrimitiveComponent.h
@@ -16,6 +16,7 @@ class AController;
 class UTexture;
 struct FEngineShowFlags;
 struct FConvexVolume;
+struct FNavigableGeometryExport;
 
 /** Information about a streaming texture that a primitive uses for rendering. */
 struct FStreamingTexturePrimitiveInfo
@@ -1655,8 +1656,8 @@ public:
 	}
 
 	// Begin INavRelevantInterface Interface
-	virtual FBox GetNavigationBounds() const;
-	virtual bool IsNavigationRelevant() const;
+	virtual FBox GetNavigationBounds() const override;
+	virtual bool IsNavigationRelevant() const override;
 	// End INavRelevantInterface Interface
 
 	FORCEINLINE EHasCustomNavigableGeometry::Type HasCustomNavigableGeometry() const { return bHasCustomNavigableGeometry; }
@@ -1665,7 +1666,11 @@ public:
 
 	/** Collects custom navigable geometry of component.
 	 *	@return true if regular navigable geometry exporting should be run as well */
-	virtual bool DoCustomNavigableGeometryExport(struct FNavigableGeometryExport* GeomExport) const { return true; }
+	DEPRECATED(4.8, "UPrimitiveComponent::DoCustomNavigableGeometryExport(FNavigableGeometryExport* GeomExport) is deprecated, use UPrimitiveComponent::DoCustomNavigableGeometryExport(FNavigableGeometryExport& GeomExport) instead (takes ref instead of a pointer)")
+	virtual bool DoCustomNavigableGeometryExport(FNavigableGeometryExport* GeomExport) const { return DoCustomNavigableGeometryExport(*GeomExport); }
+	/** Collects custom navigable geometry of component.
+	*	@return true if regular navigable geometry exporting should be run as well */
+	virtual bool DoCustomNavigableGeometryExport(FNavigableGeometryExport& GeomExport) const { return true; }
 
 	static void DispatchMouseOverEvents(UPrimitiveComponent* CurrentComponent, UPrimitiveComponent* NewComponent);
 	static void DispatchTouchOverEvents(ETouchIndex::Type FingerIndex, UPrimitiveComponent* CurrentComponent, UPrimitiveComponent* NewComponent);
diff --git a/Engine/Source/Runtime/Engine/Classes/Components/SceneComponent.h b/Engine/Source/Runtime/Engine/Classes/Components/SceneComponent.h
--- a/Engine/Source/Runtime/Engine/Classes/Components/SceneComponent.h
+++ b/Engine/Source/Runtime/Engine/Classes/Components/SceneComponent.h
@@ -417,7 +417,7 @@ public:
 	 *   @param bMaintainWorldTransform	If true, update the relative location/rotation of this component to keep its world position the same *	
 	 */
 	UFUNCTION(BlueprintCallable, Category="Utilities|Transformation")
-	virtual void DetachFromParent(bool bMaintainWorldPosition = false);
+	virtual void DetachFromParent(bool bMaintainWorldPosition = false, bool bCallModify = true);
 
 
 	/** 
diff --git a/Engine/Source/Runtime/Engine/Classes/Components/SkeletalMeshComponent.h b/Engine/Source/Runtime/Engine/Classes/Components/SkeletalMeshComponent.h
--- a/Engine/Source/Runtime/Engine/Classes/Components/SkeletalMeshComponent.h
+++ b/Engine/Source/Runtime/Engine/Classes/Components/SkeletalMeshComponent.h
@@ -566,9 +566,11 @@ public:
 
 	/**
 	 * Set Morph Target with Name and Value(0-1)
+	 *
+	 * @param bRemoveZeroWeight : Used by editor code when it should stay in the active list with zero weight
 	 */
 	UFUNCTION(BlueprintCallable, Category="Components|SkeletalMesh")
-	void SetMorphTarget(FName MorphTargetName, float Value);
+	void SetMorphTarget(FName MorphTargetName, float Value, bool bRemoveZeroWeight=true);
 
 	/**
 	 * Clear all Morph Target that are set to this mesh
@@ -638,6 +640,9 @@ public:
 	/** Temporary array of bone indices required this frame. Filled in by UpdateSkelPose. */
 	TArray<FBoneIndexType> RequiredBones;
 
+	/** Tempory array of bone indices required to populate space bases */
+	TArray<FBoneIndexType> FillSpaceBasesRequiredBones;
+
 	/** 
 	 *	Index of the 'Root Body', or top body in the asset hierarchy. 
 	 *	Filled in by InitInstance, so we don't need to save it.
diff --git a/Engine/Source/Runtime/Engine/Classes/Components/SkinnedMeshComponent.h b/Engine/Source/Runtime/Engine/Classes/Components/SkinnedMeshComponent.h
--- a/Engine/Source/Runtime/Engine/Classes/Components/SkinnedMeshComponent.h
+++ b/Engine/Source/Runtime/Engine/Classes/Components/SkinnedMeshComponent.h
@@ -283,6 +283,10 @@ public:
 	UPROPERTY()
 	uint32 bDisplayBones:1;
 
+	/** Disable Morphtarget for this component. */
+	UPROPERTY()
+	uint32 bDisableMorphTarget:1;
+
 	/** Don't bother rendering the skin. */
 	UPROPERTY()
 	uint32 bHideSkin:1;
diff --git a/Engine/Source/Runtime/Engine/Classes/Components/StaticMeshComponent.h b/Engine/Source/Runtime/Engine/Classes/Components/StaticMeshComponent.h
--- a/Engine/Source/Runtime/Engine/Classes/Components/StaticMeshComponent.h
+++ b/Engine/Source/Runtime/Engine/Classes/Components/StaticMeshComponent.h
@@ -300,16 +300,17 @@ public:
 	virtual void GetUsedMaterials(TArray<UMaterialInterface*>& OutMaterials) const override;
 	virtual UMaterialInterface* GetMaterial(int32 MaterialIndex) const override;
 
-	virtual bool DoCustomNavigableGeometryExport(struct FNavigableGeometryExport* GeomExport) const;
+	virtual bool DoCustomNavigableGeometryExport(FNavigableGeometryExport& GeomExport) const override;
 #if WITH_EDITOR
 	virtual bool ComponentIsTouchingSelectionBox(const FBox& InSelBBox, const FEngineShowFlags& ShowFlags, const bool bConsiderOnlyBSP, const bool bMustEncompassEntireComponent) const override;
 	virtual bool ComponentIsTouchingSelectionFrustum(const FConvexVolume& InSelBBox, const FEngineShowFlags& ShowFlags, const bool bConsiderOnlyBSP, const bool bMustEncompassEntireComponent) const override;
 #endif
 	// End UPrimitiveComponent interface.
 
-	// Begin UNavRelevantInterface interface.
+	// Begin INavRelevantInterface interface.
+	virtual bool IsNavigationRelevant() const override;
 	virtual void GetNavigationData(FNavigationRelevantData& Data) const override;
-	// End UPrimitiveComponent interface.
+	// End INavRelevantInterface interface.
 	
 	// Begin UMeshComponent interface
 	virtual TArray<class UMaterialInterface*> GetMaterials() const override;
diff --git a/Engine/Source/Runtime/Engine/Classes/Components/WindDirectionalSourceComponent.h b/Engine/Source/Runtime/Engine/Classes/Components/WindDirectionalSourceComponent.h
--- a/Engine/Source/Runtime/Engine/Classes/Components/WindDirectionalSourceComponent.h
+++ b/Engine/Source/Runtime/Engine/Classes/Components/WindDirectionalSourceComponent.h
@@ -15,6 +15,11 @@ class UWindDirectionalSourceComponent : public USceneComponent
 	UPROPERTY(interp, Category=WindDirectionalSourceComponent)
 	float Speed;
 
+	UPROPERTY(interp, Category = WindDirectionalSourceComponent)
+	float MinGustAmount;
+
+	UPROPERTY(interp, Category = WindDirectionalSourceComponent)
+	float MaxGustAmount;
 
 public:
 	class FWindSourceSceneProxy* SceneProxy;
diff --git a/Engine/Source/Runtime/Engine/Classes/Engine/Engine.h b/Engine/Source/Runtime/Engine/Classes/Engine/Engine.h
--- a/Engine/Source/Runtime/Engine/Classes/Engine/Engine.h
+++ b/Engine/Source/Runtime/Engine/Classes/Engine/Engine.h
@@ -1309,10 +1309,6 @@ public:
 	UPROPERTY(config)
 	float MaxOcclusionPixelsFraction;
 
-	/** Do not use Ageia PhysX hardware */
-	UPROPERTY(globalconfig)
-	uint32 bDisablePhysXHardwareSupport:1;
-
 	/** Whether to pause the game if focus is lost. */
 	UPROPERTY(config)
 	uint32 bPauseOnLossOfFocus:1;
@@ -1914,6 +1910,14 @@ public:
 		return true;
 	}
 
+	/**
+	 * @return true if level streaming should prefer to stream levels from disk instead of duplicating them from editor world
+	 */
+	virtual bool PreferToStreamLevelsInPIE() const
+	{
+		return false;
+	}
+
 	/**
 	 * Enables or disables the ScreenSaver (PC only)
 	 *
diff --git a/Engine/Source/Runtime/Engine/Classes/Engine/EngineTypes.h b/Engine/Source/Runtime/Engine/Classes/Engine/EngineTypes.h
--- a/Engine/Source/Runtime/Engine/Classes/Engine/EngineTypes.h
+++ b/Engine/Source/Runtime/Engine/Classes/Engine/EngineTypes.h
@@ -284,7 +284,11 @@ enum EMovementMode
 	/** Walking on a surface. */
 	MOVE_Walking	UMETA(DisplayName="Walking"),
 
-	/** Simplified walking on navigation data (e.g. navmesh) */
+	/** 
+	 * Simplified walking on navigation data (e.g. navmesh). 
+	 * If bGenerateOverlapEvents is true, then we will perform sweeps with each navmesh move.
+	 * If bGenerateOverlapEvents is false then movement is cheaper but characters can overlap other objects without some extra process to repel/resolve their collisions.
+	 */
 	MOVE_NavWalking	UMETA(DisplayName="Navmesh Walking"),
 
 	/** Falling under the effects of gravity, such as after jumping or walking off the edge of a surface. */
diff --git a/Engine/Source/Runtime/Engine/Classes/Engine/Level.h b/Engine/Source/Runtime/Engine/Classes/Engine/Level.h
--- a/Engine/Source/Runtime/Engine/Classes/Engine/Level.h
+++ b/Engine/Source/Runtime/Engine/Classes/Engine/Level.h
@@ -449,11 +449,6 @@ public:
 	FLevelSimplificationDetails LevelSimplification[WORLDTILE_LOD_MAX_INDEX];
 #endif //WITH_EDITORONLY_DATA
 
-#if PERF_TRACK_DETAILED_ASYNC_STATS
-	/** Mapping of how long each actor class takes to have UpdateComponents called on it */
-	TMap<const UClass*,struct FMapTimeEntry>		UpdateComponentsTimePerActorClass;
-#endif // PERF_TRACK_DETAILED_ASYNC_STATS
-
 	/** Actor which defines level logical bounding box				*/
 	TWeakObjectPtr<ALevelBounds>				LevelBoundsActor;
 
diff --git a/Engine/Source/Runtime/Engine/Classes/Engine/LevelStreaming.h b/Engine/Source/Runtime/Engine/Classes/Engine/LevelStreaming.h
--- a/Engine/Source/Runtime/Engine/Classes/Engine/LevelStreaming.h
+++ b/Engine/Source/Runtime/Engine/Classes/Engine/LevelStreaming.h
@@ -120,6 +120,17 @@ class ULevelStreaming : public UObject
 	UPROPERTY(Category=LevelStreaming, BlueprintReadWrite)
 	uint32 bShouldBlockOnLoad:1;
 
+	/** Requested LOD. Non LOD sub-levels have Index = -1  */
+	UPROPERTY(transient, Category=LevelStreaming, BlueprintReadWrite)
+	int32 LevelLODIndex;
+	
+	/** 
+	 *  Whether this level streaming object should be ignored by world composition distance streaming, 
+	 *  so streaming state can be controlled by other systems (ex: in blueprints)
+	 */
+	UPROPERTY(transient, Category=LevelStreaming, BlueprintReadWrite)
+	uint32 bDisableDistanceStreaming:1;
+
 	/** Whether this level streaming object's level should be unloaded and the object be removed from the level list.			*/
 	uint32 bIsRequestingUnloadAndRemoval:1;
 	
@@ -150,9 +161,6 @@ class ULevelStreaming : public UObject
 	UPROPERTY()
 	TArray<FString> Keywords;
 
-	/** Requested LOD */
-	int32 LevelLODIndex;
-
 	// Begin UObject Interface
 	virtual void PostLoad() override;
 	virtual void Serialize( FArchive& Ar ) override;
@@ -190,7 +198,7 @@ class ULevelStreaming : public UObject
 	 *
 	 * @return true if level should be loaded/ streamed in, false otherwise
 	 */
-	virtual bool ShouldBeLoaded();
+	virtual bool ShouldBeLoaded() const;
 
 	/**
 	 * Return whether this level should be visible/ associated with the world if it is
@@ -198,7 +206,7 @@ class ULevelStreaming : public UObject
 	 * 
 	 * @return true if the level should be visible, false otherwise
 	 */
-	virtual bool ShouldBeVisible();
+	virtual bool ShouldBeVisible() const;
 
 	virtual bool ShouldBeAlwaysLoaded() const { return false; }
 	
@@ -243,6 +251,10 @@ class ULevelStreaming : public UObject
 	UFUNCTION(BlueprintCallable, Category="Game")
 	bool IsLevelLoaded() const;
 
+	/** Returns whether level has streaming state change pending */
+	UFUNCTION(BlueprintCallable, Category="Game")
+	bool IsStreamingStatePending() const;
+
 	/** Creates a new instance of this streaming level with a provided unique instance name */
 	UFUNCTION(BlueprintCallable, Category="Game")
 	ULevelStreaming* CreateInstance(FString UniqueInstanceName);
diff --git a/Engine/Source/Runtime/Engine/Classes/Engine/LevelStreamingAlwaysLoaded.h b/Engine/Source/Runtime/Engine/Classes/Engine/LevelStreamingAlwaysLoaded.h
--- a/Engine/Source/Runtime/Engine/Classes/Engine/LevelStreamingAlwaysLoaded.h
+++ b/Engine/Source/Runtime/Engine/Classes/Engine/LevelStreamingAlwaysLoaded.h
@@ -16,7 +16,7 @@ class ULevelStreamingAlwaysLoaded : public ULevelStreaming
 	GENERATED_UCLASS_BODY()
 
 	// Begin ULevelStreaming Interface
-	virtual bool ShouldBeLoaded() override;
+	virtual bool ShouldBeLoaded() const override;
 	virtual bool ShouldBeAlwaysLoaded() const override { return true; } 
 	// End ULevelStreaming Interface
 };
diff --git a/Engine/Source/Runtime/Engine/Classes/Engine/LevelStreamingKismet.h b/Engine/Source/Runtime/Engine/Classes/Engine/LevelStreamingKismet.h
--- a/Engine/Source/Runtime/Engine/Classes/Engine/LevelStreamingKismet.h
+++ b/Engine/Source/Runtime/Engine/Classes/Engine/LevelStreamingKismet.h
@@ -29,8 +29,8 @@ class ULevelStreamingKismet : public ULevelStreaming
 	// End UObject Interface
 
 	// Begin ULevelStreaming Interface
-	virtual bool ShouldBeLoaded() override;
-	virtual bool ShouldBeVisible() override;
+	virtual bool ShouldBeLoaded() const override;
+	virtual bool ShouldBeVisible() const override;
 	// End ULevelStreaming Interface
 };
 
diff --git a/Engine/Source/Runtime/Engine/Classes/Engine/LevelStreamingPersistent.h b/Engine/Source/Runtime/Engine/Classes/Engine/LevelStreamingPersistent.h
--- a/Engine/Source/Runtime/Engine/Classes/Engine/LevelStreamingPersistent.h
+++ b/Engine/Source/Runtime/Engine/Classes/Engine/LevelStreamingPersistent.h
@@ -17,11 +17,11 @@ class ULevelStreamingPersistent : public ULevelStreaming
 
 
 	// Begin ULevelStreaming Interface
-	virtual bool ShouldBeLoaded() override
+	virtual bool ShouldBeLoaded() const override
 	{
 		return true;
 	}
-	virtual bool ShouldBeVisible() override
+	virtual bool ShouldBeVisible() const override
 	{
 		return true;
 	}
diff --git a/Engine/Source/Runtime/Engine/Classes/Engine/SkeletalMesh.h b/Engine/Source/Runtime/Engine/Classes/Engine/SkeletalMesh.h
--- a/Engine/Source/Runtime/Engine/Classes/Engine/SkeletalMesh.h
+++ b/Engine/Source/Runtime/Engine/Classes/Engine/SkeletalMesh.h
@@ -785,6 +785,12 @@ public:
 	 */
 	ENGINE_API TArray<USkeletalMeshSocket*>& GetMeshOnlySocketList();
 
+	/**
+	* Returns the "active" socket list - all sockets from this mesh plus all non-duplicates from the skeleton
+	* Const ref return value as this cannot be modified externally
+	*/
+	ENGINE_API TArray<USkeletalMeshSocket*> GetActiveSocketList() const;
+
 #if WITH_EDITOR
 	/** Retrieves the source model for this skeletal mesh. */
 	ENGINE_API FStaticLODModel& GetSourceModel();
@@ -795,12 +801,6 @@ public:
 	 */
 	ENGINE_API FStaticLODModel& PreModifyMesh();
 
-	/**
-	 * Returns the "active" socket list - all sockets from this mesh plus all non-duplicates from the skeleton
-	 * Const ref return value as this cannot be modified externally
-	 */
-	ENGINE_API const TArray<USkeletalMeshSocket*>& GetActiveSocketList() const;
-
 	/**
 	* Makes sure all attached objects are valid and removes any that aren't.
 	*
@@ -872,10 +872,8 @@ public:
 
 private:
 
-#if WITH_EDITOR
 	/** Utility function to help with building the combined socket list */
 	bool IsSocketOnMesh( const FName& InSocketName ) const;
-#endif
 
 	/**
 	* Flush current render state
diff --git a/Engine/Source/Runtime/Engine/Classes/Engine/SkeletalMeshSocket.h b/Engine/Source/Runtime/Engine/Classes/Engine/SkeletalMeshSocket.h
--- a/Engine/Source/Runtime/Engine/Classes/Engine/SkeletalMeshSocket.h
+++ b/Engine/Source/Runtime/Engine/Classes/Engine/SkeletalMeshSocket.h
@@ -30,6 +30,11 @@ class USkeletalMeshSocket : public UObject
 	UPROPERTY(EditAnywhere, BlueprintReadOnly, Category=SkeletalMeshSocket)
 	FVector RelativeScale;
 
+	/** If true then the hierarchy of bones this socket is attached to will always be 
+	    evaluated, even if it had previously been removed due to the current lod setting */
+	UPROPERTY(EditAnywhere, BlueprintReadOnly, Category=SkeletalMeshSocket)
+	bool bForceAlwaysAnimated;
+
 	UFUNCTION(BlueprintCallable, Category="Components|SkeletalMesh")
 	ENGINE_API FVector GetSocketLocation(const class USkeletalMeshComponent* SkelComp) const;
 
@@ -79,6 +84,21 @@ class USkeletalMeshSocket : public UObject
 	 *	@return	bool			true if successful, false if not
 	 */
 	ENGINE_API bool AttachActor (class AActor* Actor, class USkeletalMeshComponent* SkelComp) const;
+
+#if WITH_EDITOR
+	/** Broadcasts a notification whenever the socket property has changed. */
+	DECLARE_EVENT_TwoParams(USkeletalMeshSocket, FSocketChangedEvent, const class USkeletalMeshSocket*, const class UProperty*);
+	FSocketChangedEvent& OnPropertyChanged() { return ChangedEvent; }
+
+	// Begin UObject interface
+	virtual void PostEditChangeProperty(FPropertyChangedEvent& PropertyChangedEvent) override;
+	// End UObject interface
+
+private: 
+	/** Broadcasts a notification whenever the socket property has changed. */
+	FSocketChangedEvent ChangedEvent;
+
+#endif // WITH_EDITOR
 };
 
 
diff --git a/Engine/Source/Runtime/Engine/Classes/Engine/StaticMesh.h b/Engine/Source/Runtime/Engine/Classes/Engine/StaticMesh.h
--- a/Engine/Source/Runtime/Engine/Classes/Engine/StaticMesh.h
+++ b/Engine/Source/Runtime/Engine/Classes/Engine/StaticMesh.h
@@ -474,7 +474,7 @@ public:
 	ENGINE_API virtual void ReleaseResources();
 
 	/**
-	 * Returns the scale dependent texture factor used by the texture streaming code.
+	 * Returns the scale dependent texture factor used by the texture streaLODg code.
 	 *
 	 * @param RequestedUVIndex UVIndex to look at
 	 * @return scale dependent texture factor
diff --git a/Engine/Source/Runtime/Engine/Classes/Engine/StaticMeshActor.h b/Engine/Source/Runtime/Engine/Classes/Engine/StaticMeshActor.h
--- a/Engine/Source/Runtime/Engine/Classes/Engine/StaticMeshActor.h
+++ b/Engine/Source/Runtime/Engine/Classes/Engine/StaticMeshActor.h
@@ -1,6 +1,7 @@
 // Copyright 1998-2015 Epic Games, Inc. All Rights Reserved.
 
 #pragma once
+#include "AI/Navigation/NavigationTypes.h"
 #include "StaticMeshActor.generated.h"
 
 /**
@@ -28,6 +29,9 @@ public:
 	UPROPERTY(Category=Actor, EditAnywhere, AdvancedDisplay)
 	bool bStaticMeshReplicateMovement;
 
+	UPROPERTY(EditAnywhere, Category = Navigation, AdvancedDisplay)
+	ENavDataGatheringMode NavigationGeometryGatheringMode;
+
 	/** Function to change mobility type */
 	void SetMobility(EComponentMobility::Type InMobility);
 
@@ -38,6 +42,10 @@ public:
 	// End AActor Interface
 #endif // WITH_EDITOR	
 
+	// INavRelevantInterface begin
+	virtual ENavDataGatheringMode GetGeometryGatheringMode() const { return ENavDataGatheringMode::Default; }
+	// INavRelevantInterface end
+
 protected:
 	// Begin UObject interface.
 	virtual FString GetDetailedInfoInternal() const override;
diff --git a/Engine/Source/Runtime/Engine/Classes/Engine/Texture.h b/Engine/Source/Runtime/Engine/Classes/Engine/Texture.h
--- a/Engine/Source/Runtime/Engine/Classes/Engine/Texture.h
+++ b/Engine/Source/Runtime/Engine/Classes/Engine/Texture.h
@@ -199,6 +199,9 @@ struct FTextureSource
 	FORCEINLINE bool IsPNGCompressed() const { return bPNGCompressed; }
 	FORCEINLINE int32 GetSizeOnDisk() const { return BulkData.GetBulkDataSize(); }
 	FORCEINLINE bool IsBulkDataLoaded() const { return BulkData.IsBulkDataLoaded(); }
+	
+	/** Sets the GUID to use, and whether that GUID is actually a hash of some data. */
+	void SetId(const FGuid& InId, bool bInGuidIsHash);
 #endif
 
 private:
diff --git a/Engine/Source/Runtime/Engine/Classes/Engine/World.h b/Engine/Source/Runtime/Engine/Classes/Engine/World.h
--- a/Engine/Source/Runtime/Engine/Classes/Engine/World.h
+++ b/Engine/Source/Runtime/Engine/Classes/Engine/World.h
@@ -493,6 +493,39 @@ struct ENGINE_API FWorldAsyncTraceState
 	int32 NextAvailableOverlapIndex;
 };
 
+#if WITH_EDITOR
+/* FAsyncPreRegisterDDCRequest - info about an async DDC request that we're going to wait on before registering components */
+class ENGINE_API FAsyncPreRegisterDDCRequest
+{
+	/* DDC Key used for the request */
+	FString DDCKey;
+
+	/* Handle for Async DDC request. 0 if no longer invalid. */
+	uint32 Handle;
+public:
+	/** constructor */
+	FAsyncPreRegisterDDCRequest(const FString& InKey, uint32 InHandle)
+	: DDCKey(InKey)
+	, Handle(InHandle)
+	{}
+
+	/** destructor */
+	~FAsyncPreRegisterDDCRequest();
+
+	/** returns true if the request is complete */
+	bool PollAsynchronousCompletion();
+
+	/** waits until the request is complete */
+	void WaitAsynchronousCompletion();
+
+	/** returns true if the DDC returned the results requested. Must only be called once. */
+	bool GetAsynchronousResults(TArray<uint8>& OutData);
+
+	/** get the DDC key associated with this request */
+	const FString& GetKey() const { return DDCKey; }
+};
+#endif
+
 /** 
  * The World is the top level object representing a map or a sandbox in which Actors and Components will exist and be rendered.  
  *
@@ -509,6 +542,8 @@ class ENGINE_API UWorld : public UObject, public FNetworkNotify
 {
 	GENERATED_UCLASS_BODY()
 
+	~UWorld();
+
 #if WITH_EDITORONLY_DATA
 	/** List of all the layers referenced by the world's actors */
 	UPROPERTY()
@@ -1006,6 +1041,14 @@ public:
 	UPROPERTY(transient)
 	uint32 bAreConstraintsDirty:1;
 
+	/** Coordinates async tasks started in post load that we want completed before we register components.  May not be here for long; currently used to convert foliage instance buffers.*/
+	FThreadSafeCounter AsyncPreRegisterLevelStreamingTasks;
+
+#if WITH_EDITORONLY_DATA
+	/** List of DDC async requests we need to wait on before we register components. Game thread only. */
+	TArray<TSharedPtr<FAsyncPreRegisterDDCRequest>> AsyncPreRegisterDDCRequests;
+#endif
+
 	/**
 	 * UWorld default constructor
 	 */
diff --git a/Engine/Source/Runtime/Engine/Classes/Engine/WorldComposition.h b/Engine/Source/Runtime/Engine/Classes/Engine/WorldComposition.h
--- a/Engine/Source/Runtime/Engine/Classes/Engine/WorldComposition.h
+++ b/Engine/Source/Runtime/Engine/Classes/Engine/WorldComposition.h
@@ -73,6 +73,8 @@ class ENGINE_API UWorldComposition : public UObject
 	
 	/** Adds or removes level streaming objects to world based on distance settings from current view point */
 	void UpdateStreamingState(const FVector& InLocation);
+	void UpdateStreamingState(const FVector* InLocation, int32 Num);
+	void UpdateStreamingStateCinematic(const FVector* InLocation, int32 Num);
 
 #if WITH_EDITOR
 	/** Simulates streaming in editor world, only visibility, no loading/unloading, no LOD sub-levels 
@@ -89,6 +91,7 @@ class ENGINE_API UWorldComposition : public UObject
 
 	/** Returns currently visible and hidden levels based on distance based streaming */
 	void GetDistanceVisibleLevels(const FVector& InLocation, TArray<FDistanceVisibleLevel>& OutVisibleLevels, TArray<FDistanceVisibleLevel>& OutHiddenLevels) const;
+	void GetDistanceVisibleLevels(const FVector* InLocations, int32 Num, TArray<FDistanceVisibleLevel>& OutVisibleLevels, TArray<FDistanceVisibleLevel>& OutHiddenLevels) const;
 
 	/** @returns Whether specified streaming level is distance dependent */
 	bool IsDistanceDependentLevel(FName PackageName) const;
@@ -182,10 +185,14 @@ private:
 	 * Finds tile by package name 
 	 * @return Pointer to a found tile 
 	 */
+	int32 FindTileIndexByName(const FName& InPackageName) const;
 	FWorldCompositionTile* FindTileByName(const FName& InPackageName) const;
+	
+	/** @returns Whether specified streaming level is distance dependent */
+	bool IsDistanceDependentLevel(int32 TileIdx) const;
 
 	/** Attempts to set new streaming state for a particular tile, could be rejected if state change on 'cooldown' */
-	void CommitTileStreamingState(UWorld* PersistenWorld, int32 TileIdx, bool bShouldBeLoaded, bool bShouldBeVisible, int32 LODIdx);
+	bool CommitTileStreamingState(UWorld* PersistenWorld, int32 TileIdx, bool bShouldBeLoaded, bool bShouldBeVisible, bool bShouldBlock, int32 LODIdx);
 
 public:
 #if WITH_EDITOR
@@ -209,4 +216,9 @@ public:
 	// Time threshold between tile streaming state changes
 	UPROPERTY(config)
 	double						TilesStreamingTimeThreshold;
+
+	// Whether all distance dependent tiles should be loaded and visible during cinematic
+	UPROPERTY(config)
+	bool						bLoadAllTilesDuringCinematic;
+
 };
\ No newline at end of file
diff --git a/Engine/Source/Runtime/Engine/Classes/Kismet/GameplayStatics.h b/Engine/Source/Runtime/Engine/Classes/Kismet/GameplayStatics.h
--- a/Engine/Source/Runtime/Engine/Classes/Kismet/GameplayStatics.h
+++ b/Engine/Source/Runtime/Engine/Classes/Kismet/GameplayStatics.h
@@ -114,7 +114,11 @@ class ENGINE_API UGameplayStatics : public UBlueprintFunctionLibrary
 	/** Returns level streaming object with specified level package name */
 	UFUNCTION(BlueprintPure, meta=(WorldContext="WorldContextObject"), Category="Game")
 	static class ULevelStreaming* GetStreamingLevel(UObject* WorldContextObject, FName PackageName);
-	
+
+	/** Flushes level streaming in blocking fashion and returns when all sub-levels are loaded / visible / hidden */
+	UFUNCTION(BlueprintCallable, meta=(WorldContext="WorldContextObject"), Category = "Game")
+	static void FlushLevelStreaming(UObject* WorldContextObject);
+		
 	/** Cancels all currently queued streaming packages */
 	UFUNCTION(BlueprintCallable, Category = "Game")
 	static void CancelAsyncLoading();
@@ -184,7 +188,7 @@ class ENGINE_API UGameplayStatics : public UBlueprintFunctionLibrary
 	 * @return true if damage was applied to at least one actor.
 	 */
 	UFUNCTION(BlueprintCallable, BlueprintAuthorityOnly, Category="Game|Damage", meta=(WorldContext="WorldContextObject", AutoCreateRefTerm="IgnoreActors"))
-		static bool ApplyRadialDamage(UObject* WorldContextObject, float BaseDamage, const FVector& Origin, float DamageRadius, TSubclassOf<class UDamageType> DamageTypeClass, const TArray<AActor*>& IgnoreActors, AActor* DamageCauser = NULL, AController* InstigatedByController = NULL, bool bDoFullDamage = false, ECollisionChannel DamagePreventionChannel = ECC_Visibility);
+	static bool ApplyRadialDamage(UObject* WorldContextObject, float BaseDamage, const FVector& Origin, float DamageRadius, TSubclassOf<class UDamageType> DamageTypeClass, const TArray<AActor*>& IgnoreActors, AActor* DamageCauser = NULL, AController* InstigatedByController = NULL, bool bDoFullDamage = false, ECollisionChannel DamagePreventionChannel = ECC_Visibility);
 	
 	/** Hurt locally authoritative actors within the radius. Will only hit components that block the Visibility channel.
 	 * @param BaseDamage - The base damage to apply, i.e. the damage at the origin.
@@ -200,7 +204,7 @@ class ENGINE_API UGameplayStatics : public UBlueprintFunctionLibrary
 	 * @return true if damage was applied to at least one actor.
 	 */
 	UFUNCTION(BlueprintCallable, BlueprintAuthorityOnly, Category="Game|Damage", meta=(WorldContext="WorldContextObject", AutoCreateRefTerm="IgnoreActors"))
-		static bool ApplyRadialDamageWithFalloff(UObject* WorldContextObject, float BaseDamage, float MinimumDamage, const FVector& Origin, float DamageInnerRadius, float DamageOuterRadius, float DamageFalloff, TSubclassOf<class UDamageType> DamageTypeClass, const TArray<AActor*>& IgnoreActors, AActor* DamageCauser = NULL, AController* InstigatedByController = NULL, ECollisionChannel DamagePreventionChannel = ECC_Visibility);
+	static bool ApplyRadialDamageWithFalloff(UObject* WorldContextObject, float BaseDamage, float MinimumDamage, const FVector& Origin, float DamageInnerRadius, float DamageOuterRadius, float DamageFalloff, TSubclassOf<class UDamageType> DamageTypeClass, const TArray<AActor*>& IgnoreActors, AActor* DamageCauser = NULL, AController* InstigatedByController = NULL, ECollisionChannel DamagePreventionChannel = ECC_Visibility);
 	
 
 	/** Hurts the specified actor with the specified impact.
@@ -529,5 +533,17 @@ class ENGINE_API UGameplayStatics : public UBlueprintFunctionLibrary
 	/** Requests a new location for a world origin */
 	UFUNCTION(BlueprintCallable, Category="Game", meta=(WorldContext="WorldContextObject"))
 	static void SetWorldOriginLocation(UObject* WorldContextObject, FIntVector NewLocation);
+
+	/**
+	* Counts how many grass foliage instances overlap a given sphere.
+	*
+	* @param	Mesh			The static mesh we are interested in counting
+	* @param	CenterPosition	The center position of the sphere
+	* @param	Radius			The radius of the sphere.
+	*
+	* @return number of foliage instances with their mesh set to Mesh that overlap the sphere
+	*/
+	UFUNCTION(BlueprintCallable, Category = "Foliage", meta = (WorldContext = "WorldContextObject", UnsafeDuringActorConstruction = "true"))
+	static int32 GrassOverlappingSphereCount(UObject* WorldContextObject, const UStaticMesh* StaticMesh, FVector CenterPosition, float Radius);
 };
 
diff --git a/Engine/Source/Runtime/Engine/Classes/Materials/MaterialExpressionSpeedTree.h b/Engine/Source/Runtime/Engine/Classes/Materials/MaterialExpressionSpeedTree.h
--- a/Engine/Source/Runtime/Engine/Classes/Materials/MaterialExpressionSpeedTree.h
+++ b/Engine/Source/Runtime/Engine/Classes/Materials/MaterialExpressionSpeedTree.h
@@ -17,12 +17,13 @@ enum ESpeedTreeGeometryType
 UENUM()
 enum ESpeedTreeWindType
 {
-	STW_None	UMETA(DisplayName="None"),
-	STW_Fastest	UMETA(DisplayName="Fastest"),
-	STW_Fast	UMETA(DisplayName="Fast"),
-	STW_Better	UMETA(DisplayName="Better"),
-	STW_Best	UMETA(DisplayName="Best"),
-	STW_Palm	UMETA(DisplayName="Palm")
+	STW_None		UMETA(DisplayName="None"),
+	STW_Fastest		UMETA(DisplayName="Fastest"),
+	STW_Fast		UMETA(DisplayName="Fast"),
+	STW_Better		UMETA(DisplayName="Better"),
+	STW_Best		UMETA(DisplayName="Best"),
+	STW_Palm		UMETA(DisplayName="Palm"),
+	STW_BestPlus	UMETA(DisplayName="BestPlus"),
 };
 
 UENUM()
diff --git a/Engine/Source/Runtime/Engine/Classes/Matinee/InterpTrackFade.h b/Engine/Source/Runtime/Engine/Classes/Matinee/InterpTrackFade.h
--- a/Engine/Source/Runtime/Engine/Classes/Matinee/InterpTrackFade.h
+++ b/Engine/Source/Runtime/Engine/Classes/Matinee/InterpTrackFade.h
@@ -19,6 +19,13 @@ class UInterpTrackFade : public UInterpTrackFloatBase
 	UPROPERTY(EditAnywhere, Category=InterpTrackFade)
 	uint32 bPersistFade:1;
 
+	/** True to set master audio volume along with the visual fade. */
+	UPROPERTY(EditAnywhere, Category = InterpTrackFade)
+	uint32 bFadeAudio : 1;
+
+	/** Color to fade to. */
+	UPROPERTY(EditAnywhere, Category = InterpTrackFade)
+	FLinearColor FadeColor;
 
 	// Begin UInterpTrack interface.
 	virtual int32 AddKeyframe(float Time, UInterpTrackInst* TrInst, EInterpCurveMode InitInterpMode) override;
diff --git a/Engine/Source/Runtime/Engine/Classes/Matinee/InterpTrackFloatAnimBPParam.h b/Engine/Source/Runtime/Engine/Classes/Matinee/InterpTrackFloatAnimBPParam.h
new file mode 100644
--- /dev/null
+++ b/Engine/Source/Runtime/Engine/Classes/Matinee/InterpTrackFloatAnimBPParam.h
@@ -0,0 +1,36 @@
+// Copyright 1998-2015 Epic Games, Inc. All Rights Reserved.
+
+#pragma once
+#include "InterpTrackFloatAnimBPParam.generated.h"
+
+UCLASS(meta=( DisplayName = "Float Anim BP Parameter Track" ) )
+class UInterpTrackFloatAnimBPParam : public UInterpTrackFloatBase
+{
+	GENERATED_UCLASS_BODY()
+	
+	/** Materials whose parameters we want to change and the references to those materials. */
+	UPROPERTY(EditAnywhere, Category=InterpTrackFloatAnimBPParam)
+	class UAnimBlueprintGeneratedClass* AnimBlueprintClass;
+
+	/** Name of parameter in the MaterialInstance which this track will modify over time. */
+	UPROPERTY(EditAnywhere, Category=InterpTrackFloatAnimBPParam)
+	FName ParamName;
+
+	// Begin UObject Interface
+#if WITH_EDITOR
+	virtual void PostEditChangeProperty(FPropertyChangedEvent& PropertyChangedEvent) override;
+#endif // WITH_EDITOR
+	// End UObject Interface
+
+	// Begin UInterpTrack interface.
+	virtual int32 AddKeyframe(float Time, UInterpTrackInst* TrInst, EInterpCurveMode InitInterpMode) override;
+	virtual void PreviewUpdateTrack(float NewPosition, UInterpTrackInst* TrInst) override;
+	virtual void UpdateTrack(float NewPosition, UInterpTrackInst* TrInst, bool bJump) override;
+	// End UInterpTrack interface.
+
+private:
+	bool bRefreshParamter;
+};
+
+
+
diff --git a/Engine/Source/Runtime/Engine/Classes/Matinee/InterpTrackInstFloatAnimBPParam.h b/Engine/Source/Runtime/Engine/Classes/Matinee/InterpTrackInstFloatAnimBPParam.h
new file mode 100644
--- /dev/null
+++ b/Engine/Source/Runtime/Engine/Classes/Matinee/InterpTrackInstFloatAnimBPParam.h
@@ -0,0 +1,34 @@
+// Copyright 1998-2015 Epic Games, Inc. All Rights Reserved.
+
+
+#pragma once
+#include "InterpTrackInstFloatAnimBPParam.generated.h"
+
+UCLASS()
+class UInterpTrackInstFloatAnimBPParam : public UInterpTrackInst
+{
+	GENERATED_UCLASS_BODY()
+
+	/** MIDs we're using to set the desired parameter. */
+	UPROPERTY(transient)
+	class UAnimInstance* AnimScriptInstance;
+
+	/** Saved values for restoring state when exiting Matinee. */
+	UPROPERTY(transient)
+	float ResetFloat;
+
+	UProperty* ParamProperty;
+
+	// Begin UInterpTrackInst Instance
+	virtual void InitTrackInst(UInterpTrack* Track) override;
+	virtual void SaveActorState(UInterpTrack* Track) override;
+	virtual void RestoreActorState(UInterpTrack* Track) override;
+	// End UInterpTrackInst Instance
+
+	// set/getter for the parameter reading
+	void SetValue(float InValue);
+	float GetValue() const;
+	// refresh parameter setup
+	void RefreshParameter(UInterpTrack* Track);
+};
+
diff --git a/Engine/Source/Runtime/Engine/Classes/PhysicsEngine/BodyInstance.h b/Engine/Source/Runtime/Engine/Classes/PhysicsEngine/BodyInstance.h
--- a/Engine/Source/Runtime/Engine/Classes/PhysicsEngine/BodyInstance.h
+++ b/Engine/Source/Runtime/Engine/Classes/PhysicsEngine/BodyInstance.h
@@ -6,7 +6,6 @@
 #include "CollisionQueryParams.h"
 #include "BodyInstance.generated.h"
 
-
 #define UE_WITH_PHYSICS (WITH_PHYSX || WITH_BOX2D)
 
 struct FCollisionShape;
@@ -21,6 +20,8 @@ class UPhysicsConstraintComponent;
 DECLARE_DELEGATE_TwoParams(FCalculateCustomPhysics, float, FBodyInstance*);
 
 #if WITH_PHYSX
+struct FShapeData;
+
 namespace physx
 {
 	class PxRigidActor;
@@ -32,6 +33,7 @@ namespace physx
 	class PxMaterial;
 	class PxTransform;
 	struct PxContactPair;
+	struct PxFilterData;
 }
 #endif // WITH_PHYSX
 
@@ -418,8 +420,6 @@ public:
 	 */
 	static void InitBodies(TArray<FBodyInstance*>& Bodies, TArray<FTransform>& Transforms, class UBodySetup* BodySetup, class UPrimitiveComponent* PrimitiveComp, class FPhysScene* InRBScene, PhysXAggregateType InAggregate = NULL, bool Defer = false);
 
-	void InitBody_Legacy(class UBodySetup* Setup, const FTransform& Transform, class UPrimitiveComponent* PrimComp, class FPhysScene* InRBScene, PhysXAggregateType InAggregate = NULL);
-
 	/** Validate a body transform, outputting debug info
 	 *	@param Transform Transform to debug
 	 *	@param DebugName Name of the instance for logging
@@ -427,10 +427,41 @@ public:
 	 */
 	static bool ValidateTransform(const FTransform &Transform, FString& DebugName, UBodySetup* Setup);
 
+	/** Standalone path to batch initialise large amounts of static bodies, which will be deferred till the next scene update for fast scene addition.
+	 *	@param Bodies
+	 *	@param Transforms
+	 *	@param BodySetup
+	 *	@param PrimitiveComp
+	 *	@param InRBScene
+	 */
+	static void InitStaticBodies(TArray<FBodyInstance*>& Bodies, TArray<FTransform>& Transforms, class UBodySetup* BodySetup, class UPrimitiveComponent* PrimitiveComp, class FPhysScene* InRBScene);
+
 #if WITH_PHYSX
 	/** Initialise dynamic properties for this instance when using PhysX - this must be done after scene addition.
 	 */
 	void InitDynamicProperties();
+
+	/** Populate the filter data within the provided FShapeData with the correct filters for this instance
+	 *	@param ShapeData ShapeData to populate
+	 *	@param bForceSimpleAsComplex Whether to force simple colision as complex
+	 */
+	void GetFilterData(FShapeData& ShapeData, bool bForceSimpleAsComplex = false);
+
+	/** Set initialisation flags on a provided shape
+	 *	@param UseCollisionEnabled Whether collision is enabled for the shape
+	 *	@param PShape The shape to set the flags on
+	 *	@param SceneType Which scene we are using
+	 *	@param bUseComplexAsSimple Whether to use complex collision as simple
+	 */
+	void SetShapeFlags(TEnumAsByte<ECollisionEnabled::Type> UseCollisionEnabled, physx::PxShape* PShape, EPhysicsSceneType SceneType, const bool bUseComplexAsSimple = false);
+
+	/** Populate the flag fields of the provided FShapeData with correct initialisation flags
+	 *	@param ShapeData ShapeData to populate
+	 *	@param UseCollisionEnabled Whether collision is enabled for this instance
+	 *	@param bUseComplexAsSimple Whether to use complex collision as simple
+	 */
+	void GetShapeFlags(FShapeData& ShapeData, TEnumAsByte<ECollisionEnabled::Type> UseCollisionEnabled, const bool bUseComplexAsSimple = false);
+
 #endif // WITH_PHYSX
 
 #endif	//UE_WITH_PHYSICS
@@ -474,6 +505,9 @@ public:
 	/** Get the complex PhysicalMaterial array for this body */
 	TArray<UPhysicalMaterial*> GetComplexPhysicalMaterials() const;
 
+	/** Get the complex PhysicalMaterials for this body */
+	void GetComplexPhysicalMaterials(TArray<UPhysicalMaterial*> &PhysMaterials) const;
+
 	/** Returns the slope override struct for this instance. If we don't have our own custom setting, it will return the setting from the body setup. */
 	const struct FWalkableSlopeOverride& GetWalkableSlopeOverride() const;
 
@@ -644,7 +678,7 @@ public:
 #if WITH_PHYSX
 	static void ApplyMaterialToShape(physx::PxShape* PShape, physx::PxMaterial* PSimpleMat, TArray<UPhysicalMaterial*>& ComplexPhysMats);
 
-	void ApplyMaterialToInstanceShapes(physx::PxMaterial* PSimpleMat, TArray<UPhysicalMaterial*> ComplexPhysMats);
+	void ApplyMaterialToInstanceShapes(physx::PxMaterial* PSimpleMat, TArray<UPhysicalMaterial*>& ComplexPhysMats);
 #endif
 
 	/** Update the instances collision filtering data */
diff --git a/Engine/Source/Runtime/Engine/Classes/PhysicsEngine/BodySetup.h b/Engine/Source/Runtime/Engine/Classes/PhysicsEngine/BodySetup.h
--- a/Engine/Source/Runtime/Engine/Classes/PhysicsEngine/BodySetup.h
+++ b/Engine/Source/Runtime/Engine/Classes/PhysicsEngine/BodySetup.h
@@ -217,14 +217,14 @@ public:
 	/** 
 	 *   Add the shapes defined by this body setup to the supplied PxRigidBody. 
 	 */
-	ENGINE_API void AddShapesToRigidActor(physx::PxRigidActor* PDestActor, FVector& Scale3D, const FTransform& RelativeTM = FTransform::Identity, TArray<physx::PxShape*>* NewShapes = NULL);
+	ENGINE_API void AddShapesToRigidActor(FBodyInstance* OwningInstance, physx::PxRigidActor* PDestActor, EPhysicsSceneType SceneType, FVector& Scale3D, physx::PxMaterial* SimpleMaterial, TArray<UPhysicalMaterial*>& ComplexMaterials, FShapeData& ShapeData, const FTransform& RelativeTM = FTransform::Identity, TArray<physx::PxShape*>* NewShapes = NULL);
 
 private:
 	void AddSpheresToRigidActor(physx::PxRigidActor* PDestActor, const FTransform& RelativeTM, float MinScale, float MinScaleAbs, TArray<physx::PxShape*>* NewShapes) const;
 	void AddBoxesToRigidActor(physx::PxRigidActor* PDestActor, const FTransform& RelativeTM, const FVector& Scale3D, const FVector& Scale3DAbs, TArray<physx::PxShape*>* NewShapes) const;
 	void AddSphylsToRigidActor(physx::PxRigidActor* PDestActor, const FTransform& RelativeTM, const FVector& Scale3D, const FVector& Scale3DAbs, TArray<physx::PxShape*>* NewShapes) const;
 	void AddConvexElemsToRigidActor(physx::PxRigidActor* PDestActor, const FTransform& RelativeTM, const FVector& Scale3D, const FVector& Scale3DAbs, TArray<physx::PxShape*>* NewShapes) const;
-	void AddTriMeshToRigidActor(physx::PxRigidActor* PDestActor, const FVector& Scale3D, const FVector& Scale3DAbs) const;
+	void AddTriMeshToRigidActor(physx::PxRigidActor* PDestActor, const FVector& Scale3D, const FVector& Scale3DAbs, TArray<physx::PxShape*>* NewShapes) const;
 #endif // WITH_PHYSX
 
 };
diff --git a/Engine/Source/Runtime/Engine/Private/AI/Navigation/NavCollision.cpp b/Engine/Source/Runtime/Engine/Private/AI/Navigation/NavCollision.cpp
--- a/Engine/Source/Runtime/Engine/Private/AI/Navigation/NavCollision.cpp
+++ b/Engine/Source/Runtime/Engine/Private/AI/Navigation/NavCollision.cpp
@@ -157,7 +157,7 @@ void UNavCollision::Setup(UBodySetup* BodySetup)
 	ClearCollision(); 
 
 	bool bCalculated = false;
-	if (bGatherConvexGeometry)
+	if (ShouldUseConvexCollision())
 	{
 		// Find or create cooked navcollision data
 		FByteBulkData* FormatData = GetCookedData(NAVCOLLISION_FORMAT);
@@ -208,8 +208,13 @@ void UNavCollision::ClearCollision()
 
 void UNavCollision::GetNavigationModifier(FCompositeNavModifier& Modifier, const FTransform& LocalToWorld)
 {
+	QUICK_SCOPE_CYCLE_COUNTER(STAT_NavCollision_GetNavigationModifier);
+
 	const TSubclassOf<UNavArea> UseAreaClass = AreaClass ? AreaClass : UNavigationSystem::GetDefaultObstacleArea();
 
+	Modifier.ReserveForAdditionalAreas(CylinderCollision.Num() + BoxCollision.Num()
+		+ (ConvexCollision.VertexBuffer.Num() > 0 ? ConvexShapeIndices.Num() : 0));
+
 	for (int32 i = 0; i < CylinderCollision.Num(); i++)
 	{
 		FTransform CylinderToWorld = LocalToWorld;
@@ -232,25 +237,28 @@ void UNavCollision::GetNavigationModifier(FCompositeNavModifier& Modifier, const
 		Modifier.Add(AreaMod);
 	}
 
-	// rebuild collision data if needed
-	if (!bHasConvexGeometry)
-	{
-		bHasConvexGeometry = GatherCollision();
-	}
-
-	if (ConvexCollision.VertexBuffer.Num() > 0)
+	if (ShouldUseConvexCollision())
 	{
-		int32 LastVertIndex = 0;
-		TArray<FVector> Verts(ConvexCollision.VertexBuffer);
+		// rebuild collision data if needed
+		if (!bHasConvexGeometry)
+		{
+			bHasConvexGeometry = GatherCollision();
+		}
 
-		for (int32 i = 0; i < ConvexShapeIndices.Num(); i++)
+		if (ConvexCollision.VertexBuffer.Num() > 0)
 		{
-			int32 FirstVertIndex = LastVertIndex;
-			LastVertIndex = ConvexShapeIndices.IsValidIndex(i + 1) ? ConvexShapeIndices[i + 1] : ConvexCollision.VertexBuffer.Num();
+			int32 LastVertIndex = 0;
+			TArray<FVector> Verts(ConvexCollision.VertexBuffer);
 
-			FAreaNavModifier AreaMod(Verts, FirstVertIndex, LastVertIndex, ENavigationCoordSystem::Unreal, LocalToWorld, UseAreaClass);
-			AreaMod.SetIncludeAgentHeight(true);
-			Modifier.Add(AreaMod);
+			for (int32 i = 0; i < ConvexShapeIndices.Num(); i++)
+			{
+				int32 FirstVertIndex = LastVertIndex;
+				LastVertIndex = ConvexShapeIndices.IsValidIndex(i + 1) ? ConvexShapeIndices[i + 1] : ConvexCollision.VertexBuffer.Num();
+
+				FAreaNavModifier AreaMod(Verts, FirstVertIndex, LastVertIndex, ENavigationCoordSystem::Unreal, LocalToWorld, UseAreaClass);
+				AreaMod.SetIncludeAgentHeight(true);
+				Modifier.Add(AreaMod);
+			}
 		}
 	}
 }
@@ -371,7 +379,7 @@ void UNavCollision::Serialize(FArchive& Ar)
 		UE_LOG(LogNavigation, Fatal, TEXT("This platform requires cooked packages, and NavCollision data was not cooked into %s."), *GetFullName());
 	}
 
-	if (bCooked && bGatherConvexGeometry)
+	if (bCooked && ShouldUseConvexCollision())
 	{
 		if (Ar.IsCooking())
 		{
diff --git a/Engine/Source/Runtime/Engine/Private/AI/Navigation/NavigationData.cpp b/Engine/Source/Runtime/Engine/Private/AI/Navigation/NavigationData.cpp
--- a/Engine/Source/Runtime/Engine/Private/AI/Navigation/NavigationData.cpp
+++ b/Engine/Source/Runtime/Engine/Private/AI/Navigation/NavigationData.cpp
@@ -410,7 +410,7 @@ void ANavigationData::RebuildAll()
 {
 	ConditionalConstructGenerator(); //recreate generator
 	
-	if (NavDataGenerator)
+	if (NavDataGenerator.IsValid())
 	{
 		NavDataGenerator->RebuildAll();
 	}
@@ -418,7 +418,7 @@ void ANavigationData::RebuildAll()
 
 void ANavigationData::EnsureBuildCompletion()
 {
-	if (NavDataGenerator)
+	if (NavDataGenerator.IsValid())
 	{
 		NavDataGenerator->EnsureBuildCompletion();
 	}
@@ -426,7 +426,7 @@ void ANavigationData::EnsureBuildCompletion()
 
 void ANavigationData::CancelBuild()
 {
-	if (NavDataGenerator)
+	if (NavDataGenerator.IsValid())
 	{
 		NavDataGenerator->CancelBuild();
 	}
@@ -440,7 +440,7 @@ void ANavigationData::OnNavigationBoundsChanged()
 		ConditionalConstructGenerator();
 	}
 	
-	if (NavDataGenerator)
+	if (NavDataGenerator.IsValid())
 	{
 		NavDataGenerator->OnNavigationBoundsChanged();
 	}
@@ -448,7 +448,7 @@ void ANavigationData::OnNavigationBoundsChanged()
 
 void ANavigationData::TickAsyncBuild(float DeltaSeconds)
 {
-	if (NavDataGenerator)
+	if (NavDataGenerator.IsValid())
 	{
 		NavDataGenerator->TickAsyncBuild(DeltaSeconds);
 	}
@@ -458,8 +458,8 @@ void ANavigationData::RebuildDirtyAreas(const TArray<FNavigationDirtyArea>& Dirt
 {
 	// the 'bWantsUpdate' mechanics allows us to skip first requested update after data is loaded
 	// Can be also used to manually control navigation rebuilding, by for example forcing bWantsUpdate 
-	//	to false in Tick function effectively supressing rebuilding
-	if (bWantsUpdate == true && NavDataGenerator)
+	//	to false in Tick function effectively suppressing rebuilding
+	if (bWantsUpdate == true && NavDataGenerator.IsValid())
 	{
 		NavDataGenerator->RebuildDirtyAreas(DirtyAreas);
 	}
@@ -689,7 +689,7 @@ uint32 ANavigationData::LogMemUsed() const
 
 	UE_LOG(LogNavigation, Display, TEXT("%s: ANavigationData: %u\n    self: %d"), *GetName(), MemUsed, sizeof(ANavigationData));	
 
-	if (NavDataGenerator)
+	if (NavDataGenerator.IsValid())
 	{
 		NavDataGenerator->LogMemUsed();
 	}
diff --git a/Engine/Source/Runtime/Engine/Private/AI/Navigation/NavigationGraph.cpp b/Engine/Source/Runtime/Engine/Private/AI/Navigation/NavigationGraph.cpp
--- a/Engine/Source/Runtime/Engine/Private/AI/Navigation/NavigationGraph.cpp
+++ b/Engine/Source/Runtime/Engine/Private/AI/Navigation/NavigationGraph.cpp
@@ -58,7 +58,7 @@ ANavigationGraph::ANavigationGraph(const FObjectInitializer& ObjectInitializer)
 {
 	if (HasAnyFlags(RF_ClassDefaultObject) == false)
 	{
-		NavDataGenerator.Reset(new FNavGraphGenerator(this));
+		NavDataGenerator = MakeShareable(new FNavGraphGenerator(this));
 	}
 }
 
diff --git a/Engine/Source/Runtime/Engine/Private/AI/Navigation/NavigationOctree.cpp b/Engine/Source/Runtime/Engine/Private/AI/Navigation/NavigationOctree.cpp
--- a/Engine/Source/Runtime/Engine/Private/AI/Navigation/NavigationOctree.cpp
+++ b/Engine/Source/Runtime/Engine/Private/AI/Navigation/NavigationOctree.cpp
@@ -11,7 +11,7 @@
 //----------------------------------------------------------------------//
 FNavigationOctree::FNavigationOctree(const FVector& Origin, float Radius)
 	: TOctree<FNavigationOctreeElement, FNavigationOctreeSemantics>(Origin, Radius)
-	, bPerformGeometryGatheringLazily(false)
+	, DefaultGeometryGatheringMode(ENavDataGatheringMode::Instant)
 	, bGatherGeometry(false)
 	, NodesMemory(0)
 {
@@ -24,9 +24,10 @@ FNavigationOctree::~FNavigationOctree()
 	DEC_MEMORY_STAT_BY(STAT_Navigation_CollisionTreeMemory, NodesMemory);
 }
 
-void FNavigationOctree::SetLazyGeometryGatheringEnabled(bool bInEnable)
+void FNavigationOctree::SetDataGatheringMode(ENavDataGatheringModeConfig Mode)
 {
-	bPerformGeometryGatheringLazily = bInEnable;
+	check(Mode != ENavDataGatheringModeConfig::Invalid);
+	DefaultGeometryGatheringMode = ENavDataGatheringMode(Mode);
 }
 
 void FNavigationOctree::SetNavigableGeometryStoringMode(ENavGeometryStoringMode NavGeometryMode)
@@ -34,60 +35,96 @@ void FNavigationOctree::SetNavigableGeometryStoringMode(ENavGeometryStoringMode
 	bGatherGeometry = (NavGeometryMode == FNavigationOctree::StoreNavGeometry);
 }
 
-void FNavigationOctree::DemandGatheringGeometry(const FNavigationOctreeElement& Element) const
+void FNavigationOctree::DemandLazyDataGathering(FNavigationRelevantData& ElementData)
 {
-	UObject* ElementOb = Element.Owner.Get();
-	if (ElementOb == nullptr || bGatherGeometry == false)
+	UObject* ElementOb = ElementData.GetOwner();
+	if (ElementOb == nullptr)
 	{
 		return;
 	}
-		
-	UActorComponent* ActorComp = Cast<UActorComponent>(ElementOb);
-	if (ActorComp)
+
+	bool bShrink = false;
+	const int32 OrgElementMemory = ElementData.GetGeometryAllocatedSize();
+
+	if (ElementData.IsPendingLazyGeometryGathering() == true && ElementData.SupportsGatheringGeometrySlices() == false)
 	{
-		FNavigationOctreeElement* MutableElement = const_cast<FNavigationOctreeElement*>(&Element);
+		QUICK_SCOPE_CYCLE_COUNTER(STAT_RecastNavMeshGenerator_LazyGeometryExport);
+		UActorComponent& ActorComp = *CastChecked<UActorComponent>(ElementOb);
+		ComponentExportDelegate.ExecuteIfBound(&ActorComp, ElementData);
 
-		ComponentExportDelegate.ExecuteIfBound(ActorComp, MutableElement->Data);
+		bShrink = true;
+
+		// mark this element as no longer needing geometry gathering
+		ElementData.bPendingLazyGeometryGathering = false;
+	}
+
+	if (ElementData.IsPendingLazyModifiersGathering())
+	{
+		QUICK_SCOPE_CYCLE_COUNTER(STAT_RecastNavMeshGenerator_LazyModifiersExport);
+		INavRelevantInterface* NavElement = Cast<INavRelevantInterface>(ElementOb);
+		check(NavElement);
+		NavElement->GetNavigationData(ElementData);
+		ElementData.bPendingLazyModifiersGathering = false;
+		bShrink = true;
+	}
 
+	if (bShrink)
+	{
 		// shrink arrays before counting memory
 		// it will be reallocated when adding to octree and RemoveNode will have different value returned by GetAllocatedSize()
-		MutableElement->Shrink();
+		ElementData.Shrink();
+	}
 
-		// mark this element as no longer needing geometry gathering
-		MutableElement->Data.bPendingLazyGeometryGathering = false;
+	const int32 ElementMemoryChange = ElementData.GetGeometryAllocatedSize() - OrgElementMemory;
+	const_cast<FNavigationOctree*>(this)->NodesMemory += ElementMemoryChange;
+	INC_MEMORY_STAT_BY(STAT_Navigation_CollisionTreeMemory, ElementMemoryChange);
+}
 
-		const int32 ElementMemory = MutableElement->Data.GetGeometryAllocatedSize();
-		const_cast<FNavigationOctree*>(this)->NodesMemory += ElementMemory;
-		INC_MEMORY_STAT_BY(STAT_Navigation_CollisionTreeMemory, ElementMemory);
-	}
+void FNavigationOctree::DemandLazyDataGathering(const FNavigationOctreeElement& Element)
+{
+	FNavigationRelevantData& MutableData = const_cast<FNavigationRelevantData&>(*Element.Data);
+	DemandLazyDataGathering(MutableData);
 }
 
 void FNavigationOctree::AddNode(UObject* ElementOb, INavRelevantInterface* NavElement, const FBox& Bounds, FNavigationOctreeElement& Element)
 {
 	// we assume NavElement is ElementOb already cast
-	//check(ElementOb == NavElement);
-	Element.Owner = ElementOb;
-	Element.Bounds = Bounds;
+	Element.Bounds = Bounds;	
 
-	if (bGatherGeometry)
+	if (NavElement)
 	{
-		UActorComponent* ActorComp = Cast<UActorComponent>(ElementOb);
-		if (ActorComp && (bPerformGeometryGatheringLazily == false))
+		const ENavDataGatheringMode GatheringMode = NavElement->GetGeometryGatheringMode();
+		bool bDoInstantGathering = ((GatheringMode == ENavDataGatheringMode::Default && DefaultGeometryGatheringMode == ENavDataGatheringMode::Instant) 
+			|| GatheringMode == ENavDataGatheringMode::Instant);
+
+		if (bGatherGeometry)
 		{
-			ComponentExportDelegate.ExecuteIfBound(ActorComp, Element.Data);
+			UActorComponent* ActorComp = Cast<UActorComponent>(ElementOb);
+			if (ActorComp)
+			{
+				if (bDoInstantGathering)
+				{
+					ComponentExportDelegate.ExecuteIfBound(ActorComp, *Element.Data);
+				}
+				else
+				{
+					Element.Data->bPendingLazyGeometryGathering = true;
+					Element.Data->bSupportsGatheringGeometrySlices = NavElement && NavElement->SupportsGatheringGeometrySlices();
+				}
+			}
+		}
+
+		SCOPE_CYCLE_COUNTER(STAT_Navigation_GatheringNavigationModifiersSync);
+		if (bDoInstantGathering)
+		{
+			NavElement->GetNavigationData(*Element.Data);
 		}
 		else
 		{
-			Element.Data.bPendingLazyGeometryGathering = true;
+			Element.Data->bPendingLazyModifiersGathering = true;
 		}
 	}
 
-	if (NavElement)
-	{
-		SCOPE_CYCLE_COUNTER(STAT_Navigation_GatheringNavigationModifiersSync);
-		NavElement->GetNavigationData(Element.Data);
-	}
-
 	// shrink arrays before counting memory
 	// it will be reallocated when adding to octree and RemoveNode will have different value returned by GetAllocatedSize()
 	Element.Shrink();
@@ -99,32 +136,32 @@ void FNavigationOctree::AddNode(UObject* ElementOb, INavRelevantInterface* NavEl
 	AddElement(Element);
 }
 
-void FNavigationOctree::AppendToNode(const FOctreeElementId& Id, INavRelevantInterface* NavElement, const FBox& Bounds, FNavigationOctreeElement& Data)
+void FNavigationOctree::AppendToNode(const FOctreeElementId& Id, INavRelevantInterface* NavElement, const FBox& Bounds, FNavigationOctreeElement& Element)
 {
 	FNavigationOctreeElement OrgData = GetElementById(Id);
 
-	Data = OrgData;
-	Data.Bounds = Bounds + OrgData.Bounds.GetBox();
+	Element = OrgData;
+	Element.Bounds = Bounds + OrgData.Bounds.GetBox();
 
 	if (NavElement)
 	{
 		SCOPE_CYCLE_COUNTER(STAT_Navigation_GatheringNavigationModifiersSync);
-		NavElement->GetNavigationData(Data.Data);
+		NavElement->GetNavigationData(*Element.Data);
 	}
 
 	// shrink arrays before counting memory
 	// it will be reallocated when adding to octree and RemoveNode will have different value returned by GetAllocatedSize()
-	Data.Shrink();
+	Element.Shrink();
 
 	const int32 OrgElementMemory = OrgData.GetAllocatedSize();
-	const int32 NewElementMemory = Data.GetAllocatedSize();
+	const int32 NewElementMemory = Element.GetAllocatedSize();
 	const int32 MemoryDelta = NewElementMemory - OrgElementMemory;
 
 	NodesMemory += MemoryDelta;
 	INC_MEMORY_STAT_BY(STAT_Navigation_CollisionTreeMemory, MemoryDelta);
 
 	RemoveElement(Id);
-	AddElement(Data);
+	AddElement(Element);
 }
 
 void FNavigationOctree::UpdateNode(const FOctreeElementId& Id, const FBox& NewBounds)
@@ -137,8 +174,8 @@ void FNavigationOctree::UpdateNode(const FOctreeElementId& Id, const FBox& NewBo
 
 void FNavigationOctree::RemoveNode(const FOctreeElementId& Id)
 {
-	FNavigationOctreeElement& Data = GetElementById(Id);
-	const int32 ElementMemory = Data.GetAllocatedSize();
+	FNavigationOctreeElement& Element = GetElementById(Id);
+	const int32 ElementMemory = Element.GetAllocatedSize();
 	NodesMemory -= ElementMemory;
 	DEC_MEMORY_STAT_BY(STAT_Navigation_CollisionTreeMemory, ElementMemory);
 
@@ -154,7 +191,7 @@ const FNavigationRelevantData* FNavigationOctree::GetDataForID(const FOctreeElem
 
 	const FNavigationOctreeElement& OctreeElement = GetElementById(Id);
 
-	return &OctreeElement.Data;
+	return &*OctreeElement.Data;
 }
 
 //----------------------------------------------------------------------//
@@ -190,7 +227,7 @@ FORCENOINLINE
 void FNavigationOctreeSemantics::SetElementId(const FNavigationOctreeElement& Element, FOctreeElementId Id)
 {
 	UWorld* World = NULL;
-	UObject* ElementOwner = Element.Owner.Get();
+	UObject* ElementOwner = Element.GetOwner();
 
 	if (AActor* Actor = Cast<AActor>(ElementOwner))
 	{
