UE4 - merge GDC branch, code [REDACTED_HANDLE] to main

diff --git a/Engine/Source/Runtime/Engine/Private/AI/Navigation/NavigationSystem.cpp b/Engine/Source/Runtime/Engine/Private/AI/Navigation/NavigationSystem.cpp
--- a/Engine/Source/Runtime/Engine/Private/AI/Navigation/NavigationSystem.cpp
+++ b/Engine/Source/Runtime/Engine/Private/AI/Navigation/NavigationSystem.cpp
@@ -18,6 +18,11 @@
 #include "Editor/GeometryMode/Public/EditorGeometry.h"
 #endif
 
+#if WITH_PHYSX
+#include "PhysicsPublic.h"
+#include "../../PhysicsEngine/PhysXSupport.h"
+#endif // WITH_PHYSX
+
 // @todo this is here only due to circular dependency to AIModule. To be removed
 #include "Navigation/CrowdManager.h"
 #include "Navigation/PathFollowingComponent.h"
@@ -131,6 +136,16 @@ FNavigationInvoker::FNavigationInvoker(AActor& InActor, float InGenerationRadius
 
 }
 
+//----------------------------------------------------------------------//
+// FNavHeightfieldSample
+//----------------------------------------------------------------------//
+FNavHeightfieldSamples::FNavHeightfieldSamples()
+{
+#if WITH_PHYSX
+	//static_assert(sizeof(physx::PxI16) == sizeof(Heights.GetTypeSize()), "FNavHeightfieldSamples::Heights' type needs to be kept in sync with physx::PxI16");
+#endif // WITH_PHYSX
+}
+
 //----------------------------------------------------------------------//
 // FNavDataConfig
 //----------------------------------------------------------------------//
@@ -247,6 +262,7 @@ UNavigationSystem::UNavigationSystem(const FObjectInitializer& ObjectInitializer
 	NextInvokersUpdateTime = 0.f;
 	ActiveTilesUpdateInterval = 1.f;
 	bGenerateNavigationOnlyAroundNavigationInvokers = false;
+	DataGatheringMode = ENavDataGatheringModeConfig::Instant;
 
 	if (HasAnyFlags(RF_ClassDefaultObject) == false)
 	{
@@ -421,8 +437,8 @@ bool UNavigationSystem::ConditionalPopulateNavOctree()
 
 	if (bSupportRebuilding)
 	{
-		NavOctree = new FNavigationOctree(FVector(0,0,0), 64000);
-		NavOctree->SetLazyGeometryGatheringEnabled(bGenerateNavigationOnlyAroundNavigationInvokers);
+		NavOctree = MakeShareable(new FNavigationOctree(FVector(0,0,0), 64000));
+		NavOctree->SetDataGatheringMode(DataGatheringMode);
 		
 		const ERuntimeGenerationType RuntimeGenerationType = GetRuntimeGenerationType();
 		const bool bStoreNavGeometry = (RuntimeGenerationType == ERuntimeGenerationType::Dynamic);
@@ -483,9 +499,9 @@ void UNavigationSystem::PostEditChangeProperty(FPropertyChangedEvent& PropertyCh
 		}
 		else if (PropName == NAME_EnableActiveTiles)
 		{
-			if (NavOctree)
+			if (NavOctree.IsValid())
 			{
-				NavOctree->SetLazyGeometryGatheringEnabled(bGenerateNavigationOnlyAroundNavigationInvokers);
+				NavOctree->SetDataGatheringMode(DataGatheringMode);
 			}
 
 			for (auto NavData : NavDataSet)
@@ -2129,7 +2145,7 @@ FSetElementId UNavigationSystem::RegisterNavOctreeElement(UObject* ElementOwner,
 	}
 #endif
 
-	if (NavOctree == NULL || ElementOwner == NULL || ElementInterface == NULL)
+	if (NavOctree.IsValid() == false || ElementOwner == NULL || ElementInterface == NULL)
 	{
 		return SetId;
 	}
@@ -2192,7 +2208,7 @@ void UNavigationSystem::AddElementToNavOctree(const FNavigationDirtyElement& Dir
 		return;
 	}
 
-	FNavigationOctreeElement GeneratedData;
+	FNavigationOctreeElement GeneratedData(*ElementOwner);
 	const FBox ElementBounds = DirtyElement.NavInterface->GetNavigationBounds();
 
 	UObject* NavigationParent = DirtyElement.NavInterface->GetNavigationParent();
@@ -2232,7 +2248,7 @@ void UNavigationSystem::AddElementToNavOctree(const FNavigationDirtyElement& Dir
 
 	if (bValidBBox && !GeneratedData.IsEmpty())
 	{
-		const int32 DirtyFlag = DirtyElement.FlagsOverride ? DirtyElement.FlagsOverride : GeneratedData.Data.GetDirtyFlag();
+		const int32 DirtyFlag = DirtyElement.FlagsOverride ? DirtyElement.FlagsOverride : GeneratedData.Data->GetDirtyFlag();
 		AddDirtyArea(BBox, DirtyFlag);
 	}
 }
@@ -2246,7 +2262,7 @@ bool UNavigationSystem::GetNavOctreeElementData(UObject* NodeOwner, int32& Dirty
 		{
 			// mark area occupied by given actor as dirty
 			FNavigationOctreeElement& ElementData = NavOctree->GetElementById(*ElementId);
-			DirtyFlags = ElementData.Data.GetDirtyFlag();
+			DirtyFlags = ElementData.Data->GetDirtyFlag();
 			DirtyBounds = ElementData.Bounds.GetBox();
 			return true;
 		}
@@ -2264,7 +2280,7 @@ void UNavigationSystem::UnregisterNavOctreeElement(UObject* ElementOwner, INavRe
 	}
 #endif
 
-	if (NavOctree == NULL || ElementOwner == NULL || ElementInterface == NULL)
+	if (NavOctree.IsValid() == false || ElementOwner == NULL || ElementInterface == NULL)
 	{
 		return;
 	}
@@ -2310,7 +2326,7 @@ void UNavigationSystem::RemoveNavOctreeElementId(const FOctreeElementId& Element
 	{
 		// mark area occupied by given actor as dirty
 		FNavigationOctreeElement& ElementData = NavOctree->GetElementById(ElementId);
-		const int32 DirtyFlag = GetDirtyFlagHelper(UpdateFlags, ElementData.Data.GetDirtyFlag());
+		const int32 DirtyFlag = GetDirtyFlagHelper(UpdateFlags, ElementData.Data->GetDirtyFlag());
 		AddDirtyArea(ElementData.Bounds.GetBox(), DirtyFlag);
 		NavOctree->RemoveNode(ElementId);
 	}
@@ -2318,7 +2334,7 @@ void UNavigationSystem::RemoveNavOctreeElementId(const FOctreeElementId& Element
 
 const FNavigationRelevantData* UNavigationSystem::GetDataForObject(const UObject& Object) const
 {
-	check(NavOctree);
+	check(NavOctree.IsValid());
 
 	const FOctreeElementId* OctreeID = GetObjectsNavOctreeId(&Object);
 
@@ -2506,7 +2522,7 @@ bool UNavigationSystem::UpdateNavOctreeElementBounds(UActorComponent* Comp, cons
 			if (ElementId && ElementId->IsValidId())
 			{
 				FNavigationOctreeElement& ElementData = NavOctree->GetElementById(*ElementId);
-				AddDirtyArea(DirtyArea, ElementData.Data.GetDirtyFlag());
+				AddDirtyArea(DirtyArea, ElementData.Data->GetDirtyFlag());
 			}
 		}
 
@@ -2584,7 +2600,7 @@ void UNavigationSystem::OnActorUnregistered(AActor* Actor)
 
 void UNavigationSystem::FindElementsInNavOctree(const FBox& QueryBox, const FNavigationOctreeFilter& Filter, TArray<FNavigationOctreeElement>& Elements)
 {
-	if (NavOctree == nullptr)
+	if (NavOctree.IsValid() == false)
 	{
 		UE_LOG(LogNavigation, Warning, TEXT("UNavigationSystem::FindElementsInNavOctree gets called while NavOctree is null"));
 		return;
@@ -3019,7 +3035,7 @@ void UNavigationSystem::NavigationBuildingLock()
 		return;
 	}
 
-	GetMainNavData(bAutoCreateNavigationData && NavOctree != nullptr && IsThereAnywhereToBuildNavigation() ? FNavigationSystem::Create : FNavigationSystem::DontCreate);
+	GetMainNavData(bAutoCreateNavigationData && NavOctree.IsValid() && IsThereAnywhereToBuildNavigation() ? FNavigationSystem::Create : FNavigationSystem::DontCreate);
 
 	bNavigationBuildingLocked = true;
 }
@@ -3168,18 +3184,18 @@ void UNavigationSystem::OnLevelRemovedFromWorld(ULevel* InLevel, UWorld* InWorld
 void UNavigationSystem::AddLevelCollisionToOctree(ULevel* Level)
 {
 #if WITH_RECAST
-	if (Level && NavOctree && 
+	if (Level && NavOctree.IsValid() &&
 		NavOctree->GetNavGeometryStoringMode() == FNavigationOctree::StoreNavGeometry)
 	{
 		const TArray<FVector>* LevelGeom = Level->GetStaticNavigableGeometry();
 		const FOctreeElementId* ElementId = GetObjectsNavOctreeId(Level);
 
-		if (LevelGeom && !ElementId)
+		if (!ElementId && LevelGeom && LevelGeom->Num() > 0)
 		{
-			FNavigationOctreeElement BSPElem;
-			FRecastNavMeshGenerator::ExportVertexSoupGeometry(*LevelGeom, BSPElem.Data);
+			FNavigationOctreeElement BSPElem(*Level);
+			FRecastNavMeshGenerator::ExportVertexSoupGeometry(*LevelGeom, *BSPElem.Data);
 
-			const auto& Bounds = BSPElem.Data.Bounds;
+			const auto& Bounds = BSPElem.Data->Bounds;
 			if (!Bounds.GetExtent().IsNearlyZero())
 			{
 				NavOctree->AddNode(Level, NULL, Bounds, BSPElem);
@@ -3194,7 +3210,7 @@ void UNavigationSystem::AddLevelCollisionToOctree(ULevel* Level)
 
 void UNavigationSystem::RemoveLevelCollisionFromOctree(ULevel* Level)
 {
-	if (Level && NavOctree)
+	if (Level && NavOctree.IsValid())
 	{
 		const FOctreeElementId* ElementId = GetObjectsNavOctreeId(Level);
 		UE_LOG(LogNavOctree, Log, TEXT("UNREG %s %s"), *GetNameSafe(Level), ElementId ? TEXT("[exists]") : TEXT(""));
@@ -3276,10 +3292,9 @@ void UNavigationSystem::CleanUp(ECleanupMode Mode)
 
 void UNavigationSystem::DestroyNavOctree()
 {
-	if (NavOctree != NULL)
+	if (NavOctree.IsValid())
 	{
 		NavOctree->Destroy();
-		delete NavOctree;
 		NavOctree = NULL;
 	}
 
@@ -3556,6 +3571,31 @@ bool UNavigationSystem::DoesPathIntersectBox(const FNavigationPath* Path, const
 	return Path != NULL && Path->DoesIntersectBox(Box, AgentLocation, StartingIndex, NULL, AgentExtent);
 }
 
+void UNavigationSystem::SetMaxSimultaneousTileGenerationJobsCount(int32 MaxNumberOfJobs)
+{
+	for (auto NavigationData : NavDataSet)
+	{
+		ARecastNavMesh* RecastNavMesh = Cast<ARecastNavMesh>(NavigationData);
+		if (RecastNavMesh)
+		{
+			RecastNavMesh->SetMaxSimultaneousTileGenerationJobsCount(MaxNumberOfJobs);
+		}
+	}
+}
+
+void UNavigationSystem::ResetMaxSimultaneousTileGenerationJobsCount()
+{
+	for (auto NavigationData : NavDataSet)
+	{
+		ARecastNavMesh* RecastNavMesh = Cast<ARecastNavMesh>(NavigationData);
+		if (RecastNavMesh)
+		{
+			const ARecastNavMesh* CDO = RecastNavMesh->GetClass()->GetDefaultObject<ARecastNavMesh>();
+			RecastNavMesh->SetMaxSimultaneousTileGenerationJobsCount(CDO->MaxSimultaneousTileGenerationJobsCount);
+		}
+	}
+}
+
 //----------------------------------------------------------------------//
 // Active tiles
 //----------------------------------------------------------------------//
@@ -3580,6 +3620,15 @@ void UNavigationSystem::UnregisterNavigationInvoker(AActor& Invoker)
 	}
 }
 
+void UNavigationSystem::SetGeometryGatheringMode(ENavDataGatheringModeConfig NewMode)
+{
+	DataGatheringMode = NewMode;
+	if (NavOctree.IsValid())
+	{
+		NavOctree->SetDataGatheringMode(DataGatheringMode);
+	}
+}
+
 void UNavigationSystem::RegisterInvoker(AActor& Invoker, float TileGenerationRadius, float TileRemovalRadius)
 {
 	UE_CVLOG(bGenerateNavigationOnlyAroundNavigationInvokers == false, this, LogNavigation, Warning
diff --git a/Engine/Source/Runtime/Engine/Private/AI/Navigation/PImplRecastNavMesh.cpp b/Engine/Source/Runtime/Engine/Private/AI/Navigation/PImplRecastNavMesh.cpp
--- a/Engine/Source/Runtime/Engine/Private/AI/Navigation/PImplRecastNavMesh.cpp
+++ b/Engine/Source/Runtime/Engine/Private/AI/Navigation/PImplRecastNavMesh.cpp
@@ -1856,23 +1856,23 @@ static FORCEINLINE float PointDistToSegment2DSquared(const float* PT, const floa
  * Traverses given tile's edges and detects the ones that are either poly (i.e. not triangle, but whole navmesh polygon) 
  * or navmesh edge. Returns a pair of verts for each edge found.
  */
-void FPImplRecastNavMesh::GetDebugPolyEdges(const dtMeshTile* Tile, bool bInternalEdges, bool bNavMeshEdges, TArray<FVector>& InternalEdgeVerts, TArray<FVector>& NavMeshEdgeVerts) const
+void FPImplRecastNavMesh::GetDebugPolyEdges(const dtMeshTile& Tile, bool bInternalEdges, bool bNavMeshEdges, TArray<FVector>& InternalEdgeVerts, TArray<FVector>& NavMeshEdgeVerts) const
 {
 	static const float thr = FMath::Square(0.01f);
 
 	ensure(bInternalEdges || bNavMeshEdges);
 	const bool bExportAllEdges = bInternalEdges && !bNavMeshEdges;
 	
-	for (int i = 0; i < Tile->header->polyCount; ++i)
+	for (int i = 0; i < Tile.header->polyCount; ++i)
 	{
-		const dtPoly* Poly = &Tile->polys[i];
+		const dtPoly* Poly = &Tile.polys[i];
 
 		if (Poly->getType() != DT_POLYTYPE_GROUND)
 		{
 			continue;
 		}
 
-		const dtPolyDetail* pd = &Tile->detailMeshes[i];		
+		const dtPolyDetail* pd = &Tile.detailMeshes[i];
 		for (int j = 0, nj = (int)Poly->vertCount; j < nj; ++j)
 		{
 			bool bIsExternal = !bExportAllEdges && (Poly->neis[j] == 0 || Poly->neis[j] & DT_EXT_LINK);
@@ -1881,8 +1881,8 @@ void FPImplRecastNavMesh::GetDebugPolyEdges(const dtMeshTile* Tile, bool bIntern
 			if (Poly->getArea() == RECAST_NULL_AREA)
 			{
 				if (Poly->neis[j] && !(Poly->neis[j] & DT_EXT_LINK) &&
-					Poly->neis[j] <= Tile->header->offMeshBase &&
-					Tile->polys[Poly->neis[j] - 1].getArea() != RECAST_NULL_AREA)
+					Poly->neis[j] <= Tile.header->offMeshBase &&
+					Tile.polys[Poly->neis[j] - 1].getArea() != RECAST_NULL_AREA)
 				{
 					bIsExternal = true;
 					bIsConnected = false;
@@ -1898,7 +1898,7 @@ void FPImplRecastNavMesh::GetDebugPolyEdges(const dtMeshTile* Tile, bool bIntern
 				unsigned int k = Poly->firstLink;
 				while (k != DT_NULL_LINK)
 				{
-					const dtLink& link = DetourNavMesh->getLink(Tile, k);
+					const dtLink& link = DetourNavMesh->getLink(&Tile, k);
 					k = link.next;
 
 					if (link.edge == j)
@@ -1916,25 +1916,25 @@ void FPImplRecastNavMesh::GetDebugPolyEdges(const dtMeshTile* Tile, bool bIntern
 				continue;
 			}
 
-			const float* V0 = &Tile->verts[Poly->verts[j]*3];
-			const float* V1 = &Tile->verts[Poly->verts[(j+1) % nj]*3];
+			const float* V0 = &Tile.verts[Poly->verts[j] * 3];
+			const float* V1 = &Tile.verts[Poly->verts[(j + 1) % nj] * 3];
 
 			// Draw detail mesh edges which align with the actual poly edge.
 			// This is really slow.
 			for (int32 k = 0; k < pd->triCount; ++k)
 			{
-				const unsigned char* t = &(Tile->detailTris[(pd->triBase+k)*4]);
+				const unsigned char* t = &(Tile.detailTris[(pd->triBase + k) * 4]);
 				const float* tv[3];
 
 				for (int32 m = 0; m < 3; ++m)
 				{
 					if (t[m] < Poly->vertCount)
 					{
-						tv[m] = &Tile->verts[Poly->verts[t[m]]*3];
+						tv[m] = &Tile.verts[Poly->verts[t[m]] * 3];
 					}
 					else
 					{
-						tv[m] = &Tile->detailVerts[(pd->vertBase+(t[m] - Poly->vertCount))*3];
+						tv[m] = &Tile.detailVerts[(pd->vertBase + (t[m] - Poly->vertCount)) * 3];
 					}
 				}
 				for (int m = 0, n = 2; m < 3; n=m++)
@@ -1956,19 +1956,19 @@ void FPImplRecastNavMesh::GetDebugPolyEdges(const dtMeshTile* Tile, bool bIntern
 	}
 }
 
-uint8 GetValidEnds(const dtNavMesh* NavMesh, const dtMeshTile* Tile, const dtPoly* Poly)
+uint8 GetValidEnds(const dtNavMesh& NavMesh, const dtMeshTile& Tile, const dtPoly& Poly)
 {
-	if ((Poly->getType() & DT_POLYTYPE_GROUND) != 0)
+	if ((Poly.getType() & DT_POLYTYPE_GROUND) != 0)
 	{
 		return false;
 	}
 
 	uint8 ValidEnds = FRecastDebugGeometry::OMLE_None;
 
-	unsigned int k = Poly->firstLink;
+	unsigned int k = Poly.firstLink;
 	while (k != DT_NULL_LINK)
 	{
-		const dtLink& link = NavMesh->getLink(Tile, k);
+		const dtLink& link = NavMesh.getLink(&Tile, k);
 		k = link.next;
 
 		if (link.edge == 0)
@@ -1992,46 +1992,87 @@ uint8 GetValidEnds(const dtNavMesh* NavMesh, const dtMeshTile* Tile, const dtPol
  */
 void FPImplRecastNavMesh::GetDebugGeometry(FRecastDebugGeometry& OutGeometry, int32 TileIndex) const
 {
-	if (DetourNavMesh)
+	if (DetourNavMesh == nullptr || TileIndex >= DetourNavMesh->getMaxTiles())
 	{
-		if (TileIndex >= DetourNavMesh->getMaxTiles())
-		{
-			return;
-		}
+		return;
+	}
 				
-		check(NavMeshOwner);
-		dtNavMesh const* ConstNavMesh = DetourNavMesh;
+	check(NavMeshOwner);
+
+	const dtNavMesh* const ConstNavMesh = DetourNavMesh;
 		
-		// presize our tarrays for efficiency
-		const int32 NumTiles = TileIndex == INDEX_NONE ? DetourNavMesh->getMaxTiles() : TileIndex + 1;
-		const int32 StartingTile = TileIndex == INDEX_NONE ? 0 : TileIndex;
+	// presize our tarrays for efficiency
+	const int32 NumTiles = TileIndex == INDEX_NONE ? ConstNavMesh->getMaxTiles() : TileIndex + 1;
+	const int32 StartingTile = TileIndex == INDEX_NONE ? 0 : TileIndex;
+
+	int32 NumVertsToReserve = 0;
+	int32 NumIndicesToReserve = 0;
+	int32 NumClusters = 0;
 
-		int32 NumVertsToReserve = 0;
-		int32 NumIndicesToReserve = 0;
-		for (int32 TileIdx=StartingTile; TileIdx < NumTiles; ++TileIdx)
+	const FRecastNavMeshGenerator* Generator = static_cast<const FRecastNavMeshGenerator*>(NavMeshOwner->GetGenerator());
+
+	if (Generator && Generator->IsBuildingRestrictedToActiveTiles())
+	{
+		const TArray<FIntPoint>& ActiveTiles = NavMeshOwner->GetActiveTiles();
+		for (const FIntPoint& TileLocation : ActiveTiles)
 		{
-			dtMeshTile const* const Tile = ConstNavMesh->getTile(TileIdx);
-			dtMeshHeader const* const Header = Tile->header;
+			const int32 LayersCount = ConstNavMesh->getTileCountAt(TileLocation.X, TileLocation.Y);
 
-			if (Header != NULL)
+			for (int32 Layer = 0; Layer < LayersCount; ++Layer)
 			{
-				NumVertsToReserve += Header->vertCount + Header->detailVertCount;
+				dtMeshTile const* const Tile = ConstNavMesh->getTileAt(TileLocation.X, TileLocation.Y, Layer);
+				if (Tile != nullptr && Tile->header != nullptr)
+				{
+					NumVertsToReserve += Tile->header->vertCount + Tile->header->detailVertCount;
 
-				for (int32 PolyIdx=0; PolyIdx < Header->polyCount; ++PolyIdx)
+					for (int32 PolyIdx = 0; PolyIdx < Tile->header->polyCount; ++PolyIdx)
+					{
+						dtPolyDetail const* const DetailPoly = &Tile->detailMeshes[PolyIdx];
+						NumIndicesToReserve += (DetailPoly->triCount * 3);
+					}
+
+					NumClusters = FMath::Max(Tile->header->clusterCount, NumClusters);
+				}
+			}
+		}
+
+		OutGeometry.MeshVerts.Reserve(NumVertsToReserve);
+		OutGeometry.AreaIndices[0].Reserve(NumIndicesToReserve);
+		OutGeometry.BuiltMeshIndices.Reserve(NumIndicesToReserve);
+		OutGeometry.Clusters.AddZeroed(NumClusters);
+
+		uint32 AllTilesVertBase = 0;
+		for (const FIntPoint& TileLocation : ActiveTiles)
+		{
+			const int32 LayersCount = ConstNavMesh->getTileCountAt(TileLocation.X, TileLocation.Y);
+
+			for (int32 Layer = 0; Layer < LayersCount; ++Layer)
+			{
+				dtMeshTile const* const Tile = ConstNavMesh->getTileAt(TileLocation.X, TileLocation.Y, Layer);
+				if (Tile != nullptr && Tile->header != nullptr)
 				{
-					dtPolyDetail const* const DetailPoly = &Tile->detailMeshes[PolyIdx];
-					NumIndicesToReserve += (DetailPoly->triCount * 3);
+					AllTilesVertBase += GetTilesDebugGeometry(Generator, *Tile, AllTilesVertBase, OutGeometry);
 				}
 			}
 		}
-		int32 NumClusters = 0;
-		for (int32 TileIdx=StartingTile; TileIdx < NumTiles; ++TileIdx)
+	}
+	else
+	{
+		for (int32 TileIdx = StartingTile; TileIdx < NumTiles; ++TileIdx)
 		{
 			dtMeshTile const* const Tile = ConstNavMesh->getTile(TileIdx);
 			dtMeshHeader const* const Header = Tile->header;
 
 			if (Header != NULL)
 			{
+				NumVertsToReserve += Header->vertCount + Header->detailVertCount;
+
+				for (int32 PolyIdx = 0; PolyIdx < Header->polyCount; ++PolyIdx)
+				{
+					dtPolyDetail const* const DetailPoly = &Tile->detailMeshes[PolyIdx];
+					NumIndicesToReserve += (DetailPoly->triCount * 3);
+				}
+
 				NumClusters = FMath::Max(Header->clusterCount, NumClusters);
 			}
 		}
@@ -2041,179 +2082,185 @@ void FPImplRecastNavMesh::GetDebugGeometry(FRecastDebugGeometry& OutGeometry, in
 		OutGeometry.BuiltMeshIndices.Reserve(NumIndicesToReserve);
 		OutGeometry.Clusters.AddZeroed(NumClusters);
 
-		const FRecastNavMeshGenerator* Generator = NavMeshOwner ? (const FRecastNavMeshGenerator*)(NavMeshOwner->GetGenerator()) : NULL;
-
 		// spin through all polys in all tiles and draw them
 		// @see drawMeshTile() in recast code for reference
 		uint32 AllTilesVertBase = 0;
-		for (int32 TileIdx=StartingTile; TileIdx < NumTiles; ++TileIdx)
+		for (int32 TileIdx = StartingTile; TileIdx < NumTiles; ++TileIdx)
 		{
 			dtMeshTile const* const Tile = ConstNavMesh->getTile(TileIdx);
-			dtMeshHeader const* const Header = Tile->header;
 
-			if (Header == NULL)
+			if (Tile == nullptr || Tile->header == nullptr)
 			{
 				continue;
 			}
 
-			const bool bIsBeingBuilt = Generator != NULL && !!NavMeshOwner->bDistinctlyDrawTilesBeingBuilt 
-				&& Generator->IsTileChanged(TileIdx);
-			
-			// add all the poly verts
-			float* F = Tile->verts;
-			for (int32 VertIdx=0; VertIdx < Header->vertCount; ++VertIdx)
-			{
-				FVector const VertPos = Recast2UnrVector(F);
-				OutGeometry.MeshVerts.Add(VertPos);
-				F += 3;
-			}
-			int32 const DetailVertIndexBase = Header->vertCount;
-			// add the detail verts
-			F = Tile->detailVerts;
-			for (int32 DetailVertIdx=0; DetailVertIdx < Header->detailVertCount; ++DetailVertIdx)
-			{
-				FVector const VertPos = Recast2UnrVector(F);
-				OutGeometry.MeshVerts.Add(VertPos);
-				F += 3;
-			}
-
-			// add all the indices
-			for (int32 PolyIdx=0; PolyIdx < Header->polyCount; ++PolyIdx)
-			{
-				dtPoly const* const Poly = &Tile->polys[PolyIdx];
-				
-				if (Poly->getType() == DT_POLYTYPE_GROUND)
-				{
-					dtPolyDetail const* const DetailPoly = &Tile->detailMeshes[PolyIdx];
+			AllTilesVertBase += GetTilesDebugGeometry(Generator, *Tile, AllTilesVertBase, OutGeometry, TileIdx);
+		}
+	}
+}
 
-					TArray<int32>* Indices = bIsBeingBuilt ? &OutGeometry.BuiltMeshIndices : &OutGeometry.AreaIndices[Poly->getArea()];
+int32 FPImplRecastNavMesh::GetTilesDebugGeometry(const FRecastNavMeshGenerator* Generator, const dtMeshTile& Tile, int32 VertBase, FRecastDebugGeometry& OutGeometry, int32 TileIdx) const
+{
+	check(NavMeshOwner && DetourNavMesh);
+	dtMeshHeader const* const Header = Tile.header;
+	check(Header);
+	
+	const bool bIsBeingBuilt = Generator != nullptr && !!NavMeshOwner->bDistinctlyDrawTilesBeingBuilt
+		&& Generator->IsTileChanged(TileIdx == INDEX_NONE ? DetourNavMesh->decodePolyIdTile(DetourNavMesh->getTileRef(&Tile)) : TileIdx);
 
-					// one triangle at a time
-					for (int32 TriIdx=0; TriIdx < DetailPoly->triCount; ++TriIdx)
-					{
-						int32 DetailTriIdx = (DetailPoly->triBase + TriIdx) * 4;
-						const unsigned char* DetailTri = &Tile->detailTris[DetailTriIdx];
+	// add all the poly verts
+	float* F = Tile.verts;
+	for (int32 VertIdx = 0; VertIdx < Header->vertCount; ++VertIdx)
+	{
+		FVector const VertPos = Recast2UnrVector(F);
+		OutGeometry.MeshVerts.Add(VertPos);
+		F += 3;
+	}
+	int32 const DetailVertIndexBase = Header->vertCount;
+	// add the detail verts
+	F = Tile.detailVerts;
+	for (int32 DetailVertIdx = 0; DetailVertIdx < Header->detailVertCount; ++DetailVertIdx)
+	{
+		FVector const VertPos = Recast2UnrVector(F);
+		OutGeometry.MeshVerts.Add(VertPos);
+		F += 3;
+	}
 
-						// calc indices into the vert buffer we just populated
-						int32 TriVertIndices[3];
-						for (int32 TriVertIdx=0; TriVertIdx<3; ++TriVertIdx)
-						{
-							if (DetailTri[TriVertIdx] < Poly->vertCount)
-							{
-								TriVertIndices[TriVertIdx] = AllTilesVertBase + Poly->verts[ DetailTri[TriVertIdx] ];
-							}
-							else
-							{
-								TriVertIndices[TriVertIdx] = AllTilesVertBase + DetailVertIndexBase + (DetailPoly->vertBase + DetailTri[TriVertIdx] - Poly->vertCount);
-							}
-						}
+	// add all the indices
+	for (int32 PolyIdx = 0; PolyIdx < Header->polyCount; ++PolyIdx)
+	{
+		dtPoly const* const Poly = &Tile.polys[PolyIdx];
 
-						Indices->Add(TriVertIndices[0]);
-						Indices->Add(TriVertIndices[1]);
- 						Indices->Add(TriVertIndices[2]);
+		if (Poly->getType() == DT_POLYTYPE_GROUND)
+		{
+			dtPolyDetail const* const DetailPoly = &Tile.detailMeshes[PolyIdx];
 
-						if (Tile->polyClusters && OutGeometry.Clusters.IsValidIndex(Tile->polyClusters[PolyIdx]))
-						{
-							TArray<int32>& ClusterIndices = OutGeometry.Clusters[Tile->polyClusters[PolyIdx]].MeshIndices;
-							ClusterIndices.Add(TriVertIndices[0]);
-							ClusterIndices.Add(TriVertIndices[1]);
-							ClusterIndices.Add(TriVertIndices[2]);
-						}
- 					}
-				}
- 			}
+			TArray<int32>* Indices = bIsBeingBuilt ? &OutGeometry.BuiltMeshIndices : &OutGeometry.AreaIndices[Poly->getArea()];
 
-			for (int32 i = 0; i < Header->offMeshConCount; ++i)
+			// one triangle at a time
+			for (int32 TriIdx = 0; TriIdx < DetailPoly->triCount; ++TriIdx)
 			{
-				const dtOffMeshConnection* OffMeshConnection = &Tile->offMeshCons[i];
+				int32 DetailTriIdx = (DetailPoly->triBase + TriIdx) * 4;
+				const unsigned char* DetailTri = &Tile.detailTris[DetailTriIdx];
 
-				if (OffMeshConnection != NULL)
+				// calc indices into the vert buffer we just populated
+				int32 TriVertIndices[3];
+				for (int32 TriVertIdx = 0; TriVertIdx < 3; ++TriVertIdx)
 				{
-					dtPoly const* const LinkPoly = &Tile->polys[OffMeshConnection->poly];
-					const float* va = &Tile->verts[LinkPoly->verts[0]*3]; //OffMeshConnection->pos;
-					const float* vb = &Tile->verts[LinkPoly->verts[1]*3]; //OffMeshConnection->pos[3];
-
-					const FRecastDebugGeometry::FOffMeshLink Link = {
-						Recast2UnrVector(va)
-						, Recast2UnrVector(vb)
-						, LinkPoly->getArea()
-						, (uint8)OffMeshConnection->getBiDirectional()
-						, GetValidEnds(DetourNavMesh, Tile, LinkPoly)
-						, OffMeshConnection->rad
-					};
-
-					OutGeometry.OffMeshLinks.Add(Link);
+					if (DetailTri[TriVertIdx] < Poly->vertCount)
+					{
+						TriVertIndices[TriVertIdx] = VertBase + Poly->verts[DetailTri[TriVertIdx]];
+					}
+					else
+					{
+						TriVertIndices[TriVertIdx] = VertBase + DetailVertIndexBase + (DetailPoly->vertBase + DetailTri[TriVertIdx] - Poly->vertCount);
+					}
 				}
-			}
 
-			for (int32 i = 0; i < Header->offMeshSegConCount; ++i)
-			{
-				const dtOffMeshSegmentConnection* OffMeshSeg = &Tile->offMeshSeg[i];
-				if (OffMeshSeg != NULL)
+				Indices->Add(TriVertIndices[0]);
+				Indices->Add(TriVertIndices[1]);
+				Indices->Add(TriVertIndices[2]);
+
+				if (Tile.polyClusters && OutGeometry.Clusters.IsValidIndex(Tile.polyClusters[PolyIdx]))
 				{
-					const int32 polyBase = Header->offMeshSegPolyBase + OffMeshSeg->firstPoly;
-					for (int32 j = 0; j < OffMeshSeg->npolys; j++)
-					{
-						dtPoly const* const LinkPoly = &Tile->polys[polyBase + j];
-
-						FRecastDebugGeometry::FOffMeshSegment Link;
-						Link.LeftStart  = Recast2UnrealPoint(&Tile->verts[LinkPoly->verts[0]*3]);
-						Link.LeftEnd    = Recast2UnrealPoint(&Tile->verts[LinkPoly->verts[1]*3]);
-						Link.RightStart = Recast2UnrealPoint(&Tile->verts[LinkPoly->verts[2]*3]);
-						Link.RightEnd   = Recast2UnrealPoint(&Tile->verts[LinkPoly->verts[3]*3]);
-						Link.AreaID		= LinkPoly->getArea();
-						Link.Direction	= (uint8)OffMeshSeg->getBiDirectional();
-						Link.ValidEnds  = GetValidEnds(DetourNavMesh, Tile, LinkPoly);
-						
-						const int LinkIdx = OutGeometry.OffMeshSegments.Add(Link);
-						OutGeometry.OffMeshSegmentAreas[Link.AreaID].Add(LinkIdx);
-					}
+					TArray<int32>& ClusterIndices = OutGeometry.Clusters[Tile.polyClusters[PolyIdx]].MeshIndices;
+					ClusterIndices.Add(TriVertIndices[0]);
+					ClusterIndices.Add(TriVertIndices[1]);
+					ClusterIndices.Add(TriVertIndices[2]);
 				}
 			}
- 
-			for (int32 i = 0; i < Header->clusterCount; i++)
-			{
-				const dtCluster& c0 = Tile->clusters[i];
-				uint32 iLink = c0.firstLink;
-				while (iLink != DT_NULL_LINK)
-				{
-					const dtClusterLink& link = DetourNavMesh->getClusterLink(Tile, iLink);
-					iLink = link.next;
+		}
+	}
 
-					dtMeshTile const* const OtherTile = ConstNavMesh->getTileByRef(link.ref);
-					if (OtherTile)
-					{
-						int32 linkedIdx = ConstNavMesh->decodeClusterIdCluster(link.ref);
-						const dtCluster& c1 = OtherTile->clusters[linkedIdx];
+	for (int32 i = 0; i < Header->offMeshConCount; ++i)
+	{
+		const dtOffMeshConnection* OffMeshConnection = &Tile.offMeshCons[i];
 
-						FRecastDebugGeometry::FClusterLink LinkGeom;
-						LinkGeom.FromCluster = Recast2UnrealPoint(c0.center);
-						LinkGeom.ToCluster = Recast2UnrealPoint(c1.center);
+		if (OffMeshConnection != NULL)
+		{
+			dtPoly const* const LinkPoly = &Tile.polys[OffMeshConnection->poly];
+			const float* va = &Tile.verts[LinkPoly->verts[0] * 3]; //OffMeshConnection->pos;
+			const float* vb = &Tile.verts[LinkPoly->verts[1] * 3]; //OffMeshConnection->pos[3];
 
-						if (linkedIdx > i || TileIdx > (int32)ConstNavMesh->decodeClusterIdTile(link.ref))
-						{
-							FVector UpDir(0,0,1.0f);
-							FVector LinkDir = (LinkGeom.ToCluster - LinkGeom.FromCluster).GetSafeNormal();
-							FVector SideDir = FVector::CrossProduct(LinkDir, UpDir);
-							LinkGeom.FromCluster += SideDir * 40.0f;
-							LinkGeom.ToCluster += SideDir * 40.0f;
-						}
+			const FRecastDebugGeometry::FOffMeshLink Link = {
+				Recast2UnrVector(va)
+				, Recast2UnrVector(vb)
+				, LinkPoly->getArea()
+				, (uint8)OffMeshConnection->getBiDirectional()
+				, GetValidEnds(*DetourNavMesh, Tile, *LinkPoly)
+				, OffMeshConnection->rad
+			};
 
-						OutGeometry.ClusterLinks.Add(LinkGeom);
-					}
-				}
+			OutGeometry.OffMeshLinks.Add(Link);
+		}
+	}
+
+	for (int32 i = 0; i < Header->offMeshSegConCount; ++i)
+	{
+		const dtOffMeshSegmentConnection* OffMeshSeg = &Tile.offMeshSeg[i];
+		if (OffMeshSeg != NULL)
+		{
+			const int32 polyBase = Header->offMeshSegPolyBase + OffMeshSeg->firstPoly;
+			for (int32 j = 0; j < OffMeshSeg->npolys; j++)
+			{
+				dtPoly const* const LinkPoly = &Tile.polys[polyBase + j];
+
+				FRecastDebugGeometry::FOffMeshSegment Link;
+				Link.LeftStart = Recast2UnrealPoint(&Tile.verts[LinkPoly->verts[0] * 3]);
+				Link.LeftEnd = Recast2UnrealPoint(&Tile.verts[LinkPoly->verts[1] * 3]);
+				Link.RightStart = Recast2UnrealPoint(&Tile.verts[LinkPoly->verts[2] * 3]);
+				Link.RightEnd = Recast2UnrealPoint(&Tile.verts[LinkPoly->verts[3] * 3]);
+				Link.AreaID = LinkPoly->getArea();
+				Link.Direction = (uint8)OffMeshSeg->getBiDirectional();
+				Link.ValidEnds = GetValidEnds(*DetourNavMesh, Tile, *LinkPoly);
+
+				const int LinkIdx = OutGeometry.OffMeshSegments.Add(Link);
+				OutGeometry.OffMeshSegmentAreas[Link.AreaID].Add(LinkIdx);
 			}
+		}
+	}
 
- 			// Get tile edges and navmesh edges
-			if (OutGeometry.bGatherPolyEdges || OutGeometry.bGatherNavMeshEdges)
+	/*for (int32 i = 0; i < Header->clusterCount; i++)
+	{
+		const dtCluster& c0 = Tile.clusters[i];
+		uint32 iLink = c0.firstLink;
+		while (iLink != DT_NULL_LINK)
+		{
+			const dtClusterLink& link = DetourNavMesh->getClusterLink(&Tile, iLink);
+			iLink = link.next;
+
+			dtMeshTile const* const OtherTile = DetourNavMesh->getTileByRef(link.ref);
+			if (OtherTile)
 			{
-				GetDebugPolyEdges(Tile, !!OutGeometry.bGatherPolyEdges, !!OutGeometry.bGatherNavMeshEdges
-					, OutGeometry.PolyEdges, OutGeometry.NavMeshEdges);
+				int32 linkedIdx = DetourNavMesh->decodeClusterIdCluster(link.ref);
+				const dtCluster& c1 = OtherTile->clusters[linkedIdx];
+
+				FRecastDebugGeometry::FClusterLink LinkGeom;
+				LinkGeom.FromCluster = Recast2UnrealPoint(c0.center);
+				LinkGeom.ToCluster = Recast2UnrealPoint(c1.center);
+
+				if (linkedIdx > i || TileIdx > (int32)DetourNavMesh->decodeClusterIdTile(link.ref))
+				{
+					FVector UpDir(0, 0, 1.0f);
+					FVector LinkDir = (LinkGeom.ToCluster - LinkGeom.FromCluster).GetSafeNormal();
+					FVector SideDir = FVector::CrossProduct(LinkDir, UpDir);
+					LinkGeom.FromCluster += SideDir * 40.0f;
+					LinkGeom.ToCluster += SideDir * 40.0f;
+				}
+
+				OutGeometry.ClusterLinks.Add(LinkGeom);
 			}
- 
- 			AllTilesVertBase += Header->vertCount + Header->detailVertCount;
 		}
- 	}
+	}*/
+
+	// Get tile edges and navmesh edges
+	if (OutGeometry.bGatherPolyEdges || OutGeometry.bGatherNavMeshEdges)
+	{
+		GetDebugPolyEdges(Tile, !!OutGeometry.bGatherPolyEdges, !!OutGeometry.bGatherNavMeshEdges
+			, OutGeometry.PolyEdges, OutGeometry.NavMeshEdges);
+	}
+
+	return Header->vertCount + Header->detailVertCount;
 }
 
 FBox FPImplRecastNavMesh::GetNavMeshBounds() const
diff --git a/Engine/Source/Runtime/Engine/Private/AI/Navigation/RecastNavMesh.cpp b/Engine/Source/Runtime/Engine/Private/AI/Navigation/RecastNavMesh.cpp
--- a/Engine/Source/Runtime/Engine/Private/AI/Navigation/RecastNavMesh.cpp
+++ b/Engine/Source/Runtime/Engine/Private/AI/Navigation/RecastNavMesh.cpp
@@ -231,6 +231,8 @@ ARecastNavMesh::ARecastNavMesh(const FObjectInitializer& ObjectInitializer)
 	LayerPartitioning = ERecastPartitioning::Watershed;
 	RegionChunkSplits = 2;
 	LayerChunkSplits = 2;
+	MaxSimultaneousTileGenerationJobsCount = 1024;
+	bDoFullyAsyncNavDataGathering = false;
 
 #if RECAST_ASYNC_REBUILDING
 	BatchQueryCounter = 0;
@@ -453,7 +455,7 @@ void ARecastNavMesh::UpdateNavMeshDrawing()
 		}
 	}
 #else // DO_NAVMESH_DEBUG_DRAWING_PER_TILE
-	if (RenderingComp != NULL && RenderingComp->bVisible)
+	if (RenderingComp != NULL && RenderingComp->bVisible && UNavMeshRenderingComponent::IsNavigationShowFlagSet(GetWorld()))
 	{
 		RenderingComp->MarkRenderStateDirty();
 	}
@@ -464,7 +466,11 @@ void ARecastNavMesh::UpdateNavMeshDrawing()
 void ARecastNavMesh::CleanUp()
 {
 	Super::CleanUp();
-	NavDataGenerator.Reset();
+	if (NavDataGenerator.IsValid())
+	{
+		NavDataGenerator->CancelBuild();
+		NavDataGenerator.Reset();
+	}
 	DestroyRecastPImpl();
 }
 
@@ -1270,13 +1276,16 @@ void ARecastNavMesh::GetDebugGeometry(FRecastDebugGeometry& OutGeometry, int32 T
 void ARecastNavMesh::RequestDrawingUpdate()
 {
 #if !UE_BUILD_SHIPPING
-	DECLARE_CYCLE_STAT(TEXT("FSimpleDelegateGraphTask.Requesting navmesh redraw"),
+	if (UNavMeshRenderingComponent::IsNavigationShowFlagSet(GetWorld()))
+	{
+		DECLARE_CYCLE_STAT(TEXT("FSimpleDelegateGraphTask.Requesting navmesh redraw"),
 		STAT_FSimpleDelegateGraphTask_RequestingNavmeshRedraw,
-		STATGROUP_TaskGraphTasks);
+			STATGROUP_TaskGraphTasks);
 
-	FSimpleDelegateGraphTask::CreateAndDispatchWhenReady(
-		FSimpleDelegateGraphTask::FDelegate::CreateUObject(this, &ARecastNavMesh::UpdateDrawing),
-		GET_STATID(STAT_FSimpleDelegateGraphTask_RequestingNavmeshRedraw), NULL, ENamedThreads::GameThread);
+		FSimpleDelegateGraphTask::CreateAndDispatchWhenReady(
+			FSimpleDelegateGraphTask::FDelegate::CreateUObject(this, &ARecastNavMesh::UpdateDrawing),
+			GET_STATID(STAT_FSimpleDelegateGraphTask_RequestingNavmeshRedraw), NULL, ENamedThreads::GameThread);
+	}
 #endif // !UE_BUILD_SHIPPING
 }
 
@@ -1478,6 +1487,20 @@ void ARecastNavMesh::SetDefaultForbiddenFlags(uint16 ForbiddenAreaFlags)
 	FPImplRecastNavMesh::SetFilterForbiddenFlags((FRecastQueryFilter*)DefaultQueryFilter->GetImplementation(), ForbiddenAreaFlags);
 }
 
+void ARecastNavMesh::SetMaxSimultaneousTileGenerationJobsCount(int32 NewJobsCountLimit) 
+{
+	const int32 NewCount = NewJobsCountLimit > 0 ? NewJobsCountLimit : 1;
+	if (MaxSimultaneousTileGenerationJobsCount != NewCount)
+	{
+		MaxSimultaneousTileGenerationJobsCount = NewCount;
+		if (GetGenerator() != nullptr)
+		{
+			FRecastNavMeshGenerator* MyGenerator = static_cast<FRecastNavMeshGenerator*>(GetGenerator());
+			MyGenerator->SetMaxTileGeneratorTasks(NewCount);
+		}
+	}
+}
+
 bool ARecastNavMesh::FilterPolys(TArray<NavNodeRef>& PolyRefs, const FRecastQueryFilter* Filter, const UObject* QueryOwner) const
 {
 	bool bSuccess = false;
@@ -1862,7 +1885,7 @@ void ARecastNavMesh::PostEditChangeProperty(FPropertyChangedEvent& PropertyChang
 bool ARecastNavMesh::NeedsRebuild() const
 {
 	bool bLooksLikeNeeded = !RecastNavMeshImpl || RecastNavMeshImpl->GetRecastMesh() == 0;
-	if (NavDataGenerator)
+	if (NavDataGenerator.IsValid())
 	{
 		return bLooksLikeNeeded || NavDataGenerator->GetNumRemaningBuildTasks() > 0;
 	}
@@ -1892,7 +1915,7 @@ void ARecastNavMesh::ConditionalConstructGenerator()
 	const bool bRequiresGenerator = SupportsRuntimeGeneration() || !World->IsGameWorld();
 	if (bRequiresGenerator)
 	{
-		NavDataGenerator.Reset(new FRecastNavMeshGenerator(*this));
+		NavDataGenerator = MakeShareable(new FRecastNavMeshGenerator(*this));
 
 		if (World->GetNavigationSystem())
 		{
diff --git a/Engine/Source/Runtime/Engine/Private/AI/Navigation/RecastNavMeshGenerator.cpp b/Engine/Source/Runtime/Engine/Private/AI/Navigation/RecastNavMeshGenerator.cpp
--- a/Engine/Source/Runtime/Engine/Private/AI/Navigation/RecastNavMeshGenerator.cpp
+++ b/Engine/Source/Runtime/Engine/Private/AI/Navigation/RecastNavMeshGenerator.cpp
@@ -198,6 +198,7 @@ struct FRecastGeometryExport : public FNavigableGeometryExport
 	virtual void ExportPxConvexMesh(physx::PxConvexMesh const * const ConvexMesh, const FTransform& LocalToWorld) override;
 	virtual void ExportPxHeightField(physx::PxHeightField const * const HeightField, const FTransform& LocalToWorld) override;
 #endif // WITH_PHYSX
+	virtual void ExportHeightFieldSlice(const FNavHeightfieldSamples& PrefetchedHeightfieldSamples, const int32 NumRows, const int32 NumCols, const FTransform& LocalToWorld, const FBox& SliceBox) override;
 	virtual void ExportCustomMesh(const FVector* InVertices, int32 NumVerts, const int32* InIndices, int32 NumIndices, const FTransform& LocalToWorld) override;
 	virtual void ExportRigidBodySetup(UBodySetup& BodySetup, const FTransform& LocalToWorld) override;
 	virtual void AddNavModifiers(const FCompositeNavModifier& Modifiers) override;
@@ -268,41 +269,41 @@ static UWorld* FindEditorWorld()
 	return NULL;
 }
 
-static void StoreCollisionCache(FRecastGeometryExport* GeomExport)
+static void StoreCollisionCache(FRecastGeometryExport& GeomExport)
 {
-	const int32 NumFaces = GeomExport->IndexBuffer.Num() / 3;
-	const int32 NumVerts = GeomExport->VertexBuffer.Num() / 3;
+	const int32 NumFaces = GeomExport.IndexBuffer.Num() / 3;
+	const int32 NumVerts = GeomExport.VertexBuffer.Num() / 3;
 
 	if (NumFaces == 0 || NumVerts == 0)
 	{
-		GeomExport->Data->CollisionData.Empty();
+		GeomExport.Data->CollisionData.Empty();
 		return;
 	}
 
 	FRecastGeometryCache::FHeader HeaderInfo;
 	HeaderInfo.NumFaces = NumFaces;
 	HeaderInfo.NumVerts = NumVerts;
-	HeaderInfo.SlopeOverride = GeomExport->SlopeOverride;
+	HeaderInfo.SlopeOverride = GeomExport.SlopeOverride;
 
 	// allocate memory
 	const int32 HeaderSize = sizeof(FRecastGeometryCache);
 	const int32 CoordsSize = sizeof(float) * 3 * NumVerts;
 	const int32 IndicesSize = sizeof(int32) * 3 * NumFaces;
 	const int32 CacheSize = HeaderSize + CoordsSize + IndicesSize;
-	
+
 	// reserve + add combo to allocate exact amount (without any overhead/slack)
-	GeomExport->Data->CollisionData.Reserve(CacheSize);
-	GeomExport->Data->CollisionData.AddUninitialized(CacheSize);
+	GeomExport.Data->CollisionData.Reserve(CacheSize);
+	GeomExport.Data->CollisionData.AddUninitialized(CacheSize);
 
 	// store collisions
-	uint8* RawMemory = GeomExport->Data->CollisionData.GetData();
+	uint8* RawMemory = GeomExport.Data->CollisionData.GetData();
 	FRecastGeometryCache* CacheMemory = (FRecastGeometryCache*)RawMemory;
 	CacheMemory->Header = HeaderInfo;
 	CacheMemory->Verts = 0;
 	CacheMemory->Indices = 0;
 
-	FMemory::Memcpy(RawMemory + HeaderSize, GeomExport->VertexBuffer.GetData(), CoordsSize);
-	FMemory::Memcpy(RawMemory + HeaderSize + CoordsSize, GeomExport->IndexBuffer.GetData(), IndicesSize);
+	FMemory::Memcpy(RawMemory + HeaderSize, GeomExport.VertexBuffer.GetData(), CoordsSize);
+	FMemory::Memcpy(RawMemory + HeaderSize + CoordsSize, GeomExport.IndexBuffer.GetData(), IndicesSize);
 }
 
 #if WITH_PHYSX
@@ -437,40 +438,46 @@ FORCEINLINE_DEBUGGABLE void ExportPxTriMesh(PxTriangleMesh const * const TriMesh
 	}
 }
 
-void ExportPxHeightField(PxHeightField const * const HeightField, const FTransform& LocalToWorld,
-											TNavStatArray<float>& VertexBuffer, TNavStatArray<int32>& IndexBuffer,
-											FBox& UnrealBounds)
+void ExportPxHeightField(PxHeightField const * const HeightField, const FTransform& LocalToWorld
+	, TNavStatArray<float>& VertexBuffer, TNavStatArray<int32>& IndexBuffer
+	, FBox& UnrealBounds)
 {
 	if (HeightField == NULL)
 	{
 		return;
 	}
 
+	QUICK_SCOPE_CYCLE_COUNTER(STAT_NavMesh_ExportPxHeightField);
+
 	const int32 NumRows = HeightField->getNbRows();
 	const int32 NumCols = HeightField->getNbColumns();
-	const int32 NumVtx	= NumRows*NumCols;
+	const int32 VertexCount = NumRows * NumCols;
 
 	// Unfortunately we have to use PxHeightField::saveCells instead PxHeightField::getHeight here 
 	// because current PxHeightField interface does not provide an access to a triangle material index by HF 2D coordinates
 	// PxHeightField::getTriangleMaterialIndex uses some internal adressing which does not match HF 2D coordinates
 	TArray<PxHeightFieldSample> HFSamples;
-	HFSamples.SetNumUninitialized(NumVtx);
-	HeightField->saveCells(HFSamples.GetData(), HFSamples.Num()*HFSamples.GetTypeSize());
-	
+	HFSamples.SetNumUninitialized(VertexCount);
+	{
+		QUICK_SCOPE_CYCLE_COUNTER(STAT_NavMesh_ExportPxHeightField_saveCells);
+		HeightField->saveCells(HFSamples.GetData(), VertexCount*HFSamples.GetTypeSize());
+	}
+
 	//
-	int32 VertOffset = VertexBuffer.Num() / 3;
-	const int32 NumQuads = (NumRows-1)*(NumCols-1);
-	VertexBuffer.Reserve(VertexBuffer.Num() + NumVtx*3);
-	IndexBuffer.Reserve(IndexBuffer.Num() + NumQuads*6);
+	const int32 VertOffset = VertexBuffer.Num() / 3;
+	const int32 NumQuads = (NumRows - 1)*(NumCols - 1);
+
+	VertexBuffer.Reserve(VertexBuffer.Num() + VertexCount * 3);
+	IndexBuffer.Reserve(IndexBuffer.Num() + NumQuads * 6);
 
 	const bool bMirrored = (LocalToWorld.GetDeterminant() < 0.f);
-	
+
 	for (int32 Y = 0; Y < NumRows; Y++)
 	{
 		for (int32 X = 0; X < NumCols; X++)
 		{
-			int32 SampleIdx = (bMirrored ? X : (NumCols - X - 1))*NumCols + Y;
-	
+			const int32 SampleIdx = (bMirrored ? X : (NumCols - X - 1))*NumCols + Y;
+
 			const PxHeightFieldSample& Sample = HFSamples[SampleIdx];
 			const FVector UnrealCoords = LocalToWorld.TransformPosition(FVector(X, Y, Sample.height));
 			UnrealBounds += UnrealCoords;
@@ -480,35 +487,121 @@ void ExportPxHeightField(PxHeightField const * const HeightField, const FTransfo
 			VertexBuffer.Add(UnrealCoords.Z);
 		}
 	}
-		
-	for (int32 Y = 0; Y < NumRows-1; Y++)
+
+	for (int32 Y = 0; Y < NumRows - 1; Y++)
 	{
-		for (int32 X = 0; X < NumCols-1; X++)
+		for (int32 X = 0; X < NumCols - 1; X++)
 		{
-			int32 I00 = X+0 + (Y+0)*NumCols;
-			int32 I01 = X+0 + (Y+1)*NumCols;
-			int32 I10 = X+1 + (Y+0)*NumCols;
-			int32 I11 = X+1 + (Y+1)*NumCols;
+			const int32 I00 = X + 0 + (Y + 0)*NumCols;
+			int32 I01 = X + 0 + (Y + 1)*NumCols;
+			int32 I10 = X + 1 + (Y + 0)*NumCols;
+			const int32 I11 = X + 1 + (Y + 1)*NumCols;
 
 			if (bMirrored)
 			{
 				Swap(I01, I10);
 			}
 
-			int32 SampleIdx = (NumCols - X - 1)*NumCols + Y;
+			const int32 SampleIdx = (NumCols - X - 1)*NumCols + Y;
 			const PxHeightFieldSample& Sample = HFSamples[SampleIdx];
 			const bool HoleQuad = (Sample.materialIndex0 == PxHeightFieldMaterial::eHOLE);
 
 			IndexBuffer.Add(VertOffset + I00);
 			IndexBuffer.Add(VertOffset + (HoleQuad ? I00 : I11));
 			IndexBuffer.Add(VertOffset + (HoleQuad ? I00 : I10));
-			
+
 			IndexBuffer.Add(VertOffset + I00);
 			IndexBuffer.Add(VertOffset + (HoleQuad ? I00 : I01));
 			IndexBuffer.Add(VertOffset + (HoleQuad ? I00 : I11));
 		}
 	}
 }
+
+void ExportHeightFieldSlice(const FNavHeightfieldSamples& PrefetchedHeightfieldSamples, const int32 NumRows, const int32 NumCols, const FTransform& LocalToWorld
+	, TNavStatArray<float>& VertexBuffer, TNavStatArray<int32>& IndexBuffer, const FBox& SliceBox
+	, FBox& UnrealBounds)
+{
+	QUICK_SCOPE_CYCLE_COUNTER(STAT_NavMesh_ExportHeightFieldSlice);
+
+	static const uint32 SizeOfPx = sizeof(physx::PxI16);
+	static const uint32 SizeOfHeight = PrefetchedHeightfieldSamples.Heights.GetTypeSize();
+	ensure(SizeOfPx == SizeOfHeight);
+
+	// calculate the actual start and number of columns we want
+	const FBox LocalBox = SliceBox.TransformBy(LocalToWorld.Inverse());
+
+	const int32 StartingRow = FMath::Max(FMath::FloorToInt(LocalBox.Min.Y), 0);
+	const int32 StartingColumn = FMath::Max(FMath::FloorToInt(LocalBox.Min.X), 0);
+	const int32 RowLimit = FMath::Min(FMath::CeilToInt(LocalBox.Max.Y), NumRows);
+	const int32 ColumnLimit = FMath::Min(FMath::CeilToInt(LocalBox.Max.X), NumCols);
+	
+	const int32 VertexCount = (RowLimit - StartingRow)*(ColumnLimit - StartingColumn);
+
+	const int32 VertOffset = VertexBuffer.Num() / 3;
+	const int32 RowCount = RowLimit - StartingRow;
+	const int32 CollumnCount = ColumnLimit - StartingColumn;
+	const int32 NumQuads = (RowCount - 1)*(CollumnCount - 1);
+
+	{
+		QUICK_SCOPE_CYCLE_COUNTER(STAT_RecastGeometryExport_AllocatingMemory);
+		VertexBuffer.Reserve(VertexBuffer.Num() + VertexCount * 3);
+		IndexBuffer.Reserve(IndexBuffer.Num() + NumQuads * 6);
+	}
+
+	const bool bMirrored = (LocalToWorld.GetDeterminant() < 0.f);
+	
+	TNavStatArray<float> TmpVertexBuffer;
+	TmpVertexBuffer.AddUninitialized(VertexCount * 3);
+	FVector* Vertex = ((FVector*)TmpVertexBuffer.GetData());
+
+	for (int32 Y = StartingRow; Y < RowLimit; Y++)
+	{
+		for (int32 X = StartingColumn; X < ColumnLimit; X++)
+		{
+			const int32 SampleIdx = (bMirrored ? X : (NumCols - X - 1))*NumCols + Y;
+
+			const FVector UnrealCoords = LocalToWorld.TransformPosition(FVector(X, Y, PrefetchedHeightfieldSamples.Heights[SampleIdx]));
+			*Vertex++ = UnrealCoords;
+		}
+	}
+	VertexBuffer.Append(TmpVertexBuffer);
+
+	for (int32 Y = 0; Y < RowCount - 1; Y++)
+	{
+		for (int32 X = 0; X < CollumnCount - 1; X++)
+		{
+			const int32 SampleIdx = (NumCols - X - 1)*NumCols + Y;
+			if (PrefetchedHeightfieldSamples.Holes[SampleIdx] == false)
+			{
+				const int32 I00 = X + 0 + (Y + 0)*CollumnCount;
+				const int32 I01 = X + 0 + (Y + 1)*CollumnCount;
+				const int32 I10 = X + 1 + (Y + 0)*CollumnCount;
+				const int32 I11 = X + 1 + (Y + 1)*CollumnCount;
+
+				if (bMirrored == false)
+				{
+					IndexBuffer.Add(VertOffset + I00);
+					IndexBuffer.Add(VertOffset + I11);
+					IndexBuffer.Add(VertOffset + I10);
+
+					IndexBuffer.Add(VertOffset + I00);
+					IndexBuffer.Add(VertOffset + I01);
+					IndexBuffer.Add(VertOffset + I11);
+				}
+				else
+				{
+					IndexBuffer.Add(VertOffset + I00);
+					IndexBuffer.Add(VertOffset + I01);
+					IndexBuffer.Add(VertOffset + I11);
+
+					IndexBuffer.Add(VertOffset + I00);
+					IndexBuffer.Add(VertOffset + I11);
+					IndexBuffer.Add(VertOffset + I10);
+				}
+			}
+		}
+	}
+}
 #endif //WITH_PHYSX
 
 void ExportCustomMesh(const FVector* InVertices, int32 NumVerts, const int32* InIndices, int32 NumIndices, const FTransform& LocalToWorld,
@@ -844,7 +937,7 @@ FORCEINLINE_DEBUGGABLE void ExportRigidBodySetup(UBodySetup& BodySetup, TNavStat
 	TemporaryShapeBuffer.Reset();
 }
 
-FORCEINLINE_DEBUGGABLE void ExportComponent(UActorComponent* Component, FRecastGeometryExport* GeomExport, const FBox* ClipBounds=NULL)
+FORCEINLINE_DEBUGGABLE void ExportComponent(UActorComponent* Component, FRecastGeometryExport& GeomExport, const FBox* ClipBounds=NULL)
 {
 #if WITH_PHYSX
 	bool bHasData = false;
@@ -862,11 +955,11 @@ FORCEINLINE_DEBUGGABLE void ExportComponent(UActorComponent* Component, FRecastG
 		{
 			if (!bHasData)
 			{
-				ExportRigidBodySetup(*BodySetup, GeomExport->VertexBuffer, GeomExport->IndexBuffer, GeomExport->Data->Bounds, PrimComp->ComponentToWorld);
+				ExportRigidBodySetup(*BodySetup, GeomExport.VertexBuffer, GeomExport.IndexBuffer, GeomExport.Data->Bounds, PrimComp->ComponentToWorld);
 				bHasData = true;
 			}
 
-			GeomExport->SlopeOverride = BodySetup->WalkableSlopeOverride;
+			GeomExport.SlopeOverride = BodySetup->WalkableSlopeOverride;
 		}
 	}
 #endif // WITH_PHYSX
@@ -969,6 +1062,11 @@ void FRecastGeometryExport::ExportPxHeightField(physx::PxHeightField const * con
 {
 	RecastGeometryExport::ExportPxHeightField(HeightField, LocalToWorld, VertexBuffer, IndexBuffer, Data->Bounds);
 }
+
+void FRecastGeometryExport::ExportHeightFieldSlice(const FNavHeightfieldSamples& PrefetchedHeightfieldSamples, const int32 NumRows, const int32 NumCols, const FTransform& LocalToWorld, const FBox& SliceBox)
+{
+	RecastGeometryExport::ExportHeightFieldSlice(PrefetchedHeightfieldSamples, NumRows, NumCols, LocalToWorld, VertexBuffer, IndexBuffer, SliceBox, Data->Bounds);
+}
 #endif // WITH_PHYSX
 
 void FRecastGeometryExport::ExportCustomMesh(const FVector* InVertices, int32 NumVerts, const int32* InIndices, int32 NumIndices, const FTransform& LocalToWorld)
@@ -1433,9 +1531,10 @@ static FBox CalculateTileBounds(int32 X, int32 Y, const FVector& NavMeshOrigin,
 // FRecastTileGenerator
 //----------------------------------------------------------------------//
 
-FRecastTileGenerator::FRecastTileGenerator(const FRecastNavMeshGenerator& ParentGenerator, const FIntPoint& Location)
+FRecastTileGenerator::FRecastTileGenerator(FRecastNavMeshGenerator& ParentGenerator, const FIntPoint& Location)
 {
 	bSucceeded = false;
+	bUpdateGeometry = true;
 
 	TileX = Location.X;
 	TileY = Location.Y;
@@ -1443,6 +1542,8 @@ FRecastTileGenerator::FRecastTileGenerator(const FRecastNavMeshGenerator& Parent
 	TileConfig = ParentGenerator.GetConfig();
 	Version = ParentGenerator.GetVersion();
 	AdditionalCachedData = ParentGenerator.GetAdditionalCachedData();
+
+	ParentGeneratorWeakPtr = ParentGenerator.AsShared();
 }
 
 FRecastTileGenerator::~FRecastTileGenerator()
@@ -1510,8 +1611,15 @@ void FRecastTileGenerator::Setup(const FRecastNavMeshGenerator& ParentGenerator,
 				}
 			}
 		}
-
-		GatherGeometry(ParentGenerator, bRegenerateCompressedLayers);
+		
+		if (ParentGenerator.GatherGeometryOnGameThread())
+		{
+			GatherGeometry(ParentGenerator, bRegenerateCompressedLayers);
+		}
+		else
+		{
+			PrepareGeometrySources(ParentGenerator, bRegenerateCompressedLayers);
+		}
 	}
 	
 	//
@@ -1524,18 +1632,159 @@ bool FRecastTileGenerator::HasDataToBuild() const
 		CompressedLayers.Num()
 		|| Modifiers.Num()
 		|| OffmeshLinks.Num()
-		|| RawGeometry.Num();
+		|| RawGeometry.Num()
+		|| InclusionBounds.Num();
 }
 
 void FRecastTileGenerator::DoWork()
 {
+	TSharedPtr<FNavDataGenerator> ParentGenerator = ParentGeneratorWeakPtr.Pin();
+	FRecastNavMeshGenerator* RecastParentGenerator = static_cast<FRecastNavMeshGenerator*>(ParentGenerator.Get());
+	check(ParentGenerator.Get());
+
+	if (InclusionBounds.Num())
+	{
+		DoAsyncGeometryGathering();
+	}
+
 	bSucceeded = GenerateTile();
+
+	DumpAsyncData();
+}
+
+void FRecastTileGenerator::DumpAsyncData()
+{
+	RawGeometry.Empty();
+	Modifiers.Empty();
+	OffmeshLinks.Empty();
+
+	NavigationRelevantData.Empty();
+	NavOctree = nullptr;
+}
+
+void FRecastTileGenerator::DoAsyncGeometryGathering()
+{
+	QUICK_SCOPE_CYCLE_COUNTER(STAT_RecastNavMeshGenerator_PrepareGeometrySources);
+	
+	for (auto& ElementData : NavigationRelevantData)
+	{
+		if (ElementData->GetOwner() == nullptr)
+		{
+			UE_LOG(LogNavigation, Warning, TEXT("DoAsyncGeometryGathering: skipping an element with no longer valid Owner"));
+			continue;
+		}
+
+		bool bDumpGeometryData = false;
+		if (ElementData->IsPendingLazyGeometryGathering() && ElementData->SupportsGatheringGeometrySlices())
+		{
+			QUICK_SCOPE_CYCLE_COUNTER(STAT_RecastNavMeshGenerator_LandscapeSlicesExporting);
+
+			FRecastGeometryExport GeomExport(*ElementData);
+
+			INavRelevantInterface* NavRelevant = Cast<INavRelevantInterface>(ElementData->GetOwner());
+			if(NavRelevant)
+			{
+				NavRelevant->PrepareGeometryExportSync();
+				NavRelevant->GatherGeometrySlice(GeomExport, TileBB.ExpandBy(NavDataConfig.AgentRadius * 2));
+
+				RecastGeometryExport::CovertCoordDataToRecast(GeomExport.VertexBuffer);
+				RecastGeometryExport::StoreCollisionCache(GeomExport);
+				bDumpGeometryData = true;
+			}
+			else
+			{
+				UE_LOG(LogNavigation, Error, TEXT("DoAsyncGeometryGathering: got an invalid NavRelevant instance!"));
+			}
+		}
+
+		if (ElementData->IsPendingLazyGeometryGathering() || ElementData->IsPendingLazyModifiersGathering())
+		{
+			NavOctree->DemandLazyDataGathering(*ElementData);
+		}
+
+		const bool bExportGeometry = bUpdateGeometry && ElementData->HasGeometry();
+		if (bExportGeometry)
+		{
+			if (ARecastNavMesh::IsVoxelCacheEnabled())
+			{
+				TNavStatArray<rcSpanCache> SpanData;
+				rcSpanCache* CachedVoxels = 0;
+				int32 NumCachedVoxels = 0;
+
+				DECLARE_SCOPE_CYCLE_COUNTER(TEXT("Rasterization: prepare voxel cache"), Stat_RecastRasterCachePrep, STATGROUP_Navigation);
+
+				if (!HasVoxelCache(ElementData->VoxelData, CachedVoxels, NumCachedVoxels))
+				{
+					// rasterize
+					PrepareVoxelCache(ElementData->CollisionData, SpanData);
+					CachedVoxels = SpanData.GetData();
+					NumCachedVoxels = SpanData.Num();
+
+					// encode
+					const int32 PrevElementMemory = ElementData->GetAllocatedSize();
+					FNavigationRelevantData* ModData = (FNavigationRelevantData*)&ElementData;
+					AddVoxelCache(ModData->VoxelData, CachedVoxels, NumCachedVoxels);
+
+					const int32 NewElementMemory = ElementData->GetAllocatedSize();
+					const int32 ElementMemoryDelta = NewElementMemory - PrevElementMemory;
+					INC_MEMORY_STAT_BY(STAT_Navigation_CollisionTreeMemory, ElementMemoryDelta);
+				}
+			}
+			else
+			{
+				AppendGeometry(ElementData->CollisionData, ElementData->NavDataPerInstanceTransformDelegate);
+			}
+
+			if (bDumpGeometryData)
+			{
+				const_cast<FNavigationRelevantData&>(*ElementData).CollisionData.Empty();
+			}
+		}
+
+		const FCompositeNavModifier ModifierInstance = ElementData->Modifiers.HasMetaAreas() ? ElementData->Modifiers.GetInstantiatedMetaModifier(&NavDataConfig, ElementData->SourceObject) : ElementData->Modifiers;
+		if (ModifierInstance.IsEmpty() == false)
+		{
+			AppendModifier(ModifierInstance, ElementData->NavDataPerInstanceTransformDelegate);
+		}
+	}
+}
+
+void FRecastTileGenerator::PrepareGeometrySources(const FRecastNavMeshGenerator& ParentGenerator, bool bGeometryChanged)
+{
+	QUICK_SCOPE_CYCLE_COUNTER(STAT_RecastNavMeshGenerator_PrepareGeometrySources);
+
+	UNavigationSystem*	NavSys = UNavigationSystem::GetCurrent(ParentGenerator.GetWorld());
+	FNavigationOctree*	NavOctreeInstance = NavSys ? NavSys->GetMutableNavOctree() : nullptr;
+	check(NavOctreeInstance);
+	NavDataConfig = ParentGenerator.GetOwner()->GetConfig();
+	NavigationRelevantData.Reset();
+	NavOctree = NavOctreeInstance->AsShared();
+	bUpdateGeometry = bGeometryChanged;
+
+	for (FNavigationOctree::TConstElementBoxIterator<FNavigationOctree::DefaultStackAllocator> It(*NavOctreeInstance, ParentGenerator.GrowBoundingBox(TileBB, /*bIncludeAgentHeight*/ false));
+		It.HasPendingElements();
+		It.Advance())
+	{
+		const FNavigationOctreeElement& Element = It.GetCurrentElement();
+		const bool bShouldUse = Element.ShouldUseGeometry(NavDataConfig);
+		if (bShouldUse)
+		{
+			const bool bExportGeometry = bGeometryChanged && (Element.Data->HasGeometry() || Element.Data->IsPendingLazyGeometryGathering());
+			if (bExportGeometry
+				|| (Element.Data->IsPendingLazyModifiersGathering() || Element.Data->Modifiers.HasMetaAreas() == true || Element.Data->Modifiers.IsEmpty() == false))
+			{
+				NavigationRelevantData.Add(Element.Data);
+			}
+		}
+	}
 }
 
 void FRecastTileGenerator::GatherGeometry(const FRecastNavMeshGenerator& ParentGenerator, bool bGeometryChanged)
 {
-	const UNavigationSystem*	NavSys = UNavigationSystem::GetCurrent(ParentGenerator.GetWorld());
-	const FNavigationOctree*	NavOctree = NavSys ? NavSys->GetNavOctree() : nullptr;
+	QUICK_SCOPE_CYCLE_COUNTER(STAT_RecastNavMeshGenerator_GatherGeometry);
+
+	UNavigationSystem*	NavSys = UNavigationSystem::GetCurrent(ParentGenerator.GetWorld());
+	FNavigationOctree*	NavOctree = NavSys ? NavSys->GetMutableNavOctree() : nullptr;
 	check(NavOctree);
 	const FNavDataConfig&		NavDataConfig = ParentGenerator.GetOwner()->GetConfig();
 
@@ -1547,12 +1796,40 @@ void FRecastTileGenerator::GatherGeometry(const FRecastNavMeshGenerator& ParentG
 		const bool bShouldUse = Element.ShouldUseGeometry(NavDataConfig);
 		if (bShouldUse)
 		{
-			if (Element.Data.IsPendingLazyGeometryGathering())
+			bool bDumpGeometryData = false;
+			if (Element.Data->IsPendingLazyGeometryGathering() || Element.Data->IsPendingLazyModifiersGathering())
 			{
-				NavOctree->DemandGatheringGeometry(Element);
+				const bool bSupportsSlices = Element.Data->SupportsGatheringGeometrySlices();
+				if (bSupportsSlices == false || Element.Data->IsPendingLazyModifiersGathering() == true)
+				{
+					QUICK_SCOPE_CYCLE_COUNTER(STAT_RecastNavMeshGenerator_LazyGeometryExport);
+					NavOctree->DemandLazyDataGathering(Element);
+				}
+				
+				if (bSupportsSlices == true)
+				{
+					QUICK_SCOPE_CYCLE_COUNTER(STAT_RecastNavMeshGenerator_LandscapeSlicesExporting);
+
+					FRecastGeometryExport GeomExport(const_cast<FNavigationRelevantData&>(*Element.Data));
+
+					INavRelevantInterface* NavRelevant = const_cast<INavRelevantInterface*>(Cast<const INavRelevantInterface>(Element.GetOwner()));
+					if (NavRelevant)
+					{
+						NavRelevant->PrepareGeometryExportSync();
+						NavRelevant->GatherGeometrySlice(GeomExport, TileBB.ExpandBy(NavDataConfig.AgentRadius * 2));
+
+						RecastGeometryExport::CovertCoordDataToRecast(GeomExport.VertexBuffer);
+						RecastGeometryExport::StoreCollisionCache(GeomExport);
+						bDumpGeometryData = true;
+					}
+					else
+					{
+						UE_LOG(LogNavigation, Error, TEXT("GatherGeometry: got an invalid NavRelevant instance!"));
+					}
+				}
 			}
 
-			const bool bExportGeometry = bGeometryChanged && Element.Data.HasGeometry();
+			const bool bExportGeometry = bGeometryChanged && Element.Data->HasGeometry();
 			if (bExportGeometry)
 			{
 				if (ARecastNavMesh::IsVoxelCacheEnabled())
@@ -1563,31 +1840,39 @@ void FRecastTileGenerator::GatherGeometry(const FRecastNavMeshGenerator& ParentG
 
 					DECLARE_SCOPE_CYCLE_COUNTER(TEXT("Rasterization: prepare voxel cache"), Stat_RecastRasterCachePrep, STATGROUP_Navigation);
 
-					if (!HasVoxelCache(Element.Data.VoxelData, CachedVoxels, NumCachedVoxels))
+					if (!HasVoxelCache(Element.Data->VoxelData, CachedVoxels, NumCachedVoxels))
 					{
 						// rasterize
-						PrepareVoxelCache(Element.Data.CollisionData, SpanData);
+						PrepareVoxelCache(Element.Data->CollisionData, SpanData);
 						CachedVoxels = SpanData.GetData();
 						NumCachedVoxels = SpanData.Num();
 
 						// encode
-						const int32 PrevElementMemory = Element.Data.GetAllocatedSize();
+						const int32 PrevElementMemory = Element.Data->GetAllocatedSize();
 						FNavigationRelevantData* ModData = (FNavigationRelevantData*)&Element.Data;
 						AddVoxelCache(ModData->VoxelData, CachedVoxels, NumCachedVoxels);
 
-						const int32 NewElementMemory = Element.Data.GetAllocatedSize();
+						const int32 NewElementMemory = Element.Data->GetAllocatedSize();
 						const int32 ElementMemoryDelta = NewElementMemory - PrevElementMemory;
 						INC_MEMORY_STAT_BY(STAT_Navigation_CollisionTreeMemory, ElementMemoryDelta);
 					}
 				}
 				else
 				{
-					AppendGeometry(Element.Data.CollisionData, Element.Data.NavDataPerInstanceTransformDelegate);
+					AppendGeometry(Element.Data->CollisionData, Element.Data->NavDataPerInstanceTransformDelegate);
 				}
-			}
 
+				if (bDumpGeometryData)
+				{
+					const_cast<FNavigationRelevantData&>(*Element.Data).CollisionData.Empty();
+				}
+			}
+						
 			const FCompositeNavModifier ModifierInstance = Element.GetModifierForAgent(&NavDataConfig);
-			AppendModifier(ModifierInstance, Element.Data.NavDataPerInstanceTransformDelegate);
+			if (ModifierInstance.IsEmpty() == false)
+			{
+				AppendModifier(ModifierInstance, Element.Data->NavDataPerInstanceTransformDelegate);
+			}
 		}
 	}
 }
@@ -2668,6 +2953,18 @@ uint32 FRecastTileGenerator::GetUsedMemCount() const
 	return TotalMemory;
 }
 
+void FRecastTileGenerator::AddReferencedObjects(FReferenceCollector& Collector)
+{
+	for (auto& RelevantData : NavigationRelevantData)
+	{
+		UObject* Owner = RelevantData->GetOwner();
+		if (Owner)
+		{
+			Collector.AddReferencedObject(Owner);
+		}
+	}
+}
+
 static int32 CaclulateMaxTilesCount(const TNavStatArray<FBox>& NavigableAreas, float TileSizeinWorldUnits, float AvgLayersPerGridCell)
 {
 	int32 GridCellsCount = 0;
@@ -2755,8 +3052,6 @@ FRecastNavMeshGenerator::FRecastNavMeshGenerator(ARecastNavMesh& InDestNavMesh)
 
 FRecastNavMeshGenerator::~FRecastNavMeshGenerator()
 {
-	DiscardCurrentBuildingTasks();
-	
 	DEC_DWORD_STAT_BY( STAT_NavigationMemory, sizeof(*this) );
 }
 
@@ -2824,7 +3119,7 @@ void FRecastNavMeshGenerator::Init()
 
 	/** setup maximum number of active tile generator*/
 	const int32 NumberOfWorkerThreads = FTaskGraphInterface::Get().GetNumWorkerThreads();
-	MaxTileGeneratorTasks = FMath::Max(NumberOfWorkerThreads*2, 1);
+	MaxTileGeneratorTasks = FMath::Min(FMath::Max(NumberOfWorkerThreads * 2, 1), GetOwner() ? GetOwner()->GetMaxSimultaneousTileGenerationJobsCount() : INT_MAX);
 	UE_LOG(LogNavigation, Log, TEXT("Using max of %d workers to build navigation."), MaxTileGeneratorTasks);
 	NumActiveTiles = 0;
 
@@ -2993,8 +3288,6 @@ void FRecastNavMeshGenerator::EnsureBuildCompletion()
 void FRecastNavMeshGenerator::CancelBuild()
 {
 	DiscardCurrentBuildingTasks();
-	RunningDirtyTiles.Empty();
-	PendingDirtyTiles.Empty();
 
 #if	WITH_EDITOR	
 	RecentlyBuiltTiles.Empty();
@@ -3022,7 +3315,11 @@ void FRecastNavMeshGenerator::TickAsyncBuild(float DeltaSeconds)
 	const UNavigationSystem* NavSys = UNavigationSystem::GetCurrent(GetWorld());
 	check(NavSys);
 	const int32 NumRunningTasks = NavSys->GetNumRunningBuildTasks();
-	const int32 NumTasksToSubmit = MaxTileGeneratorTasks - NumRunningTasks;
+	// this is a temp solution to enforce only one worker thread if GatherGeometryOnGameThread == false
+	// due to missing safety features
+	const bool bDoAsyncDataGathering = GatherGeometryOnGameThread() == false;
+
+	const int32 NumTasksToSubmit = (bDoAsyncDataGathering ? 1 : MaxTileGeneratorTasks) - NumRunningTasks;
 	TArray<uint32> UpdatedTileIndices = ProcessTileTasks(NumTasksToSubmit);
 			
 	if (UpdatedTileIndices.Num() > 0)
@@ -3258,7 +3555,7 @@ TArray<uint32> FRecastNavMeshGenerator::RemoveTileLayers(const int32 TileX, cons
 	return ResultTileIndices;
 }
 
-TArray<uint32> FRecastNavMeshGenerator::AddGeneratedTiles(const FRecastTileGenerator& TileGenerator)
+TArray<uint32> FRecastNavMeshGenerator::AddGeneratedTiles(FRecastTileGenerator& TileGenerator)
 {
 	SCOPE_CYCLE_COUNTER(STAT_Navigation_AddGeneratedTiles);
 	
@@ -3293,6 +3590,8 @@ TArray<uint32> FRecastNavMeshGenerator::AddGeneratedTiles(const FRecastTileGener
 
 			if (OldTileRef)
 			{
+				QUICK_SCOPE_CYCLE_COUNTER(STAT_NavMesh_RemoveOldTIle);
+
 				NumActiveTiles--;
 				UE_LOG(LogNavigation, Log, TEXT("%s> Tile (%d,%d:%d), removing TileRef: 0x%X (active:%d)"),
 					*DestNavMesh->GetName(), TileX, TileY, LayerIndex, OldTileRef, NumActiveTiles);
@@ -3307,8 +3606,13 @@ TArray<uint32> FRecastNavMeshGenerator::AddGeneratedTiles(const FRecastTileGener
 				bool bRejectNavmesh = false;
 				dtTileRef ResultTileRef = 0;
 
-				// let navmesh know it's tile generator who owns the data
-				dtStatus status = DetourMesh->addTile(TileLayers[i].GetData(), TileLayers[i].DataSize, DT_TILE_FREE_DATA, OldTileRef, &ResultTileRef);
+				dtStatus status = 0;
+
+				{
+					QUICK_SCOPE_CYCLE_COUNTER(STAT_NavMesh_AddTileToDetourMesh);
+					// let navmesh know it's tile generator who owns the data
+					status = DetourMesh->addTile(TileLayers[i].GetData(), TileLayers[i].DataSize, DT_TILE_FREE_DATA, OldTileRef, &ResultTileRef);
+				}
 
 				if (dtStatusFailed(status))
 				{
@@ -3327,9 +3631,12 @@ TArray<uint32> FRecastNavMeshGenerator::AddGeneratedTiles(const FRecastTileGener
 
 					UE_LOG(LogNavigation, Log, TEXT("%s> Tile (%d,%d:%d), added TileRef: 0x%X (active:%d)"),
 						*DestNavMesh->GetName(), TileX, TileY, LayerIndex, ResultTileRef, NumActiveTiles);
-				
-					// NavMesh took the ownership of generated data, so we don't need to deallocate it
-					uint8* ReleasedData = TileLayers[i].Release();
+
+					{
+						QUICK_SCOPE_CYCLE_COUNTER(STAT_NavMesh_ReleaseTileLayers);
+						// NavMesh took the ownership of generated data, so we don't need to deallocate it
+						uint8* ReleasedData = TileLayers[i].Release();
+					}
 				}
 			}
 		}
@@ -3610,6 +3917,8 @@ TArray<uint32> FRecastNavMeshGenerator::ProcessTileTasks(const int32 NumTasksToS
 	// Submit pending tile elements
 	for (int32 ElementIdx = PendingDirtyTiles.Num()-1; ElementIdx >= 0 && NumSubmittedTasks < NumTasksToSubmit; ElementIdx--)
 	{
+		QUICK_SCOPE_CYCLE_COUNTER(STAT_RecastNavMeshGenerator_ProcessTileTasks_NewTasks);
+
 		FPendingTileElement& PendingElement = PendingDirtyTiles[ElementIdx];
 		FRunningTileElement RunningElement(PendingElement.Coord);
 		
@@ -3653,6 +3962,8 @@ TArray<uint32> FRecastNavMeshGenerator::ProcessTileTasks(const int32 NumTasksToS
 	// Collect completed tasks and apply generated data to navmesh
 	for (int32 Idx = RunningDirtyTiles.Num() - 1; Idx >=0; --Idx)
 	{
+		QUICK_SCOPE_CYCLE_COUNTER(STAT_RecastNavMeshGenerator_ProcessTileTasks_FinishedTasks);
+
 		FRunningTileElement& Element = RunningDirtyTiles[Idx];
 		check(Element.AsyncTask);
 
@@ -3661,23 +3972,27 @@ TArray<uint32> FRecastNavMeshGenerator::ProcessTileTasks(const int32 NumTasksToS
 			// Add generated tiles to navmesh
 			if (!Element.bShouldDiscard)
 			{
-				const FRecastTileGenerator& TileGenerator = *(Element.AsyncTask->GetTask().TileGenerator);
+				FRecastTileGenerator& TileGenerator = *(Element.AsyncTask->GetTask().TileGenerator);
 				TArray<uint32> UpdatedTileIndices = AddGeneratedTiles(TileGenerator);
 				UpdatedTiles.Append(UpdatedTileIndices);
 			
 				// Store compressed tile cache layers so it can be reused later
-				TArray<FNavMeshTileData> CompressedLayers = TileGenerator.GetCompressedLayers();
-				if (CompressedLayers.Num())
+				if (TileGenerator.GetCompressedLayers().Num())
 				{
-					DestNavMesh->AddTileCacheLayers(Element.Coord.X, Element.Coord.Y, CompressedLayers);
+					QUICK_SCOPE_CYCLE_COUNTER(STAT_RecastNavMeshGenerator_StoringCompressedLayers);
+					DestNavMesh->AddTileCacheLayers(Element.Coord.X, Element.Coord.Y, TileGenerator.GetCompressedLayers());
 				}
 			}
 
-			// Destroy tile generator task
-			delete Element.AsyncTask;
-			Element.AsyncTask = nullptr;
-			// Remove completed tile element from a list of running tasks
-			RunningDirtyTiles.RemoveAt(Idx);
+			{
+				QUICK_SCOPE_CYCLE_COUNTER(STAT_RecastNavMeshGenerator_TileGeneratorRemoval);
+
+				// Destroy tile generator task
+				delete Element.AsyncTask;
+				Element.AsyncTask = nullptr;
+				// Remove completed tile element from a list of running tasks
+				RunningDirtyTiles.RemoveAtSwap(Idx, 1, false);
+			}
 		}
 	}
 
@@ -3694,16 +4009,16 @@ TArray<uint32> FRecastNavMeshGenerator::ProcessTileTasks(const int32 NumTasksToS
 void FRecastNavMeshGenerator::ExportComponentGeometry(UActorComponent* Component, FNavigationRelevantData& Data)
 {
 	FRecastGeometryExport GeomExport(Data);
-	RecastGeometryExport::ExportComponent(Component, &GeomExport);
+	RecastGeometryExport::ExportComponent(Component, GeomExport);
 	RecastGeometryExport::CovertCoordDataToRecast(GeomExport.VertexBuffer);
-	RecastGeometryExport::StoreCollisionCache(&GeomExport);
+	RecastGeometryExport::StoreCollisionCache(GeomExport);
 }
 
 void FRecastNavMeshGenerator::ExportVertexSoupGeometry(const TArray<FVector>& Verts, FNavigationRelevantData& Data)
 {
 	FRecastGeometryExport GeomExport(Data);
 	RecastGeometryExport::ExportVertexSoup(Verts, GeomExport.VertexBuffer, GeomExport.IndexBuffer, GeomExport.Data->Bounds);
-	RecastGeometryExport::StoreCollisionCache(&GeomExport);
+	RecastGeometryExport::StoreCollisionCache(GeomExport);
 }
 
 void FRecastNavMeshGenerator::ExportRigidBodyGeometry(UBodySetup& BodySetup, TNavStatArray<FVector>& OutVertexBuffer, TNavStatArray<int32>& OutIndexBuffer, const FTransform& LocalToWorld)
@@ -3771,6 +4086,11 @@ int32 FRecastNavMeshGenerator::GetNumRunningBuildTasks() const
 	return RunningDirtyTiles.Num();
 }
 
+bool FRecastNavMeshGenerator::GatherGeometryOnGameThread() const 
+{ 
+	return DestNavMesh == nullptr || DestNavMesh->ShouldGatherDataOnGameThread() == true;
+}
+
 bool FRecastNavMeshGenerator::IsTileChanged(int32 TileIdx) const
 {
 #if WITH_EDITOR	
@@ -3838,11 +4158,11 @@ void FRecastNavMeshGenerator::ExportNavigationData(const FString& FileName) cons
 				It.Advance())
 			{
 				const FNavigationOctreeElement& Element = It.GetCurrentElement();
-				const bool bExportGeometry = Element.Data.HasGeometry() && Element.ShouldUseGeometry(DestNavMesh->GetConfig());
+				const bool bExportGeometry = Element.Data->HasGeometry() && Element.ShouldUseGeometry(DestNavMesh->GetConfig());
 
-				if (bExportGeometry && Element.Data.CollisionData.Num())
+				if (bExportGeometry && Element.Data->CollisionData.Num())
 				{
-					FRecastGeometryCache CachedGeometry(Element.Data.CollisionData.GetData());
+					FRecastGeometryCache CachedGeometry(Element.Data->CollisionData.GetData());
 					IndexBuffer.Reserve( IndexBuffer.Num() + (CachedGeometry.Header.NumFaces * 3 ));
 					CoordBuffer.Reserve( CoordBuffer.Num() + (CachedGeometry.Header.NumVerts * 3 ));
 					for (int32 i = 0; i < CachedGeometry.Header.NumFaces * 3; i++)
@@ -3856,7 +4176,7 @@ void FRecastNavMeshGenerator::ExportNavigationData(const FString& FileName) cons
 				}
 				else
 				{
-					const TArray<FAreaNavModifier>& AreaMods = Element.Data.Modifiers.GetAreas();
+					const TArray<FAreaNavModifier>& AreaMods = Element.Data->Modifiers.GetAreas();
 					for (int32 i = 0; i < AreaMods.Num(); i++)
 					{
 						FAreaExportData ExportInfo;
@@ -4044,5 +4364,3 @@ public:
 } NavigationGeomExec;
 
 #endif // WITH_RECAST
-
-
diff --git a/Engine/Source/Runtime/Engine/Private/Actor.cpp b/Engine/Source/Runtime/Engine/Private/Actor.cpp
--- a/Engine/Source/Runtime/Engine/Private/Actor.cpp
+++ b/Engine/Source/Runtime/Engine/Private/Actor.cpp
@@ -3260,6 +3260,9 @@ bool AActor::IncrementalRegisterComponents(int32 NumComponentsToRegister)
 	// Register RootComponent first so all other components can reliable use it (ie call GetLocation) when they register
 	if( RootComponent != NULL && !RootComponent->IsRegistered() )
 	{
+#if PERF_TRACK_DETAILED_ASYNC_STATS
+		FScopeCycleCounterUObject ContextScope(RootComponent);
+#endif
 		// An unregistered root component is bad news
 		check(RootComponent->bAutoRegister);
 
@@ -3290,6 +3293,9 @@ bool AActor::IncrementalRegisterComponents(int32 NumComponentsToRegister)
 				CompIdx--;
 				NumTotalRegisteredComponents--; // because we will try to register the parent again later...
 			}
+#if PERF_TRACK_DETAILED_ASYNC_STATS
+			FScopeCycleCounterUObject ContextScope(Component);
+#endif
 				
 			//Before we register our component, save it to our transaction buffer so if "undone" it will return to an unregistered state.
 			//This should prevent unwanted components hanging around when undoing a copy/paste or duplication action.
@@ -3306,6 +3312,9 @@ bool AActor::IncrementalRegisterComponents(int32 NumComponentsToRegister)
 	// See whether we are done
 	if (Components.Num() == NumTotalRegisteredComponents)
 	{
+#if PERF_TRACK_DETAILED_ASYNC_STATS
+		QUICK_SCOPE_CYCLE_COUNTER(STAT_AActor_IncrementalRegisterComponents_PostRegisterAllComponents);
+#endif
 		// Finally, call PostRegisterAllComponents
 		PostRegisterAllComponents();
 		return true;
diff --git a/Engine/Source/Runtime/Engine/Private/ActorComponent.cpp b/Engine/Source/Runtime/Engine/Private/ActorComponent.cpp
--- a/Engine/Source/Runtime/Engine/Private/ActorComponent.cpp
+++ b/Engine/Source/Runtime/Engine/Private/ActorComponent.cpp
@@ -9,6 +9,7 @@
 #include "UObjectToken.h"
 #include "MapErrors.h"
 #include "ComponentReregisterContext.h"
+#include "ComponentRecreateRenderStateContext.h"
 #include "Engine/SimpleConstructionScript.h"
 
 #define LOCTEXT_NAMESPACE "ActorComponent"
@@ -90,7 +91,6 @@ FGlobalComponentReregisterContext::FGlobalComponentReregisterContext(const TArra
 	}
 }
 
-
 FGlobalComponentReregisterContext::~FGlobalComponentReregisterContext()
 {
 	check(ActiveGlobalReregisterContextCount > 0);
@@ -99,6 +99,24 @@ FGlobalComponentReregisterContext::~FGlobalComponentReregisterContext()
 	ActiveGlobalReregisterContextCount--;
 }
 
+FGlobalComponentRecreateRenderStateContext::FGlobalComponentRecreateRenderStateContext()
+{
+	// wait until resources are released
+	FlushRenderingCommands();
+
+	// recreate render state for all components.
+	for (auto* Component : TObjectRange<UActorComponent>())
+	{
+		new(ComponentContexts) FComponentRecreateRenderStateContext(Component);
+	}
+}
+
+FGlobalComponentRecreateRenderStateContext::~FGlobalComponentRecreateRenderStateContext()
+{
+	ComponentContexts.Empty();
+}
+
+
 UActorComponent::UActorComponent(const FObjectInitializer& ObjectInitializer /*= FObjectInitializer::Get()*/)
 	: Super(ObjectInitializer)
 {
@@ -129,14 +147,14 @@ void UActorComponent::PostLoad()
 
 	if (GetLinkerUE4Version() < VER_UE4_ACTOR_COMPONENT_CREATION_METHOD)
 	{
-		if (bCreatedByConstructionScript_DEPRECATED)
-		{
+	if (bCreatedByConstructionScript_DEPRECATED)
+	{
 			CreationMethod = EComponentCreationMethod::SimpleConstructionScript;
-		}
-		else if (bInstanceComponent_DEPRECATED)
-		{
-			CreationMethod = EComponentCreationMethod::Instance;
-		}
+	}
+	else if (bInstanceComponent_DEPRECATED)
+	{
+		CreationMethod = EComponentCreationMethod::Instance;
+	}
 
 		if (CreationMethod == EComponentCreationMethod::SimpleConstructionScript)
 		{
diff --git a/Engine/Source/Runtime/Engine/Private/Animation/AnimInstance.cpp b/Engine/Source/Runtime/Engine/Private/Animation/AnimInstance.cpp
--- a/Engine/Source/Runtime/Engine/Private/Animation/AnimInstance.cpp
+++ b/Engine/Source/Runtime/Engine/Private/Animation/AnimInstance.cpp
@@ -34,6 +34,7 @@ DEFINE_STAT(STAT_UpdateLocalToWorldAndOverlaps);
 DEFINE_STAT(STAT_SkelComposeTime);
 DEFINE_STAT(STAT_GetAnimationPose);
 DEFINE_STAT(STAT_AnimNativeEvaluatePoses);
+DEFINE_STAT(STAT_AnimTriggerAnimNotifies);
 DEFINE_STAT(STAT_AnimNativeBlendPoses);
 DEFINE_STAT(STAT_AnimNativeCopyPoses);
 DEFINE_STAT(STAT_AnimGraphEvaluate);
@@ -1285,6 +1286,7 @@ void UAnimInstance::AddCurveValue(const USkeleton::AnimCurveUID Uid, float Value
 
 void UAnimInstance::TriggerAnimNotifies(float DeltaSeconds)
 {
+	SCOPE_CYCLE_COUNTER(STAT_AnimTriggerAnimNotifies);
 	USkeletalMeshComponent * SkelMeshComp = GetSkelMeshComponent();
 
 	// Array that will replace the 'ActiveAnimNotifyState' at the end of this function.
diff --git a/Engine/Source/Runtime/Engine/Private/Animation/AnimNode_StateMachine.cpp b/Engine/Source/Runtime/Engine/Private/Animation/AnimNode_StateMachine.cpp
--- a/Engine/Source/Runtime/Engine/Private/Animation/AnimNode_StateMachine.cpp
+++ b/Engine/Source/Runtime/Engine/Private/Animation/AnimNode_StateMachine.cpp
@@ -770,9 +770,9 @@ void FAnimNode_StateMachine::SetState(const FAnimationBaseContext& Context, int3
 	if (NewStateIndex != CurrentState)
 	{
 		const int32 PrevStateIndex = CurrentState;
-		if(CurrentState != INDEX_NONE && CurrentState < OnGraphStatesEntered.Num())
+		if(CurrentState != INDEX_NONE && CurrentState < OnGraphStatesExited.Num())
 		{
-			OnGraphStatesEntered[CurrentState].ExecuteIfBound(*this, CurrentState, NewStateIndex);
+			OnGraphStatesExited[CurrentState].ExecuteIfBound(*this, CurrentState, NewStateIndex);
 		}
 
 		// Determine if the new state is active or not
@@ -791,9 +791,9 @@ void FAnimNode_StateMachine::SetState(const FAnimationBaseContext& Context, int3
 			StatePoseLinks[NewStateIndex].CacheBones(CacheBoneContext);
 		}
 
-		if(CurrentState != INDEX_NONE && CurrentState < OnGraphStatesExited.Num())
+		if(CurrentState != INDEX_NONE && CurrentState < OnGraphStatesEntered.Num())
 		{
-			OnGraphStatesExited[CurrentState].ExecuteIfBound(*this, PrevStateIndex, CurrentState);
+			OnGraphStatesEntered[CurrentState].ExecuteIfBound(*this, PrevStateIndex, CurrentState);
 		}
 	}
 }
diff --git a/Engine/Source/Runtime/Engine/Private/Camera/CameraComponent.cpp b/Engine/Source/Runtime/Engine/Private/Camera/CameraComponent.cpp
--- a/Engine/Source/Runtime/Engine/Private/Camera/CameraComponent.cpp
+++ b/Engine/Source/Runtime/Engine/Private/Camera/CameraComponent.cpp
@@ -32,6 +32,7 @@ UCameraComponent::UCameraComponent(const FObjectInitializer& ObjectInitializer)
 	OrthoNearClipPlane = 0.0f;
 	OrthoFarClipPlane = WORLD_MAX;
 	bConstrainAspectRatio = false;
+	bUseFieldOfViewForLOD = true;
 	PostProcessBlendWeight = 1.0f;
 	bUseControllerViewRotation_DEPRECATED = true; // the previous default value before bUsePawnControlRotation replaced this var.
 	bUsePawnControlRotation = false;
@@ -210,6 +211,7 @@ void UCameraComponent::GetCameraView(float DeltaTime, FMinimalViewInfo& DesiredV
 	DesiredView.FOV = FieldOfView;
 	DesiredView.AspectRatio = AspectRatio;
 	DesiredView.bConstrainAspectRatio = bConstrainAspectRatio;
+	DesiredView.bUseFieldOfViewForLOD = bUseFieldOfViewForLOD;
 	DesiredView.ProjectionMode = ProjectionMode;
 	DesiredView.OrthoWidth = OrthoWidth;
 	DesiredView.OrthoNearClipPlane = OrthoNearClipPlane;
diff --git a/Engine/Source/Runtime/Engine/Private/Camera/CameraStackTypes.cpp b/Engine/Source/Runtime/Engine/Private/Camera/CameraStackTypes.cpp
--- a/Engine/Source/Runtime/Engine/Private/Camera/CameraStackTypes.cpp
+++ b/Engine/Source/Runtime/Engine/Private/Camera/CameraStackTypes.cpp
@@ -19,6 +19,7 @@ bool FMinimalViewInfo::Equals(const FMinimalViewInfo& OtherInfo) const
 		(OrthoFarClipPlane == OtherInfo.OrthoFarClipPlane) &&
 		(AspectRatio == OtherInfo.AspectRatio) &&
 		(bConstrainAspectRatio == OtherInfo.bConstrainAspectRatio) &&
+		(bUseFieldOfViewForLOD == OtherInfo.bUseFieldOfViewForLOD) &&
 		(ProjectionMode == OtherInfo.ProjectionMode);
 }
 
@@ -36,6 +37,7 @@ void FMinimalViewInfo::BlendViewInfo(FMinimalViewInfo& OtherInfo, float OtherWei
 
 	AspectRatio = FMath::Lerp(AspectRatio, OtherInfo.AspectRatio, OtherWeight);
 	bConstrainAspectRatio |= OtherInfo.bConstrainAspectRatio;
+	bUseFieldOfViewForLOD |= OtherInfo.bUseFieldOfViewForLOD;
 }
 
 void FMinimalViewInfo::ApplyBlendWeight(const float& Weight)
@@ -64,6 +66,7 @@ void FMinimalViewInfo::AddWeightedViewInfo(const FMinimalViewInfo& OtherView, co
 	AspectRatio += OtherViewWeighted.AspectRatio;
 
 	bConstrainAspectRatio |= OtherViewWeighted.bConstrainAspectRatio;
+	bUseFieldOfViewForLOD |= OtherViewWeighted.bUseFieldOfViewForLOD;
 }
 
 void FMinimalViewInfo::CalculateProjectionMatrixGivenView(const FMinimalViewInfo& ViewInfo, TEnumAsByte<enum EAspectRatioAxisConstraint> AspectRatioAxisConstraint, FViewport* Viewport, FSceneViewProjectionData& InOutProjectionData)
diff --git a/Engine/Source/Runtime/Engine/Private/Collision/WorldCollisionAsync.cpp b/Engine/Source/Runtime/Engine/Private/Collision/WorldCollisionAsync.cpp
--- a/Engine/Source/Runtime/Engine/Private/Collision/WorldCollisionAsync.cpp
+++ b/Engine/Source/Runtime/Engine/Private/Collision/WorldCollisionAsync.cpp
@@ -391,6 +391,7 @@ void UWorld::WaitForAllAsyncTraceTasks()
 	AsyncTraceData& DataBufferExecuted = AsyncTraceState.GetBufferForPreviousFrame();
 	if (DataBufferExecuted.AsyncTraceCompletionEvent.Num() > 0)
 	{
+		QUICK_SCOPE_CYCLE_COUNTER(STAT_WaitForAllAsyncTraceTasks);
 		FTaskGraphInterface::Get().WaitUntilTasksComplete(DataBufferExecuted.AsyncTraceCompletionEvent,ENamedThreads::GameThread);
 		DataBufferExecuted.AsyncTraceCompletionEvent.Reset();
 	}
