UE4 - merge GDC branch, code [REDACTED_HANDLE] to main

diff --git a/Engine/Source/Runtime/Engine/Private/ContentStreaming.cpp b/Engine/Source/Runtime/Engine/Private/ContentStreaming.cpp
--- a/Engine/Source/Runtime/Engine/Private/ContentStreaming.cpp
+++ b/Engine/Source/Runtime/Engine/Private/ContentStreaming.cpp
@@ -376,13 +376,21 @@ struct FStreamingTexture
 	 */
 	float CalcPriority( )
 	{
-//		float DistanceFactor = 1.0f - FMath::Clamp( MinDistance, 1.0f, MAX_STREAMINGDISTANCE ) / MAX_STREAMINGDISTANCE;
-		float DistanceFactor = 1 - FMath::Sqrt( FMath::Clamp( MinDistance, 1.0f, MAX_STREAMINGDISTANCE ) / MAX_STREAMINGDISTANCE );
-//		float MipFactor = WantedMips - ResidentMips;
-		float MipFactor = float(WantedMips) / float(MAX_TEXTURE_MIP_COUNT);
-		float TimeFactor = GStreamWithTimeFactor ? (FMath::Clamp( LastRenderTime, 1.0f, MAX_LASTRENDERTIME ) / MAX_LASTRENDERTIME) : 0.0f;
-		float Priority = MipFactor + DistanceFactor * (1.0f - TimeFactor * 0.5f) + bForceFullyLoad * 100.0f;
-		return Priority;
+		if (LODGroup != TEXTUREGROUP_Terrain_Heightmap)
+		{
+	//		float DistanceFactor = 1.0f - FMath::Clamp( MinDistance, 1.0f, MAX_STREAMINGDISTANCE ) / MAX_STREAMINGDISTANCE;
+			float DistanceFactor = 1 - FMath::Sqrt( FMath::Clamp( MinDistance, 1.0f, MAX_STREAMINGDISTANCE ) / MAX_STREAMINGDISTANCE );
+	//		float MipFactor = WantedMips - ResidentMips;
+			float MipFactor = float(WantedMips) / float(MAX_TEXTURE_MIP_COUNT);
+			float TimeFactor = GStreamWithTimeFactor ? (FMath::Clamp( LastRenderTime, 1.0f, MAX_LASTRENDERTIME ) / MAX_LASTRENDERTIME) : 0.0f;
+			float Priority = MipFactor + DistanceFactor * (1.0f - TimeFactor * 0.5f) + bForceFullyLoad * 100.0f;
+			return Priority;
+		}
+		else
+		{
+			// Because heightmaps are technically geometry, it is absolutely imperative that they are streamed in
+			return 1000.0f;
+		}
 	}
 
 	/**
@@ -514,14 +522,15 @@ bool FFloatMipLevel::operator>=(const FFloatMipLevel& rhs) const
 
 FFloatMipLevel FFloatMipLevel::FromScreenSizeInTexels(float ScreenSizeInTexels)
 {
-	check(ScreenSizeInTexels >= 0);
+	//check(ScreenSizeInTexels >= 0);  @todo: should we figure out why this is negative?
+	const float ScreenSizeInTexelsClamped = FMath::Max(0.f, ScreenSizeInTexels);
 
 	FFloatMipLevel Ret;
 
 	// int
 //	Ret.MipLevel = 1 + FMath::CeilLogTwo( FMath::TruncToInt(ScreenSizeInTexels) );
 	// float
-	Ret.MipLevel = 1.0f + FMath::Log2(ScreenSizeInTexels);
+	Ret.MipLevel = 1.0f + FMath::Log2(ScreenSizeInTexelsClamped);
 
 	return Ret;
 }
@@ -925,9 +934,9 @@ struct FStreamingContext
 			PendingMemoryAdjustment = Stats.PendingMemoryAdjustment;
 	
 			// set total size for the pool (used to available)
-			SET_MEMORY_STAT(STAT_TexturePoolAllocatedSize, (int32)AllocatedMemorySize);
-			SET_MEMORY_STAT(STAT_TexturePoolSize, (int32)Stats.TexturePoolSize);
-			SET_MEMORY_STAT(MCR_TexturePool, (int32)Stats.TexturePoolSize);
+			SET_MEMORY_STAT(STAT_TexturePoolAllocatedSize, AllocatedMemorySize);
+			SET_MEMORY_STAT(STAT_TexturePoolSize, Stats.TexturePoolSize);
+			SET_MEMORY_STAT(MCR_TexturePool, Stats.TexturePoolSize);
 		}
 		else
 		{
@@ -973,7 +982,7 @@ struct FStreamingContext
 	/** Currently available number of bytes in the texture pool, or INDEX_NONE if bRHISupportsMemoryStats is false. */
 	int64 AvailableMemorySize;
 	/** Pending Adjustments to allocated texture memory, due to async reallocations, etc. */
-	int32 PendingMemoryAdjustment;
+	int64 PendingMemoryAdjustment;
 	/** Whether to fill in TextureStats this frame. */
 	bool bCollectTextureStats;
 
@@ -983,28 +992,28 @@ struct FStreamingContext
 #endif
 
 	// Stats for this frame.
-	uint32 ThisFrameTotalRequestSize;
-	uint32 ThisFrameTotalLightmapRequestSize;
-	uint32 ThisFrameNumStreamingTextures;
-	uint32 ThisFrameNumRequestsInCancelationPhase;
-	uint32 ThisFrameNumRequestsInUpdatePhase;
-	uint32 ThisFrameNumRequestsInFinalizePhase;
-	uint32 ThisFrameTotalIntermediateTexturesSize;
-	uint32 ThisFrameNumIntermediateTextures;
-	uint32 ThisFrameTotalStreamingTexturesSize;
-	uint32 ThisFrameTotalStreamingTexturesMaxSize;
-	uint32 ThisFrameTotalLightmapMemorySize;
-	uint32 ThisFrameTotalLightmapDiskSize;
-	uint32 ThisFrameTotalMipCountIncreaseRequestsInFlight;
-	uint32 ThisFrameOptimalWantedSize;
+	uint64 ThisFrameTotalRequestSize;
+	uint64 ThisFrameTotalLightmapRequestSize;
+	uint64 ThisFrameNumStreamingTextures;
+	uint64 ThisFrameNumRequestsInCancelationPhase;
+	uint64 ThisFrameNumRequestsInUpdatePhase;
+	uint64 ThisFrameNumRequestsInFinalizePhase;
+	uint64 ThisFrameTotalIntermediateTexturesSize;
+	uint64 ThisFrameNumIntermediateTextures;
+	uint64 ThisFrameTotalStreamingTexturesSize;
+	uint64 ThisFrameTotalStreamingTexturesMaxSize;
+	uint64 ThisFrameTotalLightmapMemorySize;
+	uint64 ThisFrameTotalLightmapDiskSize;
+	uint64 ThisFrameTotalMipCountIncreaseRequestsInFlight;
+	uint64 ThisFrameOptimalWantedSize;
 	/** Number of bytes using StaticTexture heuristics this frame, currently in memory. */
-	uint32 ThisFrameTotalStaticTextureHeuristicSize;
+	uint64 ThisFrameTotalStaticTextureHeuristicSize;
 	/** Number of bytes using DynmicTexture heuristics this frame, currently in memory. */
-	uint32 ThisFrameTotalDynamicTextureHeuristicSize;
+	uint64 ThisFrameTotalDynamicTextureHeuristicSize;
 	/** Number of bytes using LastRenderTime heuristics this frame, currently in memory. */
-	uint32 ThisFrameTotalLastRenderHeuristicSize;
+	uint64 ThisFrameTotalLastRenderHeuristicSize;
 	/** Number of bytes using ForcedIntoMemory heuristics this frame, currently in memory. */
-	uint32 ThisFrameTotalForcedHeuristicSize;
+	uint64 ThisFrameTotalForcedHeuristicSize;
 };
 
 
@@ -1062,21 +1071,21 @@ public:
 			NumWantingTextures = 0;
 		}
 		/** Total number of bytes currently in memory */
-		int32 TotalResidentSize;
+		int64 TotalResidentSize;
 		/** Total number of bytes required, using PerfectWantedMips */
-		STAT( int32 TotalRequiredSize );
+		STAT( int64 TotalRequiredSize );
 		/** Currently being streamed in, in bytes. */
 		SIZE_T PendingStreamInSize;
 		/** Currently being streamed out, in bytes. */
 		SIZE_T PendingStreamOutSize;
 		/** How much we want to stream in, in bytes. */
-		int32 WantedInSize;
+		int64 WantedInSize;
 		/** How much we want to stream out, in bytes.  */
-		int32 WantedOutSize;
+		int64 WantedOutSize;
 		/** How much memory we're temporarily using for streaming in/out, in bytes. */
-		int32 TempStreamingSize;
+		int64 TempStreamingSize;
 		/** Number of textures that want more mips. */
-		int32 NumWantingTextures;
+		int64 NumWantingTextures;
 	};
 
 	/** Returns the resulting priorities, matching the FStreamingManagerTexture::StreamingTextures array. */
@@ -1116,7 +1125,7 @@ private:
 		{
 			FStreamingTexture& StreamingTexture = StreamingManager.StreamingTextures[ Index ];
 
-			int32 ResidentTextureSize = StreamingTexture.GetSize( StreamingTexture.ResidentMips );
+			int64 ResidentTextureSize = StreamingTexture.GetSize( StreamingTexture.ResidentMips );
 			ThreadStats.TotalResidentSize += ResidentTextureSize;
 
 			StreamingTexture.bUsesStaticHeuristics = false;
@@ -1147,10 +1156,10 @@ private:
 				}
 
 				// Accumulate streaming numbers.
-				int32 WantedTextureSize = StreamingTexture.GetSize( StreamingTexture.WantedMips );
+				int64 WantedTextureSize = StreamingTexture.GetSize( StreamingTexture.WantedMips );
 				if ( StreamingTexture.bInFlight )
 				{
-					int32 RequestedTextureSize = StreamingTexture.GetSize( StreamingTexture.RequestedMips );
+					int64 RequestedTextureSize = StreamingTexture.GetSize( StreamingTexture.RequestedMips );
 					ThreadStats.TempStreamingSize += ResidentTextureSize;	//@TODO: 0 for in-place reallocations.
 					if ( StreamingTexture.RequestedMips > StreamingTexture.ResidentMips )
 					{
@@ -2093,8 +2102,11 @@ FStreamingManagerTexture::FStreamingManagerTexture()
 ,	bPauseTextureStreaming(false)
 {
 	// Read settings from ini file.
-	verify( GConfig->GetInt( TEXT("TextureStreaming"), TEXT("MemoryMargin"),				MemoryMargin,						GEngineIni ) );
-	verify( GConfig->GetInt( TEXT("TextureStreaming"), TEXT("MinEvictSize"),				MinEvictSize,						GEngineIni ) );
+	int32 TempInt;
+	verify( GConfig->GetInt( TEXT("TextureStreaming"), TEXT("MemoryMargin"),				TempInt,						GEngineIni ) );
+	MemoryMargin = TempInt;
+	verify( GConfig->GetInt( TEXT("TextureStreaming"), TEXT("MinEvictSize"),				TempInt,						GEngineIni ) );
+	MinEvictSize = TempInt;
 
 	verify( GConfig->GetFloat( TEXT("TextureStreaming"), TEXT("LightmapStreamingFactor"),			GLightmapStreamingFactor,		GEngineIni ) );
 	verify( GConfig->GetFloat( TEXT("TextureStreaming"), TEXT("ShadowmapStreamingFactor"),			GShadowmapStreamingFactor,		GEngineIni ) );
@@ -2167,7 +2179,7 @@ FStreamingManagerTexture::FStreamingManagerTexture()
 #if STATS_FAST
 	MaxStreamingTexturesSize = 0;
 	MaxOptimalTextureSize = 0;
-	MaxStreamingOverBudget = MIN_int32;
+	MaxStreamingOverBudget = MIN_int64;
 	MaxTexturePoolAllocatedSize = 0;
 	MinLargestHoleSize = OriginalTexturePoolSize;
 	MaxNumWantingTextures = 0;
@@ -2312,7 +2324,7 @@ struct FTextureStreamingCompare
  *	@param RequiredMemorySize	- Amount of memory to try to free up, in bytes
  *	@param bSucceeded			- [out] Upon return, whether it succeeded or not
  **/
-void Renderthread_StreamOutTextureData(FRHICommandListImmediate& RHICmdList, TArray<FTextureSortElement>* InCandidateTextures, int32 RequiredMemorySize, volatile bool* bSucceeded)
+void Renderthread_StreamOutTextureData(FRHICommandListImmediate& RHICmdList, TArray<FTextureSortElement>* InCandidateTextures, int64 RequiredMemorySize, volatile bool* bSucceeded)
 {
 	*bSucceeded = false;
 	// only for debugging?
@@ -2328,7 +2340,7 @@ void Renderthread_StreamOutTextureData(FRHICommandListImmediate& RHICmdList, TAr
 	InCandidateTextures->Sort( FTextureStreamingCompare() );
 
 	// Attempt to shrink enough candidates to free up the required memory size. One mip-level at a time.
-	int32 SavedMemory = 0;
+	int64 SavedMemory = 0;
 	bool bKeepShrinking = true;
 	bool bShrinkCharacterTextures = false;	// Don't shrink any character textures the first loop.
 	while ( SavedMemory < RequiredMemorySize && bKeepShrinking )
@@ -2348,9 +2360,9 @@ void Renderthread_StreamOutTextureData(FRHICommandListImmediate& RHICmdList, TAr
 				if ( bReallocationSucceeded )
 				{
 					// Start using the new one.
-					int32 OldSize = SortElement.Size;
-					int32 NewSize = SortElement.Texture->CalcTextureMemorySize(NewMipCount);
-					int32 Savings = OldSize - NewSize;
+					int64 OldSize = SortElement.Size;
+					int64 NewSize = SortElement.Texture->CalcTextureMemorySize(NewMipCount);
+					int64 Savings = OldSize - NewSize;
 
 					// Set up UTexture2D
 					SortElement.Texture->ResidentMips = NewMipCount;
@@ -2383,9 +2395,9 @@ void Renderthread_StreamOutTextureData(FRHICommandListImmediate& RHICmdList, TAr
  *	@param RequiredMemorySize	- Additional texture memory required
  *	@return						- Whether it succeeded or not
  **/
-bool FStreamingManagerTexture::StreamOutTextureData( int32 RequiredMemorySize )
+bool FStreamingManagerTexture::StreamOutTextureData( int64 RequiredMemorySize )
 {
-	RequiredMemorySize = FMath::Max<int32>(RequiredMemorySize, MinEvictSize);
+	RequiredMemorySize = FMath::Max<int64>(RequiredMemorySize, MinEvictSize);
 
 	// Array of candidates for reducing mip-levels.
 	TArray<FTextureSortElement> CandidateTextures;
@@ -2447,7 +2459,7 @@ bool FStreamingManagerTexture::StreamOutTextureData( int32 RequiredMemorySize )
 	ENQUEUE_UNIQUE_RENDER_COMMAND_THREEPARAMETER(
 		StreamOutTextureDataCommand,
 		TArray<FTextureSortElement>*,CandidateTextures,&CandidateTextures,
-		int32,RequiredMemorySize,RequiredMemorySize,
+		int64,RequiredMemorySize,RequiredMemorySize,
 		volatile bool*,bSucceeded,&bSucceeded,
 	{
 		Renderthread_StreamOutTextureData(RHICmdList, CandidateTextures, RequiredMemorySize, bSucceeded);
@@ -3180,7 +3192,7 @@ void FStreamingManagerTexture::UpdateStreamingStats( const FStreamingContext& Co
 #if STATS_FAST
 		MaxStreamingTexturesSize = FMath::Max(MaxStreamingTexturesSize, GET_MEMORY_STAT(STAT_StreamingTexturesSize));
 		MaxOptimalTextureSize = FMath::Max(MaxOptimalTextureSize, GET_MEMORY_STAT(STAT_OptimalTextureSize));
-		MaxStreamingOverBudget = FMath::Max<int32>(MaxStreamingOverBudget, int32(GET_MEMORY_STAT(STAT_StreamingOverBudget)) - int32(GET_MEMORY_STAT(STAT_StreamingUnderBudget)));
+		MaxStreamingOverBudget = FMath::Max<int64>(MaxStreamingOverBudget, GET_MEMORY_STAT(STAT_StreamingOverBudget) - GET_MEMORY_STAT(STAT_StreamingUnderBudget));
 		MaxTexturePoolAllocatedSize = FMath::Max(MaxTexturePoolAllocatedSize, GET_MEMORY_STAT(STAT_TexturePoolAllocatedSize));
 #if STATS
 		MinLargestHoleSize = FMath::Min(MinLargestHoleSize, GET_MEMORY_STAT(STAT_TexturePool_LargestHole));
@@ -3288,9 +3300,9 @@ void FStreamingManagerTexture::StreamTextures( bool bProcessEverything )
 			}
 		};
 		Context.TextureStats.Sort( FCompareFTextureStreamingStats() );
-		int32 TotalCurrentSize	= 0;
-		int32 TotalWantedSize	= 0;
-		int32 TotalMaxSize		= 0;
+		int64 TotalCurrentSize	= 0;
+		int64 TotalWantedSize	= 0;
+		int64 TotalMaxSize		= 0;
 		TextureStatsReport.Add( FString( TEXT(",Priority,Current,Wanted,Max,Largest Resident,Current Size (KB),Wanted Size (KB),Max Size (KB),Largest Resident Size (KB),Streaming Type,Last Rendered,BoostFactor,Name") ) );
 		for( int32 TextureIndex=0; TextureIndex<Context.TextureStats.Num(); TextureIndex++ )
 		{
@@ -3347,7 +3359,7 @@ void FStreamingManagerTexture::StreamTextures( bool bProcessEverything )
 	ThreadStats.PendingStreamInSize += (Stats.PendingMemoryAdjustment > 0) ? Stats.PendingMemoryAdjustment : 0;
 	ThreadStats.PendingStreamOutSize += (Stats.PendingMemoryAdjustment < 0) ? -Stats.PendingMemoryAdjustment : 0;
 
-	int32 AvailableNow, AvailableLater, TempMemoryUsed;
+	int64 AvailableNow, AvailableLater, TempMemoryUsed;
 	if ( Stats.IsUsingLimitedPoolSize() )
 	{
 		ThreadStats.TempStreamingSize = GStreamMemoryTracker.CalcTempMemory();
@@ -3367,8 +3379,8 @@ void FStreamingManagerTexture::StreamTextures( bool bProcessEverything )
 	else
 	{
 		TempMemoryUsed = ThreadStats.TempStreamingSize;
-		AvailableNow = MAX_int32;
-		AvailableLater = MAX_int32;
+		AvailableNow = MAX_int64;
+		AvailableLater = MAX_int64;
 		SET_MEMORY_STAT( STAT_StreamingOverBudget, 0 );
 		SET_MEMORY_STAT( STAT_StreamingUnderBudget, 0 );
 	}
@@ -3400,8 +3412,8 @@ void FStreamingManagerTexture::StreamTextures( bool bProcessEverything )
 				if ( HighPrioTexture.bInFlight && HighPrioTexture.RequestedMips < HighPrioTexture.ResidentMips && HighPrioTexture.RequestedMips < HighPrioTexture.WantedMips )
 				{
 					// Should we cancel?
-					int32 RequestedSize = HighPrioTexture.GetSize( HighPrioTexture.RequestedMips );
-					int32 StreamSize = HighPrioTexture.GetSize( HighPrioTexture.ResidentMips ) - RequestedSize;
+					int64 RequestedSize = HighPrioTexture.GetSize( HighPrioTexture.RequestedMips );
+					int64 StreamSize = HighPrioTexture.GetSize( HighPrioTexture.ResidentMips ) - RequestedSize;
 					if ( StreamSize > AvailableLater )
 					{
 						// Try to cancel.
@@ -3417,8 +3429,8 @@ void FStreamingManagerTexture::StreamTextures( bool bProcessEverything )
 				if ( HighPrioTexture.bInFlight == false && HighPrioTexture.WantedMips > HighPrioTexture.ResidentMips )
 				{
 					//@TODO: TempSize=0 when using in-place reallocation, but we're conservative here.
-					int32 TempSize = HighPrioTexture.GetSize( HighPrioTexture.ResidentMips );
-					int32 StreamSize = HighPrioTexture.GetSize( HighPrioTexture.WantedMips ) - TempSize;
+					int64 TempSize = HighPrioTexture.GetSize( HighPrioTexture.ResidentMips );
+					int64 StreamSize = HighPrioTexture.GetSize( HighPrioTexture.WantedMips ) - TempSize;
 					AvailableLater -= StreamSize;
 
 					// Can we stream in this texture now?
@@ -3485,7 +3497,7 @@ void FStreamingManagerTexture::StreamTextures( bool bProcessEverything )
  * @param StreamingRequests		[in/out] Indices to textures that are going to be streamed this frame
  * @return						The last priority index considered (may be sooner that StopIndex)
  */
-int32 FStreamingManagerTexture::StreamoutTextures( FStreamoutLogic StreamoutLogic, int32& AvailableLater, int32& TempMemoryUsed, int32 StartIndex, int32 StopIndex, int32& LowPrioIndex, const TArray<FTexturePriority>& PrioritizedTextures, FStreamingRequests& StreamingRequests )
+int32 FStreamingManagerTexture::StreamoutTextures( FStreamoutLogic StreamoutLogic, int64& AvailableLater, int64& TempMemoryUsed, int32 StartIndex, int32 StopIndex, int32& LowPrioIndex, const TArray<FTexturePriority>& PrioritizedTextures, FStreamingRequests& StreamingRequests )
 {
 	// Only bump LowPrioIndex if we're starting from there.
 	bool bBumpLowPrioIndex = (StartIndex == LowPrioIndex);
@@ -3511,8 +3523,8 @@ int32 FStreamingManagerTexture::StreamoutTextures( FStreamoutLogic StreamoutLogi
 					bool bCancelled = CancelStreamingRequest( StreamingTexture );
 					if ( bCancelled )
 					{
-						int32 CurrentSize = StreamingTexture.GetSize( StreamingTexture.ResidentMips );
-						int32 StreamSize = StreamingTexture.GetSize( StreamingTexture.RequestedMips ) - CurrentSize;
+						int64 CurrentSize = StreamingTexture.GetSize( StreamingTexture.ResidentMips );
+						int64 StreamSize = StreamingTexture.GetSize( StreamingTexture.RequestedMips ) - CurrentSize;
 						AvailableLater += StreamSize;
 					}
 				}
@@ -3528,8 +3540,8 @@ int32 FStreamingManagerTexture::StreamoutTextures( FStreamoutLogic StreamoutLogi
 				if ( WantedMips < CurrentMips )
 				{
 					//@TODO: Temp memory used=0 when using in-place reallocation. Though this is conservative and should be fine.
-					int32 CurrentSize = StreamingTexture.GetSize( CurrentMips );
-					int32 StreamSize = CurrentSize - StreamingTexture.GetSize( WantedMips );
+					int64 CurrentSize = StreamingTexture.GetSize( CurrentMips );
+					int64 StreamSize = CurrentSize - StreamingTexture.GetSize( WantedMips );
 					AvailableLater += StreamSize;
 					TempMemoryUsed += CurrentSize;
 
@@ -3567,7 +3579,7 @@ int32 FStreamingManagerTexture::StreamoutTextures( FStreamoutLogic StreamoutLogi
  * @param PrioritizedTextures	Array of prioritized textures to process
  * @param TempMemoryUsed		Current amount of temporary memory used by the streaming system, in bytes
  */
-void FStreamingManagerTexture::StreamTexturesUnlimited( FStreamingContext& Context, const TArray<FTexturePriority>& PrioritizedTextures, int32 TempMemoryUsed )
+void FStreamingManagerTexture::StreamTexturesUnlimited( FStreamingContext& Context, const TArray<FTexturePriority>& PrioritizedTextures, int64 TempMemoryUsed )
 {
 	for ( int32 PrioIndex=0; PrioIndex < PrioritizedTextures.Num() && TempMemoryUsed < MaxTempMemoryUsed; ++PrioIndex )
 	{
@@ -3598,7 +3610,7 @@ void FStreamingManagerTexture::StreamTexturesUnlimited( FStreamingContext& Conte
 			if (!GNeverStreamOutTextures || StreamingTexture.ResidentMips < StreamingTexture.WantedMips)
 			{
 				//@TODO: Temp memory used=0 when using in-place reallocation.
-				int32 CurrentSize = StreamingTexture.GetSize( StreamingTexture.ResidentMips );
+				int64 CurrentSize = StreamingTexture.GetSize( StreamingTexture.ResidentMips );
 				TempMemoryUsed += CurrentSize;
 				StartStreaming( StreamingTexture, StreamingTexture.WantedMips, Context, true );
 			}
@@ -3969,16 +3981,16 @@ void FStreamingManagerTexture::CalcMinMaxMips( FStreamingTexture& StreamingTextu
 void FStreamingManagerTexture::UpdateFrameStats( FStreamingContext& Context, FStreamingTexture& StreamingTexture, int32 TextureIndex )
 {
 #if STATS_FAST || STATS
-	STAT(int32 ResidentSize = StreamingTexture.GetSize(StreamingTexture.ResidentMips));
+	STAT(int64 ResidentSize = StreamingTexture.GetSize(StreamingTexture.ResidentMips));
 	STAT(Context.ThisFrameTotalStreamingTexturesSize += ResidentSize);
 	STAT(int32 PerfectWantedMips = FMath::Clamp(StreamingTexture.PerfectWantedMips, StreamingTexture.MinAllowedMips, StreamingTexture.MaxAllowedOptimalMips) );
-	STAT(int32 PerfectWantedSize = StreamingTexture.GetSize( PerfectWantedMips ));
-	STAT(int32 MostResidentSize = StreamingTexture.GetSize( StreamingTexture.MostResidentMips ) );
+	STAT(int64 PerfectWantedSize = StreamingTexture.GetSize( PerfectWantedMips ));
+	STAT(int64 MostResidentSize = StreamingTexture.GetSize( StreamingTexture.MostResidentMips ) );
 	STAT(Context.ThisFrameOptimalWantedSize += PerfectWantedSize );
 #endif
 
 #if STATS
-	int32 PotentialSize = StreamingTexture.GetSize(StreamingTexture.MaxAllowedMips);
+	int64 PotentialSize = StreamingTexture.GetSize(StreamingTexture.MaxAllowedMips);
 	ETextureStreamingType StreamType = StreamingTexture.GetStreamingType();
 	switch ( StreamType )
 	{
@@ -4205,8 +4217,8 @@ bool FStreamingManagerTexture::StartStreaming( FStreamingTexture& StreamingTextu
 				if ( bIsLoadingRequest )
 				{
 					// Manually update size as allocations are deferred/ happening in rendering thread.
-					int32 CurrentRequestSize			= StreamingTexture.GetSize(RequestMips);
-					int32 StreamInSize					= CurrentRequestSize - StreamingTexture.GetSize(StreamingTexture.ResidentMips);
+					int64 CurrentRequestSize			= StreamingTexture.GetSize(RequestMips);
+					int64 StreamInSize					= CurrentRequestSize - StreamingTexture.GetSize(StreamingTexture.ResidentMips);
 					Context.ThisFrameTotalRequestSize	+= StreamInSize;
 					Context.ThisFrameTotalLightmapRequestSize += StreamingTexture.bIsStreamingLightmap ? StreamInSize : 0;
 					Context.AvailableMemorySize			-= CurrentRequestSize;
@@ -4263,7 +4275,7 @@ bool FStreamingManagerTexture::HandleDumpTextureStreamingStatsCommand( const TCH
 	Ar.Logf( TEXT("  Num Wanting Textures =             %d"), MaxNumWantingTextures );
 	MaxStreamingTexturesSize = 0;
 	MaxOptimalTextureSize = 0;
-	MaxStreamingOverBudget = MIN_int32;
+	MaxStreamingOverBudget = MIN_int64;
 	MaxTexturePoolAllocatedSize = 0;
 	MinLargestHoleSize = OriginalTexturePoolSize;
 	MaxNumWantingTextures = 0;
@@ -4539,7 +4551,7 @@ bool FStreamingManagerTexture::HandleUntrackTextureCommand( const TCHAR* Cmd, FO
 bool FStreamingManagerTexture::HandleStreamOutCommand( const TCHAR* Cmd, FOutputDevice& Ar )
 {
 	FString Parameter(FParse::Token(Cmd, 0));
-	int32 FreeMB = (Parameter.Len() > 0) ? FCString::Atoi(*Parameter) : 0;
+	int64 FreeMB = (Parameter.Len() > 0) ? FCString::Atoi(*Parameter) : 0;
 	if ( FreeMB > 0 )
 	{
 		bool bSucceeded = StreamOutTextureData( FreeMB * 1024 * 1024 );
@@ -4726,18 +4738,18 @@ void FStreamingManagerTexture::DumpTextureGroupStats( bool bDetailedStats )
 		}
 		int32 NumTextures;
 		int32 NumNonStreamingTextures;
-		int32 CurrentTextureSize;
-		int32 WantedTextureSize;
-		int32 MaxTextureSize;
-		int32 NonStreamingSize;
+		int64 CurrentTextureSize;
+		int64 WantedTextureSize;
+		int64 MaxTextureSize;
+		int64 NonStreamingSize;
 	};
 	FTextureGroupStats TextureGroupStats[TEXTUREGROUP_MAX];
 	FTextureGroupStats TextureGroupWaste[TEXTUREGROUP_MAX];
-	int32 NumNonStreamingTextures = 0;
-	int32 NonStreamingSize = 0;
+	int64 NumNonStreamingTextures = 0;
+	int64 NonStreamingSize = 0;
 	int32 NumNonStreamingPoolTextures = 0;
-	int32 NonStreamingPoolSize = 0;
-	int32 TotalSavings = 0;
+	int64 NonStreamingPoolSize = 0;
+	int64 TotalSavings = 0;
 //	int32 UITexels = 0;
 	int32 NumDXT[PF_MAX];
 	int32 NumNonSaved[PF_MAX];
@@ -4770,29 +4782,29 @@ void FStreamingManagerTexture::DumpTextureGroupStats( bool bDetailedStats )
 			Stat.WantedTextureSize += StreamingTexture.GetSize( StreamingTexture.WantedMips );
 			Stat.MaxTextureSize += StreamingTexture.GetSize( StreamingTexture.MaxAllowedMips );
 			
-			int32 WasteCurrent = StreamingTexture.GetSize( StreamingTexture.ResidentMips ) - RHICalcTexture2DPlatformSize(Texture2D->GetSizeX(), Texture2D->GetSizeY(), Texture2D->GetPixelFormat(), StreamingTexture.ResidentMips, 1, 0, TextureAlign);			
+			int64 WasteCurrent = StreamingTexture.GetSize( StreamingTexture.ResidentMips ) - RHICalcTexture2DPlatformSize(Texture2D->GetSizeX(), Texture2D->GetSizeY(), Texture2D->GetPixelFormat(), StreamingTexture.ResidentMips, 1, 0, TextureAlign);			
 
-			int32 WasteWanted = StreamingTexture.GetSize( StreamingTexture.WantedMips ) - RHICalcTexture2DPlatformSize(Texture2D->GetSizeX(), Texture2D->GetSizeY(), Texture2D->GetPixelFormat(), StreamingTexture.WantedMips, 1, 0, TextureAlign);			
+			int64 WasteWanted = StreamingTexture.GetSize( StreamingTexture.WantedMips ) - RHICalcTexture2DPlatformSize(Texture2D->GetSizeX(), Texture2D->GetSizeY(), Texture2D->GetPixelFormat(), StreamingTexture.WantedMips, 1, 0, TextureAlign);			
 
-			int32 WasteMaxSize = StreamingTexture.GetSize( StreamingTexture.MaxAllowedMips ) - RHICalcTexture2DPlatformSize(Texture2D->GetSizeX(), Texture2D->GetSizeY(), Texture2D->GetPixelFormat(), StreamingTexture.MaxAllowedMips, 1, 0, TextureAlign);			
+			int64 WasteMaxSize = StreamingTexture.GetSize( StreamingTexture.MaxAllowedMips ) - RHICalcTexture2DPlatformSize(Texture2D->GetSizeX(), Texture2D->GetSizeY(), Texture2D->GetPixelFormat(), StreamingTexture.MaxAllowedMips, 1, 0, TextureAlign);			
 
 			Waste.NumTextures++;
-			Waste.CurrentTextureSize += FMath::Max(WasteCurrent,0);
-			Waste.WantedTextureSize += FMath::Max(WasteWanted,0);
-			Waste.MaxTextureSize += FMath::Max(WasteMaxSize,0);
+			Waste.CurrentTextureSize += FMath::Max<int64>(WasteCurrent,0);
+			Waste.WantedTextureSize += FMath::Max<int64>(WasteWanted,0);
+			Waste.MaxTextureSize += FMath::Max<int64>(WasteMaxSize,0);
 		}
 		else
 		{
 			bool bIsPooledTexture = Texture->Resource && IsValidRef(Texture->Resource->TextureRHI) && appIsPoolTexture( Texture->Resource->TextureRHI );
-			int32 TextureSize = Texture->CalcTextureMemorySizeEnum(TMC_ResidentMips);
+			int64 TextureSize = Texture->CalcTextureMemorySizeEnum(TMC_ResidentMips);
 			Stat.NumNonStreamingTextures++;
 			Stat.NonStreamingSize += TextureSize;
 			if ( Texture2D && Texture2D->Resource )
 			{				
-				int32 WastedSize = TextureSize - RHICalcTexture2DPlatformSize(Texture2D->GetSizeX(), Texture2D->GetSizeY(), Texture2D->GetPixelFormat(), Texture2D->GetNumMips(), 1, 0, TextureAlign);				
+				int64 WastedSize = TextureSize - RHICalcTexture2DPlatformSize(Texture2D->GetSizeX(), Texture2D->GetSizeY(), Texture2D->GetPixelFormat(), Texture2D->GetNumMips(), 1, 0, TextureAlign);				
 
 				Waste.NumNonStreamingTextures++;
-				Waste.NonStreamingSize += FMath::Max(WastedSize, 0);
+				Waste.NonStreamingSize += FMath::Max<int64>(WastedSize, 0);
 			}
 			if ( bIsPooledTexture )
 			{
@@ -5522,8 +5534,8 @@ FStreamMemoryTracker::FStreamMemoryTracker()
 /** Track 'start streaming' on the gamethread. (Memory not yet affected.) */
 void FStreamMemoryTracker::GameThread_BeginUpdate( const UTexture2D& Texture )
 {
-	int32 ResidentSize = Texture.CalcTextureMemorySize( Texture.ResidentMips );
-	int32 RequestedSize = Texture.CalcTextureMemorySize( Texture.RequestedMips );
+	TSize ResidentSize = Texture.CalcTextureMemorySize( Texture.ResidentMips );
+	TSize RequestedSize = Texture.CalcTextureMemorySize( Texture.RequestedMips );
 
 	// Always plan on allocating temporary memory (we don't know if we're going to use in-place realloc yet).
 	FPlatformAtomics::InterlockedAdd( &PendingTempMemory, ResidentSize );	//@TODO: 0 for pre-planned in-place realloc 
@@ -5542,8 +5554,8 @@ void FStreamMemoryTracker::GameThread_BeginUpdate( const UTexture2D& Texture )
 
 void FStreamMemoryTracker::RenderThread_Update( const UTexture2D& Texture, bool bSuccessful )
 {
-	int32 ResidentSize = Texture.CalcTextureMemorySize( Texture.ResidentMips );
-	int32 RequestedSize = Texture.CalcTextureMemorySize( Texture.RequestedMips );
+	TSize ResidentSize = Texture.CalcTextureMemorySize( Texture.ResidentMips );
+	TSize RequestedSize = Texture.CalcTextureMemorySize( Texture.RequestedMips );
 
 	// We're not planning to use temp memory anymore - we're now either using it or not.
 	FPlatformAtomics::InterlockedAdd( &PendingTempMemory, -ResidentSize );		//@TODO: 0 for pre-planned in-place realloc
@@ -5565,8 +5577,8 @@ void FStreamMemoryTracker::RenderThread_Update( const UTexture2D& Texture, bool
 
 void FStreamMemoryTracker::RenderThread_Finalize( const UTexture2D& Texture, bool bSuccessful )
 {
-	int32 ResidentSize = Texture.CalcTextureMemorySize( Texture.ResidentMips );
-	int32 RequestedSize = Texture.CalcTextureMemorySize( Texture.RequestedMips );
+	TSize ResidentSize = Texture.CalcTextureMemorySize( Texture.ResidentMips );
+	TSize RequestedSize = Texture.CalcTextureMemorySize( Texture.RequestedMips );
 
 	// When not using in-place realloc, track freed temp memory.
 	FPlatformAtomics::InterlockedAdd( &CurrentTempMemory, -ResidentSize );	
@@ -5580,20 +5592,20 @@ void FStreamMemoryTracker::RenderThread_Finalize( const UTexture2D& Texture, boo
 }
 
 /** Calculate how much texture memory is currently available for streaming. */
-int32 FStreamMemoryTracker::CalcAvailableNow( int32 TotalFreeMemory, int32 MemoryMargin )
+FStreamMemoryTracker::TSize FStreamMemoryTracker::CalcAvailableNow( FStreamMemoryTracker::TSize TotalFreeMemory, FStreamMemoryTracker::TSize MemoryMargin )
 {
 	// Conservative.
 	return TotalFreeMemory - PendingStreamIn - PendingTempMemory - MemoryMargin;
 }
 
 /** Calculate how much texture memory will available later for streaming. */
-int32 FStreamMemoryTracker::CalcAvailableLater( int32 TotalFreeMemory, int32 MemoryMargin )
+FStreamMemoryTracker::TSize FStreamMemoryTracker::CalcAvailableLater( FStreamMemoryTracker::TSize TotalFreeMemory, FStreamMemoryTracker::TSize MemoryMargin )
 {
 	return TotalFreeMemory - PendingStreamIn + CurrentStreamOut + CurrentTempMemory - MemoryMargin;
 }
 
 /** Calculate how much texture memory is currently being used for temporary texture data during streaming. */
-int32 FStreamMemoryTracker::CalcTempMemory()
+FStreamMemoryTracker::TSize FStreamMemoryTracker::CalcTempMemory()
 {
 	return PendingTempMemory + CurrentTempMemory;
 }
diff --git a/Engine/Source/Runtime/Engine/Private/DestructibleComponent.cpp b/Engine/Source/Runtime/Engine/Private/DestructibleComponent.cpp
--- a/Engine/Source/Runtime/Engine/Private/DestructibleComponent.cpp
+++ b/Engine/Source/Runtime/Engine/Private/DestructibleComponent.cpp
@@ -1033,7 +1033,7 @@ void UDestructibleComponent::ApplyRadiusDamage(float BaseDamage, const FVector&
 #endif
 }
 
-bool UDestructibleComponent::DoCustomNavigableGeometryExport(FNavigableGeometryExport* GeomExport) const
+bool UDestructibleComponent::DoCustomNavigableGeometryExport(FNavigableGeometryExport& GeomExport) const
 {
 #if WITH_APEX
 	if (ApexDestructibleActor == NULL)
@@ -1087,7 +1087,7 @@ bool UDestructibleComponent::DoCustomNavigableGeometryExport(FNavigableGeometryE
 									++ShapesExportedCount;
 
 									// @todo address Geometry.scale not being used here
-									GeomExport->ExportPxConvexMesh(Geometry.convexMesh, LocalToWorld);
+									GeomExport.ExportPxConvexMesh(Geometry.convexMesh, LocalToWorld);
 								}
 							}
 							break;
@@ -1101,11 +1101,11 @@ bool UDestructibleComponent::DoCustomNavigableGeometryExport(FNavigableGeometryE
 
 									if ((Geometry.triangleMesh->getTriangleMeshFlags()) & PxTriangleMeshFlag::eHAS_16BIT_TRIANGLE_INDICES)
 									{
-										GeomExport->ExportPxTriMesh16Bit(Geometry.triangleMesh, LocalToWorld);
+										GeomExport.ExportPxTriMesh16Bit(Geometry.triangleMesh, LocalToWorld);
 									}
 									else
 									{
-										GeomExport->ExportPxTriMesh32Bit(Geometry.triangleMesh, LocalToWorld);
+										GeomExport.ExportPxTriMesh32Bit(Geometry.triangleMesh, LocalToWorld);
 									}
 								}
 							}
diff --git a/Engine/Source/Runtime/Engine/Private/GameEngine.cpp b/Engine/Source/Runtime/Engine/Private/GameEngine.cpp
--- a/Engine/Source/Runtime/Engine/Private/GameEngine.cpp
+++ b/Engine/Source/Runtime/Engine/Private/GameEngine.cpp
@@ -815,6 +815,7 @@ void UGameEngine::Tick( float DeltaSeconds, bool bIdleMode )
 	if ( GameViewport != NULL )
 	{
 		// Decide whether to drop high detail because of frame rate.
+		QUICK_SCOPE_CYCLE_COUNTER(STAT_UGameEngine_Tick_SetDropDetail);
 		GameViewport->SetDropDetail(DeltaSeconds);
 	}
 
@@ -849,10 +850,14 @@ void UGameEngine::Tick( float DeltaSeconds, bool bIdleMode )
 		GWorld = Context.World();
 
 		// Tick all travel and Pending NetGames (Seamless, server, client)
-		TickWorldTravel(Context, DeltaSeconds);
+		{
+			QUICK_SCOPE_CYCLE_COUNTER(STAT_UGameEngine_Tick_TickWorldTravel);
+			TickWorldTravel(Context, DeltaSeconds);
+		}
 
 		if (!IsRunningDedicatedServer() && !IsRunningCommandlet())
 		{
+			QUICK_SCOPE_CYCLE_COUNTER(STAT_UGameEngine_Tick_CheckCaptures);
 			// Only update reflection captures in game once all 'always loaded' levels have been loaded
 			// This won't work with actual level streaming though
 			if (Context.World()->AreAlwaysLoadedLevelsLoaded())
@@ -958,13 +963,17 @@ void UGameEngine::Tick( float DeltaSeconds, bool bIdleMode )
 			// Will need to take another look when trying to support multiple worlds.
 
 			// Update resource streaming after viewports have had a chance to update view information. Normal update.
-			IStreamingManager::Get().Tick( DeltaSeconds );
+			{
+				QUICK_SCOPE_CYCLE_COUNTER(STAT_UGameEngine_Tick_IStreamingManager);
+				IStreamingManager::Get().Tick( DeltaSeconds );
+			}
 
 			if ( Context.World()->bTriggerPostLoadMap )
 			{
 				Context.World()->bTriggerPostLoadMap = false;
 
 				// Turns off the loading movie (if it was turned on by LoadMap) and other post-load cleanup.
+				QUICK_SCOPE_CYCLE_COUNTER(STAT_UGameEngine_Tick_IStreamingManager);
 				PostLoadMap();
 			}
 		}
@@ -973,6 +982,7 @@ void UGameEngine::Tick( float DeltaSeconds, bool bIdleMode )
 		TickCycles=LocalTickCycles;
 
 		// See whether any map changes are pending and we requested them to be committed.
+		QUICK_SCOPE_CYCLE_COUNTER(STAT_UGameEngine_Tick_ConditionalCommitMapChange);
 		ConditionalCommitMapChange(Context);
 	}
 
@@ -983,6 +993,7 @@ void UGameEngine::Tick( float DeltaSeconds, bool bIdleMode )
 	// Restore original GWorld*. This will go away one day.
 	if (OriginalGWorldContext != NAME_None)
 	{
+		QUICK_SCOPE_CYCLE_COUNTER(STAT_UGameEngine_Tick_GetWorldContextFromHandleChecked);
 		GWorld = GetWorldContextFromHandleChecked(OriginalGWorldContext).World();
 	}
 
@@ -1009,6 +1020,7 @@ void UGameEngine::Tick( float DeltaSeconds, bool bIdleMode )
 		FAVIWriter* AVIWriter = FAVIWriter::GetInstance();
 		if (AVIWriter)
 		{
+			QUICK_SCOPE_CYCLE_COUNTER(STAT_UGameEngine_Tick_AVIWriter);
 			AVIWriter->Update(DeltaSeconds);
 		}
 
@@ -1017,6 +1029,7 @@ void UGameEngine::Tick( float DeltaSeconds, bool bIdleMode )
 		{
 			if (AVIWriter)
 			{
+				QUICK_SCOPE_CYCLE_COUNTER(STAT_UGameEngine_Tick_StartCapture);
 				AVIWriter->StartCapture(GameViewport->Viewport);
 			}
 			bCheckForMovieCapture = false;
diff --git a/Engine/Source/Runtime/Engine/Private/GameplayStatics.cpp b/Engine/Source/Runtime/Engine/Private/GameplayStatics.cpp
--- a/Engine/Source/Runtime/Engine/Private/GameplayStatics.cpp
+++ b/Engine/Source/Runtime/Engine/Private/GameplayStatics.cpp
@@ -18,6 +18,7 @@
 #include "GameFramework/DamageType.h"
 #include "PhysicalMaterials/PhysicalMaterial.h"
 #include "Components/DecalComponent.h"
+#include "LandscapeProxy.h"
 
 //////////////////////////////////////////////////////////////////////////
 // UGameplayStatics
@@ -405,6 +406,15 @@ ULevelStreaming* UGameplayStatics::GetStreamingLevel(UObject* WorldContextObject
 	return NULL;
 }
 
+void UGameplayStatics::FlushLevelStreaming(UObject* WorldContextObject)
+{
+	UWorld* World = GEngine->GetWorldFromContextObject(WorldContextObject);
+	if (World != nullptr)
+	{
+		World->FlushLevelStreaming();
+	}
+}
+
 void UGameplayStatics::CancelAsyncLoading()
 {
 	::CancelAsyncLoading();
@@ -1357,4 +1367,35 @@ void UGameplayStatics::SetWorldOriginLocation(UObject* WorldContextObject, FIntV
 	{
 		World->RequestNewWorldOrigin(NewLocation);
 	}
-}
\ No newline at end of file
+}
+
+int32 UGameplayStatics::GrassOverlappingSphereCount(UObject* WorldContextObject, const UStaticMesh* Mesh, FVector CenterPosition, float Radius)
+{
+	int32 Count = 0;
+
+	UWorld* const World = GEngine->GetWorldFromContextObject(WorldContextObject);
+	if (World)
+	{
+		const FSphere Sphere(CenterPosition, Radius);
+
+		// check every landscape
+		for (TActorIterator<ALandscapeProxy> It(World); It; ++It)
+		{
+			ALandscapeProxy* L = *It;
+			if (L)
+			{
+				for (UHierarchicalInstancedStaticMeshComponent const* HComp : L->FoliageComponents)
+				{
+					if (HComp && (HComp->StaticMesh == Mesh))
+					{
+						Count += HComp->GetOverlappingSphereCount(Sphere);
+					}
+				}
+			}
+		}
+	}
+
+	return Count;
+}
+
+
diff --git a/Engine/Source/Runtime/Engine/Private/GlobalShader.cpp b/Engine/Source/Runtime/Engine/Private/GlobalShader.cpp
--- a/Engine/Source/Runtime/Engine/Private/GlobalShader.cpp
+++ b/Engine/Source/Runtime/Engine/Private/GlobalShader.cpp
@@ -299,7 +299,10 @@ FString SaveGlobalShaderFile(EShaderPlatform Platform, FString SavePath)
 
 	// make the final name
 	FString FullPath = SavePath / GetGlobalShaderCacheFilename(Platform);
-	verify(FFileHelper::SaveArrayToFile(GlobalShaderData, *FullPath));
+	if (!FFileHelper::SaveArrayToFile(GlobalShaderData, *FullPath))
+	{
+		UE_LOG(LogMaterial, Fatal, TEXT("Could not save global shader file to '%s'"), *FullPath);
+	}
 
 	return FullPath;
 }
diff --git a/Engine/Source/Runtime/Engine/Private/HierarchicalInstancedStaticMesh.cpp b/Engine/Source/Runtime/Engine/Private/HierarchicalInstancedStaticMesh.cpp
--- a/Engine/Source/Runtime/Engine/Private/HierarchicalInstancedStaticMesh.cpp
+++ b/Engine/Source/Runtime/Engine/Private/HierarchicalInstancedStaticMesh.cpp
@@ -21,6 +21,11 @@ static TAutoConsoleVariable<int32> CVarForceLOD(
 	-1,
 	TEXT("If greater than or equal to zero, forces the foliage LOD to that level."));
 
+static TAutoConsoleVariable<int32> CVarOnlyLOD(
+	TEXT("foliage.OnlyLOD"),
+	-1,
+	TEXT("If greater than or equal to zero, only renders the foliage LOD at that level."));
+
 static TAutoConsoleVariable<int32> CVarDisableCull(
 	TEXT("foliage.DisableCull"),
 	0,
@@ -36,21 +41,16 @@ static TAutoConsoleVariable<int32> CVarDitheredLOD(
 	1,
 	TEXT("If greater than zero, dithered LOD is used, otherwise popping LOD is used."));
 
+static TAutoConsoleVariable<int32> CVarOverestimateLOD(
+	TEXT("foliage.OverestimateLOD"),
+	0,
+	TEXT("If greater than zero and dithered LOD is not used, then we use an overestimate of LOD instead of an underestimate."));
+
 static TAutoConsoleVariable<int32> CVarMaxTrianglesToRender(
 	TEXT("foliage.MaxTrianglesToRender"),
 	100000000,
 	TEXT("This is an absolute limit on the number of foliage triangles to render in one traversal. This is used to prevent a silly LOD parameter mistake from causing the OS to kill the GPU."));
 
-static TAutoConsoleVariable<float> CVarFoliageOrthoGuardband(
-	TEXT("foliage.OrthoGuardBand"),
-	1000.0f,
-	TEXT("This controls the guardband for foliage culling in ortho and shadow views. In world units."));
-
-static TAutoConsoleVariable<float> CVarFoliageGuardband(
-	TEXT("foliage.GuardBand"),
-	5.0f,
-	TEXT("This controls the guardband for foliage culling in perspective views 5 gives a 5 degree increase in FOV."));
-
 TAutoConsoleVariable<float> CVarFoliageMinimumScreenSize(
 	TEXT("foliage.MinimumScreenSize"),
 	0.000005f,
@@ -63,9 +63,23 @@ TAutoConsoleVariable<float> CVarFoliageLODDistanceScale(
 
 TAutoConsoleVariable<float> CVarRandomLODRange(
 	TEXT("foliage.RandomLODRange"),
-	500.0f,
+	0.0f,
 	TEXT("Random distance added to each instance distance to compute LOD."));
 
+static TAutoConsoleVariable<int32> CVarASyncInstaneBufferConversion(
+	TEXT("foliage.ASyncInstaneBufferConversion"),
+	1,
+	TEXT("If this is > 0, then build game instance buffering async during streaming. This is not thought to be a long term solution to this problem."));
+
+static TAutoConsoleVariable<int32> CVarRebuildTreeOnLoad(
+	TEXT("foliage.RebuildTreeOnLoad"),
+	1,
+	TEXT("If this is > 0, then we look at our heuristics and rebuild the tree if they don't match."));
+
+static TAutoConsoleVariable<int32> CVarMinVertsToSplitNode(
+	TEXT("foliage.MinVertsToSplitNode"),
+	16384,
+	TEXT("Controls the accuracy between culling and LOD accuracy and culling and CPU performance."));
 
 DECLARE_CYCLE_STAT(TEXT("Traversal Time"),STAT_FoliageTraversalTime,STATGROUP_Foliage);
 DECLARE_CYCLE_STAT(TEXT("Build Time"), STAT_FoliageBuildTime, STATGROUP_Foliage);
@@ -209,20 +223,13 @@ class FClusterBuilder
 public:
 	FClusterTree* Result;
 
-	FClusterBuilder(TArray<FMatrix>& InTransforms, const FBox& InInstBox, int32 InMaxInstancesPerLeaf = 0)
+	FClusterBuilder(TArray<FMatrix>& InTransforms, const FBox& InInstBox, int32 InMaxInstancesPerLeaf)
 		: Num(InTransforms.Num())
 		, InstBox(InInstBox)
 		, Result(nullptr)
 	{
 		InternalNodeBranchingFactor = CVarFoliageSplitFactor.GetValueOnAnyThread();
-		if (InMaxInstancesPerLeaf <= 0)
-		{
-			MaxInstancesPerLeaf = InternalNodeBranchingFactor;
-		}
-		else
-		{
-			MaxInstancesPerLeaf = InMaxInstancesPerLeaf;
-		}
+		MaxInstancesPerLeaf = InMaxInstancesPerLeaf;
 
 		Exchange(Transforms, InTransforms);
 		SortIndex.AddUninitialized(Num);
@@ -512,7 +519,7 @@ static void TestFoliage(const TArray<FString>& Args)
 	{
 		InstanceTransforms[Index] = Instances[Index].Transform;
 	}
-	FClusterBuilder Builder(InstanceTransforms, TempBox);
+	FClusterBuilder Builder(InstanceTransforms, TempBox, 16);
 	Builder.Build();
 
 	int32 Level = 0;
@@ -619,7 +626,7 @@ public:
 		, CaptureTag(0)
 #endif
 	{
-		check(!InComponent->InstanceReorderTable.Num() && !InComponent->PerInstanceSMData.Num());
+		check(!bInIsGrass || (!InComponent->InstanceReorderTable.Num() && !InComponent->PerInstanceSMData.Num()));
 	}
 
 	// FPrimitiveSceneProxy interface.
@@ -652,14 +659,16 @@ struct FFoliageRenderInstanceParams
 {
 	bool bNeedsSingleLODRuns;
 	bool bNeedsMultipleLODRuns;
+	bool bOverestimate;
 	mutable TArray<uint32, SceneRenderingAllocator> MultipleLODRuns[MAX_STATIC_MESH_LODS];
 	mutable TArray<uint32, SceneRenderingAllocator> SingleLODRuns[MAX_STATIC_MESH_LODS];
 	mutable int32 TotalSingleLODInstances[MAX_STATIC_MESH_LODS];
 	mutable int32 TotalMultipleLODInstances[MAX_STATIC_MESH_LODS];
 
-	FFoliageRenderInstanceParams(bool InbNeedsSingleLODRuns, bool InbNeedsMultipleLODRuns)
+	FFoliageRenderInstanceParams(bool InbNeedsSingleLODRuns, bool InbNeedsMultipleLODRuns, bool InbOverestimate)
 		: bNeedsSingleLODRuns(InbNeedsSingleLODRuns)
 		, bNeedsMultipleLODRuns(InbNeedsMultipleLODRuns)
+		, bOverestimate(InbOverestimate)
 	{
 		for (int32 Index = 0; Index < MAX_STATIC_MESH_LODS; Index++)
 		{
@@ -683,8 +692,8 @@ struct FFoliageRenderInstanceParams
 	{
 		if (bNeedsSingleLODRuns)
 		{
-			AddRun(SingleLODRuns[MinLod], FirstInstance, LastInstance);
-			TotalSingleLODInstances[MinLod] += 1 + LastInstance - FirstInstance;
+			AddRun(SingleLODRuns[bOverestimate ? MaxLod : MinLod], FirstInstance, LastInstance);
+			TotalSingleLODInstances[bOverestimate ? MaxLod : MinLod] += 1 + LastInstance - FirstInstance;
 		}
 		if (bNeedsMultipleLODRuns)
 		{
@@ -705,7 +714,7 @@ struct FFoliageRenderInstanceParams
 struct FFoliageCullInstanceParams : public FFoliageRenderInstanceParams
 {
 	FConvexVolume ViewFrustumLocal;
-	FConvexVolume ViewFrustumGuardBandLocal;
+	int32 MinInstancesToSplit[MAX_STATIC_MESH_LODS];
 	const TArray<FClusterNode>& Tree;
 	const FSceneView* View;
 	FVector ViewOriginInLocalZero;
@@ -715,8 +724,8 @@ struct FFoliageCullInstanceParams : public FFoliageRenderInstanceParams
 	float LODPlanesMin[MAX_STATIC_MESH_LODS];
 
 
-	FFoliageCullInstanceParams(bool InbNeedsSingleLODRuns, bool InbNeedsMultipleLODRuns, const TArray<FClusterNode>& InTree)
-	:	FFoliageRenderInstanceParams(InbNeedsSingleLODRuns, InbNeedsMultipleLODRuns)
+	FFoliageCullInstanceParams(bool InbNeedsSingleLODRuns, bool InbNeedsMultipleLODRuns, bool InbOverestimate, const TArray<FClusterNode>& InTree)
+	:	FFoliageRenderInstanceParams(InbNeedsSingleLODRuns, InbNeedsMultipleLODRuns, InbOverestimate)
 	,	Tree(InTree)
 	{
 	}
@@ -756,7 +765,6 @@ static FORCEINLINE_DEBUGGABLE bool CullNode(const FFoliageCullInstanceParams& Pa
 	if (TUseVector)
 	{
 		checkSlow(Params.ViewFrustumLocal.PermutedPlanes.Num() == 4);
-		checkSlow(Params.ViewFrustumGuardBandLocal.PermutedPlanes.Num() == 4);
 
 		//@todo, once we have more than one mesh per tree, these should be aligned
 		VectorRegister BoxMin = VectorLoad(&Node.BoundMin);
@@ -797,37 +805,15 @@ static FORCEINLINE_DEBUGGABLE bool CullNode(const FFoliageCullInstanceParams& Pa
 		VectorRegister PushX = VectorMultiply(AbsExtentX,VectorAbs(PlanesX));
 		VectorRegister PushY = VectorMultiplyAdd(AbsExtentY,VectorAbs(PlanesY),PushX);
 		VectorRegister PushOut = VectorMultiplyAdd(AbsExtentZ,VectorAbs(PlanesZ),PushY);
+		VectorRegister PushOutNegative = VectorNegate(PushOut);
 
-		const FPlane* RESTRICT PermutedGuardPlanePtr = (FPlane*)Params.ViewFrustumGuardBandLocal.PermutedPlanes.GetData();
-		PlanesX = VectorLoadAligned(PermutedGuardPlanePtr);
-		PermutedGuardPlanePtr++;
-		PlanesY = VectorLoadAligned(PermutedGuardPlanePtr);
-		PermutedGuardPlanePtr++;
-		PlanesZ = VectorLoadAligned(PermutedGuardPlanePtr);
-		PermutedGuardPlanePtr++;
-		PlanesW = VectorLoadAligned(PermutedGuardPlanePtr);
-		PermutedGuardPlanePtr++;
-
-		// Calculate the distance (x * x) + (y * y) + (z * z) - w
-		DistX = VectorMultiply(OrigX,PlanesX);
-		DistY = VectorMultiplyAdd(OrigY,PlanesY,DistX);
-		DistZ = VectorMultiplyAdd(OrigZ,PlanesZ,DistY);
-		VectorRegister DistanceGuard = VectorSubtract(DistZ,PlanesW);
-		// Now do the push out FMath::Abs(x * x) + FMath::Abs(y * y) + FMath::Abs(z * z)
-		PushX = VectorMultiply(AbsExtentX,VectorAbs(PlanesX));
-		PushY = VectorMultiplyAdd(AbsExtentY,VectorAbs(PlanesY),PushX);
-		VectorRegister PushOutGuard = VectorMultiplyAdd(AbsExtentZ,VectorAbs(PlanesZ),PushY);
-		VectorRegister PushOutNegative = VectorNegate(PushOutGuard);
-
-		// Definitely inside frustums, but check to see if it's fully contained in the guard band
-		bOutFullyContained = !VectorAnyGreaterThan(DistanceGuard,PushOutNegative);
+		bOutFullyContained = !VectorAnyGreaterThan(Distance,PushOutNegative);
 		// Check for completely outside
 		return !!VectorAnyGreaterThan(Distance,PushOut);
 	}
 	FVector Center = (Node.BoundMin + Node.BoundMax) * 0.5f;
 	FVector Extent = (Node.BoundMax - Node.BoundMin) * 0.5f;
-	if (!Params.ViewFrustumLocal.IntersectBox(Center, Extent) ||
-		!Params.ViewFrustumGuardBandLocal.IntersectBox(Center, Extent, bOutFullyContained)) // only roundoff error can cause a cull here, but whatever
+	if (!Params.ViewFrustumLocal.IntersectBox(Center, Extent, bOutFullyContained)) 
 	{
 		return true;
 	}
@@ -869,7 +855,9 @@ void FHierarchicalStaticMeshSceneProxy::Traverse(const FFoliageCullInstanceParam
 		}
 	}
 
-	bool bSplit = (!bFullyContained || MinLOD < MaxLOD) && Node.FirstChild >= 0;
+	bool bSplit = (!bFullyContained || MinLOD < MaxLOD) 
+		&& Node.FirstChild >= 0 
+		&& (1 + Node.LastInstance - Node.FirstInstance) >= Params.MinInstancesToSplit[MinLOD];
 
 	if (!bSplit)
 	{
@@ -904,7 +892,12 @@ void FHierarchicalStaticMeshSceneProxy::FillDynamicMeshElements(FMeshElementColl
 {
 	SCOPE_CYCLE_COUNTER(STAT_FoliageBatchTime);
 	int64 TotalTriangles = 0;
-	for (int32 LODIndex = 0; LODIndex < InstancedRenderData.VertexFactories.Num(); LODIndex++)
+
+	int32 OnlyLOD = FMath::Min<int32>(CVarOnlyLOD.GetValueOnRenderThread(),InstancedRenderData.VertexFactories.Num() - 1);
+	int32 FirstLOD = (OnlyLOD < 0) ? 0 : OnlyLOD;
+	int32 LastLODPlusOne = (OnlyLOD < 0) ? InstancedRenderData.VertexFactories.Num() : (OnlyLOD+1);
+
+	for (int32 LODIndex = FirstLOD; LODIndex < LastLODPlusOne; LODIndex++)
 	{
 		const FStaticMeshLODResources& LODModel = StaticMesh->RenderData->LODResources[LODIndex];
 
@@ -1072,6 +1065,9 @@ void FHierarchicalStaticMeshSceneProxy::GetDynamicMeshElements(const TArray<cons
 
 	bool bMultipleSections = ALLOW_DITHERED_LOD_FOR_INSTANCED_STATIC_MESHES && CVarDitheredLOD.GetValueOnRenderThread() > 0;
 	bool bSingleSections = !bMultipleSections;
+	bool bOverestimate = CVarOverestimateLOD.GetValueOnRenderThread() > 0;
+
+	int32 MinVertsToSplitNode = CVarMinVertsToSplitNode.GetValueOnRenderThread();
 
 	for (int32 ViewIndex = 0; ViewIndex < Views.Num(); ViewIndex++)
 	{
@@ -1097,7 +1093,7 @@ void FHierarchicalStaticMeshSceneProxy::GetDynamicMeshElements(const TArray<cons
 			if (ClusterTree.Num())
 			{
 
-				FFoliageCullInstanceParams InstanceParams(bSingleSections, bMultipleSections, ClusterTree);
+				FFoliageCullInstanceParams InstanceParams(bSingleSections, bMultipleSections, bOverestimate, ClusterTree);
 				InstanceParams.LODs = RenderData->LODResources.Num();
 
 				InstanceParams.View = View;
@@ -1109,7 +1105,6 @@ void FHierarchicalStaticMeshSceneProxy::GetDynamicMeshElements(const TArray<cons
 				ElementParams.ShadowFrustum = !!View->ViewMatrices.GetDynamicMeshElementsShadowCullFrustum;
 				if (View->ViewMatrices.GetDynamicMeshElementsShadowCullFrustum)
 				{
-					float GuardBand = CVarFoliageOrthoGuardband.GetValueOnRenderThread();
 					for (int32 Index = 0; Index < View->ViewMatrices.GetDynamicMeshElementsShadowCullFrustum->Planes.Num(); Index++)
 					{
 						FPlane Src = View->ViewMatrices.GetDynamicMeshElementsShadowCullFrustum->Planes[Index];
@@ -1119,8 +1114,6 @@ void FHierarchicalStaticMeshSceneProxy::GetDynamicMeshElements(const TArray<cons
 						FPlane Local = Norm.TransformBy(WorldToLocal);
 						FPlane LocalNorm = Local / Local.Size();
 						InstanceParams.ViewFrustumLocal.Planes.Add(LocalNorm);
-						LocalNorm.W += GuardBand;
-						InstanceParams.ViewFrustumGuardBandLocal.Planes.Add(LocalNorm);
 					}
 					bUseVectorCull = InstanceParams.ViewFrustumLocal.Planes.Num() == 4;
 				}
@@ -1139,46 +1132,26 @@ void FHierarchicalStaticMeshSceneProxy::GetDynamicMeshElements(const TArray<cons
 						ThreePlanes.SetAxes(&InstanceParams.ViewFrustumLocal.Planes[0], &InstanceParams.ViewFrustumLocal.Planes[1], &InstanceParams.ViewFrustumLocal.Planes[2]);
 						FVector ProjectionOrigin = ThreePlanes.Inverse().GetTransposed().TransformVector(FVector(InstanceParams.ViewFrustumLocal.Planes[0].W, InstanceParams.ViewFrustumLocal.Planes[1].W, InstanceParams.ViewFrustumLocal.Planes[2].W));
 
-						float GuardBand = FMath::Tan(FMath::DegreesToRadians(CVarFoliageGuardband.GetValueOnRenderThread()));
 						FVector Forward = LocalViewProjForCulling.GetColumn(3).GetSafeNormal();
 						for (int32 Index = 0; Index < InstanceParams.ViewFrustumLocal.Planes.Num(); Index++)
 						{
 							FPlane Src = InstanceParams.ViewFrustumLocal.Planes[Index];
 							FVector Normal = Src.GetSafeNormal();
 							InstanceParams.ViewFrustumLocal.Planes[Index] = FPlane(Normal, Normal | ProjectionOrigin);
-							FVector Perp = ((Normal ^ Forward).GetSafeNormal() ^ Normal).GetSafeNormal();
-							if ((Perp | Forward) > 0.0f)
-							{
-								Perp *= -1.0f;
-							}
-
-							float BlendNormal = FMath::Sqrt(1.0f - GuardBand * GuardBand);
-							 
-							Normal = ((Normal * BlendNormal) + (Perp * GuardBand)).GetSafeNormal();
-
-							InstanceParams.ViewFrustumGuardBandLocal.Planes.Add(FPlane(Normal, Normal | ProjectionOrigin));
 						}
 					}
 					else
 					{
 						bUseVectorCull = false;
-						float GuardBand = CVarFoliageOrthoGuardband.GetValueOnRenderThread();
-						for (int32 Index = 0; Index < InstanceParams.ViewFrustumLocal.Planes.Num(); Index++)
-						{
-							FPlane Src = InstanceParams.ViewFrustumLocal.Planes[Index];
-							Src.W += GuardBand;
-							InstanceParams.ViewFrustumGuardBandLocal.Planes.Add(Src);
-						}
 					}
 				}
-				if (!InstanceParams.ViewFrustumLocal.Planes.Num() || !InstanceParams.ViewFrustumGuardBandLocal.Planes.Num())
+				if (!InstanceParams.ViewFrustumLocal.Planes.Num())
 				{
 					bDisableCull = true;
 				}
 				else
 				{
 					InstanceParams.ViewFrustumLocal.Init();
-					InstanceParams.ViewFrustumGuardBandLocal.Init();
 				}
 
 				ElementParams.bBlendLODs = bMultipleSections;
@@ -1226,6 +1199,16 @@ void FHierarchicalStaticMeshSceneProxy::GetDynamicMeshElements(const TArray<cons
 					InstanceParams.LODPlanesMax[InstanceParams.LODs - 1] = FMath::Max(InstanceParams.LODPlanesMax[InstanceParams.LODs - 1], FinalCull);
 				}
 
+				for (int32 LODIndex = 0; LODIndex < InstanceParams.LODs; LODIndex++)
+				{
+					InstanceParams.MinInstancesToSplit[LODIndex] = 2;
+					int32 NumVerts = RenderData->LODResources[LODIndex].VertexBuffer.GetNumVertices();
+					if (NumVerts)
+					{
+						InstanceParams.MinInstancesToSplit[LODIndex] = MinVertsToSplitNode / NumVerts;
+					}
+				}
+
 				INC_DWORD_STAT(STAT_FoliageTraversals);
 #if !(UE_BUILD_SHIPPING || UE_BUILD_TEST)
 				if (GCaptureDebugRuns == GDebugTag && CaptureTag == GDebugTag)
@@ -1254,24 +1237,31 @@ void FHierarchicalStaticMeshSceneProxy::GetDynamicMeshElements(const TArray<cons
 					//check(UPTRINT(&ClusterTree[0].BoundMin) % 16 == 0);
 					//check(UPTRINT(&ClusterTree[0].BoundMax) % 16 == 0);
 
-					int32 MinLOD = 0;
-					int32 MaxLOD = InstanceParams.LODs;
+					int32 UseMinLOD = ClampedMinLOD;
+
+					int32 DebugMin = FMath::Min(CVarMinLOD.GetValueOnRenderThread(), InstanceParams.LODs - 1);
+					if (DebugMin >= 0)
+					{
+						UseMinLOD = FMath::Max(UseMinLOD, DebugMin);
+					}
+					int32 UseMaxLOD = InstanceParams.LODs;
 
 					int32 Force = CVarForceLOD.GetValueOnRenderThread();
 					if (Force >= 0)
 					{
-						MinLOD = Force;
-						MaxLOD = Force;
+						UseMinLOD = FMath::Clamp(Force, 0, InstanceParams.LODs - 1);
+						UseMaxLOD = FMath::Clamp(Force, 0, InstanceParams.LODs - 1);
 					}
+
 					if (CVarCullAll.GetValueOnRenderThread() < 1)
 					{
 						if (bUseVectorCull)
 						{
-							Traverse<true>(InstanceParams, 0, MinLOD, MaxLOD, bDisableCull);
+							Traverse<true>(InstanceParams, 0, UseMinLOD, UseMaxLOD, bDisableCull);
 						}
 						else
 						{
-							Traverse<false>(InstanceParams, 0, MinLOD, MaxLOD, bDisableCull);
+							Traverse<false>(InstanceParams, 0, UseMinLOD, UseMaxLOD, bDisableCull);
 						}
 					}
 #if !(UE_BUILD_SHIPPING || UE_BUILD_TEST)
@@ -1303,7 +1293,7 @@ void FHierarchicalStaticMeshSceneProxy::GetDynamicMeshElements(const TArray<cons
 			// Render unbuilt instances
 			if (FirstUnbuiltIndex <= LastUnbuiltIndex)
 			{
-				FFoliageRenderInstanceParams InstanceParams(true, false);
+				FFoliageRenderInstanceParams InstanceParams(true, false, false);
 
 				// disable LOD blending for unbuilt instances as we haven't calculated the correct LOD.
 				ElementParams.bBlendLODs = false;
@@ -1348,6 +1338,7 @@ FBoxSphereBounds UHierarchicalInstancedStaticMeshComponent::CalcBounds(const FTr
 	}
 	else
 	{
+		QUICK_SCOPE_CYCLE_COUNTER(STAT_UHierarchicalInstancedStaticMeshComponent_CalcBounds_SlowPath);
 		return Super::CalcBounds(BoundTransform);
 	}
 }
@@ -1371,6 +1362,17 @@ UHierarchicalInstancedStaticMeshComponent::~UHierarchicalInstancedStaticMeshComp
 		DEC_DWORD_STAT_BY(STAT_FoliageInstanceBuffers, ProxySize);
 }
 	ProxySize = 0;
+	FlushAsyncBuildInstanceBufferTask();
+}
+
+void UHierarchicalInstancedStaticMeshComponent::FlushAsyncBuildInstanceBufferTask()
+{
+	if (AsyncBuildInstanceBufferTask)
+	{
+		AsyncBuildInstanceBufferTask->EnsureCompletion();
+		delete AsyncBuildInstanceBufferTask;
+		AsyncBuildInstanceBufferTask = nullptr;
+	}
 }
 
 
@@ -1386,18 +1388,6 @@ void UHierarchicalInstancedStaticMeshComponent::Serialize(FArchive& Ar)
 	ClusterTree.BulkSerialize(Ar);
 }
 
-void UHierarchicalInstancedStaticMeshComponent::PostLoad()
-{
-	Super::PostLoad();
-
-	// For some reason we don't have a tree. Build one now!
-	if (PerInstanceSMData.Num() > 0 && (!ClusterTreePtr.IsValid() || ClusterTreePtr->Num() == 0))
-	{
-		check(!IsAsyncBuilding());
-		BuildTree();
-	}
-}
-
 SIZE_T UHierarchicalInstancedStaticMeshComponent::GetResourceSize( EResourceSizeMode::Type Mode )
 {
 	SIZE_T ResSize = Super::GetResourceSize(Mode);
@@ -1550,8 +1540,63 @@ bool UHierarchicalInstancedStaticMeshComponent::ShouldCreatePhysicsState() const
 	return Super::ShouldCreatePhysicsState();
 }
 
+int32 UHierarchicalInstancedStaticMeshComponent::GetVertsForLOD(int32 LODIndex)
+{
+	if (StaticMesh && StaticMesh->HasValidRenderData())
+	{
+		return StaticMesh->GetNumVertices(LODIndex);
+	}
+	return 0;
+}
+
+int32 UHierarchicalInstancedStaticMeshComponent::DesiredInstancesPerLeaf()
+{
+	int32 LOD0Verts = GetVertsForLOD(0);
+	int32 VertsToSplit = CVarMinVertsToSplitNode.GetValueOnAnyThread();
+	if (LOD0Verts)
+	{
+		return FMath::Clamp(VertsToSplit / LOD0Verts, 1, 1024);
+	}
+	return 16;
+}
+
+float UHierarchicalInstancedStaticMeshComponent::ActualInstancesPerLeaf()
+{
+	const TArray<FClusterNode>& ClusterTree = *ClusterTreePtr;
+	if (ClusterTree.Num())
+	{
+		int32 NumLeaves = 0;
+		int32 NumInstances = 0;
+		for (int32 Index = ClusterTree.Num() - 1; Index >= 0; Index--)
+		{
+			if (ClusterTree[Index].FirstChild >= 0)
+			{
+				break;
+			}
+			NumLeaves++;
+			NumInstances += 1 + ClusterTree[Index].LastInstance - ClusterTree[Index].FirstInstance;
+		}
+		if (NumLeaves)
+		{
+			return float(NumInstances) / float(NumLeaves);
+		}
+	}
+	return 0.0f;
+}
+
+void UHierarchicalInstancedStaticMeshComponent::PostBuildStats()
+{
+#if 0
+	const TArray<FClusterNode>& ClusterTree = *ClusterTreePtr;
+	int32 NumInst = WriteOncePrebuiltInstanceBuffer.Num() ? WriteOncePrebuiltInstanceBuffer.Num() : PerInstanceSMData.Num();
+	UE_LOG(LogStaticMesh, Display, TEXT("Built a foliage hierarchy with %d instances, %d nodes, %f instances / leaf (desired %d) and %d verts in LOD0. Grass? %d "), NumInst, ClusterTree.Num(), ActualInstancesPerLeaf(), DesiredInstancesPerLeaf(), GetVertsForLOD(0), !!WriteOncePrebuiltInstanceBuffer.Num());
+#endif
+}
+
+
 void UHierarchicalInstancedStaticMeshComponent::BuildTree()
 {
+	QUICK_SCOPE_CYCLE_COUNTER(STAT_UHierarchicalInstancedStaticMeshComponent_BuildTree);
 	// Verify that the mesh is valid before using it.
 	const bool bMeshIsValid = 
 		// make sure we have instances
@@ -1581,7 +1626,7 @@ void UHierarchicalInstancedStaticMeshComponent::BuildTree()
 			InstanceTransforms[Index] = PerInstanceSMData[Index].Transform;
 		}
 
-		TUniquePtr<FClusterBuilder> Builder(new FClusterBuilder(InstanceTransforms, StaticMesh->GetBounds().GetBox()));
+		TUniquePtr<FClusterBuilder> Builder(new FClusterBuilder(InstanceTransforms, StaticMesh->GetBounds().GetBox(), DesiredInstancesPerLeaf()));
 		Builder->Build();
 
 		NumBuiltInstances = Builder->Result->InstanceReorderTable.Num();
@@ -1595,7 +1640,7 @@ void UHierarchicalInstancedStaticMeshComponent::BuildTree()
 		
 		FlushAccumulatedNavigationUpdates();
 				
-		//UE_LOG(LogStaticMesh, Display, TEXT("Built a foliage hierarchy with %d elements in %.1fs."), PerInstanceSMData.Num(), float(FPlatformTime::Seconds() - StartTime));
+		PostBuildStats();
 	}
 }
 
@@ -1619,6 +1664,7 @@ void UHierarchicalInstancedStaticMeshComponent::BuildTreeAnyThread(
 
 void UHierarchicalInstancedStaticMeshComponent::AcceptPrebuiltTree(TArray<FClusterNode>& InClusterTree)
 {
+	QUICK_SCOPE_CYCLE_COUNTER(STAT_UHierarchicalInstancedStaticMeshComponent_AcceptPrebuiltTree);
 	// this is only for prebuild data, already in the correct order
 	check(!PerInstanceSMData.Num());
 	NumBuiltInstances = WriteOncePrebuiltInstanceBuffer.Num();
@@ -1651,7 +1697,11 @@ void UHierarchicalInstancedStaticMeshComponent::AcceptPrebuiltTree(TArray<FClust
 			InstancingRandomSeed = FMath::Rand();
 		}
 		Exchange(*ClusterTreePtr, InClusterTree);
+		PostBuildStats();
+
 	}
+	QUICK_SCOPE_CYCLE_COUNTER(STAT_UHierarchicalInstancedStaticMeshComponent_AcceptPrebuiltTree_Mark);
+
 	MarkRenderStateDirty();
 }
 
@@ -1744,7 +1794,7 @@ void UHierarchicalInstancedStaticMeshComponent::BuildFlatTree(const TArray<int32
 
 		FlushAccumulatedNavigationUpdates();
 
-//		UE_LOG(LogStaticMesh, Display, TEXT("Built a flat foliage hierarchy with %d elements in %.2fms."), PerInstanceSMData.Num(), float(FPlatformTime::Seconds() - StartTime) * 1000.0f);
+		PostBuildStats();
 	}
 }
 
@@ -1798,6 +1848,8 @@ void UHierarchicalInstancedStaticMeshComponent::ApplyBuildTreeAsync(ENamedThread
 
 		ReleasePerInstanceRenderData();
 	    MarkRenderStateDirty();
+
+		PostBuildStats();
 	}
 }
 
@@ -1835,9 +1887,9 @@ void UHierarchicalInstancedStaticMeshComponent::BuildTreeAsync()
 			InstanceTransforms[Index] = PerInstanceSMData[Index].Transform;
 		}
 
-		UE_LOG(LogStaticMesh, Display, TEXT("Copied %d transforms in %.3fs."), Num, float(FPlatformTime::Seconds() - StartTime));
+		//UE_LOG(LogStaticMesh, Display, TEXT("Copied %d transforms in %.3fs."), Num, float(FPlatformTime::Seconds() - StartTime));
 
-		TSharedRef<FClusterBuilder, ESPMode::ThreadSafe> Builder(new FClusterBuilder(InstanceTransforms, StaticMesh->GetBounds().GetBox()));
+		TSharedRef<FClusterBuilder, ESPMode::ThreadSafe> Builder(new FClusterBuilder(InstanceTransforms, StaticMesh->GetBounds().GetBox(), DesiredInstancesPerLeaf()));
 
 		bIsAsyncBuilding = true;
 
@@ -1864,6 +1916,8 @@ FPrimitiveSceneProxy* UHierarchicalInstancedStaticMeshComponent::CreateSceneProx
 	}
 	ProxySize = 0;
 
+	FlushAsyncBuildInstanceBufferTask();
+
 	// Verify that the mesh is valid before using it.
 	const bool bMeshIsValid = 
 		// make sure we have instances
@@ -1877,6 +1931,11 @@ FPrimitiveSceneProxy* UHierarchicalInstancedStaticMeshComponent::CreateSceneProx
 
 	if(bMeshIsValid)
 	{
+		if (bHasPerInstanceHitProxies && GetWorld() && GetWorld()->IsGameWorld())
+		{
+			// hit proxies take forever to build, don't do it if this is a game world.
+			bHasPerInstanceHitProxies = false;
+		}
 		// If we don't have a random seed for this instanced static mesh component yet, then go ahead and
 		// generate one now.  This will be saved with the static mesh component and used for future generation
 		// of random numbers for this component's instances. (Used by the PerInstanceRandom material expression)
@@ -1890,7 +1949,10 @@ FPrimitiveSceneProxy* UHierarchicalInstancedStaticMeshComponent::CreateSceneProx
 			ProxySize = FStaticMeshInstanceData::GetResourceSize(WriteOncePrebuiltInstanceBuffer.Num());
 			INC_DWORD_STAT_BY(STAT_FoliageInstanceBuffers, ProxySize);
 
-			return ::new FHierarchicalStaticMeshSceneProxy(true, this, GetWorld()->FeatureLevel, WriteOncePrebuiltInstanceBuffer);
+			bool bIsGrass = !PerInstanceSMData.Num();
+
+			check(bIsGrass || PerInstanceSMData.Num() == WriteOncePrebuiltInstanceBuffer.Num());
+			return ::new FHierarchicalStaticMeshSceneProxy(bIsGrass, this, GetWorld()->FeatureLevel, WriteOncePrebuiltInstanceBuffer);
 		}
 		ProxySize = FStaticMeshInstanceData::GetResourceSize(PerInstanceSMData.Num());
 		INC_DWORD_STAT_BY(STAT_FoliageInstanceBuffers, ProxySize);
@@ -1899,36 +1961,130 @@ FPrimitiveSceneProxy* UHierarchicalInstancedStaticMeshComponent::CreateSceneProx
 	return NULL;
 }
 
+void FAsyncBuildInstanceBuffer::DoWork()
+{
+	QUICK_SCOPE_CYCLE_COUNTER(FAsyncBuildInstanceBuffer_DoWork);
+
+	int32 Num = Component->PerInstanceSMData.Num();
+	Component->WriteOncePrebuiltInstanceBuffer.AllocateInstances(Num);
+	FRandomStream RandomStream(Component->InstancingRandomSeed);
+	bool bUseRemapTable = Component->PerInstanceSMData.Num() == Component->InstanceReorderTable.Num();
+
+	for (int32 InstanceIndex = 0; InstanceIndex < Num; InstanceIndex++)
+	{
+		FInstanceStream* InstanceRenderData = Component->WriteOncePrebuiltInstanceBuffer.GetInstanceWriteAddress(bUseRemapTable ? Component->InstanceReorderTable[InstanceIndex] : InstanceIndex);
+		const FInstancedStaticMeshInstanceData& Instance = Component->PerInstanceSMData[InstanceIndex];
+
+		InstanceRenderData->SetInstance(Instance.Transform, RandomStream.GetFraction(), Instance.LightmapUVBias, Instance.ShadowmapUVBias);
+	}
+	if (Component->RemovedInstances.Num())
+	{
+		check(bUseRemapTable);
+		for (int32 InstanceIndex = 0; InstanceIndex < Component->RemovedInstances.Num(); InstanceIndex++)
+		{
+			FInstanceStream* InstanceRenderData = Component->WriteOncePrebuiltInstanceBuffer.GetInstanceWriteAddress(Component->RemovedInstances[InstanceIndex]);
+			InstanceRenderData->NullifyInstance();
+		}
+	}
+
+
+	check(World);
+	check(World->AsyncPreRegisterLevelStreamingTasks.GetValue() > 0);
+	FPlatformMisc::MemoryBarrier();
+	World->AsyncPreRegisterLevelStreamingTasks.Decrement();
+}
+
+void UHierarchicalInstancedStaticMeshComponent::PostLoad()
+{
+	Super::PostLoad();
+#if WITH_EDITORONLY_DATA
+	if (CVarRebuildTreeOnLoad.GetValueOnGameThread() > 0 && PerInstanceSMData.Num() > 3 && StaticMesh)
+	{
+		StaticMesh->ConditionalPostLoad();
+		if (GetVertsForLOD(0))
+		{
+			float Actual = ActualInstancesPerLeaf();
+			float Desired = (float)FMath::Min(DesiredInstancesPerLeaf(), PerInstanceSMData.Num());
+			float DesiredCapped = FMath::Min(Desired, (float)PerInstanceSMData.Num());
+			if (Actual > Desired * 2.5f || Actual * 2.5f < DesiredCapped)
+			{
+				UE_LOG(LogStaticMesh, Warning, TEXT("Rebuilding foliage, please resave map %s."), *GetFullName());
+
+				BuildTree();
+				float NewActual = ActualInstancesPerLeaf();
+				ensureMsgf(!(NewActual > Desired * 2.5f || NewActual * 2.5f < DesiredCapped), TEXT("HISMC rebuild did not hit target %f %f %f"), NewActual, Desired, DesiredCapped); // if this fires, it means this will always rebuild
+			}
+		}
+	}
+#endif
+	if (CVarASyncInstaneBufferConversion.GetValueOnGameThread() > 0)
+	{
+		UWorld* World = GetWorld();
+		if (World && World->IsGameWorld() && PerInstanceSMData.Num())
+		{
+			World->AsyncPreRegisterLevelStreamingTasks.Increment();
+			while( InstancingRandomSeed == 0 )
+			{
+				InstancingRandomSeed = FMath::Rand();
+			}
+			AsyncBuildInstanceBufferTask = new FAsyncTask<FAsyncBuildInstanceBuffer>(this, World);
+			AsyncBuildInstanceBufferTask->StartBackgroundTask();
+		}
+	}
+	// For some reason we don't have a tree. Build one now!
+	if (PerInstanceSMData.Num() > 0 && (!ClusterTreePtr.IsValid() || ClusterTreePtr->Num() == 0))
+	{
+		check(!IsAsyncBuilding());
+		BuildTree();
+	}
+}
+
 static void GatherInstanceTransformsInArea(const UHierarchicalInstancedStaticMeshComponent& Component, const FBox& AreaBox, int32 Child, TArray<FTransform>& InstanceData)
 {
 	const TArray<FClusterNode>& ClusterTree = *Component.ClusterTreePtr;
-	const FClusterNode& ChildNode = ClusterTree[Child];
-	const FBox WorldNodeBox = FBox(ChildNode.BoundMin, ChildNode.BoundMax).TransformBy(Component.ComponentToWorld);
-	
-	if (AreaBox.Intersect(WorldNodeBox))
+	if (ClusterTree.Num())
 	{
-		if (ChildNode.FirstChild < 0 || AreaBox.IsInside(WorldNodeBox))
+		const FClusterNode& ChildNode = ClusterTree[Child];
+		const FBox WorldNodeBox = FBox(ChildNode.BoundMin, ChildNode.BoundMax).TransformBy(Component.ComponentToWorld);
+
+		if (AreaBox.Intersect(WorldNodeBox))
 		{
-			// Unfortunately ordering of PerInstanceSMData does not match ordering of cluster tree, so we have to use remaping
-			const bool bUseRemaping = Component.SortedInstances.Num() > 0;
-			
-			// In case there no more subdivision or node is completely encapsulated by a area box
-			// add all instances to the result
-			for (int32 i = ChildNode.FirstInstance; i <= ChildNode.LastInstance; ++i)
+			if (ChildNode.FirstChild < 0 || AreaBox.IsInside(WorldNodeBox))
 			{
-				int32 SortedIdx = bUseRemaping ? Component.SortedInstances[i] : i;
-				FTransform InstanceToComponent(Component.PerInstanceSMData[SortedIdx].Transform);
-				if (!InstanceToComponent.GetScale3D().IsZero())
+				// Unfortunately ordering of PerInstanceSMData does not match ordering of cluster tree, so we have to use remaping
+				const bool bUseRemaping = Component.SortedInstances.Num() > 0;
+
+				// In case there no more subdivision or node is completely encapsulated by a area box
+				// add all instances to the result
+				for (int32 i = ChildNode.FirstInstance; i <= ChildNode.LastInstance; ++i)
 				{
-					InstanceData.Add(InstanceToComponent*Component.ComponentToWorld);
+					int32 SortedIdx = bUseRemaping ? Component.SortedInstances[i] : i;
+
+					FTransform InstanceToComponent;
+					if (Component.PerInstanceSMData.IsValidIndex(SortedIdx))
+					{
+						InstanceToComponent = FTransform(Component.PerInstanceSMData[SortedIdx].Transform);
+					}
+					else if (Component.PerInstanceRenderData.IsValid())
+					{
+						// if there's no PerInstanceSMData (e.g. for grass), we'll go ge the transform from the render buffer
+						FInstanceStream const* Inst = Component.PerInstanceRenderData->InstanceBuffer.GetInstance(i);
+						FMatrix XformMat;
+						Inst->GetInstanceTransform(XformMat);
+						InstanceToComponent = FTransform(XformMat);
+					}
+					if (!InstanceToComponent.GetScale3D().IsZero())
+					{
+						InstanceData.Add(InstanceToComponent*Component.ComponentToWorld);
+					}
 				}
 			}
-		}
-		else
-		{
-			for (int32 i = ChildNode.FirstChild; i <= ChildNode.LastChild; ++i)
+			else
 			{
-				GatherInstanceTransformsInArea(Component, AreaBox, i, InstanceData);
+				for (int32 i = ChildNode.FirstChild; i <= ChildNode.LastChild; ++i)
+				{
+					GatherInstanceTransformsInArea(Component, AreaBox, i, InstanceData);
+				}
 			}
 		}
 	}
@@ -2040,6 +2196,8 @@ void UHierarchicalInstancedStaticMeshComponent::FlushAccumulatedNavigationUpdate
 {
 	if (AccumulatedNavigationDirtyArea.IsValid)
 	{
+		QUICK_SCOPE_CYCLE_COUNTER(STAT_UHierarchicalInstancedStaticMeshComponent_FlushAccumulatedNavigationUpdates);
+
 		const TArray<FClusterNode>& ClusterTree = *ClusterTreePtr;
 		UNavigationSystem* NavSys = UNavigationSystem::GetCurrent(GetWorld());
 		// Check if this component is registered in navigation system
@@ -2052,3 +2210,27 @@ void UHierarchicalInstancedStaticMeshComponent::FlushAccumulatedNavigationUpdate
 		AccumulatedNavigationDirtyArea.Init();
 	}
 }
+
+static void RebuildFoliageTrees(const TArray<FString>& Args)
+{
+#if !(UE_BUILD_SHIPPING || UE_BUILD_TEST)
+	UE_LOG(LogConsoleResponse, Display, TEXT("Rebuild Foliage Trees"));
+#endif
+	for (TObjectIterator<UHierarchicalInstancedStaticMeshComponent> It; It; ++It)
+	{
+		UHierarchicalInstancedStaticMeshComponent* Comp = *It;
+		if (Comp && !Comp->IsTemplate() && !Comp->IsPendingKill() && !Comp->IsAsyncBuilding())
+		{
+			Comp->BuildTree();
+			Comp->ReleasePerInstanceRenderData();
+			Comp->MarkRenderStateDirty();
+		}
+	}
+}
+
+static FAutoConsoleCommand RebuildFoliageTreesCmd(
+	TEXT("foliage.RebuildFoliageTrees"),
+	TEXT("Rebuild the trees for non-grass foliage."),
+	FConsoleCommandWithArgsDelegate::CreateStatic(&RebuildFoliageTrees)
+	);
+
diff --git a/Engine/Source/Runtime/Engine/Private/ImageUtils.cpp b/Engine/Source/Runtime/Engine/Private/ImageUtils.cpp
--- a/Engine/Source/Runtime/Engine/Private/ImageUtils.cpp
+++ b/Engine/Source/Runtime/Engine/Private/ImageUtils.cpp
@@ -153,6 +153,12 @@ UTexture2D* FImageUtils::CreateTexture2D(int32 SrcWidth, int32 SrcHeight, const
 	}
 	Tex2D->Source.UnlockMip(0);
 
+	// Set the Source Guid/Hash if specified
+	if (InParams.SourceGuidHash.IsValid())
+	{
+		Tex2D->Source.SetId(InParams.SourceGuidHash, true);
+	}
+
 	// Set compression options.
 	Tex2D->SRGB = InParams.bSRGB;
 	Tex2D->CompressionSettings	= InParams.CompressionSettings;
diff --git a/Engine/Source/Runtime/Engine/Private/InstancedStaticMesh.cpp b/Engine/Source/Runtime/Engine/Private/InstancedStaticMesh.cpp
--- a/Engine/Source/Runtime/Engine/Private/InstancedStaticMesh.cpp
+++ b/Engine/Source/Runtime/Engine/Private/InstancedStaticMesh.cpp
@@ -15,6 +15,12 @@ const int32 InstancedStaticMeshMaxTexCoord = 8;
 
 IMPLEMENT_HIT_PROXY(HInstancedStaticMeshInstance, HHitProxy);
 
+TAutoConsoleVariable<int32> CVarMinLOD(
+	TEXT("foliage.MinLOD"),
+	-1,
+	TEXT("Used to discard the top LODs for performance evaluation. -1: Disable all effects of this cvar."));
+
+
 /** InstancedStaticMeshInstance hit proxy */
 void HInstancedStaticMeshInstance::AddReferencedObjects(FReferenceCollector& Collector)
 {
@@ -57,6 +63,9 @@ void FStaticMeshInstanceBuffer::SetupCPUAccess(UInstancedStaticMeshComponent* In
 
 void FStaticMeshInstanceBuffer::Init(UInstancedStaticMeshComponent* InComponent, const TArray<TRefCountPtr<HHitProxy> >& InHitProxies)
 {
+	QUICK_SCOPE_CYCLE_COUNTER(STAT_FStaticMeshInstanceBuffer_Init);
+	double StartTime = FPlatformTime::Seconds();
+
 	bool bUseRemapTable = InComponent->PerInstanceSMData.Num() == InComponent->InstanceReorderTable.Num();
 
 	int32 NumInstances = InComponent->PerInstanceSMData.Num();
@@ -75,7 +84,7 @@ void FStaticMeshInstanceBuffer::Init(UInstancedStaticMeshComponent* InComponent,
 
 	FColor HitProxyColor(ForceInit);
 	bool bSelected = false;
-	for (int32 InstanceIndex = 0; InstanceIndex < (int32)NumInstances; InstanceIndex++)
+	for (int32 InstanceIndex = 0; InstanceIndex < NumInstances; InstanceIndex++)
 	{
 #if WITH_EDITOR
 		if (InHitProxies.Num() == NumInstances)
@@ -102,36 +111,22 @@ void FStaticMeshInstanceBuffer::Init(UInstancedStaticMeshComponent* InComponent,
 			InstanceRenderData->NullifyInstance();
 		}
 	}
+
+	float ThisTime = (StartTime - FPlatformTime::Seconds()) * 1000.0f;
+	if (ThisTime > 30.0f)
+	{
+		UE_LOG(LogStaticMesh, Display, TEXT("Took %6.2fms to set up instance buffer for %d instances for component %s."), ThisTime, NumInstances, *InComponent->GetFullName());
+	}
 }
 
 void FStaticMeshInstanceBuffer::InitFromPreallocatedData(UInstancedStaticMeshComponent* InComponent, FStaticMeshInstanceData& Other)
 {
+	QUICK_SCOPE_CYCLE_COUNTER(STAT_FStaticMeshInstanceBuffer_InitFromPreallocatedData);
 	AllocateData(Other);
 	SetupCPUAccess(InComponent);
+	InstanceData->SetAllowCPUAccess(true);		// GDC demo hack!
 }
 
-#if 0 // delete me
-
-/** Serializer. */
-FArchive& operator<<(FArchive& Ar,FStaticMeshInstanceBuffer& InstanceBuffer)
-{
-	int32 NumInstances = InstanceBuffer.GetNumInstances();
-	Ar << InstanceBuffer.Stride << NumInstances;
-
-	if(Ar.IsLoading())
-	{
-		// Allocate the vertex data storage type.
-		InstanceBuffer.AllocateData();
-	}
-
-	// Serialize the vertex data.
-	InstanceBuffer.InstanceData->Serialize(Ar);
-
-	return Ar;
-}
-
-#endif
-
 /**
  * Specialized assignment operator, only used when importing LOD's.  
  */
@@ -142,6 +137,7 @@ void FStaticMeshInstanceBuffer::operator=(const FStaticMeshInstanceBuffer &Other
 
 void FStaticMeshInstanceBuffer::InitRHI()
 {
+	QUICK_SCOPE_CYCLE_COUNTER(STAT_FStaticMeshInstanceBuffer_InitRHI);
 	check(InstanceData);
 	FResourceArrayInterface* ResourceArray = InstanceData->GetResourceArray();
 	if(ResourceArray->GetResourceDataSize())
@@ -170,7 +166,7 @@ void FStaticMeshInstanceBuffer::AllocateData(FStaticMeshInstanceData& Other)
 {
 	AllocateData();
 	Other.SetAllowCPUAccess(InstanceData->GetAllowCPUAccess());
-	Exchange(Other, *InstanceData);
+	FMemory::Memswap(&Other, InstanceData, sizeof(FStaticMeshInstanceData));
 }
 
 
@@ -935,28 +931,30 @@ void UInstancedStaticMeshComponent::InitInstanceBody(int32 InstanceIdx, FBodyIns
 
 void UInstancedStaticMeshComponent::CreateAllInstanceBodies()
 {
+	QUICK_SCOPE_CYCLE_COUNTER(STAT_UInstancedStaticMeshComponent_CreateAllInstanceBodies);
+
 	UBodySetup* BodySetup = GetBodySetup();
 	if (BodySetup)
 	{
-		int32 NumBodies = PerInstanceSMData.Num();
-		InstanceBodies.Init(NumBodies);
-
-		TArray<FTransform> Transforms;
-		Transforms.Reserve(NumBodies);
-
-		for (int32 i = 0; i < NumBodies; ++i)
-		{
-			InstanceBodies[i] = new FBodyInstance;
-			FBodyInstance* Instance = InstanceBodies[i];
-			Transforms.Add(FTransform(PerInstanceSMData[i].Transform) * ComponentToWorld);
-
-			Instance->CopyBodyInstancePropertiesFrom(&BodyInstance);
-			Instance->InstanceBodyIndex = i; // Set body index 
-			Instance->bAutoWeld = false;
-
-			// make sure we never enable bSimulatePhysics for ISMComps
-			Instance->bSimulatePhysics = false;
-		}
+	    int32 NumBodies = PerInstanceSMData.Num();
+	    InstanceBodies.Init(NumBodies);
+    
+	    TArray<FTransform> Transforms;
+	    Transforms.Reserve(NumBodies);
+    
+	    for (int32 i = 0; i < NumBodies; ++i)
+	    {
+		    InstanceBodies[i] = new FBodyInstance;
+		    FBodyInstance* Instance = InstanceBodies[i];
+		    Transforms.Add(FTransform(PerInstanceSMData[i].Transform) * ComponentToWorld);
+    
+		    Instance->CopyBodyInstancePropertiesFrom(&BodyInstance);
+		    Instance->InstanceBodyIndex = i; // Set body index 
+		    Instance->bAutoWeld = false;
+    
+		    // make sure we never enable bSimulatePhysics for ISMComps
+		    Instance->bSimulatePhysics = false;
+	    }
 
 		if (NumBodies > 0)
 		{
@@ -967,6 +965,7 @@ void UInstancedStaticMeshComponent::CreateAllInstanceBodies()
 
 void UInstancedStaticMeshComponent::ClearAllInstanceBodies()
 {
+	QUICK_SCOPE_CYCLE_COUNTER(STAT_UInstancedStaticMeshComponent_ClearAllInstanceBodies);
 	for (int32 i = 0; i < InstanceBodies.Num(); i++)
 	{
 		check(InstanceBodies[i]);
@@ -1445,7 +1444,7 @@ void UInstancedStaticMeshComponent::PartialNavigationUpdate(int32 InstanceIdx)
 	UNavigationSystem::UpdateNavOctree(this);
 }
 
-bool UInstancedStaticMeshComponent::DoCustomNavigableGeometryExport(struct FNavigableGeometryExport* GeomExport) const
+bool UInstancedStaticMeshComponent::DoCustomNavigableGeometryExport(FNavigableGeometryExport& GeomExport) const
 {
 	if (StaticMesh && StaticMesh->NavCollision)
 	{
@@ -1457,10 +1456,10 @@ bool UInstancedStaticMeshComponent::DoCustomNavigableGeometryExport(struct FNavi
 		
 		if (NavCollision->bHasConvexGeometry)
 		{
-			GeomExport->ExportCustomMesh(NavCollision->ConvexCollision.VertexBuffer.GetData(), NavCollision->ConvexCollision.VertexBuffer.Num(),
+			GeomExport.ExportCustomMesh(NavCollision->ConvexCollision.VertexBuffer.GetData(), NavCollision->ConvexCollision.VertexBuffer.Num(),
 				NavCollision->ConvexCollision.IndexBuffer.GetData(), NavCollision->ConvexCollision.IndexBuffer.Num(), FTransform::Identity);
 
-			GeomExport->ExportCustomMesh(NavCollision->TriMeshCollision.VertexBuffer.GetData(), NavCollision->TriMeshCollision.VertexBuffer.Num(),
+			GeomExport.ExportCustomMesh(NavCollision->TriMeshCollision.VertexBuffer.GetData(), NavCollision->TriMeshCollision.VertexBuffer.Num(),
 				NavCollision->TriMeshCollision.IndexBuffer.GetData(), NavCollision->TriMeshCollision.IndexBuffer.Num(), FTransform::Identity);
 		}
 		else
@@ -1468,12 +1467,12 @@ bool UInstancedStaticMeshComponent::DoCustomNavigableGeometryExport(struct FNavi
 			UBodySetup* BodySetup = StaticMesh->BodySetup;
 			if (BodySetup)
 			{
-				GeomExport->ExportRigidBodySetup(*BodySetup, FTransform::Identity);
+				GeomExport.ExportRigidBodySetup(*BodySetup, FTransform::Identity);
 			}
 		}
 
 		// Hook per instance transform delegate
-		GeomExport->SetNavDataPerInstanceTransformDelegate(FNavDataPerInstanceTransformDelegate::CreateUObject(this, &UInstancedStaticMeshComponent::GetNavigationPerInstanceTransforms));
+		GeomExport.SetNavDataPerInstanceTransformDelegate(FNavDataPerInstanceTransformDelegate::CreateUObject(this, &UInstancedStaticMeshComponent::GetNavigationPerInstanceTransforms));
 	}
 
 	// we don't want "regular" collision export for this component
@@ -1607,6 +1606,11 @@ void UInstancedStaticMeshComponent::SelectInstance(bool bInSelected, int32 InIns
 #endif
 }
 
+static TAutoConsoleVariable<int32> CVarCullAllInVertexShader(
+	TEXT("foliage.CullAllInVertexShader"),
+	0,
+	TEXT("Debugging, if this is greater than 0, cull all instances in the vertex shader."));
+
 void FInstancedStaticMeshVertexFactoryShaderParameters::SetMesh( FRHICommandList& RHICmdList, FShader* VertexShader,const class FVertexFactory* VertexFactory,const class FSceneView& View,const struct FMeshBatchElement& BatchElement,uint32 DataFlags ) const 
 {
 	FLocalVertexFactoryShaderParameters::SetMesh(RHICmdList, VertexShader, VertexFactory, View, BatchElement, DataFlags);
@@ -1623,6 +1627,15 @@ void FInstancedStaticMeshVertexFactoryShaderParameters::SetMesh( FRHICommandList
 		InstancingWorldViewOriginOne.W = 1.0f;
 		if (InstancingUserData && BatchElement.InstancedLODRange)
 		{
+			int32 FirstLOD = InstancingUserData->MinLOD;
+
+			int32 DebugMin = FMath::Min(CVarMinLOD.GetValueOnRenderThread(), InstancingUserData->MeshRenderData->LODResources.Num() - 1);
+			if (DebugMin >= 0)
+			{
+				FirstLOD = FMath::Max(FirstLOD, DebugMin);
+			}
+
+
 			float SphereRadius = InstancingUserData->MeshRenderData->Bounds.SphereRadius;
 			float MinSize = CVarFoliageMinimumScreenSize.GetValueOnRenderThread();
 			float LODScale = CVarFoliageLODDistanceScale.GetValueOnRenderThread();
@@ -1670,7 +1683,7 @@ void FInstancedStaticMeshVertexFactoryShaderParameters::SetMesh( FRHICommandList
 				}
 
 				InstancingViewZCompare.X = MIN_flt;
-				if (BatchElement.InstancedLODIndex)
+				if (BatchElement.InstancedLODIndex > FirstLOD)
 				{
 					float CurCut = FMath::Sqrt(Fac / InstancingUserData->MeshRenderData->ScreenSize[BatchElement.InstancedLODIndex]);
 					if (CurCut < FinalCull)
@@ -1729,9 +1742,17 @@ void FInstancedStaticMeshVertexFactoryShaderParameters::SetMesh( FRHICommandList
 			{
 				InstancingFadeOutParams.Y = 0.f;
 			}
+			if (CVarCullAllInVertexShader.GetValueOnRenderThread() > 0)
+			{
+				InstancingFadeOutParams.Z = 0.0f;
+				InstancingFadeOutParams.W = 0.0f;
+			}
+			else
+			{
 			InstancingFadeOutParams.Z = InstancingUserData->bRenderSelected ? 1.f : 0.f;
 			InstancingFadeOutParams.W = InstancingUserData->bRenderUnselected ? 1.f : 0.f;
 		}
+		}
 		SetShaderValue(RHICmdList, VS, InstancingFadeOutParamsParameter, InstancingFadeOutParams );
 	}
 
diff --git a/Engine/Source/Runtime/Engine/Private/InstancedStaticMesh.h b/Engine/Source/Runtime/Engine/Private/InstancedStaticMesh.h
--- a/Engine/Source/Runtime/Engine/Private/InstancedStaticMesh.h
+++ b/Engine/Source/Runtime/Engine/Private/InstancedStaticMesh.h
@@ -54,6 +54,7 @@
 extern TAutoConsoleVariable<float> CVarFoliageMinimumScreenSize;
 extern TAutoConsoleVariable<float> CVarFoliageLODDistanceScale;
 extern TAutoConsoleVariable<float> CVarRandomLODRange;
+extern TAutoConsoleVariable<int32> CVarMinLOD;
 
 
 // This must match the maximum a user could specify in the material (see 
@@ -121,6 +122,12 @@ public:
 		return InstanceData->GetData();
 	}
 
+	const FInstanceStream* GetInstance(int32 InstanceIndex) const
+	{
+		return InstanceData->GetInstanceWriteAddress(InstanceIndex);
+	}
+
+
 	// FRenderResource interface.
 	virtual void InitRHI() override;
 	virtual FString GetFriendlyName() const { return TEXT("Static-mesh instances"); }
@@ -154,6 +161,8 @@ struct FInstancingUserData
 	int32 StartCullDistance;
 	int32 EndCullDistance;
 
+	int32 MinLOD;
+
 	bool bRenderSelected;
 	bool bRenderUnselected;
 };
@@ -301,6 +310,7 @@ struct FPerInstanceRenderData
 		// Create hit proxies for each instance if the component wants
 		if (GIsEditor && InComponent->bHasPerInstanceHitProxies)
 		{
+			QUICK_SCOPE_CYCLE_COUNTER(STAT_FPerInstanceRenderData_HitProxies);
 			HitProxies.Empty(InComponent->PerInstanceSMData.Num());
 			for (int32 InstanceIdx=0; InstanceIdx < InComponent->PerInstanceSMData.Num(); InstanceIdx++)
 			{
@@ -374,7 +384,7 @@ public:
 		{
 			InComponent->PerInstanceRenderData = MakeShareable(new FPerInstanceRenderData(InComponent, Other, InFeatureLevel));
 			PerInstanceRenderData = InComponent->PerInstanceRenderData;
-			InComponent->bPerInstanceRenderDataWasPrebuilt = true;
+			InComponent->bPerInstanceRenderDataWasPrebuilt = InComponent->PerInstanceSMData.Num() == 0;
 		}
 		NumInstances = PerInstanceRenderData->InstanceBuffer.GetNumInstances();
 		InitResources();
@@ -613,6 +623,7 @@ private:
 		UserData_AllInstances.MeshRenderData = InComponent->StaticMesh->RenderData;
 		UserData_AllInstances.StartCullDistance = InComponent->InstanceStartCullDistance;
 		UserData_AllInstances.EndCullDistance = InComponent->InstanceEndCullDistance;
+		UserData_AllInstances.MinLOD = ClampedMinLOD;
 		UserData_AllInstances.bRenderSelected = true;
 		UserData_AllInstances.bRenderUnselected = true;
 		UserData_AllInstances.RenderData = bInstanced ? nullptr : &InstancedRenderData;
diff --git a/Engine/Source/Runtime/Engine/Private/Interpolation.cpp b/Engine/Source/Runtime/Engine/Private/Interpolation.cpp
--- a/Engine/Source/Runtime/Engine/Private/Interpolation.cpp
+++ b/Engine/Source/Runtime/Engine/Private/Interpolation.cpp
@@ -72,6 +72,10 @@
 #include "Particles/ParticleSystemComponent.h"
 #include "Particles/ParticleSystemReplay.h"
 #include "ContentStreaming.h"
+#include "Matinee/InterpTrackFloatAnimBPParam.h"
+#include "Matinee/InterpTrackInstFloatAnimBPParam.h"
+#include "Animation/AnimInstance.h"
+#include "Animation/AnimBlueprintGeneratedClass.h"
 
 #if WITH_EDITOR
 #include "UnrealEd.h"
@@ -2019,7 +2023,7 @@ void UInterpGroup::UpdateGroup(float NewPosition, UInterpGroupInst* GrInst, bool
 	check( InterpTracks.Num() == GrInst->TrackInst.Num() );
 
 	// Update animation state of Actor.
-#if WITH_EDITORONLY_DATA
+#if 0
 	// if in editor and preview and anim control exists, let them update weight before update track
 	if ( GIsEditor && bPreview && HasAnimControlTrack() )
 	{
@@ -2047,8 +2051,10 @@ void UInterpGroup::UpdateGroup(float NewPosition, UInterpGroupInst* GrInst, bool
 		}
 	}
 
+#if 0
 	// Update animation state of Actor.
 	UpdateAnimWeights(NewPosition, GrInst, bPreview, bJump);
+#endif // disable this, it doesn't w
 }
 
 bool UInterpGroup::HasSelectedTracks() const
@@ -6811,6 +6817,8 @@ void UInterpTrackInstDirector::TermTrackInst(UInterpTrack* Track)
 UInterpTrackFade::UInterpTrackFade(const FObjectInitializer& ObjectInitializer)
 	: Super(ObjectInitializer)
 {
+	FadeColor = FLinearColor::Black;
+	
 	bOnePerGroup = true;
 	bDirGroupOnly = true;
 	TrackInstClass = UInterpTrackInstFade::StaticClass();
@@ -6851,10 +6859,7 @@ void UInterpTrackFade::UpdateTrack(float NewPosition, UInterpTrackInst* TrInst,
 		APlayerController* PC = Cast<APlayerController>(GrInst->GetGroupActor());
 		if(PC && PC->PlayerCameraManager && !PC->PlayerCameraManager->IsPendingKill())
 		{
-			PC->PlayerCameraManager->bEnableFading = true;
-			PC->PlayerCameraManager->FadeAmount = GetFadeAmountAtTime(NewPosition);
-			// disable the Kismet fade control so that we don't thrash
-			PC->PlayerCameraManager->FadeTimeRemaining = 0.0f;
+			PC->PlayerCameraManager->SetManualCameraFade(GetFadeAmountAtTime(NewPosition), FadeColor, bFadeAudio);
 		}
 	}
 }
@@ -6879,12 +6884,12 @@ void UInterpTrackInstFade::TermTrackInst(UInterpTrack* Track)
 	UInterpTrackFade *FadeTrack = Cast<UInterpTrackFade>(Track);
 	if (FadeTrack == NULL || !FadeTrack->bPersistFade)
 	{
-		UInterpGroupInst* GrInst = CastChecked<UInterpGroupInst>( GetOuter() );
-		APlayerController* PC = Cast<APlayerController>(GrInst->GroupActor);
+		UInterpGroupInst* const GrInst = CastChecked<UInterpGroupInst>(GetOuter());
+		APlayerController* const PC = Cast<APlayerController>(GrInst->GroupActor);
 		if(PC && PC->PlayerCameraManager && !PC->PlayerCameraManager->IsPendingKill())
 		{
-			PC->PlayerCameraManager->bEnableFading = false;
-			PC->PlayerCameraManager->FadeAmount = 0.f;
+			PC->PlayerCameraManager->StopCameraFade();
+
 			// if the player is remote, ensure they got it
 			// this handles cases where the LDs stream out this level immediately afterwards,
 			// which can mean the client never gets the matinee replication if it was temporarily unresponsive
@@ -7287,14 +7292,18 @@ void UInterpTrackAnimControl::PreviewUpdateTrack(float NewPosition, class UInter
 	int32 ChannelIndex = CalcChannelIndex();
 
 	UAnimSequence* NewAnimSeq = NULL;
-	float NewAnimPosition, TimeElapsed = 0.f;
+
+	UInterpGroupInst* GrInst = CastChecked<UInterpGroupInst>(TrInst->GetOuter());
+	AMatineeActor* MatineeActor = CastChecked<AMatineeActor>(GrInst->GetOuter());
+	bool bJump = !( MatineeActor->bIsPlaying );
+	float NewAnimPosition, TimeElapsed = (bJump || NewPosition < AnimInst->LastUpdatePosition)? 0.f : NewPosition - AnimInst->LastUpdatePosition;
 	bool bNewLooping;
 	bool bResetTime = GetAnimForTime(NewPosition, &NewAnimSeq, NewAnimPosition, bNewLooping);
 
 	if( NewAnimSeq != NULL )
 	{
 		// if we're going backward or if not @ the first frame of the animation
-		bool bFireNotifier = !bSkipAnimNotifiers && (TimeElapsed < 0.f || !bResetTime) ;
+		bool bFireNotifier = !bSkipAnimNotifiers && (!bResetTime) ;
 		IMatineeAnimInterface * IMAI = Cast<IMatineeAnimInterface>(Actor);
 		if (IMAI)
 		{
@@ -8000,7 +8009,9 @@ void UInterpTrackSound::PreviewUpdateTrack(float NewPosition, UInterpTrackInst*
 	if ( bTimeChangedDrastically && MatineeActor->bIsScrubbing )
 	{
 		FSoundTrackKey& SoundTrackKey = GetSoundTrackKeyAtPosition(NewPosition);
-		const bool bIsInRange = NewPosition >= SoundTrackKey.Time && NewPosition <= ( SoundTrackKey.Time + SoundTrackKey.Sound->Duration );
+		const bool bIsInRange = (SoundTrackKey.Sound)
+				? (NewPosition >= SoundTrackKey.Time && NewPosition <= (SoundTrackKey.Time + SoundTrackKey.Sound->Duration))
+				: false;
 
 		USoundCue* TempPlaybackAudioCue = NewObject<USoundCue>();
 		UAudioComponent* Component = FAudioDevice::CreateComponent(TempPlaybackAudioCue, NULL, NULL, false, false);
@@ -9624,3 +9635,178 @@ UBillboardComponent* AMatineeActor::GetSpriteComponent() const { return SpriteCo
 /** Returns SpriteComponent subobject **/
 UBillboardComponent* AMaterialInstanceActor::GetSpriteComponent() const { return SpriteComponent; }
 #endif
+
+
+/*-----------------------------------------------------------------------------
+	Float anim BP parameter track.
+-----------------------------------------------------------------------------*/
+UInterpTrackFloatAnimBPParam::UInterpTrackFloatAnimBPParam(const FObjectInitializer& ObjectInitializer)
+	: Super(ObjectInitializer)
+	, bRefreshParamter(false)
+{
+	TrackInstClass = UInterpTrackInstFloatAnimBPParam::StaticClass();
+	TrackTitle = TEXT("Float AnimBP Param");
+}
+
+#if WITH_EDITOR
+
+void UInterpTrackFloatAnimBPParam::PostEditChangeProperty(FPropertyChangedEvent& PropertyChangedEvent)
+{
+	Super::PostEditChangeProperty(PropertyChangedEvent);
+
+	UProperty* PropertyThatChanged = PropertyChangedEvent.Property;
+	FName PropertyName = PropertyThatChanged != NULL ? PropertyThatChanged->GetFName() : NAME_None;
+
+	if (PropertyName == GET_MEMBER_NAME_CHECKED(UInterpTrackFloatAnimBPParam, ParamName) ||
+		PropertyName == GET_MEMBER_NAME_CHECKED(UInterpTrackFloatAnimBPParam, AnimBlueprintClass))
+	{
+		bRefreshParamter = true;
+	}
+}
+#endif // WITH_EDITOR
+
+int32 UInterpTrackFloatAnimBPParam::AddKeyframe(float Time, UInterpTrackInst* TrInst, EInterpCurveMode InitInterpMode)
+{
+	int32 NewKeyIndex = FloatTrack.AddPoint( Time, 0.f );
+	FloatTrack.Points[NewKeyIndex].InterpMode = InitInterpMode;
+	if ( NewKeyIndex > 0 )
+	{
+		if ( NewKeyIndex < FloatTrack.Points.Num() - 1 )
+		{
+			const float Duration = FloatTrack.Points[NewKeyIndex+1].InVal - FloatTrack.Points[NewKeyIndex-1].InVal;
+			const float Remaining = FloatTrack.Points[NewKeyIndex+1].InVal - FloatTrack.Points[NewKeyIndex].InVal;
+			const float DurationPct = ( Duration > 0.f ? ((Duration - Remaining) / Duration) : 0.f );
+			if( FloatTrack.Points[NewKeyIndex-1].InterpMode == CIM_Linear || FloatTrack.Points[NewKeyIndex-1].InterpMode == CIM_Constant )	// Linear or Constant interpolation
+			{
+				FloatTrack.Points[NewKeyIndex].OutVal = FMath::Lerp<float>(FloatTrack.Points[NewKeyIndex-1].OutVal, FloatTrack.Points[NewKeyIndex+1].OutVal, DurationPct);
+			}
+			else	// Cubic Interpolation
+			{
+				FloatTrack.Points[NewKeyIndex].OutVal = FMath::CubicInterp<float>(FloatTrack.Points[NewKeyIndex-1].OutVal, FloatTrack.Points[NewKeyIndex-1].LeaveTangent * Duration, FloatTrack.Points[NewKeyIndex+1].OutVal, FloatTrack.Points[NewKeyIndex+1].ArriveTangent * Duration, DurationPct);
+			}
+		}
+		else	// Same position as previous point
+		{
+			FloatTrack.Points[NewKeyIndex].OutVal = FloatTrack.Points[NewKeyIndex-1].OutVal;
+		}
+	}
+	else if ( NewKeyIndex < FloatTrack.Points.Num() - 1 )	// Same position as next point
+	{
+		FloatTrack.Points[NewKeyIndex].OutVal = FloatTrack.Points[NewKeyIndex+1].OutVal;
+	}
+	FloatTrack.AutoSetTangents(CurveTension);
+	return NewKeyIndex;
+}
+
+
+void UInterpTrackFloatAnimBPParam::PreviewUpdateTrack(float NewPosition, UInterpTrackInst* TrInst)
+{
+	if (bRefreshParamter)
+	{
+		UInterpTrackInstFloatAnimBPParam* ParamTrackInst = Cast<UInterpTrackInstFloatAnimBPParam>(TrInst);
+		if(ParamTrackInst != NULL)
+		{
+			ParamTrackInst->RefreshParameter(this);
+		}
+
+		bRefreshParamter = false;
+	}
+
+	UpdateTrack(NewPosition, TrInst, false);
+}
+
+
+void UInterpTrackFloatAnimBPParam::UpdateTrack(float NewPosition, UInterpTrackInst* TrInst, bool bJump)
+{
+	UInterpTrackInstFloatAnimBPParam* ParamTrackInst = Cast<UInterpTrackInstFloatAnimBPParam>(TrInst);
+	if (ParamTrackInst != NULL)
+	{
+		float NewFloatValue = FloatTrack.Eval(NewPosition, 0.f);
+		// set value
+		ParamTrackInst->SetValue(NewFloatValue);
+	}
+}
+
+UInterpTrackInstFloatAnimBPParam::UInterpTrackInstFloatAnimBPParam(const FObjectInitializer& ObjectInitializer)
+	: Super(ObjectInitializer)
+	, AnimScriptInstance(nullptr)
+	, ResetFloat(0.f)
+	, ParamProperty(nullptr)
+{
+}
+
+void UInterpTrackInstFloatAnimBPParam::InitTrackInst(UInterpTrack* Track)
+{
+	Super::InitTrackInst(Track);
+	RefreshParameter(Track);
+}
+
+void UInterpTrackInstFloatAnimBPParam::RefreshParameter(UInterpTrack* Track)
+{
+	// if currently has correct setup, restore actor state
+	RestoreActorState(Track);
+
+	AnimScriptInstance = nullptr;
+	ParamProperty = nullptr;
+
+	UInterpTrackFloatAnimBPParam* ParamTrack = Cast<UInterpTrackFloatAnimBPParam>(Track);
+	if(ParamTrack != nullptr && ParamTrack->ParamName != NAME_None)
+	{
+		AActor* Actor=GetGroupActor();
+		if(Actor)
+		{
+			TInlineComponentArray<USkeletalMeshComponent*> SkeletalMeshComponents;
+			Actor->GetComponents(SkeletalMeshComponents);
+
+			if(SkeletalMeshComponents.Num() > 0)
+			{
+				UAnimInstance* NewAnimInstance = SkeletalMeshComponents[0]->GetAnimInstance();
+
+				if(NewAnimInstance && NewAnimInstance->GetClass() == ParamTrack->AnimBlueprintClass)
+				{
+					AnimScriptInstance = NewAnimInstance;
+					// make sure the class has the parameter
+					ParamProperty = NewAnimInstance->GetClass()->FindPropertyByName(ParamTrack->ParamName);
+				}
+			}
+		}
+	}
+
+	// save actor state since now property might have changed
+	SaveActorState(Track);
+}
+
+void UInterpTrackInstFloatAnimBPParam::SetValue(float InValue)
+{
+	if (AnimScriptInstance && ParamProperty)
+	{
+		float* Value = ParamProperty->ContainerPtrToValuePtr<float>(AnimScriptInstance);
+		if(Value)
+		{
+			*Value = InValue;
+		}
+	}
+}
+
+float UInterpTrackInstFloatAnimBPParam::GetValue() const
+{
+	if(AnimScriptInstance && ParamProperty)
+	{
+		float* Value = ParamProperty->ContainerPtrToValuePtr<float>(AnimScriptInstance);
+		if(Value)
+		{
+			return *Value;
+		}
+	}
+
+	return 0.f;
+}
+
+void UInterpTrackInstFloatAnimBPParam::SaveActorState(UInterpTrack* Track) 
+{
+	ResetFloat = GetValue();
+}
+void UInterpTrackInstFloatAnimBPParam::RestoreActorState(UInterpTrack* Track)
+{
+	SetValue(ResetFloat);
+}
\ No newline at end of file
