UE4 - merge GDC branch, code [REDACTED_HANDLE] to main

diff --git a/Engine/Source/Runtime/Engine/Private/Level.cpp b/Engine/Source/Runtime/Engine/Private/Level.cpp
--- a/Engine/Source/Runtime/Engine/Private/Level.cpp
+++ b/Engine/Source/Runtime/Engine/Private/Level.cpp
@@ -163,6 +163,7 @@ FArchive& operator<<( FArchive& Ar, FPrecomputedVolumeDistanceField& D )
 
 FLevelSimplificationDetails::FLevelSimplificationDetails()
  : DetailsPercentage(70.f)
+ , bCreatePackagePerAsset(true)
  , LandscapeExportLOD(7)
  , bGenerateLandscapeNormalMap(true)
  , bGenerateLandscapeMetallicMap(false)
@@ -691,6 +692,9 @@ void ULevel::IncrementalUpdateComponents(int32 NumComponentsToUpdate, bool bReru
 		bool bAllComponentsRegistered = true;
 		if (Actor)
 		{
+#if PERF_TRACK_DETAILED_ASYNC_STATS
+			FScopeCycleCounterUObject ContextScope(Actor);
+#endif
 			bAllComponentsRegistered = Actor->IncrementalRegisterComponents(NumComponentsToUpdate);
 		}
 
@@ -714,6 +718,9 @@ void ULevel::IncrementalUpdateComponents(int32 NumComponentsToUpdate, bool bReru
 		CurrentActorIndexForUpdateComponents	= 0;
 		bAreComponentsCurrentlyRegistered		= true;
 		
+#if PERF_TRACK_DETAILED_ASYNC_STATS
+		QUICK_SCOPE_CYCLE_COUNTER(STAT_ULevel_IncrementalUpdateComponents_RerunConstructionScripts);
+#endif
 		if (bRerunConstructionScripts && !IsTemplate() && !GIsUCCMakeStandaloneHeaderGenerator)
 		{
 			// Don't rerun construction scripts until after all actors' components have been registered.  This
@@ -723,6 +730,9 @@ void ULevel::IncrementalUpdateComponents(int32 NumComponentsToUpdate, bool bReru
 			{
 				if (Actor)
 				{
+#if PERF_TRACK_DETAILED_ASYNC_STATS
+					FScopeCycleCounterUObject ContextScope(Actor);
+#endif
 					Actor->RerunConstructionScripts();
 				}
 			}
@@ -1732,8 +1742,11 @@ bool ULevel::IsCurrentLevel() const
 
 void ULevel::ApplyWorldOffset(const FVector& InWorldOffset, bool bWorldShift)
 {
-	if (bTextureStreamingBuilt)
+	QUICK_SCOPE_CYCLE_COUNTER(STAT_ULevel_ApplyWorldOffset);
+
+	if (bTextureStreamingBuilt && !InWorldOffset.IsZero())
 	{
+		QUICK_SCOPE_CYCLE_COUNTER(STAT_ULevel_ApplyWorldOffset_TextureStreaming);
 		// Update texture streaming data to account for the move
 		for (TMap< UTexture2D*, TArray<FStreamableTextureInstance> >::TIterator It(TextureToInstancesMap); It; ++It)
 		{
@@ -1743,51 +1756,68 @@ void ULevel::ApplyWorldOffset(const FVector& InWorldOffset, bool bWorldShift)
 				TextureInfo[i].BoundingSphere.Center+= InWorldOffset;
 			}
 		}
-		
-		// Re-add level data to a manager
-		IStreamingManager::Get().AddPreparedLevel( this );
+
+		// Re-add level data to a manager, in case level is visible it this point
+		if (bIsVisible)
+		{
+			IStreamingManager::Get().AddPreparedLevel( this );
+		}
 	}
 
 	// Move precomputed light samples
-	if (PrecomputedLightVolume)
+	if (PrecomputedLightVolume && !InWorldOffset.IsZero())
 	{
-		if (!PrecomputedLightVolume->IsAddedToScene())
-		{
-			PrecomputedLightVolume->ApplyWorldOffset(InWorldOffset);
-		}
-		// At world origin rebasing all registered volumes will be moved during FScene shifting
-		// Otherwise we need to send a command to move just this volume
-		else if (!bWorldShift) 
+		QUICK_SCOPE_CYCLE_COUNTER(STAT_ULevel_ApplyWorldOffset_PrecomputedLightVolume);
+		// Shift light volume only if it's going to be used
+		static const auto AllowStaticLightingVar = IConsoleManager::Get().FindTConsoleVariableDataInt(TEXT("r.AllowStaticLighting"));
+		const bool bAllowStaticLighting = (!AllowStaticLightingVar || AllowStaticLightingVar->GetValueOnGameThread() != 0);
+		if (bAllowStaticLighting) 
 		{
-			ENQUEUE_UNIQUE_RENDER_COMMAND_TWOPARAMETER(
- 				ApplyWorldOffset_PLV,
- 				FPrecomputedLightVolume*, InPrecomputedLightVolume, PrecomputedLightVolume,
- 				FVector, InWorldOffset, InWorldOffset,
- 			{
-				InPrecomputedLightVolume->ApplyWorldOffset(InWorldOffset);
- 			});
+			if (!PrecomputedLightVolume->IsAddedToScene())
+			{
+				PrecomputedLightVolume->ApplyWorldOffset(InWorldOffset);
+			}
+			// At world origin rebasing all registered volumes will be moved during FScene shifting
+			// Otherwise we need to send a command to move just this volume
+			else if (!bWorldShift) 
+			{
+				ENQUEUE_UNIQUE_RENDER_COMMAND_TWOPARAMETER(
+ 					ApplyWorldOffset_PLV,
+ 					FPrecomputedLightVolume*, InPrecomputedLightVolume, PrecomputedLightVolume,
+ 					FVector, InWorldOffset, InWorldOffset,
+ 				{
+					InPrecomputedLightVolume->ApplyWorldOffset(InWorldOffset);
+ 				});
+			}
 		}
 	}
 
-	// Iterate over all actors in the level and move them
-	for (int32 ActorIndex = 0; ActorIndex < Actors.Num(); ActorIndex++)
 	{
-		AActor* Actor = Actors[ActorIndex];
-		if (Actor)
+		QUICK_SCOPE_CYCLE_COUNTER(STAT_ULevel_ApplyWorldOffset_Actors);
+		// Iterate over all actors in the level and move them
+		for (int32 ActorIndex = 0; ActorIndex < Actors.Num(); ActorIndex++)
 		{
-			FVector Offset = (bWorldShift && Actor->bIgnoresOriginShifting) ? FVector::ZeroVector : InWorldOffset;
-						
-			if (!Actor->IsA(ANavigationData::StaticClass())) // Navigation data will be moved in NavigationSystem
+			AActor* Actor = Actors[ActorIndex];
+			if (Actor)
 			{
-				Actor->ApplyWorldOffset(Offset, bWorldShift);
+				FVector Offset = (bWorldShift && Actor->bIgnoresOriginShifting) ? FVector::ZeroVector : InWorldOffset;
+
+				if (!Actor->IsA(ANavigationData::StaticClass())) // Navigation data will be moved in NavigationSystem
+				{
+					FScopeCycleCounterUObject Context(Actor);
+					Actor->ApplyWorldOffset(Offset, bWorldShift);
+				}
 			}
 		}
 	}
 	
-	// Move model geometry
-	for (int32 CompIdx = 0; CompIdx < ModelComponents.Num(); ++CompIdx)
 	{
-		ModelComponents[CompIdx]->ApplyWorldOffset(InWorldOffset, bWorldShift);
+		QUICK_SCOPE_CYCLE_COUNTER(STAT_ULevel_ApplyWorldOffset_Model);
+		// Move model geometry
+		for (int32 CompIdx = 0; CompIdx < ModelComponents.Num(); ++CompIdx)
+		{
+			ModelComponents[CompIdx]->ApplyWorldOffset(InWorldOffset, bWorldShift);
+		}
 	}
 }
 
diff --git a/Engine/Source/Runtime/Engine/Private/LevelStreaming.cpp b/Engine/Source/Runtime/Engine/Private/LevelStreaming.cpp
--- a/Engine/Source/Runtime/Engine/Private/LevelStreaming.cpp
+++ b/Engine/Source/Runtime/Engine/Private/LevelStreaming.cpp
@@ -272,7 +272,14 @@ void ULevelStreaming::Serialize( FArchive& Ar )
 
 FName ULevelStreaming::GetLODPackageName() const
 {
-	return LODPackageNames.IsValidIndex(LevelLODIndex) ? LODPackageNames[LevelLODIndex] : GetWorldAssetPackageFName();
+	if (LODPackageNames.IsValidIndex(LevelLODIndex))
+	{
+		return LODPackageNames[LevelLODIndex];
+	}
+	else
+	{
+		return GetWorldAssetPackageFName();
+	}
 }
 
 FName ULevelStreaming::GetLODPackageNameToLoad() const
@@ -329,6 +336,8 @@ bool ULevelStreaming::RequestLevel(UWorld* PersistentWorld, bool bAllowLevelLoad
 		return false;
 	}
 	
+	QUICK_SCOPE_CYCLE_COUNTER(STAT_ULevelStreaming_RequestLevel);
+	FScopeCycleCounterUObject Context(PersistentWorld);
 	// Package name we want to load
 	FName DesiredPackageName = PersistentWorld->IsGameWorld() ? GetLODPackageName() : GetWorldAssetPackageFName();
 	FName DesiredPackageNameToLoad = PersistentWorld->IsGameWorld() ? GetLODPackageNameToLoad() : PackageNameToLoad;
@@ -407,33 +416,41 @@ bool ULevelStreaming::RequestLevel(UWorld* PersistentWorld, bool bAllowLevelLoad
 		PIEInstanceID = PersistentWorld->GetOutermost()->PIEInstanceID;
 #endif
 		const FString NonPrefixedLevelName = UWorld::StripPIEPrefixFromPackageName(DesiredPackageName.ToString(), PersistentWorld->StreamingLevelsPrefix);
-					
-		// Do the duplication
-		UWorld* PIELevelWorld = UWorld::DuplicateWorldForPIE(NonPrefixedLevelName, PersistentWorld);
-		if (PIELevelWorld)
+		UPackage* EditorLevelPackage = FindObjectFast<UPackage>(nullptr, FName(*NonPrefixedLevelName));
+		
+		bool bShouldDuplicate = EditorLevelPackage && (bBlockOnLoad || EditorLevelPackage->IsDirty() || !GEngine->PreferToStreamLevelsInPIE());
+		if (bShouldDuplicate)
 		{
-			PIELevelWorld->PersistentLevel->bAlreadyMovedActors = true; // As we have duplicated the world, the actors will already have been transformed
-			check(PendingUnloadLevel == NULL);
-			SetLoadedLevel(PIELevelWorld->PersistentLevel);
+			// Do the duplication
+			UWorld* PIELevelWorld = UWorld::DuplicateWorldForPIE(NonPrefixedLevelName, PersistentWorld);
+			if (PIELevelWorld)
+			{
+				PIELevelWorld->PersistentLevel->bAlreadyMovedActors = true; // As we have duplicated the world, the actors will already have been transformed
+				check(PendingUnloadLevel == NULL);
+				SetLoadedLevel(PIELevelWorld->PersistentLevel);
 
-			// Broadcast level loaded event to blueprints
-			OnLevelLoaded.Broadcast();
+				// Broadcast level loaded event to blueprints
+				{
+					QUICK_SCOPE_CYCLE_COUNTER(STAT_OnLevelLoaded_Broadcast);
+					OnLevelLoaded.Broadcast();
+				}
 
-			return true;
-		}
-		else if (PersistentWorld->WorldComposition == NULL) // In world composition streaming levels are not loaded by default
-		{
-			if ( bAllowLevelLoadRequests )
-			{
-				UE_LOG(LogLevelStreaming, Log, TEXT("World to duplicate for PIE '%s' not found. Attempting load."), *NonPrefixedLevelName);
+				return true;
 			}
-			else
+			else if (PersistentWorld->WorldComposition == NULL) // In world composition streaming levels are not loaded by default
 			{
-				UE_LOG(LogLevelStreaming, Warning, TEXT("Unable to duplicate PIE World: '%s'"), *NonPrefixedLevelName);
-				for (TObjectIterator<UWorld> It; It; ++It)
+				if ( bAllowLevelLoadRequests )
 				{
-					UWorld *W = *It;
-					UE_LOG(LogLevelStreaming, Warning, TEXT("    Loaded World: %s"), *W->GetPathName() );
+					UE_LOG(LogLevelStreaming, Log, TEXT("World to duplicate for PIE '%s' not found. Attempting load."), *NonPrefixedLevelName);
+				}
+				else
+				{
+					UE_LOG(LogLevelStreaming, Warning, TEXT("Unable to duplicate PIE World: '%s'"), *NonPrefixedLevelName);
+					for (TObjectIterator<UWorld> It; It; ++It)
+					{
+						UWorld *W = *It;
+						UE_LOG(LogLevelStreaming, Warning, TEXT("    Loaded World: %s"), *W->GetPathName() );
+					}
 				}
 			}
 		}
@@ -643,6 +660,26 @@ bool ULevelStreaming::IsLevelLoaded() const
 	return LoadedLevel != NULL;
 }
 
+bool ULevelStreaming::IsStreamingStatePending() const
+{
+	UWorld* PersistentWorld = GetWorld();
+	if (PersistentWorld)
+	{
+		if (IsLevelLoaded() == ShouldBeLoaded() && IsLevelVisible() == ShouldBeVisible())
+		{
+			FName DesiredPackageName = PersistentWorld->IsGameWorld() ? GetLODPackageName() : GetWorldAssetPackageFName();
+			if (!LoadedLevel || LoadedLevel->GetOutermost()->GetFName() == DesiredPackageName)
+			{
+				return false;
+			}
+		}
+		
+		return true;
+	}
+	
+	return false;
+}
+
 ULevelStreaming* ULevelStreaming::CreateInstance(FString InstanceUniqueName)
 {
 	ULevelStreaming* StreamingLevelInstance = nullptr;
@@ -772,12 +809,12 @@ void ULevelStreaming::RenameForPIE(int32 PIEInstanceID)
 	}
 }
 
-bool ULevelStreaming::ShouldBeLoaded()
+bool ULevelStreaming::ShouldBeLoaded() const
 {
 	return true;
 }
 
-bool ULevelStreaming::ShouldBeVisible()
+bool ULevelStreaming::ShouldBeVisible() const
 {
 	if( GetWorld()->IsGameWorld() )
 	{
@@ -927,12 +964,12 @@ void ULevelStreamingKismet::PostLoad()
 	}
 }
 
-bool ULevelStreamingKismet::ShouldBeVisible()
+bool ULevelStreamingKismet::ShouldBeVisible() const
 {
 	return bShouldBeVisible || (bShouldBeVisibleInEditor && !FApp::IsGame());
 }
 
-bool ULevelStreamingKismet::ShouldBeLoaded()
+bool ULevelStreamingKismet::ShouldBeLoaded() const
 {
 	return bShouldBeLoaded;
 }
@@ -946,7 +983,7 @@ ULevelStreamingAlwaysLoaded::ULevelStreamingAlwaysLoaded(const FObjectInitialize
 	bShouldBeVisible = true;
 }
 
-bool ULevelStreamingAlwaysLoaded::ShouldBeLoaded()
+bool ULevelStreamingAlwaysLoaded::ShouldBeLoaded() const
 {
 	return true;
 }
diff --git a/Engine/Source/Runtime/Engine/Private/LocalPlayer.cpp b/Engine/Source/Runtime/Engine/Private/LocalPlayer.cpp
--- a/Engine/Source/Runtime/Engine/Private/LocalPlayer.cpp
+++ b/Engine/Source/Runtime/Engine/Private/LocalPlayer.cpp
@@ -711,7 +711,7 @@ FSceneView* ULocalPlayer::CalcSceneView( class FSceneViewFamily* ViewFamily,
 		// Apply screen fade effect to screen.
 		if (PlayerController->PlayerCameraManager->bEnableFading)
 		{
-			ViewInitOptions.OverlayColor = PlayerController->PlayerCameraManager->FadeColor.ReinterpretAsLinear();
+			ViewInitOptions.OverlayColor = PlayerController->PlayerCameraManager->FadeColor;
 			ViewInitOptions.OverlayColor.A = FMath::Clamp(PlayerController->PlayerCameraManager->FadeAmount,0.0f,1.0f);
 		}
 
@@ -742,6 +742,7 @@ FSceneView* ULocalPlayer::CalcSceneView( class FSceneViewFamily* ViewFamily,
 	ViewInitOptions.WorldToMetersScale = PlayerController->GetWorldSettings()->WorldToMeters;
 	ViewInitOptions.CursorPos = Viewport->HasMouseCapture() ? FIntPoint(-1, -1) : FIntPoint(Viewport->GetMouseX(), Viewport->GetMouseY());
 	ViewInitOptions.bOriginOffsetThisFrame = PlayerController->GetWorld()->bOriginOffsetThisFrame;
+	ViewInitOptions.bUseFieldOfViewForLOD = ViewInfo.bUseFieldOfViewForLOD;
 	PlayerController->BuildHiddenComponentList(OutViewLocation, /*out*/ ViewInitOptions.HiddenPrimitives);
 
 	FSceneView* const View = new FSceneView(ViewInitOptions);
diff --git a/Engine/Source/Runtime/Engine/Private/Materials/Material.cpp b/Engine/Source/Runtime/Engine/Private/Materials/Material.cpp
--- a/Engine/Source/Runtime/Engine/Private/Materials/Material.cpp
+++ b/Engine/Source/Runtime/Engine/Private/Materials/Material.cpp
@@ -28,7 +28,7 @@
 #include "UObjectAnnotation.h"
 #include "MaterialCompiler.h"
 #include "TargetPlatform.h"
-#include "ComponentReregisterContext.h"
+#include "ComponentRecreateRenderStateContext.h"
 #include "ShaderCompiler.h"
 #include "Materials/MaterialParameterCollection.h"
 #if WITH_EDITOR
@@ -2374,12 +2374,12 @@ void UMaterial::PostEditChangeProperty(FPropertyChangedEvent& PropertyChangedEve
 			UpdateLightmassTextureTracking();
 		}
 
-		// Ensure that any components with static elements using this material are reregistered so changes
-		// are propagated to them. The preview material is only applied to the preview mesh component,
+		// Ensure that any components with static elements using this material have their render state recreated
+		// so changes are propagated to them. The preview material is only applied to the preview mesh component,
 		// and that reregister is handled by the material editor.
 		if (!bIsPreviewMaterial && !bIsMaterialEditorStatsMaterial)
 		{
-			FGlobalComponentReregisterContext RecreateComponents;
+			FGlobalComponentRecreateRenderStateContext RecreateComponentsRenderState;
 		}
 	}
 
diff --git a/Engine/Source/Runtime/Engine/Private/Materials/MaterialShared.cpp b/Engine/Source/Runtime/Engine/Private/Materials/MaterialShared.cpp
--- a/Engine/Source/Runtime/Engine/Private/Materials/MaterialShared.cpp
+++ b/Engine/Source/Runtime/Engine/Private/Materials/MaterialShared.cpp
@@ -16,6 +16,7 @@
 #include "MaterialUniformExpressions.h"
 #include "Developer/TargetPlatform/Public/TargetPlatform.h"
 #include "ComponentReregisterContext.h"
+#include "ComponentRecreateRenderStateContext.h"
 #include "EngineModule.h"
 #include "Engine/Font.h"
 
@@ -1961,11 +1962,17 @@ void FMaterial::RestoreEditorLoadedMaterialShadersFromMemory(const TMap<FMateria
 FMaterialUpdateContext::FMaterialUpdateContext(uint32 Options, EShaderPlatform InShaderPlatform)
 {
 	bool bReregisterComponents = (Options & EOptions::ReregisterComponents) != 0;
+	bool bRecreateRenderStates = (Options & EOptions::RecreateRenderStates) != 0;
+
 	bSyncWithRenderingThread = (Options & EOptions::SyncWithRenderingThread) != 0;
 	if (bReregisterComponents)
 	{
 		ComponentReregisterContext = new FGlobalComponentReregisterContext();
 	}
+	else if (bRecreateRenderStates)
+	{
+		ComponentRecreateRenderStateContext = new FGlobalComponentRecreateRenderStateContext();
+	}
 	if (bSyncWithRenderingThread)
 	{
 		FlushRenderingCommands();
@@ -2016,7 +2023,7 @@ FMaterialUpdateContext::~FMaterialUpdateContext()
 	TArray<const FMaterial*> MaterialResourcesToUpdate;
 	TArray<UMaterialInstance*> InstancesToUpdate;
 
-	bool bUpdateStaticDrawLists = ComponentReregisterContext == NULL;
+	bool bUpdateStaticDrawLists = !ComponentReregisterContext.IsValid() && !ComponentRecreateRenderStateContext.IsValid();
 
 	// If static draw lists must be updated, gather material resources from all updated materials.
 	if (bUpdateStaticDrawLists)
@@ -2104,12 +2111,14 @@ FMaterialUpdateContext::~FMaterialUpdateContext()
 		// safe, e.g. while a component is being registered.
 		GetRendererModule().UpdateStaticDrawListsForMaterials(MaterialResourcesToUpdate);
 	}
-	else
+	else if (ComponentReregisterContext.IsValid())
 	{
-		// We must be reregistering components in which case static draw lists will be recreated.
-		check(ComponentReregisterContext != NULL);
 		ComponentReregisterContext.Reset();
 	}
+	else if (ComponentRecreateRenderStateContext.IsValid())
+	{
+		ComponentRecreateRenderStateContext.Reset();
+	}
 
 	double EndTime = FPlatformTime::Seconds();
 	UE_LOG(LogMaterial, Log,
diff --git a/Engine/Source/Runtime/Engine/Private/NavMeshRenderingComponent.cpp b/Engine/Source/Runtime/Engine/Private/NavMeshRenderingComponent.cpp
--- a/Engine/Source/Runtime/Engine/Private/NavMeshRenderingComponent.cpp
+++ b/Engine/Source/Runtime/Engine/Private/NavMeshRenderingComponent.cpp
@@ -34,12 +34,11 @@ bool UNavMeshRenderingComponent::NeedsLoadForServer() const
 	return false;
 }
 
-void UNavMeshRenderingComponent::TimerFunction()
+bool UNavMeshRenderingComponent::IsNavigationShowFlagSet(const UWorld* World)
 {
 	bool bShowNavigation = false;
 
-	UWorld* MyWorld = GetWorld();
-	FWorldContext* WorldContext = GEngine->GetWorldContextFromWorld(MyWorld);
+	FWorldContext* WorldContext = GEngine->GetWorldContextFromWorld(World);
 
 #if WITH_EDITOR
 	if (GEditor && WorldContext && WorldContext->WorldType == EWorldType::Editor)
@@ -59,6 +58,13 @@ void UNavMeshRenderingComponent::TimerFunction()
 		bShowNavigation = WorldContext && WorldContext->GameViewport && WorldContext->GameViewport->EngineShowFlags.Navigation;
 	}
 
+	return bShowNavigation;
+}
+
+void UNavMeshRenderingComponent::TimerFunction()
+{
+	const bool bShowNavigation = IsNavigationShowFlagSet(GetWorld());
+
 	if (bShowNavigation != !!bCollectNavigationData)
 	{
 		bCollectNavigationData = bShowNavigation;
@@ -108,27 +114,8 @@ FPrimitiveSceneProxy* UNavMeshRenderingComponent::CreateSceneProxy()
 #if WITH_RECAST && !UE_BUILD_SHIPPING && !UE_BUILD_TEST
 	FPrimitiveSceneProxy* SceneProxy = NULL;
 
-	UWorld* MyWorld = GetWorld();
-	FWorldContext* WorldContext = GEngine->GetWorldContextFromWorld(MyWorld);
+	const bool bShowNavigation = IsNavigationShowFlagSet(GetWorld());
 
-	bool bShowNavigation = false;
-#if WITH_EDITOR
-	if (GEditor && WorldContext && WorldContext->WorldType == EWorldType::Editor)
-	{
-		for (FEditorViewportClient* CurrentViewpoer : GEditor->AllViewportClients)
-		{
-			if (CurrentViewpoer->EngineShowFlags.Navigation)
-			{
-				bShowNavigation = true;
-				break;
-			}
-		}
-	}
-	else
-#endif //WITH_EDITOR
-	{
-		bShowNavigation = WorldContext && WorldContext->GameViewport && WorldContext->GameViewport->EngineShowFlags.Navigation;
-	}
 	bCollectNavigationData = bShowNavigation;
 
 	if (bCollectNavigationData && IsVisible())
@@ -198,6 +185,7 @@ FBoxSphereBounds UNavMeshRenderingComponent::CalcBounds(const FTransform& LocalT
 
 void UNavMeshRenderingComponent::GatherData(FNavMeshSceneProxyData& CurrentData) const
 {
+	QUICK_SCOPE_CYCLE_COUNTER(STAT_NavMesh_GatherDebugDrawingGeometry);
 #if WITH_RECAST
 	const ARecastNavMesh* NavMesh = Cast<ARecastNavMesh>(GetOwner());
 
@@ -483,9 +471,9 @@ void UNavMeshRenderingComponent::GatherData(FNavMeshSceneProxyData& CurrentData)
 				for (FNavigationOctree::ElementConstIt ElementIt(Node.GetElementIt()); ElementIt; ElementIt++)
 				{
 					const FNavigationOctreeElement& Element = *ElementIt;
-					if (Element.ShouldUseGeometry(NavMesh->GetConfig()) && Element.Data.CollisionData.Num())
+					if (Element.ShouldUseGeometry(NavMesh->GetConfig()) && Element.Data->CollisionData.Num())
 					{
-						const FRecastGeometryCache CachedGeometry(Element.Data.CollisionData.GetData());
+						const FRecastGeometryCache CachedGeometry(Element.Data->CollisionData.GetData());
 						AppendGeometry(PathCollidingGeomVerts, PathCollidingGeomIndices, CachedGeometry.Verts, CachedGeometry.Header.NumVerts, CachedGeometry.Indices, CachedGeometry.Header.NumFaces);
 					}
 				}
diff --git a/Engine/Source/Runtime/Engine/Private/NavMovementComponent.cpp b/Engine/Source/Runtime/Engine/Private/NavMovementComponent.cpp
--- a/Engine/Source/Runtime/Engine/Private/NavMovementComponent.cpp
+++ b/Engine/Source/Runtime/Engine/Private/NavMovementComponent.cpp
@@ -57,6 +57,9 @@ void UNavMovementComponent::UpdateNavAgent(const AActor& Owner)
 		return;
 	}
 
+	// initialize properties from navigation system
+	NavAgentProps.NavWalkingSearchHeightScale = UNavigationSystem::GetCurrent(GetWorld())->GetDefaultSupportedAgentConfig().NavWalkingSearchHeightScale;
+
 	// Can't call GetSimpleCollisionCylinder(), because no components will be registered.
 	float BoundRadius, BoundHalfHeight;	
 	Owner.GetSimpleCollisionCylinder(BoundRadius, BoundHalfHeight);
@@ -71,6 +74,9 @@ void UNavMovementComponent::UpdateNavAgent(const UCapsuleComponent& CapsuleCompo
 		return;
 	}
 
+	// initialize properties from navigation system
+	NavAgentProps.NavWalkingSearchHeightScale = UNavigationSystem::GetCurrent(GetWorld())->GetDefaultSupportedAgentConfig().NavWalkingSearchHeightScale;
+
 	NavAgentProps.AgentRadius = CapsuleComponent.GetScaledCapsuleRadius();
 	NavAgentProps.AgentHeight = CapsuleComponent.GetScaledCapsuleHalfHeight() * 2.f;
 }
diff --git a/Engine/Source/Runtime/Engine/Private/Particles/ParticleComponents.cpp b/Engine/Source/Runtime/Engine/Private/Particles/ParticleComponents.cpp
--- a/Engine/Source/Runtime/Engine/Private/Particles/ParticleComponents.cpp
+++ b/Engine/Source/Runtime/Engine/Private/Particles/ParticleComponents.cpp
@@ -4126,6 +4126,7 @@ void UParticleSystemComponent::WaitForAsyncAndFinalize(EForceAsyncWorkCompletion
 		check(IsInGameThread());
 		SCOPE_CYCLE_COUNTER(STAT_GTSTallTime);
 		double StartTime = FPlatformTime::Seconds();
+		QUICK_SCOPE_CYCLE_COUNTER(STAT_UParticleSystemComponent_WaitForAsyncAndFinalize);
 		FTaskGraphInterface::Get().WaitUntilTaskCompletes(AsyncWork, ENamedThreads::GameThread_Local);
 		float ThisTime = float(FPlatformTime::Seconds() - StartTime) * 1000.0f;
 		if (Behavior != SILENT)
diff --git a/Engine/Source/Runtime/Engine/Private/PhysicsEngine/BodyInstance.cpp b/Engine/Source/Runtime/Engine/Private/PhysicsEngine/BodyInstance.cpp
--- a/Engine/Source/Runtime/Engine/Private/PhysicsEngine/BodyInstance.cpp
+++ b/Engine/Source/Runtime/Engine/Private/PhysicsEngine/BodyInstance.cpp
@@ -24,6 +24,17 @@
 #include "PhysicalMaterials/PhysicalMaterial.h"
 #include "PhysicsEngine/PhysicsSettings.h"
 
+DEFINE_STAT(STAT_InitBody);
+DEFINE_STAT(STAT_InitBodyDebug);
+DEFINE_STAT(STAT_InitBodySceneInteraction);
+DEFINE_STAT(STAT_InitBodyPostAdd);
+DEFINE_STAT(STAT_UpdatePhysFilter);
+DEFINE_STAT(STAT_UpdatePhysFilterPhysX);
+
+DEFINE_STAT(STAT_InitBodies);
+DEFINE_STAT(STAT_BulkSceneAdd);
+DEFINE_STAT(STAT_StaticInitBodies);
+
 ////////////////////////////////////////////////////////////////////////////
 // FCollisionResponse
 ////////////////////////////////////////////////////////////////////////////
@@ -650,7 +661,11 @@ void FBodyInstance::UpdatePhysicsShapeFilterData(uint32 SkelMeshCompID, bool bUs
 		TArray<PxShape*> AllShapes = GetAllShapes(NumSyncShapes);
 
 		PxScene* AsyncScene = (RigidActorAsync != NULL) ? RigidActorAsync->getScene() : NULL;
-		SCENE_LOCK_WRITE(AsyncScene);
+		
+		if(AsyncScene)
+		{
+			SCENE_LOCK_WRITE(AsyncScene);
+		}
 
 		bool bUpdateMassProperties = false;
 		for (int32 ShapeIdx = 0; ShapeIdx < AllShapes.Num(); ShapeIdx++)
@@ -776,7 +791,10 @@ void FBodyInstance::UpdatePhysicsShapeFilterData(uint32 SkelMeshCompID, bool bUs
 			UpdateMassProperties();
 		}
 
-		SCENE_UNLOCK_WRITE(AsyncScene);
+		if(AsyncScene)
+		{
+			SCENE_UNLOCK_WRITE(AsyncScene);
+		}
 	}
 }
 #endif
@@ -784,6 +802,8 @@ void FBodyInstance::UpdatePhysicsShapeFilterData(uint32 SkelMeshCompID, bool bUs
 /** Update the filter data on the physics shapes, based on the owning component flags. */
 void FBodyInstance::UpdatePhysicsFilterData(bool bForceSimpleAsComplex)
 {
+	SCOPE_CYCLE_COUNTER(STAT_UpdatePhysFilter);
+
 	// Do nothing if no physics actor
 	if (!IsValidBodyInstance())
 	{
@@ -914,15 +934,6 @@ void FBodyInstance::UpdatePhysicsFilterData(bool bForceSimpleAsComplex)
 #endif
 }
 
-DEFINE_STAT(STAT_InitBody);
-DEFINE_STAT(STAT_InitBodyDebug);
-DEFINE_STAT(STAT_InitBodySceneInteraction);
-DEFINE_STAT(STAT_InitBodyPostAdd);
-DEFINE_STAT(STAT_UpdatePhysFilter);
-
-DEFINE_STAT(STAT_BulkInit);
-DEFINE_STAT(STAT_BulkSceneAdd);
-
 #if UE_WITH_PHYSICS
 
 void FBodyInstance::InitBody(class UBodySetup* Setup, const FTransform& Transform, class UPrimitiveComponent* PrimComp, class FPhysScene* InRBScene, PhysXAggregateType InAggregate /*= NULL*/, bool bDefer /*= false*/ )
@@ -946,510 +957,9 @@ void FBodyInstance::InitBody(class UBodySetup* Setup, const FTransform& Transfor
 	InitBodies(Bodies, Transforms, Setup, PrimComp, InRBScene, InAggregate, bDefer);
 }
 
-void FBodyInstance::InitBody_Legacy(UBodySetup* Setup, const FTransform& Transform, UPrimitiveComponent* PrimComp, FPhysScene* InRBScene, PxAggregate* InAggregate)
-{
-	SCOPE_CYCLE_COUNTER(STAT_InitBody);
-	check(Setup);
-	AActor* Owner = PrimComp ? PrimComp->GetOwner() : NULL;
-
-#if WITH_PHYSX
-	PhysxUserData = FPhysxUserData(this);
-#endif
-
-	// Make the debug name for this geometry...
-	FString DebugName(TEXT(""));
-#if (WITH_EDITORONLY_DATA || UE_BUILD_DEBUG || LOOKING_FOR_PERF_ISSUES) && !(UE_BUILD_SHIPPING || UE_BUILD_TEST) && !NO_LOGGING
-	if (PrimComp)
-	{
-		SCOPE_CYCLE_COUNTER(STAT_InitBodyDebug);
-		DebugName += FString::Printf(TEXT("Component: %s "), *PrimComp->GetReadableName());
-	}
-
-	if (Setup->BoneName != NAME_None)
-	{
-		SCOPE_CYCLE_COUNTER(STAT_InitBodyDebug);
-		DebugName += FString::Printf(TEXT("Bone: %s "), *Setup->BoneName.ToString());
-	}
-
-#if WITH_PHYSX
-	// Convert to char* for PhysX
-	{
-		SCOPE_CYCLE_COUNTER(STAT_InitBodyDebug);
-		CharDebugName = MakeShareable(new TArray<ANSICHAR>(StringToArray<ANSICHAR>(*DebugName, DebugName.Len() + 1)));
-	}
-#endif
-#endif
-
-	if(!ValidateTransform(Transform, DebugName, Setup))
-	{
-		// Failed to validate transform.
-		return;
-	}
-
-	// remember my owning component
-	OwnerComponent = PrimComp;
-	BodySetup = Setup;
-	Scale3D = Transform.GetScale3D();
-
-	// whenever update filter, this check will trigger, it seems it's best to know when we initialize body instead. 
-	check(BodySetup.IsValid());
-
-	if(bAutoWeld)
-	{
-		UPrimitiveComponent * ParentPrimComponent = PrimComp ? Cast<UPrimitiveComponent>(PrimComp->AttachParent) : NULL;
-		
-		//if we have a parent we will now do the weld and exit any further initialization
-		if (ParentPrimComponent && PrimComp->GetWorld() && PrimComp->GetWorld()->IsGameWorld())
-		{
-			if (PrimComp->WeldToImplementation(ParentPrimComponent, PrimComp->AttachSocketName, false))	//welded new simulated body so initialization is done
-			{
-				return;
-			}
-		}
-	}
-
-	USkeletalMeshComponent* SkelMeshComp = Cast<USkeletalMeshComponent>(OwnerComponent.Get());
-	if (SkelMeshComp != nullptr)
-	{
-		// if not dedicated and if set up is done to use physics, enable flag on
-		if ((BodySetup->PhysicsType == PhysType_Simulated) || (BodySetup->PhysicsType == PhysType_Default))
-		{
-			bool bEnableSimulation = (SkelMeshComp && IsRunningDedicatedServer()) ? SkelMeshComp->bEnablePhysicsOnDedicatedServer : true;
-			bEnableSimulation &= ((BodySetup->PhysicsType == PhysType_Simulated) || (SkelMeshComp->BodyInstance.bSimulatePhysics));	//if unfixed enable. If default look at parent
-			
-			if(bEnableSimulation)
-			{
-				// set simulate to true if using physics
-				bSimulatePhysics = true;
-				if(BodySetup->PhysicsType == PhysType_Simulated)
-				{
-					PhysicsBlendWeight = 1.f;
-				}
-
-			}
-			else
-			{
-				bSimulatePhysics = false;
-				if(BodySetup->PhysicsType == PhysType_Simulated)
-				{
-					PhysicsBlendWeight = 0.f;
-				}
-			}
-		}
-	}
-
-	// See if we are 'static'
-	const bool bPhysicsStatic = (OwnerComponent == NULL) || (OwnerComponent->Mobility != EComponentMobility::Movable);
-
-	// In skeletal case, we need both our bone and skelcomponent flag to be true.
-	// This might be 'and'ing us with ourself, but thats fine.
-	const bool bUseSimulate = ShouldInstanceSimulatingPhysics();
-
-	// Determine if we should start out awake or sleeping
-	bool bShouldStartAwake = bStartAwake;
-	FVector InitialLinVel(EForceInit::ForceInitToZero);
-
-	if (SkelMeshComp != NULL)
-	{
-		bShouldStartAwake = bShouldStartAwake && SkelMeshComp->BodyInstance.bStartAwake;
-	}
-
-	if (Owner != NULL)
-	{
-		InitialLinVel = Owner->GetVelocity();
-
-		if (InitialLinVel.Size() > KINDA_SMALL_NUMBER)
-		{
-			bShouldStartAwake = true;
-		}
-	}
-
-#if WITH_BOX2D
-	if (UBodySetup2D* BodySetup2D = Cast<UBodySetup2D>(Setup))
-	{
-		if (b2World* BoxWorld = FPhysicsIntegration2D::FindAssociatedWorld(PrimComp->GetWorld()))
-		{
-			const b2Vec2 Scale2D = FPhysicsIntegration2D::ConvertUnrealVectorToBox(Scale3D);
-
-			if (Setup != nullptr)
-			{
-				OwnerComponent = PrimComp;
-
-				// Create the body definition
-				b2BodyDef BodyDefinition;
-				if (bPhysicsStatic)
-				{
-					BodyDefinition.type = b2_staticBody;
-				}
-				else
-				{
-					BodyDefinition.type = bUseSimulate ? b2_dynamicBody : b2_kinematicBody;
-				}
-				BodyDefinition.awake = bShouldStartAwake;
-
-				if (bUseSimulate)
-				{
-					BodyDefinition.linearVelocity = FPhysicsIntegration2D::ConvertUnrealVectorToBox(InitialLinVel);
-				}
-
-				// Create the body
-				BodyInstancePtr = BoxWorld->CreateBody(&BodyDefinition);
-				BodyInstancePtr->SetUserData(this);
-
-				// Circles
-				for (const FCircleElement2D& Circle : BodySetup2D->AggGeom2D.CircleElements)
-				{
-					b2CircleShape CircleShape;
-					CircleShape.m_radius = Circle.Radius * Scale3D.Size() / UnrealUnitsPerMeter;
-					CircleShape.m_p.x = Circle.Center.X * Scale2D.x;
-					CircleShape.m_p.y = Circle.Center.Y * Scale2D.y;
-
-					b2FixtureDef FixtureDef;
-					FixtureDef.shape = &CircleShape;
-
-					BodyInstancePtr->CreateFixture(&FixtureDef);
-				}
-
-				// Boxes
-				for (const FBoxElement2D& Box : BodySetup2D->AggGeom2D.BoxElements)
-				{
-					const b2Vec2 HalfBoxSize(Box.Width * 0.5f * Scale2D.x, Box.Height * 0.5f * Scale2D.y);
-					const b2Vec2 BoxCenter(Box.Center.X * Scale2D.x, Box.Center.Y * Scale2D.y);
-
-					b2PolygonShape DynamicBox;
-					DynamicBox.SetAsBox(HalfBoxSize.x, HalfBoxSize.y, BoxCenter, FMath::DegreesToRadians(Box.Angle));
-
-					b2FixtureDef FixtureDef;
-					FixtureDef.shape = &DynamicBox;
-
-					BodyInstancePtr->CreateFixture(&FixtureDef);
-				}
-
-				// Convex hulls
-				for (const FConvexElement2D& Convex : BodySetup2D->AggGeom2D.ConvexElements)
-				{
-					const int32 NumVerts = Convex.VertexData.Num();
-
-					if (NumVerts <= b2_maxPolygonVertices)
-					{
-						TArray<b2Vec2, TInlineAllocator<b2_maxPolygonVertices>> Verts;
-
-						for (int32 VertexIndex = 0; VertexIndex < Convex.VertexData.Num(); ++VertexIndex)
-						{
-							const FVector2D SourceVert = Convex.VertexData[VertexIndex];
-							new (Verts) b2Vec2(SourceVert.X * Scale2D.x, SourceVert.Y * Scale2D.y);
-						}
-
-						b2PolygonShape ConvexPoly;
-						ConvexPoly.Set(Verts.GetData(), Verts.Num());
-
-						b2FixtureDef FixtureDef;
-						FixtureDef.shape = &ConvexPoly;
-
-						BodyInstancePtr->CreateFixture(&FixtureDef);
-					}
-					else
-					{
-						UE_LOG(LogPhysics, Warning, TEXT("Too many vertices in a 2D convex body")); //@TODO: Create a better error message that indicates the asset
-					}
-				}
-
-				// Make sure it contained at least one shape
-				if (BodyInstancePtr->GetFixtureList() == nullptr)
-				{
-					if (DebugName.Len())
-					{
-						UE_LOG(LogPhysics, Log, TEXT("InitBody: failed - no shapes: %s"), *DebugName);
-					}
-
-					BodyInstancePtr->SetUserData(nullptr);
-					BoxWorld->DestroyBody(BodyInstancePtr);
-					BodyInstancePtr = nullptr;
-
-					//clear Owner and Setup info as well to properly clean up the BodyInstance.
-					OwnerComponent = NULL;
-					BodySetup = NULL;
-
-					return;
-				}
-				else
-				{
-					// Position the body
-					SetBodyTransform(Transform, /*bTeleport=*/ true);
-				}
-			}
-		}
-		else
-		{
-			//@TODO: BOX2D: Error out if the world scene mapping wasn't found?
-		}
-
-		// Apply correct physical materials to shape we created.
-		UpdatePhysicalMaterials();
-
-		// Set the filter data on the shapes (call this after setting BodyData because it uses that pointer)
-		UpdatePhysicsFilterData();
-
-		if (!bPhysicsStatic)
-		{
-			// Compute mass (call this after setting BodyData because it uses that pointer)
-			UpdateMassProperties();
-
-			// Update damping
-			UpdateDampingProperties();
-
-			SetMaxAngularVelocity(MaxAngularVelocity, false);
-
-			SetMaxDepenetrationVelocity(bOverrideMaxDepenetrationVelocity ? MaxDepenetrationVelocity : UPhysicsSettings::Get()->MaxDepenetrationVelocity);
-
-			//@TODO: BOX2D: Determine if sleep threshold and solver settings can be configured per-body or not
-#if 0
-			// Set the parameters for determining when to put the object to sleep.
-			float SleepEnergyThresh = PNewDynamic->getSleepThreshold();
-			if (SleepFamily == SF_Sensitive)
-			{
-				SleepEnergyThresh /= 20.f;
-			}
-			PNewDynamic->setSleepThreshold(SleepEnergyThresh);
-			// set solver iteration count 
-			int32 PositionIterCount = FMath::Clamp(PositionSolverIterationCount, 1, 255);
-			int32 VelocityIterCount = FMath::Clamp(VelocitySolverIterationCount, 1, 255);
-			PNewDynamic->setSolverIterationCounts(PositionIterCount, VelocityIterCount);
-#endif
-		}
-
-		// Skip over the rest of the PhysX-specific setup in this function, since it was a 2D object
-		return;
-	}
-#endif
-
-#if WITH_PHYSX
-
-	// If there is already a body instanced, or there is no scene to create it into, do nothing.
-	if (GetPxRigidActor() != NULL || !InRBScene)
-	{
-		// clear Owner and Setup info as well to properly clean up the BodyInstance.
-		OwnerComponent = NULL;
-		BodySetup = NULL;
-
-		return;
-	}
-
-	PxScene* PSceneForNewDynamic = NULL;	// TBD
-	PxScene* PSceneSync = InRBScene->GetPhysXScene(PST_Sync);
-	PxScene* PSceneAsync = InRBScene->HasAsyncScene() ? InRBScene->GetPhysXScene(PST_Async) : NULL;
-
-	PxRigidDynamic* PNewDynamic = NULL;
-	PxRigidActor* PNewActorSync = NULL;
-	PxRigidActor* PNewActorAsync = NULL;
-
-	PxTransform PTransform = U2PTransform(Transform);
-
-	if(bPhysicsStatic)
-	{
-		// Put the static actor in both scenes
-		PNewActorSync = GPhysXSDK->createRigidStatic(PTransform);
-		if(PSceneAsync)
-		{
-			PNewActorAsync = GPhysXSDK->createRigidStatic(PTransform);
-		}
-	}
-	else
-	{
-		PNewDynamic = GPhysXSDK->createRigidDynamic(PTransform);
-
-		// Put the dynamic actor in one scene or the other
-		if (!UseAsyncScene())
-		{
-			PNewActorSync = PNewDynamic;
-			PSceneForNewDynamic = PSceneSync;
-		}
-		else
-		{
-			PNewActorAsync = PNewDynamic;
-			PSceneForNewDynamic = PSceneAsync;
-		}
-
-		// Set kinematic if desired
-		if (bUseSimulate)
-		{
-			PNewDynamic->setRigidDynamicFlag(PxRigidDynamicFlag::eKINEMATIC, false);
-		}
-		else
-		{
-			PNewDynamic->setRigidDynamicFlag(PxRigidDynamicFlag::eKINEMATIC, true);
-		}
-
-		// turn off gravity if desired
-		if (!bEnableGravity)
-		{
-			PNewDynamic->setActorFlag( PxActorFlag::eDISABLE_GRAVITY, true );
-		}
-
-		// If we ever drive this body kinematically, we want to use its target for scene queries, so collision is updated right away, not on the next physics sim
-		PNewDynamic->setRigidDynamicFlag(PxRigidDynamicFlag::eUSE_KINEMATIC_TARGET_FOR_SCENE_QUERIES, true);	
-	}
-
-	// Copy geom from template and scale
-
-	// Sync:
-	if (PNewActorSync != NULL)
-	{
-		Setup->AddShapesToRigidActor(PNewActorSync, Scale3D);
-		PNewActorSync->userData = &PhysxUserData; // Store pointer to owning bodyinstance.
-		PNewActorSync->setName( CharDebugName.IsValid() ? CharDebugName->GetData() : NULL );
-
-		check(FPhysxUserData::Get<FBodyInstance>(PNewActorSync->userData) == this && FPhysxUserData::Get<FBodyInstance>(PNewActorSync->userData)->OwnerComponent != NULL);
-	}
-
-	// Async:
-	if (PNewActorAsync != NULL)
-	{
-		check(PSceneAsync);
-
-		Setup->AddShapesToRigidActor(PNewActorAsync, Scale3D);
-		PNewActorAsync->userData = &PhysxUserData; // Store pointer to owning bodyinstance.
-		PNewActorAsync->setName( CharDebugName.IsValid() ? CharDebugName->GetData() : NULL );
-
-		check(FPhysxUserData::Get<FBodyInstance>(PNewActorAsync->userData) == this && FPhysxUserData::Get<FBodyInstance>(PNewActorAsync->userData)->OwnerComponent != NULL);
-	}
-
-
-	// If we added no shapes, generate warning, destroy actor and bail out (don't add to scene).
-	if ((PNewActorSync && PNewActorSync->getNbShapes() == 0) || ((PNewActorAsync && PNewActorAsync->getNbShapes() == 0)))
-	{
-		if (DebugName.Len())
-		{
-			UE_LOG(LogPhysics, Log, TEXT("InitBody: failed - no shapes: %s"), *DebugName);
-		}
-		// else if probably a level with no bsp
-		if (PNewActorSync)
-		{
-			PNewActorSync->release();
-		}
-		if (PNewActorAsync)
-		{
-			PNewActorAsync->release();
-		}
-
-		//clear Owner and Setup info as well to properly clean up the BodyInstance.
-		OwnerComponent = NULL;
-		BodySetup = NULL;
-
-		return;
-	}
-
-	// Store pointers to PhysX data in RB_BodyInstance
-	RigidActorSync = PNewActorSync;
-	RigidActorAsync = PNewActorAsync;
-
-	// Store scene indices
-	SceneIndexSync = InRBScene->PhysXSceneIndex[PST_Sync];
-	SceneIndexAsync = PSceneAsync ? InRBScene->PhysXSceneIndex[PST_Async] : 0;
-
-	// Apply correct physical materials to shape we created.
-	UpdatePhysicalMaterials();
-
-	// Set the filter data on the shapes (call this after setting BodyData because it uses that pointer)
-	UpdatePhysicsFilterData();
-
-	// Need to add actor into scene before calling putToSleep
-	// check if InAggregate is passed in
-	// BRG N.B. : For now only using InAggregate and BodyAggregate for dynamic bodies that are in the same scene as the aggregate, since otherwise we might have actors from two scenes.
-	// If the aggregate scene is NULL we trust the caller that incoming bodies will be in the same scene.
-	if(InAggregate && PNewDynamic != NULL && (InAggregate->getScene() == PSceneForNewDynamic || InAggregate->getScene() == NULL))
-	{
-		InAggregate->addActor(*PNewDynamic);
-	}
-	else if(PNewDynamic != NULL && PNewDynamic->getNbShapes()  > AggregateBodyShapesThreshold)		// a lot of shapes for a single body case
-	{
-		BodyAggregate = GPhysXSDK->createAggregate(AggregateMaxSize, true);
-		if(BodyAggregate)
-		{
-			SCOPE_CYCLE_COUNTER(STAT_InitBodySceneInteraction);
-			BodyAggregate->addActor(*PNewDynamic);
-			SCOPED_SCENE_WRITE_LOCK(PSceneForNewDynamic);
-			PSceneForNewDynamic->addAggregate(*BodyAggregate);
-		}
-	}
-	else
-	{
-		// Actually add actor(s) to scene(s) (if not artic link)
-		if(PNewActorSync != NULL)
-		{
-			SCOPE_CYCLE_COUNTER(STAT_InitBodySceneInteraction);
-			SCOPED_SCENE_WRITE_LOCK(PSceneSync);
-			PSceneSync->addActor(*PNewActorSync);
-		}
-		if(PNewActorAsync != NULL)
-		{
-			SCOPE_CYCLE_COUNTER(STAT_InitBodySceneInteraction);
-			SCOPED_SCENE_WRITE_LOCK(PSceneAsync);
-			PSceneAsync->addActor(*PNewActorAsync);
-		}
-	}
-
-
-	// Set initial velocities
-	if (PNewDynamic != NULL)
-	{
-		SCOPE_CYCLE_COUNTER(STAT_InitBodyPostAdd);
-		// Compute mass (call this after setting BodyData because it uses that pointer)
-		UpdateMassProperties();
-		// Update damping
-		UpdateDampingProperties();
-
-		SetMaxAngularVelocity(MaxAngularVelocity, false);
-
-		SetMaxDepenetrationVelocity(bOverrideMaxDepenetrationVelocity ? MaxDepenetrationVelocity : UPhysicsSettings::Get()->MaxDepenetrationVelocity);
-
-		// Set initial velocity 
-		if(bUseSimulate)
-		{
-			PNewDynamic->setLinearVelocity( U2PVector(InitialLinVel) );
-		}
-
-		// Set the parameters for determining when to put the object to sleep.
-		float SleepEnergyThresh = PNewDynamic->getSleepThreshold();
-		if( SleepFamily == SF_Sensitive)
-		{
-			SleepEnergyThresh /= 20.f;
-		}
-		PNewDynamic->setSleepThreshold(SleepEnergyThresh);
-		// set solver iteration count 
-		int32 PositionIterCount = FMath::Clamp(PositionSolverIterationCount, 1, 255);
-		int32 VelocityIterCount = FMath::Clamp(VelocitySolverIterationCount, 1, 255);
-		PNewDynamic->setSolverIterationCounts(PositionIterCount, VelocityIterCount);
-
-		if(PNewDynamic->getScene())
-		{
-			// A scene is required to build the DOF constraint
-			CreateDOFLock();
-
-			// wakeUp and putToSleep will issue warnings on kinematic actors
-			if(IsRigidBodyNonKinematic(PNewDynamic) && PNewDynamic->getScene())
-			{
-				// Sleep/wake up as appropriate
-				if(bShouldStartAwake)
-				{
-					// Wake up bodies that are part of a moving actor.
-					PNewDynamic->wakeUp();
-				}
-				else
-				{
-					// Bodies should start out sleeping.
-					PNewDynamic->putToSleep();
-				}
-			}
-		}
-	}
-#endif // WITH_PHYSX
-}
-
-
 void FBodyInstance::InitBodies(TArray<FBodyInstance*>& Bodies, TArray<FTransform>& Transforms, class UBodySetup* BodySetup, class UPrimitiveComponent* PrimitiveComp, class FPhysScene* InRBScene, PhysXAggregateType InAggregate /*= NULL*/, bool bDefer /*= false*/ )
 {
-	SCOPE_CYCLE_COUNTER(STAT_BulkInit);
+	SCOPE_CYCLE_COUNTER(STAT_InitBodies);
 
 	check(BodySetup);
 	check(InRBScene);
@@ -1803,32 +1313,37 @@ void FBodyInstance::InitBodies(TArray<FBodyInstance*>& Bodies, TArray<FTransform
 				PSceneForNewDynamic = PSceneAsync;
 			}
 
-			// Set kinematic if desired
-			if(Instance->ShouldInstanceSimulatingPhysics())
-			{
-				PNewDynamic->setRigidDynamicFlag(PxRigidDynamicFlag::eKINEMATIC, false);
-			}
-			else
-			{
-				PNewDynamic->setRigidDynamicFlag(PxRigidDynamicFlag::eKINEMATIC, true);
-			}
-
 			// turn off gravity if desired
 			if(!Instance->bEnableGravity)
 			{
 				PNewDynamic->setActorFlag(PxActorFlag::eDISABLE_GRAVITY, true);
 			}
-
-			// If we ever drive this body kinematically, we want to use its target for scene queries, so collision is updated right away, not on the next physics sim
-			PNewDynamic->setRigidDynamicFlag(PxRigidDynamicFlag::eUSE_KINEMATIC_TARGET_FOR_SCENE_QUERIES, true);
 		}
 
+		Instance->RigidActorSync = PNewActorSync;
+		Instance->RigidActorAsync = PNewActorAsync;
 		Instance->PhysxUserData = FPhysxUserData(Instance);
 
+		UPhysicalMaterial* SimplePhysMat = Instance->GetSimplePhysicalMaterial();
+		TArray<UPhysicalMaterial*> ComplexPhysMats = Instance->GetComplexPhysicalMaterials();
+		PxMaterial* PSimpleMat = SimplePhysMat->GetPhysXMaterial();
+
+		FShapeData ShapeData;
+		Instance->GetFilterData(ShapeData);
+		Instance->GetShapeFlags(ShapeData, Instance->CollisionEnabled, BodySetup->CollisionTraceFlag == CTF_UseComplexAsSimple);
+
+		if(PNewDynamic)
+		{
+			if(!Instance->ShouldInstanceSimulatingPhysics())
+			{
+				ShapeData.SyncBodyFlags |= PxRigidBodyFlag::eKINEMATIC;
+			}
+			ShapeData.SyncBodyFlags |= PxRigidBodyFlag::eUSE_KINEMATIC_TARGET_FOR_SCENE_QUERIES;
+		}
+
 		if(PNewActorSync)
 		{
-			SCOPED_SCENE_WRITE_LOCK(PSceneSync);
-			BodySetup->AddShapesToRigidActor(PNewActorSync, Instance->Scale3D);
+			BodySetup->AddShapesToRigidActor(Instance, PNewActorSync, PST_Sync, Instance->Scale3D, PSimpleMat, ComplexPhysMats, ShapeData);
 			bInitFail |= PNewActorSync->getNbShapes() == 0;
 			PNewActorSync->userData = &Instance->PhysxUserData;
 			PNewActorSync->setName(Instance->CharDebugName.IsValid() ? Instance->CharDebugName->GetData() : nullptr);
@@ -1841,9 +1356,8 @@ void FBodyInstance::InitBodies(TArray<FBodyInstance*>& Bodies, TArray<FTransform
 
 		if(PNewActorAsync)
 		{
-			SCOPED_SCENE_WRITE_LOCK(PSceneAsync);
 			check(PSceneAsync);
-			BodySetup->AddShapesToRigidActor(PNewActorAsync, Instance->Scale3D);
+			BodySetup->AddShapesToRigidActor(Instance, PNewActorAsync, PST_Async, Instance->Scale3D, PSimpleMat, ComplexPhysMats, ShapeData);
 			bInitFail |= PNewActorAsync->getNbShapes() == 0;
 			PNewActorAsync->userData = &Instance->PhysxUserData;
 			PNewActorAsync->setName(Instance->CharDebugName.IsValid() ? Instance->CharDebugName->GetData() : nullptr);
@@ -1859,11 +1373,13 @@ void FBodyInstance::InitBodies(TArray<FBodyInstance*>& Bodies, TArray<FTransform
 			if(PNewActorSync)
 			{
 				PNewActorSync->release();
+				Instance->RigidActorSync = nullptr;
 			}
 
 			if(PNewActorAsync)
 			{
 				PNewActorAsync->release();
+				Instance->RigidActorAsync = nullptr;
 			}
 
 			Instance->OwnerComponent = nullptr;
@@ -1913,39 +1429,17 @@ void FBodyInstance::InitBodies(TArray<FBodyInstance*>& Bodies, TArray<FTransform
 			Instance->CurrentSceneState = BodyInstanceSceneState::Added;
 		}
 		
-		Instance->RigidActorSync = PNewActorSync;
-		Instance->RigidActorAsync = PNewActorAsync;
 		Instance->SceneIndexSync = InRBScene->PhysXSceneIndex[PST_Sync];
 		Instance->SceneIndexAsync = PSceneAsync ? InRBScene->PhysXSceneIndex[PST_Async] : 0;
 		Instance->InitialLinearVelocity = InitialLinVel;
 		Instance->bWokenExternally = bComponentAwake;
-		Instance->UpdatePhysicsFilterData();
 	}
 
 	// Grab a material to use, for now we only support batching of bodies using the same material
 	UPhysicalMaterial* SimpleMat = Bodies[0]->GetSimplePhysicalMaterial();
 	TArray<UPhysicalMaterial*> ComplexMats = Bodies[0]->GetComplexPhysicalMaterials();
 	PxMaterial* PhysxSimpleMat = SimpleMat->GetPhysXMaterial();
-	
-	// Update material parameters
-	//for(int32 ShapeIdx = 0 ; ShapeIdx < ShapesWritten ; ++ShapeIdx)
-	//{
-	//	SCOPED_SCENE_WRITE_LOCK(PSceneSync);
-	//	SCOPED_SCENE_WRITE_LOCK(PSceneAsync);
-	//	PxShape* Shape = PhysxShapes[ShapeIdx];
-	//	ApplyMaterialToShape(Shape, PhysxSimpleMat, ComplexMats);
-	//}
 
-	{
-		SCOPED_SCENE_WRITE_LOCK(PSceneSync);
-		SCOPED_SCENE_WRITE_LOCK(PSceneAsync);
-		for(int32 BodyIdx = 0 ; BodyIdx < NumBodies ; ++BodyIdx)
-		{
-			FBodyInstance* Instance = Bodies[BodyIdx];
-			Instance->ApplyMaterialToInstanceShapes(PhysxSimpleMat, ComplexMats);
-		}
-	}
-	
 	if(!bCanDefer)
 	{
 		SCOPE_CYCLE_COUNTER(STAT_BulkSceneAdd);
@@ -2157,9 +1651,6 @@ bool FBodyInstance::Weld(FBodyInstance* TheirBody, const FTransform& TheirTM)
 	//@TODO: BOX2D: Implement Weld
 
 #if WITH_PHYSX
-	
-	
-
 	TArray<PxShape *> PNewShapes;
 
 	FTransform MyTM = GetUnrealWorldTransform();
@@ -2172,16 +1663,25 @@ bool FBodyInstance::Weld(FBodyInstance* TheirBody, const FTransform& TheirTM)
 
 	SCOPED_SCENE_WRITE_LOCK(PSyncScene);
 	SCOPED_SCENE_WRITE_LOCK(PAsyncScene);
+	SCOPE_CYCLE_COUNTER(STAT_UpdatePhysMats);
+
+	UPhysicalMaterial* SimplePhysMat = GetSimplePhysicalMaterial();
+	TArray<UPhysicalMaterial*> ComplexPhysMats = GetComplexPhysicalMaterials();
+	PxMaterial* PSimpleMat = SimplePhysMat->GetPhysXMaterial();
 	
+	FShapeData ShapeData;
+	GetFilterData(ShapeData);
+	GetShapeFlags(ShapeData, CollisionEnabled, BodySetup->CollisionTraceFlag == CTF_UseComplexAsSimple);
+
 	//child body gets placed into the same scenes as parent body
 	if (PxRigidActor* MyBody = RigidActorSync)
 	{
-		TheirBody->BodySetup->AddShapesToRigidActor(MyBody, Scale3D, RelativeTM, &PNewShapes);
+		TheirBody->BodySetup->AddShapesToRigidActor(this, MyBody, PST_Sync, Scale3D, PSimpleMat, ComplexPhysMats, ShapeData, RelativeTM, &PNewShapes);
 	}
 
 	if (PxRigidActor* MyBody = RigidActorAsync)
 	{
-		TheirBody->BodySetup->AddShapesToRigidActor(MyBody, Scale3D, RelativeTM, &PNewShapes);
+		TheirBody->BodySetup->AddShapesToRigidActor(this, MyBody, PST_Sync, Scale3D, PSimpleMat, ComplexPhysMats, ShapeData, RelativeTM, &PNewShapes);
 	}
 
 
@@ -3294,8 +2794,13 @@ UPhysicalMaterial* FBodyInstance::GetSimplePhysicalMaterial() const
 TArray<UPhysicalMaterial*> FBodyInstance::GetComplexPhysicalMaterials() const
 {
 	TArray<UPhysicalMaterial*> PhysMaterials;
+	GetComplexPhysicalMaterials(PhysMaterials);
+	return PhysMaterials;
+}
 
-	check( GEngine->DefaultPhysMaterial != NULL );
+void FBodyInstance::GetComplexPhysicalMaterials(TArray<UPhysicalMaterial*> &PhysMaterials) const
+{
+	check(GEngine->DefaultPhysMaterial != NULL);
 	// See if the Material has a PhysicalMaterial
 	UPrimitiveComponent* PrimComp = OwnerComponent.Get();
 	if(PrimComp)
@@ -3303,7 +2808,7 @@ TArray<UPhysicalMaterial*> FBodyInstance::GetComplexPhysicalMaterials() const
 		const int32 NumMaterials = PrimComp->GetNumMaterials();
 		PhysMaterials.SetNum(NumMaterials);
 
-		for(int32 MatIdx=0; MatIdx<NumMaterials; MatIdx++)
+		for(int32 MatIdx = 0; MatIdx < NumMaterials; MatIdx++)
 		{
 			UPhysicalMaterial* PhysMat = GEngine->DefaultPhysMaterial;
 			UMaterialInterface* Material = PrimComp->GetMaterial(MatIdx);
@@ -3316,8 +2821,6 @@ TArray<UPhysicalMaterial*> FBodyInstance::GetComplexPhysicalMaterials() const
 			PhysMaterials[MatIdx] = PhysMat;
 		}
 	}
-
-	return PhysMaterials;
 }
 
 #if WITH_PHYSX
@@ -4495,7 +3998,6 @@ void FBodyInstance::ApplyMaterialToShape(PxShape* PShape, PxMaterial* PSimpleMat
 	// If a triangle mesh, need to get array of materials...
 	if(PShape->getGeometryType() == PxGeometryType::eTRIANGLEMESH)
 	{
-		//TArray<UPhysicalMaterial*> ComplexPhysMats = GetComplexPhysicalMaterials();
 		TArray<PxMaterial*> PComplexMats;
 		PComplexMats.AddUninitialized(ComplexPhysMats.Num());
 		for(int MatIdx = 0; MatIdx < ComplexPhysMats.Num(); MatIdx++)
@@ -4523,7 +4025,7 @@ void FBodyInstance::ApplyMaterialToShape(PxShape* PShape, PxMaterial* PSimpleMat
 	}
 }
 
-void FBodyInstance::ApplyMaterialToInstanceShapes(PxMaterial* PSimpleMat, TArray<UPhysicalMaterial*> ComplexPhysMats)
+void FBodyInstance::ApplyMaterialToInstanceShapes(PxMaterial* PSimpleMat, TArray<UPhysicalMaterial*>& ComplexPhysMats)
 {
 	int32 NumSyncShapes = 0;
 	TArray<PxShape*> AllShapes = GetAllShapes(NumSyncShapes);
@@ -4610,6 +4112,390 @@ void FBodyInstance::InitDynamicProperties()
 		}
 	}
 }
+
+void FBodyInstance::GetFilterData(FShapeData& ShapeData, bool bForceSimpleAsComplex)//PxFilterData& SimData, PxFilterData& QueryDataSimple, PxFilterData& QueryDataComplex, TEnumAsByte<ECollisionEnabled::Type>& UseCollisionEnabled, bool bForceSimpleAsComplex)
+{
+	// Do nothing if no physics actor
+	if(!IsValidBodyInstance())
+	{
+		return;
+	}
+
+	// this can happen in landscape height field collision component
+	if(!BodySetup.IsValid())
+	{
+		return;
+	}
+
+	// Figure out if we are static
+	AActor* Owner = OwnerComponent.IsValid() ? OwnerComponent.Get()->GetOwner() : NULL;
+	const bool bPhysicsStatic = !OwnerComponent.IsValid() || OwnerComponent.Get()->IsWorldGeometry();
+
+	// Grab collision setting from body instance
+	ShapeData.CollisionEnabled = GetCollisionEnabled(); // this checks actor override
+	bool bUseNotifyRBCollision = bNotifyRigidBodyCollision;
+	FCollisionResponseContainer UseResponse = CollisionResponses.GetResponseContainer();
+
+	bool bUseCollisionEnabledOverride = false;
+	bool bResponseOverride = false;
+	bool bNotifyOverride = false;
+
+	// Get skelmeshcomp ID
+	uint32 SkelMeshCompID = 0;
+	if(USkeletalMeshComponent* SkelMeshComp = Cast<USkeletalMeshComponent>(OwnerComponent.Get()))
+	{
+		SkelMeshCompID = SkelMeshComp->GetUniqueID();
+
+		// In skeletal case, collision enable/disable/movement should be overriden by mesh component
+		// being in the physics asset, and not having collision is a waste and it can cause a bug where disconnected bodies
+		if(Owner)
+		{
+			if(Owner->GetActorEnableCollision())	//we only override if actor has collision enabled
+			{
+				ShapeData.CollisionEnabled = SkelMeshComp->BodyInstance.CollisionEnabled;
+			}
+			else
+			{	//actor has collision disabled so disable regardless of component setting
+				ShapeData.CollisionEnabled = ECollisionEnabled::NoCollision;
+			}
+		}
+
+
+		ObjectType = SkelMeshComp->GetCollisionObjectType();
+		bUseCollisionEnabledOverride = true;
+
+		if(BodySetup->CollisionReponse == EBodyCollisionResponse::BodyCollision_Enabled)
+		{
+			UseResponse.SetAllChannels(ECR_Block);
+		}
+		else if(BodySetup->CollisionReponse == EBodyCollisionResponse::BodyCollision_Disabled)
+		{
+			UseResponse.SetAllChannels(ECR_Ignore);
+		}
+
+		UseResponse = FCollisionResponseContainer::CreateMinContainer(UseResponse, SkelMeshComp->BodyInstance.CollisionResponses.GetResponseContainer());
+		bUseNotifyRBCollision = bUseNotifyRBCollision && SkelMeshComp->BodyInstance.bNotifyRigidBodyCollision;
+		bResponseOverride = true;
+		bNotifyOverride = true;
+	}
+
+#if WITH_EDITOR
+	// if no collision, but if world wants to enable trace collision for components, allow it
+	if((ShapeData.CollisionEnabled == ECollisionEnabled::NoCollision) && Owner && (Owner->IsA(AVolume::StaticClass()) == false))
+	{
+		UWorld* World = Owner->GetWorld();
+		UPrimitiveComponent* PrimComp = OwnerComponent.Get();
+		if(World && World->bEnableTraceCollision &&
+		   (PrimComp->IsA(UStaticMeshComponent::StaticClass()) || PrimComp->IsA(USkeletalMeshComponent::StaticClass()) || PrimComp->IsA(UBrushComponent::StaticClass())))
+		{
+			//UE_LOG(LogPhysics, Warning, TEXT("Enabling collision %s : %s"), *GetNameSafe(Owner), *GetNameSafe(OwnerComponent.Get()));
+			// clear all other channel just in case other people using those channels to do something
+			UseResponse.SetAllChannels(ECR_Ignore);
+			ShapeData.CollisionEnabled = ECollisionEnabled::QueryOnly;
+			bResponseOverride = true;
+			bUseCollisionEnabledOverride = true;
+		}
+	}
+#endif
+
+	const bool bUseComplexAsSimple = !bForceSimpleAsComplex && (BodySetup.Get()->CollisionTraceFlag == CTF_UseComplexAsSimple);
+	const bool bUseSimpleAsComplex = bForceSimpleAsComplex || (BodySetup.Get()->CollisionTraceFlag == CTF_UseSimpleAsComplex);
+
+	if(PxRigidActor* PActor = GetPxRigidActor())
+	{
+		if(ShapeData.CollisionEnabled != ECollisionEnabled::NoCollision)
+		{
+			bool * bNotifyOverridePtr = bNotifyOverride ? &bUseNotifyRBCollision : NULL;
+
+			PxFilterData PSimFilterData;
+			PxFilterData PSimpleQueryData;
+			PxFilterData PComplexQueryData;
+			int32 CompID = (OwnerComponent != nullptr) ? OwnerComponent->GetUniqueID() : 0;
+			CreateShapeFilterData(ObjectType, CompID, UseResponse, SkelMeshCompID, InstanceBodyIndex, PSimpleQueryData, PSimFilterData, bUseCCD && !bPhysicsStatic, bUseNotifyRBCollision, bPhysicsStatic);	//CCD is determined by root body in case of welding
+			PComplexQueryData = PSimpleQueryData;
+			
+			// Set output sim data
+			ShapeData.FilterData.SimFilter = PSimFilterData;
+
+			// Build filterdata variations for complex and simple
+			PSimpleQueryData.word3 |= EPDF_SimpleCollision;
+			if(bUseSimpleAsComplex)
+			{
+				PSimpleQueryData.word3 |= EPDF_ComplexCollision;
+			}
+
+			PComplexQueryData.word3 |= EPDF_ComplexCollision;
+			if(bUseComplexAsSimple)
+			{
+				PComplexQueryData.word3 |= EPDF_SimpleCollision;
+			}
+
+			ShapeData.FilterData.QuerySimpleFilter = PSimpleQueryData;
+			ShapeData.FilterData.QueryComplexFilter = PComplexQueryData;
+		}
+	}
+}
+
+void FBodyInstance::InitStaticBodies(TArray<FBodyInstance*>& Bodies, TArray<FTransform>& Transforms, class UBodySetup* BodySetup, class UPrimitiveComponent* PrimitiveComp, class FPhysScene* InRBScene)
+{
+	SCOPE_CYCLE_COUNTER(STAT_StaticInitBodies);
+
+	int32 NumBodies = Bodies.Num();
+	AActor* OwningActor = PrimitiveComp ? PrimitiveComp->GetOwner() : nullptr;
+
+	UPhysicalMaterial* SimpleMat;
+	TArray<UPhysicalMaterial*> ComplexMats;
+	PxMaterial* PhysxSimpleMat = nullptr;
+
+	PxScene* PSceneSync = InRBScene->GetPhysXScene(PST_Sync);
+	PxScene* PSceneAsync = InRBScene->HasAsyncScene() ? InRBScene->GetPhysXScene(PST_Async) : nullptr;
+
+	TArray<PxActor*> NewSyncActors;
+	TArray<PxActor*> NewAsyncActors;
+
+	if(PSceneAsync)
+	{
+		NewAsyncActors.Reserve(NumBodies);
+	}
+	NewSyncActors.Reserve(NumBodies);
+
+	for(int32 BodyIdx = 0 ; BodyIdx < NumBodies ; ++BodyIdx)
+	{
+		FBodyInstance* Instance = Bodies[BodyIdx];
+		FTransform& Transform = Transforms[BodyIdx];
+		FShapeData ShapeData;
+
+		Instance->OwnerComponent = PrimitiveComp;
+		Instance->BodySetup = BodySetup;
+		Instance->Scale3D = Transform.GetScale3D();
+
+		// Handle autowelding here to avoid extra work
+		if(Instance->bAutoWeld)
+		{
+			UPrimitiveComponent * ParentPrimComponent = PrimitiveComp ? Cast<UPrimitiveComponent>(PrimitiveComp->AttachParent) : NULL;
+
+			//if we have a parent we will now do the weld and exit any further initialization
+			if(ParentPrimComponent && PrimitiveComp->GetWorld() && PrimitiveComp->GetWorld()->IsGameWorld())
+			{
+				if(PrimitiveComp->WeldToImplementation(ParentPrimComponent, PrimitiveComp->AttachSocketName, false))	//welded new simulated body so initialization is done
+				{
+					return;
+				}
+			}
+		}
+
+		PxRigidActor* PNewActorSync = nullptr;
+		PxRigidActor* PNewActorAsync = nullptr;
+		PxTransform PTransform = U2PTransform(Transform);
+		bool bInitFail = false;
+
+		// Don't process if we've already got a body
+		if(Instance->GetPxRigidActor())
+		{
+			Instance->OwnerComponent = nullptr;
+			Instance->BodySetup = nullptr;
+
+			continue;
+		}
+
+		Instance->PhysxUserData = FPhysxUserData(Instance);
+		PNewActorSync = GPhysXSDK->createRigidStatic(PTransform);
+		PNewActorSync->userData = &Instance->PhysxUserData;
+		check(FPhysxUserData::Get<FBodyInstance>(PNewActorSync->userData) == Instance && FPhysxUserData::Get<FBodyInstance>(PNewActorSync->userData)->OwnerComponent != NULL);
+		NewSyncActors.Add(PNewActorSync);
+
+		if(PSceneAsync)
+		{
+			PNewActorAsync = GPhysXSDK->createRigidStatic(PTransform);
+			PNewActorAsync->userData = &Instance->PhysxUserData;
+			check(FPhysxUserData::Get<FBodyInstance>(PNewActorAsync->userData) == Instance && FPhysxUserData::Get<FBodyInstance>(PNewActorAsync->userData)->OwnerComponent != NULL);
+			NewAsyncActors.Add(PNewActorAsync);
+		}
+
+		Instance->RigidActorSync = PNewActorSync;
+		Instance->RigidActorAsync = PNewActorAsync;
+		Instance->SceneIndexSync = InRBScene->PhysXSceneIndex[PST_Sync];
+		Instance->SceneIndexAsync = PSceneAsync ? InRBScene->PhysXSceneIndex[PST_Async] : 0;
+
+		// Get filter data
+		PxFilterData SimFilter;
+		PxFilterData QuerySimpleFilter;
+		PxFilterData QueryComplexFilter;
+		TEnumAsByte<ECollisionEnabled::Type> UseCollisionEnabled;
+		Instance->GetFilterData(ShapeData);
+		Instance->GetShapeFlags(ShapeData, Instance->CollisionEnabled, BodySetup->CollisionTraceFlag == CTF_UseComplexAsSimple);
+
+		Instance->GetComplexPhysicalMaterials(ComplexMats);
+		SimpleMat = Instance->GetSimplePhysicalMaterial();
+		PhysxSimpleMat = SimpleMat->GetPhysXMaterial();
+
+		if(PNewActorSync)
+		{
+			BodySetup->AddShapesToRigidActor(Instance, PNewActorSync, PST_Sync, Instance->Scale3D, PhysxSimpleMat, ComplexMats, ShapeData, FTransform::Identity);
+		}
+		
+		if(PNewActorAsync)
+		{
+			BodySetup->AddShapesToRigidActor(Instance, PNewActorAsync, PST_Async, Instance->Scale3D, PhysxSimpleMat, ComplexMats, ShapeData, FTransform::Identity);
+		}
+	}
+
+	InRBScene->DeferAddActors(Bodies, NewSyncActors, PST_Sync);
+	if(PSceneAsync)
+	{
+		InRBScene->DeferAddActors(Bodies, NewAsyncActors, PST_Async);
+	}
+}
+
+void FBodyInstance::SetShapeFlags(TEnumAsByte<ECollisionEnabled::Type> UseCollisionEnabled, PxShape* PShape, EPhysicsSceneType SceneType, const bool bUseComplexAsSimple)
+{
+	// If query collision is enabled..
+	bool bUpdateMassProperties = false;
+	if(UseCollisionEnabled != ECollisionEnabled::NoCollision)
+	{
+		AActor* Owner = OwnerComponent.IsValid() ? OwnerComponent.Get()->GetOwner() : NULL;
+		const bool bPhysicsStatic = !OwnerComponent.IsValid() || OwnerComponent.Get()->IsWorldGeometry();
+		
+		// Only perform scene queries in the synchronous scene for static shapes
+		if(bPhysicsStatic)
+		{
+			PShape->setFlag(PxShapeFlag::eSCENE_QUERY_SHAPE, SceneType == PST_Sync);
+		}
+		// If non-static, always enable scene queries
+		else
+		{
+			PShape->setFlag(PxShapeFlag::eSCENE_QUERY_SHAPE, true);
+		}
+
+		// See if we want physics collision
+		bool bSimCollision = (UseCollisionEnabled == ECollisionEnabled::QueryAndPhysics);
+
+		// Triangle mesh is 'complex' geom
+		if(PShape->getGeometryType() == PxGeometryType::eTRIANGLEMESH)
+		{
+			// on dynamic objects and objects which don't use complex as simple, tri mesh not used for sim
+			if(!bSimCollision || !bUseComplexAsSimple)
+			{
+				PShape->setFlag(PxShapeFlag::eSIMULATION_SHAPE, false);
+			}
+			else
+			{
+				PShape->setFlag(PxShapeFlag::eSIMULATION_SHAPE, true);
+			}
+
+			if(OwnerComponent == NULL || !OwnerComponent->IsA(UModelComponent::StaticClass()))
+			{
+				PShape->setFlag(PxShapeFlag::eVISUALIZATION, false); // dont draw the tri mesh, we can see it anyway, and its slow
+			}
+		}
+		// Everything else is 'simple'
+		else
+		{
+			// See if we currently have sim collision
+			bool bCurrentSimCollision = (PShape->getFlags() & PxShapeFlag::eSIMULATION_SHAPE);
+			// Enable sim collision
+			if(bSimCollision && !bCurrentSimCollision)
+			{
+				bUpdateMassProperties = true;
+				PShape->setFlag(PxShapeFlag::eSIMULATION_SHAPE, true);
+			}
+			// Disable sim collision
+			else if(!bSimCollision && bCurrentSimCollision)
+			{
+				bUpdateMassProperties = true;
+				PShape->setFlag(PxShapeFlag::eSIMULATION_SHAPE, false);
+			}
+
+			// enable swept bounds for CCD for this shape
+			PxRigidBody* PBody = GetPxRigidActor()->is<PxRigidBody>();
+			if(bSimCollision && !bPhysicsStatic && bUseCCD && PBody)
+			{
+				PBody->setRigidBodyFlag(PxRigidBodyFlag::eENABLE_CCD, true);
+			}
+			else if(PBody)
+			{
+
+				PBody->setRigidBodyFlag(PxRigidBodyFlag::eENABLE_CCD, false);
+			}
+		}
+	}
+	// No collision enabled
+	else
+	{
+		PShape->setFlag(PxShapeFlag::eSIMULATION_SHAPE, false);
+		PShape->setFlag(PxShapeFlag::eSCENE_QUERY_SHAPE, false);
+	}
+
+	if(bUpdateMassProperties)
+	{
+		UpdateMassProperties();
+	}
+}
+
+void FBodyInstance::GetShapeFlags(FShapeData& ShapeData, TEnumAsByte<ECollisionEnabled::Type> UseCollisionEnabled, const bool bUseComplexAsSimple /*= false*/)
+{
+	// If query collision is enabled..
+	ShapeData.SyncShapeFlags = PxShapeFlags(0);
+	ShapeData.AsyncShapeFlags = PxShapeFlags(0);
+	ShapeData.SimpleShapeFlags = PxShapeFlags(0);
+	ShapeData.ComplexShapeFlags = PxShapeFlags(0);
+	ShapeData.SyncBodyFlags = PxRigidBodyFlags(0);
+
+	// Default flags
+	ShapeData.SyncShapeFlags |= PxShapeFlag::eSCENE_QUERY_SHAPE;
+	ShapeData.AsyncShapeFlags |= PxShapeFlag::eSCENE_QUERY_SHAPE;
+	ShapeData.SimpleShapeFlags |= PxShapeFlag::eVISUALIZATION;
+	ShapeData.ComplexShapeFlags |= PxShapeFlag::eVISUALIZATION;
+
+	AActor* Owner = OwnerComponent.IsValid() ? OwnerComponent.Get()->GetOwner() : NULL;
+
+	if(ShapeData.CollisionEnabled != ECollisionEnabled::NoCollision)
+	{
+		const bool bPhysicsStatic = !OwnerComponent.IsValid() || OwnerComponent.Get()->IsWorldGeometry();
+
+		// Only perform scene queries in the synchronous scene for static shapes
+		if(bPhysicsStatic)
+		{
+			ShapeData.AsyncShapeFlags.clear(PxShapeFlag::eSCENE_QUERY_SHAPE);
+		}
+
+		// See if we want physics collision
+		bool bSimCollision = (UseCollisionEnabled == ECollisionEnabled::QueryAndPhysics);
+
+		// on dynamic objects and objects which don't use complex as simple, tri mesh not used for sim
+		if(bSimCollision && bUseComplexAsSimple)
+		{
+			ShapeData.ComplexShapeFlags |= PxShapeFlag::eSIMULATION_SHAPE;
+		}
+		
+		if(OwnerComponent == NULL || !OwnerComponent->IsA(UModelComponent::StaticClass()))
+		{
+			ShapeData.ComplexShapeFlags.clear(PxShapeFlag::eVISUALIZATION);
+		}
+
+		// Enable sim collision
+		if(bSimCollision)
+		{
+			ShapeData.SimpleShapeFlags |= PxShapeFlag::eSIMULATION_SHAPE;
+		}
+
+		// enable swept bounds for CCD for this shape
+		PxRigidBody* PBody = GetPxRigidActor()->is<PxRigidBody>();
+		if(bSimCollision && !bPhysicsStatic && bUseCCD && PBody)
+		{
+			ShapeData.SyncBodyFlags |= PxRigidBodyFlag::eENABLE_CCD;
+		}
+	}
+	// No collision enabled
+	else
+	{
+		ShapeData.SyncShapeFlags.clear(PxShapeFlag::eSCENE_QUERY_SHAPE);
+		ShapeData.SyncShapeFlags.clear(PxShapeFlag::eSIMULATION_SHAPE);
+		ShapeData.AsyncShapeFlags.clear(PxShapeFlag::eSCENE_QUERY_SHAPE);
+		ShapeData.AsyncShapeFlags.clear(PxShapeFlag::eSIMULATION_SHAPE);
+	}
+}
+
 #endif
 
 ////////////////////////////////////////////////////////////////////////////
diff --git a/Engine/Source/Runtime/Engine/Private/PhysicsEngine/BodySetup.cpp b/Engine/Source/Runtime/Engine/Private/PhysicsEngine/BodySetup.cpp
--- a/Engine/Source/Runtime/Engine/Private/PhysicsEngine/BodySetup.cpp
+++ b/Engine/Source/Runtime/Engine/Private/PhysicsEngine/BodySetup.cpp
@@ -95,9 +95,13 @@ void UBodySetup::AddCollisionFrom(class UBodySetup* FromSetup)
 	}
 }
 
+DECLARE_CYCLE_STAT(TEXT("Create Physics Meshes"), STAT_CreatePhysicsMeshes, STATGROUP_Physics);
+
 
 void UBodySetup::CreatePhysicsMeshes()
 {
+	SCOPE_CYCLE_COUNTER(STAT_CreatePhysicsMeshes);
+
 #if WITH_PHYSX
 	// Create meshes from cooked data if not already done
 	if(bCreatedPhysicsMeshes)
@@ -503,7 +507,7 @@ void UBodySetup::AddConvexElemsToRigidActor(PxRigidActor* PDestActor, const FTra
 	}
 }
 
-void UBodySetup::AddTriMeshToRigidActor(PxRigidActor* PDestActor, const FVector& Scale3D, const FVector& Scale3DAbs) const
+void UBodySetup::AddTriMeshToRigidActor(PxRigidActor* PDestActor, const FVector& Scale3D, const FVector& Scale3DAbs, TArray<PxShape*>* NewShapes) const
 {
 	float ContactOffsetFactor, MaxContactOffset;
 	GetContactOffsetParams(ContactOffsetFactor, MaxContactOffset);
@@ -544,6 +548,11 @@ void UBodySetup::AddTriMeshToRigidActor(PxRigidActor* PDestActor, const FVector&
 				{
 					NewShape->setLocalPose(PLocalPose);
 					NewShape->setContactOffset(MaxContactOffset);
+
+					if(NewShapes)
+					{
+						NewShapes->Add(NewShape);
+					}
 				}
 				else
 				{
@@ -558,7 +567,7 @@ void UBodySetup::AddTriMeshToRigidActor(PxRigidActor* PDestActor, const FVector&
 	}
 }
 
-void UBodySetup::AddShapesToRigidActor(PxRigidActor* PDestActor, FVector& Scale3D, const FTransform& RelativeTM /* = FTransform::Identity */, TArray<physx::PxShape*>* NewShapes /* = NULL */ )
+void UBodySetup::AddShapesToRigidActor(FBodyInstance* OwningInstance, physx::PxRigidActor* PDestActor, EPhysicsSceneType SceneType, FVector& Scale3D, physx::PxMaterial* SimpleMaterial, TArray<UPhysicalMaterial*>& ComplexMaterials, FShapeData& ShapeData, const FTransform& RelativeTM, TArray<physx::PxShape*>* NewShapes)
 {
 #if WITH_RUNTIME_PHYSICS_COOKING || WITH_EDITOR
 	// in editor, there are a lot of things relying on body setup to create physics meshes
@@ -568,6 +577,7 @@ void UBodySetup::AddShapesToRigidActor(PxRigidActor* PDestActor, FVector& Scale3
 	float MinScale;
 	float MinScaleAbs;
 	FVector Scale3DAbs;
+	TArray<PxShape*> Shapes;
 	SetupNonUniformHelper(Scale3D, MinScale, MinScaleAbs, Scale3DAbs);
 
 	{
@@ -588,18 +598,85 @@ void UBodySetup::AddShapesToRigidActor(PxRigidActor* PDestActor, FVector& Scale3
 	// for simple queries as well
 	if (CollisionTraceFlag != ECollisionTraceFlag::CTF_UseComplexAsSimple)
 	{
-		AddSpheresToRigidActor(PDestActor, RelativeTM, MinScale, MinScaleAbs, NewShapes);
-		AddBoxesToRigidActor(PDestActor, RelativeTM, Scale3D, Scale3DAbs, NewShapes);
-		AddSphylsToRigidActor(PDestActor, RelativeTM, Scale3D, Scale3DAbs, NewShapes);
-		AddConvexElemsToRigidActor(PDestActor, RelativeTM, Scale3D, Scale3DAbs, NewShapes);
+		AddSpheresToRigidActor(PDestActor, RelativeTM, MinScale, MinScaleAbs, &Shapes);
+		AddBoxesToRigidActor(PDestActor, RelativeTM, Scale3D, Scale3DAbs, &Shapes);
+		AddSphylsToRigidActor(PDestActor, RelativeTM, Scale3D, Scale3DAbs, &Shapes);
+		AddConvexElemsToRigidActor(PDestActor, RelativeTM, Scale3D, Scale3DAbs, &Shapes);
 	}
 
-
 	// Create tri-mesh shape, when we are not using simple collision shapes for 
 	// complex queries as well
 	if( CollisionTraceFlag != ECollisionTraceFlag::CTF_UseSimpleAsComplex )
 	{
-		AddTriMeshToRigidActor(PDestActor, Scale3D, Scale3DAbs);
+		AddTriMeshToRigidActor(PDestActor, Scale3D, Scale3DAbs, &Shapes);
+	}
+
+	for(PxShape* Shape : Shapes)
+	{
+		FShapeFilterData& Filters = ShapeData.FilterData;
+		// Apply provided materials
+		// If a triangle mesh, need to get array of materials...
+		if(Shape->getGeometryType() == PxGeometryType::eTRIANGLEMESH)
+		{
+			TArray<PxMaterial*> PComplexMats;
+			PComplexMats.AddUninitialized(ComplexMaterials.Num());
+			for(int MatIdx = 0; MatIdx < ComplexMaterials.Num(); MatIdx++)
+			{
+				check(ComplexMaterials[MatIdx] != NULL);
+				PComplexMats[MatIdx] = ComplexMaterials[MatIdx]->GetPhysXMaterial();
+				check(PComplexMats[MatIdx] != NULL);
+			}
+
+			if(PComplexMats.Num())
+			{
+				Shape->setMaterials(PComplexMats.GetData(), PComplexMats.Num());
+			}
+			else
+			{
+				UE_LOG(LogPhysics, Warning, TEXT("FBodyInstance::UpdatePhysicalMaterials : PComplexMats is empty - falling back on simple physical material."));
+				Shape->setMaterials(&SimpleMaterial, 1);
+			}
+
+			Shape->setQueryFilterData(Filters.QueryComplexFilter);
+			
+			if(SceneType == PST_Sync)
+			{
+				Shape->setFlags(ShapeData.SyncShapeFlags | ShapeData.ComplexShapeFlags);
+			}
+			else
+			{
+				Shape->setFlags(ShapeData.AsyncShapeFlags | ShapeData.ComplexShapeFlags);
+			}
+		}
+		// Simple shape, 
+		else
+		{
+			Shape->setMaterials(&SimpleMaterial, 1);
+			Shape->setQueryFilterData(Filters.QuerySimpleFilter);
+
+			if(SceneType == PST_Sync)
+			{
+				Shape->setFlags(ShapeData.SyncShapeFlags | ShapeData.SimpleShapeFlags);
+			}
+			else
+			{
+				Shape->setFlags(ShapeData.AsyncShapeFlags | ShapeData.SimpleShapeFlags);
+			}
+		}
+		Shape->setSimulationFilterData(Filters.SimFilter);
+
+		if(OwningInstance)
+		{
+			if(PxRigidBody* RigidBody = OwningInstance->GetPxRigidActor()->is<PxRigidBody>())
+			{
+				RigidBody->setRigidBodyFlags(ShapeData.SyncBodyFlags);
+			}
+		}
+	}
+
+	if(NewShapes)
+	{
+		(*NewShapes) = MoveTemp(Shapes);
 	}
 }
 
diff --git a/Engine/Source/Runtime/Engine/Private/PhysicsEngine/PhysDerivedData.cpp b/Engine/Source/Runtime/Engine/Private/PhysicsEngine/PhysDerivedData.cpp
--- a/Engine/Source/Runtime/Engine/Private/PhysicsEngine/PhysDerivedData.cpp
+++ b/Engine/Source/Runtime/Engine/Private/PhysicsEngine/PhysDerivedData.cpp
@@ -43,8 +43,13 @@ void FDerivedDataPhysXCooker::InitCooker()
 #endif
 }
 
+DECLARE_CYCLE_STAT(TEXT("PhysX Cooking"), STAT_PhysXCooking, STATGROUP_Physics);
+
+
 bool FDerivedDataPhysXCooker::Build( TArray<uint8>& OutData )
 {
+	SCOPE_CYCLE_COUNTER(STAT_PhysXCooking);
+
 	check(Cooker != NULL);
 
 	FMemoryWriter Ar( OutData );
diff --git a/Engine/Source/Runtime/Engine/Private/PhysicsEngine/PhysScene.cpp b/Engine/Source/Runtime/Engine/Private/PhysicsEngine/PhysScene.cpp
--- a/Engine/Source/Runtime/Engine/Private/PhysicsEngine/PhysScene.cpp
+++ b/Engine/Source/Runtime/Engine/Private/PhysicsEngine/PhysScene.cpp
@@ -353,16 +353,16 @@ void FPhysScene::TermBody(FBodyInstance* BodyInstance)
 #endif
 	}
 
-	// Remove body from any pending deferred addition / removal
-	for(int32 SceneIdx = 0 ; SceneIdx < PST_MAX ; ++SceneIdx)
-	{
-		int32 FoundIdx = INDEX_NONE;
-		if(DeferredAddInstances[SceneIdx].Find(BodyInstance, FoundIdx))
+		// Remove body from any pending deferred addition / removal
+		for(int32 SceneIdx = 0 ; SceneIdx < PST_MAX ; ++SceneIdx)
 		{
-			DeferredAddActors->RemoveAt(FoundIdx);
-			DeferredAddInstances->RemoveAt(FoundIdx);
+			int32 FoundIdx = INDEX_NONE;
+		if(DeferredAddInstances[SceneIdx].Find(BodyInstance, FoundIdx))
+			{
+				DeferredAddActors->RemoveAt(FoundIdx);
+				DeferredAddInstances->RemoveAt(FoundIdx);
+			}
 		}
-	}
 
 #if WITH_PHYSX
 	RemoveActiveBody(BodyInstance, PST_Sync);
@@ -661,6 +661,7 @@ void FPhysScene::WaitPhysScenes()
 	}
 	if (ThingsToComplete.Num())
 	{
+		QUICK_SCOPE_CYCLE_COUNTER(STAT_FPhysScene_WaitPhysScenes);
 		FTaskGraphInterface::Get().WaitUntilTasksComplete(ThingsToComplete, ENamedThreads::GameThread);
 	}
 }
@@ -675,6 +676,7 @@ void FPhysScene::WaitClothScene()
 
 	if (ThingsToComplete.Num())
 	{
+		QUICK_SCOPE_CYCLE_COUNTER(STAT_FPhysScene_WaitClothScene);
 		FTaskGraphInterface::Get().WaitUntilTasksComplete(ThingsToComplete, ENamedThreads::GameThread);
 	}
 }
@@ -1408,9 +1410,10 @@ void FPhysScene::FlushDeferredActors()
 			SCOPED_SCENE_WRITE_LOCK(Scene);
 
 			Scene->addActors(DeferredAddActors[SceneIdx].GetData(), DeferredAddActors[SceneIdx].Num());
-
+			int32 Idx = -1;
 			for(FBodyInstance* Instance : DeferredAddInstances[SceneIdx])
 			{
+				++Idx;
 				Instance->CurrentSceneState = BodyInstanceSceneState::Added;
 
 				if(Instance->IsDynamic())
@@ -1465,6 +1468,19 @@ void FPhysScene::DeferAddActor(FBodyInstance* OwningInstance, PxActor* Actor, EP
 	}
 }
 
+void FPhysScene::DeferAddActors(TArray<FBodyInstance*>& OwningInstances, TArray<PxActor*>& Actors, EPhysicsSceneType SceneType)
+{
+	int32 Num = OwningInstances.Num();
+
+	DeferredAddInstances[SceneType].Reserve(DeferredAddInstances->Num() + Num);
+	DeferredAddActors[SceneType].Reserve(DeferredAddActors->Num() + Num);
+
+	for(int32 Idx = 0 ; Idx < Num ; ++Idx)
+	{
+		DeferAddActor(OwningInstances[Idx], Actors[Idx], SceneType);
+	}
+}
+
 void FPhysScene::DeferRemoveActor(FBodyInstance* OwningInstance, PxActor* Actor, EPhysicsSceneType SceneType)
 {
 	check(OwningInstance && Actor && SceneType < PST_MAX);
diff --git a/Engine/Source/Runtime/Engine/Private/PhysicsEngine/PhysXSupport.h b/Engine/Source/Runtime/Engine/Private/PhysicsEngine/PhysXSupport.h
--- a/Engine/Source/Runtime/Engine/Private/PhysicsEngine/PhysXSupport.h
+++ b/Engine/Source/Runtime/Engine/Private/PhysicsEngine/PhysXSupport.h
@@ -21,6 +21,8 @@
 #define PHYSX_MEMORY_STATS		0
 #endif
 
+#define PHYSX_MEMORY_STAT_ONLY (0)
+
 #if USE_SCENE_LOCK
 
 /** Scoped scene read lock - we use this instead of PxSceneReadLock because it handles NULL scene */
@@ -352,7 +354,11 @@ public:
 
 		return (uint8*)AllocationHeader + 16;
 #else
-		return FMemory::Malloc(size, 16);
+		void* ptr = FMemory::Malloc(size, 16);
+		#if PHYSX_MEMORY_STAT_ONLY
+			INC_DWORD_STAT_BY(STAT_MemoryPhysXTotalAllocationSize, FMemory::GetAllocSize(ptr));
+		#endif
+		return ptr;
 #endif
 	}
 	 
@@ -368,6 +374,9 @@ public:
 			FMemory::Free(AllocationHeader);
 		}
 #else
+		#if PHYSX_MEMORY_STAT_ONLY
+			DEC_DWORD_STAT_BY(STAT_MemoryPhysXTotalAllocationSize, FMemory::GetAllocSize(ptr));
+		#endif
 		FMemory::Free(ptr);
 #endif
 	}
@@ -629,3 +638,35 @@ ENGINE_API SIZE_T GetPhysxObjectSize(PxBase* Obj, const PxCollection* SharedColl
 
 
 #include "../Collision/PhysicsFiltering.h"
+
+/** Helper struct holding physics body filter data during initialisation */
+struct FShapeFilterData
+{
+	PxFilterData SimFilter;
+	PxFilterData QuerySimpleFilter;
+	PxFilterData QueryComplexFilter;
+};
+
+/** Helper object to hold initialisation data for shapes */
+struct FShapeData
+{
+	FShapeData()
+		: SyncShapeFlags(0)
+		, AsyncShapeFlags(0)
+		, SimpleShapeFlags(0)
+		, ComplexShapeFlags(0)
+		, SyncBodyFlags(0)
+		, AsyncBodyFlags(0)
+	{
+
+	}
+
+	TEnumAsByte<ECollisionEnabled::Type> CollisionEnabled;
+	FShapeFilterData FilterData;
+	PxShapeFlags SyncShapeFlags;
+	PxShapeFlags AsyncShapeFlags;
+	PxShapeFlags SimpleShapeFlags;
+	PxShapeFlags ComplexShapeFlags;
+	PxRigidBodyFlags SyncBodyFlags;
+	PxRigidBodyFlags AsyncBodyFlags;
+};
diff --git a/Engine/Source/Runtime/Engine/Private/PlayerCameraManager.cpp b/Engine/Source/Runtime/Engine/Private/PlayerCameraManager.cpp
--- a/Engine/Source/Runtime/Engine/Private/PlayerCameraManager.cpp
+++ b/Engine/Source/Runtime/Engine/Private/PlayerCameraManager.cpp
@@ -481,7 +481,9 @@ void APlayerCameraManager::UpdateViewTarget(FTViewTarget& OutVT, float DeltaTime
 	OutVT.POV.FOV = DefaultFOV;
 	OutVT.POV.OrthoWidth = DefaultOrthoWidth;
 	OutVT.POV.bConstrainAspectRatio = false;
+	OutVT.POV.bUseFieldOfViewForLOD = true;
 	OutVT.POV.ProjectionMode = bIsOrthographic ? ECameraProjectionMode::Orthographic : ECameraProjectionMode::Perspective;
+	OutVT.POV.PostProcessSettings.SetBaseValues();
 	OutVT.POV.PostProcessBlendWeight = 1.0f;
 
 
@@ -593,12 +595,19 @@ void APlayerCameraManager::UpdateCameraLensEffects(const FTViewTarget& OutVT)
 	}
 }
 
-
 void APlayerCameraManager::ApplyAudioFade()
 {
 	if (GEngine && GEngine->GetAudioDevice() )
 	{
-		GEngine->GetAudioDevice()->TransientMasterVolume = 1.0 - FadeAmount;
+		GEngine->GetAudioDevice()->TransientMasterVolume = 1.f - FadeAmount;
+	}
+}
+
+void APlayerCameraManager::StopAudioFade()
+{
+	if (GEngine && GEngine->GetAudioDevice())
+	{
+		GEngine->GetAudioDevice()->TransientMasterVolume = 1.f;
 	}
 }
 
@@ -850,21 +859,26 @@ void APlayerCameraManager::DoUpdateCamera(float DeltaTime)
 	// Cache results
 	FillCameraCache(NewPOV);
 
-	if (bEnableFading && FadeTimeRemaining > 0.0f)
+	if (bEnableFading)
 	{
-		FadeTimeRemaining = FMath::Max(FadeTimeRemaining - DeltaTime, 0.0f);
-		if (FadeTime > 0.0f)
+		if (bAutoAnimateFade)
 		{
-			FadeAmount = FadeAlpha.X + ((1.f - FadeTimeRemaining/FadeTime) * (FadeAlpha.Y - FadeAlpha.X));
+			FadeTimeRemaining = FMath::Max(FadeTimeRemaining - DeltaTime, 0.0f);
+			if (FadeTime > 0.0f)
+			{
+				FadeAmount = FadeAlpha.X + ((1.f - FadeTimeRemaining / FadeTime) * (FadeAlpha.Y - FadeAlpha.X));
+			}
+
+			if ((bHoldFadeWhenFinished == false) && (FadeTimeRemaining <= 0.f))
+			{
+				// done
+				StopCameraFade();
+			}
 		}
 
 		if (bFadeAudio)
 		{
 			ApplyAudioFade();
-			if (FadeAmount == 0)
-			{
-				bFadeAudio = false;
-			}
 		}
 	}
 }
@@ -1058,8 +1072,6 @@ void APlayerCameraManager::ClearCameraLensEffects()
  *  Camera Shakes
  *  ------------------------------------------------------------ */
 
-
-
 void APlayerCameraManager::PlayCameraShake(TSubclassOf<UCameraShake> Shake, float Scale, ECameraAnimPlaySpace::Type PlaySpace, FRotator UserPlaySpaceRot)
 {
 	if ((Shake != nullptr) && (Scale > 0.0f))
@@ -1127,6 +1139,48 @@ void APlayerCameraManager::ClearAllCameraShakes()
 //	StopAllCameraAnims(true);
 }
 
+/** ------------------------------------------------------------
+ *  Camera fades
+ *  ------------------------------------------------------------ */
+
+void APlayerCameraManager::StartCameraFade(float FromAlpha, float ToAlpha, float InFadeTime, FLinearColor InFadeColor, bool bInFadeAudio, bool bInHoldWhenFinished)
+{
+	bEnableFading = true;
+
+	FadeColor = InFadeColor;
+	FadeAlpha = FVector2D(FromAlpha, ToAlpha);
+	FadeTime = InFadeTime;
+	FadeTimeRemaining = InFadeTime;
+	bFadeAudio = bInFadeAudio;
+
+	bAutoAnimateFade = true;
+	bHoldFadeWhenFinished = bInHoldWhenFinished;
+}
+
+void APlayerCameraManager::StopCameraFade()
+{
+	if (bEnableFading == true)
+	{
+		// Make sure FadeAmount finishes at the desired value
+		FadeAmount = FadeAlpha.Y;
+		bEnableFading = false;
+		StopAudioFade();
+	}
+}
+
+void APlayerCameraManager::SetManualCameraFade(float InFadeAmount, FLinearColor Color, bool bInFadeAudio)
+{
+	bEnableFading = true;
+	FadeColor = Color;
+	FadeAmount = InFadeAmount;
+	bFadeAudio = bInFadeAudio;
+
+	bAutoAnimateFade = false;
+	StopAudioFade();
+	FadeTimeRemaining = 0.0f;
+}
+
+
 //////////////////////////////////////////////////////////////////////////
 // FTViewTarget
 
diff --git a/Engine/Source/Runtime/Engine/Private/PlayerController.cpp b/Engine/Source/Runtime/Engine/Private/PlayerController.cpp
--- a/Engine/Source/Runtime/Engine/Private/PlayerController.cpp
+++ b/Engine/Source/Runtime/Engine/Private/PlayerController.cpp
@@ -1460,26 +1460,19 @@ void APlayerController::ResetCameraMode()
 
 void APlayerController::ClientSetCameraFade_Implementation(bool bEnableFading, FColor FadeColor, FVector2D FadeAlpha, float FadeTime, bool bFadeAudio)
 {
-	if (PlayerCameraManager != NULL)
+	if (PlayerCameraManager != nullptr)
 	{
-		PlayerCameraManager->bEnableFading = bEnableFading;
-		if (PlayerCameraManager->bEnableFading)
+		if (bEnableFading)
 		{
-			PlayerCameraManager->FadeColor = FadeColor;
-			PlayerCameraManager->FadeAlpha = FadeAlpha;
-			PlayerCameraManager->FadeTime = FadeTime;
-			PlayerCameraManager->FadeTimeRemaining = FadeTime;
-			PlayerCameraManager->bFadeAudio = bFadeAudio;
+			PlayerCameraManager->StartCameraFade(FadeAlpha.X, FadeAlpha.Y, FadeTime, FadeColor.ReinterpretAsLinear(), bFadeAudio);
 		}
 		else
 		{
-			// Make sure FadeAmount finishes at the correct value
-			PlayerCameraManager->FadeAmount = PlayerCameraManager->FadeAlpha.Y;
+			PlayerCameraManager->StopCameraFade();
 		}
 	}
 }
 
-
 void APlayerController::SendClientAdjustment()
 {
 	if (AcknowledgedPawn != GetPawn() && !GetSpectatorPawn())
diff --git a/Engine/Source/Runtime/Engine/Private/PrimitiveComponent.cpp b/Engine/Source/Runtime/Engine/Private/PrimitiveComponent.cpp
--- a/Engine/Source/Runtime/Engine/Private/PrimitiveComponent.cpp
+++ b/Engine/Source/Runtime/Engine/Private/PrimitiveComponent.cpp
@@ -306,8 +306,15 @@ void UPrimitiveComponent::OnRegister()
 
 	if (bCanEverAffectNavigation)
 	{
-		bNavigationRelevant = IsNavigationRelevant();
-		UNavigationSystem::OnComponentRegistered(this);	
+		const bool bNavRelevant = bNavigationRelevant = IsNavigationRelevant();
+		if (bNavRelevant)
+		{
+			UNavigationSystem::OnComponentRegistered(this);
+		}
+	}
+	else
+	{
+		bNavigationRelevant = false;
 	}
 }
 
diff --git a/Engine/Source/Runtime/Engine/Private/SceneComponent.cpp b/Engine/Source/Runtime/Engine/Private/SceneComponent.cpp
--- a/Engine/Source/Runtime/Engine/Private/SceneComponent.cpp
+++ b/Engine/Source/Runtime/Engine/Private/SceneComponent.cpp
@@ -1088,7 +1088,7 @@ void USceneComponent::SnapTo(class USceneComponent* Parent, FName InSocketName)
 	AttachTo(Parent, InSocketName, EAttachLocation::SnapToTarget);
 }
 
-void USceneComponent::DetachFromParent(bool bMaintainWorldPosition)
+void USceneComponent::DetachFromParent(bool bMaintainWorldPosition, bool bCallModify)
 {
 	if(AttachParent != NULL)
 	{
@@ -1100,8 +1100,11 @@ void USceneComponent::DetachFromParent(bool bMaintainWorldPosition)
 		// Make sure parent points to us if we're registered
 		checkf(!bRegistered || AttachParent->AttachChildren.Contains(this), TEXT("Attempt to detach SceneComponent '%s' owned by '%s' from AttachParent '%s' while not attached."), *GetName(), (GetOwner() ? *GetOwner()->GetName() : TEXT("Unowned")), *AttachParent->GetName());
 
-		Modify();
-		AttachParent->Modify();
+		if (bCallModify)
+		{
+			Modify();
+			AttachParent->Modify();
+		}
 
 		PrimaryComponentTick.RemovePrerequisite(AttachParent, AttachParent->PrimaryComponentTick); // no longer required to tick after the attachment
 
diff --git a/Engine/Source/Runtime/Engine/Private/SceneView.cpp b/Engine/Source/Runtime/Engine/Private/SceneView.cpp
--- a/Engine/Source/Runtime/Engine/Private/SceneView.cpp
+++ b/Engine/Source/Runtime/Engine/Private/SceneView.cpp
@@ -377,6 +377,7 @@ FSceneView::FSceneView(const FSceneViewInitOptions& InitOptions)
 	}
 	
 	// Compute a transform from view origin centered world-space to clip space.
+	ViewMatrices.TranslatedViewMatrix = TranslatedViewMatrix;
 	ViewMatrices.TranslatedViewProjectionMatrix = TranslatedViewMatrix * ViewMatrices.ProjMatrix;
 	ViewMatrices.InvTranslatedViewProjectionMatrix = ViewMatrices.TranslatedViewProjectionMatrix.Inverse();
 
@@ -463,6 +464,8 @@ FSceneView::FSceneView(const FSceneViewInitOptions& InitOptions)
 		bIsGameView = (Family && Family->Scene && Family->Scene->GetWorld() ) ? Family->Scene->GetWorld()->IsGameWorld() : false;
 	}
 
+	bUseFieldOfViewForLOD = InitOptions.bUseFieldOfViewForLOD;
+
 #if WITH_EDITOR
 	EditorViewBitflag = InitOptions.EditorViewBitflag;
 
@@ -470,10 +473,19 @@ FSceneView::FSceneView(const FSceneViewInitOptions& InitOptions)
 #endif
 }
 
+static TAutoConsoleVariable<int32> CVarCompensateForFOV(
+	TEXT("lod.CompensateForFOV"),
+	1,
+	TEXT("When not 0 account for FOV in LOD calculations."));
+
 float FSceneView::GetLODDistanceFactor() const
 {
-	const float ScreenMultiple = FMath::Max(ViewRect.Width() / 2.0f * ViewMatrices.ProjMatrix.M[0][0],
-		ViewRect.Height() / 2.0f * ViewMatrices.ProjMatrix.M[1][1]);
+	bool bCompensateForFOV = bUseFieldOfViewForLOD && CVarCompensateForFOV.GetValueOnAnyThread() != 0;
+	float ScreenScaleX = bCompensateForFOV ? ViewMatrices.ProjMatrix.M[0][0] : 1.0f;
+	float ScreenScaleY = bCompensateForFOV ? ViewMatrices.ProjMatrix.M[1][1] : (ViewRect.Width() / ViewRect.Height());
+
+	const float ScreenMultiple = FMath::Max(ViewRect.Width() / 2.0f * ScreenScaleX,
+		ViewRect.Height() / 2.0f * ScreenScaleY);
 	float Fac = PI * ScreenMultiple * ScreenMultiple / ViewRect.Area();
 	return Fac;
 }
diff --git a/Engine/Source/Runtime/Engine/Private/ShapeComponent.cpp b/Engine/Source/Runtime/Engine/Private/ShapeComponent.cpp
--- a/Engine/Source/Runtime/Engine/Private/ShapeComponent.cpp
+++ b/Engine/Source/Runtime/Engine/Private/ShapeComponent.cpp
@@ -66,7 +66,7 @@ void UShapeComponent::PostEditChangeProperty(FPropertyChangedEvent& PropertyChan
 }
 #endif // WITH_EDITOR
 
-bool UShapeComponent::DoCustomNavigableGeometryExport(struct FNavigableGeometryExport* GeomExport) const
+bool UShapeComponent::DoCustomNavigableGeometryExport(FNavigableGeometryExport& GeomExport) const
 {
 	return !bDynamicObstacle;
 }
diff --git a/Engine/Source/Runtime/Engine/Private/SkeletalMesh.cpp b/Engine/Source/Runtime/Engine/Private/SkeletalMesh.cpp
--- a/Engine/Source/Runtime/Engine/Private/SkeletalMesh.cpp
+++ b/Engine/Source/Runtime/Engine/Private/SkeletalMesh.cpp
@@ -3068,6 +3068,41 @@ FArchive& operator<<( FArchive& Ar, FSkeletalMaterial& Elem )
 	return Ar;
 }
 
+TArray<USkeletalMeshSocket*> USkeletalMesh::GetActiveSocketList() const
+{
+	TArray<USkeletalMeshSocket*> ActiveSockets = Sockets;
+
+	// Then the skeleton sockets that aren't in the mesh
+	if (Skeleton)
+	{
+		for (auto SkeletonSocketIt = Skeleton->Sockets.CreateConstIterator(); SkeletonSocketIt; ++SkeletonSocketIt)
+		{
+			USkeletalMeshSocket* Socket = *(SkeletonSocketIt);
+
+			if (!IsSocketOnMesh(Socket->SocketName))
+			{
+				ActiveSockets.Add(Socket);
+			}
+		}
+	}
+	return ActiveSockets;
+}
+
+bool USkeletalMesh::IsSocketOnMesh(const FName& InSocketName) const
+{
+	for(int32 SocketIdx=0; SocketIdx < Sockets.Num(); SocketIdx++)
+	{
+		USkeletalMeshSocket* Socket = Sockets[SocketIdx];
+
+		if(Socket != NULL && Socket->SocketName == InSocketName)
+		{
+			return true;
+		}
+	}
+
+	return false;
+}
+
 #if WITH_EDITOR
 
 
@@ -3094,45 +3129,6 @@ FStaticLODModel& USkeletalMesh::PreModifyMesh()
 	return GetSourceModel();
 }
 
-const TArray<USkeletalMeshSocket*>& USkeletalMesh::GetActiveSocketList() const
-{
-	// Static so we can return by reference - really don't want to go down the road of passing this array around by value!
-	static TArray<USkeletalMeshSocket*> ActiveSocketList;
-	
-	ActiveSocketList = Sockets;
-
-	// Then the skeleton sockets that aren't in the mesh
-	if ( Skeleton )
-	{
-		for ( auto SkeletonSocketIt = Skeleton->Sockets.CreateConstIterator(); SkeletonSocketIt; ++SkeletonSocketIt )
-		{
-			USkeletalMeshSocket* Socket = *(SkeletonSocketIt);
-
-			if ( !IsSocketOnMesh( Socket->SocketName ) )
-			{
-				ActiveSocketList.Add( Socket );
-			}
-		}
-	}
-
-	return ActiveSocketList;
-}
-
-bool USkeletalMesh::IsSocketOnMesh( const FName& InSocketName ) const
-{
-	for( int32 SocketIdx=0; SocketIdx < Sockets.Num(); SocketIdx++ )
-	{
-		USkeletalMeshSocket* Socket = Sockets[SocketIdx];
-
-		if ( Socket != NULL && Socket->SocketName == InSocketName )
-		{
-			return true;
-		}
-	}
-
-	return false;
-}
-
 int32 USkeletalMesh::ValidatePreviewAttachedObjects()
 {
 	int32 NumBrokenAssets = PreviewAttachedAssetContainer.ValidatePreviewAttachedObjects();
@@ -3507,6 +3503,7 @@ USkeletalMeshSocket
 -----------------------------------------------------------------------------*/
 USkeletalMeshSocket::USkeletalMeshSocket(const FObjectInitializer& ObjectInitializer)
 	: Super(ObjectInitializer)
+	, bForceAlwaysAnimated(true)
 {
 	RelativeScale = FVector(1.0f, 1.0f, 1.0f);
 }
@@ -3670,10 +3667,10 @@ void ASkeletalMeshActor::PreviewBeginAnimControl(UInterpGroup* InInterpGroup)
 		UAnimInstance* AnimInst = SkeletalMeshComponent->GetAnimInstance();
 		if (!AnimInst)
 		{
-			SkeletalMeshComponent->SetAnimationMode(EAnimationMode::Type::AnimationSingleNode);
-		}
+		SkeletalMeshComponent->SetAnimationMode(EAnimationMode::Type::AnimationSingleNode);
 	}
 }
+}
 
 
 void ASkeletalMeshActor::PreviewFinishAnimControl(UInterpGroup* InInterpGroup)
@@ -3690,13 +3687,11 @@ void ASkeletalMeshActor::PreviewFinishAnimControl(UInterpGroup* InInterpGroup)
 				AnimInst->Montage_Stop(0.f);
 				AnimInst->UpdateAnimation(0.f);
 			}
-
 		}
-		
 		// Update space bases to reset it back to ref pose
 		SkeletalMeshComponent->RefreshBoneTransforms();
 		SkeletalMeshComponent->RefreshSlaveComponents();
-		SkeletalMeshComponent->UpdateComponentToWorld();		
+		SkeletalMeshComponent->UpdateComponentToWorld();
 	}
 }
 
@@ -3742,9 +3737,14 @@ void ASkeletalMeshActor::PreviewSetAnimPosition(FName SlotName, int32 ChannelInd
 			if (AnimMontageInst)
 			{
 				AnimMontageInst->Weight = 1.f;
-				//AnimInst->Montage_Set(bLooping);
+
+				float OldMontagePosition = AnimInst->Montage_GetPosition(CurrentlyPlayingMontage.Get());
 				AnimInst->Montage_SetPosition(CurrentlyPlayingMontage.Get(), InPosition);
 				AnimInst->UpdateAnimation(DeltaTime);
+
+				// since we don't advance montage in the tick, we manually have to handle notifies
+				AnimMontageInst->HandleEvents(OldMontagePosition, InPosition, NULL);
+				AnimInst->TriggerAnimNotifies(DeltaTime);
 			}
 			else
 			{
@@ -3888,10 +3888,10 @@ void ASkeletalMeshActor::BeginAnimControl(UInterpGroup* InInterpGroup)
 		UAnimInstance* AnimInst = SkeletalMeshComponent->GetAnimInstance();
 		if (!AnimInst)
 		{
-			SkeletalMeshComponent->SetAnimationMode(EAnimationMode::Type::AnimationSingleNode);
-		}
+		SkeletalMeshComponent->SetAnimationMode(EAnimationMode::Type::AnimationSingleNode);
 	}
 }
+}
 
 bool ASkeletalMeshActor::CanPlayAnimation(class UAnimSequenceBase* AnimAssetBase/*=NULL*/) const
 {
@@ -3948,7 +3948,13 @@ void SetAnimPositionInner(FName SlotName, USkeletalMeshComponent* SkeletalMeshCo
 			if(AnimMontageInst)
 			{
 				AnimMontageInst->Weight = 1.f;
+
+				float OldMontagePosition = AnimInst->Montage_GetPosition(CurrentlyPlayingMontage.Get());
 				AnimInst->Montage_SetPosition(CurrentlyPlayingMontage.Get(), InPosition);
+
+				// since we don't advance montage in the tick, we manually have to handle notifies
+				AnimMontageInst->HandleEvents(OldMontagePosition, InPosition, NULL);
+				AnimInst->TriggerAnimNotifies(0.f);
 			}
 			else
 			{
@@ -3957,7 +3963,6 @@ void SetAnimPositionInner(FName SlotName, USkeletalMeshComponent* SkeletalMeshCo
 		}
 	}
 }
-
 void ASkeletalMeshActor::SetAnimPosition(FName SlotName, int32 ChannelIndex, UAnimSequence* InAnimSequence, float InPosition, bool bFireNotifies, bool bLooping)
 {
 	if (CanPlayAnimation(InAnimSequence))
@@ -3972,7 +3977,7 @@ void ASkeletalMeshActor::FinishAnimControl(UInterpGroup* InInterpGroup)
 	{
 		UAnimInstance* AnimInst = SkeletalMeshComponent->GetAnimInstance();
 		if(AnimInst)
-		{
+	{
 			AnimInst->Montage_Stop(0.f);
 			AnimInst->UpdateAnimation(0.f);
 		}
@@ -4676,6 +4681,8 @@ USkinnedMeshComponent::USkinnedMeshComponent(const FObjectInitializer& ObjectIni
 	CurrentEditableSpaceBases = 0;
 	CurrentReadSpaceBases = 1;
 	bNeedToFlipSpaceBaseBuffers = false;
+
+	bCanEverAffectNavigation = false;
 }
 
 
@@ -4798,3 +4805,16 @@ FPrimitiveSceneProxy* USkinnedMeshComponent::CreateSceneProxy()
 
 /** Returns SkeletalMeshComponent subobject **/
 USkeletalMeshComponent* ASkeletalMeshActor::GetSkeletalMeshComponent() { return SkeletalMeshComponent; }
+
+
+#if WITH_EDITOR
+void USkeletalMeshSocket::PostEditChangeProperty(FPropertyChangedEvent& PropertyChangedEvent)
+{
+	Super::PostEditChangeProperty(PropertyChangedEvent);
+
+	if (PropertyChangedEvent.Property)
+	{
+		ChangedEvent.Broadcast(this, PropertyChangedEvent.MemberProperty);
+	}
+}
+#endif
\ No newline at end of file
diff --git a/Engine/Source/Runtime/Engine/Private/SkeletalMeshComponent.cpp b/Engine/Source/Runtime/Engine/Private/SkeletalMeshComponent.cpp
--- a/Engine/Source/Runtime/Engine/Private/SkeletalMeshComponent.cpp
+++ b/Engine/Source/Runtime/Engine/Private/SkeletalMeshComponent.cpp
@@ -18,6 +18,7 @@
 #include "Animation/VertexAnim/VertexAnimation.h"
 #include "Particles/ParticleSystemComponent.h"
 #include "Animation/AnimSingleNodeInstance.h"
+#include "Engine/SkeletalMeshSocket.h"
 #if WITH_EDITOR
 #include "ShowFlags.h"
 #include "Collision.h"
@@ -645,7 +646,7 @@ void USkeletalMeshComponent::FillSpaceBases(const USkeletalMesh* InSkeletalMesh,
 
 	// First bone is always root bone, and it doesn't have a parent.
 	{
-		check( RequiredBones[0] == 0 );
+		check(FillSpaceBasesRequiredBones[0] == 0);
 		DestSpaceBases[0] = SourceAtoms[0];
 
 #if (UE_BUILD_DEBUG || UE_BUILD_DEVELOPMENT)
@@ -654,9 +655,9 @@ void USkeletalMeshComponent::FillSpaceBases(const USkeletalMesh* InSkeletalMesh,
 #endif
 	}
 
-	for(int32 i=1; i<RequiredBones.Num(); i++)
+	for (int32 i = 1; i<FillSpaceBasesRequiredBones.Num(); i++)
 	{
-		const int32 BoneIndex = RequiredBones[i];
+		const int32 BoneIndex = FillSpaceBasesRequiredBones[i];
 		FTransform* SpaceBase = SpaceBasesData + BoneIndex;
 
 		FPlatformMisc::Prefetch(SpaceBase);
@@ -831,9 +832,50 @@ void USkeletalMeshComponent::RecalcRequiredBones(int32 LODIndex)
 		MergeInBoneIndexArrays(RequiredBones, MirroredDesiredBones);
 	}
 
+	TArray<FBoneIndexType> NeededBonesForFillSpaceBases;
+	{
+		TArray<FBoneIndexType> ForceAnimatedSocketBones;
+
+		for (const USkeletalMeshSocket* Socket : SkeletalMesh->GetActiveSocketList())
+		{
+			int32 BoneIndex = SkeletalMesh->RefSkeleton.FindBoneIndex(Socket->BoneName);
+			if (Socket->bForceAlwaysAnimated)
+			{
+				ForceAnimatedSocketBones.Add(BoneIndex);
+			}
+			else
+			{
+				NeededBonesForFillSpaceBases.Add(BoneIndex);
+			}
+		}
+
+		// Then sort array of required bones in hierarchy order
+		ForceAnimatedSocketBones.Sort();
+
+		// Make sure all of these are in RequiredBones.
+		MergeInBoneIndexArrays(RequiredBones, ForceAnimatedSocketBones);
+	}
+
+
 	// Ensure that we have a complete hierarchy down to those bones.
 	FAnimationRuntime::EnsureParentsPresent(RequiredBones, SkeletalMesh);
 
+	FillSpaceBasesRequiredBones.Empty(RequiredBones.Num() + NeededBonesForFillSpaceBases.Num());
+	FillSpaceBasesRequiredBones = RequiredBones;
+	
+	NeededBonesForFillSpaceBases.Sort();
+	MergeInBoneIndexArrays(FillSpaceBasesRequiredBones, NeededBonesForFillSpaceBases);
+	FAnimationRuntime::EnsureParentsPresent(FillSpaceBasesRequiredBones, SkeletalMesh);
+
+	// Sanitise bones that we aren't going to be updating
+	for (int32 BoneIndex = 0; BoneIndex < LocalAtoms.Num(); ++BoneIndex)
+	{
+		if (!RequiredBones.Contains(BoneIndex))
+		{
+			LocalAtoms[BoneIndex] = SkeletalMesh->RefSkeleton.GetRefBonePose()[BoneIndex];
+		}
+	}
+
 	// make sure animation requiredBone to mark as dirty
 	if (AnimScriptInstance)
 	{
@@ -1432,10 +1474,11 @@ void USkeletalMeshComponent::RenderAxisGizmo( const FTransform& Transform, UCanv
 	DrawDebugCanvasLine(Canvas, Origin, Origin + ZAxis * AxisLength, FLinearColor( 0.3f, 0.3f, 1.f));	
 }
 
-void USkeletalMeshComponent::SetMorphTarget(FName MorphTargetName, float Value)
+void USkeletalMeshComponent::SetMorphTarget(FName MorphTargetName, float Value, bool bRemoveZeroWeight)
 {
 	float *CurveValPtr = MorphTargetCurves.Find(MorphTargetName);
-	if ( FPlatformMath::Abs(Value) > ZERO_ANIMWEIGHT_THRESH )
+	bool bShouldAddToList = !bRemoveZeroWeight || FPlatformMath::Abs(Value) > ZERO_ANIMWEIGHT_THRESH;
+	if ( bShouldAddToList )
 	{
 		if ( CurveValPtr )
 		{
@@ -1961,7 +2004,15 @@ void USkeletalMeshComponent::ValidateAnimation()
 	{
 		if(AnimBlueprintGeneratedClass && SkeletalMesh && AnimBlueprintGeneratedClass->TargetSkeleton != SkeletalMesh->Skeleton)
 		{
-			UE_LOG(LogAnimation, Warning, TEXT("AnimBP %s is incompatible with skeleton %s, removing AnimBP from actor."), *AnimBlueprintGeneratedClass->GetName(), *SkeletalMesh->Skeleton->GetName());
+			if(SkeletalMesh->Skeleton == nullptr)
+			{
+				UE_LOG(LogAnimation, Warning, TEXT("AnimBP %s is incompatible because mesh %s has no skeleton, removing AnimBP from actor."), *AnimBlueprintGeneratedClass->GetName(), *SkeletalMesh->GetName());
+			}
+			else
+			{
+				UE_LOG(LogAnimation, Warning, TEXT("AnimBP %s is incompatible with skeleton %s, removing AnimBP from actor."), *AnimBlueprintGeneratedClass->GetName(), *SkeletalMesh->Skeleton->GetName());
+			}
+
 			AnimBlueprintGeneratedClass = NULL;
 		}
 	}
diff --git a/Engine/Source/Runtime/Engine/Private/SkeletalMeshComponentPhysics.cpp b/Engine/Source/Runtime/Engine/Private/SkeletalMeshComponentPhysics.cpp
--- a/Engine/Source/Runtime/Engine/Private/SkeletalMeshComponentPhysics.cpp
+++ b/Engine/Source/Runtime/Engine/Private/SkeletalMeshComponentPhysics.cpp
@@ -2306,9 +2306,13 @@ void USkeletalMeshComponent::ApplyWindForCloth(FClothingActor& ClothingActor)
 		{
 			FVector Position = ComponentToWorld.GetTranslation();
 
-			FVector4 WindParam = World->Scene->GetWindParameters(Position);
+			FVector WindDirection;
+			float WindSpeed;
+			float WindMinGust;
+			float WindMaxGust;
+			World->Scene->GetWindParameters(Position, WindDirection, WindSpeed, WindMinGust, WindMaxGust);
 
-			physx::PxVec3 WindVelocity(WindParam.X, WindParam.Y, WindParam.Z);
+			physx::PxVec3 WindVelocity(WindDirection.X, WindDirection.Y, WindDirection.Z);
 
 			WindVelocity *= WindUnitAmout;
 			float WindAdaption = rand()%20 * 0.1f; // make range from 0 to 2
diff --git a/Engine/Source/Runtime/Engine/Private/SkinnedMeshComponent.cpp b/Engine/Source/Runtime/Engine/Private/SkinnedMeshComponent.cpp
--- a/Engine/Source/Runtime/Engine/Private/SkinnedMeshComponent.cpp
+++ b/Engine/Source/Runtime/Engine/Private/SkinnedMeshComponent.cpp
@@ -32,6 +32,8 @@ namespace FAnimUpdateRateManager
 	// Global counter to spread SkinnedMeshComponent tick updates.
 	static uint8 UpdateRateGroupCount = 0;
 
+	static float TargetFrameTimeForUpdateRate = 1.f / 30.f; //Target frame rate for lookahead URO
+
 	struct FAnimUpdateRateParametersTracker
 	{
 		FAnimUpdateRateParameters UpdateRateParameters;
@@ -152,7 +154,7 @@ namespace FAnimUpdateRateManager
 			if (bUsingRootMotionFromEverything && DesiredEvaluationRate > 1)
 			{
 				//Use look ahead mode that allows us to rate limit updates even when using root motion
-				Tracker->UpdateRateParameters.SetLookAheadMode(DeltaTime, Tracker->GetAnimUpdateRateShiftTag(), DeltaTime*DesiredEvaluationRate);
+				Tracker->UpdateRateParameters.SetLookAheadMode(DeltaTime, Tracker->GetAnimUpdateRateShiftTag(), TargetFrameTimeForUpdateRate*DesiredEvaluationRate);
 			}
 			else
 			{
@@ -343,7 +345,7 @@ void USkinnedMeshComponent::SendRenderDynamicData_Concurrent()
 		}
 
 		// Are morph targets disabled for this LOD?
-		if ( SkeletalMesh->LODInfo[ UseLOD ].bHasBeenSimplified )
+		if ( SkeletalMesh->LODInfo[ UseLOD ].bHasBeenSimplified || bDisableMorphTarget )
 		{
 			ActiveVertexAnims.Empty();
 		}
@@ -1966,10 +1968,11 @@ void FAnimUpdateRateParameters::SetLookAheadMode(float DeltaTime, uint8 UpdateRa
 
 	if (TickedPoseOffestTime < 0.f)
 	{
+		LookAheadAmount = FMath::Max(TickedPoseOffestTime*-1.f, LookAheadAmount);
 		AdditionalTime = LookAheadAmount;
 		TickedPoseOffestTime += LookAheadAmount;
 
-		bool bValid = (TickedPoseOffestTime > 0.f);
+		bool bValid = (TickedPoseOffestTime >= 0.f);
 		if (!bValid)
 		{
 			FPlatformMisc::LowLevelOutputDebugStringf(TEXT("TPO Time: %.3f | Orig TPO Time: %.3f | DT: %.3f | LookAheadAmount: %.3f\n"), TickedPoseOffestTime, OriginalTickedPoseOffestTime, DeltaTime, LookAheadAmount);
