UE4 - merge GDC branch, code [REDACTED_HANDLE] to main

diff --git a/Engine/Source/Runtime/Engine/Private/SpeedTreeWind.cpp b/Engine/Source/Runtime/Engine/Private/SpeedTreeWind.cpp
--- a/Engine/Source/Runtime/Engine/Private/SpeedTreeWind.cpp
+++ b/Engine/Source/Runtime/Engine/Private/SpeedTreeWind.cpp
@@ -517,6 +517,16 @@ void FSpeedTreeWind::SetDirection(const FVector& vDir)
 	}
 }
 
+void FSpeedTreeWind::SetGustMin(float InGustMin)
+{
+	m_sParams.m_fGustStrengthMin = InGustMin;
+}
+
+void FSpeedTreeWind::SetGustMax(float InGustMax)
+{
+	m_sParams.m_fGustStrengthMax = InGustMax;
+}
+
 
 /////////////////////////////////////////////////////////////////////////////
 // FSpeedTreeWind::SetInitDirection
diff --git a/Engine/Source/Runtime/Engine/Private/SplineMeshComponent.cpp b/Engine/Source/Runtime/Engine/Private/SplineMeshComponent.cpp
--- a/Engine/Source/Runtime/Engine/Private/SplineMeshComponent.cpp
+++ b/Engine/Source/Runtime/Engine/Private/SplineMeshComponent.cpp
@@ -662,7 +662,7 @@ UBodySetup* USplineMeshComponent::GetBodySetup()
 	return NULL;
 }
 
-bool USplineMeshComponent::DoCustomNavigableGeometryExport(struct FNavigableGeometryExport* GeomExport) const
+bool USplineMeshComponent::DoCustomNavigableGeometryExport(FNavigableGeometryExport& GeomExport) const
 {
 	// the NavCollision is supposed to be faster than exporting the regular collision,
 	// but I'm not sure that's true here, as the regular collision is pre-distorted to the spline
@@ -684,7 +684,7 @@ bool USplineMeshComponent::DoCustomNavigableGeometryExport(struct FNavigableGeom
 					Vertex = CalcSliceTransform(GetAxisValue(Vertex, ForwardAxis)).TransformPosition(Vertex);
 					VertexBuffer.Add(Vertex);
 				}
-				GeomExport->ExportCustomMesh(VertexBuffer.GetData(), VertexBuffer.Num(),
+				GeomExport.ExportCustomMesh(VertexBuffer.GetData(), VertexBuffer.Num(),
 					NavCollision->ConvexCollision.IndexBuffer.GetData(), NavCollision->ConvexCollision.IndexBuffer.Num(),
 					ComponentToWorld);
 
@@ -695,7 +695,7 @@ bool USplineMeshComponent::DoCustomNavigableGeometryExport(struct FNavigableGeom
 					Vertex = CalcSliceTransform(GetAxisValue(Vertex, ForwardAxis)).TransformPosition(Vertex);
 					VertexBuffer.Add(Vertex);
 				}
-				GeomExport->ExportCustomMesh(VertexBuffer.GetData(), VertexBuffer.Num(),
+				GeomExport.ExportCustomMesh(VertexBuffer.GetData(), VertexBuffer.Num(),
 					NavCollision->TriMeshCollision.IndexBuffer.GetData(), NavCollision->TriMeshCollision.IndexBuffer.Num(),
 					ComponentToWorld);
 
diff --git a/Engine/Source/Runtime/Engine/Private/StaticMeshComponent.cpp b/Engine/Source/Runtime/Engine/Private/StaticMeshComponent.cpp
--- a/Engine/Source/Runtime/Engine/Private/StaticMeshComponent.cpp
+++ b/Engine/Source/Runtime/Engine/Private/StaticMeshComponent.cpp
@@ -1195,6 +1195,9 @@ bool UStaticMeshComponent::SetStaticMesh(UStaticMesh* NewMesh)
 	// Update physics representation right away
 	RecreatePhysicsState();
 
+	// update navigation relevancy
+	bNavigationRelevant = IsNavigationRelevant();
+
 	// Notify the streaming system. Don't use Update(), because this may be the first time the mesh has been set
 	// and the component may have to be added to the streaming system for the first time.
 	IStreamingManager::Get().NotifyPrimitiveAttached( this, DPT_Spawned );
@@ -1604,7 +1607,7 @@ void UStaticMeshComponent::ApplyComponentInstanceData(FStaticMeshComponentInstan
 }
 
 #include "AI/Navigation/RecastHelpers.h"
-bool UStaticMeshComponent::DoCustomNavigableGeometryExport(struct FNavigableGeometryExport* GeomExport) const
+bool UStaticMeshComponent::DoCustomNavigableGeometryExport(FNavigableGeometryExport& GeomExport) const
 {
 	if (StaticMesh != NULL && StaticMesh->NavCollision != NULL)
 	{
@@ -1620,10 +1623,10 @@ bool UStaticMeshComponent::DoCustomNavigableGeometryExport(struct FNavigableGeom
 			// if any of scales is 0 there's no point in exporting it
 			if (!Scale3D.IsZero())
 			{
-				GeomExport->ExportCustomMesh(NavCollision->ConvexCollision.VertexBuffer.GetData(), NavCollision->ConvexCollision.VertexBuffer.Num(),
+				GeomExport.ExportCustomMesh(NavCollision->ConvexCollision.VertexBuffer.GetData(), NavCollision->ConvexCollision.VertexBuffer.Num(),
 					NavCollision->ConvexCollision.IndexBuffer.GetData(), NavCollision->ConvexCollision.IndexBuffer.Num(), ComponentToWorld);
 
-				GeomExport->ExportCustomMesh(NavCollision->TriMeshCollision.VertexBuffer.GetData(), NavCollision->TriMeshCollision.VertexBuffer.Num(),
+				GeomExport.ExportCustomMesh(NavCollision->TriMeshCollision.VertexBuffer.GetData(), NavCollision->TriMeshCollision.VertexBuffer.Num(),
 					NavCollision->TriMeshCollision.IndexBuffer.GetData(), NavCollision->TriMeshCollision.IndexBuffer.Num(), ComponentToWorld);
 			}
 
@@ -1635,6 +1638,11 @@ bool UStaticMeshComponent::DoCustomNavigableGeometryExport(struct FNavigableGeom
 	return true;
 }
 
+bool UStaticMeshComponent::IsNavigationRelevant() const
+{
+	return StaticMesh != nullptr && Super::IsNavigationRelevant();
+}
+
 void UStaticMeshComponent::GetNavigationData(FNavigationRelevantData& Data) const
 {
 	Super::GetNavigationData(Data);
diff --git a/Engine/Source/Runtime/Engine/Private/Texture.cpp b/Engine/Source/Runtime/Engine/Private/Texture.cpp
--- a/Engine/Source/Runtime/Engine/Private/Texture.cpp
+++ b/Engine/Source/Runtime/Engine/Private/Texture.cpp
@@ -841,6 +841,12 @@ void FTextureSource::UseHashAsGuid()
 	}
 }
 
+void FTextureSource::SetId(const FGuid& InId, bool bInGuidIsHash)
+{
+	Id = InId;
+	bGuidIsHash = bInGuidIsHash;
+}
+
 uint32 UTexture::GetMaximumDimension() const
 {
 	return GetMax2DTextureDimension();
diff --git a/Engine/Source/Runtime/Engine/Private/Texture2D.cpp b/Engine/Source/Runtime/Engine/Private/Texture2D.cpp
--- a/Engine/Source/Runtime/Engine/Private/Texture2D.cpp
+++ b/Engine/Source/Runtime/Engine/Private/Texture2D.cpp
@@ -742,7 +742,9 @@ int32 UTexture2D::GetNumNonStreamingMips() const
 void UTexture2D::CalcAllowedMips( int32 MipCount, int32 NumNonStreamingMips, int32 LODBias, int32& OutMinAllowedMips, int32& OutMaxAllowedMips )
 {
 	// Calculate the minimum number of mip-levels required.
-	int32 MinAllowedMips = FMath::Max( UTexture2D::GetMinTextureResidentMipCount(), NumNonStreamingMips );
+	int32 MinAllowedMips = UTexture2D::GetMinTextureResidentMipCount();
+	MinAllowedMips = FMath::Max( MinAllowedMips, MipCount - LODBias );
+	MinAllowedMips = FMath::Min( MinAllowedMips, NumNonStreamingMips );
 	MinAllowedMips = FMath::Min( MinAllowedMips, MipCount );
 
 	// Calculate the maximum number of mip-levels.
@@ -1799,6 +1801,9 @@ void FTexture2DResource::LoadMipData()
 					}
 					check(IORequestIndices[MipIndex]);
 				}
+
+				// For consistency with other code paths, track the pointer to the locked buffer.
+				MipData[ActualMipIndex] = TheMipData;
 			}
 
 			// Are we reducing the mip-count?
@@ -1847,18 +1852,16 @@ void FTexture2DResource::LoadMipData()
 			const FTexture2DMipMap& MipMap = OwnerMips[ActualMipIndex];
 			int32 MipSize = CalcTextureMipMapSize(MipMap.SizeX, MipMap.SizeY, TextureRHI->GetFormat(), 0);
 
-			void* TheMipData = NULL;
 			if (bUsingAsyncCreation)
 			{
-				// Allocate temporary system memory to stream in to.
-				MipData[ActualMipIndex] = FMemory::Malloc(MipSize);
-				TheMipData = MipData[ActualMipIndex];
+				// The async task will allocate temporary system memory to stream in to.
+				check(MipData[ActualMipIndex] == NULL);
 			}
 			else
 			{
 				// Lock the new texture.
 				uint32 DestPitch;
-				TheMipData = RHILockTexture2D( IntermediateTextureRHI, MipIndex, RLM_WriteOnly, DestPitch, false );
+				MipData[ActualMipIndex] = RHILockTexture2D( IntermediateTextureRHI, MipIndex, RLM_WriteOnly, DestPitch, false );
 			}
 
 			// Pass the request on to the async io manager after increasing the request count. The request count 
@@ -1875,7 +1878,7 @@ void FTexture2DResource::LoadMipData()
 			{
 				FAsyncStreamDerivedMipTask* Task = new(PendingAsyncStreamDerivedMipTasks) FAsyncStreamDerivedMipTask(
 					MipMap.DerivedDataKey,
-					TheMipData,
+					&MipData[ActualMipIndex],
 					MipSize,
 					&Owner->PendingMipChangeRequestStatus
 					);
@@ -1884,6 +1887,11 @@ void FTexture2DResource::LoadMipData()
 			else
 #endif // #if WITH_EDITORONLY_DATA
 			{
+				if (!MipData[ActualMipIndex])
+				{
+					QUICK_SCOPE_CYCLE_COUNTER(STAT_FTexture2DResource_LoadMipData_Malloc);
+					MipData[ActualMipIndex] = FMemory::Malloc(MipSize);
+				}
 				check(MipMap.BulkData.GetFilename().Len());
 				if( MipMap.BulkData.IsStoredCompressedOnDisk() )
 				{
@@ -1892,7 +1900,7 @@ void FTexture2DResource::LoadMipData()
 						MipMap.BulkData.GetBulkDataOffsetInFile(),			// offset
 						MipMap.BulkData.GetBulkDataSizeOnDisk(),			// compressed size
 						MipMap.BulkData.GetBulkDataSize(),					// uncompressed size
-						TheMipData,											// dest pointer
+						MipData[ActualMipIndex],							// dest pointer
 						MipMap.BulkData.GetDecompressionFlags(),			// compressed data format
 						&Owner->PendingMipChangeRequestStatus,				// counter to decrement
 						AsyncIOPriority										// priority
@@ -1905,7 +1913,7 @@ void FTexture2DResource::LoadMipData()
 						MipMap.BulkData.GetFilename(),						// filename
 						MipMap.BulkData.GetBulkDataOffsetInFile(),			// offset
 						MipMap.BulkData.GetBulkDataSize(),					// size
-						TheMipData,											// dest pointer
+						MipData[ActualMipIndex],							// dest pointer
 						&Owner->PendingMipChangeRequestStatus,				// counter to decrement
 						AsyncIOPriority										// priority
 						);
@@ -1969,6 +1977,7 @@ void FTexture2DResource::UploadMipData()
 			{
 				int32 ActualMipIndex = MipIndex + PendingFirstMip;
 				RHIUnlockTexture2D( Texture2DRHI, ActualMipIndex, false );
+				MipData[ActualMipIndex] = NULL;
 			}
 		}
 
@@ -2015,6 +2024,7 @@ void FTexture2DResource::UploadMipData()
 				// Intermediate texture has RequestedMips miplevels, all of which have been locked.
 				// DEXTEX: Do not unload as we didn't locked the textures in the first place
 				RHIUnlockTexture2D( IntermediateTextureRHI, MipIndex, false );
+				MipData[MipIndex + PendingFirstMip] = NULL;
 			}
 		}
 	}
diff --git a/Engine/Source/Runtime/Engine/Private/TextureDerivedData.cpp b/Engine/Source/Runtime/Engine/Private/TextureDerivedData.cpp
--- a/Engine/Source/Runtime/Engine/Private/TextureDerivedData.cpp
+++ b/Engine/Source/Runtime/Engine/Private/TextureDerivedData.cpp
@@ -1313,16 +1313,19 @@ void FTexturePlatformData::SerializeCooked(FArchive& Ar, UTexture* Owner, bool b
 /** Initialization constructor. */
 FAsyncStreamDerivedMipWorker::FAsyncStreamDerivedMipWorker(
 	const FString& InDerivedDataKey,
-	void* InDestMipData,
+	void** InDestMipDataPointer,
 	int32 InMipSize,
 	FThreadSafeCounter* InThreadSafeCounter
 	)
 	: DerivedDataKey(InDerivedDataKey)
-	, DestMipData(InDestMipData)
+	, DestMipDataPointer(InDestMipDataPointer)
 	, ExpectedMipSize(InMipSize)
 	, bRequestFailed(false)
 	, ThreadSafeCounter(InThreadSafeCounter)
 {
+	check(DestMipDataPointer != NULL);
+
+	// Make sure the pixel format enum is cached -- we access it on another thread.
 	UTexture::GetPixelFormatEnum();
 }
 
@@ -1334,10 +1337,18 @@ void FAsyncStreamDerivedMipWorker::DoWork()
 	if (GetDerivedDataCacheRef().GetSynchronous(*DerivedDataKey, DerivedMipData))
 	{
 		FMemoryReader Ar(DerivedMipData, true);
+
 		int32 MipSize = 0;
 		Ar << MipSize;
-		checkf(ExpectedMipSize == ExpectedMipSize, TEXT("MipSize(%d) == ExpectedSize(%d)"), MipSize, ExpectedMipSize);
-		Ar.Serialize(DestMipData, MipSize);
+		checkf(MipSize == ExpectedMipSize, TEXT("MipSize(%d) == ExpectedSize(%d)"), MipSize, ExpectedMipSize);
+
+		// Allocate memory for the mip unless the caller has already done so.
+		if (*DestMipDataPointer == NULL)
+		{
+			*DestMipDataPointer = FMemory::Malloc(MipSize);
+		}
+
+		Ar.Serialize(*DestMipDataPointer, MipSize);
 	}
 	else
 	{
diff --git a/Engine/Source/Runtime/Engine/Private/TickTaskManager.cpp b/Engine/Source/Runtime/Engine/Private/TickTaskManager.cpp
--- a/Engine/Source/Runtime/Engine/Private/TickTaskManager.cpp
+++ b/Engine/Source/Runtime/Engine/Private/TickTaskManager.cpp
@@ -180,6 +180,7 @@ public:
 				STATGROUP_TaskGraphTasks);
 
 			// dispatch the tick group on another thread, that way, the game thread can be processing ticks while ticks are being queued by another thread
+			QUICK_SCOPE_CYCLE_COUNTER(STAT_ReleaseTickGroup);
 			FTaskGraphInterface::Get().WaitUntilTaskCompletes(
 				FDelegateGraphTask::CreateAndDispatchWhenReady(
 					FDelegateGraphTask::FDelegate::CreateRaw(this, &FTickTaskSequencer::DispatchTickGroup, WorldTickGroup),
@@ -191,6 +192,7 @@ public:
 		{
 			if (bBlockTillComplete || SingleThreadedMode())
 			{
+				QUICK_SCOPE_CYCLE_COUNTER(STAT_ReleaseTickGroup_Block);
 				FTaskGraphInterface::Get().WaitUntilTasksComplete(TickCompletionEvents[WorldTickGroup],ENamedThreads::GameThread);
 			}
 			else
diff --git a/Engine/Source/Runtime/Engine/Private/UnrealEngine.cpp b/Engine/Source/Runtime/Engine/Private/UnrealEngine.cpp
--- a/Engine/Source/Runtime/Engine/Private/UnrealEngine.cpp
+++ b/Engine/Source/Runtime/Engine/Private/UnrealEngine.cpp
@@ -953,6 +953,7 @@ void UEngine::PreExit()
 
 void UEngine::TickDeferredCommands()
 {
+	QUICK_SCOPE_CYCLE_COUNTER(STAT_UEngine_TickDeferredCommands);
 	// Execute all currently queued deferred commands (allows commands to be queued up for next frame).
 	const int32 DeferredCommandsCount = DeferredCommands.Num();
 	for( int32 DeferredCommandsIndex=0; DeferredCommandsIndex<DeferredCommandsCount; DeferredCommandsIndex++ )
@@ -4607,7 +4608,19 @@ struct FHierarchy
 // @TODO yrx 2014-09-15 Move to ObjectCommads.cpp or ObjectExec.cpp
 bool UEngine::HandleObjCommand( const TCHAR* Cmd, FOutputDevice& Ar )
 {
-	if (FParse::Command(&Cmd,TEXT("LIST2")))
+	if( FParse::Command(&Cmd,TEXT("GARBAGE")) || FParse::Command(&Cmd,TEXT("GC")) )
+	{
+		// Purge unclaimed objects.
+		Ar.Logf(TEXT("Collecting garbage and resetting GC timers on all worlds."));
+		CollectGarbage( GARBAGE_COLLECTION_KEEPFLAGS );
+		for (TObjectIterator<UWorld> It; It; ++It)
+		{
+			UWorld* CurrentWorld = *It;
+			CurrentWorld->TimeSinceLastPendingKillPurge = 0;
+		}
+		return true;
+	}
+	else if (FParse::Command(&Cmd,TEXT("LIST2")))
 	{			
 		UClass* ClassToCheck = NULL;
 		ParseObject<UClass>(Cmd, TEXT("CLASS="  ), ClassToCheck, ANY_PACKAGE );
@@ -9199,6 +9212,8 @@ bool UEngine::LoadMap( FWorldContext& WorldContext, FURL URL, class UPendingNetG
 
 void UEngine::BlockTillLevelStreamingCompleted(UWorld* InWorld)
 {
+	QUICK_SCOPE_CYCLE_COUNTER(STAT_UEngine_BlockTillLevelStreamingCompleted);
+
 	check(InWorld);
 	
 	// Update streaming levels state using streaming volumes
diff --git a/Engine/Source/Runtime/Engine/Private/WindDirectionalSource.cpp b/Engine/Source/Runtime/Engine/Private/WindDirectionalSource.cpp
--- a/Engine/Source/Runtime/Engine/Private/WindDirectionalSource.cpp
+++ b/Engine/Source/Runtime/Engine/Private/WindDirectionalSource.cpp
@@ -53,8 +53,32 @@ AWindDirectionalSource::AWindDirectionalSource(const FObjectInitializer& ObjectI
 #endif // WITH_EDITORONLY_DATA
 }
 
+void FWindSourceSceneProxy::FWindData::PrepareForAccumulate()
+{
+	FMemory::Memset(this, 0, sizeof(FWindData));
+}
+
+void FWindSourceSceneProxy::FWindData::AddWeighted(const FWindData& InWindData, float Weight)
+{	
+	Speed += InWindData.Speed * Weight;
+	MinGustAmt += InWindData.MinGustAmt * Weight;
+	MaxGustAmt += InWindData.MaxGustAmt * Weight;
+	Direction += InWindData.Direction * Weight;
+}
+
+void FWindSourceSceneProxy::FWindData::NormalizeByTotalWeight(float TotalWeight)
+{
+	if (TotalWeight > 0.0f)
+	{
+		Speed /= TotalWeight;
+		MinGustAmt /= TotalWeight;
+		MaxGustAmt /= TotalWeight;
+		Direction /= TotalWeight;
+		Direction.Normalize();
+	}
+}
 
-bool FWindSourceSceneProxy::GetWindParameters(const FVector& EvaluatePosition, FVector4& WindDirectionAndSpeed, float& Weight) const 
+bool FWindSourceSceneProxy::GetWindParameters(const FVector& EvaluatePosition, FWindData& WindData, float& Weight) const
 { 
 	if (bIsPointSource)
 	{
@@ -63,31 +87,43 @@ bool FWindSourceSceneProxy::GetWindParameters(const FVector& EvaluatePosition, F
 		{
 			// Mimic Engine point light attenuation with a FalloffExponent of 1
 			const float RadialFalloff = FMath::Max(1.0f - ((EvaluatePosition - Position) / Radius).SizeSquared(), 0.0f);
-			WindDirectionAndSpeed = FVector4((EvaluatePosition - Position) / Distance * Strength * RadialFalloff, Speed); 
+			//WindDirectionAndSpeed = FVector4((EvaluatePosition - Position) / Distance * Strength * RadialFalloff, Speed); 
+
+			WindData.Direction = (EvaluatePosition - Position) / Distance;
+			WindData.Speed = Speed * RadialFalloff;
+			WindData.MinGustAmt = MinGustAmt * RadialFalloff;
+			WindData.MaxGustAmt = MaxGustAmt * RadialFalloff;
+
 			Weight = Distance / Radius * Strength;
 			return true;
 		}
 		Weight = 0.f;
-		WindDirectionAndSpeed = FVector4(0,0,0,0);
+		WindData = FWindData();
 		return false;
 	}
 
-	Weight = Strength;
-	WindDirectionAndSpeed = FVector4(Direction * Strength, Speed); 
+	Weight				= Strength;
+	WindData.Direction	= Direction;
+	WindData.Speed		= Speed;
+	WindData.MinGustAmt = MinGustAmt;
+	WindData.MaxGustAmt = MaxGustAmt;	
 	return true;
 }
 
-bool FWindSourceSceneProxy::GetDirectionalWindParameters(FVector4& WindDirectionAndSpeed, float& Weight) const 
+bool FWindSourceSceneProxy::GetDirectionalWindParameters(FWindData& WindData, float& Weight) const
 { 
 	if (bIsPointSource)
 	{
 		Weight = 0.f;
-		WindDirectionAndSpeed = FVector4(0,0,0,0);
+		WindData = FWindData();
 		return false;
 	}
 
-	Weight = Strength;
-	WindDirectionAndSpeed = FVector4(Direction * Strength, Speed); 
+	Weight				= Strength;
+	WindData.Direction	= Direction;
+	WindData.Speed		= Speed;
+	WindData.MinGustAmt = MinGustAmt;
+	WindData.MaxGustAmt = MaxGustAmt;
 	return true;
 }
 
@@ -102,8 +138,11 @@ void FWindSourceSceneProxy::ApplyWorldOffset(FVector InOffset)
 UWindDirectionalSourceComponent::UWindDirectionalSourceComponent(const FObjectInitializer& ObjectInitializer)
 	: Super(ObjectInitializer)
 {
-	Strength = 1.0f;
-	Speed = 1.0f;
+	Strength = 0.1f;
+	Speed = 0.1f;
+	MinGustAmount = 0.1f;
+	MaxGustAmount = 0.2f;
+
 	// wind will be activated automatically by default
 	bAutoActivate = true;
 }
@@ -132,7 +171,9 @@ FWindSourceSceneProxy* UWindDirectionalSourceComponent::CreateSceneProxy() const
 	return new FWindSourceSceneProxy(
 		ComponentToWorld.GetUnitAxis( EAxis::X ),
 		Strength,
-		Speed
+		Speed,
+		MinGustAmount,
+		MaxGustAmount
 		);
 }
 
diff --git a/Engine/Source/Runtime/Engine/Private/World.cpp b/Engine/Source/Runtime/Engine/Private/World.cpp
--- a/Engine/Source/Runtime/Engine/Private/World.cpp
+++ b/Engine/Source/Runtime/Engine/Private/World.cpp
@@ -44,6 +44,7 @@
 #include "Materials/MaterialParameterCollectionInstance.h"
 
 #if WITH_EDITOR
+	#include "DerivedDataCacheInterface.h"
 	#include "Editor/UnrealEd/Public/Kismet2/KismetEditorUtilities.h"
 	#include "Editor/UnrealEd/Public/Kismet2/BlueprintEditorUtils.h"
 	#include "Editor/UnrealEd/Classes/ThumbnailRendering/WorldThumbnailInfo.h"
@@ -122,6 +123,14 @@ UWorld::UWorld( const FObjectInitializer& ObjectInitializer )
 	FWorldDelegates::OnPostWorldCreation.Broadcast(this);
 }
 
+UWorld::~UWorld()
+{
+	while (AsyncPreRegisterLevelStreamingTasks.GetValue())
+	{
+		FPlatformProcess::Sleep(0.0f);
+	}
+}
+
 void UWorld::Serialize( FArchive& Ar )
 {
 	Super::Serialize( Ar );
@@ -1495,7 +1504,7 @@ static bool IsTimeLimitExceeded( const TCHAR* CurrentTask, double StartTime, ULe
 			// Log if a single event took way too much time.
 			if( DeltaTime > 20 )
 			{
-				UE_LOG(LogStreaming, Log, TEXT("UWorld::AddToWorld: %s for %s took (less than) %5.2f ms"), CurrentTask, *Level->GetOutermost()->GetName(), DeltaTime );
+				UE_LOG(LogStreaming, Display, TEXT("UWorld::AddToWorld: %s for %s took (less than) %5.2f ms"), CurrentTask, *Level->GetOutermost()->GetName(), DeltaTime );
 			}
 			bIsTimeLimitExceed = true;
 		}
@@ -1504,7 +1513,6 @@ static bool IsTimeLimitExceeded( const TCHAR* CurrentTask, double StartTime, ULe
 }
 
 #if PERF_TRACK_DETAILED_ASYNC_STATS
-
 // Variables for tracking how long each part of the AddToWorld process takes
 double MoveActorTime = 0.0;
 double ShiftActorsTime = 0.0;
@@ -1557,6 +1565,8 @@ void UWorld::AddToWorld( ULevel* Level, const FTransform& LevelTransform )
 	check(!Level->IsPendingKill());
 	check(!Level->HasAnyFlags(RF_Unreachable));
 
+	FScopeCycleCounterUObject ContextScope(Level);
+
 	// Set flags to indicate that we are associating a level with the world to e.g. perform slower/ better octree insertion 
 	// and such, as opposed to the fast path taken for run-time/ gameplay objects.
 	Level->bIsAssociatingLevel = true;
@@ -1617,6 +1627,54 @@ void UWorld::AddToWorld( ULevel* Level, const FTransform& LevelTransform )
 		bExecuteNextStep = (!bConsiderTimeLimit || !IsTimeLimitExceeded( TEXT("shifting actors"), StartTime, Level ));
 	}
 
+	if (bExecuteNextStep && AsyncPreRegisterLevelStreamingTasks.GetValue())
+	{
+		if (!bConsiderTimeLimit)
+		{
+			QUICK_SCOPE_CYCLE_COUNTER(UWorld_AddToWorld_WaitFor_AsyncPreRegisterLevelStreamingTasks);
+			while (AsyncPreRegisterLevelStreamingTasks.GetValue())
+			{
+				FPlatformProcess::Sleep(0.001f);
+			}
+		}
+		else
+		{
+			bExecuteNextStep = false;
+		}
+	}
+
+	// Wait on any async DDC handles
+#if WITH_EDITOR
+	if (bExecuteNextStep && AsyncPreRegisterDDCRequests.Num())
+	{
+		if (!bConsiderTimeLimit)
+		{
+			QUICK_SCOPE_CYCLE_COUNTER(UWorld_AddToWorld_WaitFor_AsyncPreRegisterLevelStreamingTasks);
+
+			for (TSharedPtr<FAsyncPreRegisterDDCRequest>& Request : AsyncPreRegisterDDCRequests)
+			{
+				Request->WaitAsynchronousCompletion();
+			}
+			AsyncPreRegisterDDCRequests.Empty();
+		}
+		else
+		{
+			for (int32 Index = 0; Index < AsyncPreRegisterDDCRequests.Num(); Index++)
+			{
+				if (AsyncPreRegisterDDCRequests[Index]->PollAsynchronousCompletion())
+				{
+					AsyncPreRegisterDDCRequests.RemoveAtSwap(Index--);
+				}
+				else
+				{
+					bExecuteNextStep = false;
+					break;
+				}
+			}
+		}
+	}
+#endif
+
 	// Updates the level components (Actor components and UModelComponents).
 	if( bExecuteNextStep && !Level->bAlreadyUpdatedComponents )
 	{
@@ -1629,9 +1687,14 @@ void UWorld::AddToWorld( ULevel* Level, const FTransform& LevelTransform )
 			// Pretend here that we are loading package to avoid package dirtying during components registration
 		TGuardValue<bool> IsEditorLoadingPackage(GIsEditorLoadingPackage, (GIsEditor ? true : GIsEditorLoadingPackage));
 #endif
+
+		// Config bool that allows disabling all construction scripts during PIE level streaming.
+		bool bRerunConstructionDuringEditorStreaming = true;
+		GConfig->GetBool(TEXT("Kismet"), TEXT("bRerunConstructionDuringEditorStreaming"), /*out*/ bRerunConstructionDuringEditorStreaming, GEngineIni);
+
 		// We don't need to rerun construction scripts if we have cooked data or we are playing in editor unless the PIE world was loaded
 		// from disk rather than duplicated
-		const bool bRerunConstructionScript = !(FPlatformProperties::RequiresCookedData() || (IsPlayInEditor() && Level->bWasDuplicatedForPIE));
+		const bool bRerunConstructionScript = !(FPlatformProperties::RequiresCookedData() || (IsGameWorld() && (Level->bWasDuplicatedForPIE || !bRerunConstructionDuringEditorStreaming)));
 		
 		// Incrementally update components.
 		int32 NumComponentsToUpdate = GEngine->LevelStreamingComponentsRegistrationGranularity;
@@ -1639,11 +1702,11 @@ void UWorld::AddToWorld( ULevel* Level, const FTransform& LevelTransform )
 		{
 			Level->IncrementalUpdateComponents( (!IsGameWorld() || IsRunningCommandlet()) ? 0 : NumComponentsToUpdate, bRerunConstructionScript );
 		}
-		while( (!bConsiderTimeLimit || !IsTimeLimitExceeded( TEXT("updating components"), StartTime, Level )) && !Level->bAreComponentsCurrentlyRegistered );
+		while( !Level->bAreComponentsCurrentlyRegistered && (!bConsiderTimeLimit || !IsTimeLimitExceeded( TEXT("updating components"), StartTime, Level )));
 
 		// We are done once all components are attached.
 		Level->bAlreadyUpdatedComponents	= Level->bAreComponentsCurrentlyRegistered;
-		bExecuteNextStep					= Level->bAreComponentsCurrentlyRegistered;
+		bExecuteNextStep					= Level->bAreComponentsCurrentlyRegistered && (!bConsiderTimeLimit || !IsTimeLimitExceeded(TEXT("updating components"), StartTime, Level));
 	}
 
 	if( IsGameWorld() && AreActorsInitialized() )
@@ -1753,25 +1816,25 @@ void UWorld::AddToWorld( ULevel* Level, const FTransform& LevelTransform )
 			SortActorListTime + 
 			PerformLastStepTime;
 
-		UE_LOG(LogStreaming, Log, TEXT("Detailed AddToWorld stats for '%s' - Total %6.2fms"), *Level->GetOutermost()->GetName(), TotalTime * 1000 );
-		UE_LOG(LogStreaming, Log, TEXT("Move Actors             : %6.2f ms"), MoveActorTime * 1000 );
-		UE_LOG(LogStreaming, Log, TEXT("Shift Actors            : %6.2f ms"), ShiftActorsTime * 1000 );
-		UE_LOG(LogStreaming, Log, TEXT("Update Components       : %6.2f ms"), UpdateComponentsTime * 1000 );
-		UE_LOG(LogStreaming, Log, TEXT("Init BSP Phys           : %6.2f ms"), InitBSPPhysTime * 1000 );
-		UE_LOG(LogStreaming, Log, TEXT("Init Actor Phys         : %6.2f ms"), InitActorPhysTime * 1000 );
-		UE_LOG(LogStreaming, Log, TEXT("Init Actors             : %6.2f ms"), InitActorTime * 1000 );
-		UE_LOG(LogStreaming, Log, TEXT("Initialize              : %6.2f ms"), RouteActorInitializeTime * 1000 );
-		UE_LOG(LogStreaming, Log, TEXT("Cross Level Refs        : %6.2f ms"), CrossLevelRefsTime * 1000 );
-		UE_LOG(LogStreaming, Log, TEXT("Sort Actor List         : %6.2f ms"), SortActorListTime * 1000 );
-		UE_LOG(LogStreaming, Log, TEXT("Perform Last Step       : %6.2f ms"), SortActorListTime * 1000 );
+		UE_LOG(LogStreaming, Display, TEXT("Detailed AddToWorld stats for '%s' - Total %6.2fms"), *Level->GetOutermost()->GetName(), TotalTime * 1000 );
+		UE_LOG(LogStreaming, Display, TEXT("Move Actors             : %6.2f ms"), MoveActorTime * 1000 );
+		UE_LOG(LogStreaming, Display, TEXT("Shift Actors            : %6.2f ms"), ShiftActorsTime * 1000 );
+		UE_LOG(LogStreaming, Display, TEXT("Update Components       : %6.2f ms"), UpdateComponentsTime * 1000 );
+		UE_LOG(LogStreaming, Display, TEXT("Init BSP Phys           : %6.2f ms"), InitBSPPhysTime * 1000 );
+		UE_LOG(LogStreaming, Display, TEXT("Init Actor Phys         : %6.2f ms"), InitActorPhysTime * 1000 );
+		UE_LOG(LogStreaming, Display, TEXT("Init Actors             : %6.2f ms"), InitActorTime * 1000 );
+		UE_LOG(LogStreaming, Display, TEXT("Initialize              : %6.2f ms"), RouteActorInitializeTime * 1000 );
+		UE_LOG(LogStreaming, Display, TEXT("Cross Level Refs        : %6.2f ms"), CrossLevelRefsTime * 1000 );
+		UE_LOG(LogStreaming, Display, TEXT("Sort Actor List         : %6.2f ms"), SortActorListTime * 1000 );
+		UE_LOG(LogStreaming, Display, TEXT("Perform Last Step       : %6.2f ms"), SortActorListTime * 1000 );
 	}
 #endif // PERF_TRACK_DETAILED_ASYNC_STATS
 }
 
-
 void UWorld::RemoveFromWorld( ULevel* Level )
 {
 	SCOPE_CYCLE_COUNTER(STAT_RemoveFromWorldTime);
+	FScopeCycleCounterUObject Context(Level);
 	check(Level);
 	check(!Level->IsPendingKill());
 	check(!Level->HasAnyFlags(RF_Unreachable));
@@ -1838,6 +1901,14 @@ void UWorld::RemoveFromWorld( ULevel* Level )
 			}
 		}
 		
+		// Remove any actors in the network list
+		for ( int i = 0; i < Level->Actors.Num(); i++ )
+		{
+			RemoveNetworkActor( Level->Actors[ i ] );
+		}
+
+		Level->bIsVisible = false;
+
 		// Notify world composition: will place a level at original position
 		if (WorldComposition)
 		{
@@ -1845,26 +1916,21 @@ void UWorld::RemoveFromWorld( ULevel* Level )
 		}
 
 		// Make sure level always has OwningWorld in the editor
-		if ( IsGameWorld() )
+		if (IsGameWorld())
 		{
 			Levels.Remove(Level);
 			Level->OwningWorld = NULL;
 		}
 				
-		// Remove any actors in the network list
-		for ( int i = 0; i < Level->Actors.Num(); i++ )
-		{
-			RemoveNetworkActor( Level->Actors[ i ] );
-		}
-
-		Level->bIsVisible = false;
-		UE_LOG(LogStreaming, Log, TEXT("UWorld::RemoveFromWorld for %s took %5.2f ms"), *Level->GetOutermost()->GetName(), (FPlatformTime::Seconds()-StartTime) * 1000.0 );
-			
 		// let the universe know we have removed a level
 		FWorldDelegates::LevelRemovedFromWorld.Broadcast(Level, this);
 		BroadcastLevelsChanged();
 
 		ULevelStreaming::BroadcastLevelVisibleStatus(this, Level->GetOutermost()->GetFName(), false);
+
+#if PERF_TRACK_DETAILED_ASYNC_STATS
+		UE_LOG(LogStreaming, Display, TEXT("UWorld::RemoveFromWorld for %s took %5.2f ms"), *Level->GetOutermost()->GetName(), (FPlatformTime::Seconds() - StartTime) * 1000.0);
+#endif // PERF_TRACK_DETAILED_ASYNC_STATS
 	}
 }
 
@@ -2080,6 +2146,9 @@ FString UWorld::BuildPIEPackagePrefix(int PIEInstanceID)
 
 UWorld* UWorld::DuplicateWorldForPIE(const FString& PackageName, UWorld* OwningWorld)
 {
+	QUICK_SCOPE_CYCLE_COUNTER(UWorld_DuplicateWorldForPIE);
+	FScopeCycleCounterUObject Context(OwningWorld);
+
 	// Find the original (non-PIE) level package
 	UPackage* EditorLevelPackage = FindObjectFast<UPackage>( NULL, FName(*PackageName));
 	if( !EditorLevelPackage )
@@ -2156,6 +2225,8 @@ UWorld* UWorld::DuplicateWorldForPIE(const FString& PackageName, UWorld* OwningW
 			&& PIELevel->Model == EditorLevel->Model
 			&& PIELevel->ModelComponents.Num() == EditorLevel->ModelComponents.Num() )
 		{
+			QUICK_SCOPE_CYCLE_COUNTER(UWorld_DuplicateWorldForPIE_UpdateModelComponents);
+
 			PIELevel->Model->ClearLocalMaterialIndexBuffersData();
 			for (int32 ComponentIndex = 0; ComponentIndex < PIELevel->ModelComponents.Num(); ++ComponentIndex)
 			{
@@ -2197,6 +2268,7 @@ void UWorld::UpdateLevelStreamingInner(ULevelStreaming* StreamingLevel)
 			return;
 		}
 	}
+	FScopeCycleCounterUObject ContextScope(StreamingLevel);
 
 	// Work performed to make a level visible is spread across several frames and we can't unload/ hide a level that is currently pending
 	// to be made visible, so we fulfill those requests first.
@@ -2263,6 +2335,7 @@ void UWorld::UpdateLevelStreamingInner(ULevelStreaming* StreamingLevel)
 
 					if (Scene)
 					{
+						QUICK_SCOPE_CYCLE_COUNTER(STAT_UpdateLevelStreamingInner_OnLevelAddedToWorld);
 						// Notify the new level has been added after the old has been discarded
 						Scene->OnLevelAddedToWorld(Level->GetOutermost()->GetFName());
 					}
@@ -2821,10 +2894,14 @@ void UWorld::InitializeActorsForPlay(const FURL& InURL, bool bResetTime)
 		URL = InURL;
 	}
 
+	// Config bool that allows disabling all construction scripts during PIE level streaming.
+	bool bRerunConstructionDuringEditorStreaming = true;
+	GConfig->GetBool(TEXT("Kismet"), TEXT("bRerunConstructionDuringEditorStreaming"), /*out*/ bRerunConstructionDuringEditorStreaming, GEngineIni);
+
 	// Update world and the components of all levels.	
 	// We don't need to rerun construction scripts if we have cooked data or we are playing in editor unless the PIE world was loaded
 	// from disk rather than duplicated
-	const bool bRerunConstructionScript = !(FPlatformProperties::RequiresCookedData() || (IsPlayInEditor() && PersistentLevel->bWasDuplicatedForPIE));
+	const bool bRerunConstructionScript = !(FPlatformProperties::RequiresCookedData() || (IsGameWorld() && (PersistentLevel->bWasDuplicatedForPIE || !bRerunConstructionDuringEditorStreaming)));
 	UpdateWorldComponents( bRerunConstructionScript, true );
 
 	// Reset indices till we have a chance to rearrange actor list at the end of this function.
@@ -5442,4 +5519,44 @@ static FAutoConsoleCommand DumpVisibleActorsCmd(
 	FConsoleCommandDelegate::CreateStatic(DumpVisibleActors)
 	);
 
+#if WITH_EDITOR
+FAsyncPreRegisterDDCRequest::~FAsyncPreRegisterDDCRequest()
+{
+	// Discard any results
+	if (Handle != 0)
+	{
+		WaitAsynchronousCompletion();
+		TArray<uint8> Junk;
+		GetAsynchronousResults(Junk);
+	}
+}
+
+bool FAsyncPreRegisterDDCRequest::PollAsynchronousCompletion()
+{
+	if (Handle != 0)
+	{
+		return GetDerivedDataCacheRef().PollAsynchronousCompletion(Handle);
+	}
+	return true;
+}
+
+void FAsyncPreRegisterDDCRequest::WaitAsynchronousCompletion()
+{
+	if (Handle != 0)
+	{
+		GetDerivedDataCacheRef().WaitAsynchronousCompletion(Handle);
+	}
+}
+
+bool FAsyncPreRegisterDDCRequest::GetAsynchronousResults(TArray<uint8>& OutData)
+{
+	check(Handle != 0);
+	bool bResult = GetDerivedDataCacheRef().GetAsynchronousResults(Handle, OutData);
+	// invalidate request after results received
+	Handle = 0;
+	DDCKey = TEXT("");
+	return bResult;
+}
+#endif
+
 #undef LOCTEXT_NAMESPACE 
diff --git a/Engine/Source/Runtime/Engine/Private/WorldComposition.cpp b/Engine/Source/Runtime/Engine/Private/WorldComposition.cpp
--- a/Engine/Source/Runtime/Engine/Private/WorldComposition.cpp
+++ b/Engine/Source/Runtime/Engine/Private/WorldComposition.cpp
@@ -341,13 +341,15 @@ void UWorldComposition::Reset()
 	TilesStreaming.Empty();
 }
 
-FWorldCompositionTile* UWorldComposition::FindTileByName(const FName& InPackageName) const
+int32 UWorldComposition::FindTileIndexByName(const FName& InPackageName) const
 {
-	for (const FWorldCompositionTile& Tile : Tiles)
+	for (int32 TileIdx = 0; TileIdx < Tiles.Num(); ++TileIdx)
 	{
+		const FWorldCompositionTile& Tile = Tiles[TileIdx];
+
 		if (Tile.PackageName == InPackageName)
 		{
-			return const_cast<FWorldCompositionTile*>(&Tile);
+			return TileIdx;
 		}
 		
 		// Check LOD names
@@ -355,12 +357,25 @@ FWorldCompositionTile* UWorldComposition::FindTileByName(const FName& InPackageN
 		{
 			if (LODPackageName == InPackageName)
 			{
-				return const_cast<FWorldCompositionTile*>(&Tile);
+				return TileIdx;
 			}
 		}
 	}
-	
-	return nullptr;
+
+	return INDEX_NONE;
+}
+
+FWorldCompositionTile* UWorldComposition::FindTileByName(const FName& InPackageName) const
+{
+	int32 TileIdx = FindTileIndexByName(InPackageName);
+	if (TileIdx != INDEX_NONE)
+	{
+		return const_cast<FWorldCompositionTile*>(&Tiles[TileIdx]);
+	}
+	else
+	{
+		return nullptr;
+	}
 }
 
 #if WITH_EDITOR
@@ -480,6 +495,15 @@ void UWorldComposition::GetDistanceVisibleLevels(
 	const FVector& InLocation, 
 	TArray<FDistanceVisibleLevel>& OutVisibleLevels,
 	TArray<FDistanceVisibleLevel>& OutHiddenLevels) const
+{
+	GetDistanceVisibleLevels(&InLocation, 1, OutVisibleLevels, OutHiddenLevels);
+}
+
+void UWorldComposition::GetDistanceVisibleLevels(
+	const FVector* InLocations,
+	int32 NumLocations,
+	TArray<FDistanceVisibleLevel>& OutVisibleLevels,
+	TArray<FDistanceVisibleLevel>& OutHiddenLevels) const
 {
 	const UWorld* OwningWorld = GetWorld();
 
@@ -490,7 +514,7 @@ void UWorldComposition::GetDistanceVisibleLevels(
 		const FWorldCompositionTile& Tile = Tiles[TileIdx];
 		
 		// Skip non distance based levels
-		if (!Tile.Info.Layer.DistanceStreamingEnabled)
+		if (!IsDistanceDependentLevel(TileIdx))
 		{
 			continue;
 		}
@@ -526,19 +550,27 @@ void UWorldComposition::GetDistanceVisibleLevels(
 			LevelBounds.Max.Z = +WORLD_MAX;
 				
 			int32 NumAvailableLOD = FMath::Min(Tile.Info.LODList.Num(), Tile.LODPackageNames.Num());
-			// Find highest visible LOD entry
+			// Find LOD
 			// INDEX_NONE for original non-LOD level
 			for (int32 LODIdx = INDEX_NONE; LODIdx < NumAvailableLOD; ++LODIdx)
 			{
-				int32 TileStreamingDistance = Tile.Info.GetStreamingDistance(LODIdx);
-				FSphere QuerySphere(InLocation, TileStreamingDistance);
-		
-				if (FMath::SphereAABBIntersection(QuerySphere, LevelBounds))
+				if (bIsVisible && LODIdx > VisibleLevel.LODIndex)
 				{
-					VisibleLevel.LODIndex = LODIdx;
-					bIsVisible = true;
+					// no point to loop more, we have visible tile with best possible LOD
 					break;
 				}
+				
+				int32 TileStreamingDistance = Tile.Info.GetStreamingDistance(LODIdx);
+				for (int32 LocationIdx = 0; LocationIdx < NumLocations; ++LocationIdx)
+				{
+					FSphere QuerySphere(InLocations[LocationIdx], TileStreamingDistance);
+					if (FMath::SphereAABBIntersection(QuerySphere, LevelBounds))
+					{
+						VisibleLevel.LODIndex = LODIdx;
+						bIsVisible = true;
+						break;
+					}
+				}
 			}
 		}
 
@@ -555,23 +587,64 @@ void UWorldComposition::GetDistanceVisibleLevels(
 
 void UWorldComposition::UpdateStreamingState(const FVector& InLocation)
 {
-	// Get the list of visible and hidden levels from current view point
-	TArray<FDistanceVisibleLevel> DistanceVisibleLevels;
-	TArray<FDistanceVisibleLevel> DistanceHiddenLevels;
-	GetDistanceVisibleLevels(InLocation, DistanceVisibleLevels, DistanceHiddenLevels);
+	UpdateStreamingState(&InLocation, 1);
+}
 
+void UWorldComposition::UpdateStreamingState(const FVector* InLocations, int32 Num)
+{
 	UWorld* OwningWorld = GetWorld();
 
+	// Get the list of visible and hidden levels from current view point
+	TArray<FDistanceVisibleLevel> DistanceVisibleLevels;
+	TArray<FDistanceVisibleLevel> DistanceHiddenLevels;
+	GetDistanceVisibleLevels(InLocations, Num, DistanceVisibleLevels, DistanceHiddenLevels);
+	
+	// Dedicated server always blocks on load
+	bool bShouldBlock = (OwningWorld->GetNetMode() == NM_DedicatedServer);
+	
 	// Set distance hidden levels to unload
 	for (const auto& Level : DistanceHiddenLevels)
 	{
-		CommitTileStreamingState(OwningWorld, Level.TileIdx, false, false, Level.LODIndex);
+		CommitTileStreamingState(OwningWorld, Level.TileIdx, false, false, bShouldBlock, Level.LODIndex);
 	}
 
 	// Set distance visible levels to load
 	for (const auto& Level : DistanceVisibleLevels)
 	{
-		CommitTileStreamingState(OwningWorld, Level.TileIdx, true, true, Level.LODIndex);
+		CommitTileStreamingState(OwningWorld, Level.TileIdx, true, true, bShouldBlock, Level.LODIndex);
+	}
+}
+
+void UWorldComposition::UpdateStreamingStateCinematic(const FVector* InLocations, int32 Num)
+{
+	if (!bLoadAllTilesDuringCinematic)
+	{
+		UpdateStreamingState(InLocations, Num);
+		return;
+	}
+
+	// Cinematic always blocks on load
+	bool bShouldBlock = true;
+	bool bStreamingStateChanged = false;
+		
+	// All tiles should be loaded and visible regardless of distance
+	for (int32 TileIdx = 0; TileIdx < Tiles.Num(); TileIdx++)
+	{
+		FWorldCompositionTile& Tile = Tiles[TileIdx];
+		// Skip non distance based levels
+		if (!IsDistanceDependentLevel(TileIdx))
+		{
+			continue;
+		}
+		// Reset streaming state cooldown to ensure that new state will be committed
+		Tile.StreamingLevelStateChangeTime = 0.0;
+		//
+		bStreamingStateChanged|= CommitTileStreamingState(GetWorld(), TileIdx, true, true, bShouldBlock, INDEX_NONE);
+	}
+
+	if (bStreamingStateChanged)
+	{
+		GetWorld()->FlushLevelStreaming(EFlushLevelStreamingType::Full);
 	}
 }
 
@@ -594,9 +667,10 @@ void UWorldComposition::UpdateStreamingState()
 	}
 
 	// calculate centroid location using local players views
-	int32 NumViews = 0;
+	bool bCinematic = false;
 	FVector CentroidLocation = FVector::ZeroVector;
-	
+	TArray<FVector, TInlineAllocator<16>> Locations;
+		
 	for (int32 PlayerIndex = 0; PlayerIndex < NumPlayers; ++PlayerIndex)
 	{
 		ULocalPlayer* Player = GEngine->GetGamePlayer(PlayWorld, PlayerIndex);
@@ -605,21 +679,29 @@ void UWorldComposition::UpdateStreamingState()
 			FVector ViewLocation;
 			FRotator ViewRotation;
 			Player->PlayerController->GetPlayerViewPoint(ViewLocation, ViewRotation);
+			Locations.Add(ViewLocation);
 			CentroidLocation+= ViewLocation;
-			NumViews++;
+			bCinematic|= Player->PlayerController->bCinematicMode;
 		}
 	}
 
 	// In case there is no valid views don't bother updating level streaming state
-	if (NumViews > 0)
+	if (Locations.Num())
 	{
-		CentroidLocation/= NumViews;
+		CentroidLocation/= Locations.Num();
 		if (PlayWorld->GetWorldSettings()->bEnableWorldOriginRebasing)
 		{
 			EvaluateWorldOriginLocation(CentroidLocation);
 		}
 		
-		UpdateStreamingState(CentroidLocation);
+		if (bCinematic)
+		{
+			UpdateStreamingStateCinematic(Locations.GetData(), Locations.Num());
+		}
+		else
+		{
+			UpdateStreamingState(Locations.GetData(), Locations.Num());
+		}
 	}
 }
 
@@ -684,17 +766,29 @@ void UWorldComposition::EvaluateWorldOriginLocation(const FVector& ViewLocation)
 	}
 }
 
+bool UWorldComposition::IsDistanceDependentLevel(int32 TileIdx) const
+{
+	if (TileIdx != INDEX_NONE)
+	{
+		return (Tiles[TileIdx].Info.Layer.DistanceStreamingEnabled && !TilesStreaming[TileIdx]->bDisableDistanceStreaming);
+	}
+	else
+	{
+		return false;
+	}
+}
+
 bool UWorldComposition::IsDistanceDependentLevel(FName PackageName) const
 {
-	FWorldCompositionTile* Tile = FindTileByName(PackageName);
-	return (Tile && Tile->Info.Layer.DistanceStreamingEnabled);
+	int32 TileIdx = FindTileIndexByName(PackageName);
+	return IsDistanceDependentLevel(TileIdx);
 }
 
-void UWorldComposition::CommitTileStreamingState(UWorld* PersistenWorld, int32 TileIdx, bool bShouldBeLoaded, bool bShouldBeVisible, int32 LODIdx)
+bool UWorldComposition::CommitTileStreamingState(UWorld* PersistenWorld, int32 TileIdx, bool bShouldBeLoaded, bool bShouldBeVisible, bool bShouldBlock, int32 LODIdx)
 {
 	if (!Tiles.IsValidIndex(TileIdx))
 	{
-		return;
+		return false;
 	}
 
 	FWorldCompositionTile& Tile = Tiles[TileIdx];
@@ -703,9 +797,10 @@ void UWorldComposition::CommitTileStreamingState(UWorld* PersistenWorld, int32 T
 	// Quit early in case state is not going to be changed
 	if (StreamingLevel->bShouldBeLoaded == bShouldBeLoaded &&
 		StreamingLevel->bShouldBeVisible == bShouldBeVisible &&
+		StreamingLevel->bShouldBlockOnLoad == bShouldBlock &&
 		StreamingLevel->LevelLODIndex == LODIdx)
 	{
-		return;
+		return false;
 	}
 
 	// Quit early in case we have cooldown on streaming state changes
@@ -716,7 +811,7 @@ void UWorldComposition::CommitTileStreamingState(UWorld* PersistenWorld, int32 T
 		const double TimePassed = CurrentTime - Tile.StreamingLevelStateChangeTime;
 		if (TimePassed < TilesStreamingTimeThreshold)
 		{
-			return;
+			return false;
 		}
 
 		// Save current time as state change time for this tile
@@ -724,15 +819,11 @@ void UWorldComposition::CommitTileStreamingState(UWorld* PersistenWorld, int32 T
 	}
 
 	// Commit new state
+	StreamingLevel->bShouldBlockOnLoad	= bShouldBlock;
 	StreamingLevel->bShouldBeLoaded		= bShouldBeLoaded;
 	StreamingLevel->bShouldBeVisible	= bShouldBeVisible;
 	StreamingLevel->LevelLODIndex		= LODIdx;
-
-	// dedicated server always block on load
-	if (PersistenWorld->GetNetMode() == NM_DedicatedServer && bShouldBeLoaded) 
-	{
-		StreamingLevel->bShouldBlockOnLoad = true;
-	}
+	return true;
 }
 
 void UWorldComposition::OnLevelAddedToWorld(ULevel* InLevel)
diff --git a/Engine/Source/Runtime/Engine/Public/AI/NavDataGenerator.h b/Engine/Source/Runtime/Engine/Public/AI/NavDataGenerator.h
--- a/Engine/Source/Runtime/Engine/Public/AI/NavDataGenerator.h
+++ b/Engine/Source/Runtime/Engine/Public/AI/NavDataGenerator.h
@@ -2,9 +2,11 @@
 
 #pragma once 
 
+#include "Templates/SharedPointer.h"
+
 struct FNavigationDirtyArea;
 
-class ENGINE_API FNavDataGenerator
+class ENGINE_API FNavDataGenerator : public TSharedFromThis<FNavDataGenerator>
 {
 public:
 	virtual ~FNavDataGenerator() {}
diff --git a/Engine/Source/Runtime/Engine/Public/AI/Navigation/PImplRecastNavMesh.h b/Engine/Source/Runtime/Engine/Public/AI/Navigation/PImplRecastNavMesh.h
--- a/Engine/Source/Runtime/Engine/Public/AI/Navigation/PImplRecastNavMesh.h
+++ b/Engine/Source/Runtime/Engine/Public/AI/Navigation/PImplRecastNavMesh.h
@@ -16,6 +16,7 @@
 #define RECAST_VERY_SMALL_AGENT_RADIUS 0.0f
 
 class ARecastNavMesh;
+class FRecastNavMeshGenerator;
 class dtNavMesh;
 struct dtMeshTile;
 
@@ -76,7 +77,7 @@ public:
 
 	/** Debug rendering. */
 	void GetDebugGeometry(FRecastDebugGeometry& OutGeometry, int32 TileIndex = INDEX_NONE) const;
-
+	
 	/** Returns bounding box for the whole navmesh. */
 	FBox GetNavMeshBounds() const;
 
@@ -221,7 +222,7 @@ public:
 	dtNavMesh* DetourNavMesh;
 
 	/** Compressed layers data, can be reused for tiles generation */
-	TMap<FIntPoint, TArray<FNavMeshTileData>> CompressedTileCacheLayers;
+	TMap<FIntPoint, TArray<FNavMeshTileData> > CompressedTileCacheLayers;
 
 	/** query used for searching data on game thread */
 	mutable dtNavMeshQuery SharedNavQuery;
@@ -246,10 +247,13 @@ public:
 		const FVector& RecastStart, FVector& RecastEnd,
 		dtQueryResult& PathResult) const;
 
-	void GetDebugPolyEdges(const dtMeshTile* Tile, bool bInternalEdges, bool bNavMeshEdges, TArray<FVector>& InternalEdgeVerts, TArray<FVector>& NavMeshEdgeVerts) const;
+	void GetDebugPolyEdges(const dtMeshTile& Tile, bool bInternalEdges, bool bNavMeshEdges, TArray<FVector>& InternalEdgeVerts, TArray<FVector>& NavMeshEdgeVerts) const;
 
 	/** workhorse function finding portal edges between corridor polys */
 	void GetEdgesForPathCorridorImpl(const TArray<NavNodeRef>* PathCorridor, TArray<FNavigationPortalEdge>* PathCorridorEdges, const dtNavMeshQuery& NavQuery) const;
+
+protected:
+	int32 GetTilesDebugGeometry(const FRecastNavMeshGenerator* Generator, const dtMeshTile& Tile, int32 VertBase, FRecastDebugGeometry& OutGeometry, int32 TileIdx = INDEX_NONE) const;
 };
 
 #endif	// WITH_RECAST
diff --git a/Engine/Source/Runtime/Engine/Public/AI/Navigation/RecastNavMeshGenerator.h b/Engine/Source/Runtime/Engine/Public/AI/Navigation/RecastNavMeshGenerator.h
--- a/Engine/Source/Runtime/Engine/Public/AI/Navigation/RecastNavMeshGenerator.h
+++ b/Engine/Source/Runtime/Engine/Public/AI/Navigation/RecastNavMeshGenerator.h
@@ -129,13 +129,12 @@ struct FRecastAreaNavModifierElement
 /**
  * Class handling generation of a single tile, caching data that can speed up subsequent tile generations
  */
-class ENGINE_API FRecastTileGenerator 
-	: public FNoncopyable
+class ENGINE_API FRecastTileGenerator : public FNoncopyable, public FGCObject
 {
 	friend FRecastNavMeshGenerator;
 
 public:
-	FRecastTileGenerator(const FRecastNavMeshGenerator& ParentGenerator, const FIntPoint& Location);
+	FRecastTileGenerator(FRecastNavMeshGenerator& ParentGenerator, const FIntPoint& Location);
 	virtual ~FRecastTileGenerator();
 		
 	void DoWork();
@@ -149,13 +148,20 @@ public:
 	/** Whether tile task has anything to build */
 	bool HasDataToBuild() const;
 
-	TArray<FNavMeshTileData> GetCompressedLayers() const { return CompressedLayers; }
-	TArray<FNavMeshTileData> GetNavigationData() const { return NavigationData; }
+	const TArray<FNavMeshTileData>& GetCompressedLayers() const { return CompressedLayers; }
+protected:
+	// to be used solely by FRecastNavMeshGenerator
+	TArray<FNavMeshTileData>& GetNavigationData() { return NavigationData; }
 	
+public:
 	uint32 GetUsedMemCount() const;
 
 	// Memory amount used to construct generator 
 	uint32 UsedMemoryOnStartup;
+
+	// FGCObject begin
+	virtual void AddReferencedObjects(FReferenceCollector& Collector);
+	// FGCObject end
 		
 protected:
 	/** Does the actual tile generations. 
@@ -168,6 +174,8 @@ protected:
 	void Setup(const FRecastNavMeshGenerator& ParentGenerator, const TArray<FBox>& DirtyAreas);
 	
 	void GatherGeometry(const FRecastNavMeshGenerator& ParentGenerator, bool bGeometryChanged);
+	void PrepareGeometrySources(const FRecastNavMeshGenerator& ParentGenerator, bool bGeometryChanged);
+	void DoAsyncGeometryGathering();
 
 	/** builds CompressedLayers array (geometry + modifiers) */
 	virtual bool GenerateCompressedLayers(FNavMeshBuildContext& BuildContext);
@@ -191,10 +199,13 @@ protected:
 	bool HasVoxelCache(const TNavStatArray<uint8>& RawVoxelCache, rcSpanCache*& CachedVoxels, int32& NumCachedVoxels) const;
 	void AddVoxelCache(TNavStatArray<uint8>& RawVoxelCache, const rcSpanCache* CachedVoxels, const int32 NumCachedVoxels) const;
 
+	void DumpAsyncData();
+
 protected:
 	uint32 bSucceeded : 1;
 	uint32 bRegenerateCompressedLayers : 1;
 	uint32 bFullyEncapsulatedByInclusionBounds : 1;
+	uint32 bUpdateGeometry : 1;
 	
 	int32 TileX;
 	int32 TileY;
@@ -224,6 +235,12 @@ protected:
 	TArray<FRecastAreaNavModifierElement> Modifiers;
 	// navigation links
 	TArray<FSimpleLinkNavModifier> OffmeshLinks;
+
+	TWeakPtr<FNavDataGenerator> ParentGeneratorWeakPtr;
+
+	TNavStatArray<TSharedRef<FNavigationRelevantData, ESPMode::ThreadSafe> > NavigationRelevantData;
+	TSharedPtr<FNavigationOctree, ESPMode::ThreadSafe> NavOctree; 
+	FNavDataConfig NavDataConfig;
 };
 
 struct ENGINE_API FRecastTileGeneratorWrapper : public FNonAbandonableTask
@@ -356,8 +373,7 @@ public:
 	virtual int32 GetNumRemaningBuildTasks() const override;
 	virtual int32 GetNumRunningBuildTasks() const override;
 
-	/** Checks if a given tile is being build or has just finished building
-	 *	If TileIndex is out of bounds function returns false (i.e. not being built) */
+	/** Checks if a given tile is being build or has just finished building */
 	bool IsTileChanged(int32 TileIdx) const;
 		
 	FORCEINLINE uint32 GetVersion() const { return Version; }
@@ -385,6 +401,8 @@ public:
 
 	bool HasDirtyTiles() const;
 
+	bool GatherGeometryOnGameThread() const;
+
 	FBox GrowBoundingBox(const FBox& BBox, bool bIncludeAgentHeight) const;
 	
 #if !(UE_BUILD_SHIPPING || UE_BUILD_TEST)
@@ -419,11 +437,12 @@ protected:
 	/** Marks grid tiles affected by specified areas as dirty */
 	void MarkDirtyTiles(const TArray<FNavigationDirtyArea>& DirtyAreas);
 	
-	/** Processes pending tile generattion tasks */
+	/** Processes pending tile generation tasks */
 	TArray<uint32> ProcessTileTasks(const int32 NumTasksToSubmit);
 
+public:
 	/** Adds generated tiles to NavMesh, replacing old ones */
-	TArray<uint32> AddGeneratedTiles(const FRecastTileGenerator& TileGenerator);
+	TArray<uint32> AddGeneratedTiles(FRecastTileGenerator& TileGenerator);
 
 public:
 	/** Removes all tiles at specified grid location */
@@ -432,6 +451,13 @@ public:
 	void RemoveTiles(const TArray<FIntPoint>& Tiles);
 	void ReAddTiles(const TArray<FIntPoint>& Tiles);
 
+	bool IsBuildingRestrictedToActiveTiles() const { return bRestrictBuildingToActiveTiles; }
+
+	/** sets a limit to number of asynchronous tile generators running at one time
+	 *	@note if used at runtime will not result in killing tasks above limit count
+	 *	@mote function does not validate the input parameter - it's on caller */
+	void SetMaxTileGeneratorTasks(int32 NewLimit) { MaxTileGeneratorTasks = NewLimit; }
+
 protected:
 	bool IsInActiveSet(const FIntPoint& Tile) const;
 	void RestrictBuildingToActiveTiles(bool InRestrictBuildingToActiveTiles);
@@ -454,6 +480,7 @@ private:
 	FRecastBuildConfig Config;
 	
 	int32 NumActiveTiles;
+	/** the limit to number of asynchronous tile generators running at one time */
 	int32 MaxTileGeneratorTasks;
 	float AvgLayersPerTile;
 
diff --git a/Engine/Source/Runtime/Engine/Public/AI/NavigationModifier.h b/Engine/Source/Runtime/Engine/Public/AI/NavigationModifier.h
--- a/Engine/Source/Runtime/Engine/Public/AI/NavigationModifier.h
+++ b/Engine/Source/Runtime/Engine/Public/AI/NavigationModifier.h
@@ -250,6 +250,8 @@ struct ENGINE_API FCompositeNavModifier : public FNavigationModifier
 	FORCEINLINE bool HasAgentHeightAdjust() const { return bAdjustHeight; }
 	FORCEINLINE bool HasAreas() const { return Areas.Num() > 0; }
 
+	FORCEINLINE void ReserveForAdditionalAreas(int32 AdditionalElementsCount) { Areas.Reserve(Areas.Num() + AdditionalElementsCount); }
+
 	void MarkPotentialLinks() { bHasPotentialLinks = true; }
 
 	/** returns a copy of Modifier */
diff --git a/Engine/Source/Runtime/Engine/Public/AI/NavigationOctree.h b/Engine/Source/Runtime/Engine/Public/AI/NavigationOctree.h
--- a/Engine/Source/Runtime/Engine/Public/AI/NavigationOctree.h
+++ b/Engine/Source/Runtime/Engine/Public/AI/NavigationOctree.h
@@ -26,7 +26,7 @@ struct ENGINE_API FNavigationOctreeFilter
 
 // @todo consider optional structures that can contain a delegate instead of 
 // actual copy of collision data
-struct ENGINE_API FNavigationRelevantData
+struct ENGINE_API FNavigationRelevantData : public TSharedFromThis<FNavigationRelevantData, ESPMode::ThreadSafe>
 {
 	DECLARE_DELEGATE_RetVal_OneParam(bool, FFilterNavDataDelegate, const struct FNavDataConfig*);
 
@@ -39,56 +39,74 @@ struct ENGINE_API FNavigationRelevantData
 	/** bounds of geometry (unreal coords) */
 	FBox Bounds;
 
+	/** Gathers per instance data for navigation geometry in a specified area box */
+	FNavDataPerInstanceTransformDelegate NavDataPerInstanceTransformDelegate;
+
 	/** called to check if hosted geometry should be used for given FNavDataConfig. If not set then "true" is assumed. */
 	FFilterNavDataDelegate ShouldUseGeometryDelegate;
 
 	/** additional modifiers: areas and external links */
 	FCompositeNavModifier Modifiers;
 
+	/** UObject these data represents */
+	const TWeakObjectPtr<UObject> SourceObject;
+
 	/** get set to true when lazy navigation exporting is enabled and this navigation data has "potential" of 
 	 *	containing geometry data. First access will result in gathering the data and setting this flag back to false.
 	 *	Mind that this flag can go back to 'true' if related data gets cleared out. */
 	uint32 bPendingLazyGeometryGathering : 1;
+	uint32 bPendingLazyModifiersGathering : 1;
 
-	// Gathers per instance data for navigation geometry in a specified area box
-	FNavDataPerInstanceTransformDelegate NavDataPerInstanceTransformDelegate;
-
-	FNavigationRelevantData() : bPendingLazyGeometryGathering(false)
+	uint32 bSupportsGatheringGeometrySlices : 1;
+	
+	FNavigationRelevantData(UObject& Source) 
+		: SourceObject(&Source)
+		, bPendingLazyGeometryGathering(false)
+		, bPendingLazyModifiersGathering(false)
 	{}
 
 	FORCEINLINE bool HasGeometry() const { return VoxelData.Num() || CollisionData.Num(); }
 	FORCEINLINE bool HasModifiers() const { return !Modifiers.IsEmpty(); }
 	FORCEINLINE bool IsPendingLazyGeometryGathering() const { return bPendingLazyGeometryGathering; }
+	FORCEINLINE bool IsPendingLazyModifiersGathering() const { return bPendingLazyModifiersGathering; }
+	FORCEINLINE bool SupportsGatheringGeometrySlices() const { return bSupportsGatheringGeometrySlices; }
 	FORCEINLINE bool IsEmpty() const { return !HasGeometry() && !HasModifiers(); }
 	FORCEINLINE uint32 GetAllocatedSize() const { return CollisionData.GetAllocatedSize() + VoxelData.GetAllocatedSize() + Modifiers.GetAllocatedSize(); }
 	FORCEINLINE uint32 GetGeometryAllocatedSize() const { return CollisionData.GetAllocatedSize() + VoxelData.GetAllocatedSize(); }
 	FORCEINLINE int32 GetDirtyFlag() const
 	{
 		return ((HasGeometry() || IsPendingLazyGeometryGathering()) ? ENavigationDirtyFlag::Geometry : 0) |
-			(HasModifiers() ? ENavigationDirtyFlag::DynamicModifier : 0) |
+			((HasModifiers() || IsPendingLazyModifiersGathering())? ENavigationDirtyFlag::DynamicModifier : 0) |
 			(Modifiers.HasAgentHeightAdjust() ? ENavigationDirtyFlag::UseAgentHeight : 0);
 	}
 	
 	bool HasPerInstanceTransforms() const;
 	bool IsMatchingFilter(const FNavigationOctreeFilter& Filter) const;
 	void Shrink();
+
+	FORCEINLINE UObject* GetOwner() const { return SourceObject.Get(); }
 };
 
 struct ENGINE_API FNavigationOctreeElement
 {
 	FBoxSphereBounds Bounds;
-	TWeakObjectPtr<UObject> Owner;
-	FNavigationRelevantData Data;
+	TSharedRef<FNavigationRelevantData, ESPMode::ThreadSafe> Data;
+
+	FNavigationOctreeElement(UObject& SourceObject)
+		: Data(new FNavigationRelevantData(SourceObject))
+	{
+
+	}
 
 	FORCEINLINE bool IsEmpty() const
 	{
 		const FBox BBox = Bounds.GetBox();
-		return Data.IsEmpty() && (BBox.IsValid == 0 || BBox.GetSize().IsNearlyZero());
+		return Data->IsEmpty() && (BBox.IsValid == 0 || BBox.GetSize().IsNearlyZero());
 	}
 
 	FORCEINLINE bool IsMatchingFilter(const FNavigationOctreeFilter& Filter) const
 	{
-		return Data.IsMatchingFilter(Filter);
+		return Data->IsMatchingFilter(Filter);
 	}
 
 	/** 
@@ -98,23 +116,25 @@ struct ENGINE_API FNavigationOctreeElement
 	 */
 	FORCEINLINE FCompositeNavModifier GetModifierForAgent(const struct FNavAgentProperties* NavAgent = NULL) const 
 	{ 
-		return Data.Modifiers.HasMetaAreas() ? Data.Modifiers.GetInstantiatedMetaModifier(NavAgent, Owner) : Data.Modifiers;
+		return Data->Modifiers.HasMetaAreas() ? Data->Modifiers.GetInstantiatedMetaModifier(NavAgent, Data->SourceObject) : Data->Modifiers;
 	}
 
 	FORCEINLINE bool ShouldUseGeometry(const FNavDataConfig& NavConfig) const
 	{ 
-		return !Data.ShouldUseGeometryDelegate.IsBound() || Data.ShouldUseGeometryDelegate.Execute(&NavConfig);
+		return !Data->ShouldUseGeometryDelegate.IsBound() || Data->ShouldUseGeometryDelegate.Execute(&NavConfig);
 	}
 
 	FORCEINLINE int32 GetAllocatedSize() const
 	{
-		return Data.GetAllocatedSize();
+		return Data->GetAllocatedSize();
 	}
 
 	FORCEINLINE void Shrink()
 	{
-		Data.Shrink();
+		Data->Shrink();
 	}
+
+	FORCEINLINE UObject* GetOwner() const { return Data->SourceObject.Get(); }
 };
 
 struct FNavigationOctreeSemantics
@@ -133,7 +153,7 @@ struct FNavigationOctreeSemantics
 
 	FORCEINLINE static bool AreElementsEqual(const FNavigationOctreeElement& A, const FNavigationOctreeElement& B)
 	{
-		return A.Owner == B.Owner;
+		return A.Data->SourceObject == B.Data->SourceObject;
 	}
 
 #if NAVSYS_DEBUG
@@ -142,7 +162,7 @@ struct FNavigationOctreeSemantics
 	static void SetElementId(const FNavigationOctreeElement& Element, FOctreeElementId Id);
 };
 
-class FNavigationOctree : public TOctree<FNavigationOctreeElement, FNavigationOctreeSemantics>
+class FNavigationOctree : public TOctree<FNavigationOctreeElement, FNavigationOctreeSemantics>, public TSharedFromThis<FNavigationOctree, ESPMode::ThreadSafe>
 {
 public:
 	DECLARE_DELEGATE_TwoParams(FNavigableGeometryComponentExportDelegate, UActorComponent*, FNavigationRelevantData&);
@@ -177,13 +197,14 @@ public:
 		return bGatherGeometry ? StoreNavGeometry : SkipNavGeometry;
 	}
 
-	void SetLazyGeometryGatheringEnabled(bool bInEnable);
-
+	void SetDataGatheringMode(ENavDataGatheringModeConfig Mode);
+	
 	// @hack! TO BE FIXED
-	void DemandGatheringGeometry(const FNavigationOctreeElement& Element) const;
+	void DemandLazyDataGathering(const FNavigationOctreeElement& Element);
+	void DemandLazyDataGathering(FNavigationRelevantData& ElementData);
 
 protected:
-	uint32 bPerformGeometryGatheringLazily : 1;
+	ENavDataGatheringMode DefaultGeometryGatheringMode;
 	uint32 bGatherGeometry : 1;
 	uint32 NodesMemory;
 };
diff --git a/Engine/Source/Runtime/Engine/Public/AI/NavigationSystemHelpers.h b/Engine/Source/Runtime/Engine/Public/AI/NavigationSystemHelpers.h
--- a/Engine/Source/Runtime/Engine/Public/AI/NavigationSystemHelpers.h
+++ b/Engine/Source/Runtime/Engine/Public/AI/NavigationSystemHelpers.h
@@ -29,6 +29,7 @@ struct FNavigableGeometryExport
 	virtual void ExportPxConvexMesh(physx::PxConvexMesh const * const ConvexMesh, const FTransform& LocalToWorld) = 0;
 	virtual void ExportPxHeightField(physx::PxHeightField const * const HeightField, const FTransform& LocalToWorld) = 0;
 #endif // WITH_PHYSX
+	virtual void ExportHeightFieldSlice(const FNavHeightfieldSamples& PrefetchedHeightfieldSamples, const int32 NumRows, const int32 NumCols, const FTransform& LocalToWorld, const FBox& SliceBox) = 0;
 	virtual void ExportRigidBodySetup(UBodySetup& BodySetup, const FTransform& LocalToWorld) = 0;
 	virtual void ExportCustomMesh(const FVector* VertexBuffer, int32 NumVerts, const int32* IndexBuffer, int32 NumIndices, const FTransform& LocalToWorld) = 0;
 	
diff --git a/Engine/Source/Runtime/Engine/Public/AnimTree.h b/Engine/Source/Runtime/Engine/Public/AnimTree.h
--- a/Engine/Source/Runtime/Engine/Public/AnimTree.h
+++ b/Engine/Source/Runtime/Engine/Public/AnimTree.h
@@ -21,6 +21,7 @@ DECLARE_CYCLE_STAT_EXTERN(TEXT("Perform Anim Evaluation"), STAT_PerformAnimEvalu
 DECLARE_CYCLE_STAT_EXTERN(TEXT("Post Anim Evaluation"), STAT_PostAnimEvaluation, STATGROUP_Anim, );
 DECLARE_CYCLE_STAT_EXTERN(TEXT("GraphTime"), STAT_AnimGraphEvaluate, STATGROUP_Anim, );
 DECLARE_CYCLE_STAT_EXTERN(TEXT("Eval Pose"), STAT_AnimNativeEvaluatePoses, STATGROUP_Anim, );
+DECLARE_CYCLE_STAT_EXTERN(TEXT("Trigger Notifies"), STAT_AnimTriggerAnimNotifies, STATGROUP_Anim, );
 DECLARE_CYCLE_STAT_EXTERN(TEXT("Anim Decompression"), STAT_GetAnimationPose, STATGROUP_Anim, );
 DECLARE_CYCLE_STAT_EXTERN(TEXT("Blend Poses"), STAT_AnimNativeBlendPoses, STATGROUP_Anim, );
 DECLARE_CYCLE_STAT_EXTERN(TEXT("Copy Pose"), STAT_AnimNativeCopyPoses, STATGROUP_Anim, );
diff --git a/Engine/Source/Runtime/Engine/Public/ComponentRecreateRenderStateContext.h b/Engine/Source/Runtime/Engine/Public/ComponentRecreateRenderStateContext.h
new file mode 100644
--- /dev/null
+++ b/Engine/Source/Runtime/Engine/Public/ComponentRecreateRenderStateContext.h
@@ -0,0 +1,55 @@
+// Copyright 1998-2015 Epic Games, Inc. All Rights Reserved.
+
+#pragma once
+
+#include "ContentStreaming.h"
+#include "AI/Navigation/NavigationSystem.h"
+
+/** Destroys render state for a component and then recreates it when this object is destroyed */
+class FComponentRecreateRenderStateContext
+{
+private:
+	/** Pointer to component we are recreating render state for */
+	UActorComponent* Component;
+public:
+	FComponentRecreateRenderStateContext(UActorComponent* InComponent)
+	{
+		check(InComponent);
+		checkf(!InComponent->HasAnyFlags(RF_Unreachable), TEXT("%s"), *InComponent->GetFullName());
+
+		if (InComponent->IsRegistered() && InComponent->IsRenderStateCreated())
+		{
+			InComponent->DestroyRenderState_Concurrent();
+			Component = InComponent;
+		}
+		else
+		{
+			Component = nullptr;
+		}
+	}
+
+	~FComponentRecreateRenderStateContext()
+	{
+		if (Component != nullptr)
+		{
+			Component->CreateRenderState_Concurrent();
+		}
+	}
+};
+
+/** Destroys render states for all components and then recreates them when this object is destroyed */
+class FGlobalComponentRecreateRenderStateContext
+{
+public:
+	/** 
+	* Initialization constructor. 
+	*/
+	ENGINE_API FGlobalComponentRecreateRenderStateContext();
+
+	/** Destructor */
+	ENGINE_API ~FGlobalComponentRecreateRenderStateContext();
+
+private:
+	/** The recreate contexts for the individual components. */
+	TIndirectArray<FComponentRecreateRenderStateContext> ComponentContexts;
+};
diff --git a/Engine/Source/Runtime/Engine/Public/ContentStreaming.h b/Engine/Source/Runtime/Engine/Public/ContentStreaming.h
--- a/Engine/Source/Runtime/Engine/Public/ContentStreaming.h
+++ b/Engine/Source/Runtime/Engine/Public/ContentStreaming.h
@@ -464,7 +464,7 @@ struct ITextureStreamingManager : public IStreamingManager
 	*	@param RequiredMemorySize	- Required minimum available texture memory
 	*	@return						- Whether it succeeded or not
 	**/
-	virtual bool StreamOutTextureData(int32 RequiredMemorySize) = 0;
+	virtual bool StreamOutTextureData(int64 RequiredMemorySize) = 0;
 
 	/** Adds a new texture to the streaming manager. */
 	virtual void AddStreamingTexture(UTexture2D* Texture) = 0;
@@ -724,14 +724,19 @@ struct FSpawnedTextureInstance
  */
 struct FStreamMemoryTracker
 {
+#if PLATFORM_64BITS
+	typedef int64 TSize;
+#else
+	typedef int32 TSize;
+#endif
 	/** Stream-in memory that hasn't been allocated yet. */
-	volatile int32 PendingStreamIn;
+	volatile TSize PendingStreamIn;
 	/** Temp memory that hasn't been allocated yet. */
-	volatile int32 PendingTempMemory;
+	volatile TSize PendingTempMemory;
 	/** Stream-out memory that is allocated but hasn't been freed yet. */
-	volatile int32 CurrentStreamOut;
+	volatile TSize CurrentStreamOut;
 	/** Temp memory that is allocated, but not freed yet. */
-	volatile int32 CurrentTempMemory;
+	volatile TSize CurrentTempMemory;
 
 	//@TODO: Possibly track canceling early (on the gamethread), not just in finalize.
 	//@TODO: Add support for pre-planned in-place reallocs
@@ -760,13 +765,13 @@ struct FStreamMemoryTracker
 	void RenderThread_Finalize( const UTexture2D& Texture, bool bSuccessful );
 
 	/** Calculate how much texture memory is currently available for streaming. */
-	int32 CalcAvailableNow( int32 TotalFreeMemory, int32 MemoryMargin );
+	TSize CalcAvailableNow( TSize TotalFreeMemory, TSize MemoryMargin );
 
 	/** Calculate how much texture memory will available later for streaming. */
-	int32 CalcAvailableLater( int32 TotalFreeMemory, int32 MemoryMargin );
+	TSize CalcAvailableLater( TSize TotalFreeMemory, TSize MemoryMargin );
 
 	/** Calculate how much texture memory is currently being used for temporary texture data during streaming. */
-	int32 CalcTempMemory();
+	TSize CalcTempMemory();
 };
 
 /*-----------------------------------------------------------------------------
@@ -890,7 +895,7 @@ struct FStreamingManagerTexture : public ITextureStreamingManager
 	 *	@param RequiredMemorySize	- Required minimum available texture memory
 	 *	@return						- Whether it succeeded or not
 	 **/
-	virtual bool StreamOutTextureData( int32 RequiredMemorySize );
+	virtual bool StreamOutTextureData( int64 RequiredMemorySize );
 
 	/**
 	 * Allows the streaming manager to process exec commands.
@@ -1125,7 +1130,7 @@ protected:
 		 * @param StreamingRequests		[in/out] Indices to textures that are going to be streamed this frame
 		 * @return						The last priority index considered (may be sooner that StopIndex)
 		 */
-		int32 StreamoutTextures( FStreamoutLogic StreamoutLogic, int32& AvailableLater, int32& TempMemoryUsed, int32 StartIndex, int32 StopIndex, int32& LowPrioIndex, const TArray<FTexturePriority>& PrioritizedTextures, FStreamingRequests& StreamingRequests );
+		int32 StreamoutTextures( FStreamoutLogic StreamoutLogic, int64& AvailableLater, int64& TempMemoryUsed, int32 StartIndex, int32 StopIndex, int32& LowPrioIndex, const TArray<FTexturePriority>& PrioritizedTextures, FStreamingRequests& StreamingRequests );
 
 		/**
 		 * Stream textures in/out, when no texture pool with limited size is used by the platform.
@@ -1134,7 +1139,7 @@ protected:
 		 * @param PrioritizedTextures	Array of prioritized textures to process
 		 * @param TempMemoryUsed		Current amount of temporary memory used by the streaming system, in bytes
 		 */
-		void StreamTexturesUnlimited( FStreamingContext& Context, const TArray<FTexturePriority>& PrioritizedTextures, int32 TempMemoryUsed );
+		void StreamTexturesUnlimited( FStreamingContext& Context, const TArray<FTexturePriority>& PrioritizedTextures, int64 TempMemoryUsed );
 
 		/** Thread-safe helper struct for per-level information. */
 		struct FThreadLevelData
@@ -1263,7 +1268,7 @@ protected:
 	int32					NumTextureProcessingStages;
 
 	/** Maximum amount of temp memory used for streaming, at any given time. */
-	int32					MaxTempMemoryUsed;
+	int64					MaxTempMemoryUsed;
 
 	/** Whether to support texture instance streaming for dynamic (movable/spawned) objects. */
 	bool					bUseDynamicStreaming;
@@ -1274,10 +1279,10 @@ protected:
 	TArray<FStreamingHandlerTextureBase*> TextureStreamingHandlers;
 
 	/** Amount of memory to leave free in the texture pool. */
-	int32					MemoryMargin;
+	int64					MemoryMargin;
 
 	/** Minimum number of bytes to evict when we need to stream out textures because of a failed allocation. */
-	int32					MinEvictSize;
+	int64					MinEvictSize;
 
 	/** If set, UpdateResourceStreaming() will only process this texture. */
 	UTexture2D*				IndividualStreamingTexture;
@@ -1297,25 +1302,25 @@ protected:
 	/** Number of intermediate textures in flight. */
 	uint32 NumIntermediateTextures;
 	/** Size of all streaming testures. */
-	uint32 TotalStreamingTexturesSize;
+	uint64 TotalStreamingTexturesSize;
 	/** Maximum size of all streaming textures. */
-	uint32 TotalStreamingTexturesMaxSize;
+	uint64 TotalStreamingTexturesMaxSize;
 	/** Total number of bytes in memory, for all streaming lightmap textures. */
-	uint32 TotalLightmapMemorySize;
+	uint64 TotalLightmapMemorySize;
 	/** Total number of bytes on disk, for all streaming lightmap textures. */
-	uint32 TotalLightmapDiskSize;
+	uint64 TotalLightmapDiskSize;
 	/** Number of mip count increase requests in flight. */
 	uint32 TotalMipCountIncreaseRequestsInFlight;
 	/** Total number of bytes in memory, if all textures were streamed perfectly with a 1.0 fudge factor. */
-	uint32 TotalOptimalWantedSize;
+	uint64 TotalOptimalWantedSize;
 	/** Total number of bytes using StaticTexture heuristics, currently in memory. */
-	uint32 TotalStaticTextureHeuristicSize;
+	uint64 TotalStaticTextureHeuristicSize;
 	/** Total number of bytes using DynmicTexture heuristics, currently in memory. */
-	uint32 TotalDynamicTextureHeuristicSize;
+	uint64 TotalDynamicTextureHeuristicSize;
 	/** Total number of bytes using LastRenderTime heuristics, currently in memory. */
-	uint32 TotalLastRenderHeuristicSize;
+	uint64 TotalLastRenderHeuristicSize;
 	/** Total number of bytes using ForcedIntoMemory heuristics, currently in memory. */
-	uint32 TotalForcedHeuristicSize;
+	uint64 TotalForcedHeuristicSize;
 
 	/** Unmodified texture pool size, in bytes, as specified in the .ini file. */
 	int64 OriginalTexturePoolSize;
@@ -1364,11 +1369,11 @@ protected:
 #endif
 
 #if STATS_FAST
-	uint32 MaxStreamingTexturesSize;
-	uint32 MaxOptimalTextureSize;
-	int32 MaxStreamingOverBudget;
-	uint32 MaxTexturePoolAllocatedSize;
-	uint32 MinLargestHoleSize;
+	uint64 MaxStreamingTexturesSize;
+	uint64 MaxOptimalTextureSize;
+	int64 MaxStreamingOverBudget;
+	uint64 MaxTexturePoolAllocatedSize;
+	uint64 MinLargestHoleSize;
 	uint32 MaxNumWantingTextures;
 #endif
 	
diff --git a/Engine/Source/Runtime/Engine/Public/EngineStats.h b/Engine/Source/Runtime/Engine/Public/EngineStats.h
--- a/Engine/Source/Runtime/Engine/Public/EngineStats.h
+++ b/Engine/Source/Runtime/Engine/Public/EngineStats.h
@@ -139,7 +139,7 @@ DECLARE_DWORD_COUNTER_STAT_EXTERN(TEXT("Num Destructible shapes exported"),STAT_
 DECLARE_DWORD_COUNTER_STAT_EXTERN(TEXT("Num update nav octree"),STAT_Navigation_UpdateNavOctree,STATGROUP_Navigation, );
 DECLARE_CYCLE_STAT_EXTERN(TEXT("Adding actors to navoctree"),STAT_Navigation_AddingActorsToNavOctree,STATGROUP_Navigation, );
 DECLARE_CYCLE_STAT_EXTERN(TEXT("    Adjusting nav links"),STAT_Navigation_AdjustingNavLinks,STATGROUP_Navigation, );
-DECLARE_CYCLE_STAT_EXTERN(TEXT("    Async AddTile"),STAT_Navigation_AddGeneratedTiles,STATGROUP_Navigation, );
+DECLARE_CYCLE_STAT_EXTERN(TEXT("    Sync AddGeneratedTiles"),STAT_Navigation_AddGeneratedTiles,STATGROUP_Navigation, );
 DECLARE_CYCLE_STAT_EXTERN(TEXT("Recast tick"),STAT_Navigation_RecastTick,STATGROUP_Navigation, );
 DECLARE_CYCLE_STAT_EXTERN(TEXT("Recast: build compressed layers"),STAT_Navigation_RecastBuildCompressedLayers,STATGROUP_Navigation, );
 DECLARE_CYCLE_STAT_EXTERN(TEXT("Recast: build navmesh"),STAT_Navigation_RecastBuildNavigation,STATGROUP_Navigation, );
diff --git a/Engine/Source/Runtime/Engine/Public/ImageUtils.h b/Engine/Source/Runtime/Engine/Public/ImageUtils.h
--- a/Engine/Source/Runtime/Engine/Public/ImageUtils.h
+++ b/Engine/Source/Runtime/Engine/Public/ImageUtils.h
@@ -24,6 +24,9 @@ struct FCreateTexture2DParameters
 	/** If texture should be set as SRGB */
 	bool						bSRGB;
 
+	/* The Guid hash to use part of the texture source's DDC key */
+	FGuid						SourceGuidHash;
+
 	FCreateTexture2DParameters()
 		:	bUseAlpha(false),
 			CompressionSettings(TC_Default),
diff --git a/Engine/Source/Runtime/Engine/Public/MaterialShared.h b/Engine/Source/Runtime/Engine/Public/MaterialShared.h
--- a/Engine/Source/Runtime/Engine/Public/MaterialShared.h
+++ b/Engine/Source/Runtime/Engine/Public/MaterialShared.h
@@ -1582,6 +1582,8 @@ class FMaterialUpdateContext
 	TSet<UMaterialInterface*> UpdatedMaterialInterfaces;
 	/** Active global component reregister context, if any. */
 	TScopedPointer<class FGlobalComponentReregisterContext> ComponentReregisterContext;
+	/** Active global component render state recreation context, if any. */
+	TScopedPointer<class FGlobalComponentRecreateRenderStateContext> ComponentRecreateRenderStateContext;
 	/** The shader platform that was being processed - can control if we need to update components */
 	EShaderPlatform ShaderPlatform;
 	/** True if the SyncWithRenderingThread option was specified. */
@@ -1602,8 +1604,10 @@ public:
 			 * you have already flushed rendering commands.
 			 */
 			SyncWithRenderingThread = 0x2,
-			/** Default options: reregister components, sync with rendering thread. */
-			Default = ReregisterComponents | SyncWithRenderingThread,
+			/* Recreates only the render state for all components (mutually exclusive with ReregisterComponents) */
+			RecreateRenderStates = 0x4,
+			/** Default options: Recreate render state, sync with rendering thread. */
+			Default = RecreateRenderStates | SyncWithRenderingThread,
 		};
 	};
 
diff --git a/Engine/Source/Runtime/Engine/Public/PhysicsPublic.h b/Engine/Source/Runtime/Engine/Public/PhysicsPublic.h
--- a/Engine/Source/Runtime/Engine/Public/PhysicsPublic.h
+++ b/Engine/Source/Runtime/Engine/Public/PhysicsPublic.h
@@ -31,9 +31,11 @@ DECLARE_CYCLE_STAT_EXTERN(TEXT("Term Body"), STAT_TermBody, STATGROUP_Physics, )
 DECLARE_CYCLE_STAT_EXTERN(TEXT("Update Materials"), STAT_UpdatePhysMats, STATGROUP_Physics, );
 DECLARE_CYCLE_STAT_EXTERN(TEXT("Update Materials Scene Interaction"), STAT_UpdatePhysMatsSceneInteraction, STATGROUP_Physics, );
 DECLARE_CYCLE_STAT_EXTERN(TEXT("Filter Update"), STAT_UpdatePhysFilter, STATGROUP_Physics, );
+DECLARE_CYCLE_STAT_EXTERN(TEXT("Filter Update (PhysX Code)"), STAT_UpdatePhysFilterPhysX, STATGROUP_Physics, );
 
-DECLARE_CYCLE_STAT_EXTERN(TEXT("Bulk Body Init"), STAT_BulkInit, STATGROUP_Physics, );
+DECLARE_CYCLE_STAT_EXTERN(TEXT("Init Bodies"), STAT_InitBodies, STATGROUP_Physics, );
 DECLARE_CYCLE_STAT_EXTERN(TEXT("Bulk Body Scene Add"), STAT_BulkSceneAdd, STATGROUP_Physics, );
+DECLARE_CYCLE_STAT_EXTERN(TEXT("Static Init Bodies"), STAT_StaticInitBodies, STATGROUP_Physics, );
 
 DECLARE_DWORD_COUNTER_STAT_EXTERN(TEXT("Broadphase Adds"), STAT_NumBroadphaseAdds, STATGROUP_Physics, );
 DECLARE_DWORD_COUNTER_STAT_EXTERN(TEXT("Broadphase Removes"), STAT_NumBroadphaseRemoves, STATGROUP_Physics, );
@@ -316,6 +318,15 @@ public:
 	 */
 	void DeferAddActor(FBodyInstance* OwningInstance, PxActor* Actor, EPhysicsSceneType SceneType);
 
+	/** Defer the addition of a group of actors to a scene, this will actually be performed before the *next*
+	 *  Physics tick. 
+	 *
+	 *	@param OwningInstance - The FBodyInstance that owns the actor
+	 *	@param Actor - The actual PhysX actor to add
+	 *	@param SceneType - The scene type to add the actor to
+	 */
+	void DeferAddActors(TArray<FBodyInstance*>& OwningInstances, TArray<PxActor*>& Actors, EPhysicsSceneType SceneType);
+
 	/** Defer the removal of an actor to a scene, this will actually be performed before the *next*
 	 *  Physics tick
 	 *	@param OwningInstance - The FBodyInstance that owns the actor
diff --git a/Engine/Source/Runtime/Engine/Public/SceneInterface.h b/Engine/Source/Runtime/Engine/Public/SceneInterface.h
--- a/Engine/Source/Runtime/Engine/Public/SceneInterface.h
+++ b/Engine/Source/Runtime/Engine/Public/SceneInterface.h
@@ -193,10 +193,10 @@ public:
 	virtual const TArray<class FWindSourceSceneProxy*>& GetWindSources_RenderThread() const = 0;
 
 	/** Accesses wind parameters.  XYZ will contain wind direction * Strength, W contains wind speed. */
-	virtual FVector4 GetWindParameters(const FVector& Position) const = 0;
+	virtual void GetWindParameters(const FVector& Position, FVector& OutDirection, float& OutSpeed, float& OutMinGustAmt, float& OutMaxGustAmt) const = 0;
 
 	/** Same as GetWindParameters, but ignores point wind sources. */
-	virtual FVector4 GetDirectionalWindParameters() const = 0;
+	virtual void GetDirectionalWindParameters(FVector& OutDirection, float& OutSpeed, float& OutMinGustAmt, float& OutMaxGustAmt) const = 0;
 
 	/** 
 	 * Adds a SpeedTree wind computation object to the scene.
diff --git a/Engine/Source/Runtime/Engine/Public/SceneManagement.h b/Engine/Source/Runtime/Engine/Public/SceneManagement.h
--- a/Engine/Source/Runtime/Engine/Public/SceneManagement.h
+++ b/Engine/Source/Runtime/Engine/Public/SceneManagement.h
@@ -1023,30 +1023,55 @@ public:
 /** Represents a wind source component to the scene manager in the rendering thread. */
 class ENGINE_API FWindSourceSceneProxy
 {
-public:
+public:	
+
+	class ENGINE_API FWindData
+	{
+	public:
+		FWindData()
+			: Speed(0.0f)
+			, MinGustAmt(0.0f)
+			, MaxGustAmt(0.0f)
+			, Direction(1.0f, 0.0f, 0.0f)
+		{
+		}
+
+		void PrepareForAccumulate();
+		void AddWeighted(const FWindData& InWindData, float Weight);
+		void NormalizeByTotalWeight(float TotalWeight);
+
+		float Speed;
+		float MinGustAmt;
+		float MaxGustAmt;
+		FVector Direction;
+	};
 
 	/** Initialization constructor. */
-	FWindSourceSceneProxy(const FVector& InDirection,float InStrength,float InSpeed):
+	FWindSourceSceneProxy(const FVector& InDirection, float InStrength, float InSpeed, float InMinGustAmt, float InMaxGustAmt) :
 	  Position(FVector::ZeroVector),
 		  Direction(InDirection),
 		  Strength(InStrength),
 		  Speed(InSpeed),
+		  MinGustAmt(InMinGustAmt),
+		  MaxGustAmt(InMaxGustAmt),
 		  Radius(0),
 		  bIsPointSource(false)
 	  {}
 
 	  /** Initialization constructor. */
-	  FWindSourceSceneProxy(const FVector& InPosition,float InStrength,float InSpeed,float InRadius):
+	FWindSourceSceneProxy(const FVector& InPosition, float InStrength, float InSpeed, float InMinGustAmt, float InMaxGustAmt, float InRadius) :
 	  Position(InPosition),
 		  Direction(FVector::ZeroVector),
 		  Strength(InStrength),
 		  Speed(InSpeed),
+		  MinGustAmt(InMinGustAmt),
+		  MaxGustAmt(InMaxGustAmt),
 		  Radius(InRadius),
 		  bIsPointSource(true)
 	  {}
 
-	  bool GetWindParameters(const FVector& EvaluatePosition, FVector4& WindDirectionAndSpeed, float& Strength) const;
-	  bool GetDirectionalWindParameters(FVector4& WindDirectionAndSpeed, float& Strength) const;
+	  bool GetWindParameters(const FVector& EvaluatePosition, FWindData& WindData, float& Weight) const;
+	  bool GetDirectionalWindParameters(FWindData& WindData, float& Weight) const;
 	  void ApplyWorldOffset(FVector InOffset);
 
 private:
@@ -1055,6 +1080,8 @@ private:
 	FVector	Direction;
 	float Strength;
 	float Speed;
+	float MinGustAmt;
+	float MaxGustAmt;
 	float Radius;
 	bool bIsPointSource;
 };
diff --git a/Engine/Source/Runtime/Engine/Public/SceneView.h b/Engine/Source/Runtime/Engine/Public/SceneView.h
--- a/Engine/Source/Runtime/Engine/Public/SceneView.h
+++ b/Engine/Source/Runtime/Engine/Public/SceneView.h
@@ -99,6 +99,9 @@ struct FSceneViewInitOptions : public FSceneViewProjectionData
 	// Whether world origin was rebased this frame?
 	bool bOriginOffsetThisFrame;
 
+	// Whether to use FOV when computing mesh LOD.
+	bool bUseFieldOfViewForLOD;
+
 #if WITH_EDITOR
 	// default to 0'th view index, which is a bitfield of 1
 	uint64 EditorViewBitflag;
@@ -125,6 +128,7 @@ struct FSceneViewInitOptions : public FSceneViewProjectionData
 		, OverrideFarClippingPlaneDistance(-1.0f)
 		, bInCameraCut(false)
 		, bOriginOffsetThisFrame(false)
+		, bUseFieldOfViewForLOD(true)
 #if WITH_EDITOR
 		, EditorViewBitflag(1)
 		, OverrideLODViewOrigin(ForceInitToZero)
@@ -156,6 +160,8 @@ struct FViewMatrices
 	FMatrix		ProjMatrix;
 	// WorldToView..
 	FMatrix		ViewMatrix;
+	/** WorldToView with PreViewTranslation. */
+	FMatrix		TranslatedViewMatrix;
 	/** The view-projection transform, starting from world-space points translated by -ViewOrigin. */
 	FMatrix		TranslatedViewProjectionMatrix;
 	/** The inverse view-projection transform, ending with world-space points translated by -ViewOrigin. */
@@ -298,6 +304,7 @@ BEGIN_UNIFORM_BUFFER_STRUCT_WITH_CONSTRUCTOR(FViewUniformShaderParameters,ENGINE
 	DECLARE_UNIFORM_BUFFER_STRUCT_MEMBER(float,RealTime)
 	DECLARE_UNIFORM_BUFFER_STRUCT_MEMBER(uint32,Random)
 	DECLARE_UNIFORM_BUFFER_STRUCT_MEMBER(uint32,FrameNumber)
+	DECLARE_UNIFORM_BUFFER_STRUCT_MEMBER_EX(float,CameraCut, EShaderPrecisionModifier::Half)
 	DECLARE_UNIFORM_BUFFER_STRUCT_MEMBER_EX(float,UseLightmaps, EShaderPrecisionModifier::Half)
 	DECLARE_UNIFORM_BUFFER_STRUCT_MEMBER_EX(float,UnlitViewmodeMask, EShaderPrecisionModifier::Half)
 	DECLARE_UNIFORM_BUFFER_STRUCT_MEMBER_EX(FLinearColor,DirectionalLightColor, EShaderPrecisionModifier::Half)
@@ -417,6 +424,9 @@ public:
 	/** Whether this view should render the first instance only of any meshes using instancing. */
 	bool bRenderFirstInstanceOnly;
 
+	// Whether to use FOV when computing mesh LOD.
+	bool bUseFieldOfViewForLOD;
+
 	/** Current buffer visualization mode */
 	FName CurrentBufferVisualizationMode;
 
diff --git a/Engine/Source/Runtime/Engine/Public/SpeedTreeWind.h b/Engine/Source/Runtime/Engine/Public/SpeedTreeWind.h
--- a/Engine/Source/Runtime/Engine/Public/SpeedTreeWind.h
+++ b/Engine/Source/Runtime/Engine/Public/SpeedTreeWind.h
@@ -241,6 +241,11 @@ public:
 	ENGINE_API	const FSpeedTreeWind::SParams&	GetParams(void) const;
 	ENGINE_API	void							SetStrength(float fStrength);						// use this function to set a new desired strength (it will reach that strength smoothly)
 	ENGINE_API	void							SetDirection(const FVector& vDir);					// use this function to set a new desired direction (it will reach that direction smoothly)
+
+	//instantly set gust min/max.   Trees will pop if visible.  Don't call during gameplay while trees are visible.
+	ENGINE_API	void							SetGustMin(float InGustMin);
+	ENGINE_API	void							SetGustMax(float InGustMax);
+
 	ENGINE_API	void							SetInitDirection(const FVector& vDir);				// use this function to set a starting direction, once
 	ENGINE_API	void							EnableGusting(bool bEnabled);
 	ENGINE_API	void							SetGustFrequency(float fGustFreq);
diff --git a/Engine/Source/Runtime/Engine/Public/StaticMeshResources.h b/Engine/Source/Runtime/Engine/Public/StaticMeshResources.h
--- a/Engine/Source/Runtime/Engine/Public/StaticMeshResources.h
+++ b/Engine/Source/Runtime/Engine/Public/StaticMeshResources.h
@@ -939,6 +939,30 @@ struct FInstanceStream
 		Me->InstanceLightmapAndShadowMapUVBias[3] = 0;
 	}
 
+	FORCEINLINE void GetInstanceTransform(FMatrix& Transform) const
+	{
+		Transform.M[3][0] = InstanceOrigin.X;
+		Transform.M[3][1] = InstanceOrigin.Y;
+		Transform.M[3][2] = InstanceOrigin.Z;
+
+		Transform.M[0][0] = InstanceTransform1[0];
+		Transform.M[1][0] = InstanceTransform1[1];
+		Transform.M[2][0] = InstanceTransform1[2];
+
+		Transform.M[0][1] = InstanceTransform2[0];
+		Transform.M[1][1] = InstanceTransform2[1];
+		Transform.M[2][1] = InstanceTransform2[2];
+
+		Transform.M[0][2] = InstanceTransform3[0];
+		Transform.M[1][2] = InstanceTransform3[1];
+		Transform.M[2][2] = InstanceTransform3[2];
+
+		Transform.M[0][3] = 0.f;
+		Transform.M[1][3] = 0.f;
+		Transform.M[2][3] = 0.f;
+		Transform.M[3][3] = 0.f;
+	}
+
 	FORCEINLINE void SetInstance(const FMatrix& Transform, float RandomInstanceID, const FVector2D& LightmapUVBias, const FVector2D& ShadowmapUVBias, FColor HitProxyColor, bool bSelected)
 	{
 		const FMatrix* RESTRICT rTransform = (const FMatrix* RESTRICT)&Transform;
@@ -972,6 +996,39 @@ struct FInstanceStream
 		Me->InstanceLightmapAndShadowMapUVBias[3] = FMath::Clamp<int32>(FMath::TruncToInt(rShadowmapUVBias->Y * 32767.0f), MIN_int16, MAX_int16);
 	}
 
+	FORCEINLINE void SetInstance(const FMatrix& Transform, float RandomInstanceID, const FVector2D& LightmapUVBias, const FVector2D& ShadowmapUVBias)
+	{
+		const FMatrix* RESTRICT rTransform = (const FMatrix* RESTRICT)&Transform;
+		FInstanceStream* RESTRICT Me = (FInstanceStream* RESTRICT)this;
+		const FVector2D* RESTRICT rLightmapUVBias = (const FVector2D* RESTRICT)&LightmapUVBias;
+		const FVector2D* RESTRICT rShadowmapUVBias = (const FVector2D* RESTRICT)&ShadowmapUVBias;
+
+		Me->InstanceOrigin.X = rTransform->M[3][0];
+		Me->InstanceOrigin.Y = rTransform->M[3][1];
+		Me->InstanceOrigin.Z = rTransform->M[3][2];
+		Me->InstanceOrigin.W = RandomInstanceID;
+
+		Me->InstanceTransform1[0] = Transform.M[0][0];
+		Me->InstanceTransform1[1] = Transform.M[1][0];
+		Me->InstanceTransform1[2] = Transform.M[2][0];
+		Me->InstanceTransform1[3] = FFloat16();
+
+		Me->InstanceTransform2[0] = Transform.M[0][1];
+		Me->InstanceTransform2[1] = Transform.M[1][1];
+		Me->InstanceTransform2[2] = Transform.M[2][1];
+		Me->InstanceTransform2[3] = FFloat16();
+
+		Me->InstanceTransform3[0] = Transform.M[0][2];
+		Me->InstanceTransform3[1] = Transform.M[1][2];
+		Me->InstanceTransform3[2] = Transform.M[2][2];
+		Me->InstanceTransform3[3] = FFloat16();
+
+		Me->InstanceLightmapAndShadowMapUVBias[0] = FMath::Clamp<int32>(FMath::TruncToInt(rLightmapUVBias->X  * 32767.0f), MIN_int16, MAX_int16);
+		Me->InstanceLightmapAndShadowMapUVBias[1] = FMath::Clamp<int32>(FMath::TruncToInt(rLightmapUVBias->Y  * 32767.0f), MIN_int16, MAX_int16);
+		Me->InstanceLightmapAndShadowMapUVBias[2] = FMath::Clamp<int32>(FMath::TruncToInt(rShadowmapUVBias->X * 32767.0f), MIN_int16, MAX_int16);
+		Me->InstanceLightmapAndShadowMapUVBias[3] = FMath::Clamp<int32>(FMath::TruncToInt(rShadowmapUVBias->Y * 32767.0f), MIN_int16, MAX_int16);
+	}
+
 	FORCEINLINE void NullifyInstance()
 	{
 		FInstanceStream* RESTRICT Me = (FInstanceStream* RESTRICT)this;
diff --git a/Engine/Source/Runtime/Engine/Public/TextureResource.h b/Engine/Source/Runtime/Engine/Public/TextureResource.h
--- a/Engine/Source/Runtime/Engine/Public/TextureResource.h
+++ b/Engine/Source/Runtime/Engine/Public/TextureResource.h
@@ -51,7 +51,7 @@ public:
 	/** Initialization constructor. */
 	FAsyncStreamDerivedMipWorker(
 		const FString& InDerivedDataKey,
-		void* InDestMipData,
+		void** InDestMipDataPointer,
 		int32 InMipSize,
 		FThreadSafeCounter* InThreadSafeCounter
 		);
@@ -78,7 +78,7 @@ private:
 	/** Key for retrieving mip data from the derived data cache. */
 	FString DerivedDataKey;
 	/** The location to which the mip data should be copied. */
-	void* DestMipData;
+	void** DestMipDataPointer;
 	/** The size of the mip in bytes. */
 	int32 ExpectedMipSize;
 	/** true if the mip data was not present in the derived data cache. */
diff --git a/Engine/Source/Runtime/Landscape/Classes/LandscapeComponent.h b/Engine/Source/Runtime/Landscape/Classes/LandscapeComponent.h
--- a/Engine/Source/Runtime/Landscape/Classes/LandscapeComponent.h
+++ b/Engine/Source/Runtime/Landscape/Classes/LandscapeComponent.h
@@ -127,7 +127,12 @@ struct FLandscapeComponentGrassData
 
 	SIZE_T GetAllocatedSize() const
 	{
-		return sizeof(*this) + HeightData.GetAllocatedSize() + WeightData.GetAllocatedSize();
+		SIZE_T WeightSize = 0; 
+		for (auto It = WeightData.CreateConstIterator(); It; ++It)
+		{
+			WeightSize += It.Value().GetAllocatedSize();
+		}
+		return sizeof(*this) + HeightData.GetAllocatedSize() + WeightData.GetAllocatedSize() + WeightSize;
 	}
 
 	friend FArchive& operator<<(FArchive& Ar, FLandscapeComponentGrassData& Data);
@@ -352,6 +357,9 @@ public:
 
 	/* Is the grassmap data outdated, eg by a material */
 	bool IsGrassMapOutdated() const;
+
+	/* Serialize all hashes/guids that record the current state of this component */
+	void SerializeStateHashes(FArchive& Ar);
 #endif
 
 	/** @todo document */
diff --git a/Engine/Source/Runtime/Landscape/Classes/LandscapeGrassType.h b/Engine/Source/Runtime/Landscape/Classes/LandscapeGrassType.h
--- a/Engine/Source/Runtime/Landscape/Classes/LandscapeGrassType.h
+++ b/Engine/Source/Runtime/Landscape/Classes/LandscapeGrassType.h
@@ -16,7 +16,11 @@ struct FGrassVariety
 	UPROPERTY(EditAnywhere, Category=Grass)
 	float GrassDensity;
 
+	/* If true, use a jittered grid sequence for placement, otherwise use a halton sequence. */
 	UPROPERTY(EditAnywhere, Category=Grass)
+	bool bUseGrid;
+
+	UPROPERTY(EditAnywhere, Category=Grass, meta = (EditCondition = "bUseGrid", UIMin = 0, ClampMin = 0, UIMax = 1, ClampMax = 1))
 	float PlacementJitter;
 
 	/* The distance where instances will begin to fade out if using a PerInstanceFadeAmount material node. 0 disables. */
@@ -30,6 +34,13 @@ struct FGrassVariety
 	UPROPERTY(EditAnywhere, Category = Grass)
 	int32 EndCullDistance;
 
+	/** 
+	 * Specifies the smallest LOD that will be used for this component.  
+	 * If -1 (default), the MinLOD of the static mesh asset will be used instead.
+	 */
+	UPROPERTY(EditAnywhere, Category = Grass)
+	int32 MinLOD;
+
 	UPROPERTY(EditAnywhere, Category = Grass)
 	bool RandomRotation;
 
@@ -43,8 +54,10 @@ struct FGrassVariety
 		StartCullDistance = 10000.0f;
 		EndCullDistance = 10000.0f;
 		PlacementJitter = 1.0f;
+		MinLOD = -1;
 		RandomRotation = true;
 		AlignToSurface = true;
+		bUseGrid = true;
 	}
 };
 
diff --git a/Engine/Source/Runtime/Landscape/Classes/LandscapeHeightfieldCollisionComponent.h b/Engine/Source/Runtime/Landscape/Classes/LandscapeHeightfieldCollisionComponent.h
--- a/Engine/Source/Runtime/Landscape/Classes/LandscapeHeightfieldCollisionComponent.h
+++ b/Engine/Source/Runtime/Landscape/Classes/LandscapeHeightfieldCollisionComponent.h
@@ -6,6 +6,7 @@
 class ULandscapeLayerInfoObject;
 class AInstancedFoliageActor;
 
+#include "AI/Navigation/NavigationTypes.h"
 #include "LandscapeHeightfieldCollisionComponent.generated.h"
 
 #if WITH_PHYSX
@@ -106,6 +107,13 @@ class ULandscapeHeightfieldCollisionComponent : public UPrimitiveComponent
 	 *  For editor and full version of collision objects
 	 */
 	mutable bool								bShouldSaveCookedDataToDDC[2];
+
+	/**
+	  * Async DCC load for cooked collision representation. We speculatively
+	  * load this to remove hitch when streaming 
+	  */
+	mutable TSharedPtr<FAsyncPreRegisterDDCRequest>	SpeculativeDDCRequest;
+
 #endif //WITH_EDITORONLY_DATA
 
 	/** 
@@ -120,6 +128,11 @@ class ULandscapeHeightfieldCollisionComponent : public UPrimitiveComponent
 	
 	/** Physics engine version of heightfield data. */
 	TRefCountPtr<struct FPhysXHeightfieldRef>	HeightfieldRef;
+	
+	/** Cached PxHeightFieldSamples values for navmesh generation. Note that it's being used only if navigation octree is set up for lazy geometry exporting */
+	int32 HeightfieldRowsCount;
+	int32 HeightfieldColumnsCount;
+	mutable FNavHeightfieldSamples CachedHeightFieldSamples;
 
 	enum ECollisionQuadFlags
 	{
@@ -144,13 +157,20 @@ class ULandscapeHeightfieldCollisionComponent : public UPrimitiveComponent
 	// End USceneComponent interface.
 
 	// Begin UPrimitiveComponent interface
-	virtual bool DoCustomNavigableGeometryExport(struct FNavigableGeometryExport* GeomExport) const override;
+	virtual bool DoCustomNavigableGeometryExport(FNavigableGeometryExport& GeomExport) const override;
 #if WITH_EDITOR
 	virtual bool ComponentIsTouchingSelectionBox(const FBox& InSelBBox, const FEngineShowFlags& ShowFlags, const bool bConsiderOnlyBSP, const bool bMustEncompassEntireComponent) const override;
 	virtual bool ComponentIsTouchingSelectionFrustum(const FConvexVolume& InFrustum, const FEngineShowFlags& ShowFlags, const bool bConsiderOnlyBSP, const bool bMustEncompassEntireComponent) const override;
 #endif
 	//End UPrimitiveComponent interface
 
+	// Begin INavRelevantInterface interface
+	virtual bool SupportsGatheringGeometrySlices() const override { return true; }
+	virtual void GatherGeometrySlice(FNavigableGeometryExport& GeomExport, const FBox& SliceBox) const override;
+	virtual ENavDataGatheringMode GetGeometryGatheringMode() const override;
+	virtual void PrepareGeometryExportSync() override;
+	// End INavRelevantInterface interface
+
 	// Begin UObject Interface.
 	virtual void Serialize(FArchive& Ar) override;
 	virtual void BeginDestroy() override;
@@ -169,6 +189,9 @@ class ULandscapeHeightfieldCollisionComponent : public UPrimitiveComponent
 	// @todo document
 	class ULandscapeInfo* GetLandscapeInfo(bool bSpawnNewActor = true) const;
 
+	/**  We speculatively async load collision object from DDC to remove hitch when streaming */
+	void SpeculativelyLoadAsyncDDCCollsionData();
+
 	/** 
 	 * Cooks raw height data into collision object binary stream
 	 */
diff --git a/Engine/Source/Runtime/Landscape/Classes/LandscapeProxy.h b/Engine/Source/Runtime/Landscape/Classes/LandscapeProxy.h
--- a/Engine/Source/Runtime/Landscape/Classes/LandscapeProxy.h
+++ b/Engine/Source/Runtime/Landscape/Classes/LandscapeProxy.h
@@ -17,6 +17,7 @@
 #include "LandscapeLayerInfoObject.h"
 #include "LandscapeGrassType.h"
 #include "Tickable.h"
+#include "AI/Navigation/NavigationTypes.h"
 
 #include "LandscapeProxy.generated.h"
 
@@ -454,6 +455,9 @@ public:
 	UPROPERTY(EditAnywhere, Category=Landscape)
 	uint32 bUsedForNavigation:1;
 
+	UPROPERTY(EditAnywhere, Category = Landscape, AdvancedDisplay)
+	ENavDataGatheringMode NavigationGeometryGatheringMode;
+
 	UPROPERTY(EditAnywhere, Category=LOD)
 	TEnumAsByte<enum ELandscapeLODFalloff::Type> LODFalloff;
 
@@ -659,6 +663,9 @@ public:
 	/* For the texture baking notification */
 	int32 NumComponentsNeedingTextureBaking;
 	LANDSCAPE_API static int32 TotalComponentsNeedingTextureBaking;
+
+	/** remove an overlapping component. Called from MapCheck. */
+	LANDSCAPE_API void RemoveOverlappingComponent(ULandscapeComponent* Component);
 #endif
 };
 
