UE4 - occlusion queries for foliage and grass

diff --git a/Engine/Source/Runtime/Core/Public/UObject/ObjectVersion.h b/Engine/Source/Runtime/Core/Public/UObject/ObjectVersion.h
--- a/Engine/Source/Runtime/Core/Public/UObject/ObjectVersion.h
+++ b/Engine/Source/Runtime/Core/Public/UObject/ObjectVersion.h
@@ -478,9 +478,12 @@ enum EUnrealEngineObjectUE4Version
 	VER_UE4_K2NODE_EVENT_MEMBER_REFERENCE,
 	// FPropertyTag stores GUID of struct
 	VER_UE4_STRUCT_GUID_IN_PROPERTY_TAG,
-
 	// Remove unused UPolys from UModel cooked content
 	VER_UE4_REMOVE_UNUSED_UPOLYS_FROM_UMODEL,
+	// This doesn't do anything except trigger a rebuild on HISMC cluster trees, in this case to get a good "occlusion query" level
+	VER_UE4_REBUILD_HIERARCHICAL_INSTANCE_TREES,
+
+
 	// -----<new versions can be added before this line>-------------------------------------------------
 	// - this needs to be the last line (see note below)
 	VER_UE4_AUTOMATIC_VERSION_PLUS_ONE,
diff --git a/Engine/Source/Runtime/Engine/Classes/Components/HierarchicalInstancedStaticMeshComponent.h b/Engine/Source/Runtime/Engine/Classes/Components/HierarchicalInstancedStaticMeshComponent.h
--- a/Engine/Source/Runtime/Engine/Classes/Components/HierarchicalInstancedStaticMeshComponent.h
+++ b/Engine/Source/Runtime/Engine/Classes/Components/HierarchicalInstancedStaticMeshComponent.h
@@ -100,6 +100,10 @@ class ENGINE_API UHierarchicalInstancedStaticMeshComponent : public UInstancedSt
 	UPROPERTY()
 	FBox UnbuiltInstanceBounds;
 
+	// The number of nodes in the occlusion layer
+	UPROPERTY()
+	int32 OcclusionLayerNumNodes;
+
 	bool bIsAsyncBuilding;
 	bool bConcurrentRemoval;
 
@@ -142,9 +146,10 @@ public:
 		TArray<FClusterNode>& OutClusterTree,
 		TArray<int32>& OutSortedInstances,
 		TArray<int32>& OutInstanceReorderTable,
+		int32& OutOcclusionLayerNum,
 		int32 MaxInstancesPerLeaf
 		);
-	void AcceptPrebuiltTree(TArray<FClusterNode>& InClusterTree);
+	void AcceptPrebuiltTree(TArray<FClusterNode>& InClusterTree, int InOcclusionLayerNumNodes);
 	void BuildFlatTree(const TArray<int32>& LeafInstanceCounts);
 	bool IsAsyncBuilding() const { return bIsAsyncBuilding; }
 	bool IsTreeFullyBuilt() const { return NumBuiltInstances == PerInstanceSMData.Num() && RemovedInstances.Num() == 0; }
diff --git a/Engine/Source/Runtime/Engine/Classes/Components/PrimitiveComponent.h b/Engine/Source/Runtime/Engine/Classes/Components/PrimitiveComponent.h
--- a/Engine/Source/Runtime/Engine/Classes/Components/PrimitiveComponent.h
+++ b/Engine/Source/Runtime/Engine/Classes/Components/PrimitiveComponent.h
@@ -426,7 +426,7 @@ protected:
 	TEnumAsByte<EHasCustomNavigableGeometry::Type> bHasCustomNavigableGeometry;
 
 	/** Next id to be used by a component. */
-	static uint64 NextComponentId;
+	static uint32 NextComponentId;
 
 public:
 	/** 
diff --git a/Engine/Source/Runtime/Engine/Private/HierarchicalInstancedStaticMesh.cpp b/Engine/Source/Runtime/Engine/Private/HierarchicalInstancedStaticMesh.cpp
--- a/Engine/Source/Runtime/Engine/Private/HierarchicalInstancedStaticMesh.cpp
+++ b/Engine/Source/Runtime/Engine/Private/HierarchicalInstancedStaticMesh.cpp
@@ -71,16 +71,21 @@ static TAutoConsoleVariable<int32> CVarASyncInstaneBufferConversion(
 	1,
 	TEXT("If this is > 0, then build game instance buffering async during streaming. This is not thought to be a long term solution to this problem."));
 
-static TAutoConsoleVariable<int32> CVarRebuildTreeOnLoad(
-	TEXT("foliage.RebuildTreeOnLoad"),
-	1,
-	TEXT("If this is > 0, then we look at our heuristics and rebuild the tree if they don't match."));
-
 static TAutoConsoleVariable<int32> CVarMinVertsToSplitNode(
 	TEXT("foliage.MinVertsToSplitNode"),
 	16384,
 	TEXT("Controls the accuracy between culling and LOD accuracy and culling and CPU performance."));
 
+static TAutoConsoleVariable<int32> CVarMaxOcclusionQueriesPerComponent(
+	TEXT("foliage.MaxOcclusionQueriesPerComponent"),
+	16,
+	TEXT("Controls the granualrity of occlusion culling. 16-128 is a reasonable range."));
+
+static TAutoConsoleVariable<int32> CVarMinInstancesPerOcclusionQuery(
+	TEXT("foliage.MinInstancesPerOcclusionQuery"),
+	256,
+	TEXT("Controls the granualrity of occlusion culling. 1024 to 65536 is a reasonable range. This is not exact, actual minimum might be off by a factor of two."));
+
 DECLARE_CYCLE_STAT(TEXT("Traversal Time"),STAT_FoliageTraversalTime,STATGROUP_Foliage);
 DECLARE_CYCLE_STAT(TEXT("Build Time"), STAT_FoliageBuildTime, STATGROUP_Foliage);
 DECLARE_CYCLE_STAT(TEXT("Batch Time"),STAT_FoliageBatchTime,STATGROUP_Foliage);
@@ -88,6 +93,7 @@ DECLARE_DWORD_COUNTER_STAT(TEXT("Runs"), STAT_FoliageRuns, STATGROUP_Foliage);
 DECLARE_DWORD_COUNTER_STAT(TEXT("Mesh Batches"), STAT_FoliageMeshBatches, STATGROUP_Foliage);
 DECLARE_DWORD_COUNTER_STAT(TEXT("Triangles"), STAT_FoliageTriangles, STATGROUP_Foliage);
 DECLARE_DWORD_COUNTER_STAT(TEXT("Instances"), STAT_FoliageInstances, STATGROUP_Foliage);
+DECLARE_DWORD_COUNTER_STAT(TEXT("Occlusion Culled Instances"), STAT_OcclusionCulledFoliageInstances, STATGROUP_Foliage);
 DECLARE_DWORD_COUNTER_STAT(TEXT("Traversals"),STAT_FoliageTraversals,STATGROUP_Foliage);
 DECLARE_MEMORY_STAT(TEXT("Instance Buffers"),STAT_FoliageInstanceBuffers,STATGROUP_Foliage);
 
@@ -96,6 +102,7 @@ struct FClusterTree
 	TArray<FClusterNode> Nodes;
 	TArray<int32> SortedInstances;
 	TArray<int32> InstanceReorderTable;
+	int32 OutOcclusionLayerNum;
 };
 
 
@@ -105,6 +112,7 @@ class FClusterBuilder
 	FBox InstBox;
 	int32 BranchingFactor;
 	int32 InternalNodeBranchingFactor;
+	int32 OcclusionLayerTarget;
 	int32 MaxInstancesPerLeaf;
 	int32 NumRoots;
 	TArray<int32> SortIndex;
@@ -228,6 +236,17 @@ public:
 		, InstBox(InInstBox)
 		, Result(nullptr)
 	{
+		OcclusionLayerTarget = CVarMaxOcclusionQueriesPerComponent.GetValueOnAnyThread();
+		int32 MinInstancesPerOcclusionQuery = CVarMinInstancesPerOcclusionQuery.GetValueOnAnyThread();
+
+		if (Num / MinInstancesPerOcclusionQuery < OcclusionLayerTarget)
+		{
+			OcclusionLayerTarget = Num / MinInstancesPerOcclusionQuery;
+			if (OcclusionLayerTarget < 6)
+			{
+				OcclusionLayerTarget = 0;
+			}
+		}
 		InternalNodeBranchingFactor = CVarFoliageSplitFactor.GetValueOnAnyThread();
 		MaxInstancesPerLeaf = InMaxInstancesPerLeaf;
 
@@ -249,11 +268,22 @@ public:
 
 	void Build()
 	{
+		Result = new FClusterTree;
+		bool bIsOcclusionLayer = false;
 		BranchingFactor = MaxInstancesPerLeaf;
+		if (BranchingFactor > 2 && OcclusionLayerTarget && Num / BranchingFactor <= OcclusionLayerTarget)
+		{
+			BranchingFactor = FMath::Max<int32>(2, (Num + OcclusionLayerTarget - 1) / OcclusionLayerTarget);
+			OcclusionLayerTarget = 0;
+			bIsOcclusionLayer = true;
+		}
 		Split(Num);
-		BranchingFactor = InternalNodeBranchingFactor;
+		if (bIsOcclusionLayer)
+		{
+			Result->OutOcclusionLayerNum = Clusters.Num();
+			bIsOcclusionLayer = false;
+		}
 
-		Result = new FClusterTree;
 		TArray<int32>& SortedInstances = Result->SortedInstances;
 		SortedInstances.AddUninitialized(Num);
 		for (int32 Index = 0; Index < Num; Index++)
@@ -306,7 +336,19 @@ public:
 				FClusterNode& Node = Result->Nodes[Index];
 				SortPoints[Index] = (Node.BoundMin + Node.BoundMax) * 0.5f;
 			}
+			BranchingFactor = InternalNodeBranchingFactor;
+			if (BranchingFactor > 2 && OcclusionLayerTarget && NumRoots / BranchingFactor <= OcclusionLayerTarget)
+			{
+				BranchingFactor = FMath::Max<int32>(2, (NumRoots + OcclusionLayerTarget - 1) / OcclusionLayerTarget);
+				OcclusionLayerTarget = 0;
+				bIsOcclusionLayer = true;
+			}
 			Split(NumRoots);
+			if (bIsOcclusionLayer)
+			{
+				Result->OutOcclusionLayerNum = Clusters.Num();
+				bIsOcclusionLayer = false;
+			}
 
 			InverseSortIndex.Reset();
 			InverseSortIndex.AddUninitialized(NumRoots);
@@ -575,7 +617,20 @@ static FAutoConsoleCommand UnFreezeFoliageCullingCmd(
 	);
 
 
+struct FFoliageOcclusionResults
+{
+	const bool* Results;
+	int32 NumResults;
+	uint32 FrameNumberRenderThread;
 
+	FFoliageOcclusionResults(const bool* InResults, int32 InNumResults)
+		: Results(InResults)
+		, NumResults(InNumResults)
+		, FrameNumberRenderThread(GFrameNumberRenderThread)
+	{
+
+	}
+};
 
 struct FFoliageElementParams;
 struct FFoliageRenderInstanceParams;
@@ -589,6 +644,10 @@ class FHierarchicalStaticMeshSceneProxy : public FInstancedStaticMeshSceneProxy
 	int32 FirstUnbuiltIndex;
 	int32 LastUnbuiltIndex;
 
+	int32 FirstOcclusionNode;
+	int32 LastOcclusionNode;
+	TArray<FBoxSphereBounds> OcclusionBounds;
+	TMap<uint32, FFoliageOcclusionResults> OcclusionResults;
 	bool bIsGrass;
 
 #if !(UE_BUILD_SHIPPING || UE_BUILD_TEST)
@@ -613,6 +672,7 @@ public:
 #endif
 	{
 		check(InComponent->InstanceReorderTable.Num() == InComponent->PerInstanceSMData.Num());
+		SetupOcclusion(InComponent);
 	}
 
 	FHierarchicalStaticMeshSceneProxy(bool bInIsGrass, UHierarchicalInstancedStaticMeshComponent* InComponent, ERHIFeatureLevel::Type InFeatureLevel, FStaticMeshInstanceData& Other)
@@ -627,6 +687,55 @@ public:
 #endif
 	{
 		check(!bInIsGrass || (!InComponent->InstanceReorderTable.Num() && !InComponent->PerInstanceSMData.Num()));
+		SetupOcclusion(InComponent);
+	}
+
+	void SetupOcclusion(UHierarchicalInstancedStaticMeshComponent* InComponent)
+	{
+		FirstOcclusionNode = 0;
+		LastOcclusionNode = 0;
+		if (ClusterTree.Num() && InComponent->OcclusionLayerNumNodes)
+		{
+			while (true)
+			{
+				int32 NextFirstOcclusionNode = ClusterTree[FirstOcclusionNode].FirstChild;
+				int32 NextLastOcclusionNode = ClusterTree[LastOcclusionNode].LastChild;
+
+				if (NextFirstOcclusionNode < 0 || NextLastOcclusionNode < 0)
+				{
+					break;
+				}
+				int32 NumNodes = 1 + NextLastOcclusionNode - NextFirstOcclusionNode;
+				if (NumNodes > InComponent->OcclusionLayerNumNodes)
+				{
+					break;
+				}
+				FirstOcclusionNode = NextFirstOcclusionNode;
+				LastOcclusionNode = NextLastOcclusionNode;
+			}
+		}
+		int32 NumNodes = 1 + LastOcclusionNode - FirstOcclusionNode;
+		if (NumNodes < 2)
+		{
+			FirstOcclusionNode = -1;
+			LastOcclusionNode = -1;
+			NumNodes = 0;
+			if (ClusterTree.Num())
+			{
+				//UE_LOG(LogTemp, Display, TEXT("No SubOcclusion %d inst"), 1 + ClusterTree[0].LastInstance - ClusterTree[0].FirstInstance);
+			}
+		}
+		else
+		{
+			//int32 NumPerNode = (1 + ClusterTree[0].LastInstance - ClusterTree[0].FirstInstance) / NumNodes;
+			//UE_LOG(LogTemp, Display, TEXT("Occlusion level %d   %d inst / node"), NumNodes, NumPerNode);
+			OcclusionBounds.Reserve(NumNodes);
+			FMatrix XForm = InComponent->ComponentToWorld.ToMatrixWithScale();
+			for (int32 Index = FirstOcclusionNode; Index <= LastOcclusionNode; Index++)
+			{
+				OcclusionBounds.Add(FBoxSphereBounds(FBox(ClusterTree[Index].BoundMin, ClusterTree[Index].BoundMax).TransformBy(XForm)));
+			}
+		}
 	}
 
 	// FPrimitiveSceneProxy interface.
@@ -645,6 +754,14 @@ public:
 
 	virtual void GetDynamicMeshElements(const TArray<const FSceneView*>& Views, const FSceneViewFamily& ViewFamily, uint32 VisibilityMap, FMeshElementCollector& Collector) const override;
 
+	virtual const TArray<FBoxSphereBounds>* GetOcclusionQueries(const FSceneView* View) const override;
+	virtual void AcceptOcclusionResults(const FSceneView* View, const bool* Results, int32 NumResults) override;
+	virtual bool HasSubprimitiveOcclusionQueries() const
+	{
+		return FirstOcclusionNode > 0;
+	}
+
+
 	virtual void DrawStaticElements(FStaticPrimitiveDrawInterface* PDI) override
 	{
 	}
@@ -722,11 +839,18 @@ struct FFoliageCullInstanceParams : public FFoliageRenderInstanceParams
 	int32 LODs;
 	float LODPlanesMax[MAX_STATIC_MESH_LODS];
 	float LODPlanesMin[MAX_STATIC_MESH_LODS];
+	int32 FirstOcclusionNode;
+	int32 LastOcclusionNode;
+	const bool* OcclusionResults;
+
 
 
 	FFoliageCullInstanceParams(bool InbNeedsSingleLODRuns, bool InbNeedsMultipleLODRuns, bool InbOverestimate, const TArray<FClusterNode>& InTree)
 	:	FFoliageRenderInstanceParams(InbNeedsSingleLODRuns, InbNeedsMultipleLODRuns, InbOverestimate)
 	,	Tree(InTree)
+	,	FirstOcclusionNode(-1)
+	,	LastOcclusionNode(-1)
+	,	OcclusionResults(nullptr)
 	{
 	}
 };
@@ -854,8 +978,16 @@ void FHierarchicalStaticMeshSceneProxy::Traverse(const FFoliageCullInstanceParam
 			MaxLOD--;
 		}
 	}
+	if (Index >= Params.FirstOcclusionNode && Index <= Params.LastOcclusionNode)
+	{
+		if (Params.OcclusionResults[Index - Params.FirstOcclusionNode])
+		{
+			INC_DWORD_STAT_BY(STAT_OcclusionCulledFoliageInstances, 1 + Node.LastInstance - Node.FirstInstance);
+			return;
+		}
+	}
 
-	bool bSplit = (!bFullyContained || MinLOD < MaxLOD) 
+	bool bSplit = (!bFullyContained || MinLOD < MaxLOD || Index < Params.FirstOcclusionNode) 
 		&& Node.FirstChild >= 0 
 		&& (1 + Node.LastInstance - Node.FirstInstance) >= Params.MinInstancesToSplit[MinLOD];
 
@@ -1209,6 +1341,21 @@ void FHierarchicalStaticMeshSceneProxy::GetDynamicMeshElements(const TArray<cons
 					}
 				}
 
+				if (FirstOcclusionNode >= 0 && LastOcclusionNode >= 0 && FirstOcclusionNode <= LastOcclusionNode)
+				{
+					uint32 ViewId = View->GetViewKey();
+					const FFoliageOcclusionResults* OldResults = OcclusionResults.Find(ViewId);
+					if (OldResults && 
+						OldResults->FrameNumberRenderThread == GFrameNumberRenderThread &&
+						1 + LastOcclusionNode - FirstOcclusionNode == OldResults->NumResults
+						)
+					{
+						InstanceParams.FirstOcclusionNode = FirstOcclusionNode;
+						InstanceParams.LastOcclusionNode = LastOcclusionNode;
+						InstanceParams.OcclusionResults = OldResults->Results;
+					}
+				}
+
 				INC_DWORD_STAT(STAT_FoliageTraversals);
 #if !(UE_BUILD_SHIPPING || UE_BUILD_TEST)
 				if (GCaptureDebugRuns == GDebugTag && CaptureTag == GDebugTag)
@@ -1314,6 +1461,38 @@ void FHierarchicalStaticMeshSceneProxy::GetDynamicMeshElements(const TArray<cons
 	}
 }
 
+void FHierarchicalStaticMeshSceneProxy::AcceptOcclusionResults(const FSceneView* View, const bool* Results, int32 NumResults)
+{
+	if (OcclusionBounds.Num() == NumResults)
+	{
+		uint32 ViewId = View->GetViewKey();
+		FFoliageOcclusionResults* OldResults = OcclusionResults.Find(ViewId);
+		if (OldResults)
+		{
+			OldResults->FrameNumberRenderThread = GFrameNumberRenderThread;
+			OldResults->Results = Results;
+			OldResults->NumResults = NumResults;
+		}
+		else
+		{
+			// now is a good time to clean up any stale entries
+			for (auto Iter = OcclusionResults.CreateIterator(); Iter; ++Iter)
+			{
+				if (Iter.Value().FrameNumberRenderThread != GFrameNumberRenderThread)
+				{
+					Iter.RemoveCurrent();
+				}
+			}
+			OcclusionResults.Add(ViewId, FFoliageOcclusionResults(Results, NumResults));
+		}
+	}
+}
+
+const TArray<FBoxSphereBounds>* FHierarchicalStaticMeshSceneProxy::GetOcclusionQueries(const FSceneView* View) const 
+{
+	return &OcclusionBounds;
+}
+
 FBoxSphereBounds UHierarchicalInstancedStaticMeshComponent::CalcBounds(const FTransform& BoundTransform) const
 {
 	const TArray<FClusterNode>& ClusterTree = *ClusterTreePtr;
@@ -1360,7 +1539,7 @@ UHierarchicalInstancedStaticMeshComponent::~UHierarchicalInstancedStaticMeshComp
 	if (ProxySize)
 	{
 		DEC_DWORD_STAT_BY(STAT_FoliageInstanceBuffers, ProxySize);
-}
+	}
 	ProxySize = 0;
 	FlushAsyncBuildInstanceBufferTask();
 }
@@ -1475,7 +1654,7 @@ bool UHierarchicalInstancedStaticMeshComponent::UpdateInstanceTransform(int32 In
 
 	if (StaticMesh)
 	{
-	UnbuiltInstanceBounds += StaticMesh->GetBounds().GetBox().TransformBy(NewInstanceTransform);
+		UnbuiltInstanceBounds += StaticMesh->GetBounds().GetBox().TransformBy(NewInstanceTransform);
 	}
 
 	if (!IsAsyncBuilding())
@@ -1630,6 +1809,7 @@ void UHierarchicalInstancedStaticMeshComponent::BuildTree()
 		Builder->Build();
 
 		NumBuiltInstances = Builder->Result->InstanceReorderTable.Num();
+		OcclusionLayerNumNodes = Builder->Result->OutOcclusionLayerNum;
 		UnbuiltInstanceBounds.Init();
 		RemovedInstances.Empty();
 
@@ -1650,11 +1830,13 @@ void UHierarchicalInstancedStaticMeshComponent::BuildTreeAnyThread(
 	TArray<FClusterNode>& OutClusterTree,
 	TArray<int32>& OutSortedInstances,
 	TArray<int32>& OutInstanceReorderTable,
+	int32& OutOcclusionLayerNum,
 	int32 MaxInstancesPerLeaf
 	)
 {
 	TUniquePtr<FClusterBuilder> Builder(new FClusterBuilder(InstanceTransforms, MeshBox, MaxInstancesPerLeaf));
 	Builder->Build();
+	OutOcclusionLayerNum = Builder->Result->OutOcclusionLayerNum;
 
 	Exchange(OutClusterTree, Builder->Result->Nodes);
 	Exchange(OutInstanceReorderTable, Builder->Result->InstanceReorderTable);
@@ -1662,7 +1844,7 @@ void UHierarchicalInstancedStaticMeshComponent::BuildTreeAnyThread(
 
 }
 
-void UHierarchicalInstancedStaticMeshComponent::AcceptPrebuiltTree(TArray<FClusterNode>& InClusterTree)
+void UHierarchicalInstancedStaticMeshComponent::AcceptPrebuiltTree(TArray<FClusterNode>& InClusterTree, int InOcclusionLayerNumNodes)
 {
 	QUICK_SCOPE_CYCLE_COUNTER(STAT_UHierarchicalInstancedStaticMeshComponent_AcceptPrebuiltTree);
 	// this is only for prebuild data, already in the correct order
@@ -1674,6 +1856,7 @@ void UHierarchicalInstancedStaticMeshComponent::AcceptPrebuiltTree(TArray<FClust
 	ClusterTreePtr = MakeShareable(new TArray<FClusterNode>);
 	InstanceReorderTable.Empty();
 	SortedInstances.Empty();
+	OcclusionLayerNumNodes = InOcclusionLayerNumNodes;
 
 	// Verify that the mesh is valid before using it.
 	const bool bMeshIsValid = 
@@ -1791,6 +1974,7 @@ void UHierarchicalInstancedStaticMeshComponent::BuildFlatTree(const TArray<int32
 			InstanceReorderTable.Add(Index);
 			SortedInstances.Add(Index);
 		}
+		OcclusionLayerNumNodes = 0;
 
 		FlushAccumulatedNavigationUpdates();
 
@@ -1832,6 +2016,8 @@ void UHierarchicalInstancedStaticMeshComponent::ApplyBuildTreeAsync(ENamedThread
 		Exchange(InstanceReorderTable, Builder->Result->InstanceReorderTable);
 		Exchange(SortedInstances, Builder->Result->SortedInstances);
 		RemovedInstances.Empty();
+		OcclusionLayerNumNodes = Builder->Result->OutOcclusionLayerNum;
+
 
 	    //UE_LOG(LogStaticMesh, Display, TEXT("Built a foliage hierarchy with %d of %d elements in %.1fs."), NumBuiltInstances, PerInstanceSMData.Num(), float(FPlatformTime::Seconds() - StartTime));
     
@@ -1997,26 +2183,6 @@ void FAsyncBuildInstanceBuffer::DoWork()
 void UHierarchicalInstancedStaticMeshComponent::PostLoad()
 {
 	Super::PostLoad();
-#if WITH_EDITORONLY_DATA
-	if (CVarRebuildTreeOnLoad.GetValueOnGameThread() > 0 && PerInstanceSMData.Num() > 3 && StaticMesh)
-	{
-		StaticMesh->ConditionalPostLoad();
-		if (GetVertsForLOD(0))
-		{
-			float Actual = ActualInstancesPerLeaf();
-			float Desired = (float)FMath::Min(DesiredInstancesPerLeaf(), PerInstanceSMData.Num());
-			float DesiredCapped = FMath::Min(Desired, (float)PerInstanceSMData.Num());
-			if (Actual > Desired * 2.5f || Actual * 2.5f < DesiredCapped)
-			{
-				UE_LOG(LogStaticMesh, Warning, TEXT("Rebuilding foliage, please resave map %s."), *GetFullName());
-
-				BuildTree();
-				float NewActual = ActualInstancesPerLeaf();
-				ensureMsgf(!(NewActual > Desired * 2.5f || NewActual * 2.5f < DesiredCapped), TEXT("HISMC rebuild did not hit target %f %f %f"), NewActual, Desired, DesiredCapped); // if this fires, it means this will always rebuild
-			}
-		}
-	}
-#endif
 	if (CVarASyncInstaneBufferConversion.GetValueOnGameThread() > 0)
 	{
 		UWorld* World = GetWorld();
@@ -2031,10 +2197,12 @@ void UHierarchicalInstancedStaticMeshComponent::PostLoad()
 			AsyncBuildInstanceBufferTask->StartBackgroundTask();
 		}
 	}
-	// For some reason we don't have a tree. Build one now!
-	if (PerInstanceSMData.Num() > 0 && (!ClusterTreePtr.IsValid() || ClusterTreePtr->Num() == 0))
+	// For some reason we don't have a tree, or it is out of date. Build one now!
+	if (PerInstanceSMData.Num() > 0 && (!ClusterTreePtr.IsValid() || ClusterTreePtr->Num() == 0 || GetLinkerUE4Version() < VER_UE4_REBUILD_HIERARCHICAL_INSTANCE_TREES))
 	{
+		UE_LOG(LogStaticMesh, Warning, TEXT("Rebuilding foliage, please resave map %s."), *GetFullName());
 		check(!IsAsyncBuilding());
+		StaticMesh->ConditionalPostLoad();
 		BuildTree();
 	}
 }
@@ -2044,21 +2212,21 @@ static void GatherInstanceTransformsInArea(const UHierarchicalInstancedStaticMes
 	const TArray<FClusterNode>& ClusterTree = *Component.ClusterTreePtr;
 	if (ClusterTree.Num())
 	{
-		const FClusterNode& ChildNode = ClusterTree[Child];
-		const FBox WorldNodeBox = FBox(ChildNode.BoundMin, ChildNode.BoundMax).TransformBy(Component.ComponentToWorld);
-
-		if (AreaBox.Intersect(WorldNodeBox))
+	const FClusterNode& ChildNode = ClusterTree[Child];
+	const FBox WorldNodeBox = FBox(ChildNode.BoundMin, ChildNode.BoundMax).TransformBy(Component.ComponentToWorld);
+	
+	if (AreaBox.Intersect(WorldNodeBox))
+	{
+		if (ChildNode.FirstChild < 0 || AreaBox.IsInside(WorldNodeBox))
 		{
-			if (ChildNode.FirstChild < 0 || AreaBox.IsInside(WorldNodeBox))
+			// Unfortunately ordering of PerInstanceSMData does not match ordering of cluster tree, so we have to use remaping
+			const bool bUseRemaping = Component.SortedInstances.Num() > 0;
+			
+			// In case there no more subdivision or node is completely encapsulated by a area box
+			// add all instances to the result
+			for (int32 i = ChildNode.FirstInstance; i <= ChildNode.LastInstance; ++i)
 			{
-				// Unfortunately ordering of PerInstanceSMData does not match ordering of cluster tree, so we have to use remaping
-				const bool bUseRemaping = Component.SortedInstances.Num() > 0;
-
-				// In case there no more subdivision or node is completely encapsulated by a area box
-				// add all instances to the result
-				for (int32 i = ChildNode.FirstInstance; i <= ChildNode.LastInstance; ++i)
-				{
-					int32 SortedIdx = bUseRemaping ? Component.SortedInstances[i] : i;
+				int32 SortedIdx = bUseRemaping ? Component.SortedInstances[i] : i;
 
 					FTransform InstanceToComponent;
 					if (Component.PerInstanceSMData.IsValidIndex(SortedIdx))
@@ -2073,22 +2241,22 @@ static void GatherInstanceTransformsInArea(const UHierarchicalInstancedStaticMes
 						Inst->GetInstanceTransform(XformMat);
 						InstanceToComponent = FTransform(XformMat);
 					}
-					if (!InstanceToComponent.GetScale3D().IsZero())
-					{
-						InstanceData.Add(InstanceToComponent*Component.ComponentToWorld);
-					}
+				if (!InstanceToComponent.GetScale3D().IsZero())
+				{
+					InstanceData.Add(InstanceToComponent*Component.ComponentToWorld);
 				}
 			}
-			else
+		}
+		else
+		{
+			for (int32 i = ChildNode.FirstChild; i <= ChildNode.LastChild; ++i)
 			{
-				for (int32 i = ChildNode.FirstChild; i <= ChildNode.LastChild; ++i)
-				{
-					GatherInstanceTransformsInArea(Component, AreaBox, i, InstanceData);
-				}
+				GatherInstanceTransformsInArea(Component, AreaBox, i, InstanceData);
 			}
 		}
 	}
 }
+}
 
 int32 UHierarchicalInstancedStaticMeshComponent::GetOverlappingSphereCount(const FSphere& Sphere) const
 {
diff --git a/Engine/Source/Runtime/Engine/Private/PrimitiveComponent.cpp b/Engine/Source/Runtime/Engine/Private/PrimitiveComponent.cpp
--- a/Engine/Source/Runtime/Engine/Private/PrimitiveComponent.cpp
+++ b/Engine/Source/Runtime/Engine/Private/PrimitiveComponent.cpp
@@ -66,7 +66,7 @@ TAutoConsoleVariable<int32> CVarShowInitialOverlaps(
 int32 UPrimitiveComponent::CurrentTag = 2147483647 / 4;
 
 // 0 is reserved to mean invalid
-uint64 UPrimitiveComponent::NextComponentId = 1;
+uint32 UPrimitiveComponent::NextComponentId = 1;
 
 UPrimitiveComponent::UPrimitiveComponent(const FObjectInitializer& ObjectInitializer /*= FObjectInitializer::Get()*/)
 	: Super(ObjectInitializer)
@@ -99,7 +99,7 @@ UPrimitiveComponent::UPrimitiveComponent(const FObjectInitializer& ObjectInitial
 	VisibilityId = -1;
 	CanBeCharacterBase_DEPRECATED = ECB_Yes;
 	CanCharacterStepUpOn = ECB_Yes;
-	ComponentId.Value = NextComponentId;
+	ComponentId.PrimIDValue = NextComponentId;
 	check(IsInGameThread());
 	NextComponentId++;
 
diff --git a/Engine/Source/Runtime/Engine/Private/PrimitiveDrawingUtils.cpp b/Engine/Source/Runtime/Engine/Private/PrimitiveDrawingUtils.cpp
--- a/Engine/Source/Runtime/Engine/Private/PrimitiveDrawingUtils.cpp
+++ b/Engine/Source/Runtime/Engine/Private/PrimitiveDrawingUtils.cpp
@@ -1438,7 +1438,7 @@ void ApplyViewModeOverrides(
 		{
 			// Determine unique color for model component.
 			FLinearColor BSPSplitColor;
-			FRandomStream RandomStream(GetTypeHash(PrimitiveSceneProxy->GetPrimitiveComponentId().Value));
+			FRandomStream RandomStream(GetTypeHash(PrimitiveSceneProxy->GetPrimitiveComponentId().PrimIDValue));
 			BSPSplitColor.R = RandomStream.GetFraction();
 			BSPSplitColor.G = RandomStream.GetFraction();
 			BSPSplitColor.B = RandomStream.GetFraction();
diff --git a/Engine/Source/Runtime/Engine/Private/SceneView.cpp b/Engine/Source/Runtime/Engine/Private/SceneView.cpp
--- a/Engine/Source/Runtime/Engine/Private/SceneView.cpp
+++ b/Engine/Source/Runtime/Engine/Private/SceneView.cpp
@@ -526,6 +526,15 @@ float FSceneView::GetTemporalLODTransition() const
 	return 0.0f;
 }
 
+uint32 FSceneView::GetViewKey() const
+{
+	if (State)
+	{
+		return State->GetViewKey();
+	}
+	return 0;
+}
+
 
 void FSceneView::UpdateViewMatrix()
 {
diff --git a/Engine/Source/Runtime/Engine/Public/PrimitiveSceneProxy.h b/Engine/Source/Runtime/Engine/Public/PrimitiveSceneProxy.h
--- a/Engine/Source/Runtime/Engine/Public/PrimitiveSceneProxy.h
+++ b/Engine/Source/Runtime/Engine/Public/PrimitiveSceneProxy.h
@@ -187,6 +187,24 @@ public:
 	 */
 	virtual void GetDynamicMeshElements(const TArray<const FSceneView*>& Views, const FSceneViewFamily& ViewFamily, uint32 VisibilityMap, class FMeshElementCollector& Collector) const {}
 
+	/** 
+	 * Gets the boxes for sub occlusion queries
+	 * @param View - the view the occlusion results are for
+	 * @return pointer to the boxes, must remain valid until the queries are built
+	 */
+	virtual const TArray<FBoxSphereBounds>* GetOcclusionQueries(const FSceneView* View) const 
+	{
+		return nullptr;
+	}
+
+	/** 
+	 * Gives the primitive the results of sub-occlusion-queries
+	 * @param View - the view the occlusion results are for
+	 * @param Results - visibility results, allocated from the scene allocator, so valid until the end of the frame
+	 * @param NumResults - number of visibility bools
+	 */
+	virtual void AcceptOcclusionResults(const FSceneView* View, const bool* Results, int32 NumResults) {}
+
 	/**
 	 * Determines the relevance of this primitive's elements to the given view.
 	 * Called in the rendering thread.
@@ -279,6 +297,14 @@ public:
 		return true;
 	}
 
+	/**
+	* @return true if the proxy has custom occlusion queries
+	*/
+	virtual bool HasSubprimitiveOcclusionQueries() const
+	{
+		return false;
+	}
+
 	virtual bool ShowInBSPSplitViewmode() const
 	{
 		return false;
diff --git a/Engine/Source/Runtime/Engine/Public/SceneManagement.h b/Engine/Source/Runtime/Engine/Public/SceneManagement.h
--- a/Engine/Source/Runtime/Engine/Public/SceneManagement.h
+++ b/Engine/Source/Runtime/Engine/Public/SceneManagement.h
@@ -143,6 +143,10 @@ public:
 	 * Returns the blend factor between the last two LOD samples
 	 */
 	virtual float GetTemporalLODTransition() const = 0;
+	/** 
+	 * returns a unique key for the view state, non-zero
+	 */
+	virtual uint32 GetViewKey() const = 0;
 protected:
 	// Don't allow direct deletion of the view state, Destroy should be called instead.
 	virtual ~FSceneViewStateInterface() {}
diff --git a/Engine/Source/Runtime/Engine/Public/SceneTypes.h b/Engine/Source/Runtime/Engine/Public/SceneTypes.h
--- a/Engine/Source/Runtime/Engine/Public/SceneTypes.h
+++ b/Engine/Source/Runtime/Engine/Public/SceneTypes.h
@@ -22,25 +22,25 @@ class FPrimitiveComponentId
 {
 public:
 
-	FPrimitiveComponentId() : Value(0)
+	FPrimitiveComponentId() : PrimIDValue(0)
 	{}
 
 	inline bool IsValid() const
 	{
-		return Value > 0;
+		return PrimIDValue > 0;
 	}
 
 	inline bool operator==(FPrimitiveComponentId OtherId) const
 	{
-		return Value == OtherId.Value;
+		return PrimIDValue == OtherId.PrimIDValue;
 	}
 
 	friend uint32 GetTypeHash( FPrimitiveComponentId Id )
 	{
-		return GetTypeHash(Id.Value);
+		return GetTypeHash(Id.PrimIDValue);
 	}
 
-	uint64 Value;
+	uint32 PrimIDValue;
 };
 
 /** 
diff --git a/Engine/Source/Runtime/Engine/Public/SceneView.h b/Engine/Source/Runtime/Engine/Public/SceneView.h
--- a/Engine/Source/Runtime/Engine/Public/SceneView.h
+++ b/Engine/Source/Runtime/Engine/Public/SceneView.h
@@ -635,6 +635,11 @@ public:
 	 */
 	float GetTemporalLODTransition() const;
 
+	/** 
+	 * returns a unique key for the view state if one exists, otherwise returns zero
+	 */
+	uint32 GetViewKey() const;
+
 	/** Allow things like HMD displays to update the view matrix at the last minute, to minimize perceived latency */
 	void UpdateViewMatrix();
 
diff --git a/Engine/Source/Runtime/Landscape/Private/LandscapeGrass.cpp b/Engine/Source/Runtime/Landscape/Private/LandscapeGrass.cpp
--- a/Engine/Source/Runtime/Landscape/Private/LandscapeGrass.cpp
+++ b/Engine/Source/Runtime/Landscape/Private/LandscapeGrass.cpp
@@ -970,6 +970,7 @@ struct FAsyncGrassBuilder : public FGrassBuilderBase
 	// output
 	FStaticMeshInstanceData InstanceBuffer;
 	TArray<FClusterNode> ClusterTree;
+	int32 OutOcclusionLayerNum;
 
 	FAsyncGrassBuilder(ALandscapeProxy* Landscape, ULandscapeComponent* Component, const ULandscapeGrassType* GrassType, const FGrassVariety& GrassVariety, UHierarchicalInstancedStaticMeshComponent* HierarchicalInstancedStaticMeshComponent, int32 SqrtSubsections, int32 SubX, int32 SubY, uint32 InHaltonBaseIndex)
 		: FGrassBuilderBase(Landscape, Component, GrassVariety, SqrtSubsections, SubX, SubY)
@@ -988,6 +989,7 @@ struct FAsyncGrassBuilder : public FGrassBuilderBase
 
 		TotalInstances = 0;
 		HaltonBaseIndex = InHaltonBaseIndex;
+		OutOcclusionLayerNum = 0;
 	}
 
 	void Build()
@@ -1166,7 +1168,7 @@ struct FAsyncGrassBuilder : public FGrassBuilderBase
 		{
 			TArray<int32> SortedInstances;
 			TArray<int32> InstanceReorderTable;
-			UHierarchicalInstancedStaticMeshComponent::BuildTreeAnyThread(InstanceTransforms, MeshBox, ClusterTree, SortedInstances, InstanceReorderTable, DesiredInstancesPerLeaf);
+			UHierarchicalInstancedStaticMeshComponent::BuildTreeAnyThread(InstanceTransforms, MeshBox, ClusterTree, SortedInstances, InstanceReorderTable, DesiredInstancesPerLeaf, OutOcclusionLayerNum);
 
 			// in-place sort the instances
 			FInstanceStream SwapBuffer;
@@ -1807,7 +1809,7 @@ void ALandscapeProxy::UpdateGrass(const TArray<FVector>& Cameras, bool bForceSyn
 					{
 						QUICK_SCOPE_CYCLE_COUNTER(STAT_FoliageGrassEndComp_AcceptPrebuiltTree);
 						FMemory::Memswap(&HierarchicalInstancedStaticMeshComponent->WriteOncePrebuiltInstanceBuffer, &Inner.Builder->InstanceBuffer, sizeof(FStaticMeshInstanceData));
-						HierarchicalInstancedStaticMeshComponent->AcceptPrebuiltTree(Inner.Builder->ClusterTree);
+						HierarchicalInstancedStaticMeshComponent->AcceptPrebuiltTree(Inner.Builder->ClusterTree, Inner.Builder->OutOcclusionLayerNum);
 						if (bForceSync && GetWorld())
 						{
 							QUICK_SCOPE_CYCLE_COUNTER(STAT_FoliageGrassEndComp_SyncUpdate);
diff --git a/Engine/Source/Runtime/Renderer/Private/PrimitiveSceneInfo.cpp b/Engine/Source/Runtime/Renderer/Private/PrimitiveSceneInfo.cpp
--- a/Engine/Source/Runtime/Renderer/Private/PrimitiveSceneInfo.cpp
+++ b/Engine/Source/Runtime/Renderer/Private/PrimitiveSceneInfo.cpp
@@ -221,6 +221,10 @@ void FPrimitiveSceneInfo::AddToScene(FRHICommandListImmediate& RHICmdList, bool
 	{
 		OcclusionFlags |= EOcclusionFlags::CanBeOccluded;
 	}
+	if (Proxy->HasSubprimitiveOcclusionQueries())
+	{
+		OcclusionFlags |= EOcclusionFlags::HasSubprimitiveQueries;
+	}
 	if (Proxy->AllowApproximateOcclusion()
 		// Allow approximate occlusion if attached, even if the parent does not have bLightAttachmentsAsGroup enabled
 		|| LightingAttachmentRoot.IsValid())
diff --git a/Engine/Source/Runtime/Renderer/Private/RendererScene.cpp b/Engine/Source/Runtime/Renderer/Private/RendererScene.cpp
--- a/Engine/Source/Runtime/Renderer/Private/RendererScene.cpp
+++ b/Engine/Source/Runtime/Renderer/Private/RendererScene.cpp
@@ -30,6 +30,8 @@ TGlobalResource< FGlobalDistanceCullFadeUniformBuffer > GDistanceCullFadedInUnif
 
 SIZE_T FStaticMeshDrawListBase::TotalBytesUsed = 0;
 
+static FThreadSafeCounter FSceneViewState_UniqueID;
+
 /**
  * Holds the info to update SpeedTree wind per unique tree object in the scene, instead of per instance
  */
@@ -54,6 +56,7 @@ struct FSpeedTreeWindComputation
 FSceneViewState::FSceneViewState()
 	: OcclusionQueryPool(RQT_Occlusion)
 {
+	UniqueID = FSceneViewState_UniqueID.Increment();
 	OcclusionFrameCounter = 0;
 	LastRenderTime = -FLT_MAX;
 	LastRenderTimeDelta = 0.0f;
@@ -919,8 +922,8 @@ void FScene::SetSkyLight(FSkyLightSceneProxy* LightProxy)
 
 		if (!bHadSkylight)
 		{
-			// Mark the scene as needing static draw lists to be recreated if needed
-			// The base pass chooses shaders based on whether there's a skylight in the scene, and that is cached in static draw lists
+		// Mark the scene as needing static draw lists to be recreated if needed
+		// The base pass chooses shaders based on whether there's a skylight in the scene, and that is cached in static draw lists
 			Scene->bScenesPrimitivesNeedStaticMeshElementUpdate = true;
 		}
 	});
@@ -1707,7 +1710,7 @@ void FScene::UpdateSpeedTreeWind(double CurrentTime)
 				WindComputation->UniformBuffer.SetContents(UniformParameters);
 			}
 		});
-#undef SET_SPEEDTREE_TABLE_FLOAT4V
+	#undef SET_SPEEDTREE_TABLE_FLOAT4V
 }
 
 FUniformBufferRHIParamRef FScene::GetSpeedTreeUniformBuffer(const FVertexFactory* VertexFactory)
diff --git a/Engine/Source/Runtime/Renderer/Private/ScenePrivate.h b/Engine/Source/Runtime/Renderer/Private/ScenePrivate.h
--- a/Engine/Source/Runtime/Renderer/Private/ScenePrivate.h
+++ b/Engine/Source/Runtime/Renderer/Private/ScenePrivate.h
@@ -158,8 +158,13 @@ public:
 	 */
 	int32 NumBufferedFrames;
 
+	/** 
+	 * For things that have subqueries (folaige), this is the non-zero
+	 */
+	int32 CustomIndex;
+
 	/** Initialization constructor. */
-	FORCEINLINE FPrimitiveOcclusionHistory(FPrimitiveComponentId InPrimitiveId = FPrimitiveComponentId())
+	FORCEINLINE FPrimitiveOcclusionHistory(FPrimitiveComponentId InPrimitiveId, int32 SubQuery)
 		: PrimitiveId(InPrimitiveId)
 		, HZBTestIndex(0)
 		, HZBTestFrameNumber(~0u)
@@ -167,6 +172,7 @@ public:
 		, LastConsideredTime(0.0f)
 		, LastPixelsPercentage(0.0f)
 		, bGroupedQuery(false)
+		, CustomIndex(SubQuery)
 	{
 #if BUFFERED_OCCLUSION_QUERIES
 		NumBufferedFrames = FOcclusionQueryHelpers::GetNumBufferedFrames();
@@ -217,22 +223,41 @@ public:
 	}
 };
 
+struct FPrimitiveOcclusionHistoryKey
+{
+	FPrimitiveComponentId PrimitiveId;
+	int32 CustomIndex;
+
+	FPrimitiveOcclusionHistoryKey(const FPrimitiveOcclusionHistory& Element)
+		: PrimitiveId(Element.PrimitiveId)
+		, CustomIndex(Element.CustomIndex)
+	{
+	}
+	FPrimitiveOcclusionHistoryKey(FPrimitiveComponentId InPrimitiveId, int32 InCustomIndex)
+		: PrimitiveId(InPrimitiveId)
+		, CustomIndex(InCustomIndex)
+	{
+	}
+};
+
 /** Defines how the hash set indexes the FPrimitiveOcclusionHistory objects. */
-struct FPrimitiveOcclusionHistoryKeyFuncs : BaseKeyFuncs<FPrimitiveOcclusionHistory,FPrimitiveComponentId>
+struct FPrimitiveOcclusionHistoryKeyFuncs : BaseKeyFuncs<FPrimitiveOcclusionHistory,FPrimitiveOcclusionHistoryKey>
 {
+	typedef FPrimitiveOcclusionHistoryKey KeyInitType;
+
 	static KeyInitType GetSetKey(const FPrimitiveOcclusionHistory& Element)
 	{
-		return Element.PrimitiveId;
+		return FPrimitiveOcclusionHistoryKey(Element);
 	}
 
 	static bool Matches(KeyInitType A,KeyInitType B)
 	{
-		return A == B;
+		return A.PrimitiveId == B.PrimitiveId && A.CustomIndex == B.CustomIndex;
 	}
 
 	static uint32 GetKeyHash(KeyInitType Key)
 	{
-		return GetTypeHash(Key.Value);
+		return GetTypeHash(Key.PrimitiveId.PrimIDValue) ^ (GetTypeHash(Key.CustomIndex) >> 20);
 	}
 };
 
@@ -404,6 +429,7 @@ public:
 	};
 
 	int32 NumBufferedFrames;
+	uint32 UniqueID;
 	typedef TMap<FSceneViewState::FProjectedShadowKey, FRenderQueryRHIRef> ShadowKeyOcclusionQueryMap;
 #if BUFFERED_OCCLUSION_QUERIES
 	TArray<ShadowKeyOcclusionQueryMap> ShadowOcclusionQueryMaps;
@@ -745,13 +771,14 @@ public:
 	{
 		return TemporalLODState;
 	}
-	/** 
-	 * Returns the blend factor between the last two LOD samples
-	 */
 	float GetTemporalLODTransition() const override
 	{
 		return TemporalLODState.GetTemporalLODTransition(LastRenderTime);
 	}
+	uint32 GetViewKey() const override
+	{
+		return UniqueID;
+	}
 
 
 
@@ -1221,6 +1248,8 @@ namespace EOcclusionFlags
 		AllowApproximateOcclusion = 0x4,
 		/** Indicates the primitive has a valid ID for precomputed visibility. */
 		HasPrecomputedVisibility = 0x8,
+		/** Indicates the primitive has a valid ID for precomputed visibility. */
+		HasSubprimitiveQueries = 0x10,
 	};
 };
 
diff --git a/Engine/Source/Runtime/Renderer/Private/SceneVisibility.cpp b/Engine/Source/Runtime/Renderer/Private/SceneVisibility.cpp
--- a/Engine/Source/Runtime/Renderer/Private/SceneVisibility.cpp
+++ b/Engine/Source/Runtime/Renderer/Private/SceneVisibility.cpp
@@ -77,6 +77,15 @@ static FAutoConsoleVariableRef CVarVisualizeOccludedPrimitives(
 	ECVF_RenderThreadSafe | ECVF_Cheat
 	);
 
+static int32 GAllowSubPrimitiveQueries = 1;
+static FAutoConsoleVariableRef CVarAllowSubPrimitiveQueries(
+	TEXT("r.AllowSubPrimitiveQueries"),
+	GAllowSubPrimitiveQueries,
+	TEXT("Enables sub primitive queries, currently only used by hierarchical instanced static meshes. 1: Enable, 0 Disabled. When disabled, one query is used for the entire proxy."),
+	ECVF_RenderThreadSafe
+	);
+
+
 static TAutoConsoleVariable<int32> CVarLightShaftQuality(
 	TEXT("r.LightShaftQuality"),
 	1,
@@ -408,201 +417,259 @@ static int32 OcclusionCull(FRHICommandListImmediate& RHICmdList, const FScene* S
 						bCanBeOccluded = false;
 					}
 				}
+				int32 NumSubQueries = 1;
+				bool bSubQueries = false;
+				const TArray<FBoxSphereBounds>* SubBounds = nullptr;
+				TArray<bool, SceneRenderingAllocator> SubIsOccluded;
 
-				FPrimitiveComponentId PrimitiveId = Scene->PrimitiveComponentIds[BitIt.GetIndex()];
-				FPrimitiveOcclusionHistory* PrimitiveOcclusionHistory = ViewState->PrimitiveOcclusionHistorySet.Find(PrimitiveId);
-				bool bIsOccluded = false;
-				bool bOcclusionStateIsDefinite = false;
-				if (!PrimitiveOcclusionHistory)
+				if ((OcclusionFlags & EOcclusionFlags::HasSubprimitiveQueries) && GAllowSubPrimitiveQueries)
 				{
-					// If the primitive doesn't have an occlusion history yet, create it.
-					PrimitiveOcclusionHistory = &ViewState->PrimitiveOcclusionHistorySet[
-						ViewState->PrimitiveOcclusionHistorySet.Add(FPrimitiveOcclusionHistory(PrimitiveId))
-						];
-
-					// If the primitive hasn't been visible recently enough to have a history, treat it as unoccluded this frame so it will be rendered as an occluder and its true occlusion state can be determined.
-					// already set bIsOccluded = false;
-
-					// Flag the primitive's occlusion state as indefinite, which will force it to be queried this frame.
-					// The exception is if the primitive isn't occludable, in which case we know that it's definitely unoccluded.
-					bOcclusionStateIsDefinite = bCanBeOccluded ? false : true;
+					FPrimitiveSceneProxy* Proxy = Scene->Primitives[BitIt.GetIndex()]->Proxy;
+					SubBounds = Proxy->GetOcclusionQueries(&View);
+					NumSubQueries = SubBounds->Num();
+					bSubQueries = true;
+					if (!NumSubQueries)
+					{
+						View.PrimitiveVisibilityMap.AccessCorrespondingBit(BitIt) = false;
+						continue;
+					}
+					SubIsOccluded.Reserve(NumSubQueries);
 				}
-				else
+
+				bool bAllSubOcclusionStateIsDefinite = true;
+				bool bAllSubOccluded = true;
+				FPrimitiveComponentId PrimitiveId = Scene->PrimitiveComponentIds[BitIt.GetIndex()];
+				for (int32 SubQuery = 0; SubQuery < NumSubQueries; SubQuery++)
 				{
-					if (View.bIgnoreExistingQueries)
+
+					FPrimitiveOcclusionHistory* PrimitiveOcclusionHistory = ViewState->PrimitiveOcclusionHistorySet.Find(FPrimitiveOcclusionHistoryKey(PrimitiveId, SubQuery));
+					bool bIsOccluded = false;
+					bool bOcclusionStateIsDefinite = false;
+					if (!PrimitiveOcclusionHistory)
 					{
-						// If the view is ignoring occlusion queries, the primitive is definitely unoccluded.
+						// If the primitive doesn't have an occlusion history yet, create it.
+						PrimitiveOcclusionHistory = &ViewState->PrimitiveOcclusionHistorySet[
+							ViewState->PrimitiveOcclusionHistorySet.Add(FPrimitiveOcclusionHistory(PrimitiveId, SubQuery))
+							];
+
+						// If the primitive hasn't been visible recently enough to have a history, treat it as unoccluded this frame so it will be rendered as an occluder and its true occlusion state can be determined.
 						// already set bIsOccluded = false;
-						bOcclusionStateIsDefinite = View.bDisableQuerySubmissions;
+
+						// Flag the primitive's occlusion state as indefinite, which will force it to be queried this frame.
+						// The exception is if the primitive isn't occludable, in which case we know that it's definitely unoccluded.
+						bOcclusionStateIsDefinite = bCanBeOccluded ? false : true;
 					}
-					else if (bCanBeOccluded)
+					else
 					{
-						if( bHZBOcclusion )
+						if (View.bIgnoreExistingQueries)
 						{
-							if( ViewState->HZBOcclusionTests.IsValidFrame(PrimitiveOcclusionHistory->HZBTestFrameNumber) )
-							{
-								bIsOccluded = !ViewState->HZBOcclusionTests.IsVisible( PrimitiveOcclusionHistory->HZBTestIndex );
-								bOcclusionStateIsDefinite = true;
-							}
+							// If the view is ignoring occlusion queries, the primitive is definitely unoccluded.
+							// already set bIsOccluded = false;
+							bOcclusionStateIsDefinite = View.bDisableQuerySubmissions;
 						}
-						else
+						else if (bCanBeOccluded)
 						{
-							// Read the occlusion query results.
-							uint64 NumSamples = 0;
-							FRenderQueryRHIRef& PastQuery = PrimitiveOcclusionHistory->GetPastQuery(ViewState->OcclusionFrameCounter);
-							if (IsValidRef(PastQuery))
+							if( bHZBOcclusion )
 							{
-								// NOTE: RHIGetOcclusionQueryResult should never fail when using a blocking call, rendering artifacts may show up.
-								if (RHICmdList.GetRenderQueryResult(PastQuery, NumSamples, true))
+								if( ViewState->HZBOcclusionTests.IsValidFrame(PrimitiveOcclusionHistory->HZBTestFrameNumber) )
 								{
-									// we render occlusion without MSAA
-									uint32 NumPixels = (uint32)NumSamples;
-
-									// The primitive is occluded if none of its bounding box's pixels were visible in the previous frame's occlusion query.
-									bIsOccluded = (NumPixels == 0);
-									if (!bIsOccluded)
+									bIsOccluded = !ViewState->HZBOcclusionTests.IsVisible( PrimitiveOcclusionHistory->HZBTestIndex );
+									bOcclusionStateIsDefinite = true;
+								}
+							}
+							else
+							{
+								// Read the occlusion query results.
+								uint64 NumSamples = 0;
+								FRenderQueryRHIRef& PastQuery = PrimitiveOcclusionHistory->GetPastQuery(ViewState->OcclusionFrameCounter);
+								if (IsValidRef(PastQuery))
+								{
+									// NOTE: RHIGetOcclusionQueryResult should never fail when using a blocking call, rendering artifacts may show up.
+									if (RHICmdList.GetRenderQueryResult(PastQuery, NumSamples, true))
 									{
-										checkSlow(View.OneOverNumPossiblePixels > 0.0f);
-										PrimitiveOcclusionHistory->LastPixelsPercentage = NumPixels * View.OneOverNumPossiblePixels;
+										// we render occlusion without MSAA
+										uint32 NumPixels = (uint32)NumSamples;
+
+										// The primitive is occluded if none of its bounding box's pixels were visible in the previous frame's occlusion query.
+										bIsOccluded = (NumPixels == 0);
+										if (!bIsOccluded)
+										{
+											checkSlow(View.OneOverNumPossiblePixels > 0.0f);
+											PrimitiveOcclusionHistory->LastPixelsPercentage = NumPixels * View.OneOverNumPossiblePixels;
+										}
+										else
+										{
+											PrimitiveOcclusionHistory->LastPixelsPercentage = 0.0f;
+										}
+
+
+										// Flag the primitive's occlusion state as definite if it wasn't grouped.
+										bOcclusionStateIsDefinite = !PrimitiveOcclusionHistory->bGroupedQuery;
 									}
 									else
 									{
-										PrimitiveOcclusionHistory->LastPixelsPercentage = 0.0f;
+										// If the occlusion query failed, treat the primitive as visible.  
+										// already set bIsOccluded = false;
 									}
-
-
-									// Flag the primitive's occlusion state as definite if it wasn't grouped.
-									bOcclusionStateIsDefinite = !PrimitiveOcclusionHistory->bGroupedQuery;
 								}
 								else
 								{
-									// If the occlusion query failed, treat the primitive as visible.  
-									// already set bIsOccluded = false;
+									// If there's no occlusion query for the primitive, set it's visibility state to whether it has been unoccluded recently.
+									bIsOccluded = (PrimitiveOcclusionHistory->LastVisibleTime + GEngine->PrimitiveProbablyVisibleTime < CurrentRealTime);
+									if (bIsOccluded)
+									{
+										PrimitiveOcclusionHistory->LastPixelsPercentage = 0.0f;
+									}
+									else
+									{
+										PrimitiveOcclusionHistory->LastPixelsPercentage = GEngine->MaxOcclusionPixelsFraction;
+									}
+
+									// the state was definite last frame, otherwise we would have ran a query
+									bOcclusionStateIsDefinite = true;
 								}
 							}
-							else
-							{
-								// If there's no occlusion query for the primitive, set it's visibility state to whether it has been unoccluded recently.
-								bIsOccluded = (PrimitiveOcclusionHistory->LastVisibleTime + GEngine->PrimitiveProbablyVisibleTime < CurrentRealTime);
-								if (bIsOccluded)
-								{
-									PrimitiveOcclusionHistory->LastPixelsPercentage = 0.0f;
-								}
-								else
-								{
-									PrimitiveOcclusionHistory->LastPixelsPercentage = GEngine->MaxOcclusionPixelsFraction;
-								}
 
-								// the state was definite last frame, otherwise we would have ran a query
-								bOcclusionStateIsDefinite = true;
+							if( GVisualizeOccludedPrimitives && bIsOccluded )
+							{
+								const FBoxSphereBounds& Bounds = bSubQueries ? (*SubBounds)[SubQuery] : Scene->PrimitiveOcclusionBounds[ BitIt.GetIndex() ];
+								DrawWireBox( &OcclusionPDI, Bounds.GetBox(), FColor(50, 255, 50), SDPG_Foreground );
 							}
 						}
-
-						if( GVisualizeOccludedPrimitives && bIsOccluded )
+						else
 						{
-							const FBoxSphereBounds& Bounds = Scene->PrimitiveOcclusionBounds[ BitIt.GetIndex() ];
-							DrawWireBox( &OcclusionPDI, Bounds.GetBox(), FColor(50, 255, 50), SDPG_Foreground );
+							// Primitives that aren't occludable are considered definitely unoccluded.
+							// already set bIsOccluded = false;
+							bOcclusionStateIsDefinite = true;
 						}
-					}
-					else
-					{
-						// Primitives that aren't occludable are considered definitely unoccluded.
-						// already set bIsOccluded = false;
-						bOcclusionStateIsDefinite = true;
-					}
 
-					if (bClearQueries)
-					{
-						ViewState->OcclusionQueryPool.ReleaseQuery(RHICmdList, PrimitiveOcclusionHistory->GetPastQuery(ViewState->OcclusionFrameCounter));
+						if (bClearQueries)
+						{
+							ViewState->OcclusionQueryPool.ReleaseQuery(RHICmdList, PrimitiveOcclusionHistory->GetPastQuery(ViewState->OcclusionFrameCounter));
+						}
 					}
-				}
-
-				// Set the primitive's considered time to keep its occlusion history from being trimmed.
-				PrimitiveOcclusionHistory->LastConsideredTime = CurrentRealTime;
 
-				if (bSubmitQueries && bCanBeOccluded)
-				{
-					bool bAllowBoundsTest;
-					const FBoxSphereBounds& OcclusionBounds = Scene->PrimitiveOcclusionBounds[BitIt.GetIndex()];
-					if (View.bHasNearClippingPlane)
-					{
-						bAllowBoundsTest = View.NearClippingPlane.PlaneDot(OcclusionBounds.Origin) < 
-							-(FVector::BoxPushOut(View.NearClippingPlane,OcclusionBounds.BoxExtent));
-					}
-					else
-					{
-						bAllowBoundsTest = OcclusionBounds.SphereRadius < HALF_WORLD_MAX;
-					}
+					// Set the primitive's considered time to keep its occlusion history from being trimmed.
+					PrimitiveOcclusionHistory->LastConsideredTime = CurrentRealTime;
 
-					if (bAllowBoundsTest)
+					if (bSubmitQueries && bCanBeOccluded)
 					{
-						if( bHZBOcclusion )
+						bool bAllowBoundsTest;
+						const FBoxSphereBounds& OcclusionBounds = bSubQueries ? (*SubBounds)[SubQuery] : Scene->PrimitiveOcclusionBounds[ BitIt.GetIndex() ];
+						if (View.bHasNearClippingPlane)
 						{
-							// Always run
-							PrimitiveOcclusionHistory->HZBTestIndex = ViewState->HZBOcclusionTests.AddBounds( OcclusionBounds.Origin, OcclusionBounds.BoxExtent );
-							PrimitiveOcclusionHistory->HZBTestFrameNumber = ViewState->OcclusionFrameCounter;
+							bAllowBoundsTest = View.NearClippingPlane.PlaneDot(OcclusionBounds.Origin) < 
+								-(FVector::BoxPushOut(View.NearClippingPlane,OcclusionBounds.BoxExtent));
 						}
 						else
 						{
-							// decide if a query should be run this frame
-							bool bRunQuery,bGroupedQuery;
+							bAllowBoundsTest = OcclusionBounds.SphereRadius < HALF_WORLD_MAX;
+						}
 
-							if (OcclusionFlags & EOcclusionFlags::AllowApproximateOcclusion)
+						if (bAllowBoundsTest)
+						{
+							if( bHZBOcclusion )
 							{
-								if (bIsOccluded)
-								{
-									// Primitives that were occluded the previous frame use grouped queries.
-									bGroupedQuery = true;
-									bRunQuery = true;
-								}
-								else if (bOcclusionStateIsDefinite)
+								// Always run
+								PrimitiveOcclusionHistory->HZBTestIndex = ViewState->HZBOcclusionTests.AddBounds( OcclusionBounds.Origin, OcclusionBounds.BoxExtent );
+								PrimitiveOcclusionHistory->HZBTestFrameNumber = ViewState->OcclusionFrameCounter;
+							}
+							else
+							{
+								// decide if a query should be run this frame
+								bool bRunQuery,bGroupedQuery;
+
+								if (!bSubQueries && // sub queries are never grouped, we assume the custom code knows what it is doing and will group internally if it wants
+									(OcclusionFlags & EOcclusionFlags::AllowApproximateOcclusion))
 								{
-									// If the primitive's is definitely unoccluded, only requery it occasionally.
-									float FractionMultiplier = FMath::Max(PrimitiveOcclusionHistory->LastPixelsPercentage/GEngine->MaxOcclusionPixelsFraction, 1.0f);
-									bRunQuery = (FractionMultiplier * GOcclusionRandomStream.GetFraction() < GEngine->MaxOcclusionPixelsFraction);
-									bGroupedQuery = false;
+									if (bIsOccluded)
+									{
+										// Primitives that were occluded the previous frame use grouped queries.
+										bGroupedQuery = true;
+										bRunQuery = true;
+									}
+									else if (bOcclusionStateIsDefinite)
+									{
+										// If the primitive's is definitely unoccluded, only requery it occasionally.
+										float FractionMultiplier = FMath::Max(PrimitiveOcclusionHistory->LastPixelsPercentage/GEngine->MaxOcclusionPixelsFraction, 1.0f);
+										bRunQuery = (FractionMultiplier * GOcclusionRandomStream.GetFraction() < GEngine->MaxOcclusionPixelsFraction);
+										bGroupedQuery = false;
+									}
+									else
+									{
+										bGroupedQuery = false;
+										bRunQuery = true;
+									}
 								}
 								else
 								{
+									// Primitives that need precise occlusion results use individual queries.
 									bGroupedQuery = false;
 									bRunQuery = true;
 								}
+
+								if (bRunQuery)
+								{
+									PrimitiveOcclusionHistory->SetCurrentQuery(ViewState->OcclusionFrameCounter, 
+										bGroupedQuery ? 
+										View.GroupedOcclusionQueries.BatchPrimitive(OcclusionBounds.Origin + View.ViewMatrices.PreViewTranslation,OcclusionBounds.BoxExtent) :
+										View.IndividualOcclusionQueries.BatchPrimitive(OcclusionBounds.Origin + View.ViewMatrices.PreViewTranslation,OcclusionBounds.BoxExtent)
+										);	
+								}
+								PrimitiveOcclusionHistory->bGroupedQuery = bGroupedQuery;
 							}
-							else
-							{
-								// Primitives that need precise occlusion results use individual queries.
-								bGroupedQuery = false;
-								bRunQuery = true;
-							}
+						}
+						else
+						{
+							// If the primitive's bounding box intersects the near clipping plane, treat it as definitely unoccluded.
+							bIsOccluded = false;
+							bOcclusionStateIsDefinite = true;
+						}
+					}
 
-							if (bRunQuery)
+					if (bSubQueries)
+					{
+						SubIsOccluded.Add(bIsOccluded);
+						if (!bIsOccluded)
+						{
+							bAllSubOccluded = false;
+							if (bOcclusionStateIsDefinite)
 							{
-								PrimitiveOcclusionHistory->SetCurrentQuery(ViewState->OcclusionFrameCounter, 
-									bGroupedQuery ? 
-									View.GroupedOcclusionQueries.BatchPrimitive(OcclusionBounds.Origin + View.ViewMatrices.PreViewTranslation,OcclusionBounds.BoxExtent) :
-									View.IndividualOcclusionQueries.BatchPrimitive(OcclusionBounds.Origin + View.ViewMatrices.PreViewTranslation,OcclusionBounds.BoxExtent)
-									);	
+								PrimitiveOcclusionHistory->LastVisibleTime = CurrentRealTime;
 							}
-							PrimitiveOcclusionHistory->bGroupedQuery = bGroupedQuery;
+						}
+						if (bIsOccluded || !bOcclusionStateIsDefinite)
+						{
+							bAllSubOcclusionStateIsDefinite = false;
 						}
 					}
 					else
 					{
-						// If the primitive's bounding box intersects the near clipping plane, treat it as definitely unoccluded.
-						bIsOccluded = false;
-						bOcclusionStateIsDefinite = true;
+						if (bIsOccluded)
+						{
+							View.PrimitiveVisibilityMap.AccessCorrespondingBit(BitIt) = false;
+							STAT(NumOccludedPrimitives++);
+						}
+						else if (bOcclusionStateIsDefinite)
+						{
+							PrimitiveOcclusionHistory->LastVisibleTime = CurrentRealTime;
+							View.PrimitiveDefinitelyUnoccludedMap.AccessCorrespondingBit(BitIt) = true;
+						}
 					}
 				}
-
-				if (bIsOccluded)
+				if (bSubQueries)
 				{
-					View.PrimitiveVisibilityMap.AccessCorrespondingBit(BitIt) = false;
-					STAT(NumOccludedPrimitives++);
-				}
-				else if (bOcclusionStateIsDefinite)
-				{
-					PrimitiveOcclusionHistory->LastVisibleTime = CurrentRealTime;
-					View.PrimitiveDefinitelyUnoccludedMap.AccessCorrespondingBit(BitIt) = true;
+					FPrimitiveSceneProxy* Proxy = Scene->Primitives[BitIt.GetIndex()]->Proxy;
+					Proxy->AcceptOcclusionResults(&View, &SubIsOccluded[0], SubIsOccluded.Num());
+					if (bAllSubOccluded)
+					{
+						View.PrimitiveVisibilityMap.AccessCorrespondingBit(BitIt) = false;
+						STAT(NumOccludedPrimitives++);
+					}
+					else if (bAllSubOcclusionStateIsDefinite)
+					{
+						View.PrimitiveDefinitelyUnoccludedMap.AccessCorrespondingBit(BitIt) = true;
+					}
 				}
 			}
 
