Dynamically spawned Actors/Components that are static will be affected by [REDACTED_PERSON]
Made some small optimizations to reduce CullDistanceVolume update time

diff --git a/Engine/Source/Runtime/Engine/Classes/Components/PrimitiveComponent.h b/Engine/Source/Runtime/Engine/Classes/Components/PrimitiveComponent.h
--- a/Engine/Source/Runtime/Engine/Classes/Components/PrimitiveComponent.h
+++ b/Engine/Source/Runtime/Engine/Classes/Components/PrimitiveComponent.h
@@ -1556,6 +1556,11 @@ public:
 	UFUNCTION(BlueprintCallable, Category="LOD", meta=(FriendlyName="Set Max Draw Distance"))
 	void SetCullDistance(float NewCullDistance);
 	
+	/**
+	 * Utility to cache the max draw distance based on cull distance volumes or the desired max draw distance
+	 */
+	void SetCachedMaxDrawDistance(const float NewCachedMaxDrawDistance);
+
 	/**
 	 * Changes the value of DepthPriorityGroup.
 	 * @param NewDepthPriorityGroup - The value to assign to DepthPriorityGroup.
diff --git a/Engine/Source/Runtime/Engine/Classes/Engine/World.h b/Engine/Source/Runtime/Engine/Classes/Engine/World.h
--- a/Engine/Source/Runtime/Engine/Classes/Engine/World.h
+++ b/Engine/Source/Runtime/Engine/Classes/Engine/World.h
@@ -2318,9 +2318,11 @@ public:
 	void UpdateWorldComponents(bool bRerunConstructionScripts, bool bCurrentLevelOnly);
 
 	/**
-	 * Updates all cull distance volumes.
+	 * Updates cull distance volumes for a specified component or a specified actor or all actors
+         * @param ComponentToUpdate If specified just that Component will be updated
+	 * @param ActorToUpdate If specified (and ComponentToUpdate is not specified), all Components owned by this Actor will be updated
 	 */
-	void UpdateCullDistanceVolumes();
+	void UpdateCullDistanceVolumes(AActor* ActorToUpdate = nullptr, UPrimitiveComponent* ComponentToUpdate = nullptr);
 
 	/**
 	 * Cleans up components, streaming data and assorted other intermediate data.
diff --git a/Engine/Source/Runtime/Engine/Private/ActorConstruction.cpp b/Engine/Source/Runtime/Engine/Private/ActorConstruction.cpp
--- a/Engine/Source/Runtime/Engine/Private/ActorConstruction.cpp
+++ b/Engine/Source/Runtime/Engine/Private/ActorConstruction.cpp
@@ -485,6 +485,8 @@ void AActor::ExecuteConstruction(const FTransform& Transform, const FComponentIn
 		}
 	}
 
+	GetWorld()->UpdateCullDistanceVolumes(this);
+
 	// Now run virtual notification
 	OnConstruction(Transform);
 }
@@ -561,7 +563,6 @@ UActorComponent* AActor::CreateComponentFromTemplate(UActorComponent* Template,
 
 		// Note we aren't copying the the RF_ArchetypeObject flag. Also note the result is non-transactional by default.
 		NewActorComp = (UActorComponent*)StaticDuplicateObject(Template, this, *InName, RF_AllFlags & ~(RF_ArchetypeObject|RF_Transactional|RF_WasLoaded|RF_Public|RF_InheritableComponentTemplate) );
-		//NewActorComp = ConstructObject<UActorComponent>(Template->GetClass(), this, *InName, RF_NoFlags, Template);
 
 		NewActorComp->CreationMethod = EComponentCreationMethod::UserConstructionScript;
 
@@ -585,6 +586,7 @@ UActorComponent* AActor::AddComponent(FName TemplateName, bool bManualAttachment
 		BlueprintGeneratedClass = Cast<UBlueprintGeneratedClass>(BlueprintGeneratedClass->GetSuperClass());
 	}
 
+	bool bIsSceneComponent = false;
 	UActorComponent* NewActorComp = CreateComponentFromTemplate(Template);
 	if(NewActorComp != nullptr)
 	{
@@ -609,10 +611,22 @@ UActorComponent* AActor::AddComponent(FName TemplateName, bool bManualAttachment
 			}
 
 			NewSceneComp->SetRelativeTransform(RelativeTransform);
+
+			bIsSceneComponent = true;
 		}
 
 		// Register component, which will create physics/rendering state, now component is in correct position
 		NewActorComp->RegisterComponent();
+
+		UWorld* World = GetWorld();
+		if (!bRunningUserConstructionScript && World && bIsSceneComponent)
+		{
+			UPrimitiveComponent* NewPrimitiveComponent = Cast<UPrimitiveComponent>(NewActorComp);
+			if (NewPrimitiveComponent && ACullDistanceVolume::CanBeAffectedByVolumes(NewPrimitiveComponent))
+			{
+				World->UpdateCullDistanceVolumes(this, NewPrimitiveComponent);
+			}
+		}
 	}
 
 	return NewActorComp;
diff --git a/Engine/Source/Runtime/Engine/Private/CullDistanceVolume.cpp b/Engine/Source/Runtime/Engine/Private/CullDistanceVolume.cpp
--- a/Engine/Source/Runtime/Engine/Private/CullDistanceVolume.cpp
+++ b/Engine/Source/Runtime/Engine/Private/CullDistanceVolume.cpp
@@ -53,14 +53,14 @@ bool ACullDistanceVolume::CanBeAffectedByVolumes( UPrimitiveComponent* Primitive
 	if(	Owner
 	// Disregard dynamic actors
 	&& (PrimitiveComponent->Mobility == EComponentMobility::Static)
-	// Disregard prefabs.
-	&& !PrimitiveComponent->IsTemplate()
+	// Require cull distance volume support to be enabled.
+	&&	PrimitiveComponent->bAllowCullDistanceVolume 
 	// Skip primitives that is hidden set as we don't want to cull out brush rendering or other helper objects.
 	&&	PrimitiveComponent->IsVisible()
+	// Disregard prefabs.
+	&& !PrimitiveComponent->IsTemplate()
 	// Only operate on primitives attached to the owners world.			
-	&&	PrimitiveComponent->GetScene() == Owner->GetWorld()->Scene 
-	// Require cull distance volume support to be enabled.
-	&&	PrimitiveComponent->bAllowCullDistanceVolume )
+	&&	PrimitiveComponent->GetScene() == Owner->GetWorld()->Scene)
 	{
 		return true;
 	}
diff --git a/Engine/Source/Runtime/Engine/Private/PrimitiveComponent.cpp b/Engine/Source/Runtime/Engine/Private/PrimitiveComponent.cpp
--- a/Engine/Source/Runtime/Engine/Private/PrimitiveComponent.cpp
+++ b/Engine/Source/Runtime/Engine/Private/PrimitiveComponent.cpp
@@ -946,10 +946,23 @@ void UPrimitiveComponent::PushHoveredToProxy(const bool bInHovered)
 
 void UPrimitiveComponent::SetCullDistance(float NewCullDistance)
 {
-	if( CachedMaxDrawDistance != NewCullDistance )
+	if (NewCullDistance > 0)
 	{
-	    CachedMaxDrawDistance = NewCullDistance;
-	    MarkRenderStateDirty();
+		LDMaxDrawDistance = NewCullDistance;
+	
+		if (LDMaxDrawDistance < CachedMaxDrawDistance)
+		{
+			SetCachedMaxDrawDistance(LDMaxDrawDistance);
+		}
+	}
+}
+
+void UPrimitiveComponent::SetCachedMaxDrawDistance(const float NewCachedMaxDrawDistance)
+{
+	if( !FMath::IsNearlyEqual(CachedMaxDrawDistance, NewCachedMaxDrawDistance) )
+	{
+		CachedMaxDrawDistance = NewCachedMaxDrawDistance;
+		MarkRenderStateDirty();
 	}
 }
 
diff --git a/Engine/Source/Runtime/Engine/Private/World.cpp b/Engine/Source/Runtime/Engine/Private/World.cpp
--- a/Engine/Source/Runtime/Engine/Private/World.cpp
+++ b/Engine/Source/Runtime/Engine/Private/World.cpp
@@ -1238,7 +1238,7 @@ void UWorld::UpdateWorldComponents(bool bRerunConstructionScripts, bool bCurrent
 }
 
 
-void UWorld::UpdateCullDistanceVolumes()
+void UWorld::UpdateCullDistanceVolumes(AActor* ActorToUpdate, UPrimitiveComponent* ComponentToUpdate)
 {
 	// Map that will store new max draw distance for every primitive
 	TMap<UPrimitiveComponent*,float> CompToNewMaxDrawMap;
@@ -1251,26 +1251,61 @@ void UWorld::UpdateCullDistanceVolumes()
 		TArray<ACullDistanceVolume*> CullDistanceVolumes;
 
 		// Establish base line of LD specified cull distances.
-		for( FActorIterator It(this); It; ++It )
+		if (ActorToUpdate || ComponentToUpdate)
 		{
-			TInlineComponentArray<UPrimitiveComponent*> PrimitiveComponents;
-			It->GetComponents(PrimitiveComponents);
-			for (UPrimitiveComponent* PrimitiveComponent : PrimitiveComponents)
+			if (ComponentToUpdate)
 			{
-				CompToNewMaxDrawMap.Add(PrimitiveComponent, PrimitiveComponent->LDMaxDrawDistance);
+				check((ActorToUpdate == nullptr) || (ActorToUpdate == ComponentToUpdate->GetOwner()));
+				CompToNewMaxDrawMap.Add(ComponentToUpdate, ComponentToUpdate->LDMaxDrawDistance);
+			}
+			else
+			{
+				TInlineComponentArray<UPrimitiveComponent*> PrimitiveComponents(ActorToUpdate);
+				for (UPrimitiveComponent* PrimitiveComponent : PrimitiveComponents)
+				{
+					if (ACullDistanceVolume::CanBeAffectedByVolumes(PrimitiveComponent))
+					{
+						CompToNewMaxDrawMap.Add(PrimitiveComponent, PrimitiveComponent->LDMaxDrawDistance);
+					}
+				}
 			}
 
-			ACullDistanceVolume* CullDistanceVolume = Cast<ACullDistanceVolume>(*It);
-			if (CullDistanceVolume)
+			if (CompToNewMaxDrawMap.Num() > 0)
 			{
-				CullDistanceVolumes.Add(CullDistanceVolume);
+				for (TActorIterator<ACullDistanceVolume> It(this); It; ++It)
+				{
+					CullDistanceVolumes.Add(*It);
+				}
+			}
+		}
+		else
+		{
+			for( FActorIterator It(this); It; ++It )
+			{
+				TInlineComponentArray<UPrimitiveComponent*> PrimitiveComponents(*It);
+				for (UPrimitiveComponent* PrimitiveComponent : PrimitiveComponents)
+				{
+					if (ACullDistanceVolume::CanBeAffectedByVolumes(PrimitiveComponent))
+					{
+						CompToNewMaxDrawMap.Add(PrimitiveComponent, PrimitiveComponent->LDMaxDrawDistance);
+					}
+				}
+
+				ACullDistanceVolume* CullDistanceVolume = Cast<ACullDistanceVolume>(*It);
+				if (CullDistanceVolume)
+				{
+					CullDistanceVolumes.Add(CullDistanceVolume);
+				}
 			}
 		}
 
 		// Iterate over all cull distance volumes and get new cull distances.
-		for (ACullDistanceVolume* CullDistanceVolume : CullDistanceVolumes)
+		if (CompToNewMaxDrawMap.Num() > 0)
 		{
-			CullDistanceVolume->GetPrimitiveMaxDrawDistances(CompToNewMaxDrawMap);
+			for (ACullDistanceVolume* CullDistanceVolume : CullDistanceVolumes)
+			{
+				CullDistanceVolume->GetPrimitiveMaxDrawDistances(CompToNewMaxDrawMap);
+			}
 		}
 
 		// Finally, go over all primitives, and see if they need to change.
@@ -1278,13 +1313,9 @@ void UWorld::UpdateCullDistanceVolumes()
 		for ( TMap<UPrimitiveComponent*,float>::TIterator It(CompToNewMaxDrawMap); It; ++It )
 		{
 			UPrimitiveComponent* PrimComp = It.Key();
-			float NewMaxDrawDist = It.Value();
+			const float NewMaxDrawDist = It.Value();
 
-			if( !FMath::IsNearlyEqual(PrimComp->CachedMaxDrawDistance, NewMaxDrawDist) )
-			{
-				PrimComp->CachedMaxDrawDistance = NewMaxDrawDist;
-				PrimComp->RecreateRenderState_Concurrent();
-			}
+			PrimComp->SetCachedMaxDrawDistance(NewMaxDrawDist);
 		}
 	}
 
