Removed broken/unused "navigation seeds" mechanism #UE4

diff --git a/Engine/Source/Runtime/Engine/Classes/AI/Navigation/NavigationSystem.h b/Engine/Source/Runtime/Engine/Classes/AI/Navigation/NavigationSystem.h
--- a/Engine/Source/Runtime/Engine/Classes/AI/Navigation/NavigationSystem.h
+++ b/Engine/Source/Runtime/Engine/Classes/AI/Navigation/NavigationSystem.h
@@ -166,10 +166,6 @@ public:
 	//UPROPERTY(config, EditAnywhere, Category=NavigationSystem)
 	uint32 bWholeWorldNavigable:1;
 
-	/** If set to true (default) generation seeds will include locations of all player controlled pawns */
-	UPROPERTY(config, EditAnywhere, Category=NavigationSystem)
-	uint32 bAddPlayersToGenerationSeeds:1;
-
 	/** false by default, if set to true will result in not caring about nav agent height 
 	 *	when trying to match navigation data to passed in nav agent */
 	UPROPERTY(config, EditAnywhere, Category=NavigationSystem)
@@ -712,11 +708,6 @@ public:
 
 	FORCEINLINE bool IsSetUpForLazyGeometryExporting() const { return bGenerateNavigationOnlyAroundNavigationInvokers; }
 
-	/** register actor important for generation (navigation data will be build around them first) */
-	void RegisterGenerationSeed(AActor* SeedActor);
-	void UnregisterGenerationSeed(AActor* SeedActor);
-	void GetGenerationSeeds(TArray<FVector>& SeedLocations) const;
-
 	static UNavigationSystem* CreateNavigationSystem(UWorld* WorldOwner);
 
 	static UNavigationSystem* GetCurrent(UWorld* World);
@@ -767,9 +758,6 @@ protected:
 	/** Map of all custom navigation links, that are relevant for path following */
 	TMap<uint32, INavLinkCustomInterface*> CustomLinksMap;
 
-	/** List of actors relevant to generation of navigation data */
-	TArray<TWeakObjectPtr<AActor> > GenerationSeeds;
-
 	/** stores areas marked as dirty throughout the frame, processes them 
 	 *	once a frame in Tick function */
 	TArray<FNavigationDirtyArea> DirtyAreas;
diff --git a/Engine/Source/Runtime/Engine/Private/AI/Navigation/NavigationSystem.cpp b/Engine/Source/Runtime/Engine/Private/AI/Navigation/NavigationSystem.cpp
--- a/Engine/Source/Runtime/Engine/Private/AI/Navigation/NavigationSystem.cpp
+++ b/Engine/Source/Runtime/Engine/Private/AI/Navigation/NavigationSystem.cpp
@@ -242,7 +242,6 @@ bool UNavigationSystem::bUpdateNavOctreeOnComponentChange = true;
 UNavigationSystem::UNavigationSystem(const FObjectInitializer& ObjectInitializer)
 	: Super(ObjectInitializer)
 	, bWholeWorldNavigable(false)
-	, bAddPlayersToGenerationSeeds(true)
 	, bSkipAgentHeightCheckWhenPickingNavData(false)
 	, DirtyAreasUpdateFreq(60)
 	, OperationMode(FNavigationSystem::InvalidMode)
@@ -1994,39 +1993,6 @@ void UNavigationSystem::ResetCachedFilter(TSubclassOf<UNavigationQueryFilter> Fi
 	}
 }
 
-void UNavigationSystem::RegisterGenerationSeed(AActor* SeedActor)
-{
-	GenerationSeeds.Add(SeedActor);
-}
-
-void UNavigationSystem::UnregisterGenerationSeed(AActor* SeedActor)
-{
-	GenerationSeeds.RemoveSingleSwap(SeedActor);
-}
-
-void UNavigationSystem::GetGenerationSeeds(TArray<FVector>& SeedLocations) const
-{
-	if (bAddPlayersToGenerationSeeds)
-	{
-		for (FConstPlayerControllerIterator it = GetWorld()->GetPlayerControllerIterator(); it ;it++)
-		{
-			if (*it && (*it)->GetPawn() != NULL)
-			{
-				const FVector SeedLoc((*it)->GetPawn()->GetActorLocation());
-				SeedLocations.Add(SeedLoc);
-			}
-		}
-	}
-
-	for (int32 SeedIndex = 0; SeedIndex < GenerationSeeds.Num(); SeedIndex++)
-	{
-		if (GenerationSeeds[SeedIndex].IsValid())
-		{
-			SeedLocations.Add(GenerationSeeds[SeedIndex]->GetActorLocation());
-		}
-	}
-}
-
 UNavigationSystem* UNavigationSystem::CreateNavigationSystem(UWorld* WorldOwner)
 {
 	UNavigationSystem* NavSys = NULL;
