adjusted showflag visibility

diff --git a/Engine/Source/Runtime/Engine/Public/ShowFlagsValues.inl b/Engine/Source/Runtime/Engine/Public/ShowFlagsValues.inl
--- a/Engine/Source/Runtime/Engine/Public/ShowFlagsValues.inl
+++ b/Engine/Source/Runtime/Engine/Public/ShowFlagsValues.inl
@@ -245,7 +245,7 @@ SHOWFLAG_ALWAYS_ACCESSIBLE(TextRender, SFG_Advanced, LOCTEXT("TextRenderSF", "Re
 /** Any rendering/buffer clearing  (good for benchmarking) */
 SHOWFLAG_ALWAYS_ACCESSIBLE(Rendering, SFG_Hidden, LOCTEXT("RenderingSF", "Any Rendering"))
 /** Show the current mask being used by the highres screenshot capture */
-SHOWFLAG_ALWAYS_ACCESSIBLE(HighResScreenshotMask, SFG_Advanced, LOCTEXT("HighResScreenshotMaskSF", "High Res Screenshot Mask"))
+SHOWFLAG_ALWAYS_ACCESSIBLE(HighResScreenshotMask, SFG_Hidden, LOCTEXT("HighResScreenshotMaskSF", "High Res Screenshot Mask"))
 /** Distortion of output for HMD devices */
 SHOWFLAG_ALWAYS_ACCESSIBLE(HMDDistortion, SFG_PostProcess, LOCTEXT("HMDDistortionSF", "HMD Distortion"))
 /** Whether to render in stereoscopic 3d */
