Allow TInlineComponentArray to have the Actor provided it to it and not require a subsequence GetComponents call

diff --git a/Engine/Source/Runtime/Engine/Classes/Engine/EngineTypes.h b/Engine/Source/Runtime/Engine/Classes/Engine/EngineTypes.h
--- a/Engine/Source/Runtime/Engine/Classes/Engine/EngineTypes.h
+++ b/Engine/Source/Runtime/Engine/Classes/Engine/EngineTypes.h
@@ -25,6 +25,17 @@ enum { NumInlinedActorComponents = 24 };
 template<class T, uint32 NumElements = NumInlinedActorComponents>
 class TInlineComponentArray : public TArray<T, TInlineAllocator<NumElements>>
 {
+	typedef TArray<T, TInlineAllocator<NumElements>> Super;
+
+public:
+	TInlineComponentArray() : Super() { }
+	TInlineComponentArray(class AActor* Actor) : Super()
+	{
+		if (Actor)
+		{
+			Actor->GetComponents(*this);
+		}
+	};
 };
 
 
