Fix compile for non-editor builds

diff --git a/Engine/Source/Runtime/Engine/Private/ActorConstruction.cpp b/Engine/Source/Runtime/Engine/Private/ActorConstruction.cpp
--- a/Engine/Source/Runtime/Engine/Private/ActorConstruction.cpp
+++ b/Engine/Source/Runtime/Engine/Private/ActorConstruction.cpp
@@ -5,6 +5,7 @@
 #include "LatentActions.h"
 #include "ComponentInstanceDataCache.h"
 #include "Engine/LevelScriptActor.h"
+#include "Engine/CullDistanceVolume.h"
 
 #if WITH_EDITOR
 #include "Editor.h"
