Move definition of TInlineComponentArray in to Actor.h to avoid incomplete type compile errors

diff --git a/Engine/Source/Runtime/Engine/Classes/Engine/EngineTypes.h b/Engine/Source/Runtime/Engine/Classes/Engine/EngineTypes.h
--- a/Engine/Source/Runtime/Engine/Classes/Engine/EngineTypes.h
+++ b/Engine/Source/Runtime/Engine/Classes/Engine/EngineTypes.h
@@ -10,35 +10,12 @@
 #include "GameFramework/DamageType.h"
 #include "EngineTypes.generated.h"
 
-
 /**
  * Default number of components to expect in TInlineAllocators used with AActor component arrays.
  * Used by engine code to try to avoid allocations in AActor::GetComponents(), among others.
  */
 enum { NumInlinedActorComponents = 24 };
 
-
-/**
- * TInlineComponentArray is simply a TArray that reserves a fixed amount of space on the stack
- * to try to avoid heap allocation when there are fewer than a specified number of elements expected in the result.
- */
-template<class T, uint32 NumElements = NumInlinedActorComponents>
-class TInlineComponentArray : public TArray<T, TInlineAllocator<NumElements>>
-{
-	typedef TArray<T, TInlineAllocator<NumElements>> Super;
-
-public:
-	TInlineComponentArray() : Super() { }
-	TInlineComponentArray(class AActor* Actor) : Super()
-	{
-		if (Actor)
-		{
-			Actor->GetComponents(*this);
-		}
-	};
-};
-
-
 UENUM()
 enum EAspectRatioAxisConstraint
 {
diff --git a/Engine/Source/Runtime/Engine/Classes/GameFramework/Actor.h b/Engine/Source/Runtime/Engine/Classes/GameFramework/Actor.h
--- a/Engine/Source/Runtime/Engine/Classes/GameFramework/Actor.h
+++ b/Engine/Source/Runtime/Engine/Classes/GameFramework/Actor.h
@@ -2186,8 +2186,7 @@ public:
 	 * It's recommended to use TArrays with a TInlineAllocator to potentially avoid memory allocation costs.
 	 * TInlineComponentArray is defined to make this easier, for example:
 	 * {
-	 * 	   TInlineComponentArray<UPrimitiveComponent*> PrimComponents;
-	 *     Actor->GetComponents(PrimComponents);
+	 * 	   TInlineComponentArray<UPrimitiveComponent*> PrimComponents(Actor);
 	 * }
 	 */
 	template<class T, class AllocatorType>
@@ -2377,6 +2376,26 @@ public:
 	class UClass* GetActorClass() const;
 };
 
+/**
+ * TInlineComponentArray is simply a TArray that reserves a fixed amount of space on the stack
+ * to try to avoid heap allocation when there are fewer than a specified number of elements expected in the result.
+ */
+template<class T, uint32 NumElements = NumInlinedActorComponents>
+class TInlineComponentArray : public TArray<T, TInlineAllocator<NumElements>>
+{
+	typedef TArray<T, TInlineAllocator<NumElements>> Super;
+
+public:
+	TInlineComponentArray() : Super() { }
+	TInlineComponentArray(class AActor* Actor) : Super()
+	{
+		if (Actor)
+		{
+			Actor->GetComponents(*this);
+		}
+	};
+};
+
 
 //////////////////////////////////////////////////////////////////////////
 // Inlines
