#ue4 - Removed redundant scoped stat counter.

diff --git a/Engine/Source/Runtime/Engine/Private/CharacterMovementComponent.cpp b/Engine/Source/Runtime/Engine/Private/CharacterMovementComponent.cpp
--- a/Engine/Source/Runtime/Engine/Private/CharacterMovementComponent.cpp
+++ b/Engine/Source/Runtime/Engine/Private/CharacterMovementComponent.cpp
@@ -1005,7 +1005,6 @@ void UCharacterMovementComponent::SimulatedTick(float DeltaSeconds)
 		// Tick animations before physics.
 		if( CharacterOwner->GetMesh() )
 		{
-			SCOPE_CYCLE_COUNTER(STAT_CharacterMovement);
 			TickCharacterPose(DeltaSeconds);
 
 			// Make sure animation didn't trigger an event that destroyed us
