Anchor Decal and Emitters spawned at a location to WorldSettings to prevent garbage collection unexpectedly
Make world settings "visible" so that components owned by it will be rendered

diff --git a/Engine/Source/Runtime/Engine/Private/GameplayStatics.cpp b/Engine/Source/Runtime/Engine/Private/GameplayStatics.cpp
--- a/Engine/Source/Runtime/Engine/Private/GameplayStatics.cpp
+++ b/Engine/Source/Runtime/Engine/Private/GameplayStatics.cpp
@@ -540,16 +540,17 @@ void UGameplayStatics::PlayWorldCameraShake(UObject* WorldContextObject, TSubcla
 	}
 }
 
-UParticleSystemComponent* CreateParticleSystem(UParticleSystem* EmitterTemplate, UWorld* World, AActor* Actor, bool bAutoDestroy)
+UParticleSystemComponent* CreateParticleSystem(UParticleSystem* EmitterTemplate, AActor* Actor, bool bAutoDestroy)
 {
-	UParticleSystemComponent* PSC = NewObject<UParticleSystemComponent>((Actor ? Actor : (UObject*)World));
+	check(Actor);
+	UParticleSystemComponent* PSC = NewObject<UParticleSystemComponent>(Actor);
 	PSC->bAutoDestroy = bAutoDestroy;
 	PSC->SecondsBeforeInactive = 0.0f;
 	PSC->bAutoActivate = false;
 	PSC->SetTemplate(EmitterTemplate);
 	PSC->bOverrideLODMethod = false;
 
-	PSC->RegisterComponentWithWorld(World);
+	PSC->RegisterComponent();
 
 	return PSC;
 }
@@ -562,7 +563,7 @@ UParticleSystemComponent* UGameplayStatics::SpawnEmitterAtLocation(UObject* Worl
 		UWorld* World = GEngine->GetWorldFromContextObject(WorldContextObject);
 		if(World != nullptr)
 		{
-			PSC = CreateParticleSystem(EmitterTemplate, World, NULL, bAutoDestroy);
+			PSC = CreateParticleSystem(EmitterTemplate, World->GetWorldSettings(), bAutoDestroy);
 
 			PSC->SetAbsolute(true, true, true);
 			PSC->SetWorldLocationAndRotation(SpawnLocation, SpawnRotation);
@@ -584,7 +585,7 @@ UParticleSystemComponent* UGameplayStatics::SpawnEmitterAttached(UParticleSystem
 		}
 		else
 		{
-			PSC = CreateParticleSystem(EmitterTemplate, AttachToComponent->GetWorld(), AttachToComponent->GetOwner(), bAutoDestroy);
+			PSC = CreateParticleSystem(EmitterTemplate, AttachToComponent->GetOwner(), bAutoDestroy);
 
 			PSC->AttachTo(AttachToComponent, AttachPointName);
 			if (LocationType == EAttachLocation::KeepWorldPosition)
@@ -683,7 +684,7 @@ void UGameplayStatics::PlaySoundAtLocation(UObject* WorldContextObject, class US
 
 	const bool bIsInGameWorld = ThisWorld->IsGameWorld();
 
-	if (GEngine && GEngine->UseSound() && ThisWorld->bAllowAudioPlayback && ThisWorld->GetNetMode() != NM_DedicatedServer)
+	if (GEngine->UseSound() && ThisWorld->bAllowAudioPlayback && ThisWorld->GetNetMode() != NM_DedicatedServer)
 	{
 		if( Sound->IsAudibleSimple( Location, AttenuationSettings ) )
 		{
@@ -940,15 +941,17 @@ void UGameplayStatics::DeactivateReverbEffect(FName TagName)
 	}
 }
 
-UDecalComponent* CreateDecalComponent(class UMaterialInterface* DecalMaterial, FVector DecalSize, UWorld* World, AActor* Actor, float LifeSpan)
+UDecalComponent* CreateDecalComponent(class UMaterialInterface* DecalMaterial, FVector DecalSize, AActor* Actor, float LifeSpan)
 {
+	check(Actor);
+
 	const FMatrix DecalInternalTransform = FRotationMatrix(FRotator(0.f, 90.0f, -90.0f));
 
-	UDecalComponent* DecalComp = NewObject<UDecalComponent>((Actor ? Actor : (UObject*)GetTransientPackage()));
+	UDecalComponent* DecalComp = NewObject<UDecalComponent>(Actor);
 	DecalComp->DecalMaterial = DecalMaterial;
 	DecalComp->RelativeScale3D = DecalInternalTransform.TransformVector(DecalSize);
 	DecalComp->bAbsoluteScale = true;
-	DecalComp->RegisterComponentWithWorld(World);
+	DecalComp->RegisterComponent();
 
 	if (LifeSpan > 0.f)
 	{
@@ -967,7 +970,7 @@ UDecalComponent* UGameplayStatics::SpawnDecalAtLocation(UObject* WorldContextObj
 		UWorld* World = GEngine->GetWorldFromContextObject(WorldContextObject);
 		if(World != nullptr)
 		{
-			DecalComp = CreateDecalComponent(DecalMaterial, DecalSize, World, NULL, LifeSpan);
+			DecalComp = CreateDecalComponent(DecalMaterial, DecalSize, World->GetWorldSettings(), LifeSpan);
 			DecalComp->SetWorldLocationAndRotation(Location, Rotation);
 		}
 	}
@@ -998,7 +1001,7 @@ UDecalComponent* UGameplayStatics::SpawnDecalAttached(class UMaterialInterface*
 				}
 				else
 				{
-					DecalComp = CreateDecalComponent(DecalMaterial, DecalSize, AttachToComponent->GetWorld(), AttachToComponent->GetOwner(), LifeSpan);
+					DecalComp = CreateDecalComponent(DecalMaterial, DecalSize, AttachToComponent->GetOwner(), LifeSpan);
 					DecalComp->AttachTo(AttachToComponent, AttachPointName);
 
 					if (LocationType == EAttachLocation::KeepWorldPosition)
diff --git a/Engine/Source/Runtime/Engine/Private/WorldSettings.cpp b/Engine/Source/Runtime/Engine/Private/WorldSettings.cpp
--- a/Engine/Source/Runtime/Engine/Private/WorldSettings.cpp
+++ b/Engine/Source/Runtime/Engine/Private/WorldSettings.cpp
@@ -53,9 +53,7 @@ AWorldSettings::AWorldSettings(const FObjectInitializer& ObjectInitializer)
 	MatineeTimeDilation = 1.0f;
 	DemoPlayTimeDilation = 1.0f;
 	PackedLightAndShadowMapTextureSize = 1024;
-#if WITH_EDITORONLY_DATA
-	bHiddenEd = true;
-#endif // WITH_EDITORONLY_DATA
+	bHidden = false;
 
 	DefaultColorScale = FVector(1.0f, 1.0f, 1.0f);
 
