Fix destructibles ignoring impact damage for depth layer higher than 0

diff --git a/Engine/Source/Runtime/Engine/Private/DestructibleComponent.cpp b/Engine/Source/Runtime/Engine/Private/DestructibleComponent.cpp
--- a/Engine/Source/Runtime/Engine/Private/DestructibleComponent.cpp
+++ b/Engine/Source/Runtime/Engine/Private/DestructibleComponent.cpp
@@ -96,6 +96,37 @@ FBoxSphereBounds UDestructibleComponent::CalcBounds(const FTransform& LocalToWor
 #endif	// #if WITH_APEX
 }
 
+bool IsImpactDamageEnabled(const UDestructibleMesh* TheDestructibleMesh, int32 Level)
+{
+	if(TheDestructibleMesh->DefaultDestructibleParameters.DamageParameters.ImpactDamage == 0.f)
+	{
+		return false;
+	}
+
+	bool bEnableImpactDamage = false;
+	const FDestructibleDepthParameters& DepthParams = TheDestructibleMesh->DefaultDestructibleParameters.DepthParameters[Level];
+	const EImpactDamageOverride LevelOverride = DepthParams.ImpactDamageOverride; 
+
+	switch(LevelOverride)
+	{
+		case EImpactDamageOverride::IDO_On:
+		{
+			return true;
+		}
+
+		case EImpactDamageOverride::IDO_Off:
+		{
+			return false;
+		}
+
+		default:
+		{
+			//return default if we're within the default level
+		    return TheDestructibleMesh->DefaultDestructibleParameters.DamageParameters.DefaultImpactDamageDepth >= Level ? TheDestructibleMesh->DefaultDestructibleParameters.DamageParameters.bEnableImpactDamage : false;
+		}
+	}
+}
+
 void UDestructibleComponent::OnUpdateTransform(bool bSkipPhysicsMove)
 {
 	// We are handling the physics move below, so don't handle it at higher levels
@@ -249,13 +280,8 @@ void UDestructibleComponent::CreatePhysicsState()
 		SmallChunkCollisionResponse.SetAllChannels(ECR_Ignore);
 	}
 
-	bool bEnableImpactDamage = TheDestructibleMesh->DefaultDestructibleParameters.DamageParameters.bEnableImpactDamage;
-	for (const FDestructibleDepthParameters& DepthParams : TheDestructibleMesh->DefaultDestructibleParameters.DepthParameters)
-	{
-		bEnableImpactDamage |= DepthParams.ImpactDamageOverride == EImpactDamageOverride::IDO_On;
-	}
-
-	bool bEnableContactModification = TheDestructibleMesh->DefaultDestructibleParameters.DamageParameters.bCustomImpactResistance && TheDestructibleMesh->DefaultDestructibleParameters.DamageParameters.ImpactResistance > 0.f;
+	const bool bEnableImpactDamage = IsImpactDamageEnabled(TheDestructibleMesh, 0);
+	const bool bEnableContactModification = TheDestructibleMesh->DefaultDestructibleParameters.DamageParameters.bCustomImpactResistance && TheDestructibleMesh->DefaultDestructibleParameters.DamageParameters.ImpactResistance > 0.f;
 
 	// Passing AssetInstanceID = 0 so we'll have self-collision
 	AActor* Owner = GetOwner();
@@ -314,6 +340,8 @@ void UDestructibleComponent::CreatePhysicsState()
 	// Enable hard sleeping if requested
 	verify( NxParameterized::setParamBool(*ActorParams,"useHardSleeping", bEnableHardSleeping) );
 
+	
+
 	// Destructibles are always dynamic or kinematic, and therefore only go into one of the scenes
 	const uint32 SceneType = BodyInstance.UseAsyncScene() ? PST_Async : PST_Sync;
 	NxApexScene* ApexScene = PhysScene->GetApexScene(SceneType);
@@ -1303,11 +1331,13 @@ void UDestructibleComponent::SetCollisionResponseForActor(PxRigidDynamic* Actor,
 	uint8 MoveChannel = GetCollisionObjectType();
 	if(IsCollisionEnabled())
 	{
+		UDestructibleMesh* TheDestructibleMesh = GetDestructibleMesh();
 		AActor* Owner = GetOwner();
 		bool bLargeChunk = IsChunkLarge(ChunkIdx);
 		const FCollisionResponse& ColResponse = bLargeChunk ? LargeChunkCollisionResponse : SmallChunkCollisionResponse;
-		//TODO: we currently assume chunks will not have impact damage as it's very expensive. Should look into exposing this a bit more
-		CreateShapeFilterData(MoveChannel, GetUniqueID(), ColResponse.GetResponseContainer(), 0, ChunkIdxToBoneIdx(ChunkIdx), PQueryFilterData, PSimFilterData, BodyInstance.bUseCCD, false, false);
+		physx::PxU32 SupportDepth = TheDestructibleMesh->ApexDestructibleAsset->getChunkDepth(ChunkIdx);
+		const bool bEnableImpactDamage = IsImpactDamageEnabled(TheDestructibleMesh, SupportDepth);
+		CreateShapeFilterData(MoveChannel, GetUniqueID(), ColResponse.GetResponseContainer(), 0, ChunkIdxToBoneIdx(ChunkIdx), PQueryFilterData, PSimFilterData, BodyInstance.bUseCCD, bEnableImpactDamage, false);
 		
 		PQueryFilterData.word3 |= EPDF_SimpleCollision | EPDF_ComplexCollision;
 
