#ue4 - Fix crash opening Character blueprints.

UE-11307

diff --git a/Engine/Source/Runtime/Engine/Private/NavMovementComponent.cpp b/Engine/Source/Runtime/Engine/Private/NavMovementComponent.cpp
--- a/Engine/Source/Runtime/Engine/Private/NavMovementComponent.cpp
+++ b/Engine/Source/Runtime/Engine/Private/NavMovementComponent.cpp
@@ -58,7 +58,11 @@ void UNavMovementComponent::UpdateNavAgent(const AActor& Owner)
 	}
 
 	// initialize properties from navigation system
-	NavAgentProps.NavWalkingSearchHeightScale = UNavigationSystem::GetCurrent(GetWorld())->GetDefaultSupportedAgentConfig().NavWalkingSearchHeightScale;
+	UNavigationSystem* NavSys = UNavigationSystem::GetCurrent(GetWorld());
+	if (NavSys != nullptr)
+	{
+		NavAgentProps.NavWalkingSearchHeightScale = NavSys->GetDefaultSupportedAgentConfig().NavWalkingSearchHeightScale;
+	}
 
 	// Can't call GetSimpleCollisionCylinder(), because no components will be registered.
 	float BoundRadius, BoundHalfHeight;	
@@ -75,7 +79,11 @@ void UNavMovementComponent::UpdateNavAgent(const UCapsuleComponent& CapsuleCompo
 	}
 
 	// initialize properties from navigation system
-	NavAgentProps.NavWalkingSearchHeightScale = UNavigationSystem::GetCurrent(GetWorld())->GetDefaultSupportedAgentConfig().NavWalkingSearchHeightScale;
+	UNavigationSystem* NavSys = UNavigationSystem::GetCurrent(GetWorld());
+	if (NavSys != nullptr)
+	{
+		NavAgentProps.NavWalkingSearchHeightScale = NavSys->GetDefaultSupportedAgentConfig().NavWalkingSearchHeightScale;
+	}
 
 	NavAgentProps.AgentRadius = CapsuleComponent.GetScaledCapsuleRadius();
 	NavAgentProps.AgentHeight = CapsuleComponent.GetScaledCapsuleHalfHeight() * 2.f;
