Don't allow CullDistanceVolumes to be modified by blueprints

diff --git a/Engine/Source/Runtime/Engine/Classes/Engine/CullDistanceVolume.h b/Engine/Source/Runtime/Engine/Classes/Engine/CullDistanceVolume.h
--- a/Engine/Source/Runtime/Engine/Classes/Engine/CullDistanceVolume.h
+++ b/Engine/Source/Runtime/Engine/Classes/Engine/CullDistanceVolume.h
@@ -15,11 +15,11 @@ struct FCullDistanceSizePair
 	GENERATED_USTRUCT_BODY()
 
 	/** Size to associate with cull distance. */
-	UPROPERTY(EditAnywhere, BlueprintReadWrite, Category=CullDistanceSizePair)
+	UPROPERTY(EditAnywhere, BlueprintReadOnly, Category=CullDistanceSizePair)
 	float Size;
 
 	/** Cull distance associated with size. */
-	UPROPERTY(EditAnywhere, BlueprintReadWrite, Category=CullDistanceSizePair)
+	UPROPERTY(EditAnywhere, BlueprintReadOnly, Category=CullDistanceSizePair)
 	float CullDistance;
 
 
@@ -45,13 +45,13 @@ class ACullDistanceVolume
 	 * Array of size and cull distance pairs. The code will calculate the sphere diameter of a primitive's BB and look for a best
 	 * fit in this array to determine which cull distance to use.
 	 */
-	UPROPERTY(EditAnywhere, BlueprintReadWrite, Category=CullDistanceVolume)
+	UPROPERTY(EditAnywhere, BlueprintReadOnly, Category=CullDistanceVolume)
 	TArray<struct FCullDistanceSizePair> CullDistances;
 
 	/**
 	 * Whether the volume is currently enabled or not.
 	 */
-	UPROPERTY(EditAnywhere, BlueprintReadWrite, Category=CullDistanceVolume)
+	UPROPERTY(EditAnywhere, BlueprintReadOnly, Category=CullDistanceVolume)
 	uint32 bEnabled:1;
 
 public:
