#ue4 - Exposed AActor::IsOverlappingActor(), UPrimitiveComponent::IsOverlappingActor(), UPrimitiveComponent::IsOverlappingComponent() to blueprints and improved comments.

diff --git a/Engine/Source/Runtime/Engine/Classes/Components/PrimitiveComponent.h b/Engine/Source/Runtime/Engine/Classes/Components/PrimitiveComponent.h
--- a/Engine/Source/Runtime/Engine/Classes/Components/PrimitiveComponent.h
+++ b/Engine/Source/Runtime/Engine/Classes/Components/PrimitiveComponent.h
@@ -517,16 +517,26 @@ public:
 	 */
 	void EndComponentOverlap(const FOverlapInfo& OtherOverlap, bool bDoNotifies=true, bool bNoNotifySelf=false);
 
-	/** Returns true if this component is overlapping OtherComp, false otherwise. */
-	bool IsOverlappingComponent(UPrimitiveComponent const* OtherComp) const;
+	/**
+	 * Check whether this component is overlapping another component.
+	 * @param OtherComp Component to test this component against.
+	 * @return Whether this component is overlapping another component.
+	 */
+	UFUNCTION(BlueprintCallable, Category="Collision", meta=(UnsafeDuringActorConstruction="true"))
+	bool IsOverlappingComponent(const UPrimitiveComponent* OtherComp) const;
 	
-	/** Returns true if this component is overlapping OtherComp, false otherwise. */
+	/** Check whether this component has the specified overlap. */
 	bool IsOverlappingComponent(const FOverlapInfo& Overlap) const;
 
-	/** Return true if this component is overlapping any component of the given actor, false otherwise. */
+	/**
+	 * Check whether this component is overlapping any component of the given Actor.
+	 * @param Other Actor to test this component against.
+	 * @return Whether this component is overlapping any component of the given Actor.
+	 */
+	UFUNCTION(BlueprintCallable, Category="Collision", meta=(UnsafeDuringActorConstruction="true"))
 	bool IsOverlappingActor(const AActor* Other) const;
 
-	/** Appends list of overlaps with components owned by the given actor to the 'OutOverlaps' array. Returns the number of overlaps that were added. */
+	/** Appends list of overlaps with components owned by the given actor to the 'OutOverlaps' array. Returns true if any overlaps were added. */
 	bool GetOverlapsWithActor(const AActor* Actor, TArray<FOverlapInfo>& OutOverlaps) const;
 
 	/** 
diff --git a/Engine/Source/Runtime/Engine/Classes/GameFramework/Actor.h b/Engine/Source/Runtime/Engine/Classes/GameFramework/Actor.h
--- a/Engine/Source/Runtime/Engine/Classes/GameFramework/Actor.h
+++ b/Engine/Source/Runtime/Engine/Classes/GameFramework/Actor.h
@@ -1622,10 +1622,11 @@ public:
 	void UpdateOverlaps(bool bDoNotifies=true);
 	
 	/** 
-	 * Check to see if current Actor is overlapping specified Actor
-	 * @param Other the Actor to test for
-	 * Returns true if any component of this actor is overlapping any component of Other. 
+	 * Check whether any component of this Actor is overlapping any component of another Actor.
+	 * @param Other The other Actor to test against
+	 * @return Whether any component of this Actor is overlapping any component of another Actor.
 	 */
+	UFUNCTION(BlueprintCallable, Category="Collision", meta=(UnsafeDuringActorConstruction="true"))
 	bool IsOverlappingActor(const AActor* Other) const;
 
 	/** Returns whether a MatineeActor is currently controlling this Actor */
