UE4 - Fix for skeletal meshes SetVertex calls in game (used on dynamic meshes); remove unused bool

diff --git a/Engine/Source/Runtime/Engine/Private/SkeletalMesh.cpp b/Engine/Source/Runtime/Engine/Private/SkeletalMesh.cpp
--- a/Engine/Source/Runtime/Engine/Private/SkeletalMesh.cpp
+++ b/Engine/Source/Runtime/Engine/Private/SkeletalMesh.cpp
@@ -128,10 +128,9 @@ FSkeletalMeshVertexBuffer::FSkeletalMeshVertexBuffer()
 :	bInfluencesByteSwapped(false)
 ,	bUseFullPrecisionUVs(false)
 ,	bNeedsCPUAccess(false)
-,	bProcessedPackedPositions(false)
 ,	bExtraBoneInfluences(false)
-,	VertexData(NULL)
-,	Data(NULL)
+,	VertexData(nullptr)
+,	Data(nullptr)
 ,	Stride(0)
 ,	NumVertices(0)
 ,	MeshOrigin(FVector::ZeroVector)
@@ -166,7 +165,6 @@ FSkeletalMeshVertexBuffer::FSkeletalMeshVertexBuffer(const FSkeletalMeshVertexBu
 :	bInfluencesByteSwapped(false)
 ,	bUseFullPrecisionUVs(Other.bUseFullPrecisionUVs)
 ,	bNeedsCPUAccess(Other.bNeedsCPUAccess)
-,	bProcessedPackedPositions(Other.bProcessedPackedPositions)
 ,	bExtraBoneInfluences(Other.bExtraBoneInfluences)
 ,	VertexData(NULL)
 ,	Data(NULL)
@@ -357,27 +355,21 @@ void FSkeletalMeshVertexBuffer::SetVertexFast(uint32 VertexIndex,const FSoftSkin
 	FMemory::Memcpy(VertBase->InfluenceWeights,SrcVertex.InfluenceWeights, TGPUSkinVertexBase<bExtraBoneInfluencesT>::NumInfluences);
 	if( !bUseFullPrecisionUVs )
 	{
-#if !WITH_EDITORONLY_DATA // I don't expect this to happen on console. If so this won't work with PackedPosition. Having check just in case.
-		check (false);
-#else
-		((TGPUSkinVertexFloat16Uvs<MAX_TEXCOORDS, bExtraBoneInfluencesT>*)(VertBase))->Position = SrcVertex.Position;
+		auto* Vertex = (TGPUSkinVertexFloat16Uvs<MAX_TEXCOORDS, bExtraBoneInfluencesT>*)VertBase;
+		Vertex->Position = SrcVertex.Position;
 		for( uint32 UVIndex = 0; UVIndex < NumTexCoords; ++UVIndex )
 		{
-			((TGPUSkinVertexFloat16Uvs<MAX_TEXCOORDS, bExtraBoneInfluencesT>*)(VertBase))->UVs[UVIndex] = FVector2DHalf( SrcVertex.UVs[UVIndex] );
+			Vertex->UVs[UVIndex] = FVector2DHalf( SrcVertex.UVs[UVIndex] );
 		}
-#endif
 	}
 	else
 	{
-#if !WITH_EDITORONLY_DATA // I don't expect this to happen on console. If so this won't work with PackedPosition. Having check just in case.
-		check (false);
-#else
-		((TGPUSkinVertexFloat32Uvs<MAX_TEXCOORDS, bExtraBoneInfluencesT>*)(VertBase))->Position = SrcVertex.Position;
+		auto* Vertex = (TGPUSkinVertexFloat32Uvs<MAX_TEXCOORDS, bExtraBoneInfluencesT>*)VertBase;
+		Vertex->Position = SrcVertex.Position;
 		for( uint32 UVIndex = 0; UVIndex < NumTexCoords; ++UVIndex )
 		{
-			((TGPUSkinVertexFloat32Uvs<MAX_TEXCOORDS, bExtraBoneInfluencesT>*)(VertBase))->UVs[UVIndex] = FVector2D( SrcVertex.UVs[UVIndex] );
+			Vertex->UVs[UVIndex] = FVector2D( SrcVertex.UVs[UVIndex] );
 		}
-#endif
 	}
 }
 
diff --git a/Engine/Source/Runtime/Engine/Public/SkeletalMeshTypes.h b/Engine/Source/Runtime/Engine/Public/SkeletalMeshTypes.h
--- a/Engine/Source/Runtime/Engine/Public/SkeletalMeshTypes.h
+++ b/Engine/Source/Runtime/Engine/Public/SkeletalMeshTypes.h
@@ -1078,8 +1078,6 @@ private:
 	bool bUseFullPrecisionUVs;
 	/** true if this vertex buffer will be used with CPU skinning. Resource arrays are set to cpu accessible if this is true */
 	bool bNeedsCPUAccess;
-	/** Position data has already been packed. Used during cooking to avoid packing twice. */
-	bool bProcessedPackedPositions;
 	/** Has extra bone influences per Vertex, which means using a different TGPUSkinVertexBase */
 	bool bExtraBoneInfluences;
 	/** The vertex data storage type */
