USkeletalMeshComponent::RecalcRequiredBones's check for Sockets could add the same bone indices multiple times, triggering an assert in FillSpaceBases() later. Made sure all bones added to ForceAnimatedSocketBones and NeededBonesForFillSpaceBases are unique.

diff --git a/Engine/Source/Runtime/Engine/Private/SkeletalMeshComponent.cpp b/Engine/Source/Runtime/Engine/Private/SkeletalMeshComponent.cpp
--- a/Engine/Source/Runtime/Engine/Private/SkeletalMeshComponent.cpp
+++ b/Engine/Source/Runtime/Engine/Private/SkeletalMeshComponent.cpp
@@ -691,7 +691,8 @@ void USkeletalMeshComponent::FillSpaceBases(const USkeletalMesh* InSkeletalMesh,
 	FAnimationRuntime::NormalizeRotations(DestSpaceBases);
 }
 
-/** Takes sorted array Base and then adds any elements from sorted array Insert which is missing from it, preserving order. */
+/** Takes sorted array Base and then adds any elements from sorted array Insert which is missing from it, preserving order.
+ * this assumes both arrays are sorted and contain unique bone indices. */
 static void MergeInBoneIndexArrays(TArray<FBoneIndexType>& BaseArray, TArray<FBoneIndexType>& InsertArray)
 {
 	// Then we merge them into the array of required bones.
@@ -841,11 +842,11 @@ void USkeletalMeshComponent::RecalcRequiredBones(int32 LODIndex)
 			int32 BoneIndex = SkeletalMesh->RefSkeleton.FindBoneIndex(Socket->BoneName);
 			if (Socket->bForceAlwaysAnimated)
 			{
-				ForceAnimatedSocketBones.Add(BoneIndex);
+				ForceAnimatedSocketBones.AddUnique(BoneIndex);
 			}
 			else
 			{
-				NeededBonesForFillSpaceBases.Add(BoneIndex);
+				NeededBonesForFillSpaceBases.AddUnique(BoneIndex);
 			}
 		}
 
