Fix error message when calling SetStaticMesh on a static component

diff --git a/Engine/Source/Runtime/Engine/Private/StaticMeshComponent.cpp b/Engine/Source/Runtime/Engine/Private/StaticMeshComponent.cpp
--- a/Engine/Source/Runtime/Engine/Private/StaticMeshComponent.cpp
+++ b/Engine/Source/Runtime/Engine/Private/StaticMeshComponent.cpp
@@ -1182,7 +1182,7 @@ bool UStaticMeshComponent::SetStaticMesh(UStaticMesh* NewMesh)
 	if(!AreDynamicDataChangesAllowed() && Owner != NULL)
 	{
 		FMessageLog("PIE").Warning(FText::Format(LOCTEXT("SetMeshOnStatic", "Calling SetStaticMesh on '{0}' but Mobility is Static."), 
-			FText::FromString(GetPathName(this))));
+			FText::FromString(GetPathName())));
 		return false;
 	}
 
