Fix non APEX build

diff --git a/Engine/Source/Runtime/Engine/Private/DestructibleComponent.cpp b/Engine/Source/Runtime/Engine/Private/DestructibleComponent.cpp
--- a/Engine/Source/Runtime/Engine/Private/DestructibleComponent.cpp
+++ b/Engine/Source/Runtime/Engine/Private/DestructibleComponent.cpp
@@ -1335,7 +1335,9 @@ void UDestructibleComponent::SetCollisionResponseForActor(PxRigidDynamic* Actor,
 		AActor* Owner = GetOwner();
 		bool bLargeChunk = IsChunkLarge(ChunkIdx);
 		const FCollisionResponse& ColResponse = bLargeChunk ? LargeChunkCollisionResponse : SmallChunkCollisionResponse;
+#if WITH_APEX
 		physx::PxU32 SupportDepth = TheDestructibleMesh->ApexDestructibleAsset->getChunkDepth(ChunkIdx);
+
 		const bool bEnableImpactDamage = IsImpactDamageEnabled(TheDestructibleMesh, SupportDepth);
 		CreateShapeFilterData(MoveChannel, GetUniqueID(), ColResponse.GetResponseContainer(), 0, ChunkIdxToBoneIdx(ChunkIdx), PQueryFilterData, PSimFilterData, BodyInstance.bUseCCD, bEnableImpactDamage, false);
 		
@@ -1358,6 +1360,7 @@ void UDestructibleComponent::SetCollisionResponseForActor(PxRigidDynamic* Actor,
 			Shape->setFlag(PxShapeFlag::eSIMULATION_SHAPE, true); 
 			Shape->setFlag(PxShapeFlag::eVISUALIZATION, true);
 		}
+#endif
 	}
 }
 
