#ue4 - Perf: Convert some stack TArray usage to TInlineComponentArray, and avoid allocation/deallocation in other cases when possible.

diff --git a/Engine/Source/Runtime/Engine/Private/Actor.cpp b/Engine/Source/Runtime/Engine/Private/Actor.cpp
--- a/Engine/Source/Runtime/Engine/Private/Actor.cpp
+++ b/Engine/Source/Runtime/Engine/Private/Actor.cpp
@@ -941,11 +941,9 @@ void AActor::UpdateOverlaps(bool bDoNotifies)
 
 bool AActor::IsOverlappingActor(const AActor* Other) const
 {
-	// using a stack to walk this actor's attached component tree
-	TArray<USceneComponent*> ComponentStack;
-
+	// use a stack to walk this actor's attached component tree
+	TInlineComponentArray<USceneComponent*> ComponentStack;
 	USceneComponent const* CurrentComponent = GetRootComponent();
-
 	while (CurrentComponent)
 	{
 		// push children on the stack so they get tested later
@@ -958,22 +956,24 @@ bool AActor::IsOverlappingActor(const AActor* Other) const
 			return true;
 		}
 
-		CurrentComponent = (ComponentStack.Num() > 0) ? ComponentStack.Pop() : NULL;
+		// advance to next component
+		const bool bAllowShrinking = false;
+		CurrentComponent = (ComponentStack.Num() > 0) ? ComponentStack.Pop(bAllowShrinking) : nullptr;
 	}
 
 	return false;
 }
 
 
-void AActor::GetOverlappingActors(TArray<AActor*>& OverlappingActors, UClass* ClassFilter) const
+void AActor::GetOverlappingActors(TArray<AActor*>& OutOverlappingActors, UClass* ClassFilter) const
 {
 	// prepare output
-	OverlappingActors.Empty();
-
-	TArray<USceneComponent*> ComponentStack;
+	OutOverlappingActors.Reset();
+	TArray<AActor*> OverlappingActorsForCurrentComponent;
 
+	// use a stack to walk this actor's attached component tree
+	TInlineComponentArray<USceneComponent*> ComponentStack;
 	USceneComponent const* CurrentComponent = GetRootComponent();
-
 	while (CurrentComponent)
 	{
 		// push children on the stack so they get tested later
@@ -983,7 +983,6 @@ void AActor::GetOverlappingActors(TArray<AActor*>& OverlappingActors, UClass* Cl
 		UPrimitiveComponent const* const PrimComp = Cast<const UPrimitiveComponent>(CurrentComponent);
 		if (PrimComp)
 		{
-			TArray<AActor*> OverlappingActorsForCurrentComponent;
 			PrimComp->GetOverlappingActors(OverlappingActorsForCurrentComponent, ClassFilter);
 
 			// then merge it into our final list
@@ -992,24 +991,25 @@ void AActor::GetOverlappingActors(TArray<AActor*>& OverlappingActors, UClass* Cl
 				AActor* OverlappingActor = *CompIt;
 				if(OverlappingActor != this)
 				{
-					OverlappingActors.AddUnique(OverlappingActor);
+					OutOverlappingActors.AddUnique(OverlappingActor);
 				}
 			}
 		}
 
 		// advance to next component
-		CurrentComponent = (ComponentStack.Num() > 0) ? ComponentStack.Pop() : NULL;
+		const bool bAllowShrinking = false;
+		CurrentComponent = (ComponentStack.Num() > 0) ? ComponentStack.Pop(bAllowShrinking) : nullptr;
 	}
 }
 
 void AActor::GetOverlappingComponents(TArray<UPrimitiveComponent*>& OutOverlappingComponents) const
 {
-	OutOverlappingComponents.Empty();
-
-	TArray<USceneComponent*> ComponentStack;
+	OutOverlappingComponents.Reset();
+	TArray<UPrimitiveComponent*> OverlappingComponentsForCurrentComponent;
 
+	// use a stack to walk this actor's attached component tree
+	TInlineComponentArray<USceneComponent*> ComponentStack;
 	USceneComponent* CurrentComponent = GetRootComponent();
-
 	while (CurrentComponent)
 	{
 		// push children on the stack so they get tested later
@@ -1019,7 +1019,6 @@ void AActor::GetOverlappingComponents(TArray<UPrimitiveComponent*>& OutOverlappi
 		if (PrimComp)
 		{
 			// get list of components from the component
-			TArray<UPrimitiveComponent*> OverlappingComponentsForCurrentComponent;
 			PrimComp->GetOverlappingComponents(OverlappingComponentsForCurrentComponent);
 
 			// then merge it into our final list
@@ -1030,7 +1029,8 @@ void AActor::GetOverlappingComponents(TArray<UPrimitiveComponent*>& OutOverlappi
 		}
 
 		// advance to next component
-		CurrentComponent = (ComponentStack.Num() > 0) ? ComponentStack.Pop() : NULL;
+		const bool bAllowShrinking = false;
+		CurrentComponent = (ComponentStack.Num() > 0) ? ComponentStack.Pop(bAllowShrinking) : nullptr;
 	}
 }
 
@@ -1284,14 +1284,14 @@ FName AActor::GetAttachParentSocketName() const
 
 void AActor::GetAttachedActors(TArray<class AActor*>& OutActors) const
 {
-	OutActors.Empty();
+	OutActors.Reset();
 	if (RootComponent != NULL)
 	{
 		// Current set of components to check
-		TArray< USceneComponent*, TInlineAllocator<NumInlinedActorComponents> > CompsToCheck;
+		TInlineComponentArray<USceneComponent*> CompsToCheck;
 
 		// Set of all components we have checked
-		TArray< USceneComponent*, TInlineAllocator<NumInlinedActorComponents> > CheckedComps;
+		TInlineComponentArray<USceneComponent*> CheckedComps;
 
 		CompsToCheck.Push(RootComponent);
 
@@ -1299,7 +1299,8 @@ void AActor::GetAttachedActors(TArray<class AActor*>& OutActors) const
 		while(CompsToCheck.Num() > 0)
 		{
 			// Get the next off the queue
-			USceneComponent* SceneComp = CompsToCheck.Pop();
+			const bool bAllowShrinking = false;
+			USceneComponent* SceneComp = CompsToCheck.Pop(bAllowShrinking);
 
 			// Add it to the 'checked' set, should not already be there!
 			if (!CheckedComps.Contains(SceneComp))
diff --git a/Engine/Source/Runtime/Engine/Private/PrimitiveComponent.cpp b/Engine/Source/Runtime/Engine/Private/PrimitiveComponent.cpp
--- a/Engine/Source/Runtime/Engine/Private/PrimitiveComponent.cpp
+++ b/Engine/Source/Runtime/Engine/Private/PrimitiveComponent.cpp
@@ -1842,7 +1842,10 @@ bool UPrimitiveComponent::IsOverlappingComponent(UPrimitiveComponent const* Othe
 {
 	for (int32 i=0; i < OverlappingComponents.Num(); ++i)
 	{
-		if (OverlappingComponents[i].OverlapInfo.Component.Get() == OtherComp) { return true; }
+		if (OverlappingComponents[i].OverlapInfo.Component.Get() == OtherComp)
+		{
+			return true;
+		}
 	}
 	return false;
 }
@@ -2062,7 +2065,7 @@ void UPrimitiveComponent::EndComponentOverlap(const FOverlapInfo& OtherOverlap,
 
 void UPrimitiveComponent::GetOverlappingActors(TArray<AActor*>& OutOverlappingActors, UClass* ClassFilter) const
 {
-	OutOverlappingActors.Empty();
+	OutOverlappingActors.Reset();
 
 	for (auto CompIt = OverlappingComponents.CreateConstIterator(); CompIt; ++CompIt)
 	{
