Implemented DrawMaterialTile (GitHub)
(https://github.com/EpicGames/UnrealEngine/pull/605)

diff --git a/Engine/Source/Runtime/Engine/Classes/Engine/Canvas.h b/Engine/Source/Runtime/Engine/Classes/Engine/Canvas.h
--- a/Engine/Source/Runtime/Engine/Classes/Engine/Canvas.h
+++ b/Engine/Source/Runtime/Engine/Classes/Engine/Canvas.h
@@ -469,6 +469,14 @@ public:
 	UFUNCTION(BlueprintCallable, Category=Canvas, meta=(FriendlyName="Draw Triangles"))
 	void K2_DrawTriangle(UTexture* RenderTexture, TArray<FCanvasUVTri> Triangles);
 
+	/**
+	 * Draws a set of triangles on the Canvas.
+	 *
+	 * @param RenderMaterial			Material to use when rendering. Remember that only the emissive channel is able to be rendered as no lighting is performed when rendering to the Canvas.
+	 * @param Triangles					Triangles to render.
+	 */
+	UFUNCTION(BlueprintCallable, Category = Canvas, meta = (FriendlyName = "Draw Material Triangles"))
+	void K2_DrawMaterialTriangle(UMaterialInterface* RenderMaterial, TArray<FCanvasUVTri> Triangles);
 	/**
 	 * Draws a polygon on the Canvas.
 	 *
diff --git a/Engine/Source/Runtime/Engine/Classes/GameFramework/HUD.h b/Engine/Source/Runtime/Engine/Classes/GameFramework/HUD.h
--- a/Engine/Source/Runtime/Engine/Classes/GameFramework/HUD.h
+++ b/Engine/Source/Runtime/Engine/Classes/GameFramework/HUD.h
@@ -302,6 +302,9 @@ public:
 	UFUNCTION(BlueprintCallable, Category=HUD)
 	void DrawMaterialSimple(UMaterialInterface* Material, float ScreenX, float ScreenY, float ScreenW, float ScreenH, float Scale=1.f, bool bScalePosition=false);
 
+	UFUNCTION(BlueprintCallable, Category = HUD)
+	void DrawMaterialTriangle(UMaterialInterface* Material, FVector2D V0_Pos, FVector2D V1_Pos, FVector2D V2_Pos, FVector2D V0_UV, FVector2D V1_UV, FVector2D V2_UV, FLinearColor V0_Color = FLinearColor::White, FLinearColor V1_Color = FLinearColor::White, FLinearColor V2_Color = FLinearColor::White);
+	
 	/** Transforms a 3D world-space vector into 2D screen coordinates */
 	UFUNCTION(BlueprintCallable, Category = HUD)
 	FVector Project(FVector Location) const;
diff --git a/Engine/Source/Runtime/Engine/Private/HUD.cpp b/Engine/Source/Runtime/Engine/Private/HUD.cpp
--- a/Engine/Source/Runtime/Engine/Private/HUD.cpp
+++ b/Engine/Source/Runtime/Engine/Private/HUD.cpp
@@ -660,6 +660,18 @@ void AHUD::DrawTextureSimple(UTexture* Texture, float ScreenX, float ScreenY, fl
 	}
 }
 
+void AHUD::DrawMaterialTriangle(UMaterialInterface* Material, FVector2D V0_Pos, FVector2D V1_Pos, FVector2D V2_Pos, FVector2D V0_UV, FVector2D V1_UV, FVector2D V2_UV, FLinearColor V0_Color, FLinearColor V1_Color, FLinearColor V2_Color)
+{
+	if (IsCanvasValid_WarnIfNot() && Material)
+	{
+		FCanvasTriangleItem TriangleItem(V0_Pos, V1_Pos, V2_Pos, V0_UV, V1_UV, V2_UV, NULL);
+		TriangleItem.TriangleList[0].V0_Color = V0_Color;
+		TriangleItem.TriangleList[0].V1_Color = V1_Color;
+		TriangleItem.TriangleList[0].V2_Color = V2_Color;
+		TriangleItem.MaterialRenderProxy = Material->GetRenderProxy(0);
+		Canvas->DrawItem(TriangleItem);
+	}
+}
 FVector AHUD::Project(FVector Location) const
 {
 	if (IsCanvasValid_WarnIfNot())
diff --git a/Engine/Source/Runtime/Engine/Private/TileRendering.cpp b/Engine/Source/Runtime/Engine/Private/TileRendering.cpp
--- a/Engine/Source/Runtime/Engine/Private/TileRendering.cpp
+++ b/Engine/Source/Runtime/Engine/Private/TileRendering.cpp
@@ -1,7 +1,7 @@
 // Copyright 1998-2015 Epic Games, Inc. All Rights Reserved.
 
 /*=============================================================================
-	TranslucentRendering.cpp: Translucent rendering implementation.
+	TileRendering.cpp: Tile rendering implementation.
 =============================================================================*/
 
 #include "EnginePrivate.h"
@@ -10,6 +10,8 @@
 #include "LocalVertexFactory.h"
 #include "MeshBatch.h"
 #include "RendererInterface.h"
+#include "SceneUtils.h"
+#include "CanvasTypes.h"
 
 /** 
 * vertex data for a screen quad 
@@ -17,22 +19,25 @@
 struct FMaterialTileVertex
 {
 	FVector			Position;
-	FPackedNormal	TangentX,
-					TangentZ;
+	FPackedNormal	TangentX;
+	FPackedNormal	TangentZ;
 	uint32			Color;
-	float			U,
-					V;
+	float			U;
+	float			V;
 
 	inline void Initialize(float InX, float InY, float InU, float InV)
 	{
-		Position.X = InX; Position.Y = InY; Position.Z = 0.0f;
+		Position.X = InX; 
+		Position.Y = InY; 
+		Position.Z = 0.0f;
 		TangentX = FVector(1, 0, 0); 
 		//TangentY = FVector(0, 1, 0); 
 		TangentZ = FVector(0, 0, 1);
 		// TangentZ.w contains the sign of the tangent basis determinant. Assume +1
 		TangentZ.Vector.W = 255;
 		Color = FColor(255,255,255,255).DWColor();
-		U = InU; V = InV;
+		U = InU; 
+		V = InV;
 	}
 };
 
@@ -167,3 +172,157 @@ void FTileRenderer::DrawTile(FRHICommandListImmediate& RHICmdList, const class F
 
 	GetRendererModule().DrawTileMesh(RHICmdList, View, Mesh, bIsHitTesting, HitProxyId);
 }
+
+bool FCanvasTileRendererItem::Render_RenderThread(FRHICommandListImmediate& RHICmdList, const FCanvas* Canvas)
+{
+	float CurrentRealTime = 0.f;
+	float CurrentWorldTime = 0.f;
+	float DeltaWorldTime = 0.f;
+
+	if (!bFreezeTime)
+	{
+		CurrentRealTime = Canvas->GetCurrentRealTime();
+		CurrentWorldTime = Canvas->GetCurrentWorldTime();
+		DeltaWorldTime = Canvas->GetCurrentDeltaWorldTime();
+	}
+
+	checkSlow(Data);
+	// current render target set for the canvas
+	const FRenderTarget* CanvasRenderTarget = Canvas->GetRenderTarget();
+	FSceneViewFamily* ViewFamily = new FSceneViewFamily(FSceneViewFamily::ConstructionValues(
+		CanvasRenderTarget,
+		nullptr,
+		FEngineShowFlags(ESFIM_Game))
+		.SetWorldTimes(CurrentWorldTime, DeltaWorldTime, CurrentRealTime)
+		.SetGammaCorrection(CanvasRenderTarget->GetDisplayGamma()));
+
+	FIntRect ViewRect(FIntPoint(0, 0), CanvasRenderTarget->GetSizeXY());
+
+	// make a temporary view
+	FSceneViewInitOptions ViewInitOptions;
+	ViewInitOptions.ViewFamily = ViewFamily;
+	ViewInitOptions.SetViewRectangle(ViewRect);
+	ViewInitOptions.ViewOrigin = FVector::ZeroVector;
+	ViewInitOptions.ViewRotationMatrix = FMatrix::Identity;
+	ViewInitOptions.ProjectionMatrix = Data->Transform.GetMatrix();
+	ViewInitOptions.BackgroundColor = FLinearColor::Black;
+	ViewInitOptions.OverlayColor = FLinearColor::White;
+
+	FSceneView* View = new FSceneView(ViewInitOptions);
+	
+	bool bNeedsToSwitchVerticalAxis = RHINeedsToSwitchVerticalAxis(Canvas->GetShaderPlatform()) && !Canvas->GetAllowSwitchVerticalAxis();
+
+	for (int32 TileIdx = 0; TileIdx < Data->Tiles.Num(); TileIdx++)
+	{
+		const FRenderData::FTileInst& Tile = Data->Tiles[TileIdx];
+		FTileRenderer::DrawTile(
+			RHICmdList,
+			*View,
+			Data->MaterialRenderProxy,
+			bNeedsToSwitchVerticalAxis,
+			Tile.X, Tile.Y, Tile.SizeX, Tile.SizeY,
+			Tile.U, Tile.V, Tile.SizeU, Tile.SizeV,
+			Canvas->IsHitTesting(), Tile.HitProxyId,
+			Tile.InColor
+			);
+	}
+
+	delete View->Family;
+	delete View;
+	if (Canvas->GetAllowedModes() & FCanvas::Allow_DeleteOnRender)
+	{
+		delete Data;
+	}
+	if (Canvas->GetAllowedModes() & FCanvas::Allow_DeleteOnRender)
+	{
+		Data = NULL;
+	}
+	return true;
+}
+
+bool FCanvasTileRendererItem::Render_GameThread(const FCanvas* Canvas)
+{
+	float CurrentRealTime = 0.f;
+	float CurrentWorldTime = 0.f;
+	float DeltaWorldTime = 0.f;
+
+	if (!bFreezeTime)
+	{
+		CurrentRealTime = Canvas->GetCurrentRealTime();
+		CurrentWorldTime = Canvas->GetCurrentWorldTime();
+		DeltaWorldTime = Canvas->GetCurrentDeltaWorldTime();
+	}
+
+	checkSlow(Data);
+	// current render target set for the canvas
+	const FRenderTarget* CanvasRenderTarget = Canvas->GetRenderTarget();
+	FSceneViewFamily* ViewFamily = new FSceneViewFamily(FSceneViewFamily::ConstructionValues(
+		CanvasRenderTarget,
+		NULL,
+		FEngineShowFlags(ESFIM_Game))
+		.SetWorldTimes(CurrentWorldTime, DeltaWorldTime, CurrentRealTime)
+		.SetGammaCorrection(CanvasRenderTarget->GetDisplayGamma()));
+
+	FIntRect ViewRect(FIntPoint(0, 0), CanvasRenderTarget->GetSizeXY());
+
+	// make a temporary view
+	FSceneViewInitOptions ViewInitOptions;
+	ViewInitOptions.ViewFamily = ViewFamily;
+	ViewInitOptions.SetViewRectangle(ViewRect);
+	ViewInitOptions.ViewOrigin = FVector::ZeroVector;
+	ViewInitOptions.ViewRotationMatrix = FMatrix::Identity;
+	ViewInitOptions.ProjectionMatrix = Data->Transform.GetMatrix();
+	ViewInitOptions.BackgroundColor = FLinearColor::Black;
+	ViewInitOptions.OverlayColor = FLinearColor::White;
+
+	FSceneView* View = new FSceneView(ViewInitOptions);
+
+	bool bNeedsToSwitchVerticalAxis = RHINeedsToSwitchVerticalAxis(Canvas->GetShaderPlatform()) && !Canvas->GetAllowSwitchVerticalAxis();
+	struct FDrawTileParameters
+	{
+		FSceneView* View;
+		FRenderData* RenderData;
+		uint32 bIsHitTesting : 1;
+		uint32 bNeedsToSwitchVerticalAxis : 1;
+		uint32 AllowedCanvasModes;
+	};
+	FDrawTileParameters DrawTileParameters =
+	{
+		View,
+		Data,
+		Canvas->IsHitTesting(),
+		bNeedsToSwitchVerticalAxis,
+		Canvas->GetAllowedModes()
+	};
+	ENQUEUE_UNIQUE_RENDER_COMMAND_ONEPARAMETER(
+		DrawTileCommand,
+		FDrawTileParameters, Parameters, DrawTileParameters,
+		{
+		SCOPED_DRAW_EVENT(RHICmdList, CanvasDrawTile);
+		for (int32 TileIdx = 0; TileIdx < Parameters.RenderData->Tiles.Num(); TileIdx++)
+		{
+			const FRenderData::FTileInst& Tile = Parameters.RenderData->Tiles[TileIdx];
+			FTileRenderer::DrawTile(
+				RHICmdList,
+				*Parameters.View,
+				Parameters.RenderData->MaterialRenderProxy,
+				Parameters.bNeedsToSwitchVerticalAxis,
+				Tile.X, Tile.Y, Tile.SizeX, Tile.SizeY,
+				Tile.U, Tile.V, Tile.SizeU, Tile.SizeV,
+				Parameters.bIsHitTesting, Tile.HitProxyId,
+				Tile.InColor);
+		}
+
+		delete Parameters.View->Family;
+		delete Parameters.View;
+		if (Parameters.AllowedCanvasModes & FCanvas::Allow_DeleteOnRender)
+		{
+			delete Parameters.RenderData;
+		}
+		});
+	if (Canvas->GetAllowedModes() & FCanvas::Allow_DeleteOnRender)
+	{
+		Data = NULL;
+	}
+	return true;
+}
diff --git a/Engine/Source/Runtime/Engine/Private/TriangleRendering.cpp b/Engine/Source/Runtime/Engine/Private/TriangleRendering.cpp
new file mode 100644
--- /dev/null
+++ b/Engine/Source/Runtime/Engine/Private/TriangleRendering.cpp
@@ -0,0 +1,324 @@
+// Copyright 1998-2015 Epic Games, Inc. All Rights Reserved.
+
+/*=============================================================================
+	TriangleRendering.cpp: Simple triangle rendering implementation.
+=============================================================================*/
+
+#include "EnginePrivate.h"
+#include "TriangleRendering.h"
+#include "EngineModule.h"
+#include "LocalVertexFactory.h"
+#include "MeshBatch.h"
+#include "RendererInterface.h"
+#include "SceneUtils.h"
+#include "CanvasTypes.h"
+
+/** 
+ * vertex data for a screen triangle
+ */
+struct FMaterialTriangleVertex
+{
+	FVector			Position;
+	FPackedNormal	TangentX;
+	FPackedNormal	TangentZ;
+	uint32			Color;
+	float			U;
+	float			V;
+
+	inline void Initialize(float InX, float InY, float InU, float InV)
+	{
+		Position.X = InX; 
+		Position.Y = InY; 
+		Position.Z = 0.0f;
+		TangentX = FVector(1.0f, 0.0f, 0.0f); 
+		//TangentY = FVector(0.0f, 1.0f, 0.0f); 
+		TangentZ = FVector(0.0f, 0.0f, 1.0f);
+		// TangentZ.w contains the sign of the tangent basis determinant. Assume +1
+		TangentZ.Vector.W = 255;
+		Color = FColor(255,255,255,255).DWColor();
+		U = InU; 
+		V = InV;
+	}
+};
+
+/** 
+* Vertex buffer
+*/
+class FMaterialTriangleVertexBuffer : public FVertexBuffer
+{
+public:
+	/** 
+	* Initialize the RHI for this rendering resource 
+	*/
+	virtual void InitRHI() override
+	{
+		// used with a tristrip, so only 3 vertices are needed
+		uint32 Size = 3 * sizeof(FMaterialTriangleVertex);
+		// create vertex buffer
+		FRHIResourceCreateInfo CreateInfo;
+		VertexBufferRHI = RHICreateVertexBuffer(Size,BUF_Static,CreateInfo);
+		// lock it
+		void* Buffer = RHILockVertexBuffer(VertexBufferRHI,0,Size,RLM_WriteOnly);
+		// first vertex element
+		FMaterialTriangleVertex* DestVertex = (FMaterialTriangleVertex*)Buffer;
+		// fill out the verts
+		DestVertex[0].Initialize(1, -1, 1, 1);
+		DestVertex[1].Initialize(1, 1, 1, 0);
+		DestVertex[2].Initialize(-1, -1, 0, 1);
+			
+		// Unlock the buffer.
+		RHIUnlockVertexBuffer(VertexBufferRHI);        
+	}
+};
+
+TGlobalResource<FMaterialTriangleVertexBuffer> GTriangleRendererVertexBuffer;
+
+/**
+* Vertex factory for rendering tiles.
+*/
+class FTriangleVertexFactory : public FLocalVertexFactory
+{
+public:
+
+/** Default constructor. */
+FTriangleVertexFactory()
+{
+	FLocalVertexFactory::DataType Data;
+	// position
+	Data.PositionComponent = FVertexStreamComponent(
+		&GTriangleRendererVertexBuffer, STRUCT_OFFSET(FMaterialTriangleVertex, Position), sizeof(FMaterialTriangleVertex), VET_Float3);
+		// tangents
+		Data.TangentBasisComponents[0] = FVertexStreamComponent(
+		&GTriangleRendererVertexBuffer, STRUCT_OFFSET(FMaterialTriangleVertex, TangentX), sizeof(FMaterialTriangleVertex), VET_PackedNormal);
+		Data.TangentBasisComponents[1] = FVertexStreamComponent(
+		&GTriangleRendererVertexBuffer, STRUCT_OFFSET(FMaterialTriangleVertex, TangentZ), sizeof(FMaterialTriangleVertex), VET_PackedNormal);
+		// color
+		Data.ColorComponent = FVertexStreamComponent(
+		&GTriangleRendererVertexBuffer, STRUCT_OFFSET(FMaterialTriangleVertex, Color), sizeof(FMaterialTriangleVertex), VET_Color);
+		// UVs
+		Data.TextureCoordinates.Add(FVertexStreamComponent(
+		&GTriangleRendererVertexBuffer, STRUCT_OFFSET(FMaterialTriangleVertex, U), sizeof(FMaterialTriangleVertex), VET_Float2));
+		
+		// update the data
+		SetData(Data);
+	}
+};
+
+TGlobalResource<FTriangleVertexFactory> GTriangleVertexFactory;
+
+/**
+ * Mesh used to render triangles.
+ */
+class FTriangleMesh : public FRenderResource
+{
+public:
+
+	/** The mesh element. */
+	FMeshBatch TriMeshElement;
+
+	virtual void InitRHI() override
+	{
+		FMeshBatchElement& BatchElement = TriMeshElement.Elements[0];
+		TriMeshElement.VertexFactory = &GTriangleVertexFactory;
+		TriMeshElement.DynamicVertexStride = sizeof(FMaterialTriangleVertex);
+		BatchElement.FirstIndex = 0;
+		BatchElement.NumPrimitives = 1;
+		BatchElement.MinVertexIndex = 0;
+		BatchElement.MaxVertexIndex = 2;
+		TriMeshElement.ReverseCulling = false;
+		TriMeshElement.bDisableBackfaceCulling = true;
+		TriMeshElement.UseDynamicData = true;
+		TriMeshElement.Type = PT_TriangleList;
+		TriMeshElement.DepthPriorityGroup = SDPG_Foreground;
+		BatchElement.PrimitiveUniformBufferResource = &GIdentityPrimitiveUniformBuffer;
+	}
+	
+	virtual void ReleaseRHI() override
+	{
+		TriMeshElement.Elements[0].PrimitiveUniformBuffer.SafeRelease();
+	}
+};
+TGlobalResource<FTriangleMesh> GTriangleMesh;
+
+void FTriangleRenderer::DrawTriangle(FRHICommandListImmediate& RHICmdList, const class FSceneView& View, const FMaterialRenderProxy* MaterialRenderProxy, bool bNeedsToSwitchVerticalAxis, const FCanvasUVTri& Tri, bool bIsHitTesting, const FHitProxyId HitProxyId, const FColor InVertexColor)
+{
+	FMaterialTriangleVertex DestVertex[3];
+
+	// create verts
+	if (bNeedsToSwitchVerticalAxis)
+	{
+		DestVertex[0].Initialize(Tri.V1_Pos.X, View.ViewRect.Height() - Tri.V1_Pos.Y, Tri.V1_UV.X, Tri.V1_UV.Y);
+		DestVertex[1].Initialize(Tri.V0_Pos.X, View.ViewRect.Height() - Tri.V0_Pos.Y, Tri.V0_UV.X, Tri.V0_UV.Y);
+		DestVertex[2].Initialize(Tri.V2_Pos.X, View.ViewRect.Height() - Tri.V2_Pos.Y, Tri.V2_UV.X, Tri.V2_UV.Y);
+	}
+	else
+	{
+		DestVertex[0].Initialize(Tri.V1_Pos.X, Tri.V1_Pos.Y, Tri.V1_UV.X, Tri.V1_UV.Y);
+		DestVertex[1].Initialize(Tri.V0_Pos.X, Tri.V0_Pos.Y, Tri.V0_UV.X, Tri.V0_UV.Y);
+		DestVertex[2].Initialize(Tri.V2_Pos.X, Tri.V2_Pos.Y, Tri.V2_UV.X, Tri.V2_UV.Y);
+	}
+	
+	DestVertex[0].Color = InVertexColor.DWColor();
+	DestVertex[1].Color = InVertexColor.DWColor();
+	DestVertex[2].Color = InVertexColor.DWColor();
+	
+	// update the FMeshBatch
+	FMeshBatch& TriMesh = GTriangleMesh.TriMeshElement;
+	TriMesh.UseDynamicData = true;
+	TriMesh.DynamicVertexData = DestVertex;
+	TriMesh.MaterialRenderProxy = MaterialRenderProxy;
+	
+	GetRendererModule().DrawTileMesh(RHICmdList, View, TriMesh, bIsHitTesting, HitProxyId);
+}
+
+bool FCanvasTriangleRendererItem::Render_RenderThread(FRHICommandListImmediate& RHICmdList, const FCanvas* Canvas)
+{
+	float CurrentRealTime = 0.f;
+	float CurrentWorldTime = 0.f;
+	float DeltaWorldTime = 0.f;
+
+	if (!bFreezeTime)
+	{
+		CurrentRealTime = Canvas->GetCurrentRealTime();
+		CurrentWorldTime = Canvas->GetCurrentWorldTime();
+		DeltaWorldTime = Canvas->GetCurrentDeltaWorldTime();
+	}
+
+	checkSlow(Data);
+
+	// current render target set for the canvas
+	const FRenderTarget* CanvasRenderTarget = Canvas->GetRenderTarget();
+	FSceneViewFamily* ViewFamily = new FSceneViewFamily(FSceneViewFamily::ConstructionValues(
+		CanvasRenderTarget,
+		nullptr,
+		FEngineShowFlags(ESFIM_Game))
+		.SetWorldTimes(CurrentWorldTime, DeltaWorldTime, CurrentRealTime)
+		.SetGammaCorrection(CanvasRenderTarget->GetDisplayGamma()));
+
+	FIntRect ViewRect(FIntPoint(0, 0), CanvasRenderTarget->GetSizeXY());
+
+	// make a temporary viewk
+
+	FSceneViewInitOptions ViewInitOptions;
+	ViewInitOptions.ViewFamily = ViewFamily;
+	ViewInitOptions.SetViewRectangle(ViewRect);
+	ViewInitOptions.ViewOrigin = FVector::ZeroVector;
+	ViewInitOptions.ViewRotationMatrix = FMatrix::Identity;
+	ViewInitOptions.ProjectionMatrix = Data->Transform.GetMatrix();
+	ViewInitOptions.BackgroundColor = FLinearColor::Black;
+	ViewInitOptions.OverlayColor = FLinearColor::White;
+
+	bool bNeedsToSwitchVerticalAxis = RHINeedsToSwitchVerticalAxis(Canvas->GetShaderPlatform()) && !Canvas->GetAllowSwitchVerticalAxis();
+
+	FSceneView* View = new FSceneView(ViewInitOptions);
+
+	for (int32 TriIdx = 0; TriIdx < Data->Triangles.Num(); TriIdx++)
+	{
+		const FRenderData::FTriangleInst& Tri = Data->Triangles[TriIdx];
+		FTriangleRenderer::DrawTriangle(
+			RHICmdList,
+			*View,
+			Data->MaterialRenderProxy,
+			bNeedsToSwitchVerticalAxis,
+			Tri.Tri,
+			Canvas->IsHitTesting(), Tri.HitProxyId
+			);
+	}
+
+	delete View->Family;
+	delete View;
+	if (Canvas->GetAllowedModes() & FCanvas::Allow_DeleteOnRender)
+	{
+		delete Data;
+	}
+	if (Canvas->GetAllowedModes() & FCanvas::Allow_DeleteOnRender)
+	{
+		Data = nullptr;
+	}
+	return true;
+}
+
+bool FCanvasTriangleRendererItem::Render_GameThread(const FCanvas* Canvas)
+{
+	float CurrentRealTime = 0.f;
+	float CurrentWorldTime = 0.f;
+	float DeltaWorldTime = 0.f;
+
+	if (!bFreezeTime)
+	{
+		CurrentRealTime = Canvas->GetCurrentRealTime();
+		CurrentWorldTime = Canvas->GetCurrentWorldTime();
+		DeltaWorldTime = Canvas->GetCurrentDeltaWorldTime();
+	}
+
+	checkSlow(Data);
+	// current render target set for the canvas
+	const FRenderTarget* CanvasRenderTarget = Canvas->GetRenderTarget();
+	FSceneViewFamily* ViewFamily = new FSceneViewFamily(FSceneViewFamily::ConstructionValues(
+		CanvasRenderTarget,
+		nullptr,
+		FEngineShowFlags(ESFIM_Game))
+		.SetWorldTimes(CurrentWorldTime, DeltaWorldTime, CurrentRealTime)
+		.SetGammaCorrection(CanvasRenderTarget->GetDisplayGamma()));
+
+	FIntRect ViewRect(FIntPoint(0, 0), CanvasRenderTarget->GetSizeXY());
+
+	// make a temporary view
+	FSceneViewInitOptions ViewInitOptions;
+	ViewInitOptions.ViewFamily = ViewFamily;
+	ViewInitOptions.SetViewRectangle(ViewRect);
+	ViewInitOptions.ViewOrigin = FVector::ZeroVector;
+	ViewInitOptions.ViewRotationMatrix = FMatrix::Identity;
+	ViewInitOptions.ProjectionMatrix = Data->Transform.GetMatrix();
+	ViewInitOptions.BackgroundColor = FLinearColor::Black;
+	ViewInitOptions.OverlayColor = FLinearColor::White;
+
+	FSceneView* View = new FSceneView(ViewInitOptions);
+
+	bool bNeedsToSwitchVerticalAxis = RHINeedsToSwitchVerticalAxis(Canvas->GetShaderPlatform()) && !Canvas->GetAllowSwitchVerticalAxis();
+	struct FDrawTriangleParameters
+	{
+		FSceneView* View;
+		FRenderData* RenderData;
+		uint32 bIsHitTesting : 1;
+		uint32 bNeedsToSwitchVerticalAxis : 1;
+		uint32 AllowedCanvasModes;
+	};
+	FDrawTriangleParameters DrawTriangleParameters =
+	{
+		View,
+		Data,
+		bNeedsToSwitchVerticalAxis,
+		Canvas->IsHitTesting(),
+		Canvas->GetAllowedModes()
+	};
+	ENQUEUE_UNIQUE_RENDER_COMMAND_ONEPARAMETER(
+		DrawTriangleCommand,
+		FDrawTriangleParameters, Parameters, DrawTriangleParameters,
+		{
+		SCOPED_DRAW_EVENT(RHICmdList, CanvasDrawTriangle);
+		for (int32 TriIdx = 0; TriIdx < Parameters.RenderData->Triangles.Num(); TriIdx++)
+		{
+			const FRenderData::FTriangleInst& Tri = Parameters.RenderData->Triangles[TriIdx];
+			FTriangleRenderer::DrawTriangle(
+				RHICmdList,
+				*Parameters.View,
+				Parameters.RenderData->MaterialRenderProxy,
+				Parameters.bNeedsToSwitchVerticalAxis,
+				Tri.Tri,
+				Parameters.bIsHitTesting, Tri.HitProxyId);
+		}
+
+		delete Parameters.View->Family;
+		delete Parameters.View;
+		if (Parameters.AllowedCanvasModes & FCanvas::Allow_DeleteOnRender)
+		{
+			delete Parameters.RenderData;
+		}
+		});
+	if (Canvas->GetAllowedModes() & FCanvas::Allow_DeleteOnRender)
+	{
+		Data = nullptr;
+	}
+	return true;
+}
\ No newline at end of file
diff --git a/Engine/Source/Runtime/Engine/Private/UserInterface/Canvas.cpp b/Engine/Source/Runtime/Engine/Private/UserInterface/Canvas.cpp
--- a/Engine/Source/Runtime/Engine/Private/UserInterface/Canvas.cpp
+++ b/Engine/Source/Runtime/Engine/Private/UserInterface/Canvas.cpp
@@ -9,6 +9,7 @@
 #include "Engine/Font.h"
 #include "EngineFontServices.h"
 #include "TileRendering.h"
+#include "TriangleRendering.h"
 #include "RHIStaticStates.h"
 #include "BreakIterator.h"
 
@@ -26,6 +27,7 @@ DEFINE_STAT(STAT_Canvas_DrawStringTime);
 DEFINE_STAT(STAT_Canvas_WordWrappingTime);
 DEFINE_STAT(STAT_Canvas_GetBatchElementsTime);
 DEFINE_STAT(STAT_Canvas_AddTileRenderTime);
+DEFINE_STAT(STAT_Canvas_AddTriangleRenderTime);
 DEFINE_STAT(STAT_Canvas_NumBatchesCreated);
 
 FCanvasWordWrapper::FCanvasWordWrapper()
@@ -455,159 +457,7 @@ bool FCanvasBatchedElementRenderItem::Render_GameThread(const FCanvas* Canvas)
 }
 
 
-bool FCanvasTileRendererItem::Render_RenderThread(FRHICommandListImmediate& RHICmdList, const FCanvas* Canvas)
-{
-	float CurrentRealTime = 0.f;
-	float CurrentWorldTime = 0.f;
-	float DeltaWorldTime = 0.f;
-
-	if (!bFreezeTime)
-	{
-		CurrentRealTime = Canvas->GetCurrentRealTime();
-		CurrentWorldTime = Canvas->GetCurrentWorldTime();
-		DeltaWorldTime = Canvas->GetCurrentDeltaWorldTime();
-	}
-	
-	checkSlow(Data);
-	// current render target set for the canvas
-	const FRenderTarget* CanvasRenderTarget = Canvas->GetRenderTarget();
-	FSceneViewFamily* ViewFamily = new FSceneViewFamily( FSceneViewFamily::ConstructionValues(
-		CanvasRenderTarget,
-		NULL,
-		FEngineShowFlags(ESFIM_Game))
-		.SetWorldTimes( CurrentWorldTime, DeltaWorldTime, CurrentRealTime )
-		.SetGammaCorrection( CanvasRenderTarget->GetDisplayGamma() ) );
-
-	FIntRect ViewRect(FIntPoint(0, 0), CanvasRenderTarget->GetSizeXY());
-
-	// make a temporary view
-	FSceneViewInitOptions ViewInitOptions;
-	ViewInitOptions.ViewFamily = ViewFamily;
-	ViewInitOptions.SetViewRectangle(ViewRect);
-	ViewInitOptions.ViewOrigin = FVector::ZeroVector;
-	ViewInitOptions.ViewRotationMatrix = FMatrix::Identity;
-	ViewInitOptions.ProjectionMatrix = Data->Transform.GetMatrix();
-	ViewInitOptions.BackgroundColor = FLinearColor::Black;
-	ViewInitOptions.OverlayColor = FLinearColor::White;
-
-	bool bNeedsToSwitchVerticalAxis = RHINeedsToSwitchVerticalAxis(Canvas->GetShaderPlatform()) && !Canvas->GetAllowSwitchVerticalAxis();
-
-	FSceneView* View = new FSceneView(ViewInitOptions);
-
-	for( int32 TileIdx=0; TileIdx < Data->Tiles.Num(); TileIdx++ )
-	{
-		const FRenderData::FTileInst& Tile = Data->Tiles[TileIdx];
-		FTileRenderer::DrawTile(
-			RHICmdList, 
-			*View, 
-			Data->MaterialRenderProxy, 
-			bNeedsToSwitchVerticalAxis,
-			Tile.X, Tile.Y, Tile.SizeX, Tile.SizeY, 
-			Tile.U, Tile.V, Tile.SizeU, Tile.SizeV,
-			Canvas->IsHitTesting(), Tile.HitProxyId,
-			Tile.InColor
-			);
-	}
-	
-	delete View->Family;
-	delete View;
-	if( Canvas->GetAllowedModes() & FCanvas::Allow_DeleteOnRender )
-	{
-		delete Data;
-	}
-	if( Canvas->GetAllowedModes() & FCanvas::Allow_DeleteOnRender )
-	{
-		Data = NULL;
-	}
-	return true;
-}
-
-bool FCanvasTileRendererItem::Render_GameThread(const FCanvas* Canvas)
-{
-	float CurrentRealTime = 0.f;
-	float CurrentWorldTime = 0.f;
-	float DeltaWorldTime = 0.f;
 
-	if (!bFreezeTime)
-	{
-		CurrentRealTime = Canvas->GetCurrentRealTime();
-		CurrentWorldTime = Canvas->GetCurrentWorldTime();
-		DeltaWorldTime = Canvas->GetCurrentDeltaWorldTime();
-	}
-
-	checkSlow(Data);
-	// current render target set for the canvas
-	const FRenderTarget* CanvasRenderTarget = Canvas->GetRenderTarget();
-	FSceneViewFamily* ViewFamily = new FSceneViewFamily(FSceneViewFamily::ConstructionValues(
-		CanvasRenderTarget,
-		NULL,
-		FEngineShowFlags(ESFIM_Game))
-		.SetWorldTimes(CurrentWorldTime, DeltaWorldTime, CurrentRealTime)
-		.SetGammaCorrection(CanvasRenderTarget->GetDisplayGamma()));
-
-	FIntRect ViewRect(FIntPoint(0, 0), CanvasRenderTarget->GetSizeXY());
-
-	// make a temporary view
-	FSceneViewInitOptions ViewInitOptions;
-	ViewInitOptions.ViewFamily = ViewFamily;
-	ViewInitOptions.SetViewRectangle(ViewRect);
-	ViewInitOptions.ViewOrigin = FVector::ZeroVector;
-	ViewInitOptions.ViewRotationMatrix = FMatrix::Identity;
-	ViewInitOptions.ProjectionMatrix = Data->Transform.GetMatrix();
-	ViewInitOptions.BackgroundColor = FLinearColor::Black;
-	ViewInitOptions.OverlayColor = FLinearColor::White;
-
-	FSceneView* View = new FSceneView(ViewInitOptions);
-
-	bool bNeedsToSwitchVerticalAxis = RHINeedsToSwitchVerticalAxis(Canvas->GetShaderPlatform()) && !Canvas->GetAllowSwitchVerticalAxis();
-	struct FDrawTileParameters
-	{
-		FSceneView* View;
-		FRenderData* RenderData;
-		uint32 bIsHitTesting : 1;
-		uint32 bNeedsToSwitchVerticalAxis : 1;
-		uint32 AllowedCanvasModes;
-	};
-	FDrawTileParameters DrawTileParameters =
-	{
-		View,
-		Data,
-		Canvas->IsHitTesting(),
-		bNeedsToSwitchVerticalAxis,
-		Canvas->GetAllowedModes()
-	};
-	ENQUEUE_UNIQUE_RENDER_COMMAND_ONEPARAMETER(
-		DrawTileCommand,
-		FDrawTileParameters, Parameters, DrawTileParameters,
-	{
-		SCOPED_DRAW_EVENT(RHICmdList, CanvasDrawTile);
-		for (int32 TileIdx = 0; TileIdx < Parameters.RenderData->Tiles.Num(); TileIdx++)
-		{
-			const FRenderData::FTileInst& Tile = Parameters.RenderData->Tiles[TileIdx];
-			FTileRenderer::DrawTile(
-				RHICmdList,
-				*Parameters.View,
-				Parameters.RenderData->MaterialRenderProxy,
-				Parameters.bNeedsToSwitchVerticalAxis,
-				Tile.X, Tile.Y, Tile.SizeX, Tile.SizeY,
-				Tile.U, Tile.V, Tile.SizeU, Tile.SizeV,
-				Parameters.bIsHitTesting, Tile.HitProxyId,
-				Tile.InColor);
-		}
-
-		delete Parameters.View->Family;
-		delete Parameters.View;
-		if (Parameters.AllowedCanvasModes & FCanvas::Allow_DeleteOnRender)
-		{
-			delete Parameters.RenderData;
-		}
-	});
-	if (Canvas->GetAllowedModes() & FCanvas::Allow_DeleteOnRender)
-	{
-		Data = NULL;
-	}
-	return true;
-}
 
 FCanvas::FCanvasSortElement& FCanvas::GetSortElement(int32 DepthSortKey)
 {
@@ -703,6 +553,36 @@ void FCanvas::AddTileRenderItem(float X, float Y, float SizeX, float SizeY, floa
 	RenderBatch->AddTile( X,Y,SizeX,SizeY,U,V,SizeU,SizeV,HitProxyId,InColor);
 }
 
+void FCanvas::AddTriangleRenderItem(const FCanvasUVTri& Tri, const FMaterialRenderProxy* MaterialRenderProxy, FHitProxyId HitProxyId, bool bFreezeTime)
+{
+	SCOPE_CYCLE_COUNTER(STAT_Canvas_AddTriangleRenderTime);
+	
+	// get sort element based on the current sort key from top of sort key stack
+	FCanvasSortElement& SortElement = FCanvas::GetSortElement(TopDepthSortKey());
+	// find a batch to use 
+	FCanvasTriangleRendererItem* RenderBatch = NULL;
+
+	// get the current transform entry from top of transform stack
+	const FTransformEntry& TopTransformEntry = TransformStack.Top();
+	
+	// try to use the current top entry in the render batch array
+	if (SortElement.RenderBatchArray.Num() > 0)
+	{
+		checkSlow(SortElement.RenderBatchArray.Last());
+		RenderBatch = SortElement.RenderBatchArray.Last()->GetCanvasTriangleRendererItem();
+	}
+	// if a matching entry for this batch doesn't exist then allocate a new entry
+	if (RenderBatch == nullptr || !RenderBatch->IsMatch(MaterialRenderProxy, TopTransformEntry))
+	{
+		INC_DWORD_STAT(STAT_Canvas_NumBatchesCreated);
+	
+		RenderBatch = new FCanvasTriangleRendererItem(MaterialRenderProxy, TopTransformEntry, bFreezeTime);
+		SortElement.RenderBatchArray.Add(RenderBatch);
+	}
+	// add the triangle to the triangle render batch
+	RenderBatch->AddTriangle(Tri, HitProxyId);
+}
+
 FCanvas::~FCanvas()
 {
 	// delete batches from elements entries
@@ -1974,6 +1854,17 @@ void UCanvas::K2_DrawTriangle(UTexture* RenderTexture, TArray<FCanvasUVTri> Tria
 	}
 }
 
+void UCanvas::K2_DrawMaterialTriangle(UMaterialInterface* RenderMaterial, TArray<FCanvasUVTri> Triangles)
+{
+	if (RenderMaterial && Triangles.Num() > 0)
+	{
+		FCanvasTriangleItem TriangleItem(FVector2D::ZeroVector, FVector2D::ZeroVector, FVector2D::ZeroVector, NULL);
+		TriangleItem.MaterialRenderProxy = RenderMaterial->GetRenderProxy(0);
+		TriangleItem.TriangleList = Triangles;
+		DrawItem(TriangleItem);
+	}
+}
+
 void UCanvas::K2_DrawPolygon(UTexture* RenderTexture, FVector2D ScreenPosition, FVector2D Radius, int32 NumberOfSides, FLinearColor RenderColor)
 {
 	if (Radius.X > 0.0f && Radius.Y > 0.0f && NumberOfSides >= 3)
diff --git a/Engine/Source/Runtime/Engine/Private/UserInterface/CanvasItem.cpp b/Engine/Source/Runtime/Engine/Private/UserInterface/CanvasItem.cpp
--- a/Engine/Source/Runtime/Engine/Private/UserInterface/CanvasItem.cpp
+++ b/Engine/Source/Runtime/Engine/Private/UserInterface/CanvasItem.cpp
@@ -7,6 +7,7 @@
 #include "EnginePrivate.h"
 #include "CanvasItem.h"
 #include "TileRendering.h"
+#include "TriangleRendering.h"
 #include "RHIStaticStates.h"
 #include "Internationalization.h"
 #include "SlateBasics.h"
@@ -18,18 +19,29 @@ DECLARE_CYCLE_STAT(TEXT("CanvasTileMaterialItem Time"),STAT_Canvas_TileMaterialI
 DECLARE_CYCLE_STAT(TEXT("CanvasTextItem Time"),STAT_Canvas_TextItemTime,STATGROUP_Canvas);
 DECLARE_CYCLE_STAT(TEXT("CanvasLineItem Time"),STAT_Canvas_LineItemTime,STATGROUP_Canvas);
 DECLARE_CYCLE_STAT(TEXT("CanvasBoxItem Time"),STAT_Canvas_BoxItemTime,STATGROUP_Canvas);
-DECLARE_CYCLE_STAT(TEXT("CanvasTriItem Time"),STAT_Canvas_TriItemTime,STATGROUP_Canvas);
+DECLARE_CYCLE_STAT(TEXT("CanvasTriTextureItem Time"),STAT_Canvas_TriTextureItemTime,STATGROUP_Canvas);
+DECLARE_CYCLE_STAT(TEXT("CanvasTriMaterialItem Time"), STAT_Canvas_TriMaterialItemTime, STATGROUP_Canvas);
 DECLARE_CYCLE_STAT(TEXT("CanvasBorderItem Time"),STAT_Canvas_BorderItemTime,STATGROUP_Canvas);
 
 
 #if WITH_EDITOR
 #include "UnrealEd.h"
 #include "ObjectTools.h"
-#define LOCTEXT_NAMESPACE "CanvasItem"
 
 FCanvasItemTestbed::LineVars FCanvasItemTestbed::TestLine;
 bool FCanvasItemTestbed::bTestState = false;
-bool FCanvasItemTestbed::bShowTestbed = false;
+bool FCanvasItemTestbed::bShowTestbed = true;
+
+bool FCanvasItemTestbed::bShowLines = false;
+bool FCanvasItemTestbed::bShowBoxes = false;
+bool FCanvasItemTestbed::bShowTris = true;
+bool FCanvasItemTestbed::bShowText = false;
+bool FCanvasItemTestbed::bShowTiles = false;
+
+FCanvasItemTestbed::FCanvasItemTestbed()
+{
+	TestMaterial=nullptr;
+}
 
 void FCanvasItemTestbed::Draw( class FViewport* Viewport, class FCanvas* Canvas )
 {
@@ -40,6 +52,10 @@ void FCanvasItemTestbed::Draw( class FViewport* Viewport, class FCanvas* Canvas
 		return;
 	}
 
+	if (TestMaterial == nullptr)
+	{
+		TestMaterial = LoadObject<UMaterial>(NULL, TEXT("/Game/NewMaterial.NewMaterial"));
+	}
 	
 	// A little ott for a testbed - but I wanted to draw several lines to ensure it worked :)
 	if( TestLine.bTestSet == false )
@@ -93,76 +109,114 @@ void FCanvasItemTestbed::Draw( class FViewport* Viewport, class FCanvas* Canvas
 	}
 	
 	// Text
-	float CenterX = 0.0f;
-	float YTest = 16.0f;
-	FCanvasTextItem TextItem( FVector2D( CenterX, YTest ), LOCTEXT( "stringhere", "String Here" ), GEngine->GetSmallFont(), FLinearColor::Red );	
-	TextItem.Draw( Canvas );
-	
-	// Shadowed text
-	TextItem.Position.Y += TextItem.DrawnSize.Y;
-	TextItem.Scale.X = 2.0f;
-	TextItem.EnableShadow( FLinearColor::Green, FVector2D( 2.0f, 2.0f ) );
-	TextItem.Text = LOCTEXT( "Scaled String here", "Scaled String here" );
-	TextItem.Draw( Canvas );
-	TextItem.DisableShadow();
-
-	TextItem.Position.Y += TextItem.DrawnSize.Y;;
-	TextItem.Text = LOCTEXT( "CenterdStringhere", "CenterdStringhere" );
-	TextItem.Scale.X = 1.0f;
-	TextItem.bCentreX = true;
-	TextItem.Draw( Canvas );
-
-	// Outlined text
-	TextItem.Position.Y += TextItem.DrawnSize.Y;
-	TextItem.Text = LOCTEXT( "ScaledCentredStringhere", "Scaled Centred String here" );	
-	TextItem.OutlineColor = FLinearColor::Black;
-	TextItem.bOutlined = true;
-	TextItem.Scale = FVector2D( 2.0f, 2.0f );
-	TextItem.SetColor( FLinearColor::Green );
-	TextItem.Text = LOCTEXT( "ScaledCentredOutlinedStringhere", "Scaled Centred Outlined String here" );	
-	TextItem.Draw( Canvas );
+	if( bShowText == true )
+	{
+		float CenterX = Canvas->GetViewRect().Width() / 2.0f;
+		float YTest = 16.0f;
+		FCanvasTextItem TextItem(FVector2D(CenterX, YTest), FText::FromString(TEXT("String Here")), GEngine->GetSmallFont(), FLinearColor::Red);
+		TextItem.Draw(Canvas);
+
+		// Shadowed text
+		TextItem.Position.Y += TextItem.DrawnSize.Y;
+		TextItem.Scale.X = 2.0f;
+		TextItem.EnableShadow(FLinearColor::Green, FVector2D(2.0f, 2.0f));
+		TextItem.Text = FText::FromString(TEXT("Scaled String here"));
+		TextItem.Draw(Canvas);
+		TextItem.DisableShadow();
+
+		TextItem.Position.Y += TextItem.DrawnSize.Y;;
+		TextItem.Text = FText::FromString(TEXT("Centered String Here"));
+		TextItem.Scale.X = 1.0f;
+		TextItem.bCentreX = true;
+		TextItem.Draw(Canvas);
+
+		// Outlined text
+		TextItem.Position.Y += TextItem.DrawnSize.Y;
+		TextItem.Text = FText::FromString(TEXT("Scaled Centred String here"));
+		TextItem.OutlineColor = FLinearColor::Black;
+		TextItem.bOutlined = true;
+		TextItem.Scale = FVector2D(2.0f, 2.0f);
+		TextItem.SetColor(FLinearColor::Green);
+		TextItem.Text = FText::FromString(TEXT("Scaled Centred Outlined String here"));
+		TextItem.Draw(Canvas);
+	}
 	
 	// a line
-	FCanvasLineItem LineItem( TestLine.LineStart, TestLine.LineEnd );
-	LineItem.Draw( Canvas );
-
-	// some boxes
-	FCanvasBoxItem BoxItem( FVector2D( 88.0f, 88.0f ), FVector2D( 188.0f, 188.0f ) );
-	BoxItem.SetColor( FLinearColor::Yellow );
-	BoxItem.Draw( Canvas );
-
-	BoxItem.SetColor( FLinearColor::Red );
-	BoxItem.Position = FVector2D( 256.0f, 256.0f );
-	BoxItem.Draw( Canvas );
-
-	BoxItem.SetColor( FLinearColor::Blue );
-	BoxItem.Position = FVector2D( 6.0f, 6.0f );
-	BoxItem.Size = FVector2D( 48.0f, 96.0f );
-	BoxItem.Draw( Canvas );
-	
-	// Triangle
-	FCanvasTriangleItem TriItem(  FVector2D( 48.0f, 48.0f ), FVector2D( 148.0f, 48.0f ), FVector2D( 48.0f, 148.0f ), GWhiteTexture );
-	TriItem.Draw( Canvas );
-
-	// Triangle list
-	TArray< FCanvasUVTri >	TriangleList;
-	FCanvasUVTri SingleTri;
-	SingleTri.V0_Pos = FVector2D( 128.0f, 128.0f );
-	SingleTri.V1_Pos = FVector2D( 248.0f, 108.0f );
-	SingleTri.V2_Pos = FVector2D( 100.0f, 348.0f );
-	SingleTri.V0_UV = FVector2D::ZeroVector;
-	SingleTri.V1_UV = FVector2D::ZeroVector;
-	SingleTri.V2_UV = FVector2D::ZeroVector;
-	TriangleList.Add( SingleTri );
-	SingleTri.V0_Pos = FVector2D( 348.0f, 128.0f );
-	SingleTri.V1_Pos = FVector2D( 448.0f, 148.0f );
-	SingleTri.V2_Pos = FVector2D( 438.0f, 308.0f );
-	TriangleList.Add( SingleTri );
-
-	FCanvasTriangleItem TriItemList( TriangleList, GWhiteTexture );
-	TriItemList.SetColor( FLinearColor::Red );
-	TriItemList.Draw( Canvas );
+	if( bShowLines == true )
+	{
+		FCanvasLineItem LineItem(TestLine.LineStart, TestLine.LineEnd);
+		LineItem.Draw(Canvas);
+	}
+
+// 	some boxes
+	if( bShowBoxes == true )
+	{
+		FCanvasBoxItem BoxItem(FVector2D(88.0f, 88.0f), FVector2D(188.0f, 188.0f));
+		BoxItem.SetColor(FLinearColor::Yellow);
+		BoxItem.Draw(Canvas);
+
+		BoxItem.SetColor(FLinearColor::Red);
+		BoxItem.Position = FVector2D(256.0f, 256.0f);
+		BoxItem.Draw(Canvas);
+
+		BoxItem.SetColor(FLinearColor::Blue);
+		BoxItem.Position = FVector2D(6.0f, 6.0f);
+		BoxItem.Size = FVector2D(48.0f, 96.0f);
+		BoxItem.Draw(Canvas);
+
+	}
 	
+	if (bShowTris == true)
+	{
+		// Triangle
+		FCanvasTriangleItem TriItem(FVector2D(32.0f, 32.0f), FVector2D(64.0f, 32.0f), FVector2D(64.0f, 64.0f), GWhiteTexture);
+		//TriItem.Draw(Canvas);
+
+		// Triangle list
+		TArray< FCanvasUVTri >	TriangleList;
+		FCanvasUVTri SingleTri;
+		SingleTri.V0_Pos = FVector2D(128.0f, 128.0f);
+		SingleTri.V1_Pos = FVector2D(248.0f, 108.0f);
+		SingleTri.V2_Pos = FVector2D(100.0f, 348.0f);
+		TriangleList.Add(SingleTri);
+		SingleTri.V0_Pos = FVector2D(348.0f, 128.0f);
+		SingleTri.V1_Pos = FVector2D(448.0f, 148.0f);
+		SingleTri.V2_Pos = FVector2D(438.0f, 308.0f);
+		TriangleList.Add(SingleTri);
+
+		FCanvasTriangleItem TriItemList(TriangleList, GWhiteTexture);
+		TriItemList.SetColor(FLinearColor::Red);
+		//TriItemList.Draw( Canvas );
+
+		if (TestMaterial)
+		{
+			FCanvasTileItem TileItemMat(FVector2D(256.0f, 256.0f), TestMaterial->GetRenderProxy(false), FVector2D(128.0f, 128.0f));
+			//TileItemMat.Draw(Canvas);
+
+			FCanvasTriangleItem TriItem(FVector2D(512.0f, 256.0f), FVector2D(512.0f, 256.0f), FVector2D(640.0f, 384.0f), FVector2D::ZeroVector, FVector2D(1.0f, 0.0f), FVector2D(1.0f, 1.0f), nullptr);
+			TriItem.MaterialRenderProxy = TestMaterial->GetRenderProxy(false);
+			//TriItem.Draw(Canvas);
+
+			SingleTri.V0_Pos = FVector2D(228.0f, 228.0f);
+			SingleTri.V1_Pos = FVector2D(348.0f, 208.0f);
+			SingleTri.V2_Pos = FVector2D(200.0f, 448.0f);
+			SingleTri.V0_UV = FVector2D(0.0f, 0.0f);
+			SingleTri.V1_UV = FVector2D(1.0f, 0.0f);
+			SingleTri.V2_UV = FVector2D(1.0f, 1.0f);
+			TriangleList.Add(SingleTri);
+			SingleTri.V0_Pos = FVector2D(448.0f, 228.0f);
+			SingleTri.V1_Pos = FVector2D(548.0f, 248.0f);
+			SingleTri.V2_Pos = FVector2D(538.0f, 408.0f);
+			SingleTri.V0_UV = FVector2D(0.0f, 1.0f);
+			SingleTri.V1_UV = FVector2D(0.0f, 0.0f);
+			SingleTri.V2_UV = FVector2D(1.0f, 0.0f);
+			TriangleList.Add(SingleTri);
+			FCanvasTriangleItem TriItemList(TriangleList, nullptr);
+			TriItemList.MaterialRenderProxy = TestMaterial->GetRenderProxy(false);
+			TriItemList.Draw(Canvas);
+		}
+	}
+
 // 	FCanvasNGonItem NGon( FVector2D( 256.0f, 256.0f ), FVector2D( 256.0f, 256.0f ), 6, GWhiteTexture, FLinearColor::White );
 // 	NGon.Draw( Canvas );
 // 
@@ -171,7 +225,7 @@ void FCanvasItemTestbed::Draw( class FViewport* Viewport, class FCanvas* Canvas
 
 	// Texture
 	UTexture* SelectedTexture = GEditor->GetSelectedObjects()->GetTop<UTexture>();	
-	if( SelectedTexture )
+	if( ( SelectedTexture ) && ( bShowTiles == true ))
 	{
 		// Plain tex
 		FCanvasTileItem TileItem( FVector2D( 128.0f,128.0f ), SelectedTexture->Resource, FLinearColor::White );
@@ -226,18 +280,17 @@ void FCanvasItemTestbed::Draw( class FViewport* Viewport, class FCanvas* Canvas
 		TileItem.Draw( Canvas );
 
 		TestLine.Testangle = FMath::Fmod( TestLine.Testangle + 2.0f, 360.0f );
-
-		// textured tri
-		FCanvasTriangleItem TriItemTex(  FVector2D( 48.0f, 48.0f ), FVector2D( 148.0f, 48.0f ), FVector2D( 48.0f, 148.0f ), FVector2D( 0.0f, 0.0f ), FVector2D( 1.0f, 0.0f ), FVector2D( 0.0f, 1.0f ), SelectedTexture->Resource  );
-		TriItem.Texture = GWhiteTexture;
-		TriItemTex.Draw( Canvas );
-
-		// moving tri (only 1 point moves !)
-		TriItemTex.Position = TestLine.LineStart;
-		TriItemTex.Draw( Canvas );
+// 
+// 		// textured tri
+// 		FCanvasTriangleItem TriItemTex(  FVector2D( 48.0f, 48.0f ), FVector2D( 148.0f, 48.0f ), FVector2D( 48.0f, 148.0f ), FVector2D( 0.0f, 0.0f ), FVector2D( 1.0f, 0.0f ), FVector2D( 0.0f, 1.0f ), SelectedTexture->Resource  );
+// 		TriItem.Texture = GWhiteTexture;
+// 		TriItemTex.Draw( Canvas );
+// 
+// 		// moving tri (only 1 point moves !)
+// 		TriItemTex.Position = TestLine.LineStart;
+// 		TriItemTex.Draw( Canvas );
 	}
 }
-#undef LOCTEXT_NAMESPACE
 
 #endif // WITH_EDITOR
 
@@ -1236,27 +1289,66 @@ void FCanvasBoxItem::SetupBox()
 
 void FCanvasTriangleItem::Draw( class FCanvas* InCanvas )
 {
-	SCOPE_CYCLE_COUNTER(STAT_Canvas_TriItemTime);
+	if (Texture != nullptr)
+	{
+		SCOPE_CYCLE_COUNTER(STAT_Canvas_TriTextureItemTime);
+		FBatchedElements* BatchedElements = InCanvas->GetBatchedElements(FCanvas::ET_Triangle, BatchedElementParameters, Texture, BlendMode);
 
-	FBatchedElements* BatchedElements = InCanvas->GetBatchedElements(FCanvas::ET_Triangle, BatchedElementParameters, Texture, BlendMode);
-	
-	FHitProxyId HitProxyId = InCanvas->GetHitProxyId();
+		FHitProxyId HitProxyId = InCanvas->GetHitProxyId();
 
-	for(int32 i=0; i<TriangleList.Num(); i++)
+		for (int32 i = 0; i < TriangleList.Num(); i++)
+		{
+			const FCanvasUVTri& Tri = TriangleList[i];
+			int32 V0 = BatchedElements->AddVertex(FVector4(Tri.V0_Pos.X, Tri.V0_Pos.Y, 0, 1), Tri.V0_UV, Tri.V0_Color, HitProxyId);
+			int32 V1 = BatchedElements->AddVertex(FVector4(Tri.V1_Pos.X, Tri.V1_Pos.Y, 0, 1), Tri.V1_UV, Tri.V1_Color, HitProxyId);
+			int32 V2 = BatchedElements->AddVertex(FVector4(Tri.V2_Pos.X, Tri.V2_Pos.Y, 0, 1), Tri.V2_UV, Tri.V2_Color, HitProxyId);
+
+			if (BatchedElementParameters)
+			{
+				BatchedElements->AddTriangle(V0, V1, V2, BatchedElementParameters, BlendMode);
+			}
+			else
+			{
+				check(Texture);
+				BatchedElements->AddTriangle(V0, V1, V2, Texture, BlendMode);
+			}
+		}
+	}
+	else
 	{
-		const FCanvasUVTri& Tri = TriangleList[i];
-		int32 V0 = BatchedElements->AddVertex(FVector4(Tri.V0_Pos.X,Tri.V0_Pos.Y,0,1), Tri.V0_UV, Tri.V0_Color, HitProxyId);
-		int32 V1 = BatchedElements->AddVertex(FVector4(Tri.V1_Pos.X,Tri.V1_Pos.Y,0,1), Tri.V1_UV, Tri.V1_Color, HitProxyId);
-		int32 V2 = BatchedElements->AddVertex(FVector4(Tri.V2_Pos.X,Tri.V2_Pos.Y,0,1), Tri.V2_UV, Tri.V2_Color, HitProxyId);
+		SCOPE_CYCLE_COUNTER(STAT_Canvas_TriMaterialItemTime);
+
+		// get sort element based on the current sort key from top of sort key stack
+		FCanvas::FCanvasSortElement& SortElement = InCanvas->GetSortElement(InCanvas->TopDepthSortKey());
+
+		// find a batch to use 
+		FCanvasTriangleRendererItem* RenderBatch = nullptr;
 
-		if( BatchedElementParameters )
+		// get the current transform entry from top of transform stack
+		const FCanvas::FTransformEntry& TopTransformEntry = InCanvas->GetTransformStack().Top();
+
+		// try to use the current top entry in the render batch array
+		if (SortElement.RenderBatchArray.Num() > 0)
 		{
-			BatchedElements->AddTriangle(V0,V1,V2, BatchedElementParameters, BlendMode);
+			checkSlow(SortElement.RenderBatchArray.Last());
+			RenderBatch = SortElement.RenderBatchArray.Last()->GetCanvasTriangleRendererItem();
 		}
-		else
+
+		// if a matching entry for this batch doesn't exist then allocate a new entry
+		if (RenderBatch == nullptr || !RenderBatch->IsMatch(MaterialRenderProxy, TopTransformEntry))
+		{
+			INC_DWORD_STAT(STAT_Canvas_NumBatchesCreated);
+
+			RenderBatch = new FCanvasTriangleRendererItem(MaterialRenderProxy, TopTransformEntry, bFreezeTime);
+			SortElement.RenderBatchArray.Add(RenderBatch);
+		}
+
+		FHitProxyId HitProxyId = InCanvas->GetHitProxyId();
+
+		// add the triangles to the triangle render batch
+		for (int32 i = 0; i < TriangleList.Num(); i++)
 		{
-			check( Texture );
-			BatchedElements->AddTriangle(V0,V1,V2,Texture, BlendMode);
+			RenderBatch->AddTriangle(TriangleList[i], HitProxyId);
 		}
 	}	
 }
diff --git a/Engine/Source/Runtime/Engine/Public/CanvasItem.h b/Engine/Source/Runtime/Engine/Public/CanvasItem.h
--- a/Engine/Source/Runtime/Engine/Public/CanvasItem.h
+++ b/Engine/Source/Runtime/Engine/Public/CanvasItem.h
@@ -15,6 +15,7 @@ class FBatchedElementParameters;
 class FBatchedElements;
 class FMaterialRenderProxy;
 class FTexture;
+class UMaterial;
 
 class FCanvasItem
 {
@@ -631,7 +632,8 @@ public:
 	FCanvasTriangleItem( const FVector2D& InPointA, const FVector2D& InPointB, const FVector2D& InPointC, const FTexture* InTexture )
 		: FCanvasItem( InPointA )
 		, Texture( InTexture )
-		, BatchedElementParameters( NULL )
+		, MaterialRenderProxy( nullptr )
+		, BatchedElementParameters( nullptr )
 	{
 		FCanvasUVTri SingleTri;
 		SingleTri.V0_Pos = InPointA;
@@ -663,7 +665,8 @@ public:
 	FCanvasTriangleItem( const FVector2D& InPointA, const FVector2D& InPointB, const FVector2D& InPointC, const FVector2D& InTexCoordPointA, const FVector2D& InTexCoordPointB, const FVector2D& InTexCoordPointC, const FTexture* InTexture )
 		:FCanvasItem( InPointA )
 		, Texture( InTexture )
-		, BatchedElementParameters( NULL )
+		, MaterialRenderProxy( nullptr )
+		, BatchedElementParameters( nullptr )
 	{
 		FCanvasUVTri SingleTri;
 		SingleTri.V0_Pos = InPointA;
@@ -689,7 +692,8 @@ public:
 	FCanvasTriangleItem( FCanvasUVTri InSingleTri, const FTexture* InTexture )
 		:FCanvasItem( InSingleTri.V0_Pos )
 		, Texture( InTexture )
-		, BatchedElementParameters( NULL )
+		, MaterialRenderProxy( nullptr )
+		, BatchedElementParameters( nullptr )
 	{
 		TriangleList.Add( InSingleTri );
 	};
@@ -703,7 +707,8 @@ public:
 	FCanvasTriangleItem( const TArray< FCanvasUVTri >&	InTriangleList, const FTexture* InTexture )
 		:FCanvasItem( FVector2D::ZeroVector )
 		, Texture( InTexture )
-		, BatchedElementParameters( NULL )
+		, MaterialRenderProxy( nullptr )
+		, BatchedElementParameters( nullptr )
 	{
 		check( InTriangleList.Num() >= 1 );
 		
@@ -733,6 +738,9 @@ public:
 	/* texture to use for triangle(s). */
 	const FTexture* Texture;
 
+	/* Material proxy for rendering. */
+	const FMaterialRenderProxy* MaterialRenderProxy;
+
 	FBatchedElementParameters* BatchedElementParameters;
 
 	/* List of triangles. */
@@ -824,7 +832,7 @@ private:
 class ENGINE_API FCanvasItemTestbed
 {
 public:
-	FCanvasItemTestbed(){};
+	FCanvasItemTestbed();
 	void  Draw( class FViewport* Viewport, FCanvas* Canvas );
 	struct LineVars
 	{
@@ -847,6 +855,14 @@ public:
 	static LineVars TestLine;
 	static bool	bTestState;
 	static bool	bShowTestbed;
+	static bool bShowLines;
+	static bool bShowBoxes;
+	static bool bShowTris;
+	static bool bShowText;
+	static bool bShowTiles;
+
+	UMaterial* TestMaterial;
+	
 };
 #endif // WITH_EDITOR
 
diff --git a/Engine/Source/Runtime/Engine/Public/CanvasTypes.h b/Engine/Source/Runtime/Engine/Public/CanvasTypes.h
--- a/Engine/Source/Runtime/Engine/Public/CanvasTypes.h
+++ b/Engine/Source/Runtime/Engine/Public/CanvasTypes.h
@@ -249,10 +249,15 @@ public:
 	ENGINE_API FBatchedElements* GetBatchedElements(EElementType InElementType, FBatchedElementParameters* InBatchedElementParameters=NULL, const FTexture* Texture=NULL, ESimpleElementBlendMode BlendMode=SE_BLEND_MAX,const FDepthFieldGlowInfo& GlowInfo = FDepthFieldGlowInfo());
 	
 	/**
-	* Generates a new FCanvasTileRendererItem for the current sortkey and adds it to the sortelement list of itmes to render
+	* Generates a new FCanvasTileRendererItem for the current sortkey and adds it to the sortelement list of items to render
 	*/
 	ENGINE_API void AddTileRenderItem(float X,float Y,float SizeX,float SizeY,float U,float V,float SizeU,float SizeV,const FMaterialRenderProxy* MaterialRenderProxy,FHitProxyId HitProxyId,bool bFreezeTime,FColor InColor);
 
+	/**
+	* Generates a new FCanvasTriangleRendererItem for the current sortkey and adds it to the sortelement list of items to render
+	*/
+	ENGINE_API void AddTriangleRenderItem(const FCanvasUVTri& Tri, const FMaterialRenderProxy* MaterialRenderProxy, FHitProxyId HitProxyId, bool bFreezeTime);
+	
 	/** 
 	* Sends a message to the rendering thread to draw the batched elements. 
 	* @param RHICmdList - command list to use
@@ -738,6 +743,7 @@ public:
 	* @return true if anything rendered
 	*/
 	virtual bool Render_RenderThread(FRHICommandListImmediate& RHICmdList, const FCanvas* Canvas) = 0;
+	
 	/**
 	* Renders the canvas item
 	*
@@ -745,18 +751,27 @@ public:
 	* @return true if anything rendered
 	*/
 	virtual bool Render_GameThread(const FCanvas* Canvas) = 0;
+	
 	/**
 	* FCanvasBatchedElementRenderItem instance accessor
 	*
 	* @return FCanvasBatchedElementRenderItem instance
 	*/
 	virtual class FCanvasBatchedElementRenderItem* GetCanvasBatchedElementRenderItem() { return NULL; }
+	
 	/**
 	* FCanvasTileRendererItem instance accessor
 	*
 	* @return FCanvasTileRendererItem instance
 	*/
 	virtual class FCanvasTileRendererItem* GetCanvasTileRendererItem() { return NULL; }
+
+	/**
+	* FCanvasTriangleRendererItem instance accessor
+	*
+	* @return FCanvasTriangleRendererItem instance
+	*/
+	virtual class FCanvasTriangleRendererItem* GetCanvasTriangleRendererItem() { return NULL; }
 };
 
 
@@ -807,6 +822,7 @@ public:
 	* @return true if anything rendered
 	*/
 	virtual bool Render_RenderThread(FRHICommandListImmediate& RHICmdList, const FCanvas* Canvas) override;
+	
 	/**
 	* Renders the canvas item.
 	* Iterates over all batched elements and draws them with their own transforms
@@ -883,6 +899,7 @@ private:
 		/** info for optional glow effect when using depth field rendering */
 		FDepthFieldGlowInfo GlowInfo;
 	};
+	
 	/**
 	* Render data which is allocated when a new FCanvasBatchedElementRenderItem is added for rendering.
 	* This data is only freed on the rendering thread once the item has finished rendering
@@ -1016,6 +1033,119 @@ private:
 	const bool bFreezeTime;
 };
 
+/**
+* Info needed to render a single FTriangleRenderer
+*/
+class FCanvasTriangleRendererItem : public FCanvasBaseRenderItem
+{
+public:
+	/**
+	* Init constructor
+	*/
+	FCanvasTriangleRendererItem(
+		const FMaterialRenderProxy* InMaterialRenderProxy = NULL,
+		const FCanvas::FTransformEntry& InTransform = FCanvas::FTransformEntry(FMatrix::Identity),
+		bool bInFreezeTime = false)
+		// this data is deleted after rendering has completed
+		: Data(new FRenderData(InMaterialRenderProxy, InTransform))
+		, bFreezeTime(bInFreezeTime)
+	{}
+
+	/**
+	* Destructor to delete data in case nothing rendered
+	*/
+	virtual ~FCanvasTriangleRendererItem()
+	{
+		delete Data;
+	}
+
+	/**
+	 * FCanvasTriangleRendererItem instance accessor
+	 *
+	 * @return this instance
+	 */
+	virtual class FCanvasTriangleRendererItem* GetCanvasTriangleRendererItem()
+	{
+		return this;
+	}
+
+	/**
+	* Renders the canvas item.
+	* Iterates over each triangle to be rendered and draws it with its own transforms
+	*
+	* @param Canvas - canvas currently being rendered
+	* @param RHICmdList - command list to use
+	* @return true if anything rendered
+	*/
+	virtual bool Render_RenderThread(FRHICommandListImmediate& RHICmdList, const FCanvas* Canvas) override;
+
+	/**
+	* Renders the canvas item.
+	* Iterates over each triangle to be rendered and draws it with its own transforms
+	*
+	* @param Canvas - canvas currently being rendered
+	* @return true if anything rendered
+	*/
+	virtual bool Render_GameThread(const FCanvas* Canvas) override;
+
+	/**
+	* Determine if this is a matching set by comparing material,transform. All must match
+	*
+	* @param IInMaterialRenderProxy - material proxy resource for the item being rendered
+	* @param InTransform - the transform for the item being rendered
+	* @return true if the parameters match this render item
+	*/
+	bool IsMatch(const FMaterialRenderProxy* InMaterialRenderProxy, const FCanvas::FTransformEntry& InTransform)
+	{
+		return(Data->MaterialRenderProxy == InMaterialRenderProxy &&
+			Data->Transform.GetMatrixCRC() == InTransform.GetMatrixCRC());
+	};
+
+	/**
+	* Add a new triangle to the render data. These triangles all use the same transform and material proxy
+	*
+	* @param return number of triangles added
+	*/
+	FORCEINLINE int32 AddTriangle(const FCanvasUVTri& Tri, FHitProxyId HitProxyId)
+	{
+		return Data->AddTriangle(Tri, HitProxyId);
+	};
+
+private:
+	class FRenderData
+	{
+	public:
+		FRenderData(
+			const FMaterialRenderProxy* InMaterialRenderProxy = NULL,
+			const FCanvas::FTransformEntry& InTransform = FCanvas::FTransformEntry(FMatrix::Identity))
+			: MaterialRenderProxy(InMaterialRenderProxy)
+			, Transform(InTransform)
+		{}
+		const FMaterialRenderProxy* MaterialRenderProxy;
+		FCanvas::FTransformEntry Transform;
+
+		struct FTriangleInst
+		{
+			FCanvasUVTri Tri;
+			FHitProxyId HitProxyId;
+		};
+		TArray<FTriangleInst> Triangles;
+
+		FORCEINLINE int32 AddTriangle(const FCanvasUVTri& Tri, FHitProxyId HitProxyId)
+		{
+			FTriangleInst NewTri = { Tri, HitProxyId };
+			return Triangles.Add(NewTri);
+		};
+	};
+	/**
+	* Render data which is allocated when a new FCanvasTriangleRendererItem is added for rendering.
+	* This data is only freed on the rendering thread once the item has finished rendering
+	*/
+	FRenderData* Data;
+
+	const bool bFreezeTime;
+};
+
 /**
 * Render string using both a font and a material. The material should have a font exposed as a 
 * parameter so that the correct font page can be set based on the character being drawn.
diff --git a/Engine/Source/Runtime/Engine/Public/EngineStats.h b/Engine/Source/Runtime/Engine/Public/EngineStats.h
--- a/Engine/Source/Runtime/Engine/Public/EngineStats.h
+++ b/Engine/Source/Runtime/Engine/Public/EngineStats.h
@@ -166,6 +166,7 @@ DECLARE_CYCLE_STAT_EXTERN(TEXT("Draw String Time"),STAT_Canvas_DrawStringTime,ST
 DECLARE_CYCLE_STAT_EXTERN(TEXT("Word Wrapping Time"),STAT_Canvas_WordWrappingTime,STATGROUP_Canvas, );
 DECLARE_CYCLE_STAT_EXTERN(TEXT("Get Batched Element Time"),STAT_Canvas_GetBatchElementsTime,STATGROUP_Canvas, );
 DECLARE_CYCLE_STAT_EXTERN(TEXT("Add Material Tile Time"),STAT_Canvas_AddTileRenderTime,STATGROUP_Canvas, );
+DECLARE_CYCLE_STAT_EXTERN(TEXT("Add Material Triangle Time"),STAT_Canvas_AddTriangleRenderTime,STATGROUP_Canvas, );
 DECLARE_DWORD_COUNTER_STAT_EXTERN(TEXT("Num Batches Created"),STAT_Canvas_NumBatchesCreated,STATGROUP_Canvas, );	
 
 /**
diff --git a/Engine/Source/Runtime/Engine/Public/TriangleRendering.h b/Engine/Source/Runtime/Engine/Public/TriangleRendering.h
new file mode 100644
--- /dev/null
+++ b/Engine/Source/Runtime/Engine/Public/TriangleRendering.h
@@ -0,0 +1,26 @@
+// Copyright 1998-2015 Epic Games, Inc. All Rights Reserved.
+
+/*=============================================================================
+	TriangleRendering.h: Simple tile rendering implementation.
+=============================================================================*/
+
+#ifndef _INC_TRIANGLERENDERING
+#define _INC_TRIANGLERENDERING
+
+class FTriangleRenderer
+{
+public:
+
+	/**
+	 * Draw a tile at the given location and size, using the given UVs
+	 * (UV = [0..1]
+	 */
+	ENGINE_API static void DrawTriangle(FRHICommandListImmediate& RHICmdList, const class FSceneView& View, const FMaterialRenderProxy* MaterialRenderProxy, bool bNeedsToSwitchVerticalAxis, const FCanvasUVTri& Tri, bool bIsHitTesting = false, const FHitProxyId HitProxyId = FHitProxyId(), const FColor InVertexColor = FColor(255, 255, 255, 255));
+
+private:
+
+	/** This class never needs to be instantiated. */
+	FTriangleRenderer() {}
+};
+
+#endif
