Fixed issue where DialogueWaves play with incorrect SoundClass

Also added virtual function to SoundBase so that you can get the SoundClass after the proxy has been created

diff --git a/Engine/Source/Runtime/Engine/Classes/Sound/DialogueSoundWaveProxy.h b/Engine/Source/Runtime/Engine/Classes/Sound/DialogueSoundWaveProxy.h
--- a/Engine/Source/Runtime/Engine/Classes/Sound/DialogueSoundWaveProxy.h
+++ b/Engine/Source/Runtime/Engine/Classes/Sound/DialogueSoundWaveProxy.h
@@ -40,4 +40,6 @@ public:
 	 * Parses the Sound to generate the WaveInstances to play
 	 */
 	virtual void Parse( class FAudioDevice* AudioDevice, const UPTRINT NodeWaveInstanceHash, FActiveSound& AudioComponent, const FSoundParseParameters& ParseParams, TArray<FWaveInstance*>& WaveInstances );
+
+	virtual USoundClass* GetSoundClass() const override;
 };
\ No newline at end of file
diff --git a/Engine/Source/Runtime/Engine/Classes/Sound/SoundBase.h b/Engine/Source/Runtime/Engine/Classes/Sound/SoundBase.h
--- a/Engine/Source/Runtime/Engine/Classes/Sound/SoundBase.h
+++ b/Engine/Source/Runtime/Engine/Classes/Sound/SoundBase.h
@@ -33,10 +33,12 @@ class ENGINE_API USoundBase : public UObject
 {
 	GENERATED_UCLASS_BODY()
 
+protected:
 	/** Sound class this sound belongs to */
 	UPROPERTY(EditAnywhere, Category=Sound, meta=(DisplayName = "Sound Class"))
 	USoundClass* SoundClassObject;
 
+public:
 	/** When "stat sounds -debug" has been specified, draw this sound's attenuation shape when the sound is audible. For debugging purpose only. */
 	UPROPERTY(EditAnywhere, Category=Playback)
 	uint32 bDebug:1;
@@ -101,5 +103,10 @@ public:
 	 * Parses the Sound to generate the WaveInstances to play
 	 */
 	virtual void Parse( class FAudioDevice* AudioDevice, const UPTRINT NodeWaveInstanceHash, FActiveSound& ActiveSound, const FSoundParseParameters& ParseParams, TArray<FWaveInstance*>& WaveInstances ) { }
+
+	/**
+	 * Returns the SoundClass used for this sound
+	 */
+	virtual USoundClass* GetSoundClass() const;
 };
 
diff --git a/Engine/Source/Runtime/Engine/Private/ActiveSound.cpp b/Engine/Source/Runtime/Engine/Private/ActiveSound.cpp
--- a/Engine/Source/Runtime/Engine/Private/ActiveSound.cpp
+++ b/Engine/Source/Runtime/Engine/Private/ActiveSound.cpp
@@ -98,7 +98,7 @@ USoundClass* FActiveSound::GetSoundClass() const
 	}
 	else if (Sound)
 	{
-		return Sound->SoundClassObject;
+		return Sound->GetSoundClass();
 	}
 
 	return NULL;
diff --git a/Engine/Source/Runtime/Engine/Private/Audio.cpp b/Engine/Source/Runtime/Engine/Private/Audio.cpp
--- a/Engine/Source/Runtime/Engine/Private/Audio.cpp
+++ b/Engine/Source/Runtime/Engine/Private/Audio.cpp
@@ -90,9 +90,9 @@ FName FSoundBuffer::GetSoundClassName()
 				{
 					if (WaveNode->ResourceID == ResourceID)
 					{
-						if (Cue->SoundClassObject)
+						if (Cue->GetSoundClass())
 						{
-							return Cue->SoundClassObject->GetFName();
+							return Cue->GetSoundClass()->GetFName();
 						}
 						else
 						{
@@ -107,9 +107,9 @@ FName FSoundBuffer::GetSoundClassName()
 			USoundWave* Wave = Cast<USoundWave>(*It);
 			if (Wave && Wave->ResourceID == ResourceID)
 			{
-				if (Wave->SoundClassObject)
+				if (Wave->GetSoundClass())
 				{
-					return Wave->SoundClassObject->GetFName();
+					return Wave->GetSoundClass()->GetFName();
 				}
 				else
 				{
diff --git a/Engine/Source/Runtime/Engine/Private/AudioDevice.cpp b/Engine/Source/Runtime/Engine/Private/AudioDevice.cpp
--- a/Engine/Source/Runtime/Engine/Private/AudioDevice.cpp
+++ b/Engine/Source/Runtime/Engine/Private/AudioDevice.cpp
@@ -274,9 +274,9 @@ void FAudioDevice::GetSoundClassInfo( TMap<FName, FAudioClassInfo>& AudioClassIn
 		{
 			// Presume one class per sound node wave
 			USoundWave *SoundWave = WavePlayers[ WaveIndex ]->SoundWave;
-			if (SoundWave && SoundCue->SoundClassObject)
+			if (SoundWave && SoundCue->GetSoundClass())
 			{
-				SoundWaveClasses.Add( SoundWave, SoundCue->SoundClassObject->GetFName() );
+				SoundWaveClasses.Add( SoundWave, SoundCue->GetSoundClass()->GetFName() );
 			}
 		}
 	}
diff --git a/Engine/Source/Runtime/Engine/Private/DialogueWave.cpp b/Engine/Source/Runtime/Engine/Private/DialogueWave.cpp
--- a/Engine/Source/Runtime/Engine/Private/DialogueWave.cpp
+++ b/Engine/Source/Runtime/Engine/Private/DialogueWave.cpp
@@ -92,6 +92,11 @@ void UDialogueSoundWaveProxy::Parse( class FAudioDevice* AudioDevice, const UPTR
 	}
 }
 
+USoundClass* UDialogueSoundWaveProxy::GetSoundClass() const
+{
+	return SoundWave->GetSoundClass();
+}
+
 UDialogueWave::UDialogueWave(const FObjectInitializer& ObjectInitializer)
 	: Super(ObjectInitializer)
 	, LocalizationGUID( FGuid::NewGuid() )
@@ -238,7 +243,9 @@ void UDialogueWave::UpdateMappingProxy(FDialogueContextMapping& ContextMapping)
 
 	if(ContextMapping.Proxy)
 	{
+		// Copy the properties that the proxy shares with the sound in case it's used as a SoundBase
 		ContextMapping.Proxy->SoundWave = ContextMapping.SoundWave;
+		UEngine::CopyPropertiesForUnrelatedObjects(ContextMapping.SoundWave, ContextMapping.Proxy);
 
 		FSubtitleCue NewSubtitleCue;
 		FString Key = GetContextLocalizationKey( ContextMapping.Context );
diff --git a/Engine/Source/Runtime/Engine/Private/SoundBase.cpp b/Engine/Source/Runtime/Engine/Private/SoundBase.cpp
--- a/Engine/Source/Runtime/Engine/Private/SoundBase.cpp
+++ b/Engine/Source/Runtime/Engine/Private/SoundBase.cpp
@@ -100,4 +100,9 @@ float USoundBase::GetVolumeMultiplier()
 float USoundBase::GetPitchMultiplier()
 {
 	return 1.f;
-}
\ No newline at end of file
+}
+
+USoundClass* USoundBase::GetSoundClass() const
+{
+	return SoundClassObject;
+}
diff --git a/Engine/Source/Runtime/Engine/Private/SoundCue.cpp b/Engine/Source/Runtime/Engine/Private/SoundCue.cpp
--- a/Engine/Source/Runtime/Engine/Private/SoundCue.cpp
+++ b/Engine/Source/Runtime/Engine/Private/SoundCue.cpp
@@ -183,7 +183,7 @@ FString USoundCue::GetDesc()
 
 	// Display group
 	Description += TEXT( " [" );
-	Description += *SoundClassObject->GetName();
+	Description += *GetSoundClass()->GetName();
 	Description += TEXT( "]" );
 
 	return Description;
