Added nullptr check to SpawnDecalAttached

diff --git a/Engine/Source/Runtime/Engine/Private/GameplayStatics.cpp b/Engine/Source/Runtime/Engine/Private/GameplayStatics.cpp
--- a/Engine/Source/Runtime/Engine/Private/GameplayStatics.cpp
+++ b/Engine/Source/Runtime/Engine/Private/GameplayStatics.cpp
@@ -984,11 +984,11 @@ UDecalComponent* UGameplayStatics::SpawnDecalAttached(class UMaterialInterface*
 
 	if (DecalMaterial)
 	{
-		if (AttachToComponent == nullptr)
+		if (AttachToComponent == nullptr) 
 		{
 			UE_LOG(LogScript, Warning, TEXT("UGameplayStatics::SpawnDecalAttached: NULL AttachComponent specified!"));
 		}
-		else
+		else 
 		{
 			UPrimitiveComponent* AttachToPrimitive = Cast<UPrimitiveComponent>(AttachToComponent);
 			if (AttachToPrimitive == nullptr || AttachToPrimitive->bReceivesDecals)
@@ -1001,18 +1001,26 @@ UDecalComponent* UGameplayStatics::SpawnDecalAttached(class UMaterialInterface*
 				}
 				else
 				{
-					DecalComp = CreateDecalComponent(DecalMaterial, DecalSize, AttachToComponent->GetOwner(), LifeSpan);
-					DecalComp->AttachTo(AttachToComponent, AttachPointName);
-
-					if (LocationType == EAttachLocation::KeepWorldPosition)
+					if (AttachToComponent->GetOwner() == nullptr)
 					{
-						DecalComp->SetWorldLocationAndRotation(Location, Rotation);
+						UE_LOG(LogScript, Warning, TEXT("UGameplayStatics::SpawnDecalAttached: AttachToComponent has no valid Owner !"));
 					}
 					else
 					{
-						DecalComp->SetRelativeLocationAndRotation(Location, Rotation);
+						DecalComp = CreateDecalComponent(DecalMaterial, DecalSize, AttachToComponent->GetOwner(), LifeSpan);
+						DecalComp->AttachTo(AttachToComponent, AttachPointName);
+
+						if (LocationType == EAttachLocation::KeepWorldPosition)
+						{
+							DecalComp->SetWorldLocationAndRotation(Location, Rotation);
+						}
+						else
+						{
+							DecalComp->SetRelativeLocationAndRotation(Location, Rotation);
+						}
 					}
 				}
+
 			}
 		}
 	}
