Disabled canvas testbed

diff --git a/Engine/Source/Runtime/Engine/Private/UserInterface/CanvasItem.cpp b/Engine/Source/Runtime/Engine/Private/UserInterface/CanvasItem.cpp
--- a/Engine/Source/Runtime/Engine/Private/UserInterface/CanvasItem.cpp
+++ b/Engine/Source/Runtime/Engine/Private/UserInterface/CanvasItem.cpp
@@ -30,7 +30,7 @@ DECLARE_CYCLE_STAT(TEXT("CanvasBorderItem Time"),STAT_Canvas_BorderItemTime,STAT
 
 FCanvasItemTestbed::LineVars FCanvasItemTestbed::TestLine;
 bool FCanvasItemTestbed::bTestState = false;
-bool FCanvasItemTestbed::bShowTestbed = true;
+bool FCanvasItemTestbed::bShowTestbed = false;
 
 bool FCanvasItemTestbed::bShowLines = false;
 bool FCanvasItemTestbed::bShowBoxes = false;
