Still allow World based Particle/Decal creation for non-game viewports

diff --git a/Engine/Source/Runtime/Engine/Private/GameplayStatics.cpp b/Engine/Source/Runtime/Engine/Private/GameplayStatics.cpp
--- a/Engine/Source/Runtime/Engine/Private/GameplayStatics.cpp
+++ b/Engine/Source/Runtime/Engine/Private/GameplayStatics.cpp
@@ -540,17 +540,16 @@ void UGameplayStatics::PlayWorldCameraShake(UObject* WorldContextObject, TSubcla
 	}
 }
 
-UParticleSystemComponent* CreateParticleSystem(UParticleSystem* EmitterTemplate, AActor* Actor, bool bAutoDestroy)
+UParticleSystemComponent* CreateParticleSystem(UParticleSystem* EmitterTemplate, UWorld* World, AActor* Actor, bool bAutoDestroy)
 {
-	check(Actor);
-	UParticleSystemComponent* PSC = NewObject<UParticleSystemComponent>(Actor);
+	UParticleSystemComponent* PSC = NewObject<UParticleSystemComponent>((Actor ? Actor : (UObject*)World));
 	PSC->bAutoDestroy = bAutoDestroy;
 	PSC->SecondsBeforeInactive = 0.0f;
 	PSC->bAutoActivate = false;
 	PSC->SetTemplate(EmitterTemplate);
 	PSC->bOverrideLODMethod = false;
 
-	PSC->RegisterComponent();
+	PSC->RegisterComponentWithWorld(World);
 
 	return PSC;
 }
@@ -563,7 +562,7 @@ UParticleSystemComponent* UGameplayStatics::SpawnEmitterAtLocation(UObject* Worl
 		UWorld* World = GEngine->GetWorldFromContextObject(WorldContextObject);
 		if(World != nullptr)
 		{
-			PSC = CreateParticleSystem(EmitterTemplate, World->GetWorldSettings(), bAutoDestroy);
+			PSC = CreateParticleSystem(EmitterTemplate, World, World->GetWorldSettings(), bAutoDestroy);
 
 			PSC->SetAbsolute(true, true, true);
 			PSC->SetWorldLocationAndRotation(SpawnLocation, SpawnRotation);
@@ -585,7 +584,7 @@ UParticleSystemComponent* UGameplayStatics::SpawnEmitterAttached(UParticleSystem
 		}
 		else
 		{
-			PSC = CreateParticleSystem(EmitterTemplate, AttachToComponent->GetOwner(), bAutoDestroy);
+			PSC = CreateParticleSystem(EmitterTemplate, AttachToComponent->GetWorld(), AttachToComponent->GetOwner(), bAutoDestroy);
 
 			PSC->AttachTo(AttachToComponent, AttachPointName);
 			if (LocationType == EAttachLocation::KeepWorldPosition)
@@ -941,17 +940,15 @@ void UGameplayStatics::DeactivateReverbEffect(FName TagName)
 	}
 }
 
-UDecalComponent* CreateDecalComponent(class UMaterialInterface* DecalMaterial, FVector DecalSize, AActor* Actor, float LifeSpan)
+UDecalComponent* CreateDecalComponent(class UMaterialInterface* DecalMaterial, FVector DecalSize, UWorld* World, AActor* Actor, float LifeSpan)
 {
-	check(Actor);
-
 	const FMatrix DecalInternalTransform = FRotationMatrix(FRotator(0.f, 90.0f, -90.0f));
 
-	UDecalComponent* DecalComp = NewObject<UDecalComponent>(Actor);
+	UDecalComponent* DecalComp = NewObject<UDecalComponent>((Actor ? Actor : (UObject*)World));
 	DecalComp->DecalMaterial = DecalMaterial;
 	DecalComp->RelativeScale3D = DecalInternalTransform.TransformVector(DecalSize);
 	DecalComp->bAbsoluteScale = true;
-	DecalComp->RegisterComponent();
+	DecalComp->RegisterComponentWithWorld(World);
 
 	if (LifeSpan > 0.f)
 	{
@@ -970,7 +967,7 @@ UDecalComponent* UGameplayStatics::SpawnDecalAtLocation(UObject* WorldContextObj
 		UWorld* World = GEngine->GetWorldFromContextObject(WorldContextObject);
 		if(World != nullptr)
 		{
-			DecalComp = CreateDecalComponent(DecalMaterial, DecalSize, World->GetWorldSettings(), LifeSpan);
+			DecalComp = CreateDecalComponent(DecalMaterial, DecalSize, World, World->GetWorldSettings(), LifeSpan);
 			DecalComp->SetWorldLocationAndRotation(Location, Rotation);
 		}
 	}
@@ -984,11 +981,11 @@ UDecalComponent* UGameplayStatics::SpawnDecalAttached(class UMaterialInterface*
 
 	if (DecalMaterial)
 	{
-		if (AttachToComponent == nullptr) 
+		if (AttachToComponent == nullptr)
 		{
 			UE_LOG(LogScript, Warning, TEXT("UGameplayStatics::SpawnDecalAttached: NULL AttachComponent specified!"));
 		}
-		else 
+		else
 		{
 			UPrimitiveComponent* AttachToPrimitive = Cast<UPrimitiveComponent>(AttachToComponent);
 			if (AttachToPrimitive == nullptr || AttachToPrimitive->bReceivesDecals)
@@ -1001,26 +998,18 @@ UDecalComponent* UGameplayStatics::SpawnDecalAttached(class UMaterialInterface*
 				}
 				else
 				{
-					if (AttachToComponent->GetOwner() == nullptr)
+					DecalComp = CreateDecalComponent(DecalMaterial, DecalSize, AttachToComponent->GetWorld(), AttachToComponent->GetOwner(), LifeSpan);
+					DecalComp->AttachTo(AttachToComponent, AttachPointName);
+
+					if (LocationType == EAttachLocation::KeepWorldPosition)
 					{
-						UE_LOG(LogScript, Warning, TEXT("UGameplayStatics::SpawnDecalAttached: AttachToComponent has no valid Owner !"));
+						DecalComp->SetWorldLocationAndRotation(Location, Rotation);
 					}
 					else
 					{
-						DecalComp = CreateDecalComponent(DecalMaterial, DecalSize, AttachToComponent->GetOwner(), LifeSpan);
-						DecalComp->AttachTo(AttachToComponent, AttachPointName);
-
-						if (LocationType == EAttachLocation::KeepWorldPosition)
-						{
-							DecalComp->SetWorldLocationAndRotation(Location, Rotation);
-						}
-						else
-						{
-							DecalComp->SetRelativeLocationAndRotation(Location, Rotation);
-						}
+						DecalComp->SetRelativeLocationAndRotation(Location, Rotation);
 					}
 				}
-
 			}
 		}
 	}
