Don't call begin play on Actors that have been marked pending kill

diff --git a/Engine/Source/Runtime/Engine/Private/WorldSettings.cpp b/Engine/Source/Runtime/Engine/Private/WorldSettings.cpp
--- a/Engine/Source/Runtime/Engine/Private/WorldSettings.cpp
+++ b/Engine/Source/Runtime/Engine/Private/WorldSettings.cpp
@@ -124,7 +124,7 @@ void AWorldSettings::NotifyBeginPlay()
 		for (FActorIterator It(World); It; ++It)
 		{
 			// Actors that have traveled seamlessly from other levels already had BeginPlay called in that level
-			if (!It->HasActorBegunPlay())
+			if (!It->IsPendingKill() && !It->HasActorBegunPlay())
 			{
 				It->BeginPlay();
 			}
