fixes for the case where WITH_RECAST is false

diff --git a/Engine/Source/Runtime/Engine/Private/AI/Navigation/NavigationSystem.cpp b/Engine/Source/Runtime/Engine/Private/AI/Navigation/NavigationSystem.cpp
--- a/Engine/Source/Runtime/Engine/Private/AI/Navigation/NavigationSystem.cpp
+++ b/Engine/Source/Runtime/Engine/Private/AI/Navigation/NavigationSystem.cpp
@@ -3539,6 +3539,7 @@ bool UNavigationSystem::DoesPathIntersectBox(const FNavigationPath* Path, const
 
 void UNavigationSystem::SetMaxSimultaneousTileGenerationJobsCount(int32 MaxNumberOfJobs)
 {
+#if WITH_RECAST
 	for (auto NavigationData : NavDataSet)
 	{
 		ARecastNavMesh* RecastNavMesh = Cast<ARecastNavMesh>(NavigationData);
@@ -3547,10 +3548,12 @@ void UNavigationSystem::SetMaxSimultaneousTileGenerationJobsCount(int32 MaxNumbe
 			RecastNavMesh->SetMaxSimultaneousTileGenerationJobsCount(MaxNumberOfJobs);
 		}
 	}
+#endif
 }
 
 void UNavigationSystem::ResetMaxSimultaneousTileGenerationJobsCount()
 {
+#if WITH_RECAST
 	for (auto NavigationData : NavDataSet)
 	{
 		ARecastNavMesh* RecastNavMesh = Cast<ARecastNavMesh>(NavigationData);
@@ -3560,6 +3563,7 @@ void UNavigationSystem::ResetMaxSimultaneousTileGenerationJobsCount()
 			RecastNavMesh->SetMaxSimultaneousTileGenerationJobsCount(CDO->MaxSimultaneousTileGenerationJobsCount);
 		}
 	}
+#endif
 }
 
 //----------------------------------------------------------------------//
