#ue4 - Perf: don't force a reallocate in UPrimitiveComponent::GetOverlappingComponents(). Reduce allocations in AActor::ClearComponentOverlaps().

diff --git a/Engine/Source/Runtime/Engine/Private/Actor.cpp b/Engine/Source/Runtime/Engine/Private/Actor.cpp
--- a/Engine/Source/Runtime/Engine/Private/Actor.cpp
+++ b/Engine/Source/Runtime/Engine/Private/Actor.cpp
@@ -909,12 +909,12 @@ void AActor::ClearComponentOverlaps()
 	// Remove owned components from overlap tracking
 	// We don't traverse the RootComponent attachment tree since that might contain
 	// components owned by other actors.
+	TArray<UPrimitiveComponent*> OverlappingComponentsForCurrentComponent;
 	for (UPrimitiveComponent* const PrimComp : PrimitiveComponents)
 	{
-		TArray<UPrimitiveComponent*> OverlappingComponents;
-		PrimComp->GetOverlappingComponents(OverlappingComponents);
+		PrimComp->GetOverlappingComponents(OverlappingComponentsForCurrentComponent);
 
-		for (UPrimitiveComponent* const OverlapComp : OverlappingComponents)
+		for (UPrimitiveComponent* const OverlapComp : OverlappingComponentsForCurrentComponent)
 		{
 			if (OverlapComp)
 			{
diff --git a/Engine/Source/Runtime/Engine/Private/PrimitiveComponent.cpp b/Engine/Source/Runtime/Engine/Private/PrimitiveComponent.cpp
--- a/Engine/Source/Runtime/Engine/Private/PrimitiveComponent.cpp
+++ b/Engine/Source/Runtime/Engine/Private/PrimitiveComponent.cpp
@@ -2088,7 +2088,7 @@ void UPrimitiveComponent::GetOverlappingActors(TArray<AActor*>& OutOverlappingAc
 
 void UPrimitiveComponent::GetOverlappingComponents(TArray<UPrimitiveComponent*>& OutOverlappingComponents) const
 {
-	OutOverlappingComponents.Empty( OverlappingComponents.Num() );
+	OutOverlappingComponents.Reset( OverlappingComponents.Num() );
 
 	for (auto CompIt = OverlappingComponents.CreateConstIterator(); CompIt; ++CompIt)
 	{
