Add ability to skip fixed amount of network replay time

diff --git a/Engine/Source/Runtime/Engine/Private/DemoNetDriver.cpp b/Engine/Source/Runtime/Engine/Private/DemoNetDriver.cpp
--- a/Engine/Source/Runtime/Engine/Private/DemoNetDriver.cpp
+++ b/Engine/Source/Runtime/Engine/Private/DemoNetDriver.cpp
@@ -23,6 +23,7 @@ DEFINE_LOG_CATEGORY_STATIC( LogDemo, Log, All );
 
 static TAutoConsoleVariable<float> CVarDemoRecordHz( TEXT( "demo.RecordHz" ), 10, TEXT( "Number of demo frames recorded per second" ) );
 static TAutoConsoleVariable<float> CVarDemoTimeDilation( TEXT( "demo.TimeDilation" ), -1.0f, TEXT( "Override time dilation during demo playback (-1 = don't override)" ) );
+static TAutoConsoleVariable<float> CVarDemoSkipTime( TEXT( "demo.SkipTime" ), 0, TEXT( "Skip fixed amount of network replay time (in seconds)" ) );
 
 static const int32 MAX_DEMO_READ_WRITE_BUFFER = 1024 * 2;
 
@@ -918,22 +919,52 @@ void UDemoNetDriver::TickDemoPlayback( float DeltaSeconds )
 	}
 
 	DemoDeltaTime += DeltaSeconds;
+
+	if ( DemoCurrentTime + DemoDeltaTime > DemoTotalTime )
+	{
+		DemoDeltaTime = DemoTotalTime - DemoCurrentTime;
+	}
+
 	DemoCurrentTime += DeltaSeconds;
 
-	if ( DemoCurrentTime >= DemoTotalTime )
+	if ( CVarDemoSkipTime.GetValueOnGameThread() > 0 )
+	{
+		DemoDeltaTime += CVarDemoSkipTime.GetValueOnGameThread();
+
+		if ( DemoCurrentTime + DemoDeltaTime > DemoTotalTime )
+		{
+			DemoDeltaTime = DemoTotalTime - DemoCurrentTime;
+		}
+
+		DemoCurrentTime += CVarDemoSkipTime.GetValueOnGameThread();
+
+		CVarDemoSkipTime.AsVariable()->Set( TEXT( "0" ), ECVF_SetByConsole );
+	}
+
+	if ( DemoCurrentTime > DemoTotalTime )
 	{
 		DemoCurrentTime = DemoTotalTime;
 	}
 
 	while ( true )
 	{
-		// Read demo frames until we are caught up
-		if ( !ReadDemoFrame() )
+		const int32 StartFrame = DemoFrameNum;
+
+		while ( true )
 		{
-			break;
+			// Read demo frames until we are caught up
+			if ( !ReadDemoFrame() )
+			{
+				break;
+			}
+
+			DemoFrameNum++;
 		}
 
-		DemoFrameNum++;
+		if ( StartFrame == DemoFrameNum )
+		{
+			break;
+		}
 	}
 }
 
