Removed pedantic check().

diff --git a/Engine/Source/Runtime/Engine/Private/TextureDerivedData.cpp b/Engine/Source/Runtime/Engine/Private/TextureDerivedData.cpp
--- a/Engine/Source/Runtime/Engine/Private/TextureDerivedData.cpp
+++ b/Engine/Source/Runtime/Engine/Private/TextureDerivedData.cpp
@@ -917,8 +917,6 @@ public:
 		{
 			GetBuildInfo();
 			BuildTexture();
-
-			check( bSucceeded );
 		}
 	}
 
