Fix BlendPhysicsBones ignoring the scale of the root atom

diff --git a/Engine/Source/Runtime/Engine/Private/PhysicsEngine/PhysAnim.cpp b/Engine/Source/Runtime/Engine/Private/PhysicsEngine/PhysAnim.cpp
--- a/Engine/Source/Runtime/Engine/Private/PhysicsEngine/PhysAnim.cpp
+++ b/Engine/Source/Runtime/Engine/Private/PhysicsEngine/PhysAnim.cpp
@@ -147,6 +147,13 @@ void USkeletalMeshComponent::BlendPhysicsBones( TArray<FBoneIndexType>& InRequir
 					// Now blend in this atom. See if we are forcing this bone to always be blended in
 					LocalAtoms[BoneIndex].Blend( LocalAtoms[BoneIndex], PhysAtom, UsePhysWeight );
 
+					if(BoneIndex == 0)
+					{
+						//We must update RecipScale3D based on the atom scale of the root
+						TotalScale3D *= LocalAtoms[0].GetScale3D();
+						RecipScale3D = TotalScale3D.Reciprocal();
+					}
+
 					if (UsePhysWeight < 1.f)
 					{
 						bUpdatePhysics = true;
