Renamed UEdGraphNode::DiscardPin to ::RemovePin for discoverability purposes. If the pin has a parent (from splitting), it will remove all related pins.

diff --git a/Engine/Source/Editor/BlueprintGraph/Private/K2Node_CommutativeAssociativeBinaryOperator.cpp b/Engine/Source/Editor/BlueprintGraph/Private/K2Node_CommutativeAssociativeBinaryOperator.cpp
--- a/Engine/Source/Editor/BlueprintGraph/Private/K2Node_CommutativeAssociativeBinaryOperator.cpp
+++ b/Engine/Source/Editor/BlueprintGraph/Private/K2Node_CommutativeAssociativeBinaryOperator.cpp
@@ -179,12 +179,8 @@ void UK2Node_CommutativeAssociativeBinaryOperator::RemoveInputPin(UEdGraphPin* P
 		FScopedTransaction Transaction( LOCTEXT("RemovePinTx", "RemovePin") );
 		Modify();
 
-		int32 PinRemovalIndex = INDEX_NONE;
-		if (Pins.Find(Pin, /*out*/ PinRemovalIndex))
+		if (RemovePin(Pin))
 		{
-			Pins.RemoveAt(PinRemovalIndex);
-			Pin->Modify();
-			Pin->BreakAllPinLinks();
 			--NumAdditionalInputs;
 
 			int32 NameIndex = 0;
diff --git a/Engine/Source/Editor/BlueprintGraph/Private/K2Node_MatineeController.cpp b/Engine/Source/Editor/BlueprintGraph/Private/K2Node_MatineeController.cpp
--- a/Engine/Source/Editor/BlueprintGraph/Private/K2Node_MatineeController.cpp
+++ b/Engine/Source/Editor/BlueprintGraph/Private/K2Node_MatineeController.cpp
@@ -240,7 +240,7 @@ void UK2Node_MatineeController::OnEventKeyframeRemoved(const AMatineeActor* InMa
 		{
 			if(UEdGraphPin* Pin = FindPin(InPinNames[PinIdx].ToString()))
 			{
-				DiscardPin(Pin);
+				RemovePin(Pin);
 				bNeedsRefresh = true;
 			}
 		}
diff --git a/Engine/Source/Runtime/Engine/Classes/EdGraph/EdGraphNode.h b/Engine/Source/Runtime/Engine/Classes/EdGraph/EdGraphNode.h
--- a/Engine/Source/Runtime/Engine/Classes/EdGraph/EdGraphNode.h
+++ b/Engine/Source/Runtime/Engine/Classes/EdGraph/EdGraphNode.h
@@ -184,8 +184,8 @@ public:
 	/** Find a pin on this node with the supplied name and assert if it is not present */
 	UEdGraphPin* FindPinChecked(const FString& PinName) const;
 	
-	/** Find a pin on this node with the supplied name and remove it */
-	void DiscardPin(UEdGraphPin* Pin);
+	/** Find a pin on this node with the supplied name and remove it, returns TRUE if successful */
+	bool RemovePin(UEdGraphPin* Pin);
 
 	/** Whether or not this node should be given the chance to override pin names.  If this returns true, then GetPinNameOverride() will be called for each pin, each frame */
 	virtual bool ShouldOverridePinNames() const { return false; }
diff --git a/Engine/Source/Runtime/Engine/Private/EdGraph/EdGraphNode.cpp b/Engine/Source/Runtime/Engine/Private/EdGraph/EdGraphNode.cpp
--- a/Engine/Source/Runtime/Engine/Private/EdGraph/EdGraphNode.cpp
+++ b/Engine/Source/Runtime/Engine/Private/EdGraph/EdGraphNode.cpp
@@ -107,13 +107,26 @@ UEdGraphPin* UEdGraphNode::FindPinChecked(const FString& PinName) const
 	return Result;
 }
 
-void UEdGraphNode::DiscardPin(UEdGraphPin* Pin)
+bool UEdGraphNode::RemovePin(UEdGraphPin* Pin)
 {
 	check( Pin );
 	
 	Modify();
-	Pins.Remove( Pin );
-	Pin->BreakAllPinLinks();
+	UEdGraphPin* RootPin = (Pin->ParentPin != nullptr)? Pin->ParentPin : Pin;
+	RootPin->BreakAllPinLinks();
+
+	if (Pins.Remove( RootPin ))
+	{
+		// Remove any children pins to ensure the entirety of the pin's representation is removed
+		for (UEdGraphPin* ChildPin : RootPin->SubPins)
+		{
+			Pins.Remove(ChildPin);
+			ChildPin->Modify();
+			ChildPin->BreakAllPinLinks();
+		}
+		return true;
+	}
+	return false;
 }
 
 void UEdGraphNode::BreakAllNodeLinks()
