Use NetBitReader when deserializing stored unmapped object bunches. Fixes issues where custom netserialize struct uses << operator to serialize object references.

diff --git a/Engine/Source/Runtime/Engine/Classes/Engine/NetSerialization.h b/Engine/Source/Runtime/Engine/Classes/Engine/NetSerialization.h
--- a/Engine/Source/Runtime/Engine/Classes/Engine/NetSerialization.h
+++ b/Engine/Source/Runtime/Engine/Classes/Engine/NetSerialization.h
@@ -473,7 +473,7 @@ bool FFastArraySerializer::FastArrayDeltaSerialize( TArray<Type> &Items, FNetDel
 				Type* ThisElement = &Items[ArraySerializer.ItemMap.FindChecked( ElementID )];
 
 				// Initialize the reader with the stored buffer that we need to read from
-				FBitReader Reader( UnmappedItem.Buffer.GetData(), UnmappedItem.NumBufferBits );
+				FNetBitReader Reader( Parms.Map, UnmappedItem.Buffer.GetData(), UnmappedItem.NumBufferBits );
 
 				// Read the property (which should serialize any newly mapped objects as well)
 				bool bHasUnmapped = false;
