Dont discard unreliable multicasts if connection is saturated at the time of UNetDriver::InternalProcessRemoteFunction. They are queued at the actor channel level and will be subject to ServerReplicateActor's bandwidth management when it is time for them to be sent.

diff --git a/Engine/Source/Runtime/Engine/Private/NetworkDriver.cpp b/Engine/Source/Runtime/Engine/Private/NetworkDriver.cpp
--- a/Engine/Source/Runtime/Engine/Private/NetworkDriver.cpp
+++ b/Engine/Source/Runtime/Engine/Private/NetworkDriver.cpp
@@ -724,8 +724,8 @@ void UNetDriver::InternalProcessRemoteFunction
 		Function = Function->GetSuperFunction();
 	}
 
-	// If saturated and function is unimportant, skip it.
-	if( !(Function->FunctionFlags & FUNC_NetReliable) && !Connection->IsNetReady(0) )
+	// If saturated and function is unimportant, skip it. Note unreliable multicasts are queued at the actor channel level so they are not gated here.
+	if( !(Function->FunctionFlags & FUNC_NetReliable) && (!(Function->FunctionFlags & FUNC_NetMulticast)) && !Connection->IsNetReady(0) )
 	{
 		DEBUG_REMOTEFUNCTION(TEXT("Network saturated, not calling %s::%s"), *Actor->GetName(), *Function->GetName());
 		return;
