Expose FDebugDisplayInfo::NumDisplayNames

diff --git a/Engine/Source/Runtime/Engine/Public/DisplayDebugHelpers.h b/Engine/Source/Runtime/Engine/Public/DisplayDebugHelpers.h
--- a/Engine/Source/Runtime/Engine/Public/DisplayDebugHelpers.h
+++ b/Engine/Source/Runtime/Engine/Public/DisplayDebugHelpers.h
@@ -17,6 +17,7 @@ public:
 
 	bool IsDisplayOn(FName DisplayName) const { return DisplayNames.Contains(DisplayName); }
 	bool IsCategoryToggledOn(FName Category, bool bDefaultsToOn) const { return ToggledCategories.Contains(Category) != bDefaultsToOn; }
+	int32 NumDisplayNames() const { return DisplayNames.Num(); }
 };
 
 // Simple object to track scope indentation
