merging replay change from dev branch

diff --git a/Engine/Source/Runtime/Engine/Private/GameMode.cpp b/Engine/Source/Runtime/Engine/Private/GameMode.cpp
--- a/Engine/Source/Runtime/Engine/Private/GameMode.cpp
+++ b/Engine/Source/Runtime/Engine/Private/GameMode.cpp
@@ -1750,5 +1750,11 @@ FString AGameMode::GetRedirectURL(const FString& MapName) const
 
 bool AGameMode::IsHandlingReplays()
 {
+	// If we're running in PIE, don't record demos
+	if (GetWorld() != nullptr && GetWorld()->IsPlayInEditor())
+	{
+		return false;
+	}
+
 	return bHandleDedicatedServerReplays && GetNetMode() == ENetMode::NM_DedicatedServer;
 }
