Fixed parameter grouping for parameters defined in functions.

diff --git a/Engine/Source/Runtime/Engine/Classes/Materials/MaterialFunction.h b/Engine/Source/Runtime/Engine/Classes/Materials/MaterialFunction.h
--- a/Engine/Source/Runtime/Engine/Classes/Materials/MaterialFunction.h
+++ b/Engine/Source/Runtime/Engine/Classes/Materials/MaterialFunction.h
@@ -109,14 +109,6 @@ public:
 		check(OutParameterNames.Num() == OutParameterIds.Num());
 	}
 
-	/** Searches the material function for a parameter of the given name and returns whether it was found and its value */
-	bool GetVectorParameterValue(FName ParameterName, FLinearColor& OutValue) const;
-	bool GetScalarParameterValue(FName ParameterName, float& OutValue) const;
-	bool GetTextureParameterValue(FName ParameterName, UTexture*& OutValue) const;
-	bool GetFontParameterValue(FName ParameterName, class UFont*& OutFontValue, int32& OutFontPage) const;
-	bool GetStaticSwitchParameterValue(FName ParameterName, bool& OutValue, FGuid& OutExpressionGuid) const;
-	bool GetStaticComponentMaskParameterValue(FName ParameterName, bool& OutR, bool& OutG, bool& OutB, bool& OutA, FGuid& OutExpressionGuid) const;
-
 #if WITH_EDITOR
 	ENGINE_API UMaterial* GetPreviewMaterial();
 
diff --git a/Engine/Source/Runtime/Engine/Classes/Materials/MaterialInstance.h b/Engine/Source/Runtime/Engine/Classes/Materials/MaterialInstance.h
--- a/Engine/Source/Runtime/Engine/Classes/Materials/MaterialInstance.h
+++ b/Engine/Source/Runtime/Engine/Classes/Materials/MaterialInstance.h
@@ -474,15 +474,52 @@ bool UpdateParameterSet(TArray<ParameterType> &Parameters, UMaterial* ParentMate
 		// No reference to the material expression exists, so try to find one in the material expression's array if we are in the editor.
 		if (bTryToFindByName && GIsEditor && !FApp::IsGame())
 		{
-			for (int32 ExpressionIndex = 0; ExpressionIndex < ParentMaterial->Expressions.Num(); ExpressionIndex++)
+			for (const UMaterialExpression* Expression : ParentMaterial->Expressions)
 			{
-				ExpressionType* ParameterExpression = Cast<ExpressionType>(ParentMaterial->Expressions[ExpressionIndex]);
-
-				if (ParameterExpression && ParameterExpression->ParameterName == Parameter.ParameterName)
+				if (const ExpressionType* ParameterExpression = Cast<const ExpressionType>(Expression))
 				{
-					Parameter.ExpressionGUID = ParameterExpression->ExpressionGUID;
-					bChanged = true;
-					break;
+					if (ParameterExpression->ParameterName == Parameter.ParameterName)
+					{
+						Parameter.ExpressionGUID = ParameterExpression->ExpressionGUID;
+						bChanged = true;
+						break;
+					}
+				}
+				else if (const UMaterialExpressionMaterialFunctionCall* FunctionCall = Cast<const UMaterialExpressionMaterialFunctionCall>(Expression))
+				{
+					if (FunctionCall->MaterialFunction)
+					{
+						auto UpdateParameterSetInFunctions = [&]()
+						{
+							TArray<UMaterialFunction*> Functions;
+							Functions.Add(FunctionCall->MaterialFunction);
+							FunctionCall->MaterialFunction->GetDependentFunctions(Functions);
+
+							for (UMaterialFunction* Function : Functions)
+							{
+								for (const UMaterialExpression* FunctionExpression : Function->FunctionExpressions)
+								{
+									if (const ExpressionType* ParameterExpression = Cast<const ExpressionType>(FunctionExpression))
+									{
+										if (ParameterExpression->ParameterName == Parameter.ParameterName)
+										{
+											Parameter.ExpressionGUID = ParameterExpression->ExpressionGUID;
+											bChanged = true;
+											break;
+										}
+									}
+								}
+							}
+
+							return false;
+						};
+
+						if (UpdateParameterSetInFunctions())
+						{
+							bChanged = true;
+							break;
+						}
+					}
 				}
 			}
 		}
diff --git a/Engine/Source/Runtime/Engine/Private/Materials/Material.cpp b/Engine/Source/Runtime/Engine/Private/Materials/Material.cpp
--- a/Engine/Source/Runtime/Engine/Private/Materials/Material.cpp
+++ b/Engine/Source/Runtime/Engine/Private/Materials/Material.cpp
@@ -1163,91 +1163,162 @@ const UMaterial* UMaterial::GetMaterial_Concurrent(TMicRecursionGuard&) const
 
 bool UMaterial::GetGroupName(FName ParameterName, FName& OutDesc) const
 {
-	bool bSuccess = false;
-	for(int32 ExpressionIndex = 0; ExpressionIndex < Expressions.Num(); ExpressionIndex++)
+	for (const UMaterialExpression* Expression : Expressions)
 	{
-		UMaterialExpression* Expression = Expressions[ExpressionIndex];
 		// Parameter is a basic Expression Parameter
-		if(Expression->IsA(UMaterialExpressionParameter::StaticClass()))
+		if (const UMaterialExpressionParameter* Parameter = Cast<const UMaterialExpressionParameter>(Expression))
 		{
-			UMaterialExpressionParameter* Parameter = CastChecked<UMaterialExpressionParameter>(Expression);
-			if(Parameter && Parameter->ParameterName == ParameterName)
+			if (Parameter->ParameterName == ParameterName)
 			{
 				OutDesc = Parameter->Group;
-				bSuccess = true;
-				break;
+				return true;
 			}
 		}
 		// Parameter is a Texture Sample Parameter
-		else if(Expression->IsA(UMaterialExpressionTextureSampleParameter::StaticClass()))
+		else if (const UMaterialExpressionTextureSampleParameter* Parameter = Cast<const UMaterialExpressionTextureSampleParameter>(Expression))
 		{
-			UMaterialExpressionTextureSampleParameter* Parameter = CastChecked<UMaterialExpressionTextureSampleParameter>(Expression);
-			if(Parameter && Parameter->ParameterName == ParameterName)
+			if (Parameter->ParameterName == ParameterName)
 			{
 				OutDesc = Parameter->Group;
-				bSuccess = true;
-				break;
+				return true;
 			}
 		}
 		// Parameter is a Font Sample Parameter
-		else if(Expression->IsA(UMaterialExpressionFontSampleParameter::StaticClass()))
+		else if (const UMaterialExpressionFontSampleParameter* Parameter = Cast<const UMaterialExpressionFontSampleParameter>(Expression))
 		{
-			UMaterialExpressionFontSampleParameter* Parameter = CastChecked<UMaterialExpressionFontSampleParameter>(Expression);
-			if(Parameter && Parameter->ParameterName == ParameterName)
+			if (Parameter->ParameterName == ParameterName)
 			{
 				OutDesc = Parameter->Group;
-				bSuccess = true;
-				break;
+				return true;
+			}
+		}
+		// Parameter is a function call
+		else if (const UMaterialExpressionMaterialFunctionCall* FunctionCall = Cast<const UMaterialExpressionMaterialFunctionCall>(Expression))
+		{
+			if (FunctionCall->MaterialFunction)
+			{
+				TArray<UMaterialFunction*> Functions;
+				Functions.Add(FunctionCall->MaterialFunction);
+				FunctionCall->MaterialFunction->GetDependentFunctions(Functions);
+
+				for (UMaterialFunction* Function : Functions)
+				{
+					for (UMaterialExpression* FunctionExpression : Function->FunctionExpressions)
+					{
+						// Parameter is a basic Expression Parameter
+						if (const UMaterialExpressionParameter* Parameter = Cast<const UMaterialExpressionParameter>(FunctionExpression))
+						{
+							if (Parameter->ParameterName == ParameterName)
+							{
+								OutDesc = Parameter->Group;
+								return true;
+							}
+						}
+						// Parameter is a Texture Sample Parameter
+						else if (const UMaterialExpressionTextureSampleParameter* Parameter = Cast<const UMaterialExpressionTextureSampleParameter>(FunctionExpression))
+						{
+							if (Parameter->ParameterName == ParameterName)
+							{
+								OutDesc = Parameter->Group;
+								return true;
+							}
+						}
+						// Parameter is a Font Sample Parameter
+						else if (const UMaterialExpressionFontSampleParameter* Parameter = Cast<const UMaterialExpressionFontSampleParameter>(FunctionExpression))
+						{
+							if (Parameter->ParameterName == ParameterName)
+							{
+								OutDesc = Parameter->Group;
+								return true;
+							}
+						}
+					}
+				}
 			}
 		}
 	}
-	return bSuccess;
+
+	return false;
 }
 
 bool UMaterial::GetParameterDesc(FName ParameterName, FString& OutDesc) const
 {
-	bool bSuccess = false;
-	for(int32 ExpressionIndex = 0; ExpressionIndex < Expressions.Num(); ExpressionIndex++)
+	for (const UMaterialExpression* Expression : Expressions)
 	{
-		UMaterialExpression* Expression = Expressions[ExpressionIndex];
-		if ( Expression )
+		// Parameter is a basic Expression Parameter
+		if (const UMaterialExpressionParameter* Parameter = Cast<const UMaterialExpressionParameter>(Expression))
 		{
-			// Parameter is a basic Expression Parameter
-			if(Expression->IsA(UMaterialExpressionParameter::StaticClass()))
+			if (Parameter->ParameterName == ParameterName)
 			{
-				UMaterialExpressionParameter* Parameter = CastChecked<UMaterialExpressionParameter>(Expression);
-				if(Parameter && Parameter->ParameterName == ParameterName)
-				{
-					OutDesc = Parameter->Desc;
-					bSuccess = true;
-					break;
-				}
+				OutDesc = Parameter->Desc;
+				return true;
 			}
-			// Parameter is a Texture Sample Parameter
-			else if(Expression->IsA(UMaterialExpressionTextureSampleParameter::StaticClass()))
+		}
+		// Parameter is a Texture Sample Parameter
+		else if (const UMaterialExpressionTextureSampleParameter* Parameter = Cast<const UMaterialExpressionTextureSampleParameter>(Expression))
+		{
+			if (Parameter->ParameterName == ParameterName)
 			{
-				UMaterialExpressionTextureSampleParameter* Parameter = CastChecked<UMaterialExpressionTextureSampleParameter>(Expression);
-				if(Parameter && Parameter->ParameterName == ParameterName)
-				{
-					OutDesc = Parameter->Desc;
-					bSuccess = true;
-					break;
-				}
+				OutDesc = Parameter->Desc;
+				return true;
 			}
-			// Parameter is a Font Sample Parameter
-			else if(Expression->IsA(UMaterialExpressionFontSampleParameter::StaticClass()))
+		}
+		// Parameter is a Font Sample Parameter
+		else if (const UMaterialExpressionFontSampleParameter* Parameter = Cast<const UMaterialExpressionFontSampleParameter>(Expression))
+		{
+			if (Parameter->ParameterName == ParameterName)
 			{
-				UMaterialExpressionFontSampleParameter* Parameter = CastChecked<UMaterialExpressionFontSampleParameter>(Expression);
-				if(Parameter && Parameter->ParameterName == ParameterName)
+				OutDesc = Parameter->Desc;
+				return true;
+			}
+		}
+		// Parameter is a function call
+		else if (const UMaterialExpressionMaterialFunctionCall* FunctionCall = Cast<const UMaterialExpressionMaterialFunctionCall>(Expression))
+		{
+			if (FunctionCall->MaterialFunction)
+			{
+				TArray<UMaterialFunction*> Functions;
+				Functions.Add(FunctionCall->MaterialFunction);
+				FunctionCall->MaterialFunction->GetDependentFunctions(Functions);
+
+				for (UMaterialFunction* Function : Functions)
 				{
-					OutDesc = Parameter->Desc;
-					bSuccess = true;
-					break;
+					for (UMaterialExpression* FunctionExpression : Function->FunctionExpressions)
+					{
+						// Parameter is a basic Expression Parameter
+						if (const UMaterialExpressionParameter* Parameter = Cast<const UMaterialExpressionParameter>(FunctionExpression))
+						{
+							if (Parameter->ParameterName == ParameterName)
+							{
+								OutDesc = Parameter->Desc;
+								return true;
+							}
+						}
+						// Parameter is a Texture Sample Parameter
+						else if (const UMaterialExpressionTextureSampleParameter* Parameter = Cast<const UMaterialExpressionTextureSampleParameter>(FunctionExpression))
+						{
+							if (Parameter->ParameterName == ParameterName)
+							{
+								OutDesc = Parameter->Desc;
+								return true;
+							}
+						}
+						// Parameter is a Font Sample Parameter
+						else if (const UMaterialExpressionFontSampleParameter* Parameter = Cast<const UMaterialExpressionFontSampleParameter>(FunctionExpression))
+						{
+							if (Parameter->ParameterName == ParameterName)
+							{
+								OutDesc = Parameter->Desc;
+								return true;
+							}
+						}
+					}
 				}
 			}
 		}
 	}
-	return bSuccess;
+
+	return false;
 }
 
 bool UMaterial::GetVectorParameterValue(FName ParameterName, FLinearColor& OutValue) const
@@ -1261,11 +1332,27 @@ bool UMaterial::GetVectorParameterValue(FName ParameterName, FLinearColor& OutVa
 				return true;
 			}
 		}
-		else if (const UMaterialExpressionMaterialFunctionCall* FunctionCall = Cast<UMaterialExpressionMaterialFunctionCall>(Expression))
+		else if (const UMaterialExpressionMaterialFunctionCall* FunctionCall = Cast<const UMaterialExpressionMaterialFunctionCall>(Expression))
 		{
-			if (FunctionCall->MaterialFunction && FunctionCall->MaterialFunction->GetVectorParameterValue(ParameterName, OutValue))
+			if (FunctionCall->MaterialFunction)
 			{
-				return true;
+				TArray<UMaterialFunction*> Functions;
+				Functions.Add(FunctionCall->MaterialFunction);
+				FunctionCall->MaterialFunction->GetDependentFunctions(Functions);
+
+				for (UMaterialFunction* Function : Functions)
+				{
+					for (UMaterialExpression* FunctionExpression : Function->FunctionExpressions)
+					{
+						if (const UMaterialExpressionVectorParameter* Parameter = Cast<const UMaterialExpressionVectorParameter>(FunctionExpression))
+						{
+							if (Parameter->IsNamedParameter(ParameterName, OutValue))
+							{
+								return true;
+							}
+						}
+					}
+				}
 			}
 		}
 	}
@@ -1284,11 +1371,27 @@ bool UMaterial::GetScalarParameterValue(FName ParameterName, float& OutValue) co
 				return true;
 			}
 		}
-		else if (const UMaterialExpressionMaterialFunctionCall* FunctionCall = Cast<UMaterialExpressionMaterialFunctionCall>(Expression))
+		else if (const UMaterialExpressionMaterialFunctionCall* FunctionCall = Cast<const UMaterialExpressionMaterialFunctionCall>(Expression))
 		{
-			if (FunctionCall->MaterialFunction && FunctionCall->MaterialFunction->GetScalarParameterValue(ParameterName, OutValue))
+			if (FunctionCall->MaterialFunction)
 			{
-				return true;
+				TArray<UMaterialFunction*> Functions;
+				Functions.Add(FunctionCall->MaterialFunction);
+				FunctionCall->MaterialFunction->GetDependentFunctions(Functions);
+
+				for (UMaterialFunction* Function : Functions)
+				{
+					for (UMaterialExpression* FunctionExpression : Function->FunctionExpressions)
+					{
+						if (const UMaterialExpressionScalarParameter* Parameter = Cast<const UMaterialExpressionScalarParameter>(FunctionExpression))
+						{
+							if (Parameter->IsNamedParameter(ParameterName, OutValue))
+							{
+								return true;
+							}
+						}
+					}
+				}
 			}
 		}
 	}
@@ -1307,11 +1410,27 @@ bool UMaterial::GetTextureParameterValue(FName ParameterName, UTexture*& OutValu
 				return true;
 			}
 		}
-		else if (const UMaterialExpressionMaterialFunctionCall* FunctionCall = Cast<UMaterialExpressionMaterialFunctionCall>(Expression))
+		else if (const UMaterialExpressionMaterialFunctionCall* FunctionCall = Cast<const UMaterialExpressionMaterialFunctionCall>(Expression))
 		{
-			if (FunctionCall->MaterialFunction && FunctionCall->MaterialFunction->GetTextureParameterValue(ParameterName, OutValue))
+			if (FunctionCall->MaterialFunction)
 			{
-				return true;
+				TArray<UMaterialFunction*> Functions;
+				Functions.Add(FunctionCall->MaterialFunction);
+				FunctionCall->MaterialFunction->GetDependentFunctions(Functions);
+
+				for (UMaterialFunction* Function : Functions)
+				{
+					for (UMaterialExpression* FunctionExpression : Function->FunctionExpressions)
+					{
+						if (const UMaterialExpressionTextureSampleParameter* Parameter = Cast<const UMaterialExpressionTextureSampleParameter>(FunctionExpression))
+						{
+							if (Parameter->IsNamedParameter(ParameterName, OutValue))
+							{
+								return true;
+							}
+						}
+					}
+				}
 			}
 		}
 	}
@@ -1330,11 +1449,27 @@ bool UMaterial::GetFontParameterValue(FName ParameterName, UFont*& OutFontValue,
 				return true;
 			}
 		}
-		else if (const UMaterialExpressionMaterialFunctionCall* FunctionCall = Cast<UMaterialExpressionMaterialFunctionCall>(Expression))
+		else if (const UMaterialExpressionMaterialFunctionCall* FunctionCall = Cast<const UMaterialExpressionMaterialFunctionCall>(Expression))
 		{
-			if (FunctionCall->MaterialFunction && FunctionCall->MaterialFunction->GetFontParameterValue(ParameterName, OutFontValue, OutFontPage))
+			if (FunctionCall->MaterialFunction)
 			{
-				return true;
+				TArray<UMaterialFunction*> Functions;
+				Functions.Add(FunctionCall->MaterialFunction);
+				FunctionCall->MaterialFunction->GetDependentFunctions(Functions);
+
+				for (UMaterialFunction* Function : Functions)
+				{
+					for (UMaterialExpression* FunctionExpression : Function->FunctionExpressions)
+					{
+						if (const UMaterialExpressionFontSampleParameter* Parameter = Cast<const UMaterialExpressionFontSampleParameter>(FunctionExpression))
+						{
+							if (Parameter->IsNamedParameter(ParameterName, OutFontValue, OutFontPage))
+							{
+								return true;
+							}
+						}
+					}
+				}
 			}
 		}
 	}
@@ -1354,11 +1489,27 @@ bool UMaterial::GetStaticSwitchParameterValue(FName ParameterName, bool& OutValu
 				return true;
 			}
 		}
-		else if (const UMaterialExpressionMaterialFunctionCall* FunctionCall = Cast<UMaterialExpressionMaterialFunctionCall>(Expression))
+		else if (const UMaterialExpressionMaterialFunctionCall* FunctionCall = Cast<const UMaterialExpressionMaterialFunctionCall>(Expression))
 		{
-			if (FunctionCall->MaterialFunction && FunctionCall->MaterialFunction->GetStaticSwitchParameterValue(ParameterName, OutValue, OutExpressionGuid))
+			if (FunctionCall->MaterialFunction)
 			{
-				return true;
+				TArray<UMaterialFunction*> Functions;
+				Functions.Add(FunctionCall->MaterialFunction);
+				FunctionCall->MaterialFunction->GetDependentFunctions(Functions);
+
+				for (UMaterialFunction* Function : Functions)
+				{
+					for (UMaterialExpression* FunctionExpression : Function->FunctionExpressions)
+					{
+						if (const UMaterialExpressionStaticBoolParameter* Parameter = Cast<const UMaterialExpressionStaticBoolParameter>(FunctionExpression))
+						{
+							if (Parameter->IsNamedParameter(ParameterName, OutValue, OutExpressionGuid))
+							{
+								return true;
+							}
+						}
+					}
+				}
 			}
 		}
 	}
@@ -1378,11 +1529,27 @@ bool UMaterial::GetStaticComponentMaskParameterValue(FName ParameterName, bool&
 				return true;
 			}
 		}
-		else if (const UMaterialExpressionMaterialFunctionCall* FunctionCall = Cast<UMaterialExpressionMaterialFunctionCall>(Expression))
+		else if (const UMaterialExpressionMaterialFunctionCall* FunctionCall = Cast<const UMaterialExpressionMaterialFunctionCall>(Expression))
 		{
-			if (FunctionCall->MaterialFunction && FunctionCall->MaterialFunction->GetStaticComponentMaskParameterValue(ParameterName, OutR, OutG, OutB, OutA, OutExpressionGuid))
+			if (FunctionCall->MaterialFunction)
 			{
-				return true;
+				TArray<UMaterialFunction*> Functions;
+				Functions.Add(FunctionCall->MaterialFunction);
+				FunctionCall->MaterialFunction->GetDependentFunctions(Functions);
+
+				for (UMaterialFunction* Function : Functions)
+				{
+					for (UMaterialExpression* FunctionExpression : Function->FunctionExpressions)
+					{
+						if (const UMaterialExpressionStaticComponentMaskParameter* Parameter = Cast<const UMaterialExpressionStaticComponentMaskParameter>(FunctionExpression))
+						{
+							if (Parameter->IsNamedParameter(ParameterName, OutR, OutG, OutB, OutA, OutExpressionGuid))
+							{
+								return true;
+							}
+						}
+					}
+				}
 			}
 		}
 	}
diff --git a/Engine/Source/Runtime/Engine/Private/Materials/MaterialExpressions.cpp b/Engine/Source/Runtime/Engine/Private/Materials/MaterialExpressions.cpp
--- a/Engine/Source/Runtime/Engine/Private/Materials/MaterialExpressions.cpp
+++ b/Engine/Source/Runtime/Engine/Private/Materials/MaterialExpressions.cpp
@@ -6639,146 +6639,6 @@ bool UMaterialFunction::HasFlippedCoordinates() const
 }
 #endif //WITH_EDITORONLY_DATA
 
-bool UMaterialFunction::GetVectorParameterValue(FName ParameterName, FLinearColor& OutValue) const
-{
-	for (const UMaterialExpression* Expression : FunctionExpressions)
-	{
-		if (const UMaterialExpressionVectorParameter* Parameter = Cast<const UMaterialExpressionVectorParameter>(Expression))
-		{
-			if (Parameter->IsNamedParameter(ParameterName, OutValue))
-			{
-				return true;
-			}
-		}
-		else if (const UMaterialExpressionMaterialFunctionCall* FunctionCall = Cast<UMaterialExpressionMaterialFunctionCall>(Expression))
-		{
-			if (FunctionCall->MaterialFunction && FunctionCall->MaterialFunction->GetVectorParameterValue(ParameterName, OutValue))
-			{
-				return true;
-			}
-		}
-	}
-
-	return false;
-}
-
-bool UMaterialFunction::GetScalarParameterValue(FName ParameterName, float& OutValue) const
-{
-	for (const UMaterialExpression* Expression : FunctionExpressions)
-	{
-		if (const UMaterialExpressionScalarParameter* Parameter = Cast<const UMaterialExpressionScalarParameter>(Expression))
-		{
-			if (Parameter->IsNamedParameter(ParameterName, OutValue))
-			{
-				return true;
-			}
-		}
-		else if (const UMaterialExpressionMaterialFunctionCall* FunctionCall = Cast<UMaterialExpressionMaterialFunctionCall>(Expression))
-		{
-			if (FunctionCall->MaterialFunction && FunctionCall->MaterialFunction->GetScalarParameterValue(ParameterName, OutValue))
-			{
-				return true;
-			}
-		}
-	}
-
-	return false;
-}
-
-bool UMaterialFunction::GetTextureParameterValue(FName ParameterName, UTexture*& OutValue) const
-{
-	for (const UMaterialExpression* Expression : FunctionExpressions)
-	{
-		if (const UMaterialExpressionTextureSampleParameter* Parameter = Cast<const UMaterialExpressionTextureSampleParameter>(Expression))
-		{
-			if (Parameter->IsNamedParameter(ParameterName, OutValue))
-			{
-				return true;
-			}
-		}
-		else if (const UMaterialExpressionMaterialFunctionCall* FunctionCall = Cast<UMaterialExpressionMaterialFunctionCall>(Expression))
-		{
-			if (FunctionCall->MaterialFunction && FunctionCall->MaterialFunction->GetTextureParameterValue(ParameterName, OutValue))
-			{
-				return true;
-			}
-		}
-	}
-
-	return false;
-}
-
-bool UMaterialFunction::GetFontParameterValue(FName ParameterName, UFont*& OutFontValue, int32& OutFontPage) const
-{
-	for (const UMaterialExpression* Expression : FunctionExpressions)
-	{
-		if (const UMaterialExpressionFontSampleParameter* Parameter = Cast<const UMaterialExpressionFontSampleParameter>(Expression))
-		{
-			if (Parameter->IsNamedParameter(ParameterName, OutFontValue, OutFontPage))
-			{
-				return true;
-			}
-		}
-		else if (const UMaterialExpressionMaterialFunctionCall* FunctionCall = Cast<UMaterialExpressionMaterialFunctionCall>(Expression))
-		{
-			if (FunctionCall->MaterialFunction && FunctionCall->MaterialFunction->GetFontParameterValue(ParameterName, OutFontValue, OutFontPage))
-			{
-				return true;
-			}
-		}
-	}
-
-	return false;
-}
-
-
-bool UMaterialFunction::GetStaticSwitchParameterValue(FName ParameterName, bool& OutValue, FGuid& OutExpressionGuid) const
-{
-	for (const UMaterialExpression* Expression : FunctionExpressions)
-	{
-		if (const UMaterialExpressionStaticBoolParameter* Parameter = Cast<const UMaterialExpressionStaticBoolParameter>(Expression))
-		{
-			if (Parameter->IsNamedParameter(ParameterName, OutValue, OutExpressionGuid))
-			{
-				return true;
-			}
-		}
-		else if (const UMaterialExpressionMaterialFunctionCall* FunctionCall = Cast<UMaterialExpressionMaterialFunctionCall>(Expression))
-		{
-			if (FunctionCall->MaterialFunction && FunctionCall->MaterialFunction->GetStaticSwitchParameterValue(ParameterName, OutValue, OutExpressionGuid))
-			{
-				return true;
-			}
-		}
-	}
-
-	return false;
-}
-
-
-bool UMaterialFunction::GetStaticComponentMaskParameterValue(FName ParameterName, bool& OutR, bool& OutG, bool& OutB, bool& OutA, FGuid& OutExpressionGuid) const
-{
-	for (const UMaterialExpression* Expression : FunctionExpressions)
-	{
-		if (const UMaterialExpressionStaticComponentMaskParameter* Parameter = Cast<const UMaterialExpressionStaticComponentMaskParameter>(Expression))
-		{
-			if (Parameter->IsNamedParameter(ParameterName, OutR, OutG, OutB, OutA, OutExpressionGuid))
-			{
-				return true;
-			}
-		}
-		else if (const UMaterialExpressionMaterialFunctionCall* FunctionCall = Cast<UMaterialExpressionMaterialFunctionCall>(Expression))
-		{
-			if (FunctionCall->MaterialFunction && FunctionCall->MaterialFunction->GetStaticComponentMaskParameterValue(ParameterName, OutR, OutG, OutB, OutA, OutExpressionGuid))
-			{
-				return true;
-			}
-		}
-	}
-
-	return false;
-}
-
 
 ///////////////////////////////////////////////////////////////////////////////
 // UMaterialExpressionMaterialFunctionCall
