Fixed asset tooltips for static/skeletal mesh so that resource size is reported with the correct suffix.

diff --git a/Engine/Source/Runtime/Engine/Private/SkeletalMesh.cpp b/Engine/Source/Runtime/Engine/Private/SkeletalMesh.cpp
--- a/Engine/Source/Runtime/Engine/Private/SkeletalMesh.cpp
+++ b/Engine/Source/Runtime/Engine/Private/SkeletalMesh.cpp
@@ -2711,6 +2711,7 @@ void USkeletalMesh::GetAssetRegistryTags(TArray<FAssetRegistryTag>& OutTags) con
 #if WITH_EDITOR
 void USkeletalMesh::GetAssetRegistryTagMetadata(TMap<FName, FAssetRegistryTagMetadata>& OutMetadata) const
 {
+	Super::GetAssetRegistryTagMetadata(OutMetadata);
 	OutMetadata.Add("PhysicsAsset", FAssetRegistryTagMetadata().SetImportantValue(TEXT("None")));
 }
 #endif
diff --git a/Engine/Source/Runtime/Engine/Private/StaticMesh.cpp b/Engine/Source/Runtime/Engine/Private/StaticMesh.cpp
--- a/Engine/Source/Runtime/Engine/Private/StaticMesh.cpp
+++ b/Engine/Source/Runtime/Engine/Private/StaticMesh.cpp
@@ -1541,6 +1541,8 @@ void UStaticMesh::GetAssetRegistryTags(TArray<FAssetRegistryTag>& OutTags) const
 #if WITH_EDITOR
 void UStaticMesh::GetAssetRegistryTagMetadata(TMap<FName, FAssetRegistryTagMetadata>& OutMetadata) const
 {
+	Super::GetAssetRegistryTagMetadata(OutMetadata);
+
 	OutMetadata.Add("CollisionPrims",
 		FAssetRegistryTagMetadata()
 			.SetTooltip(NSLOCTEXT("UStaticMesh", "CollisionPrimsTooltip", "The number of collision primitives in the static mesh"))
