Force formatted assertion macros to require output message.

diff --git a/Engine/Plugins/Experimental/SteamVR/Source/SteamVR/Private/SteamVRHMD.cpp b/Engine/Plugins/Experimental/SteamVR/Source/SteamVR/Private/SteamVRHMD.cpp
--- a/Engine/Plugins/Experimental/SteamVR/Source/SteamVR/Private/SteamVRHMD.cpp
+++ b/Engine/Plugins/Experimental/SteamVR/Source/SteamVR/Private/SteamVRHMD.cpp
@@ -178,7 +178,7 @@ namespace
 
 void FSteamVRHMD::GetCurrentOrientationAndPosition(FQuat& CurrentOrientation, FVector& CurrentPosition)
 {
-	checkf(IsInGameThread());
+	check(IsInGameThread());
 	GetCurrentPose(CurHmdOrientation, CurHmdPosition, MotionPredictionInSecondsGame);
 	CurrentOrientation = LastHmdOrientation = CurHmdOrientation;
 
diff --git a/Engine/Plugins/Runtime/OculusRift/Source/OculusRift/Private/OculusRiftHMD.cpp b/Engine/Plugins/Runtime/OculusRift/Source/OculusRift/Private/OculusRiftHMD.cpp
--- a/Engine/Plugins/Runtime/OculusRift/Source/OculusRift/Private/OculusRiftHMD.cpp
+++ b/Engine/Plugins/Runtime/OculusRift/Source/OculusRift/Private/OculusRiftHMD.cpp
@@ -279,7 +279,7 @@ void FOculusRiftHMD::PoseToOrientationAndPosition(const ovrPosef& InPose, FQuat&
 void FOculusRiftHMD::GetCurrentOrientationAndPosition(FQuat& CurrentOrientation, FVector& CurrentPosition)
 {
 	// only supposed to be used from the game thread
-	checkf(IsInGameThread());
+	check(IsInGameThread());
 	GetCurrentPose(CurrentOrientation, CurrentPosition);
 	LastHmdOrientation = CurrentOrientation;
 }
diff --git a/Engine/Source/Editor/UnrealEd/Private/Lightmass/Lightmass.cpp b/Engine/Source/Editor/UnrealEd/Private/Lightmass/Lightmass.cpp
--- a/Engine/Source/Editor/UnrealEd/Private/Lightmass/Lightmass.cpp
+++ b/Engine/Source/Editor/UnrealEd/Private/Lightmass/Lightmass.cpp
@@ -3870,7 +3870,7 @@ bool FLightmassProcessor::ImportLightMapData2DData(int32 Channel, FQuantizedLigh
 	// decompress the temp buffer into another temp buffer 
 	if (!FCompression::UncompressMemory(COMPRESS_ZLIB, DataBuffer, UncompressedSize, CompressedBuffer, CompressedSize))
 	{
-		checkf(TEXT("Uncompress failed, which is unexpected"));
+		checkf(false, TEXT("Uncompress failed, which is unexpected"));
 	}
 
 	// can free one buffer now
@@ -3914,7 +3914,7 @@ bool FLightmassProcessor::ImportSignedDistanceFieldShadowMapData2D(int32 Channel
 		// Decompress the temp buffer into another temp buffer 
 		if (!FCompression::UncompressMemory(COMPRESS_ZLIB, DataBuffer, UncompressedSize, CompressedBuffer, CompressedSize))
 		{
-			checkf(TEXT("Uncompress failed, which is unexpected"));
+			checkf(false, TEXT("Uncompress failed, which is unexpected"));
 		}
 		FMemory::Free(CompressedBuffer);
 
diff --git a/Engine/Source/Runtime/Core/Private/ProfilingDebugging/ProfilingHelpers.cpp b/Engine/Source/Runtime/Core/Private/ProfilingDebugging/ProfilingHelpers.cpp
--- a/Engine/Source/Runtime/Core/Private/ProfilingDebugging/ProfilingHelpers.cpp
+++ b/Engine/Source/Runtime/Core/Private/ProfilingDebugging/ProfilingHelpers.cpp
@@ -79,7 +79,7 @@ FString CreateProfileFilename( const FString& InFileExtension, bool bIncludeDate
 	FString MapNameStr;
 
 #if WITH_ENGINE
-	checkf(GGetMapNameDelegate.IsBound());
+	check(GGetMapNameDelegate.IsBound());
 	MapNameStr = GGetMapNameDelegate.Execute();
 #endif		// WITH_ENGINE
 
@@ -139,7 +139,7 @@ FString CreateProfileDirectoryAndFilename( const FString& InSubDirectoryName, co
 {
 	FString MapNameStr;
 #if WITH_ENGINE
-	checkf(GGetMapNameDelegate.IsBound());
+	check(GGetMapNameDelegate.IsBound());
 	MapNameStr = GGetMapNameDelegate.Execute();
 #endif		// WITH_ENGINE
 	const FString PlatformStr = FString(TEXT("PC"));
diff --git a/Engine/Source/Runtime/Core/Private/Windows/WindowsCursor.cpp b/Engine/Source/Runtime/Core/Private/Windows/WindowsCursor.cpp
--- a/Engine/Source/Runtime/Core/Private/Windows/WindowsCursor.cpp
+++ b/Engine/Source/Runtime/Core/Private/Windows/WindowsCursor.cpp
@@ -146,7 +146,7 @@ void FWindowsCursor::SetPosition( const int32 X, const int32 Y )
 void FWindowsCursor::SetType( const EMouseCursor::Type InNewCursor )
 {
 	// NOTE: Watch out for contention with FWindowsViewport::UpdateMouseCursor
-	checkf( InNewCursor < EMouseCursor::TotalCursorCount, TEXT("Invalid cursor(%d) supplied"), InNewCursor );
+	checkf( InNewCursor < EMouseCursor::TotalCursorCount, TEXT("Invalid cursor(%d) supplied"), (int)InNewCursor );
 	CurrentType = InNewCursor;
 	::SetCursor( CursorHandles[ InNewCursor ] );
 }
diff --git a/Engine/Source/Runtime/Core/Public/HAL/MallocBinned.h b/Engine/Source/Runtime/Core/Public/HAL/MallocBinned.h
--- a/Engine/Source/Runtime/Core/Public/HAL/MallocBinned.h
+++ b/Engine/Source/Runtime/Core/Public/HAL/MallocBinned.h
@@ -523,7 +523,7 @@ private:
 			for (UPTRINT i=(UPTRINT)PageSize, Offset=0; i<OsBytes; i+=PageSize, ++Offset)
 			{
 				FPoolInfo* TrailingPool = GetPoolInfo(((UPTRINT)Free)+i);
-				checkf(TrailingPool)
+				check(TrailingPool);
 				//Set trailing pools to point back to first pool
 				TrailingPool->SetAllocationSizes(0, 0, Offset, BinnedOSTableIndex);
 			}
@@ -849,7 +849,7 @@ public:
 		BinnedSizeLimit = PAGE_SIZE_LIMIT/2;
 		BinnedOSTableIndex = BinnedSizeLimit+EXTENED_PAGE_POOL_ALLOCATION_COUNT;
 
-		checkf((BinnedSizeLimit & (BinnedSizeLimit-1)) == 0);
+		check((BinnedSizeLimit & (BinnedSizeLimit-1)) == 0);
 
 
 		// Init tables.
@@ -1045,7 +1045,7 @@ public:
 			if( Pool->TableIndex < BinnedOSTableIndex )
 			{
 				// Allocated from pool, so grow or shrink if necessary.
-				checkf(Pool->TableIndex > 0); // it isn't possible to allocate a size of 0, Malloc will increase the size to DEFAULT_BINNED_ALLOCATOR_ALIGNMENT
+				check(Pool->TableIndex > 0); // it isn't possible to allocate a size of 0, Malloc will increase the size to DEFAULT_BINNED_ALLOCATOR_ALIGNMENT
 				if( NewSize>MemSizeToPoolTable[Pool->TableIndex]->BlockSize || NewSize<=MemSizeToPoolTable[Pool->TableIndex-1]->BlockSize )
 				{
 					NewPtr = Malloc( NewSize, Alignment );
diff --git a/Engine/Source/Runtime/Core/Public/Math/UnrealMatrix.h b/Engine/Source/Runtime/Core/Public/Math/UnrealMatrix.h
--- a/Engine/Source/Runtime/Core/Public/Math/UnrealMatrix.h
+++ b/Engine/Source/Runtime/Core/Public/Math/UnrealMatrix.h
@@ -263,7 +263,7 @@ inline FMatrix FMatrix::InverseFast() const
 		GetScaledAxis( EAxis::Z ).IsNearlyZero(SMALL_NUMBER) ) 
 	{
 		UE_LOG(LogUnrealMath, Error, TEXT("FMatrix::InverseFast(), trying to invert a NIL matrix, this results in NaNs! Use Inverse() instead."));
-		ensureMsgf(false, TEXT("FMatrix::InverseFast(), trying to invert a NIL matrix, this results in NaNs! Use Inverse() instead."));
+		ensureMsg(false, TEXT("FMatrix::InverseFast(), trying to invert a NIL matrix, this results in NaNs! Use Inverse() instead."));
 	}
 #endif
 	FMatrix Result;
diff --git a/Engine/Source/Runtime/Core/Public/Misc/AssertionMacros.h b/Engine/Source/Runtime/Core/Public/Misc/AssertionMacros.h
--- a/Engine/Source/Runtime/Core/Public/Misc/AssertionMacros.h
+++ b/Engine/Source/Runtime/Core/Public/Misc/AssertionMacros.h
@@ -24,8 +24,8 @@
 	 * verifyf, checkf: Same as verify, check but with printf style additional parameters
 	 * Read about __VA_ARGS__ (variadic macros) on http://gcc.gnu.org/onlinedocs/gcc-3.4.4/cpp.pdf.
 	 */
-	#define verifyf(expr, ...)				{ if(!(expr)) { FPlatformMisc::DebugBreak(); FDebug::AssertFailed( #expr, __FILE__, __LINE__, ##__VA_ARGS__ ); CA_ASSUME(expr); } }
-	#define checkf(expr, ...)				{ if(!(expr)) { FPlatformMisc::DebugBreak(); FDebug::AssertFailed( #expr, __FILE__, __LINE__, ##__VA_ARGS__ ); CA_ASSUME(expr); } }
+	#define verifyf(expr, format, ...)		{ if(!(expr)) { FPlatformMisc::DebugBreak(); FDebug::AssertFailed( #expr, __FILE__, __LINE__, format, ##__VA_ARGS__ ); CA_ASSUME(expr); } }
+	#define checkf(expr, format, ...)		{ if(!(expr)) { FPlatformMisc::DebugBreak(); FDebug::AssertFailed( #expr, __FILE__, __LINE__, format, ##__VA_ARGS__ ); CA_ASSUME(expr); } }
 	/**
 	 * Denotes code paths that should never be reached.
 	 */
@@ -58,13 +58,13 @@
 
 #else
 	#define checkCode(...)
-	#define check(...)
-	#define checkf(...)
+	#define check(expr)
+	#define checkf(expr, format, ...)
 	#define checkNoEntry()
 	#define checkNoReentry()
 	#define checkNoRecursion()
 	#define verify(expr)				{ if(!(expr)){} }
-	#define verifyf(expr, ...)			{ if(!(expr)){} }
+	#define verifyf(expr, format, ...)	{ if(!(expr)){} }
 	#define unimplemented()
 #endif
 
@@ -72,12 +72,12 @@
 // Check for development only.
 //
 #if DO_GUARD_SLOW
-	#define checkSlow(expr, ...)	{ if(!(expr)) { FPlatformMisc::DebugBreak(); FDebug::AssertFailed( #expr, __FILE__, __LINE__ ); CA_ASSUME(expr); } }
-	#define checkfSlow(expr, ...)	{ if(!(expr)) { FPlatformMisc::DebugBreak(); FDebug::AssertFailed( #expr, __FILE__, __LINE__, __VA_ARGS__ ); CA_ASSUME(expr); } }
-	#define verifySlow(expr)		{ if(!(expr)) { FPlatformMisc::DebugBreak(); FDebug::AssertFailed( #expr, __FILE__, __LINE__ );} }
+	#define checkSlow(expr)					{ if(!(expr)) { FDebug::AssertFailed( #expr, __FILE__, __LINE__ ); CA_ASSUME(expr); } }
+	#define checkfSlow(expr, format, ...)	{ if(!(expr)) { FDebug::AssertFailed( #expr, __FILE__, __LINE__, format, __VA_ARGS__ ); CA_ASSUME(expr); } }
+	#define verifySlow(expr)				{ if(!(expr)) { FDebug::AssertFailed( #expr, __FILE__, __LINE__ );} }
 #else
-	#define checkSlow(...)
-	#define checkfSlow(...)
+	#define checkSlow(expr)
+	#define checkfSlow(expr, format, ...)
 	#define verifySlow(expr)  if(expr){}
 #endif
 
@@ -112,8 +112,8 @@
 	#define ensureMsg( InExpression, InMsg ) \
 		(((InExpression) == 0) ? FDebug::EnsureNotFalse(false, #InExpression, __FILE__, __LINE__, InMsg ) : true)
 
-	#define ensureMsgf( InExpression, ... ) \
-		(((InExpression) == 0) ? (FPlatformMisc::DebugBreakReturningFalse() || FDebug::EnsureNotFalseFormatted(false, #InExpression, __FILE__, __LINE__, ##__VA_ARGS__ )) : true)
+	#define ensureMsgf( InExpression, InFormat, ... ) \
+		(((InExpression) == 0) ? FDebug::EnsureNotFalseFormatted(false, #InExpression, __FILE__, __LINE__, InFormat, ##__VA_ARGS__ ) : true)
 
 #else	// DO_CHECK
 
@@ -121,7 +121,7 @@
 
 	#define ensureMsg( InExpression, InMsg ) ( ( InExpression ) != 0 )
 
-	#define ensureMsgf( InExpression, ... ) ( ( InExpression ) != 0 )
+	#define ensureMsgf( InExpression, InFormat, ... ) ( ( InExpression ) != 0 )
 
 #endif	// DO_CHECK
 
diff --git a/Engine/Source/Runtime/CoreUObject/Public/UObject/UObjectBase.h b/Engine/Source/Runtime/CoreUObject/Public/UObject/UObjectBase.h
--- a/Engine/Source/Runtime/CoreUObject/Public/UObject/UObjectBase.h
+++ b/Engine/Source/Runtime/CoreUObject/Public/UObject/UObjectBase.h
@@ -152,7 +152,7 @@ protected:
 	 **/
 	FORCEINLINE void SetFlagsTo( EObjectFlags NewFlags )
 	{
-		checkfSlow((NewFlags & ~RF_AllFlags) == 0, TEXT("%s flagged as 0x%x but is trying to set flags to RF_AllFlags"), *GetFName().ToString(), ObjectFlags);
+		checkfSlow((NewFlags & ~RF_AllFlags) == 0, TEXT("%s flagged as 0x%x but is trying to set flags to RF_AllFlags"), *GetFName().ToString(), (int)ObjectFlags);
 		ObjectFlags = NewFlags;
 	}
 public:
diff --git a/Engine/Source/Runtime/Engine/Private/Animation/AnimCompress_PerTrackCompression.cpp b/Engine/Source/Runtime/Engine/Private/Animation/AnimCompress_PerTrackCompression.cpp
--- a/Engine/Source/Runtime/Engine/Private/Animation/AnimCompress_PerTrackCompression.cpp
+++ b/Engine/Source/Runtime/Engine/Private/Animation/AnimCompress_PerTrackCompression.cpp
@@ -346,7 +346,7 @@ protected:
 			FQuat DecompressedQ;
 			Compressor.ToQuat(DecompressedQ);
 
-			checkf(DecompressedQ.IsNormalized());
+			check(DecompressedQ.IsNormalized());
 			const float Error = FQuat::ErrorAutoNormalize(Q, DecompressedQ);
 			MaxError = FMath::Max(MaxError, Error);
 			SumError += Error;
@@ -447,7 +447,7 @@ protected:
 				UE_LOG(LogAnimationCompression, Log, TEXT(" XYZ: %i, %i, %i"), X, Y, Z);
 			}
 
-			checkf(DecompressedQ.IsNormalized());
+			check(DecompressedQ.IsNormalized());
 			const float Error = FQuat::ErrorAutoNormalize(Q, DecompressedQ);
 			MaxError = FMath::Max(MaxError, Error);
 			SumError += Error;
@@ -555,7 +555,7 @@ protected:
 				UE_LOG(LogAnimationCompression, Log, TEXT(" XYZ: %f, %f, %f, %f"), QRenorm.X, QRenorm.Y, QRenorm.Z, QRenorm.W);
 				UE_LOG(LogAnimationCompression, Log, TEXT(" Mins(%f, %f, %f)   Maxs(%f, %f,%f)"), KeyBounds.Min.X, KeyBounds.Min.Y, KeyBounds.Min.Z, KeyBounds.Max.X, KeyBounds.Max.Y, KeyBounds.Max.Z);
 			}
-			checkf(DecompressedQ.IsNormalized());
+			check(DecompressedQ.IsNormalized());
 			const float Error = FQuat::ErrorAutoNormalize(Q, DecompressedQ);
 			MaxError = FMath::Max(MaxError, Error);
 			SumError += Error;
@@ -1173,7 +1173,7 @@ void UAnimCompress_PerTrackCompression::CompressUsingUnderlyingCompressor(
 			int32 ScaleOffset = INDEX_NONE;
 			if (BestScale.CompressedBytes.Num() > 0)
 			{
-				checkf(BestScale.ActualCompressionMode < ACF_MAX);
+				check(BestScale.ActualCompressionMode < ACF_MAX);
 				ScaleOffset = AnimSeq->CompressedByteStream.Num();
 				AnimSeq->CompressedByteStream.Append(BestScale.CompressedBytes);
 			}
@@ -1195,7 +1195,7 @@ void UAnimCompress_PerTrackCompression::CompressUsingUnderlyingCompressor(
 		int32 TranslationOffset = INDEX_NONE;
 		if( !bSkipTranslationTrack && BestTranslation.CompressedBytes.Num() > 0 )
 		{
-			checkf(BestTranslation.ActualCompressionMode < ACF_MAX);
+			check(BestTranslation.ActualCompressionMode < ACF_MAX);
 			TranslationOffset = AnimSeq->CompressedByteStream.Num();
 			AnimSeq->CompressedByteStream.Append(BestTranslation.CompressedBytes);
 		}
@@ -1204,7 +1204,7 @@ void UAnimCompress_PerTrackCompression::CompressUsingUnderlyingCompressor(
 		int32 RotationOffset = INDEX_NONE;
 		if (BestRotation.CompressedBytes.Num() > 0)
 		{
-			checkf(BestRotation.ActualCompressionMode < ACF_MAX);
+			check(BestRotation.ActualCompressionMode < ACF_MAX);
 			RotationOffset = AnimSeq->CompressedByteStream.Num();
 			AnimSeq->CompressedByteStream.Append(BestRotation.CompressedBytes);
 		}
@@ -1272,10 +1272,10 @@ void ResamplePositionKeys(
 
 		FVector Value;
 
-		checkf(Track.Times[CachedIndex] <= Time);
+		check(Track.Times[CachedIndex] <= Time);
 		if (CachedIndex + 1 < KeyCount)
 		{
-			checkf(Track.Times[CachedIndex+1] >= Time);
+			check(Track.Times[CachedIndex+1] >= Time);
 
 			FVector A = Track.PosKeys[CachedIndex];
 			FVector B = Track.PosKeys[CachedIndex + 1];
@@ -1343,10 +1343,10 @@ void ResampleScaleKeys(
 
 		FVector Value;
 
-		checkf(Track.Times[CachedIndex] <= Time);
+		check(Track.Times[CachedIndex] <= Time);
 		if (CachedIndex + 1 < KeyCount)
 		{
-			checkf(Track.Times[CachedIndex+1] >= Time);
+			check(Track.Times[CachedIndex+1] >= Time);
 
 			FVector A = Track.ScaleKeys[CachedIndex];
 			FVector B = Track.ScaleKeys[CachedIndex + 1];
@@ -1414,10 +1414,10 @@ void ResampleRotationKeys(
 
 		FQuat Value;
 
-		checkf(Track.Times[CachedIndex] <= Time);
+		check(Track.Times[CachedIndex] <= Time);
 		if (CachedIndex + 1 < KeyCount)
 		{
-			checkf(Track.Times[CachedIndex+1] >= Time);
+			check(Track.Times[CachedIndex+1] >= Time);
 
 			FQuat A = Track.RotKeys[CachedIndex];
 			FQuat B = Track.RotKeys[CachedIndex + 1];
diff --git a/Engine/Source/Runtime/Engine/Private/Animation/AnimEncoding.cpp b/Engine/Source/Runtime/Engine/Private/Animation/AnimEncoding.cpp
--- a/Engine/Source/Runtime/Engine/Private/Animation/AnimEncoding.cpp
+++ b/Engine/Source/Runtime/Engine/Private/Animation/AnimEncoding.cpp
@@ -865,8 +865,8 @@ void AnimationFormat_SetInterfaceLinks(UAnimSequence& Seq)
 		Seq.TranslationCodec = &StaticCodec;
 		Seq.ScaleCodec = &StaticCodec;
 
-		checkf(Seq.RotationCompressionFormat == ACF_Identity);
-		checkf(Seq.TranslationCompressionFormat == ACF_Identity);
+		check(Seq.RotationCompressionFormat == ACF_Identity);
+		check(Seq.TranslationCompressionFormat == ACF_Identity);
 		// commenting out scale check because the older versions won't have this set correctly
 		// and I can't get the version VER_UE4_ANIM_SUPPORT_NONUNIFORM_SCALE_ANIMATION here because this function
 		// is called in Serialize where GetLinker is too early to call
diff --git a/Engine/Source/Runtime/Engine/Private/Animation/AnimationUtils.cpp b/Engine/Source/Runtime/Engine/Private/Animation/AnimationUtils.cpp
--- a/Engine/Source/Runtime/Engine/Private/Animation/AnimationUtils.cpp
+++ b/Engine/Source/Runtime/Engine/Private/Animation/AnimationUtils.cpp
@@ -1217,7 +1217,7 @@ void FAnimationUtils::CompressAnimSequenceExplicit(
 						(*WinningCompressorErrorSum) += WinningCompressorError;
 						AlternativeCompressorSavings += WinningCompressorSavings;
 						*WinningCompressorMarginalSavingsSum += WinningCompressorMarginalSavings;
-						checkf(WinningCompressorSavings == SizeDecrease);
+						check(WinningCompressorSavings == SizeDecrease);
 
 					UE_LOG(LogAnimation, Warning, TEXT("  Recompressing '%s' with compressor '%s' saved %i bytes (%i -> %i -> %i) (max diff=%f)\n"),
 							*AnimSeq->GetName(),
diff --git a/Engine/Source/Runtime/Engine/Private/GameInstance.cpp b/Engine/Source/Runtime/Engine/Private/GameInstance.cpp
--- a/Engine/Source/Runtime/Engine/Private/GameInstance.cpp
+++ b/Engine/Source/Runtime/Engine/Private/GameInstance.cpp
@@ -329,7 +329,7 @@ ULocalPlayer* UGameInstance::CreateInitialPlayer(FString& OutError)
 
 ULocalPlayer* UGameInstance::CreateLocalPlayer(int32 ControllerId, FString& OutError, bool bSpawnActor)
 {
-	checkf(GetEngine()->LocalPlayerClass != NULL);
+	check(GetEngine()->LocalPlayerClass != NULL);
 
 	ULocalPlayer* NewPlayer = NULL;
 	int32 InsertIndex = INDEX_NONE;
diff --git a/Engine/Source/Runtime/Engine/Private/GameViewportClient.cpp b/Engine/Source/Runtime/Engine/Private/GameViewportClient.cpp
--- a/Engine/Source/Runtime/Engine/Private/GameViewportClient.cpp
+++ b/Engine/Source/Runtime/Engine/Private/GameViewportClient.cpp
@@ -1433,7 +1433,7 @@ void UGameViewportClient::SetConsoleTarget(int32 PlayerIndex)
 
 ULocalPlayer* UGameViewportClient::SetupInitialLocalPlayer(FString& OutError)
 {
-	checkf(GetOuterUEngine()->ConsoleClass != NULL);
+	check(GetOuterUEngine()->ConsoleClass != NULL);
 
 	ActiveSplitscreenType = ESplitScreenType::None;
 
diff --git a/Engine/Source/Runtime/Engine/Private/Materials/HLSLMaterialTranslator.h b/Engine/Source/Runtime/Engine/Private/Materials/HLSLMaterialTranslator.h
--- a/Engine/Source/Runtime/Engine/Private/Materials/HLSLMaterialTranslator.h
+++ b/Engine/Source/Runtime/Engine/Private/Materials/HLSLMaterialTranslator.h
@@ -762,7 +762,7 @@ protected:
 			return Default;
 		}
 
-		checkf(Index >= 0 && Index < CurrentScopeChunks->Num(), TEXT("Index %d/%d, Platform=%d"), Index, CurrentScopeChunks->Num(), Platform);
+		checkf(Index >= 0 && Index < CurrentScopeChunks->Num(), TEXT("Index %d/%d, Platform=%d"), Index, CurrentScopeChunks->Num(), (int)Platform);
 		const FShaderCodeChunk& CodeChunk = (*CurrentScopeChunks)[Index];
 		if((CodeChunk.UniformExpression && CodeChunk.UniformExpression->IsConstant()) || CodeChunk.bInline)
 		{
diff --git a/Engine/Source/Runtime/Engine/Private/NetworkProfiler.cpp b/Engine/Source/Runtime/Engine/Private/NetworkProfiler.cpp
--- a/Engine/Source/Runtime/Engine/Private/NetworkProfiler.cpp
+++ b/Engine/Source/Runtime/Engine/Private/NetworkProfiler.cpp
@@ -565,7 +565,7 @@ void FNetworkProfiler::TrackSessionChange( bool bShouldContinueTracking, const F
 			// Create folder and file writer.
 			IFileManager::Get().MakeDirectory( *FPaths::GetPath(TempFileName) );
 			FileWriter = IFileManager::Get().CreateFileWriter( *TempFileName, FILEWRITE_EvenIfReadOnly );
-			checkf( FileWriter );
+			check( FileWriter );
 
 			// Serialize dummy header, overwritten when session ends.
 			FNetworkProfilerHeader DummyHeader;
diff --git a/Engine/Source/Runtime/Engine/Private/World.cpp b/Engine/Source/Runtime/Engine/Private/World.cpp
--- a/Engine/Source/Runtime/Engine/Private/World.cpp
+++ b/Engine/Source/Runtime/Engine/Private/World.cpp
@@ -684,7 +684,7 @@ UMaterialParameterCollectionInstance* UWorld::GetParameterCollectionInstance(con
 	}
 
 	// Instance should always exist due to SetupParameterCollectionInstances() and UMaterialParameterCollection::PostLoad()
-	checkf(0);
+	check(0);
 	return NULL;
 }
 
diff --git a/Engine/Source/Runtime/Slate/Private/Widgets/Input/SEditableText.cpp b/Engine/Source/Runtime/Slate/Private/Widgets/Input/SEditableText.cpp
--- a/Engine/Source/Runtime/Slate/Private/Widgets/Input/SEditableText.cpp
+++ b/Engine/Source/Runtime/Slate/Private/Widgets/Input/SEditableText.cpp
@@ -739,7 +739,7 @@ FReply SEditableText::MoveCursor( FMoveCursor Args )
 		}
 		else
 		{
-			checkSlow(false, "Unknown ECursorMoveMethod value");
+			checkfSlow(false, TEXT("Unknown ECursorMoveMethod value"));
 		}
 	}
 
diff --git a/Engine/Source/Runtime/Slate/Private/Widgets/Text/SMultiLineEditableText.cpp b/Engine/Source/Runtime/Slate/Private/Widgets/Text/SMultiLineEditableText.cpp
--- a/Engine/Source/Runtime/Slate/Private/Widgets/Text/SMultiLineEditableText.cpp
+++ b/Engine/Source/Runtime/Slate/Private/Widgets/Text/SMultiLineEditableText.cpp
@@ -973,7 +973,7 @@ FReply SMultiLineEditableText::MoveCursor( FMoveCursor Args )
 		}
 		else
 		{
-			checkSlow(false, "Unknown ECursorMoveMethod value");
+			checkfSlow(false, TEXT("Unknown ECursorMoveMethod value"));
 		}
 	}
 
diff --git a/Engine/Source/Runtime/SlateCore/Private/Animation/CurveHandle.cpp b/Engine/Source/Runtime/SlateCore/Private/Animation/CurveHandle.cpp
--- a/Engine/Source/Runtime/SlateCore/Private/Animation/CurveHandle.cpp
+++ b/Engine/Source/Runtime/SlateCore/Private/Animation/CurveHandle.cpp
@@ -130,7 +130,7 @@ float FCurveHandle::ApplyEasing( float Time, ECurveEaseFunction::Type EaseFuncti
 
 		default:
 			// Unrecognized curve easing function type
-			checkf(0, *LOCTEXT("CurveFunction_Error", "Unrecognized curve easing function type [%i] for FCurveHandle").ToString(), EaseFunction);
+			checkf(0, *LOCTEXT("CurveFunction_Error", "Unrecognized curve easing function type [%i] for FCurveHandle").ToString(), (int)EaseFunction);
 			break;
 	}
 
