Fixed navigation data deleted in the editor not being properly removed from navigation system's NavDataSet #UE4

diff --git a/Engine/Source/Runtime/Engine/Classes/AI/Navigation/NavigationData.h b/Engine/Source/Runtime/Engine/Classes/AI/Navigation/NavigationData.h
--- a/Engine/Source/Runtime/Engine/Classes/AI/Navigation/NavigationData.h
+++ b/Engine/Source/Runtime/Engine/Classes/AI/Navigation/NavigationData.h
@@ -459,7 +459,7 @@ public:
 #if WITH_EDITOR
 	virtual void PostEditUndo() override;
 #endif // WITH_EDITOR
-	virtual void BeginDestroy() override;
+	virtual void Destroyed() override;
 	// End UObject Interface
 		
 	virtual void CleanUp();
diff --git a/Engine/Source/Runtime/Engine/Private/AI/Navigation/NavigationData.cpp b/Engine/Source/Runtime/Engine/Private/AI/Navigation/NavigationData.cpp
--- a/Engine/Source/Runtime/Engine/Private/AI/Navigation/NavigationData.cpp
+++ b/Engine/Source/Runtime/Engine/Private/AI/Navigation/NavigationData.cpp
@@ -362,7 +362,7 @@ bool ANavigationData::DoesSupportAgent(const FNavAgentProperties& AgentProps) co
 	return NavDataConfig.IsEquivalent(AgentProps);
 }
 
-void ANavigationData::BeginDestroy()
+void ANavigationData::Destroyed()
 {
 	UWorld* WorldOuter = GetWorld();
 
@@ -374,7 +374,7 @@ void ANavigationData::BeginDestroy()
 
 	CleanUp();
 
-	Super::BeginDestroy();
+	Super::Destroyed();
 }
 
 void ANavigationData::CleanUp()
