Fixed navmesh not rebuilding to proper settings when changing navmesh generation properties in the editor #UE4

diff --git a/Engine/Source/Runtime/Engine/Private/AI/Navigation/RecastNavMesh.cpp b/Engine/Source/Runtime/Engine/Private/AI/Navigation/RecastNavMesh.cpp
--- a/Engine/Source/Runtime/Engine/Private/AI/Navigation/RecastNavMesh.cpp
+++ b/Engine/Source/Runtime/Engine/Private/AI/Navigation/RecastNavMesh.cpp
@@ -1927,26 +1927,24 @@ bool ARecastNavMesh::SupportsStreaming() const
 
 void ARecastNavMesh::ConditionalConstructGenerator()
 {	
+	if (NavDataGenerator.IsValid())
+	{
+		NavDataGenerator->CancelBuild();
+		NavDataGenerator.Reset();
+	}
+
 	UWorld* World = GetWorld();
 	check(World);
 	const bool bRequiresGenerator = SupportsRuntimeGeneration() || !World->IsGameWorld();
 	if (bRequiresGenerator)
 	{
-		if (NavDataGenerator.IsValid() == false)
-		{
-			NavDataGenerator = MakeShareable(new FRecastNavMeshGenerator(*this));
+		NavDataGenerator = MakeShareable(new FRecastNavMeshGenerator(*this));
 
-			if (World->GetNavigationSystem())
-			{
-				RestrictBuildingToActiveTiles(World->GetNavigationSystem()->IsActiveTilesGenerationEnabled());
-			}
+		if (World->GetNavigationSystem())
+		{
+			RestrictBuildingToActiveTiles(World->GetNavigationSystem()->IsActiveTilesGenerationEnabled());
 		}
 	}
-	else
-	{
-		NavDataGenerator->CancelBuild();
-		NavDataGenerator.Reset();
-	}
 }
 
 bool ARecastNavMesh::IsVoxelCacheEnabled()
