UMG/Slate - Introducing an option to provide a custom DPI Scaler object that can perform any DPI scale it wants based on custom logic.

diff --git a/Engine/Source/Runtime/Engine/Classes/Engine/DPICustomScalingRule.h b/Engine/Source/Runtime/Engine/Classes/Engine/DPICustomScalingRule.h
new file mode 100644
--- /dev/null
+++ b/Engine/Source/Runtime/Engine/Classes/Engine/DPICustomScalingRule.h
@@ -0,0 +1,18 @@
+// Copyright 1998-2015 Epic Games, Inc. All Rights Reserved.
+
+#pragma once
+
+#include "DPICustomScalingRule.generated.h"
+
+/**
+ * 
+ */
+UCLASS(Abstract)
+class ENGINE_API UDPICustomScalingRule : public UObject
+{
+	GENERATED_BODY()
+
+public:
+
+	virtual float GetDPIScaleBasedOnSize(FIntPoint Size) const;
+};
diff --git a/Engine/Source/Runtime/Engine/Classes/Engine/RendererSettings.h b/Engine/Source/Runtime/Engine/Classes/Engine/RendererSettings.h
--- a/Engine/Source/Runtime/Engine/Classes/Engine/RendererSettings.h
+++ b/Engine/Source/Runtime/Engine/Classes/Engine/RendererSettings.h
@@ -264,7 +264,7 @@ public:
 private:
 	
 	UPROPERTY(config)
-	TEnumAsByte<EUIScalingRule::Type> UIScaleRule_DEPRECATED;
+	TEnumAsByte<EUIScalingRule> UIScaleRule_DEPRECATED;
 
 	UPROPERTY(config)
 	FRuntimeFloatCurve UIScaleCurve_DEPRECATED;
diff --git a/Engine/Source/Runtime/Engine/Classes/Engine/UserInterfaceSettings.h b/Engine/Source/Runtime/Engine/Classes/Engine/UserInterfaceSettings.h
--- a/Engine/Source/Runtime/Engine/Classes/Engine/UserInterfaceSettings.h
+++ b/Engine/Source/Runtime/Engine/Classes/Engine/UserInterfaceSettings.h
@@ -14,28 +14,28 @@ enum class ERenderFocusRule : uint8
 	NonPointer,
 	/** Focus Brush will be rendered for widgets that have user focus only if the focus was set by navigation. */
 	NavigationOnly,
+	/** Focus Brush will not be rendered */
+	Never,
 };
 
 /** The Side to use when scaling the UI. */
 UENUM()
-namespace EUIScalingRule
+enum class EUIScalingRule : uint8
 {
-	enum Type
-	{
-		/** Evaluates the scale curve based on the shortest side of the viewport */
-		ShortestSide,
-		/** Evaluates the scale curve based on the longest side of the viewport */
-		LongestSide,
-		/** Evaluates the scale curve based on the X axis of the viewport */
-		Horizontal,
-		/** Evaluates the scale curve based on the Y axis of the viewport */
-		Vertical,
-		/** Custom - Allows custom rule interpretation */
-		//Custom
-	};
-}
+	/** Evaluates the scale curve based on the shortest side of the viewport */
+	ShortestSide,
+	/** Evaluates the scale curve based on the longest side of the viewport */
+	LongestSide,
+	/** Evaluates the scale curve based on the X axis of the viewport */
+	Horizontal,
+	/** Evaluates the scale curve based on the Y axis of the viewport */
+	Vertical,
+	/** Custom - Allows custom rule interpretation */
+	Custom
+};
 
 class UWidget;
+class UDPICustomScalingRule;
 
 /**
  * Implements user interface related settings.
@@ -69,10 +69,16 @@ public:
 	UPROPERTY(config, EditAnywhere, Category = "Cursors", meta = (MetaClass = "UserWidget", ToolTip = "Widget to use when the SlashedCircle Cursor is requested."))
 	FStringClassReference SlashedCircleCursor;
 
+	UPROPERTY(config, EditAnywhere, Category="DPI Scaling", meta = ( ToolTip = "An optional application scale to apply on top of the custom scaling rules.  So if you want to expose a property in your game title, you can modify this underlying value to scale the entire UI." ))
+	float ApplicationScale;
+
 	UPROPERTY(config, EditAnywhere, Category="DPI Scaling", meta=(
 		DisplayName="DPI Scale Rule",
 		ToolTip="The rule used when trying to decide what scale to apply." ))
-	TEnumAsByte<EUIScalingRule::Type> UIScaleRule;
+	TEnumAsByte<EUIScalingRule> UIScaleRule;
+
+	UPROPERTY(config, EditAnywhere, Category="DPI Scaling", meta=( MetaClass="DPICustomScalingRule", ToolTip = "Set DPI Scale Rule to Custom, and this class will be used instead of any of the built-in rules." ))
+	FStringClassReference CustomScalingRuleClass;
 
 	UPROPERTY(config, EditAnywhere, Category="DPI Scaling", meta=(
 		DisplayName="DPI Curve",
@@ -81,15 +87,11 @@ public:
 		YAxisName="Scale"))
 	FRuntimeFloatCurve UIScaleCurve;
 
-public:
-	UPROPERTY(Transient)
-	TArray<UObject*> CursorClasses;
-
 public:
 
 	virtual void PostInitProperties() override;
 
-	void LoadCursors();
+	void ForceLoadResources();
 
 	/** Gets the current scale of the UI based on the size */
 	float GetDPIScaleBasedOnSize(FIntPoint Size) const;
@@ -97,4 +99,14 @@ public:
 #if WITH_EDITOR
 	void PostEditChangeProperty(FPropertyChangedEvent& PropertyChangedEvent);
 #endif
+
+private:
+	UPROPERTY(Transient)
+	TArray<UObject*> CursorClasses;
+
+	UPROPERTY(Transient)
+	mutable UClass* CustomScalingRuleClassInstance;
+
+	UPROPERTY(Transient)
+	mutable UDPICustomScalingRule* CustomScalingRule;
 };
diff --git a/Engine/Source/Runtime/Engine/Private/DPICustomScalingRule.cpp b/Engine/Source/Runtime/Engine/Private/DPICustomScalingRule.cpp
new file mode 100644
--- /dev/null
+++ b/Engine/Source/Runtime/Engine/Private/DPICustomScalingRule.cpp
@@ -0,0 +1,10 @@
+// Copyright 1998-2015 Epic Games, Inc. All Rights Reserved.
+
+#include "EnginePrivate.h"
+
+#include "Engine/DPICustomScalingRule.h"
+
+float UDPICustomScalingRule::GetDPIScaleBasedOnSize(FIntPoint Size) const
+{
+	return 1;
+}
diff --git a/Engine/Source/Runtime/Engine/Private/UnrealEngine.cpp b/Engine/Source/Runtime/Engine/Private/UnrealEngine.cpp
--- a/Engine/Source/Runtime/Engine/Private/UnrealEngine.cpp
+++ b/Engine/Source/Runtime/Engine/Private/UnrealEngine.cpp
@@ -1466,7 +1466,7 @@ void UEngine::InitializeObjectReferences()
 	}
 
 	UUserInterfaceSettings* UISettings = GetMutableDefault<UUserInterfaceSettings>(UUserInterfaceSettings::StaticClass());
-	UISettings->LoadCursors();
+	UISettings->ForceLoadResources();
 }
 
 //
diff --git a/Engine/Source/Runtime/Engine/Private/UserInterfaceSettings.cpp b/Engine/Source/Runtime/Engine/Private/UserInterfaceSettings.cpp
--- a/Engine/Source/Runtime/Engine/Private/UserInterfaceSettings.cpp
+++ b/Engine/Source/Runtime/Engine/Private/UserInterfaceSettings.cpp
@@ -1,17 +1,19 @@
 // Copyright 1998-2015 Epic Games, Inc. All Rights Reserved.
 
-/*=============================================================================
-	PlayerInput.cpp: Unreal input system.
-=============================================================================*/
-
 #include "EnginePrivate.h"
+
+#include "MessageLog.h"
+
 #include "Engine/RendererSettings.h"
 #include "Engine/UserInterfaceSettings.h"
+#include "Engine/DPICustomScalingRule.h"
 
+#define LOCTEXT_NAMESPACE "Engine"
 
 UUserInterfaceSettings::UUserInterfaceSettings(const FObjectInitializer& ObjectInitializer)
-	: Super(ObjectInitializer),
-	  RenderFocusRule(ERenderFocusRule::NavigationOnly)
+	: Super(ObjectInitializer)
+	, RenderFocusRule(ERenderFocusRule::NavigationOnly)
+	, ApplicationScale(1)
 {
 }
 
@@ -21,7 +23,7 @@ void UUserInterfaceSettings::PostInitProperties()
 
 	if ( HasAnyFlags(RF_ClassDefaultObject) == false )
 	{
-		LoadCursors();
+		ForceLoadResources();
 	}
 }
 
@@ -29,11 +31,28 @@ float UUserInterfaceSettings::GetDPIScaleBasedOnSize(FIntPoint Size) const
 {
 	float Scale = 1;
 
-	//if ( UIScaleRule == EUIScalingRule::Custom )
-	//{
+	if ( UIScaleRule == EUIScalingRule::Custom )
+	{
+		if ( CustomScalingRuleClassInstance == nullptr )
+		{
+			CustomScalingRuleClassInstance = CustomScalingRuleClass.TryLoadClass<UDPICustomScalingRule>();
+
+			if ( CustomScalingRuleClassInstance == nullptr )
+			{
+				FMessageLog("MapCheck").Error()
+					->AddToken(FTextToken::Create(FText::Format(LOCTEXT("CustomScalingRule_NotFound", "Project Settings - User Interface Custom Scaling Rule '{0}' could not be found."), FText::FromString(CustomScalingRuleClass.AssetLongPathname))));
+				return 1;
+			}
+		}
+		
+		if ( CustomScalingRule == nullptr )
+		{
+			CustomScalingRule = CustomScalingRuleClassInstance->GetDefaultObject<UDPICustomScalingRule>();
+		}
 
-	//}
-	//else
+		Scale = CustomScalingRule->GetDPIScaleBasedOnSize(Size);
+	}
+	else
 	{
 		int32 EvalPoint = 0;
 		switch ( UIScaleRule )
@@ -56,10 +75,10 @@ float UUserInterfaceSettings::GetDPIScaleBasedOnSize(FIntPoint Size) const
 		Scale = DPICurve->Eval((float)EvalPoint, 1.0f);
 	}
 
-	return FMath::Max(Scale, 0.01f);
+	return FMath::Max(Scale * ApplicationScale, 0.01f);
 }
 
-void UUserInterfaceSettings::LoadCursors()
+void UUserInterfaceSettings::ForceLoadResources()
 {
 	TArray<UObject*> LoadedClasses;
 	LoadedClasses.Add(DefaultCursor.TryLoad());
@@ -76,6 +95,8 @@ void UUserInterfaceSettings::LoadCursors()
 			CursorClasses.Add(Cursor);
 		}
 	}
+
+	CustomScalingRuleClassInstance = CustomScalingRuleClass.TryLoadClass<UDPICustomScalingRule>();
 }
 
 #if WITH_EDITOR
@@ -90,3 +111,5 @@ void UUserInterfaceSettings::PostEditChangeProperty(FPropertyChangedEvent& Prope
 	}
 }
 #endif // #if WITH_EDITOR
+
+#undef LOCTEXT_NAMESPACE
