UMG/Slate - Adding some comments to the DPICustomScalingRule class.

diff --git a/Engine/Source/Runtime/Engine/Classes/Engine/DPICustomScalingRule.h b/Engine/Source/Runtime/Engine/Classes/Engine/DPICustomScalingRule.h
--- a/Engine/Source/Runtime/Engine/Classes/Engine/DPICustomScalingRule.h
+++ b/Engine/Source/Runtime/Engine/Classes/Engine/DPICustomScalingRule.h
@@ -5,7 +5,8 @@
 #include "DPICustomScalingRule.generated.h"
 
 /**
- * 
+ * Custom Scaling Rules for Slate and UMG Widgets can be implemented by sub-classing from this class
+ * and setting this rule to be used in your project settings.
  */
 UCLASS(Abstract)
 class ENGINE_API UDPICustomScalingRule : public UObject
@@ -14,5 +15,10 @@ class ENGINE_API UDPICustomScalingRule : public UObject
 
 public:
 
+	/**
+	 * Return the scale to use given the size of the viewport.
+	 * @param Size The size of the viewport.
+	 * @return The Scale to apply to the entire UI.
+	 */
 	virtual float GetDPIScaleBasedOnSize(FIntPoint Size) const;
 };
