Merge part of CL2472449 from Orion branch:

- Fix issue with the last packet not always being processed during replay viewing

diff --git a/Engine/Source/Runtime/Engine/Classes/Engine/DemoNetDriver.h b/Engine/Source/Runtime/Engine/Classes/Engine/DemoNetDriver.h
--- a/Engine/Source/Runtime/Engine/Classes/Engine/DemoNetDriver.h
+++ b/Engine/Source/Runtime/Engine/Classes/Engine/DemoNetDriver.h
@@ -57,7 +57,7 @@ class UDemoNetDriver
 	double		LastRecordAvgFlush;
 	double		MaxRecordTime;
 	int32		RecordCountSinceFlush;
-
+	float		TimeToSkip;
 public:
 
 	// UNetDriver interface.
@@ -71,7 +71,7 @@ public:
 	virtual void TickFlush( float DeltaSeconds ) override;
 	virtual void ProcessRemoteFunction( class AActor* Actor, class UFunction* Function, void* Parameters, struct FOutParmRec* OutParms, struct FFrame* Stack, class UObject* SubObject = nullptr );
 	virtual bool IsAvailable() const override { return true; }
-	
+	ENGINE_API void SkipTime(float InTimeToSkip);
 	bool InitConnectInternal( FString& Error );
 
 public:
diff --git a/Engine/Source/Runtime/Engine/Private/DemoNetDriver.cpp b/Engine/Source/Runtime/Engine/Private/DemoNetDriver.cpp
--- a/Engine/Source/Runtime/Engine/Private/DemoNetDriver.cpp
+++ b/Engine/Source/Runtime/Engine/Private/DemoNetDriver.cpp
@@ -47,6 +47,7 @@ bool UDemoNetDriver::InitBase( bool bInitAsClient, FNetworkNotify* InNotify, con
 		bIsRecordingDemoFrame	= false;
 		bDemoPlaybackDone		= false;
 		bChannelsArePaused		= false;
+		TimeToSkip = 0.0f;
 
 		ResetDemoState();
 
@@ -903,6 +904,11 @@ bool UDemoNetDriver::ReadDemoFrame()
 	return true;
 }
 
+void UDemoNetDriver::SkipTime(float InTimeToSkip)
+{
+	TimeToSkip = InTimeToSkip;
+}
+
 void UDemoNetDriver::TickDemoPlayback( float DeltaSeconds )
 {
 	if ( ServerConnection == NULL || ServerConnection->State == USOCK_Closed )
@@ -920,13 +926,8 @@ void UDemoNetDriver::TickDemoPlayback( float DeltaSeconds )
 
 	DemoDeltaTime += DeltaSeconds;
 
-	if ( DemoCurrentTime + DemoDeltaTime > DemoTotalTime )
-	{
-		DemoDeltaTime = DemoTotalTime - DemoCurrentTime;
-	}
-
 	DemoCurrentTime += DeltaSeconds;
-
+	
 	if ( CVarDemoSkipTime.GetValueOnGameThread() > 0 )
 	{
 		DemoDeltaTime += CVarDemoSkipTime.GetValueOnGameThread();
@@ -940,6 +941,20 @@ void UDemoNetDriver::TickDemoPlayback( float DeltaSeconds )
 
 		CVarDemoSkipTime.AsVariable()->Set( TEXT( "0" ), ECVF_SetByConsole );
 	}
+	
+	if (TimeToSkip > 0.0f)
+	{
+		DemoDeltaTime += TimeToSkip;
+
+		if (DemoCurrentTime + DemoDeltaTime > DemoTotalTime)
+		{
+			DemoDeltaTime = DemoTotalTime - DemoCurrentTime;
+		}
+
+		DemoCurrentTime += TimeToSkip;
+
+		TimeToSkip = 0.0f;
+	}
 
 	if ( DemoCurrentTime > DemoTotalTime )
 	{
