Convex elements draw using scale

diff --git a/Engine/Source/Editor/PhAT/Private/PhATRender.cpp b/Engine/Source/Editor/PhAT/Private/PhATRender.cpp
--- a/Engine/Source/Editor/PhAT/Private/PhATRender.cpp
+++ b/Engine/Source/Editor/PhAT/Private/PhATRender.cpp
@@ -190,7 +190,7 @@ void UPhATEdSkeletalMeshComponent::RenderAssetTools(const FSceneView* View, clas
 				{
 					if (CollisionViewMode == FPhATSharedData::PRM_Solid || CollisionViewMode == FPhATSharedData::PRM_Wireframe)
 					{
-						AggGeom->ConvexElems[j].DrawElemWire(PDI, ElemTM, GetPrimitiveColor(i, KPT_Convex, j));
+						AggGeom->ConvexElems[j].DrawElemWire(PDI, ElemTM, Scale, GetPrimitiveColor(i, KPT_Convex, j));
 					}
 				}
 				
diff --git a/Engine/Source/Editor/StaticMeshEditor/Private/StaticMeshEditorViewportClient.cpp b/Engine/Source/Editor/StaticMeshEditor/Private/StaticMeshEditorViewportClient.cpp
--- a/Engine/Source/Editor/StaticMeshEditor/Private/StaticMeshEditorViewportClient.cpp
+++ b/Engine/Source/Editor/StaticMeshEditor/Private/StaticMeshEditorViewportClient.cpp
@@ -485,7 +485,7 @@ void FStaticMeshEditorViewportClient::Draw(const FSceneView* View,FPrimitiveDraw
 			const FColor CollisionColor = StaticMeshEditorPtr.Pin()->IsSelectedPrim(HitProxy->PrimData) ? SelectedColor : UnselectedColor;
 			const FKConvexElem& ConvexElem = AggGeom->ConvexElems[i];
 			const FTransform ElemTM = ConvexElem.GetTransform();
-			ConvexElem.DrawElemWire(PDI, ElemTM, CollisionColor);
+			ConvexElem.DrawElemWire(PDI, ElemTM, 1.f, CollisionColor);
 
 			PDI->SetHitProxy(NULL);
 		}
diff --git a/Engine/Source/Runtime/Engine/Classes/PhysicsEngine/ConvexElem.h b/Engine/Source/Runtime/Engine/Classes/PhysicsEngine/ConvexElem.h
--- a/Engine/Source/Runtime/Engine/Classes/PhysicsEngine/ConvexElem.h
+++ b/Engine/Source/Runtime/Engine/Classes/PhysicsEngine/ConvexElem.h
@@ -42,9 +42,12 @@ struct FKConvexElem
 		, ConvexMeshNegX(NULL)
 	{}
 
+	DEPRECATED(4.8, "Please call DrawElemWire which takes in a Scale parameter")
 	ENGINE_API void	DrawElemWire(class FPrimitiveDrawInterface* PDI, const FTransform& ElemTM, const FColor Color) const;
+	
+	ENGINE_API void	DrawElemWire(class FPrimitiveDrawInterface* PDI, const FTransform& ElemTM, const float Scale, const FColor Color) const;
 
-	void AddCachedSolidConvexGeom(TArray<FDynamicMeshVertex>& VertexBuffer, TArray<int32>& IndexBuffer, const FColor VertexColor) const;
+	void AddCachedSolidConvexGeom(TArray<FDynamicMeshVertex>& VertexBuffer, TArray<int32>& IndexBuffer, const float Scale, const FColor VertexColor) const;
 
 	/** Reset the hull to empty all arrays */
 	ENGINE_API void	Reset();
diff --git a/Engine/Source/Runtime/Engine/Private/PhysicsEngine/PhysDrawing.cpp b/Engine/Source/Runtime/Engine/Private/PhysicsEngine/PhysDrawing.cpp
--- a/Engine/Source/Runtime/Engine/Private/PhysicsEngine/PhysDrawing.cpp
+++ b/Engine/Source/Runtime/Engine/Private/PhysicsEngine/PhysDrawing.cpp
@@ -355,7 +355,13 @@ void FKSphylElem::DrawElemSolid(FPrimitiveDrawInterface* PDI, const FTransform&
 // FKConvexElem
 /////////////////////////////////////////////////////////////////////////////////////
 
+
 void FKConvexElem::DrawElemWire(FPrimitiveDrawInterface* PDI, const FTransform& ElemTM, const FColor Color) const
+{
+	DrawElemWire(PDI, ElemTM, 1.f, Color);
+}
+
+void FKConvexElem::DrawElemWire(FPrimitiveDrawInterface* PDI, const FTransform& ElemTM, const float Scale, const FColor Color) const
 {
 #if WITH_PHYSX
 
@@ -371,7 +377,7 @@ void FKConvexElem::DrawElemWire(FPrimitiveDrawInterface* PDI, const FTransform&
 		TransformedVerts.AddUninitialized(NbVerts);
 		for(PxU32 i=0; i<NbVerts; i++)
 		{
-			TransformedVerts[i] = ElemTM.TransformPosition(P2UVector(Vertices[i]));
+			TransformedVerts[i] = ElemTM.TransformPosition(P2UVector(Vertices[i]) * Scale);
 		}
 						
 		const PxU8* PIndexBuffer = Mesh->getIndexBuffer();
@@ -408,7 +414,7 @@ void FKConvexElem::DrawElemWire(FPrimitiveDrawInterface* PDI, const FTransform&
 #endif // WITH_PHYSX
 }
 
-void FKConvexElem::AddCachedSolidConvexGeom(TArray<FDynamicMeshVertex>& VertexBuffer, TArray<int32>& IndexBuffer, const FColor VertexColor) const
+void FKConvexElem::AddCachedSolidConvexGeom(TArray<FDynamicMeshVertex>& VertexBuffer, TArray<int32>& IndexBuffer, const float Scale, const FColor VertexColor) const
 {
 #if WITH_PHYSX
 	if(ConvexMesh)
@@ -439,7 +445,7 @@ void FKConvexElem::AddCachedSolidConvexGeom(TArray<FDynamicMeshVertex>& VertexBu
 				int32 VertIndex = indices[j];
 
 				FDynamicMeshVertex Vert1;
-				Vert1.Position = Transform.TransformPosition( P2UVector(PVertices[VertIndex]) ); // Apply element transform to get geom in component space
+				Vert1.Position = Transform.TransformPosition( P2UVector(PVertices[VertIndex]) * Scale); // Apply element transform to get geom in component space
 				Vert1.Color = VertexColor;
 				Vert1.SetTangents(
 					TangentX,
@@ -590,7 +596,7 @@ void FKAggregateGeom::GetAggGeom(const FTransform& Transform, const FColor Color
 				for(int32 i=0; i<ConvexElems.Num(); i++)
 				{
 					// Get vertices/triangles from this hull.
-					ConvexElems[i].AddCachedSolidConvexGeom(ThisGeom.RenderInfo->VertexBuffer->Vertices, ThisGeom.RenderInfo->IndexBuffer->Indices, FColor::White);
+					ConvexElems[i].AddCachedSolidConvexGeom(ThisGeom.RenderInfo->VertexBuffer->Vertices, ThisGeom.RenderInfo->IndexBuffer->Indices, Scale3D.X, FColor::White);
 				}
 
 				// Only continue if we actually got some valid geometry
@@ -640,7 +646,7 @@ void FKAggregateGeom::GetAggGeom(const FTransform& Transform, const FColor Color
 				FColor ConvexColor = bPerHullColor ? DebugUtilColor[i%NUM_DEBUG_UTIL_COLORS] : Color;
 				FTransform ElemTM = ConvexElems[i].GetTransform();
 				ElemTM *= Transform;
-				ConvexElems[i].DrawElemWire(Collector.GetPDI(ViewIndex), ElemTM, ConvexColor);
+				ConvexElems[i].DrawElemWire(Collector.GetPDI(ViewIndex), ElemTM, Scale3D.X, ConvexColor);
 			}
 		}
 	}
