Fix crash inside ensure when component is null. UE-10333

diff --git a/Engine/Source/Runtime/Engine/Private/PhysicsEngine/BodyInstance.cpp b/Engine/Source/Runtime/Engine/Private/PhysicsEngine/BodyInstance.cpp
--- a/Engine/Source/Runtime/Engine/Private/PhysicsEngine/BodyInstance.cpp
+++ b/Engine/Source/Runtime/Engine/Private/PhysicsEngine/BodyInstance.cpp
@@ -2371,7 +2371,8 @@ void FBodyInstance::SetBodyTransform(const FTransform& NewTransform, bool bTelep
 #endif // !(UE_BUILD_SHIPPING || UE_BUILD_TEST)
 
 	// Catch NaNs and elegantly bail out.
-	if( !ensureMsgf(!NewTransform.ContainsNaN(), TEXT("SetBodyTransform contains NaN (%s: %s)\n%s"), *GetPathNameSafe(OwnerComponent.Get()), *GetPathNameSafe(OwnerComponent.Get()->GetOuter()), *NewTransform.ToString()) )
+
+	if( !ensureMsgf(!NewTransform.ContainsNaN(), TEXT("SetBodyTransform contains NaN (%s)\n%s"), (OwnerComponent.Get() ? *OwnerComponent->GetPathName() : TEXT("NONE")), *NewTransform.ToString()) )
 	{
 		return;
 	}
