Consolidate code

diff --git a/Engine/Source/Runtime/Engine/Private/DemoNetDriver.cpp b/Engine/Source/Runtime/Engine/Private/DemoNetDriver.cpp
--- a/Engine/Source/Runtime/Engine/Private/DemoNetDriver.cpp
+++ b/Engine/Source/Runtime/Engine/Private/DemoNetDriver.cpp
@@ -930,15 +930,7 @@ void UDemoNetDriver::TickDemoPlayback( float DeltaSeconds )
 	
 	if ( CVarDemoSkipTime.GetValueOnGameThread() > 0 )
 	{
-		DemoDeltaTime += CVarDemoSkipTime.GetValueOnGameThread();
-
-		if ( DemoCurrentTime + DemoDeltaTime > DemoTotalTime )
-		{
-			DemoDeltaTime = DemoTotalTime - DemoCurrentTime;
-		}
-
-		DemoCurrentTime += CVarDemoSkipTime.GetValueOnGameThread();
-
+		SkipTime( CVarDemoSkipTime.GetValueOnGameThread() );		// Just overwrite existing value, cvar wins in this case
 		CVarDemoSkipTime.AsVariable()->Set( TEXT( "0" ), ECVF_SetByConsole );
 	}
 	
