Fixing stat memory tracking in audio engine.

Not incrementing the tracked memory variable was making the memory stats say audio was leaking as thats the value used to decrement the stat tracker.

diff --git a/Engine/Source/Runtime/Engine/Private/AudioDecompress.cpp b/Engine/Source/Runtime/Engine/Private/AudioDecompress.cpp
--- a/Engine/Source/Runtime/Engine/Private/AudioDecompress.cpp
+++ b/Engine/Source/Runtime/Engine/Private/AudioDecompress.cpp
@@ -326,6 +326,7 @@ void FAsyncAudioDecompressWorker::DoWork( void )
 			AudioInfo->ExpandFile(Wave->RawPCMData, &QualityInfo);
 
 			const SIZE_T ResSize = Wave->GetResourceSize(EResourceSizeMode::Exclusive);
+			Wave->TrackedMemoryUsage += ResSize;
 			INC_DWORD_STAT_BY( STAT_AudioMemorySize, ResSize );
 			INC_DWORD_STAT_BY( STAT_AudioMemory, ResSize );
 		}
