Don't mark PendingKill on actors and components until after we've cleaned up the world. This is needed because UnregisterComponent code often uses WeakObjectPtr etc.. which assume not pending kill.

UE-5940

diff --git a/Engine/Source/Runtime/Engine/Private/World.cpp b/Engine/Source/Runtime/Engine/Private/World.cpp
--- a/Engine/Source/Runtime/Engine/Private/World.cpp
+++ b/Engine/Source/Runtime/Engine/Private/World.cpp
@@ -4698,8 +4698,6 @@ UWorld* FSeamlessTravelHandler::Tick()
 
 					// otherwise, set to be deleted
 					KeepAnnotation.Clear(TheActor);
-					TheActor->MarkPendingKill();
-					TheActor->MarkComponentsAsPendingKill();
 					// close any channels for this actor
 					if (NetDriver != NULL)
 					{
