CIS error - caused by 2472887

diff --git a/Engine/Source/Runtime/Engine/Classes/Engine/BlueprintGeneratedClass.h b/Engine/Source/Runtime/Engine/Classes/Engine/BlueprintGeneratedClass.h
--- a/Engine/Source/Runtime/Engine/Classes/Engine/BlueprintGeneratedClass.h
+++ b/Engine/Source/Runtime/Engine/Classes/Engine/BlueprintGeneratedClass.h
@@ -468,7 +468,7 @@ public:
 	UPROPERTY()
 	TArray<FEventGraphFastCallPair> FastCallPairs;
 
-#if WITH_EDITOR
+#if WITH_EDITORONLY_DATA
 	UPROPERTY(Transient)
 	UObject* OverridenArchetypeForCDO;
 #endif //WITH_EDITOR
diff --git a/Engine/Source/Runtime/Engine/Private/BlueprintGeneratedClass.cpp b/Engine/Source/Runtime/Engine/Private/BlueprintGeneratedClass.cpp
--- a/Engine/Source/Runtime/Engine/Private/BlueprintGeneratedClass.cpp
+++ b/Engine/Source/Runtime/Engine/Private/BlueprintGeneratedClass.cpp
@@ -664,7 +664,9 @@ void UBlueprintGeneratedClass::PurgeClass(bool bRecompilingOnLoad)
 
 	UberGraphFramePointerProperty = NULL;
 	UberGraphFunction = NULL;
+#if WITH_EDITORONLY_DATA
 	OverridenArchetypeForCDO = NULL;
+#endif //WITH_EDITOR
 
 #if UE_BLUEPRINT_EVENTGRAPH_FASTCALLS
 	FastCallPairs.Empty();
