Fixed an issue with one of the FCanvasTileItem constructors

diff --git a/Engine/Source/Runtime/Engine/Private/UserInterface/CanvasItem.cpp b/Engine/Source/Runtime/Engine/Private/UserInterface/CanvasItem.cpp
--- a/Engine/Source/Runtime/Engine/Private/UserInterface/CanvasItem.cpp
+++ b/Engine/Source/Runtime/Engine/Private/UserInterface/CanvasItem.cpp
@@ -54,7 +54,7 @@ void FCanvasItemTestbed::Draw( class FViewport* Viewport, class FCanvas* Canvas
 
 	if (TestMaterial == nullptr)
 	{
-		TestMaterial = LoadObject<UMaterial>(NULL, TEXT("/Game/NewMaterial.NewMaterial"));
+		TestMaterial = LoadObject<UMaterial>(nullptr, TEXT("/Game/NewMaterial.NewMaterial"));
 	}
 	
 	// A little ott for a testbed - but I wanted to draw several lines to ensure it worked :)
@@ -301,7 +301,7 @@ FCanvasTileItem::FCanvasTileItem(const FVector2D& InPosition, const FTexture* In
 	, UV0(0.0f, 0.0f)
 	, UV1(1.0f, 1.0f)
 	, Texture(InTexture)
-	, MaterialRenderProxy(NULL)
+	, MaterialRenderProxy(nullptr)
 	, Rotation(ForceInitToZero)
 	, PivotPoint(FVector2D::ZeroVector)
 {
@@ -319,7 +319,7 @@ FCanvasTileItem::FCanvasTileItem(const FVector2D& InPosition, const FTexture* In
 	, UV0(0.0f, 0.0f)
 	, UV1(1.0f, 1.0f)
 	, Texture(InTexture)
-	, MaterialRenderProxy(NULL)
+	, MaterialRenderProxy(nullptr)
 	, Rotation(ForceInitToZero)
 	, PivotPoint(FVector2D::ZeroVector)
 {
@@ -335,7 +335,7 @@ FCanvasTileItem::FCanvasTileItem( const FVector2D& InPosition, const FVector2D&
 	, UV0( 0.0f, 0.0f )
 	, UV1( 1.0f, 1.0f )
 	, Texture( GWhiteTexture )
-	, MaterialRenderProxy( NULL )
+	, MaterialRenderProxy( nullptr )
 	, Rotation( ForceInitToZero )
 	, PivotPoint( FVector2D::ZeroVector )
 {		
@@ -347,7 +347,8 @@ FCanvasTileItem::FCanvasTileItem(const FVector2D& InPosition, const FTexture* In
 	, Z(1.0f)
 	, UV0(InUV0)
 	, UV1(InUV1)
-	, MaterialRenderProxy(NULL)
+	, Texture(InTexture)
+	, MaterialRenderProxy(nullptr)
 	, Rotation(ForceInitToZero)
 	, PivotPoint(FVector2D::ZeroVector)
 {
@@ -366,13 +367,13 @@ FCanvasTileItem::FCanvasTileItem(const FVector2D& InPosition, const FTexture* In
 	, UV0(InUV0)
 	, UV1(InUV1)
 	, Texture(InTexture)
-	, MaterialRenderProxy(NULL)
+	, MaterialRenderProxy(nullptr)
 	, Rotation(ForceInitToZero)
 	, PivotPoint(FVector2D::ZeroVector)
 {
 	SetColor(InColor);
 	// Ensure texture is valid.
-	check(InTexture != NULL);
+	check(InTexture != nullptr);
 }
 
 FCanvasTileItem::FCanvasTileItem(const FVector2D& InPosition, const FMaterialRenderProxy* InMaterialRenderProxy, const FVector2D& InSize)
@@ -381,7 +382,7 @@ FCanvasTileItem::FCanvasTileItem(const FVector2D& InPosition, const FMaterialRen
 	, Z(1.0f)
 	, UV0(0.0f, 0.0f)
 	, UV1(1.0f, 1.0f)
-	, Texture(NULL)
+	, Texture(nullptr)
 	, MaterialRenderProxy(InMaterialRenderProxy)
 	, Rotation(ForceInitToZero)
 	, PivotPoint(FVector2D::ZeroVector)
@@ -396,7 +397,7 @@ FCanvasTileItem::FCanvasTileItem(const FVector2D& InPosition, const FMaterialRen
 	, Z(1.0f)
 	, UV0(InUV0)
 	, UV1(InUV1)
-	, Texture(NULL)
+	, Texture(nullptr)
 	, MaterialRenderProxy(InMaterialRenderProxy)
 	, Rotation(ForceInitToZero)
 	, PivotPoint(FVector2D::ZeroVector)
@@ -485,7 +486,7 @@ void FCanvasTileItem::RenderMaterialTile( class FCanvas* InCanvas, const FVector
 	// get sort element based on the current sort key from top of sort key stack
 	FCanvas::FCanvasSortElement& SortElement = InCanvas->GetSortElement(InCanvas->TopDepthSortKey());
 	// find a batch to use 
-	FCanvasTileRendererItem* RenderBatch = NULL;
+	FCanvasTileRendererItem* RenderBatch = nullptr;
 	// get the current transform entry from top of transform stack
 	const FCanvas::FTransformEntry& TopTransformEntry = InCanvas->GetTransformStack().Top();	
 
@@ -496,7 +497,7 @@ void FCanvasTileItem::RenderMaterialTile( class FCanvas* InCanvas, const FVector
 		RenderBatch = SortElement.RenderBatchArray.Last()->GetCanvasTileRendererItem();
 	}	
 	// if a matching entry for this batch doesn't exist then allocate a new entry
-	if( RenderBatch == NULL ||		
+	if( RenderBatch == nullptr ||		
 		!RenderBatch->IsMatch(MaterialRenderProxy,TopTransformEntry) )
 	{
 		INC_DWORD_STAT(STAT_Canvas_NumBatchesCreated);
@@ -823,7 +824,7 @@ void FCanvasTextItem::Draw( class FCanvas* InCanvas )
 {	
 	SCOPE_CYCLE_COUNTER(STAT_Canvas_TextItemTime);
 
-	if (InCanvas == NULL || Font == NULL || Text.IsEmpty())
+	if (InCanvas == nullptr || Font == nullptr || Text.IsEmpty())
 	{
 		return;
 	}
@@ -917,7 +918,7 @@ void FCanvasTextItem::Draw( class FCanvas* InCanvas )
 	}
 	
 	FLinearColor DrawColor;
-	BatchedElements = NULL;
+	BatchedElements = nullptr;
 	if( bOutlined )
 	{
 		DrawColor = OutlineColor;
@@ -976,8 +977,8 @@ void FCanvasTextItem::DrawStringInternal_OfflineCache( FCanvas* InCanvas, const
 
 	FVector2D CurrentPos = FVector2D(EForceInit::ForceInitToZero);
 	FHitProxyId HitProxyId = InCanvas->GetHitProxyId();
-	FTexture* LastTexture = NULL;
-	UTexture* Tex = NULL;
+	FTexture* LastTexture = nullptr;
+	UTexture* Tex = nullptr;
 	FVector2D InvTextureSize(1.0f,1.0f);
 
 	const float CharIncrement = ( (float)Font->Kerning + HorizSpacingAdjust ) * Scale.X;
@@ -1016,14 +1017,14 @@ void FCanvasTextItem::DrawStringInternal_OfflineCache( FCanvas* InCanvas, const
 		}
 
 		if( Font->Textures.IsValidIndex(Char.TextureIndex) && 
-			(Tex=Font->Textures[Char.TextureIndex])!=NULL && 
-			Tex->Resource != NULL )
+			(Tex=Font->Textures[Char.TextureIndex])!=nullptr && 
+			Tex->Resource != nullptr )
 		{
-			if( LastTexture != Tex->Resource || BatchedElements == NULL )
+			if( LastTexture != Tex->Resource || BatchedElements == nullptr )
 			{
-				FBatchedElementParameters* BatchedElementParameters = NULL;
+				FBatchedElementParameters* BatchedElementParameters = nullptr;
 				BatchedElements = InCanvas->GetBatchedElements(FCanvas::ET_Triangle, BatchedElementParameters, Tex->Resource, BlendMode, FontRenderInfo.GlowInfo);
-				check(BatchedElements != NULL);
+				check(BatchedElements != nullptr);
 				// trade-off between memory and performance by pre-allocating more reserved space 
 				// for the triangles/vertices of the batched elements used to render the text tiles
 				//BatchedElements->AddReserveTriangles(TextLen*2,Tex->Resource,BlendMode);
@@ -1378,7 +1379,7 @@ void FCanvasNGonItem::SetColor( const FLinearColor& InColor )
 
 FCanvasNGonItem::FCanvasNGonItem( const FVector2D& InPosition, const FVector2D& InRadius, int32 InNumSides, const FLinearColor& InColor )
 	: FCanvasItem( InPosition )
-	, TriListItem( NULL )
+	, TriListItem( nullptr )
 	, Texture( GWhiteTexture )
 {
 	Color = InColor;
diff --git a/Engine/Source/Runtime/Engine/Public/CanvasItem.h b/Engine/Source/Runtime/Engine/Public/CanvasItem.h
--- a/Engine/Source/Runtime/Engine/Public/CanvasItem.h
+++ b/Engine/Source/Runtime/Engine/Public/CanvasItem.h
@@ -29,7 +29,7 @@ public:
 		: Position( InPosition )
 		, BlendMode( SE_BLEND_Opaque )
 		, bFreezeTime( false )
-		, BatchedElementParameters( NULL )
+		, BatchedElementParameters( nullptr )
 		, Color( FLinearColor::White ) {};
 
 	virtual void Draw( FCanvas* InCanvas ) = 0;
@@ -762,7 +762,7 @@ public:
 	 */
 	 FCanvasNGonItem( const FVector2D& InPosition, const FVector2D& InRadius, int32 InNumSides, const FTexture* InTexture, const FLinearColor& InColor )
 	  : FCanvasItem( InPosition )
-		, TriListItem( NULL )
+		, TriListItem( nullptr )
 		, Texture( InTexture )
 	 {
 		 Color = InColor;
