added option to skip navigable geometry export from primitive components that are releveant for navigation (modifier only)
#ue4 UE-7608

diff --git a/Engine/Source/Runtime/Engine/Classes/Components/PrimitiveComponent.h b/Engine/Source/Runtime/Engine/Classes/Components/PrimitiveComponent.h
--- a/Engine/Source/Runtime/Engine/Classes/Components/PrimitiveComponent.h
+++ b/Engine/Source/Runtime/Engine/Classes/Components/PrimitiveComponent.h
@@ -55,6 +55,9 @@ namespace EHasCustomNavigableGeometry
 
 		// DoCustomNavigableGeometryExport() should be called even if the mesh is non-collidable and wouldn't normally affect the navmesh
 		EvenIfNotCollidable,
+
+		// Don't export navigable geometry even if primitive is relevant for navigation (can still add modifiers)
+		DontExport,
 	};
 }
 
diff --git a/Engine/Source/Runtime/Engine/Private/AI/Navigation/RecastNavMeshGenerator.cpp b/Engine/Source/Runtime/Engine/Private/AI/Navigation/RecastNavMeshGenerator.cpp
--- a/Engine/Source/Runtime/Engine/Private/AI/Navigation/RecastNavMeshGenerator.cpp
+++ b/Engine/Source/Runtime/Engine/Private/AI/Navigation/RecastNavMeshGenerator.cpp
@@ -943,7 +943,7 @@ FORCEINLINE_DEBUGGABLE void ExportComponent(UActorComponent* Component, FRecastG
 	bool bHasData = false;
 
 	UPrimitiveComponent* PrimComp = Cast<UPrimitiveComponent>(Component);
-	if (PrimComp && PrimComp->IsNavigationRelevant())
+	if (PrimComp && PrimComp->IsNavigationRelevant() && (PrimComp->HasCustomNavigableGeometry() != EHasCustomNavigableGeometry::DontExport))
 	{
 		if (PrimComp->HasCustomNavigableGeometry() && !PrimComp->DoCustomNavigableGeometryExport(GeomExport)) 
 		{
