SF_Max disambiguities between Shader Frequency and Sleep Family.

diff --git a/Engine/Source/Runtime/Engine/Classes/Engine/EngineTypes.h b/Engine/Source/Runtime/Engine/Classes/Engine/EngineTypes.h
--- a/Engine/Source/Runtime/Engine/Classes/Engine/EngineTypes.h
+++ b/Engine/Source/Runtime/Engine/Classes/Engine/EngineTypes.h
@@ -792,13 +792,12 @@ enum ERadialImpulseFalloff
 
 /** Presets of values used in considering when put this body to sleep. */
 UENUM()
-enum ESleepFamily
+enum class ESleepFamily : uint8
 {
 	/** Engine defaults. */
-	SF_Normal,
+	Normal,
 	/** A family of values with a lower sleep threshold; good for slower pendulum-like physics. */
-	SF_Sensitive,
-	SF_MAX,
+	Sensitive,
 };
 
 
diff --git a/Engine/Source/Runtime/Engine/Classes/PhysicsEngine/BodyInstance.h b/Engine/Source/Runtime/Engine/Classes/PhysicsEngine/BodyInstance.h
--- a/Engine/Source/Runtime/Engine/Classes/PhysicsEngine/BodyInstance.h
+++ b/Engine/Source/Runtime/Engine/Classes/PhysicsEngine/BodyInstance.h
@@ -339,7 +339,7 @@ protected:
 public:
 	/** The set of values used in considering when put this body to sleep. */
 	UPROPERTY(EditAnywhere, AdvancedDisplay, BlueprintReadWrite, Category=Physics)
-	TEnumAsByte<enum ESleepFamily> SleepFamily;
+	ESleepFamily SleepFamily;
 
 	/**	Influence of rigid body physics (blending) on the mesh's pose (0.0 == use only animation, 1.0 == use only physics) */
 	/** Provide appropriate interface for doing this instead of allowing BlueprintReadWrite **/
diff --git a/Engine/Source/Runtime/Engine/Private/DestructibleComponent.cpp b/Engine/Source/Runtime/Engine/Private/DestructibleComponent.cpp
--- a/Engine/Source/Runtime/Engine/Private/DestructibleComponent.cpp
+++ b/Engine/Source/Runtime/Engine/Private/DestructibleComponent.cpp
@@ -325,7 +325,7 @@ void UDestructibleComponent::CreatePhysicsState()
 	verify( NxParameterized::setParamF32(*ActorParams,"p3BodyDescTemplate.linearDamping", BodyInstance.LinearDamping ) );
 	const PxTolerancesScale& PScale = GPhysXSDK->getTolerancesScale();
 	PxF32 SleepEnergyThreshold = 0.00005f*PScale.speed*PScale.speed;	// 1/1000 Default, since the speed scale is quite high
-	if (BodyInstance.SleepFamily == SF_Sensitive)
+	if (BodyInstance.SleepFamily == ESleepFamily::Sensitive)
 	{
 		SleepEnergyThreshold /= 20.0f;
 	}
diff --git a/Engine/Source/Runtime/Engine/Private/PhysicsEngine/BodyInstance.cpp b/Engine/Source/Runtime/Engine/Private/PhysicsEngine/BodyInstance.cpp
--- a/Engine/Source/Runtime/Engine/Private/PhysicsEngine/BodyInstance.cpp
+++ b/Engine/Source/Runtime/Engine/Private/PhysicsEngine/BodyInstance.cpp
@@ -249,7 +249,7 @@ FBodyInstance::FBodyInstance()
 , bOverrideMaxDepenetrationVelocity(false)
 , MaxDepenetrationVelocity(0.f)
 , PhysMaterialOverride(NULL)
-, SleepFamily(SF_Normal)
+, SleepFamily(ESleepFamily::Normal)
 , PhysicsBlendWeight(0.f)
 , PositionSolverIterationCount(8)
 , VelocitySolverIterationCount(1)
@@ -1219,7 +1219,7 @@ void FBodyInstance::InitBodies(TArray<FBodyInstance*>& Bodies, TArray<FTransform
 #if 0
 					// Set the parameters for determining when to put the object to sleep.
 					float SleepEnergyThresh = PNewDynamic->getSleepThreshold();
-					if(SleepFamily == SF_Sensitive)
+					if (SleepFamily == ESleepFamily::Sensitive)
 					{
 						SleepEnergyThresh /= 20.f;
 					}
@@ -4086,7 +4086,7 @@ void FBodyInstance::InitDynamicProperties()
 		}
 
 		float SleepEnergyThresh = RigidActor->getSleepThreshold();
-		if(SleepFamily == SF_Sensitive)
+		if (SleepFamily == ESleepFamily::Sensitive)
 		{
 			SleepEnergyThresh /= 20.f;
 		}
diff --git a/Engine/Source/Runtime/Windows/D3D11RHI/Private/D3D11StateCache.cpp b/Engine/Source/Runtime/Windows/D3D11RHI/Private/D3D11StateCache.cpp
--- a/Engine/Source/Runtime/Windows/D3D11RHI/Private/D3D11StateCache.cpp
+++ b/Engine/Source/Runtime/Windows/D3D11RHI/Private/D3D11StateCache.cpp
@@ -286,7 +286,7 @@ void FD3D11StateCacheBase::ClearState()
 
 #if D3D11_ALLOW_STATE_CACHE
 	// Shader Resource View State Cache
-	for (uint32 ShaderFrequency = 0; ShaderFrequency < SF_MAX; ShaderFrequency++ )
+	for (uint32 ShaderFrequency = 0; ShaderFrequency < SF_NumFrequencies; ShaderFrequency++)
 	{
 		for (uint32 Index = 0; Index < D3D11_COMMONSHADER_INPUT_RESOURCE_SLOT_COUNT; Index++)
 		{
