Made navigation generators thread-safe shared pointers #UE4

diff --git a/Engine/Source/Runtime/Engine/Classes/AI/Navigation/NavigationData.h b/Engine/Source/Runtime/Engine/Classes/AI/Navigation/NavigationData.h
--- a/Engine/Source/Runtime/Engine/Classes/AI/Navigation/NavigationData.h
+++ b/Engine/Source/Runtime/Engine/Classes/AI/Navigation/NavigationData.h
@@ -804,7 +804,7 @@ protected:
 	FNavRaycastPtr RaycastImplementation; 
 
 protected:
-	TSharedPtr<FNavDataGenerator> NavDataGenerator;
+	TSharedPtr<FNavDataGenerator, ESPMode::ThreadSafe> NavDataGenerator;
 	/** 
 	 *	Container for all path objects generated with this Navigation Data instance. 
 	 *	Is meant to be added to only on GameThread, and in fact should user should never 
diff --git a/Engine/Source/Runtime/Engine/Private/AI/Navigation/RecastNavMeshGenerator.cpp b/Engine/Source/Runtime/Engine/Private/AI/Navigation/RecastNavMeshGenerator.cpp
--- a/Engine/Source/Runtime/Engine/Private/AI/Navigation/RecastNavMeshGenerator.cpp
+++ b/Engine/Source/Runtime/Engine/Private/AI/Navigation/RecastNavMeshGenerator.cpp
@@ -1638,7 +1638,7 @@ bool FRecastTileGenerator::HasDataToBuild() const
 
 void FRecastTileGenerator::DoWork()
 {
-	TSharedPtr<FNavDataGenerator> ParentGenerator = ParentGeneratorWeakPtr.Pin();
+	TSharedPtr<FNavDataGenerator, ESPMode::ThreadSafe> ParentGenerator = ParentGeneratorWeakPtr.Pin();
 	FRecastNavMeshGenerator* RecastParentGenerator = static_cast<FRecastNavMeshGenerator*>(ParentGenerator.Get());
 	check(ParentGenerator.Get());
 
diff --git a/Engine/Source/Runtime/Engine/Public/AI/NavDataGenerator.h b/Engine/Source/Runtime/Engine/Public/AI/NavDataGenerator.h
--- a/Engine/Source/Runtime/Engine/Public/AI/NavDataGenerator.h
+++ b/Engine/Source/Runtime/Engine/Public/AI/NavDataGenerator.h
@@ -6,7 +6,7 @@
 
 struct FNavigationDirtyArea;
 
-class ENGINE_API FNavDataGenerator : public TSharedFromThis<FNavDataGenerator>
+class ENGINE_API FNavDataGenerator : public TSharedFromThis<FNavDataGenerator, ESPMode::ThreadSafe>
 {
 public:
 	virtual ~FNavDataGenerator() {}
diff --git a/Engine/Source/Runtime/Engine/Public/AI/Navigation/RecastNavMeshGenerator.h b/Engine/Source/Runtime/Engine/Public/AI/Navigation/RecastNavMeshGenerator.h
--- a/Engine/Source/Runtime/Engine/Public/AI/Navigation/RecastNavMeshGenerator.h
+++ b/Engine/Source/Runtime/Engine/Public/AI/Navigation/RecastNavMeshGenerator.h
@@ -236,7 +236,7 @@ protected:
 	// navigation links
 	TArray<FSimpleLinkNavModifier> OffmeshLinks;
 
-	TWeakPtr<FNavDataGenerator> ParentGeneratorWeakPtr;
+	TWeakPtr<FNavDataGenerator, ESPMode::ThreadSafe> ParentGeneratorWeakPtr;
 
 	TNavStatArray<TSharedRef<FNavigationRelevantData, ESPMode::ThreadSafe> > NavigationRelevantData;
 	TSharedPtr<FNavigationOctree, ESPMode::ThreadSafe> NavOctree; 
