Thread safety for multiple calls to RefreshBoneTransforms

diff --git a/Engine/Source/Runtime/Engine/Classes/Components/SkeletalMeshComponent.h b/Engine/Source/Runtime/Engine/Classes/Components/SkeletalMeshComponent.h
--- a/Engine/Source/Runtime/Engine/Classes/Components/SkeletalMeshComponent.h
+++ b/Engine/Source/Runtime/Engine/Classes/Components/SkeletalMeshComponent.h
@@ -1174,11 +1174,16 @@ private:
 	//Data for parallel evaluation of animation
 	FAnimationEvaluationContext AnimEvaluationContext;
 
+	// Reference to our current parallel animation evaluation task (if there is one)
+	FGraphEventRef				ParallelAnimationEvaluationTask;
+
 public:
 	// Parallel evaluation wrappers
 	void ParallelAnimationEvaluation() { PerformAnimationEvaluation(AnimEvaluationContext.SkeletalMesh, AnimEvaluationContext.AnimInstance, AnimEvaluationContext.SpaceBases, AnimEvaluationContext.LocalAtoms, AnimEvaluationContext.VertexAnims, AnimEvaluationContext.RootBoneTranslation); }
 	void CompleteParallelAnimationEvaluation()
 	{
+		ParallelAnimationEvaluationTask.SafeRelease(); //We are done with this task now, clean up!
+
 		if ((AnimEvaluationContext.AnimInstance == AnimScriptInstance) && (AnimEvaluationContext.SkeletalMesh == SkeletalMesh) && (AnimEvaluationContext.SpaceBases.Num() == GetNumSpaceBases()))
 		{
 			Exchange(AnimEvaluationContext.SpaceBases, AnimEvaluationContext.bDoInterpolation ? CachedSpaceBases : GetEditableSpaceBases() );
@@ -1188,10 +1193,7 @@ public:
 
 			PostAnimEvaluation(AnimEvaluationContext);
 		}
-		else
-		{
-			AnimEvaluationContext.Clear();
-		}
+		AnimEvaluationContext.Clear();
 	}
 
 	friend class FSkeletalMeshComponentDetails;
diff --git a/Engine/Source/Runtime/Engine/Private/SkeletalMeshComponent.cpp b/Engine/Source/Runtime/Engine/Private/SkeletalMeshComponent.cpp
--- a/Engine/Source/Runtime/Engine/Private/SkeletalMeshComponent.cpp
+++ b/Engine/Source/Runtime/Engine/Private/SkeletalMeshComponent.cpp
@@ -993,11 +993,41 @@ void USkeletalMeshComponent::RefreshBoneTransforms(FActorComponentTickFunction*
 	SCOPE_CYCLE_COUNTER(STAT_RefreshBoneTransforms);
 	SCOPE_CYCLE_COUNTER(STAT_AnimGameThreadTime);
 
+	check(IsInGameThread()); //Only want to call this from the game thread as we set up tasks etc
+
 	if (!SkeletalMesh || GetNumSpaceBases() == 0)
 	{
 		return;
 	}
 
+	const bool bDoEvaluationRateOptimization = bEnableUpdateRateOptimizations && AnimUpdateRateParams->DoEvaluationRateOptimizations();
+
+	//Handle update rate optimization setup
+	//Dont mark cache as invalid if we aren't performing optimization anyway
+	const bool bInvalidCachedBones = bDoEvaluationRateOptimization &&
+									 ((LocalAtoms.Num() != SkeletalMesh->RefSkeleton.GetNum())
+									 || (LocalAtoms.Num() != CachedLocalAtoms.Num())
+									 || (GetNumSpaceBases() != CachedSpaceBases.Num()));
+
+	const bool bShouldDoEvaluation = !bDoEvaluationRateOptimization || bInvalidCachedBones || !AnimUpdateRateParams->ShouldSkipEvaluation();
+
+	const bool bDoPAE = !!CVarUseParallelAnimationEvaluation.GetValueOnGameThread() && FApp::ShouldUseThreadingForPerformance();
+
+	const bool bDoParallelEvaluation = AnimEvaluationContext.bDoEvaluation && TickFunction && bDoPAE;
+
+	if (IsValidRef(ParallelAnimationEvaluationTask))
+	{
+		//Are already processing eval on another thread, we wait for eval thread to finish
+		if (!bDoParallelEvaluation) //we are already running parallel evaluation, so if we are going to try it again just return
+		{
+			//If we are not going to attempt parallel evaluation then wait here for the existing task to finish and then
+			//perform complete directly so that when we return all the calculations are complete
+			FTaskGraphInterface::Get().WaitUntilTaskCompletes(ParallelAnimationEvaluationTask, ENamedThreads::GameThread);
+			CompleteParallelAnimationEvaluation(); //Perform completion now
+		}
+		return;
+	}
+
 	AActor* Owner = GetOwner();
 	UE_LOG(LogAnimation, Verbose, TEXT("RefreshBoneTransforms(%s)"), *GetNameSafe(Owner));
 
@@ -1010,15 +1040,7 @@ void USkeletalMeshComponent::RefreshBoneTransforms(FActorComponentTickFunction*
 	AnimEvaluationContext.SkeletalMesh = SkeletalMesh;
 	AnimEvaluationContext.AnimInstance = AnimScriptInstance;
 
-	//Handle update rate optimization setup
-	const bool bDoEvaluationRateOptimization = bEnableUpdateRateOptimizations && AnimUpdateRateParams->DoEvaluationRateOptimizations();
-	//Dont mark cache as invalid if we aren't performing optimization anyway
-	const bool bInvalidCachedBones = bDoEvaluationRateOptimization &&
-									( (LocalAtoms.Num() != SkeletalMesh->RefSkeleton.GetNum())
-									  || (LocalAtoms.Num() != CachedLocalAtoms.Num())
-									  || (GetNumSpaceBases() != CachedSpaceBases.Num()) );
-
-	AnimEvaluationContext.bDoEvaluation = !bDoEvaluationRateOptimization || bInvalidCachedBones || !AnimUpdateRateParams->ShouldSkipEvaluation();
+	AnimEvaluationContext.bDoEvaluation = bShouldDoEvaluation;
 	
 	AnimEvaluationContext.bDoInterpolation = bDoEvaluationRateOptimization && !bInvalidCachedBones && AnimUpdateRateParams->ShouldInterpolateSkippedFrames();
 	AnimEvaluationContext.bDuplicateToCacheBones = bInvalidCachedBones || (bDoEvaluationRateOptimization && AnimEvaluationContext.bDoEvaluation && !AnimEvaluationContext.bDoInterpolation);
@@ -1030,8 +1052,7 @@ void USkeletalMeshComponent::RefreshBoneTransforms(FActorComponentTickFunction*
 		CachedSpaceBases.Empty();
 	}
 
-	const bool bDoPAE = !!CVarUseParallelAnimationEvaluation.GetValueOnGameThread() && FApp::ShouldUseThreadingForPerformance();
-	if (AnimEvaluationContext.bDoEvaluation && TickFunction && bDoPAE)
+	if (bDoParallelEvaluation)
 	{
 		if (SkeletalMesh->RefSkeleton.GetNum() != AnimEvaluationContext.LocalAtoms.Num())
 		{
@@ -1042,11 +1063,12 @@ void USkeletalMeshComponent::RefreshBoneTransforms(FActorComponentTickFunction*
 		}
 
 		// start parallel work
-		FGraphEventRef EvaluationTickEvent = TGraphTask<FParallelAnimationEvaluationTask>::CreateTask().ConstructAndDispatchWhenReady(this);
+		check(!IsValidRef(ParallelAnimationEvaluationTask));
+		ParallelAnimationEvaluationTask = TGraphTask<FParallelAnimationEvaluationTask>::CreateTask().ConstructAndDispatchWhenReady(this);
 
 		// set up a task to run on the game thread to accept the results
 		FGraphEventArray Prerequistes;
-		Prerequistes.Add(EvaluationTickEvent);
+		Prerequistes.Add(ParallelAnimationEvaluationTask);
 		FGraphEventRef TickCompletionEvent = TGraphTask<FParallelAnimationCompletionTask>::CreateTask(&Prerequistes).ConstructAndDispatchWhenReady(this);
 
 		TickFunction->GetCompletionHandle()->DontCompleteUntil(TickCompletionEvent);
@@ -1072,7 +1094,6 @@ void USkeletalMeshComponent::RefreshBoneTransforms(FActorComponentTickFunction*
 
 		PostAnimEvaluation(AnimEvaluationContext);
 	}
-
 }
 
 void USkeletalMeshComponent::PostAnimEvaluation(FAnimationEvaluationContext& EvaluationContext)
