Fix crash in USkeletalMeshComponent::FillSpaceBases due to -1 in requiredbones

diff --git a/Engine/Source/Runtime/Engine/Private/SkeletalMeshComponent.cpp b/Engine/Source/Runtime/Engine/Private/SkeletalMeshComponent.cpp
--- a/Engine/Source/Runtime/Engine/Private/SkeletalMeshComponent.cpp
+++ b/Engine/Source/Runtime/Engine/Private/SkeletalMeshComponent.cpp
@@ -840,13 +840,16 @@ void USkeletalMeshComponent::RecalcRequiredBones(int32 LODIndex)
 		for (const USkeletalMeshSocket* Socket : SkeletalMesh->GetActiveSocketList())
 		{
 			int32 BoneIndex = SkeletalMesh->RefSkeleton.FindBoneIndex(Socket->BoneName);
-			if (Socket->bForceAlwaysAnimated)
+			if (BoneIndex != INDEX_NONE)
 			{
-				ForceAnimatedSocketBones.AddUnique(BoneIndex);
-			}
-			else
-			{
-				NeededBonesForFillSpaceBases.AddUnique(BoneIndex);
+				if (Socket->bForceAlwaysAnimated)
+				{
+					ForceAnimatedSocketBones.AddUnique(BoneIndex);
+				}
+				else
+				{
+					NeededBonesForFillSpaceBases.AddUnique(BoneIndex);
+				}
 			}
 		}
 
