PR #870: BUG FIX : Add missing parameter (Contributed by [REDACTED_PERSON])

https://github.com/EpicGames/UnrealEngine/pull/870

diff --git a/Engine/Source/Runtime/Engine/Private/Character.cpp b/Engine/Source/Runtime/Engine/Private/Character.cpp
--- a/Engine/Source/Runtime/Engine/Private/Character.cpp
+++ b/Engine/Source/Runtime/Engine/Private/Character.cpp
@@ -1298,7 +1298,7 @@ void ACharacter::StopAnimMontage(class UAnimMontage* AnimMontage)
 
 	if ( bShouldStopMontage )
 	{
-		AnimInstance->Montage_Stop(MontageToStop->BlendOutTime);
+		AnimInstance->Montage_Stop(MontageToStop->BlendOutTime, MontageToStop);
 	}
 }
 
