UE4 - fix for missing instanced static meshes

diff --git a/Engine/Source/Runtime/Engine/Private/InstancedStaticMesh.cpp b/Engine/Source/Runtime/Engine/Private/InstancedStaticMesh.cpp
--- a/Engine/Source/Runtime/Engine/Private/InstancedStaticMesh.cpp
+++ b/Engine/Source/Runtime/Engine/Private/InstancedStaticMesh.cpp
@@ -30,6 +30,7 @@ void HInstancedStaticMeshInstance::AddReferencedObjects(FReferenceCollector& Col
 FStaticMeshInstanceBuffer::FStaticMeshInstanceBuffer(ERHIFeatureLevel::Type InFeatureLevel):
 	InstanceData(NULL)
 	, Stride(0)
+	, NumInstances(0)
 {
 	SetFeatureLevel(InFeatureLevel);
 }
@@ -47,6 +48,7 @@ void FStaticMeshInstanceBuffer::CleanUp()
 		delete InstanceData;
 		InstanceData = NULL;
 	}
+	NumInstances = 0;
 }
 
 
@@ -68,7 +70,7 @@ void FStaticMeshInstanceBuffer::Init(UInstancedStaticMeshComponent* InComponent,
 
 	bool bUseRemapTable = InComponent->PerInstanceSMData.Num() == InComponent->InstanceReorderTable.Num();
 
-	int32 NumInstances = InComponent->PerInstanceSMData.Num();
+	int32 NumRealInstances = InComponent->PerInstanceSMData.Num();
 	int32 NumRemoved = InComponent->RemovedInstances.Num();
 
 	// Allocate the vertex data storage type.
@@ -76,7 +78,8 @@ void FStaticMeshInstanceBuffer::Init(UInstancedStaticMeshComponent* InComponent,
 
 	SetupCPUAccess(InComponent);
 
-	InstanceData->AllocateInstances(NumInstances + NumRemoved);
+	NumInstances = NumRealInstances + NumRemoved;
+	InstanceData->AllocateInstances(NumInstances);
 
 	// Setup our random number generator such that random values are generated consistently for any
 	// given instance index between reattaches
@@ -84,10 +87,10 @@ void FStaticMeshInstanceBuffer::Init(UInstancedStaticMeshComponent* InComponent,
 
 	FColor HitProxyColor(ForceInit);
 	bool bSelected = false;
-	for (int32 InstanceIndex = 0; InstanceIndex < NumInstances; InstanceIndex++)
+	for (int32 InstanceIndex = 0; InstanceIndex < NumRealInstances; InstanceIndex++)
 	{
 #if WITH_EDITOR
-		if (InHitProxies.Num() == NumInstances)
+		if (InHitProxies.Num() == NumRealInstances)
 		{
 			HitProxyColor = InHitProxies[InstanceIndex]->Id.GetColor();
 		}
@@ -115,7 +118,7 @@ void FStaticMeshInstanceBuffer::Init(UInstancedStaticMeshComponent* InComponent,
 	float ThisTime = (StartTime - FPlatformTime::Seconds()) * 1000.0f;
 	if (ThisTime > 30.0f)
 	{
-		UE_LOG(LogStaticMesh, Display, TEXT("Took %6.2fms to set up instance buffer for %d instances for component %s."), ThisTime, NumInstances, *InComponent->GetFullName());
+		UE_LOG(LogStaticMesh, Display, TEXT("Took %6.2fms to set up instance buffer for %d instances for component %s."), ThisTime, NumRealInstances, *InComponent->GetFullName());
 	}
 }
 
diff --git a/Engine/Source/Runtime/Engine/Private/InstancedStaticMesh.h b/Engine/Source/Runtime/Engine/Private/InstancedStaticMesh.h
--- a/Engine/Source/Runtime/Engine/Private/InstancedStaticMesh.h
+++ b/Engine/Source/Runtime/Engine/Private/InstancedStaticMesh.h
@@ -109,7 +109,7 @@ public:
 	}
 	FORCEINLINE uint32 GetNumInstances() const
 	{
-		return InstanceData ? InstanceData->Num() : 0;
+		return NumInstances;
 	}
 
 	const void* GetRawData() const
@@ -140,6 +140,9 @@ private:
 	/** The cached vertex stride. */
 	uint32 Stride;
 
+	/** The cached number of instances. */
+	uint32 NumInstances;
+
 	/** Allocates the vertex data storage type. */
 	void AllocateData();
 
