Engine Analytics changes to detect games vs. editor better. #UETOOL-249

diff --git a/Engine/Source/Runtime/Engine/Private/EngineAnalytics.cpp b/Engine/Source/Runtime/Engine/Private/EngineAnalytics.cpp
--- a/Engine/Source/Runtime/Engine/Private/EngineAnalytics.cpp
+++ b/Engine/Source/Runtime/Engine/Private/EngineAnalytics.cpp
@@ -5,6 +5,7 @@
 #include "EngineBuildSettings.h"
 #include "Runtime/Analytics/Analytics/Public/Analytics.h"
 #include "Runtime/Analytics/Analytics/Public/Interfaces/IAnalyticsProvider.h"
+#include "GeneralProjectSettings.h"
 
 bool FEngineAnalytics::bIsInitialized;
 TSharedPtr<IAnalyticsProvider> FEngineAnalytics::Analytics;
@@ -36,29 +37,34 @@ void FEngineAnalytics::Initialize()
 {
 	checkf(!bIsInitialized, TEXT("FEngineAnalytics::Initialize called more than once."));
 
-	// Never use analytics when running a commandlet tool
-	const bool bShouldInitAnalytics = !IsRunningCommandlet();
-
 	check(GEngine);
-#if WITH_EDITORONLY_DATA
-	if (GIsEditor)
-	{
-		bShouldSendUsageEvents = GEngine->AreEditorAnalyticsEnabled();
-	}
-	else
-#endif
-	{
-		// Outside of the editor, the only engine analytics usage is the hardware survey
-		bShouldSendUsageEvents = GEngine->bHardwareSurveyEnabled;
-	}
+
+	// this will only be true for builds that have editor support (currently PC, Mac, Linux)
+	// The idea here is to only send editor events for actual editor runs, not for things like -game runs of the editor.
+	const bool bIsEditorRun = WITH_EDITOR && GIsEditor && !IsRunningCommandlet();
+
+	// We also want to identify a real run of a game, which is NOT necessarily the opposite of an editor run.
+	// Ideally we'd be able to tell explicitly, but with content-only games, it becomes difficult.
+	// So we ensure we are not an editor run, we don't have EDITOR stuff compiled in, we are not running a commandlet,
+	// we are not a generic, utility program, and we require cooked data.
+	const bool bIsGameRun = !WITH_EDITOR && !IsRunningCommandlet() && !FPlatformProperties::IsProgram() && FPlatformProperties::RequiresCookedData();
+
+	const bool bShouldInitAnalytics = bIsEditorRun || bIsGameRun;
+
+	// Outside of the editor, the only engine analytics usage is the hardware survey
+	bShouldSendUsageEvents = bIsEditorRun 
+		? GEngine->AreEditorAnalyticsEnabled() 
+		: bIsGameRun 
+			? GEngine->bHardwareSurveyEnabled 
+			: false;
 
 	if (bShouldInitAnalytics)
 	{
 		{
 			// Setup some default engine analytics if there is nothing custom bound
 			FAnalytics::FProviderConfigurationDelegate DefaultEngineAnalyticsConfig;
-			DefaultEngineAnalyticsConfig.BindStatic( 
-				[]( const FString& KeyName, bool bIsValueRequired ) -> FString
+			DefaultEngineAnalyticsConfig.BindLambda( 
+				[=]( const FString& KeyName, bool bIsValueRequired ) -> FString
 				{
 					static TMap<FString, FString> ConfigMap;
 					if (ConfigMap.Num() == 0)
@@ -68,28 +74,20 @@ void FEngineAnalytics::Initialize()
 
 						// We always use the "Release" analytics account unless we're running in analytics test mode (usually with
 						// a command-line parameter), or we're an internal Epic build
-						bool bUseReleaseAccount =
-							(FAnalytics::Get().GetBuildType() == FAnalytics::Development ||  
-							 FAnalytics::Get().GetBuildType() == FAnalytics::Release) &&
+						const FAnalytics::BuildType AnalyticsBuildType = FAnalytics::Get().GetBuildType();
+						const bool bUseReleaseAccount =
+							(AnalyticsBuildType == FAnalytics::Development || AnalyticsBuildType == FAnalytics::Release) &&
 							!FEngineBuildSettings::IsInternalBuild();	// Internal Epic build
-
-						if( FRocketSupport::IsRocket() )
-						{
-							const TCHAR* DevelopmentAccountAPIKeyET = TEXT("Rocket.Dev");
-							const TCHAR* ReleaseAccountAPIKeyET = TEXT("Rocket.Release");
-							ConfigMap.Add(TEXT("APIKeyET"), bUseReleaseAccount ? ReleaseAccountAPIKeyET : DevelopmentAccountAPIKeyET);
-						}
-						else if( FEngineBuildSettings::IsPerforceBuild() )
+						const TCHAR* BuildTypeStr = bUseReleaseAccount ? TEXT("Dev") : TEXT("Release");
+						const TCHAR* UE4TypeStr = FRocketSupport::IsRocket() ? TEXT("Rocket") : FEngineBuildSettings::IsPerforceBuild() ? TEXT("Perforce") : TEXT("UnrealEngine");
+						if (bIsEditorRun)
 						{
-							const TCHAR* DevelopmentAccountAPIKeyET = TEXT("Perforce.Dev");
-							const TCHAR* ReleaseAccountAPIKeyET = TEXT("Perforce.Release");
-							ConfigMap.Add(TEXT("APIKeyET"), bUseReleaseAccount ? ReleaseAccountAPIKeyET : DevelopmentAccountAPIKeyET);
+							ConfigMap.Add(TEXT("APIKeyET"), FString::Printf(TEXT("UEEditor.%s.%s"), UE4TypeStr, BuildTypeStr));
 						}
 						else
 						{
-							const TCHAR* DevelopmentAccountAPIKeyET = TEXT("UnrealEngine.Dev");
-							const TCHAR* ReleaseAccountAPIKeyET = TEXT("UnrealEngine.Release");
-							ConfigMap.Add(TEXT("APIKeyET"), bUseReleaseAccount ? ReleaseAccountAPIKeyET : DevelopmentAccountAPIKeyET);
+							const UGeneralProjectSettings& ProjectSettings = *GetDefault<UGeneralProjectSettings>();
+							ConfigMap.Add(TEXT("APIKeyET"), FString::Printf(TEXT("UEGame.%s.%s|%s|%s"), UE4TypeStr, BuildTypeStr, *ProjectSettings.ProjectID.ToString(), *ProjectSettings.ProjectName));
 						}
 					}
 
@@ -117,6 +115,12 @@ void FEngineAnalytics::Initialize()
 
 				TArray<FAnalyticsEventAttribute> StartSessionAttributes;
 				GEngine->CreateStartupAnalyticsAttributes( StartSessionAttributes );
+				// Add project info whether we are in editor or game.
+				const UGeneralProjectSettings& ProjectSettings = *GetDefault<UGeneralProjectSettings>();
+				StartSessionAttributes.Emplace(TEXT("ProjectName"), ProjectSettings.ProjectName);
+				StartSessionAttributes.Emplace(TEXT("ProjectID"), ProjectSettings.ProjectID);
+				StartSessionAttributes.Emplace(TEXT("ProjectDescription"), ProjectSettings.Description);
+				StartSessionAttributes.Emplace(TEXT("ProjectVersion"), ProjectSettings.ProjectVersion);
 
 				Analytics->StartSession( StartSessionAttributes );
 			}
