#ue4 - TryToLeaveNavWalking() should be called even if DefaultLandMovementMode was NavWalking and we are trying to enter normal Walking mode.

diff --git a/Engine/Source/Runtime/Engine/Classes/GameFramework/CharacterMovementComponent.h b/Engine/Source/Runtime/Engine/Classes/GameFramework/CharacterMovementComponent.h
--- a/Engine/Source/Runtime/Engine/Classes/GameFramework/CharacterMovementComponent.h
+++ b/Engine/Source/Runtime/Engine/Classes/GameFramework/CharacterMovementComponent.h
@@ -1135,7 +1135,7 @@ protected:
 	virtual void SetNavWalkingPhysics(bool bEnable);
 
 	/** 
-	 * Tries to switch to DefaultLandMovementMode
+	 * Checks to see if the current location is not encroaching blocking geometry so the character can leave NavWalking.
 	 * Restores collision settings and adjusts character location to avoid getting stuck in geometry.
 	 * If it's not possible, MovementMode change will be delayed until character reach collision free spot.
 	 * @return True if movement mode was successfully changed
diff --git a/Engine/Source/Runtime/Engine/Private/CharacterMovementComponent.cpp b/Engine/Source/Runtime/Engine/Private/CharacterMovementComponent.cpp
--- a/Engine/Source/Runtime/Engine/Private/CharacterMovementComponent.cpp
+++ b/Engine/Source/Runtime/Engine/Private/CharacterMovementComponent.cpp
@@ -682,7 +682,7 @@ void UCharacterMovementComponent::OnMovementModeChanged(EMovementMode PreviousMo
 	}
 	else if (PreviousMovementMode == MOVE_NavWalking)
 	{
-		if (MovementMode == DefaultLandMovementMode)
+		if (MovementMode == DefaultLandMovementMode || IsWalking())
 		{
 			const bool bCanSwitchMode = TryToLeaveNavWalking();
 			if (!bCanSwitchMode)
