#ue4 - Remove component from OverlappingComponents array before firing OnComponentEndOverlap events.

Jira: UE-7000

diff --git a/Engine/Source/Runtime/Engine/Private/PrimitiveComponent.cpp b/Engine/Source/Runtime/Engine/Private/PrimitiveComponent.cpp
--- a/Engine/Source/Runtime/Engine/Private/PrimitiveComponent.cpp
+++ b/Engine/Source/Runtime/Engine/Private/PrimitiveComponent.cpp
@@ -2016,6 +2016,12 @@ void UPrimitiveComponent::EndComponentOverlap(const FOverlapInfo& OtherOverlap,
 
 	//	UE_LOG(LogActor, Log, TEXT("END OVERLAP! Self=%s SelfComp=%s, Other=%s, OtherComp=%s"), *GetNameSafe(this), *GetNameSafe(MyComp), *GetNameSafe(OtherActor), *GetNameSafe(OtherComp));
 
+	const int32 OtherOverlapIdx = OtherComp ? OtherComp->OverlappingComponents.Find(FOverlapInfo(this, INDEX_NONE)) : INDEX_NONE;
+	if (OtherOverlapIdx != INDEX_NONE)
+	{
+		OtherComp->OverlappingComponents.RemoveAtSwap(OtherOverlapIdx, 1, false);
+	}
+
 	const int32 OverlapIdx = OverlappingComponents.Find(OtherOverlap);
 	if (OverlapIdx != INDEX_NONE)
 	{
@@ -2035,28 +2041,19 @@ void UPrimitiveComponent::EndComponentOverlap(const FOverlapInfo& OtherOverlap,
 		}
 	}
 
-	if (OtherComp)
-	{
-		const int32 OtherOverlapIdx = OtherComp->OverlappingComponents.Find(FOverlapInfo(this, INDEX_NONE));
-		if (OtherOverlapIdx != INDEX_NONE)
+	// if this was the last touch on the other actor, notify that we've untouched the actor as well
+	if (bDoNotifies && MyActor && OtherActor && !MyActor->IsOverlappingActor(OtherActor) )
+	{			
+		if (IsActorValidToNotify(MyActor))
 		{
-			OtherComp->OverlappingComponents.RemoveAtSwap(OtherOverlapIdx, 1, false);
+			MyActor->ReceiveActorEndOverlap(OtherActor);
+			MyActor->OnActorEndOverlap.Broadcast(OtherActor);
 		}
-		
-		// if this was the last touch on the other actor, notify that we've untouched the actor as well
-		if (bDoNotifies && MyActor && OtherActor && !MyActor->IsOverlappingActor(OtherActor) )
-		{			
-			if (IsActorValidToNotify(MyActor))
-			{
-				MyActor->ReceiveActorEndOverlap(OtherActor);
-				MyActor->OnActorEndOverlap.Broadcast(OtherActor);
-			}
 
-			if (IsActorValidToNotify(OtherActor))
-			{
-				OtherActor->ReceiveActorEndOverlap(MyActor);
-				OtherActor->OnActorEndOverlap.Broadcast(MyActor);
-			}
+		if (IsActorValidToNotify(OtherActor))
+		{
+			OtherActor->ReceiveActorEndOverlap(MyActor);
+			OtherActor->OnActorEndOverlap.Broadcast(MyActor);
 		}
 	}
 }
