Fixing compile errors

diff --git a/Engine/Source/Editor/LevelEditor/Private/SActorDetails.cpp b/Engine/Source/Editor/LevelEditor/Private/SActorDetails.cpp
--- a/Engine/Source/Editor/LevelEditor/Private/SActorDetails.cpp
+++ b/Engine/Source/Editor/LevelEditor/Private/SActorDetails.cpp
@@ -527,7 +527,7 @@ void SActorDetails::OnNativeComponentWarningHyperlinkClicked(const FSlateHyperli
 {
 	// Find the closest native parent
 	UBlueprint* Blueprint = SCSEditor->GetBlueprint();
-	UClass* ParentClass = Blueprint ? Blueprint->ParentClass : GetActorContext()->GetClass();
+	UClass* ParentClass = Blueprint ? *Blueprint->ParentClass : GetActorContext()->GetClass();
 	while (ParentClass && !ParentClass->HasAllClassFlags(CLASS_Native))
 	{
 		ParentClass = ParentClass->GetSuperClass();
diff --git a/Engine/Source/Runtime/Engine/Classes/Engine/Selection.h b/Engine/Source/Runtime/Engine/Classes/Engine/Selection.h
--- a/Engine/Source/Runtime/Engine/Classes/Engine/Selection.h
+++ b/Engine/Source/Runtime/Engine/Classes/Engine/Selection.h
@@ -6,7 +6,11 @@
  */
 
 #pragma once
+
+#if WITH_EDITOR
 #include "ComponentEditorUtils.h"
+#endif
+
 #include "Selection.generated.h"
 
 UCLASS(customConstructor, transient)
@@ -470,7 +474,13 @@ public:
 		const UActorComponent* Comp = Cast<UActorComponent>( Object );
 		if (Comp)
 		{
-			bool bCanEditIfNative = ( Comp->CreationMethod == EComponentCreationMethod::Native ) ? FComponentEditorUtils::CanEditNativeComponent(Comp) : true;
+			bool bCanEditIfNative = true;
+#if WITH_EDITOR
+			if (Comp->CreationMethod == EComponentCreationMethod::Native)
+			{
+				bCanEditIfNative = FComponentEditorUtils::CanEditNativeComponent(Comp);
+			}
+#endif
 			return !Comp->IsCreatedByConstructionScript() && bCanEditIfNative;
 		}
 		
