Fixing non-editor build errors

diff --git a/Engine/Source/Runtime/Engine/Private/ActorEditorUtils.cpp b/Engine/Source/Runtime/Engine/Private/ActorEditorUtils.cpp
--- a/Engine/Source/Runtime/Engine/Private/ActorEditorUtils.cpp
+++ b/Engine/Source/Runtime/Engine/Private/ActorEditorUtils.cpp
@@ -2,7 +2,10 @@
 
 #include "EnginePrivate.h"
 #include "ActorEditorUtils.h"
+
+#if WITH_EDITOR
 #include "ComponentEditorUtils.h"
+#endif
 
 namespace FActorEditorUtils
 {
@@ -40,6 +43,7 @@ namespace FActorEditorUtils
 		InActor->GetComponents(InstanceComponents);
 		for (auto Component : InstanceComponents)
 		{
+#if WITH_EDITOR
 			if (Component->CreationMethod == EComponentCreationMethod::Native)
 			{
 				// Make sure it's an exposed native component
@@ -49,8 +53,10 @@ namespace FActorEditorUtils
 				}
 			}
 			else if (Component->CreationMethod == EComponentCreationMethod::Instance)
+#else
+			if (!Component->IsCreatedByConstructionScript())
+#endif
 			{
-				// Instance components are always editable
 				OutEditableComponents.Add(Component);
 			}
 		}
