Network replay skip time fixes and cleanup.

* Added initial ability to set a replay time range as high priority so that the streamer will quickly download that section
* Added initial ability to check if a range of time has been downloaded for a replay
* Refactored replays to store absolute time in MS each frame
* Simplified the ReadDemoFrame loop
* Added ability for http streamer to report that the file stream is at the end (bAtEndOfReplay)
* Bumped replay file version

diff --git a/Engine/Source/Runtime/Engine/Classes/Engine/DemoNetDriver.h b/Engine/Source/Runtime/Engine/Classes/Engine/DemoNetDriver.h
--- a/Engine/Source/Runtime/Engine/Classes/Engine/DemoNetDriver.h
+++ b/Engine/Source/Runtime/Engine/Classes/Engine/DemoNetDriver.h
@@ -22,9 +22,6 @@ class UDemoNetDriver
 	/** Last time (in real seconds) that we recorded a frame */
 	double LastRecordTime;
 
-	/** Time (in game seconds) that have elapsed between recorded frames */
-	float DemoDeltaTime;
-
 	/** Total time of demo in seconds */
 	float DemoTotalTime;
 
@@ -52,6 +49,8 @@ class UDemoNetDriver
 	/** Our network replay streamer */
 	TSharedPtr< class INetworkReplayStreamer > ReplayStreamer;
 
+	uint32 GetDemoCurrentTimeInMS() { return (uint32)( (double)DemoCurrentTime * 1000 ); }
+
 	/** Internal debug timing/tracking */
 	double		AccumulatedRecordTime;
 	double		LastRecordAvgFlush;
diff --git a/Engine/Source/Runtime/Engine/Private/DemoNetDriver.cpp b/Engine/Source/Runtime/Engine/Private/DemoNetDriver.cpp
--- a/Engine/Source/Runtime/Engine/Private/DemoNetDriver.cpp
+++ b/Engine/Source/Runtime/Engine/Private/DemoNetDriver.cpp
@@ -78,8 +78,14 @@ FString UDemoNetDriver::LowLevelGetNetworkNumber()
 	return FString( TEXT( "" ) );
 }
 
+enum ENetworkVersionHistory
+{
+	HISTORY_INITIAL				= 1,
+	HISTORY_SAVE_ABS_TIME_MS	= 2			// We now save the abs demo time in ms for each frame (solves accumulation errors)
+};
+
 static const uint32 NETWORK_DEMO_MAGIC				= 0x2CF5A13D;
-static const uint32 NETWORK_DEMO_VERSION			= 1;
+static const uint32 NETWORK_DEMO_VERSION			= HISTORY_SAVE_ABS_TIME_MS;
 
 static const uint32 NETWORK_DEMO_METADATA_MAGIC		= 0x3D06B24E;
 static const uint32 NETWORK_DEMO_METADATA_VERSION	= 0;
@@ -141,7 +147,6 @@ void UDemoNetDriver::ResetDemoState()
 {
 	DemoFrameNum	= 0;
 	LastRecordTime	= 0;
-	DemoDeltaTime	= 0;
 	DemoTotalTime	= 0;
 	DemoCurrentTime	= 0;
 	DemoTotalFrames	= 0;
@@ -401,6 +406,8 @@ void UDemoNetDriver::TickFlush( float DeltaSeconds )
 
 	if ( FileAr == nullptr )
 	{
+		UE_LOG( LogDemo, Error, TEXT( "UDemoNetDriver::TickFlush: FileAr == nullptr" ) );
+		StopDemo();
 		return;
 	}
 
@@ -666,10 +673,9 @@ void UDemoNetDriver::TickDemoRecord( float DeltaSeconds )
 		return;
 	}
 
-	DemoDeltaTime += DeltaSeconds;
 	DemoCurrentTime += DeltaSeconds;
 
-	ReplayStreamer->UpdateTotalDemoTime( DemoCurrentTime * 1000 );
+	ReplayStreamer->UpdateTotalDemoTime( GetDemoCurrentTimeInMS() );
 
 	const double CurrentSeconds = FPlatformTime::Seconds();
 
@@ -687,16 +693,15 @@ void UDemoNetDriver::TickDemoRecord( float DeltaSeconds )
 	DemoFrameNum++;
 	ReplicationFrame++;
 
-	// Save elapsed game time for this frame
-	*FileAr << DemoDeltaTime;
+	// Save total absolute demo time in MS
+	uint32 SavedAbsTimeMS = GetDemoCurrentTimeInMS();
+	*FileAr << SavedAbsTimeMS;
 
 #if DEMO_CHECKSUMS == 1
-	uint32 DeltaTimeChecksum = FCrc::MemCrc32( &DemoDeltaTime, sizeof( DemoDeltaTime ), 0 );
+	uint32 DeltaTimeChecksum = FCrc::MemCrc32( &SavedAbsTimeMS, sizeof( SavedAbsTimeMS ), 0 );
 	*FileAr << DeltaTimeChecksum;
 #endif
 
-	DemoDeltaTime = 0;
-
 	// Make sure we don't have anything in the buffer for this new frame
 	check( ClientConnections[0]->SendBuffer.GetNumBits() == 0 );
 
@@ -767,14 +772,6 @@ void UDemoNetDriver::PauseChannels( const bool bPause )
 
 bool UDemoNetDriver::ReadDemoFrame()
 {
-	if ( !ReplayStreamer->IsDataAvailable() )
-	{
-		PauseChannels( true );
-		return false;
-	} 
-
-	PauseChannels( false );
-
 	FArchive* FileAr = ReplayStreamer->GetStreamingArchive();
 
 	if ( FileAr->IsError() )
@@ -783,26 +780,41 @@ bool UDemoNetDriver::ReadDemoFrame()
 		return false;
 	}
 
-	if ( FileAr->AtEnd() )		// FIXME: Resolve how we're going to really do this with the streamers
+	if ( FileAr->AtEnd() )
 	{
 		bDemoPlaybackDone = true;
 		PauseChannels( true );
 		return false;
 	}
 
+	if ( ReplayStreamer->GetLastError() != ENetworkReplayError::None )
+	{
+		UE_LOG( LogDemo, Error, TEXT( "UDemoNetDriver::ReadDemoFrame: ReplayStreamer ERROR: %s" ), ENetworkReplayError::ToString( ReplayStreamer->GetLastError() ) );
+		StopDemo();
+		return false;
+	}
+
+	if ( !ReplayStreamer->IsDataAvailable() )
+	{
+		PauseChannels( true );
+		return false;
+	} 
+
+	PauseChannels( false );
+
 	const int32 OldFilePos = FileAr->Tell();
 
-	float ServerDeltaTime;
+	uint32 SavedAbsTimeMS = 0;
 
 	// Peek at the next demo delta time, and see if we should process this frame
-	*FileAr << ServerDeltaTime;
+	*FileAr << SavedAbsTimeMS;
 
 #if DEMO_CHECKSUMS == 1
 	{
 		uint32 ServerDeltaTimeCheksum = 0;
 		*FileAr << ServerDeltaTimeCheksum;
 
-		const uint32 DeltaTimeChecksum = FCrc::MemCrc32( &ServerDeltaTime, sizeof( ServerDeltaTime ), 0 );
+		const uint32 DeltaTimeChecksum = FCrc::MemCrc32( &SavedAbsTimeMS, sizeof( SavedAbsTimeMS ), 0 );
 
 		if ( DeltaTimeChecksum != ServerDeltaTimeCheksum )
 		{
@@ -820,15 +832,13 @@ bool UDemoNetDriver::ReadDemoFrame()
 		return false;
 	}
 
-	if ( DemoDeltaTime < ServerDeltaTime )//&& ServerConnection->State != USOCK_Pending )
+	if ( GetDemoCurrentTimeInMS() < SavedAbsTimeMS )
 	{
 		// Not enough time has passed to read another frame
 		FileAr->Seek( OldFilePos );
 		return false;
 	}
 
-	DemoDeltaTime -= ServerDeltaTime;
-
 	while ( true )
 	{
 		uint8 ReadBuffer[ MAX_DEMO_READ_WRITE_BUFFER ];
@@ -924,54 +934,40 @@ void UDemoNetDriver::TickDemoPlayback( float DeltaSeconds )
 		DemoTotalTime = (float)TotalDemoTimeInMS / 1000.0f;
 	}
 
-	DemoDeltaTime += DeltaSeconds;
-
-	DemoCurrentTime += DeltaSeconds;
-	
 	if ( CVarDemoSkipTime.GetValueOnGameThread() > 0 )
 	{
 		SkipTime( CVarDemoSkipTime.GetValueOnGameThread() );		// Just overwrite existing value, cvar wins in this case
 		CVarDemoSkipTime.AsVariable()->Set( TEXT( "0" ), ECVF_SetByConsole );
 	}
 	
-	if (TimeToSkip > 0.0f)
+	if ( TimeToSkip > 0.0f )
 	{
-		DemoDeltaTime += TimeToSkip;
+		ReplayStreamer->SetHighPriorityTimeRange( 0, DemoTotalTime );		// HACK for now, download the entire demo until we get support to download to any time
 
-		if (DemoCurrentTime + DemoDeltaTime > DemoTotalTime)
+		if ( !ReplayStreamer->IsDataAvailableForTimeRange( 0, DemoTotalTime ) )
 		{
-			DemoDeltaTime = DemoTotalTime - DemoCurrentTime;
+			PauseChannels( true );
+			return;
 		}
 
+		PauseChannels( false );
+
 		DemoCurrentTime += TimeToSkip;
 
 		TimeToSkip = 0.0f;
 	}
 
+	DemoCurrentTime += DeltaSeconds;
+
 	if ( DemoCurrentTime > DemoTotalTime )
 	{
 		DemoCurrentTime = DemoTotalTime;
 	}
 
-	while ( true )
+	// Read demo frames until we are caught up
+	while ( ReadDemoFrame() )
 	{
-		const int32 StartFrame = DemoFrameNum;
-
-		while ( true )
-		{
-			// Read demo frames until we are caught up
-			if ( !ReadDemoFrame() )
-			{
-				break;
-			}
-
-			DemoFrameNum++;
-		}
-
-		if ( StartFrame == DemoFrameNum )
-		{
-			break;
-		}
+		DemoFrameNum++;
 	}
 }
 
diff --git a/Engine/Source/Runtime/NetworkReplayStreaming/HttpNetworkReplayStreaming/Private/HttpNetworkReplayStreaming.cpp b/Engine/Source/Runtime/NetworkReplayStreaming/HttpNetworkReplayStreaming/Private/HttpNetworkReplayStreaming.cpp
--- a/Engine/Source/Runtime/NetworkReplayStreaming/HttpNetworkReplayStreaming/Private/HttpNetworkReplayStreaming.cpp
+++ b/Engine/Source/Runtime/NetworkReplayStreaming/HttpNetworkReplayStreaming/Private/HttpNetworkReplayStreaming.cpp
@@ -94,6 +94,11 @@ void HttpStreamFArchive::Seek( int64 InPos )
 	Pos = InPos;
 }
 
+bool HttpStreamFArchive::AtEnd() 
+{
+	return bAtEndOfReplay;
+}
+
 FHttpNetworkReplayStreamer::FHttpNetworkReplayStreamer() : 
 	StreamFileCount( 0 ), 
 	LastChunkTime( 0 ), 
@@ -104,6 +109,7 @@ FHttpNetworkReplayStreamer::FHttpNetworkReplayStreamer() :
 	bStreamIsLive( false ),
 	NumDownloadChunks( 0 ),
 	DemoTimeInMS( 0 ),
+	HighPriorityEndTime( 0 ),
 	StreamerLastError( ENetworkReplayError::None )
 {
 	// Initialize the server URL
@@ -138,6 +144,7 @@ void FHttpNetworkReplayStreamer::StartStreaming( const FString& StreamName, bool
 	// Setup the archives
 	StreamArchive.ArIsLoading = !bRecord;
 	StreamArchive.ArIsSaving = !StreamArchive.ArIsLoading;
+	StreamArchive.bAtEndOfReplay = false;
 
 	HeaderArchive.ArIsLoading = !bRecord;
 	HeaderArchive.ArIsSaving = !StreamArchive.ArIsLoading;
@@ -319,7 +326,7 @@ void FHttpNetworkReplayStreamer::DownloadNextChunk()
 	const double CHECK_FOR_NEXT_CHUNK_IN_SECONDS = 5;
 
 	const bool bTimeToDownloadChunk = ( FPlatformTime::Seconds() - LastChunkTime > CHECK_FOR_NEXT_CHUNK_IN_SECONDS );
-	const bool bReallyNeedToDownloadChunk = StreamArchive.Pos >= StreamArchive.Buffer.Num() && bMoreChunksDefinitelyAvilable;
+	const bool bReallyNeedToDownloadChunk = ( StreamArchive.Pos >= StreamArchive.Buffer.Num() || HighPriorityEndTime > 0 ) && bMoreChunksDefinitelyAvilable;
 
 	if ( !bTimeToDownloadChunk && !bReallyNeedToDownloadChunk )
 	{
@@ -400,7 +407,7 @@ FArchive* FHttpNetworkReplayStreamer::GetHeaderArchive()
 
 FArchive* FHttpNetworkReplayStreamer::GetStreamingArchive()
 {
-	return ( StreamArchive.IsSaving() || IsDataAvailable() ) ? &StreamArchive : NULL;
+	return &StreamArchive;//( StreamArchive.IsSaving() || IsDataAvailable() ) ? &StreamArchive : NULL;
 }
 
 FArchive* FHttpNetworkReplayStreamer::GetMetadataArchive()
@@ -415,6 +422,11 @@ void FHttpNetworkReplayStreamer::UpdateTotalDemoTime( uint32 TimeInMS )
 
 bool FHttpNetworkReplayStreamer::IsDataAvailable() const
 {
+	if ( GetLastError() != ENetworkReplayError::None )
+	{
+		return false;
+	}
+
 	// If we are loading, and we have more data
 	if ( StreamArchive.IsLoading() && StreamArchive.Pos < StreamArchive.Buffer.Num() && NumDownloadChunks > 0 )
 	{
@@ -424,6 +436,22 @@ bool FHttpNetworkReplayStreamer::IsDataAvailable() const
 	return false;
 }
 
+void FHttpNetworkReplayStreamer::SetHighPriorityTimeRange( const uint32 StartTimeInMS, const uint32 EndTimeInMS )
+{
+	HighPriorityEndTime = EndTimeInMS;
+}
+
+bool FHttpNetworkReplayStreamer::IsDataAvailableForTimeRange( const uint32 StartTimeInMS, const uint32 EndTimeInMS )
+{
+	if ( GetLastError() != ENetworkReplayError::None )
+	{
+		return false;
+	}
+
+	// HACK: We only support the total time range for now
+	return StreamFileCount >= NumDownloadChunks;
+}
+
 bool FHttpNetworkReplayStreamer::IsLive(const FString& StreamName) const 
 {
 	return bStreamIsLive;
@@ -747,13 +775,18 @@ void FHttpNetworkReplayStreamer::HttpEnumerateSessionsFinished( FHttpRequestPtr
 	EnumerateStreamsDelegate = FOnEnumerateStreamsComplete();
 }
 
-void FHttpNetworkReplayStreamer::Tick( float DeltaTime )
+void FHttpNetworkReplayStreamer::Tick( const float DeltaTime )
 {
 	if ( IsHttpBusy() )
 	{
 		return;
 	}
 
+	if ( StreamFileCount >= NumDownloadChunks && bStreamIsLive )
+	{
+		StreamArchive.bAtEndOfReplay = true;
+	}
+
 	if ( SessionName.IsEmpty() )
 	{
 		check( StreamerState == EStreamerState::Idle );
diff --git a/Engine/Source/Runtime/NetworkReplayStreaming/HttpNetworkReplayStreaming/Private/HttpNetworkReplayStreaming.h b/Engine/Source/Runtime/NetworkReplayStreaming/HttpNetworkReplayStreaming/Private/HttpNetworkReplayStreaming.h
--- a/Engine/Source/Runtime/NetworkReplayStreaming/HttpNetworkReplayStreaming/Private/HttpNetworkReplayStreaming.h
+++ b/Engine/Source/Runtime/NetworkReplayStreaming/HttpNetworkReplayStreaming/Private/HttpNetworkReplayStreaming.h
@@ -11,15 +11,17 @@
 class HttpStreamFArchive : public FArchive
 {
 public:
-	HttpStreamFArchive() : Pos( 0 ) {}
+	HttpStreamFArchive() : Pos( 0 ), bAtEndOfReplay( false ) {}
 
 	virtual void	Serialize( void* V, int64 Length );
 	virtual int64	Tell();
 	virtual int64	TotalSize();
 	virtual void	Seek( int64 InPos );
+	virtual bool	AtEnd();
 
 	TArray< uint8 >	Buffer;
 	int32			Pos;
+	bool			bAtEndOfReplay;
 };
 
 /**
@@ -39,6 +41,8 @@ public:
 	virtual void		UpdateTotalDemoTime( uint32 TimeInMS ) override;
 	virtual uint32		GetTotalDemoTime() const override { return DemoTimeInMS; }
 	virtual bool		IsDataAvailable() const override;
+	virtual void		SetHighPriorityTimeRange( const uint32 StartTimeInMS, const uint32 EndTimeInMS ) override;
+	virtual bool		IsDataAvailableForTimeRange( const uint32 StartTimeInMS, const uint32 EndTimeInMS ) override;
 	virtual bool		IsLive( const FString& StreamName ) const override;
 	virtual void		DeleteFinishedStream( const FString& StreamName, const FOnDeleteFinishedStreamComplete& Delegate ) const override;
 	virtual void		EnumerateStreams( const FString& VersionString, const FOnEnumerateStreamsComplete& Delegate ) override;
@@ -63,7 +67,7 @@ public:
 	void HttpUploadFinished( FHttpRequestPtr HttpRequest, FHttpResponsePtr HttpResponse, bool bSucceeded );
 	void HttpEnumerateSessionsFinished( FHttpRequestPtr HttpRequest, FHttpResponsePtr HttpResponse, bool bSucceeded );
 
-	void Tick( float DeltaTime );
+	void Tick( const float DeltaTime );
 	bool IsHttpBusy() const;		// True if we are waiting on http request response
 	bool IsStreaming() const;		// True if we are streaming a replay up or down
 
@@ -109,6 +113,7 @@ public:
 	int32					NumDownloadChunks;
 	uint32					DemoTimeInMS;
 	FString					ViewerName;
+	uint32					HighPriorityEndTime;
 
 	ENetworkReplayError::Type		StreamerLastError;
 
diff --git a/Engine/Source/Runtime/NetworkReplayStreaming/NetworkReplayStreaming/Public/NetworkReplayStreaming.h b/Engine/Source/Runtime/NetworkReplayStreaming/NetworkReplayStreaming/Public/NetworkReplayStreaming.h
--- a/Engine/Source/Runtime/NetworkReplayStreaming/NetworkReplayStreaming/Public/NetworkReplayStreaming.h
+++ b/Engine/Source/Runtime/NetworkReplayStreaming/NetworkReplayStreaming/Public/NetworkReplayStreaming.h
@@ -88,6 +88,8 @@ public:
 	virtual void UpdateTotalDemoTime( uint32 TimeInMS ) = 0;
 	virtual uint32 GetTotalDemoTime() const = 0;
 	virtual bool IsDataAvailable() const = 0;
+	virtual void SetHighPriorityTimeRange( const uint32 StartTimeInMS, const uint32 EndTimeInMS ) = 0;
+	virtual bool IsDataAvailableForTimeRange( const uint32 StartTimeInMS, const uint32 EndTimeInMS ) = 0;
 	
 	/** Returns true if the given StreamName is a game currently in progress */
 	virtual bool IsLive( const FString& StreamName ) const = 0;
diff --git a/Engine/Source/Runtime/NetworkReplayStreaming/NullNetworkReplayStreaming/Public/NullNetworkReplayStreaming.h b/Engine/Source/Runtime/NetworkReplayStreaming/NullNetworkReplayStreaming/Public/NullNetworkReplayStreaming.h
--- a/Engine/Source/Runtime/NetworkReplayStreaming/NullNetworkReplayStreaming/Public/NullNetworkReplayStreaming.h
+++ b/Engine/Source/Runtime/NetworkReplayStreaming/NullNetworkReplayStreaming/Public/NullNetworkReplayStreaming.h
@@ -24,6 +24,8 @@ public:
 	virtual void UpdateTotalDemoTime( uint32 TimeInMS ) override { }
 	virtual uint32 GetTotalDemoTime() const override { return 0; }
 	virtual bool IsDataAvailable() const override { return true; }
+	virtual void SetHighPriorityTimeRange( const uint32 StartTimeInMS, const uint32 EndTimeInMS ) override { }
+	virtual bool IsDataAvailableForTimeRange( const uint32 StartTimeInMS, const uint32 EndTimeInMS ) override { return true; }
 	virtual bool IsLive( const FString& StreamName ) const override;
 	virtual void DeleteFinishedStream( const FString& StreamName, const FOnDeleteFinishedStreamComplete& Delegate) const override;
 	virtual void EnumerateStreams( const FString& VersionString, const FOnEnumerateStreamsComplete& Delegate ) override;
