UE-11605 Overridden components are broken in static Actors in cooked build

diff --git a/Engine/Source/Editor/UnrealEd/Private/Kismet2/BlueprintEditorUtils.cpp b/Engine/Source/Editor/UnrealEd/Private/Kismet2/BlueprintEditorUtils.cpp
--- a/Engine/Source/Editor/UnrealEd/Private/Kismet2/BlueprintEditorUtils.cpp
+++ b/Engine/Source/Editor/UnrealEd/Private/Kismet2/BlueprintEditorUtils.cpp
@@ -1158,16 +1158,7 @@ UClass* FBlueprintEditorUtils::RegenerateBlueprintClass(UBlueprint* Blueprint, U
 		// Preload Overriden Components
 		if (Blueprint->InheritableComponentHandler)
 		{
-			if (Blueprint->InheritableComponentHandler->HasAllFlags(RF_NeedLoad))
-			{
-				auto Linker = Blueprint->InheritableComponentHandler->GetLinker();
-				if (Linker)
-				{
-					Linker->Preload(Blueprint->InheritableComponentHandler);
-				}
-			}
-
-			Blueprint->InheritableComponentHandler->PreloadAllTempates();
+			Blueprint->InheritableComponentHandler->PreloadAll();
 		}
 
 		// Purge any NULL graphs
diff --git a/Engine/Source/Runtime/Engine/Classes/Engine/InheritableComponentHandler.h b/Engine/Source/Runtime/Engine/Classes/Engine/InheritableComponentHandler.h
--- a/Engine/Source/Runtime/Engine/Classes/Engine/InheritableComponentHandler.h
+++ b/Engine/Source/Runtime/Engine/Classes/Engine/InheritableComponentHandler.h
@@ -84,14 +84,15 @@ public:
 		return 0 == Records.Num();
 	}
 
-	void PreloadAllTempates();
-
 	bool RenameTemplate(FComponentKey OldKey, FName NewName);
 
 	FComponentKey FindKey(UActorComponent* ComponentTemplate) const;
 #endif
 
 public:
+	void PreloadAllTempates();
+	void PreloadAll();
+
 	FComponentKey FindKey(const FName VariableName) const;
 
 	UActorComponent* GetOverridenComponentTemplate(FComponentKey Key) const;
diff --git a/Engine/Source/Runtime/Engine/Private/BlueprintGeneratedClass.cpp b/Engine/Source/Runtime/Engine/Private/BlueprintGeneratedClass.cpp
--- a/Engine/Source/Runtime/Engine/Private/BlueprintGeneratedClass.cpp
+++ b/Engine/Source/Runtime/Engine/Private/BlueprintGeneratedClass.cpp
@@ -291,6 +291,10 @@ UInheritableComponentHandler* UBlueprintGeneratedClass::GetInheritableComponentH
 	{
 		return nullptr;
 	}
+	else
+	{
+		InheritableComponentHandler->PreloadAll();
+	}
 
 	if (!InheritableComponentHandler && bCreateIfNecessary)
 	{
diff --git a/Engine/Source/Runtime/Engine/Private/InheritableComponentHandler.cpp b/Engine/Source/Runtime/Engine/Private/InheritableComponentHandler.cpp
--- a/Engine/Source/Runtime/Engine/Private/InheritableComponentHandler.cpp
+++ b/Engine/Source/Runtime/Engine/Private/InheritableComponentHandler.cpp
@@ -192,21 +192,6 @@ bool UInheritableComponentHandler::IsRecordNecessary(const FComponentOverrideRec
 	return !FComponentComparisonHelper::AreIdentical(ChildComponentTemplate, ParentComponentTemplate);
 }
 
-void UInheritableComponentHandler::PreloadAllTempates()
-{
-	for (auto Record : Records)
-	{
-		if (Record.ComponentTemplate && Record.ComponentTemplate->HasAllFlags(RF_NeedLoad))
-		{
-			auto Linker = Record.ComponentTemplate->GetLinker();
-			if (Linker)
-			{
-				Linker->Preload(Record.ComponentTemplate);
-			}
-		}
-	}
-}
-
 UActorComponent* UInheritableComponentHandler::FindBestArchetype(FComponentKey Key) const
 {
 	UActorComponent* ClosestArchetype = nullptr;
@@ -260,6 +245,34 @@ FComponentKey UInheritableComponentHandler::FindKey(UActorComponent* ComponentTe
 
 #endif
 
+void UInheritableComponentHandler::PreloadAllTempates()
+{
+	for (auto Record : Records)
+	{
+		if (Record.ComponentTemplate && Record.ComponentTemplate->HasAllFlags(RF_NeedLoad))
+		{
+			auto Linker = Record.ComponentTemplate->GetLinker();
+			if (Linker)
+			{
+				Linker->Preload(Record.ComponentTemplate);
+			}
+		}
+	}
+}
+
+void UInheritableComponentHandler::PreloadAll()
+{
+	if (HasAllFlags(RF_NeedLoad))
+	{
+		auto Linker = GetLinker();
+		if (Linker)
+		{
+			Linker->Preload(this);
+		}
+	}
+	PreloadAllTempates();
+}
+
 FComponentKey UInheritableComponentHandler::FindKey(const FName VariableName) const
 {
 	for (const FComponentOverrideRecord& Record : Records)
diff --git a/Engine/Source/Runtime/Engine/Private/SCS_Node.cpp b/Engine/Source/Runtime/Engine/Private/SCS_Node.cpp
--- a/Engine/Source/Runtime/Engine/Private/SCS_Node.cpp
+++ b/Engine/Source/Runtime/Engine/Private/SCS_Node.cpp
@@ -39,9 +39,10 @@ UActorComponent* USCS_Node::ExecuteNodeOnActor(AActor* Actor, USceneComponent* P
 		auto ActualBPGC = Cast<UBlueprintGeneratedClass>(Actor->GetClass());
 		while (!OverridenComponentTemplate && ActualBPGC)
 		{
-			if (ActualBPGC->InheritableComponentHandler)
+			auto InheritableComponentHandler = ActualBPGC->GetInheritableComponentHandler();
+			if (InheritableComponentHandler)
 			{
-				OverridenComponentTemplate = ActualBPGC->InheritableComponentHandler->GetOverridenComponentTemplate(ComponentKey);
+				OverridenComponentTemplate = InheritableComponentHandler->GetOverridenComponentTemplate(ComponentKey);
 			}
 			ActualBPGC = Cast<UBlueprintGeneratedClass>(ActualBPGC->GetSuperClass());
 		}
