UE-11605 Overridden components are broken in static Actors in cooked build - fix

diff --git a/Engine/Source/Runtime/Engine/Private/BlueprintGeneratedClass.cpp b/Engine/Source/Runtime/Engine/Private/BlueprintGeneratedClass.cpp
--- a/Engine/Source/Runtime/Engine/Private/BlueprintGeneratedClass.cpp
+++ b/Engine/Source/Runtime/Engine/Private/BlueprintGeneratedClass.cpp
@@ -291,7 +291,8 @@ UInheritableComponentHandler* UBlueprintGeneratedClass::GetInheritableComponentH
 	{
 		return nullptr;
 	}
-	else
+	
+	if (InheritableComponentHandler)
 	{
 		InheritableComponentHandler->PreloadAll();
 	}
