BP: UberGraphFunction is preloaded in UBlueprintGeneratedClass::Link

diff --git a/Engine/Source/Runtime/Engine/Private/BlueprintGeneratedClass.cpp b/Engine/Source/Runtime/Engine/Private/BlueprintGeneratedClass.cpp
--- a/Engine/Source/Runtime/Engine/Private/BlueprintGeneratedClass.cpp
+++ b/Engine/Source/Runtime/Engine/Private/BlueprintGeneratedClass.cpp
@@ -651,6 +651,8 @@ void UBlueprintGeneratedClass::Link(FArchive& Ar, bool bRelinkExistingProperties
 
 	if (UsePersistentUberGraphFrame() && UberGraphFunction)
 	{
+		Ar.Preload(UberGraphFunction);
+
 		for (auto Property : TFieldRange<UStructProperty>(this, EFieldIteratorFlags::ExcludeSuper))
 		{
 			if (Property->GetFName() == GetUberGraphFrameName())
