Added project settings flag to disable landscape GI data generation

diff --git a/Engine/Source/Runtime/Core/Private/HAL/ConsoleManager.cpp b/Engine/Source/Runtime/Core/Private/HAL/ConsoleManager.cpp
--- a/Engine/Source/Runtime/Core/Private/HAL/ConsoleManager.cpp
+++ b/Engine/Source/Runtime/Core/Private/HAL/ConsoleManager.cpp
@@ -1790,6 +1790,13 @@ static TAutoConsoleVariable<int32> CVarDistField(
 	TEXT("Enabling will increase mesh build times and memory usage.  Changing this value will cause a rebuild of all static meshes."),
 	ECVF_ReadOnly);
 
+static TAutoConsoleVariable<int32> CVarLandscapeGI(
+	TEXT("r.GenerateLandscapeGIData"),
+	1,
+	TEXT("Whether to generate a low-resolution base color texture for landscapes for rendering real-time global illumination.\n")
+	TEXT("This feature requires GenerateMeshDistanceFields is also enabled, and will increase mesh build times and memory usage.\n"),
+	ECVF_Default);
+
 static TAutoConsoleVariable<int32> CVarMinLogVerbosity(
 	TEXT("con.MinLogVerbosity"),
 	0,
diff --git a/Engine/Source/Runtime/Engine/Classes/Engine/RendererSettings.h b/Engine/Source/Runtime/Engine/Classes/Engine/RendererSettings.h
--- a/Engine/Source/Runtime/Engine/Classes/Engine/RendererSettings.h
+++ b/Engine/Source/Runtime/Engine/Classes/Engine/RendererSettings.h
@@ -154,6 +154,11 @@ class ENGINE_API URendererSettings
 		ToolTip="Whether to build distance fields of static meshes, needed for distance field AO, which is used to implement Movable SkyLight shadows, and ray traced distance field shadows on directional lights.  Enabling will increase mesh build times and memory usage.  Changing this setting requires restarting the editor."))
 	uint32 bGenerateMeshDistanceFields:1;
 
+	UPROPERTY(config, EditAnywhere, Category = Lighting, meta = (
+		ConsoleVariable = "r.GenerateLandscapeGIData", DisplayName = "Generate Landscape Real-time GI Data",
+		ToolTip = "Whether to generate a low-resolution base color texture for landscapes for rendering real-time global illumination.  This feature requires GenerateMeshDistanceFields is also enabled, and will increase mesh build times and memory usage."))
+		uint32 bGenerateLandscapeGIData : 1;
+
 	UPROPERTY(config, EditAnywhere, Category=Lighting, meta=(
 		ConsoleVariable="r.Shadow.DistanceFieldPenumbraSize",
 		ToolTip="Controls the size of the uniform penumbra produced by static shadowing."))
diff --git a/Engine/Source/Runtime/Landscape/Private/Landscape.cpp b/Engine/Source/Runtime/Landscape/Private/Landscape.cpp
--- a/Engine/Source/Runtime/Landscape/Private/Landscape.cpp
+++ b/Engine/Source/Runtime/Landscape/Private/Landscape.cpp
@@ -2271,6 +2271,28 @@ void ALandscapeProxy::UpdateBakedTextures()
 		return;
 	}
 
+	// Check if we can want to generate landscape GI data
+	static const auto DistanceFieldCVar = IConsoleManager::Get().FindTConsoleVariableDataInt(TEXT("r.GenerateMeshDistanceFields"));
+	static const auto LandscapeGICVar = IConsoleManager::Get().FindTConsoleVariableDataInt(TEXT("r.GenerateLandscapeGIData"));
+	if (DistanceFieldCVar->GetValueOnGameThread() == 0 || LandscapeGICVar->GetValueOnGameThread() == 0)
+	{
+		// Clear out any existing GI textures
+		for (ULandscapeComponent* Component : LandscapeComponents)
+		{
+			if (Component->GIBakedBaseColorTexture != nullptr)
+			{
+				Component->BakedTextureMaterialGuid.Invalidate();
+				Component->GIBakedBaseColorTexture = nullptr;
+				Component->MarkRenderStateDirty();
+			}
+		}
+
+		// Don't check if we need to update anything for another 60 frames
+		UpdateBakedTexturesCountdown = 60;
+
+		return;
+	}
+	
 	// Stores the components and their state hash data for a single atlas
 	struct FBakedTextureSourceInfo
 	{
