Prevent log spam when particle systems collide with static objects

A warning was added to advise users to avoid applying impulses/forces to things that cannot move. Unfortunately this meant that the log was spammed with thousands of particle collisions.

Reported here:
https://answers.unrealengine.com/questions/179345/47-recent-horrible-lag-while-loading-level-caused.html

diff --git a/Engine/Source/Runtime/Engine/Private/Particles/ParticleModules_Collision.cpp b/Engine/Source/Runtime/Engine/Private/Particles/ParticleModules_Collision.cpp
--- a/Engine/Source/Runtime/Engine/Private/Particles/ParticleModules_Collision.cpp
+++ b/Engine/Source/Runtime/Engine/Private/Particles/ParticleModules_Collision.cpp
@@ -433,11 +433,15 @@ void UParticleModuleCollision::Update(FParticleEmitterInstance* Owner, int32 Off
 						FVector	NewLocation		= Location + NewVelocity * (1.f - Hit.Time);
 						Particle.Location		= Owner->Component->ComponentToWorld.InverseTransformPosition(NewLocation);
 
-						if (bApplyPhysics && Hit.Component.Get())
+						if (bApplyPhysics)
 						{
-							FVector vImpulse;
-							vImpulse = -(NewVelocity - OldVelocity) * ParticleMass.GetValue(Particle.RelativeTime, Owner->Component);
-							Hit.Component->AddImpulseAtLocation(vImpulse, Hit.Location, Hit.BoneName);
+							UPrimitiveComponent* PrimitiveComponent = Hit.Component.Get();
+							if(PrimitiveComponent && PrimitiveComponent->IsAnySimulatingPhysics())
+							{
+								FVector vImpulse;
+								vImpulse = -(NewVelocity - OldVelocity) * ParticleMass.GetValue(Particle.RelativeTime, Owner->Component);
+								PrimitiveComponent->AddImpulseAtLocation(vImpulse, Hit.Location, Hit.BoneName);
+							}
 						}
 					}
 					else
@@ -458,11 +462,15 @@ void UParticleModuleCollision::Update(FParticleEmitterInstance* Owner, int32 Off
 						Particle.Velocity		= FVector::ZeroVector;
 						Particle.Location	   += vNewVelocity * (1.f - Hit.Time);
 
-						if (bApplyPhysics && Hit.Component.Get())
+						if (bApplyPhysics)
 						{
-							FVector vImpulse;
-							vImpulse = -(vNewVelocity - vOldVelocity) * ParticleMass.GetValue(Particle.RelativeTime, Owner->Component);
-							Hit.Component->AddImpulseAtLocation(vImpulse, Hit.Location, Hit.BoneName);
+							UPrimitiveComponent* PrimitiveComponent = Hit.Component.Get();
+							if(PrimitiveComponent && PrimitiveComponent->IsAnySimulatingPhysics())
+							{
+								FVector vImpulse;
+								vImpulse = -(vNewVelocity - vOldVelocity) * ParticleMass.GetValue(Particle.RelativeTime, Owner->Component);
+								PrimitiveComponent->AddImpulseAtLocation(vImpulse, Hit.Location, Hit.BoneName);
+							}
 						}
 					}
 
