Fixed moving to actor occasionally failing due to path's final point not being on navmesh #UE4

Basically made sure path's end (and start!) are _on_ the navmesh (not in the air, potentially)
Also:
- added an optional parameter to BP NavigationSystem::ProjectPointToNavigation so that users can specify query extent

diff --git a/Engine/Source/Runtime/Engine/Classes/AI/Navigation/NavigationSystem.h b/Engine/Source/Runtime/Engine/Classes/AI/Navigation/NavigationSystem.h
--- a/Engine/Source/Runtime/Engine/Classes/AI/Navigation/NavigationSystem.h
+++ b/Engine/Source/Runtime/Engine/Classes/AI/Navigation/NavigationSystem.h
@@ -238,7 +238,7 @@ public:
 	
 	/** Project a point onto the NavigationData */
 	UFUNCTION(BlueprintPure, Category="AI|Navigation", meta=(WorldContext="WorldContext" ) )
-	static FVector ProjectPointToNavigation(UObject* WorldContext, const FVector& Point, ANavigationData* NavData = NULL, TSubclassOf<UNavigationQueryFilter> FilterClass = NULL);
+	static FVector ProjectPointToNavigation(UObject* WorldContext, const FVector& Point, ANavigationData* NavData = NULL, TSubclassOf<UNavigationQueryFilter> FilterClass = NULL, const FVector QueryExtent = FVector::ZeroVector);
 	
 	UFUNCTION(BlueprintPure, Category="AI|Navigation", meta=(WorldContext="WorldContext" ) )
 	static FVector GetRandomPoint(UObject* WorldContext, ANavigationData* NavData = NULL, TSubclassOf<UNavigationQueryFilter> FilterClass = NULL);
diff --git a/Engine/Source/Runtime/Engine/Private/AI/Navigation/NavigationSystem.cpp b/Engine/Source/Runtime/Engine/Private/AI/Navigation/NavigationSystem.cpp
--- a/Engine/Source/Runtime/Engine/Private/AI/Navigation/NavigationSystem.cpp
+++ b/Engine/Source/Runtime/Engine/Private/AI/Navigation/NavigationSystem.cpp
@@ -3321,7 +3321,7 @@ UNavigationSystem* UNavigationSystem::GetNavigationSystem(UObject* WorldContext)
 	return GetCurrent(WorldContext);
 }
 
-FVector UNavigationSystem::ProjectPointToNavigation(UObject* WorldContextObject, const FVector& Point, ANavigationData* NavData, TSubclassOf<UNavigationQueryFilter> FilterClass)
+FVector UNavigationSystem::ProjectPointToNavigation(UObject* WorldContextObject, const FVector& Point, ANavigationData* NavData, TSubclassOf<UNavigationQueryFilter> FilterClass, const FVector QueryExtent)
 {
 	FNavLocation ProjectedPoint(Point);
 
@@ -3330,7 +3330,7 @@ FVector UNavigationSystem::ProjectPointToNavigation(UObject* WorldContextObject,
 	if (NavSys)
 	{
 		ANavigationData* UseNavData = NavData ? NavData : NavSys->GetMainNavData(FNavigationSystem::DontCreate);
-		NavSys->ProjectPointToNavigation(Point, ProjectedPoint, INVALID_NAVEXTENT, UseNavData, UNavigationQueryFilter::GetQueryFilter(UseNavData, FilterClass));
+		NavSys->ProjectPointToNavigation(Point, ProjectedPoint, QueryExtent.IsNearlyZero() ? INVALID_NAVEXTENT : QueryExtent, UseNavData, UNavigationQueryFilter::GetQueryFilter(UseNavData, FilterClass));
 	}
 
 	return ProjectedPoint.Location;
diff --git a/Engine/Source/Runtime/Engine/Private/AI/Navigation/PImplRecastNavMesh.cpp b/Engine/Source/Runtime/Engine/Private/AI/Navigation/PImplRecastNavMesh.cpp
--- a/Engine/Source/Runtime/Engine/Private/AI/Navigation/PImplRecastNavMesh.cpp
+++ b/Engine/Source/Runtime/Engine/Private/AI/Navigation/PImplRecastNavMesh.cpp
@@ -849,7 +849,7 @@ ENavigationQueryResult::Type FPImplRecastNavMesh::FindPath(const FVector& StartL
 		FVector RecastHandPlacedPathEnd;
 		NavQuery.closestPointOnPolyBoundary(StartPolyID, &RecastEndPos.X, &RecastHandPlacedPathEnd.X);
 
-		new(Path.GetPathPoints()) FNavPathPoint(StartLoc, StartPolyID);
+		new(Path.GetPathPoints()) FNavPathPoint(Recast2UnrVector(&RecastStartPos.X), StartPolyID);
 		new(Path.GetPathPoints()) FNavPathPoint(Recast2UnrVector(&RecastHandPlacedPathEnd.X), StartPolyID);
 
 		Path.PathCorridor.Add(PathResult.getRef(0));
@@ -858,7 +858,7 @@ ENavigationQueryResult::Type FPImplRecastNavMesh::FindPath(const FVector& StartL
 	else
 	{
 		PostProcessPath(FindPathStatus, Path, NavQuery, QueryFilter,
-			StartPolyID, EndPolyID, StartLoc, EndLoc, RecastStartPos, RecastEndPos,
+			StartPolyID, EndPolyID, Recast2UnrVector(&RecastStartPos.X), Recast2UnrVector(&RecastEndPos.X), RecastStartPos, RecastEndPos,
 			PathResult);
 	}
 
@@ -939,11 +939,11 @@ bool FPImplRecastNavMesh::InitPathfinding(const FVector& UnrealStart, const FVec
 	const FVector& NavExtent = NavMeshOwner->GetDefaultQueryExtent();
 	const float Extent[3] = { NavExtent.X, NavExtent.Z, NavExtent.Y };
 
-	RecastStart = Unreal2RecastPoint(UnrealStart);
-	RecastEnd = Unreal2RecastPoint(UnrealEnd);
+	const FVector RecastStartToProject = Unreal2RecastPoint(UnrealStart);
+	const FVector RecastEndToProject = Unreal2RecastPoint(UnrealEnd);
 
 	StartPoly = INVALID_NAVNODEREF;
-	Query.findNearestPoly(&RecastStart.X, Extent, Filter, &StartPoly, NULL);
+	Query.findNearestPoly(&RecastStartToProject.X, Extent, Filter, &StartPoly, &RecastStart.X);
 	if (StartPoly == INVALID_NAVNODEREF)
 	{
 		UE_VLOG(NavMeshOwner, LogNavigation, Warning, TEXT("FPImplRecastNavMesh::InitPathfinding start point not on navmesh"));
@@ -955,7 +955,7 @@ bool FPImplRecastNavMesh::InitPathfinding(const FVector& UnrealStart, const FVec
 	}
 
 	EndPoly = INVALID_NAVNODEREF;
-	Query.findNearestPoly(&RecastEnd.X, Extent, Filter, &EndPoly, NULL);
+	Query.findNearestPoly(&RecastEndToProject.X, Extent, Filter, &EndPoly, &RecastEnd.X);
 	if (EndPoly == INVALID_NAVNODEREF)
 	{
 		UE_VLOG(NavMeshOwner, LogNavigation, Warning, TEXT("FPImplRecastNavMesh::InitPathfinding end point not on navmesh"));
