Changed the default parameter of BP NavigationSystem::ProjectPointToNavigation to be the default extent rather then just "zero vector" #UE4

diff --git a/Engine/Source/Runtime/Engine/Classes/AI/Navigation/NavigationSystem.h b/Engine/Source/Runtime/Engine/Classes/AI/Navigation/NavigationSystem.h
--- a/Engine/Source/Runtime/Engine/Classes/AI/Navigation/NavigationSystem.h
+++ b/Engine/Source/Runtime/Engine/Classes/AI/Navigation/NavigationSystem.h
@@ -238,7 +238,7 @@ public:
 	
 	/** Project a point onto the NavigationData */
 	UFUNCTION(BlueprintPure, Category="AI|Navigation", meta=(WorldContext="WorldContext" ) )
-	static FVector ProjectPointToNavigation(UObject* WorldContext, const FVector& Point, ANavigationData* NavData = NULL, TSubclassOf<UNavigationQueryFilter> FilterClass = NULL, const FVector QueryExtent = FVector::ZeroVector);
+	static FVector ProjectPointToNavigation(UObject* WorldContext, const FVector& Point, ANavigationData* NavData = NULL, TSubclassOf<UNavigationQueryFilter> FilterClass = NULL, const FVector QueryExtent = INVALID_NAVEXTENT);
 	
 	UFUNCTION(BlueprintPure, Category="AI|Navigation", meta=(WorldContext="WorldContext" ) )
 	static FVector GetRandomPoint(UObject* WorldContext, ANavigationData* NavData = NULL, TSubclassOf<UNavigationQueryFilter> FilterClass = NULL);
