Flip normal for hit normal if the object is not moving. This is needed so that normal is always pushing the object away from the impact point

diff --git a/Engine/Source/Runtime/Engine/Private/Actor.cpp b/Engine/Source/Runtime/Engine/Private/Actor.cpp
--- a/Engine/Source/Runtime/Engine/Private/Actor.cpp
+++ b/Engine/Source/Runtime/Engine/Private/Actor.cpp
@@ -1779,23 +1779,33 @@ void AActor::DispatchBlockingHit(UPrimitiveComponent* MyComp, UPrimitiveComponen
 	check(MyComp);
 
 	AActor* OtherActor = (OtherComp != NULL) ? OtherComp->GetOwner() : NULL;
+	const FHitResult* HitResult = &Hit;
 
+	FHitResult FlippedHit(Hit);
+	if(bSelfMoved == false)	//if we didn't move we flip the normals so they push us away from moving object
+	{
+		FlippedHit.Normal = -FlippedHit.Normal;
+		FlippedHit.ImpactNormal = -FlippedHit.ImpactNormal;
+		HitResult = &FlippedHit;
+	}
+	
 	// Call virtual
 	if(IsActorValidToNotify(this))
 	{
-		ReceiveHit(MyComp, OtherActor, OtherComp, bSelfMoved, Hit.ImpactPoint, Hit.ImpactNormal, FVector(0,0,0), Hit);
+		
+		ReceiveHit(MyComp, OtherActor, OtherComp, bSelfMoved, Hit.ImpactPoint, HitResult->ImpactNormal, FVector(0,0,0), *HitResult);
 	}
 
 	// If we are still ok, call delegate on actor
 	if(IsActorValidToNotify(this))
 	{
-		OnActorHit.Broadcast(this, OtherActor, FVector(0,0,0), Hit);
+		OnActorHit.Broadcast(this, OtherActor, FVector(0,0,0), *HitResult);
 	}
 
 	// If component is still alive, call delegate on component
 	if(!MyComp->IsPendingKill())
 	{
-		MyComp->OnComponentHit.Broadcast(OtherActor, OtherComp, FVector(0,0,0), Hit);
+		MyComp->OnComponentHit.Broadcast(OtherActor, OtherComp, FVector(0,0,0), *HitResult);
 	}
 }
 
