BP - The UserDefinedStruct now preloads all editor data in Serialize so that all variable definition data is available at editor time.

diff --git a/Engine/Source/Runtime/Engine/Private/UserDefinedStruct.cpp b/Engine/Source/Runtime/Engine/Private/UserDefinedStruct.cpp
--- a/Engine/Source/Runtime/Engine/Private/UserDefinedStruct.cpp
+++ b/Engine/Source/Runtime/Engine/Private/UserDefinedStruct.cpp
@@ -20,6 +20,13 @@ void UUserDefinedStruct::Serialize(FArchive& Ar)
 
 	if (Ar.IsLoading() && (EUserDefinedStructureStatus::UDSS_UpToDate == Status))
 	{
+		// We need to force the editor data to be preload in case anyone needs to extract variable
+		// information at editor time about the user structure.
+		if ( EditorData != nullptr )
+		{
+			Ar.Preload(EditorData);
+		}
+
 		const FStructureEditorUtils::EStructureError Result = FStructureEditorUtils::IsStructureValid(this, NULL, &ErrorMessage);
 		if (FStructureEditorUtils::EStructureError::Ok != Result)
 		{
