PR#897: Can now "add pins" to the Append node for strings from user byte3

Submitted without changes

diff --git a/Engine/Source/Runtime/Engine/Classes/Kismet/KismetStringLibrary.h b/Engine/Source/Runtime/Engine/Classes/Kismet/KismetStringLibrary.h
--- a/Engine/Source/Runtime/Engine/Classes/Kismet/KismetStringLibrary.h
+++ b/Engine/Source/Runtime/Engine/Classes/Kismet/KismetStringLibrary.h
@@ -176,7 +176,7 @@ class UKismetStringLibrary : public UBlueprintFunctionLibrary
 	 * @param B - The string to append to A
 	 * @returns A new string which is the concatenation of A+B
 	 */
-	UFUNCTION(BlueprintPure, meta=(FriendlyName = "Append"), Category="Utilities|String")
+	UFUNCTION(BlueprintPure, meta=(FriendlyName = "Append", CommutativeAssociativeBinaryOperator = "true"), Category="Utilities|String")
 	static FString Concat_StrStr(const FString& A, const FString& B);
 
 	/**
