Fix .Release and .Dev engine analytics suffix determination. Was swapped.

diff --git a/Engine/Source/Runtime/Engine/Private/EngineAnalytics.cpp b/Engine/Source/Runtime/Engine/Private/EngineAnalytics.cpp
--- a/Engine/Source/Runtime/Engine/Private/EngineAnalytics.cpp
+++ b/Engine/Source/Runtime/Engine/Private/EngineAnalytics.cpp
@@ -78,7 +78,7 @@ void FEngineAnalytics::Initialize()
 						const bool bUseReleaseAccount =
 							(AnalyticsBuildType == FAnalytics::Development || AnalyticsBuildType == FAnalytics::Release) &&
 							!FEngineBuildSettings::IsInternalBuild();	// Internal Epic build
-						const TCHAR* BuildTypeStr = bUseReleaseAccount ? TEXT("Dev") : TEXT("Release");
+						const TCHAR* BuildTypeStr = bUseReleaseAccount ? TEXT("Release") : TEXT("Dev");
 						const TCHAR* UE4TypeStr = FRocketSupport::IsRocket() ? TEXT("Rocket") : FEngineBuildSettings::IsPerforceBuild() ? TEXT("Perforce") : TEXT("UnrealEngine");
 						if (bIsEditorRun)
 						{
