PR #898: Exposed FMath::GridSnap() to Blueprints (Contributed by [REDACTED_PERSON])

Modified comments and variable names

diff --git a/Engine/Source/Runtime/Engine/Classes/Kismet/KismetMathLibrary.h b/Engine/Source/Runtime/Engine/Classes/Kismet/KismetMathLibrary.h
--- a/Engine/Source/Runtime/Engine/Classes/Kismet/KismetMathLibrary.h
+++ b/Engine/Source/Runtime/Engine/Classes/Kismet/KismetMathLibrary.h
@@ -296,6 +296,14 @@ class UKismetMathLibrary : public UBlueprintFunctionLibrary
 	UFUNCTION(BlueprintPure, meta=(FriendlyName = "InRange (float)", Min="0.0", Max="1.0"), Category="Math|Float")
 	static bool InRange_FloatFloat(float Value, float Min, float Max, bool InclusiveMin = true, bool InclusiveMax = true);
 
+	/** Snaps a value to the nearest grid multiple. E.g.,
+	 *		Location = 5.1, GridSize = 10.0 : return value = 10.0
+	 * If GridSize is 0 Location is returned
+	 * if GridSize is very small precision issues may occur.
+	 */
+	UFUNCTION(BlueprintPure, meta = (FriendlyName = "Snap to grid (float)"), Category = "Math|Float")
+	static float GridSnap_Float(float Location, float GridSize);
+
 	/* Returns the absolute (positive) value of A */
 	UFUNCTION(BlueprintPure, meta=(FriendlyName = "Absolute (float)", CompactNodeTitle = "ABS"), Category="Math|Float")
 	static float Abs(float A);
diff --git a/Engine/Source/Runtime/Engine/Private/KismetMathLibrary.cpp b/Engine/Source/Runtime/Engine/Private/KismetMathLibrary.cpp
--- a/Engine/Source/Runtime/Engine/Private/KismetMathLibrary.cpp
+++ b/Engine/Source/Runtime/Engine/Private/KismetMathLibrary.cpp
@@ -346,6 +346,11 @@ bool UKismetMathLibrary::InRange_FloatFloat(float Value, float Min, float Max, b
 	return ((InclusiveMin ? (Value >= Min) : (Value > Min)) && (InclusiveMax ? (Value <= Max) : (Value < Max)));
 }	
 
+float UKismetMathLibrary::GridSnap_Float(float Location, float GridSize)
+{
+	return FMath::GridSnap(Location, GridSize);
+}
+
 float UKismetMathLibrary::GetPI()
 {
     return PI;
