#UE4 Component templates are now prefixed with their variable name, so it is now VariableNameGUID

diff --git a/Engine/Source/Runtime/Engine/Private/SimpleConstructionScript.cpp b/Engine/Source/Runtime/Engine/Private/SimpleConstructionScript.cpp
--- a/Engine/Source/Runtime/Engine/Private/SimpleConstructionScript.cpp
+++ b/Engine/Source/Runtime/Engine/Private/SimpleConstructionScript.cpp
@@ -834,9 +834,7 @@ USCS_Node* USimpleConstructionScript::CreateNodeImpl(UActorComponent* NewCompone
 	auto NewNode = NewObject<USCS_Node>(this, MakeUniqueObjectName(this, USCS_Node::StaticClass()));
 	NewNode->SetFlags(RF_Transactional);
 	NewNode->ComponentTemplate = NewComponentTemplate;
-
-	// Now create a name for the new component.
-	NewNode->VariableName = GenerateNewComponentName(NewComponentTemplate->GetClass(), ComponentVariableName);
+	NewNode->VariableName = ComponentVariableName;
 
 	// Note: This should match up with UEdGraphSchema_K2::VR_DefaultCategory
 	NewNode->CategoryName = TEXT("Default");
@@ -851,7 +849,10 @@ USCS_Node* USimpleConstructionScript::CreateNode(UClass* NewComponentClass, FNam
 	check(NewComponentClass->IsChildOf(UActorComponent::StaticClass()));
 	ensure(Cast<UBlueprintGeneratedClass>(Blueprint->GeneratedClass));
 
-	auto NewComponentTemplate = NewObject<UActorComponent>(Blueprint->GeneratedClass, NewComponentClass, *(FGuid::NewGuid().ToString()) /*NewComponentName*/);
+	// note that naming logic is duplicated in CreateNodeAndRenameComponent:
+	NewComponentVariableName = GenerateNewComponentName(NewComponentClass, NewComponentVariableName);
+
+	auto NewComponentTemplate = NewObject<UActorComponent>(Blueprint->GeneratedClass, NewComponentClass, *(NewComponentVariableName.GetPlainNameString() + FGuid::NewGuid().ToString() ) );
 	NewComponentTemplate->SetFlags(RF_ArchetypeObject|RF_Transactional|RF_Public);
 
 	return CreateNodeImpl(NewComponentTemplate, NewComponentVariableName);
@@ -861,10 +862,13 @@ USCS_Node* USimpleConstructionScript::CreateNodeAndRenameComponent(UActorCompone
 {
 	check(NewComponentTemplate);
 
+	// note that naming logic is duplicated in CreateNode:
+	FName NewComponentVariableName = GenerateNewComponentName(NewComponentTemplate->GetClass());
+
 	// Relocate the instance from the transient package to the BPGC and assign it a unique object name
-	NewComponentTemplate->Rename(*(FGuid::NewGuid().ToString()), GetBlueprint()->GeneratedClass, REN_DontCreateRedirectors | REN_DoNotDirty);
+	NewComponentTemplate->Rename(*(NewComponentVariableName.GetPlainNameString() + FGuid::NewGuid().ToString()), GetBlueprint()->GeneratedClass, REN_DontCreateRedirectors | REN_DoNotDirty);
 
-	return CreateNodeImpl(NewComponentTemplate, NAME_None);
+	return CreateNodeImpl(NewComponentTemplate, NewComponentVariableName);
 }
 
 void USimpleConstructionScript::ValidateNodeVariableNames(FCompilerResultsLog& MessageLog)
