Engine: Fix Canvas having an invalid ViewRect when passed to ViewportClient->Draw

diff --git a/Engine/Source/Runtime/Engine/Private/UnrealClient.cpp b/Engine/Source/Runtime/Engine/Private/UnrealClient.cpp
--- a/Engine/Source/Runtime/Engine/Private/UnrealClient.cpp
+++ b/Engine/Source/Runtime/Engine/Private/UnrealClient.cpp
@@ -987,6 +987,7 @@ void FViewport::Draw( bool bShouldPresent /*= true */)
 
 				UWorld* ViewportWorld = ViewportClient->GetWorld();
 				FCanvas Canvas(this, NULL, ViewportWorld, ViewportWorld ? ViewportWorld->FeatureLevel : GMaxRHIFeatureLevel);
+				Canvas.SetRenderTargetRect(FIntRect(0, 0, SizeX, SizeY));
 				{
 					// Make sure the Canvas is not rendered upside down
 					Canvas.SetAllowSwitchVerticalAxis(false);
