Changed calls to FProcHandle.Close() to FPlatformProcess::CloseProc().

We already had FPlatformProcess::CreateProc/CloseProc pair (which is part of larger set with WaitForProc(), TerminateProc(), GetProcReturnCode(), etc). FProcHandle::Close() made the API more confusing (and actually introduced bugs, because it did not call CloseProc() until recently).

diff --git a/Engine/Plugins/Developer/KDevelopSourceCodeAccess/Source/KDevelopSourceCodeAccess/Private/KDevelopSourceCodeAccessor.cpp b/Engine/Plugins/Developer/KDevelopSourceCodeAccess/Source/KDevelopSourceCodeAccess/Private/KDevelopSourceCodeAccessor.cpp
--- a/Engine/Plugins/Developer/KDevelopSourceCodeAccess/Source/KDevelopSourceCodeAccess/Private/KDevelopSourceCodeAccessor.cpp
+++ b/Engine/Plugins/Developer/KDevelopSourceCodeAccess/Source/KDevelopSourceCodeAccess/Private/KDevelopSourceCodeAccessor.cpp
@@ -45,9 +45,13 @@ bool FKDevelopSourceCodeAccessor::OpenSolution()
 		return false;
 	}
 	
-	// FIXME: leaks the process :(
 	FProcHandle Proc = FPlatformProcess::CreateProc(*IDEPath, *Solution, true, false, false, nullptr, 0, nullptr, nullptr);
-	return Proc.IsValid();
+	if (Proc.IsValid())
+	{
+		FPlatformProcess::CloseProc(Proc);
+		return true;
+	}
+	return false;
 }
 
 bool FKDevelopSourceCodeAccessor::CanRunKDevelop(FString& OutPath) const
diff --git a/Engine/Plugins/Developer/VisualStudioSourceCodeAccess/Source/VisualStudioSourceCodeAccess/Private/VisualStudioSourceCodeAccessor.cpp b/Engine/Plugins/Developer/VisualStudioSourceCodeAccess/Source/VisualStudioSourceCodeAccess/Private/VisualStudioSourceCodeAccessor.cpp
--- a/Engine/Plugins/Developer/VisualStudioSourceCodeAccess/Source/VisualStudioSourceCodeAccess/Private/VisualStudioSourceCodeAccessor.cpp
+++ b/Engine/Plugins/Developer/VisualStudioSourceCodeAccess/Source/VisualStudioSourceCodeAccess/Private/VisualStudioSourceCodeAccessor.cpp
@@ -1091,7 +1091,7 @@ bool FVisualStudioSourceCodeAccessor::RunVisualStudioAndOpenSolutionAndFiles(con
 
 	FProcHandle WorkerHandle = FPlatformProcess::CreateProc(*ExecutablePath, *Params, true, false, false, nullptr, 0, nullptr, nullptr);
 	bool bSuccess = WorkerHandle.IsValid();
-	WorkerHandle.Close();
+	FPlatformProcess::CloseProc(WorkerHandle);
 	return bSuccess;
 }
 
diff --git a/Engine/Source/Developer/Android/AndroidTargetPlatform/Private/AndroidTargetDevice.inl b/Engine/Source/Developer/Android/AndroidTargetPlatform/Private/AndroidTargetDevice.inl
--- a/Engine/Source/Developer/Android/AndroidTargetPlatform/Private/AndroidTargetDevice.inl
+++ b/Engine/Source/Developer/Android/AndroidTargetPlatform/Private/AndroidTargetDevice.inl
@@ -15,7 +15,7 @@ inline bool FAndroidTargetDevice::Deploy( const FString& SourceFolder, FString&
 	FProcHandle RunningProc = FPlatformProcess::CreateProc(TEXT("../DotNet/Android/AndroidPackager"), *RepackageCommand, true, false, false, NULL, 0, TEXT("../DotNET/"), NULL);
 	FPlatformProcess::WaitForProc(RunningProc);
 	FPlatformProcess::GetProcReturnCode(RunningProc, &ReturnCode);
-	RunningProc.Close();
+	FPlatformProcess::CloseProc(RunningProc);
 
 	return 0 == ReturnCode;
 }
diff --git a/Engine/Source/Developer/DesktopPlatform/Private/DesktopPlatformBase.cpp b/Engine/Source/Developer/DesktopPlatform/Private/DesktopPlatformBase.cpp
--- a/Engine/Source/Developer/DesktopPlatform/Private/DesktopPlatformBase.cpp
+++ b/Engine/Source/Developer/DesktopPlatform/Private/DesktopPlatformBase.cpp
@@ -1019,7 +1019,7 @@ bool FDesktopPlatformBase::BuildUnrealBuildTool(const FString& RootDir, FOutputD
 		return false;
 	}
 	FPlatformProcess::WaitForProc(ProcHandle);
-	ProcHandle.Close();
+	FPlatformProcess::CloseProc(ProcHandle);
 
 	// If the executable appeared where we expect it, then we were successful
 	FString UnrealBuildToolExePath = GetUnrealBuildToolExecutableFilename(RootDir);
diff --git a/Engine/Source/Developer/HotReload/Private/HotReload.cpp b/Engine/Source/Developer/HotReload/Private/HotReload.cpp
--- a/Engine/Source/Developer/HotReload/Private/HotReload.cpp
+++ b/Engine/Source/Developer/HotReload/Private/HotReload.cpp
@@ -1433,7 +1433,7 @@ void FHotReloadModule::CheckForFinishedModuleDLLCompile(const bool bWaitForCompl
 			ModulesThatWereBeingRecompiled.Empty();
 
 			// We're done with the process handle now
-			ModuleCompileProcessHandle.Close();
+			FPlatformProcess::CloseProc(ModuleCompileProcessHandle);
 			ModuleCompileProcessHandle.Reset();
 
 			FPlatformProcess::ClosePipe(ModuleCompileReadPipe, 0);
diff --git a/Engine/Source/Developer/Linux/LinuxTargetPlatform/Private/LinuxTargetDevice.cpp b/Engine/Source/Developer/Linux/LinuxTargetPlatform/Private/LinuxTargetDevice.cpp
--- a/Engine/Source/Developer/Linux/LinuxTargetPlatform/Private/LinuxTargetDevice.cpp
+++ b/Engine/Source/Developer/Linux/LinuxTargetPlatform/Private/LinuxTargetDevice.cpp
@@ -81,7 +81,12 @@ bool FLinuxTargetDevice::Launch( const FString& AppId, EBuildConfigurations::Typ
 
 	// launch the game
 	FProcHandle ProcessHandle = FPlatformProcess::CreateProc(*ExecutablePath, *Params, true, false, false, OutProcessId, 0, NULL, NULL);
-	return ProcessHandle.Close();
+	if (ProcessHandle.IsValid())
+	{
+		FPlatformProcess::CloseProc(ProcessHandle);
+		return true;
+	}
+	return false;
 #else
 	// @todo: support launching on a remote machine
 	STUBBED("FLinuxTargetDevice::Launch");
@@ -99,7 +104,12 @@ bool FLinuxTargetDevice::Run( const FString& ExecutablePath, const FString& Para
 {
 #if PLATFORM_LINUX	// if running natively, support simplified, local deployment	
 	FProcHandle ProcessHandle = FPlatformProcess::CreateProc(*ExecutablePath, *Params, true, false, false, OutProcessId, 0, NULL, NULL);
-	return ProcessHandle.Close();
+	if (ProcessHandle.IsValid())
+	{
+		FPlatformProcess::CloseProc(ProcessHandle);
+		return true;
+	}
+	return false;
 #else
 	// @todo: support remote run
 	STUBBED("FLinuxTargetDevice::Run");
diff --git a/Engine/Source/Developer/Mac/MacTargetPlatform/Private/LocalMacTargetDevice.h b/Engine/Source/Developer/Mac/MacTargetPlatform/Private/LocalMacTargetDevice.h
--- a/Engine/Source/Developer/Mac/MacTargetPlatform/Private/LocalMacTargetDevice.h
+++ b/Engine/Source/Developer/Mac/MacTargetPlatform/Private/LocalMacTargetDevice.h
@@ -116,7 +116,12 @@ public:
 
 		// launch the game
 		FProcHandle ProcessHandle = FPlatformProcess::CreateProc(*ExecutablePath, *Params, true, false, false, OutProcessId, 0, NULL, NULL);
-		return ProcessHandle.Close();
+		if (ProcessHandle.IsValid())
+		{
+			FPlatformProcess::CloseProc(ProcessHandle);
+			return true;
+		}
+		return false;
 	}
 
 	virtual bool PowerOff( bool Force ) override
@@ -142,7 +147,12 @@ public:
 	virtual bool Run( const FString& ExecutablePath, const FString& Params, uint32* OutProcessId ) override
 	{
 		FProcHandle ProcessHandle = FPlatformProcess::CreateProc(*ExecutablePath, *Params, true, false, false, OutProcessId, 0, NULL, NULL);
-		return ProcessHandle.Close();
+		if (ProcessHandle.IsValid())
+		{
+			FPlatformProcess::CloseProc(ProcessHandle);
+			return true;
+		}
+		return false;
 	}
 
 	virtual bool SupportsFeature( ETargetDeviceFeatures Feature ) const override
diff --git a/Engine/Source/Developer/ProjectLauncher/Private/Widgets/Build/SProjectLauncherBuildPage.cpp b/Engine/Source/Developer/ProjectLauncher/Private/Widgets/Build/SProjectLauncherBuildPage.cpp
--- a/Engine/Source/Developer/ProjectLauncher/Private/Widgets/Build/SProjectLauncherBuildPage.cpp
+++ b/Engine/Source/Developer/ProjectLauncher/Private/Widgets/Build/SProjectLauncherBuildPage.cpp
@@ -158,8 +158,12 @@ bool SProjectLauncherBuildPage::GenerateDSYMForProject( const FString& ProjectNa
 
     // launch the builder and monitor its process
     FProcHandle ProcessHandle = FPlatformProcess::CreateProc(*(ExecutablePath / Executable), *CommandLine, false, false, false, NULL, 0, *ExecutablePath, NULL);
-
-	return ProcessHandle.Close();
+	if (ProcessHandle.IsValid())
+	{
+		FPlatformProcess::CloseProc(ProcessHandle);
+		return true;
+	}
+	return false;
 }
 
 
diff --git a/Engine/Source/Developer/TextureFormatPVR/Private/TextureFormatPVR.cpp b/Engine/Source/Developer/TextureFormatPVR/Private/TextureFormatPVR.cpp
--- a/Engine/Source/Developer/TextureFormatPVR/Private/TextureFormatPVR.cpp
+++ b/Engine/Source/Developer/TextureFormatPVR/Private/TextureFormatPVR.cpp
@@ -480,7 +480,7 @@ class FTextureFormatPVR : public ITextureFormat
 			{
 				FPlatformProcess::Sleep(0.01f);
 			}
-			Proc.Close();
+			FPlatformProcess::CloseProc(Proc);
 
 			// Did it fail?
 			if ( ReturnCode != 0 )
diff --git a/Engine/Source/Developer/Windows/WindowsTargetPlatform/Private/LocalPcTargetDevice.h b/Engine/Source/Developer/Windows/WindowsTargetPlatform/Private/LocalPcTargetDevice.h
--- a/Engine/Source/Developer/Windows/WindowsTargetPlatform/Private/LocalPcTargetDevice.h
+++ b/Engine/Source/Developer/Windows/WindowsTargetPlatform/Private/LocalPcTargetDevice.h
@@ -257,7 +257,12 @@ public:
 
 		// launch the game
 		FProcHandle ProcessHandle = FPlatformProcess::CreateProc(*ExecutablePath, *Params, true, false, false, OutProcessId, 0, nullptr, nullptr);
-		return ProcessHandle.Close();
+		if (ProcessHandle.IsValid())
+		{
+			FPlatformProcess::CloseProc(ProcessHandle);
+			return true;
+		}
+		return false;
 	}
 
 	virtual bool PowerOff( bool Force ) override
@@ -288,7 +293,12 @@ public:
 	virtual bool Run( const FString& ExecutablePath, const FString& Params, uint32* OutProcessId )
 	{
 		FProcHandle ProcessHandle = FPlatformProcess::CreateProc(*ExecutablePath, *Params, true, false, false, OutProcessId, 0, nullptr, nullptr);
-		return ProcessHandle.Close();
+		if (ProcessHandle.IsValid())
+		{
+			FPlatformProcess::CloseProc(ProcessHandle);
+			return true;
+		}
+		return false;
 	}
 
 	virtual bool SupportsFeature( ETargetDeviceFeatures Feature ) const override
diff --git a/Engine/Source/Editor/SwarmInterface/Private/SwarmInterfaceLocal.cpp b/Engine/Source/Editor/SwarmInterface/Private/SwarmInterfaceLocal.cpp
--- a/Engine/Source/Editor/SwarmInterface/Private/SwarmInterfaceLocal.cpp
+++ b/Engine/Source/Editor/SwarmInterface/Private/SwarmInterfaceLocal.cpp
@@ -534,7 +534,12 @@ int32 FSwarmInterfaceLocalImpl::EndJobSpecification( void )
 	{
 		FString Parameters = FString(JobSpecification.Parameters) + (FApp::IsEngineInstalled() ? TEXT(" -installed") : TEXT(""));
 		LightmassProcHandle = FPlatformProcess::CreateProc( JobSpecification.ExecutableName, *Parameters, true, false, false, NULL, 0, NULL, NULL);
-		return LightmassProcHandle.Close() ? 0 : SWARM_ERROR_CONNECTION_DISCONNECTED;
+		if (LightmassProcHandle.IsValid())
+		{
+			FPlatformProcess::CloseProc(LightmassProcHandle);
+			return 0;
+		}
+		return SWARM_ERROR_CONNECTION_DISCONNECTED;
 	}
 #endif
 	return 0;
diff --git a/Engine/Source/Editor/UnrealEd/Private/PlayLevel.cpp b/Engine/Source/Editor/UnrealEd/Private/PlayLevel.cpp
--- a/Engine/Source/Editor/UnrealEd/Private/PlayLevel.cpp
+++ b/Engine/Source/Editor/UnrealEd/Private/PlayLevel.cpp
@@ -1764,7 +1764,7 @@ void UEditorEngine::PlayForMovieCapture()
 	{
 		UE_LOG(LogPlayLevel, Error,  TEXT("Failed to run a copy of the game for matinee capture."));
 	}
-	ProcessHandle.Close();
+	FPlatformProcess::CloseProc(ProcessHandle);
 }
 
 
diff --git a/Engine/Source/Editor/UnrealEd/Private/UnrealEdMisc.cpp b/Engine/Source/Editor/UnrealEd/Private/UnrealEdMisc.cpp
--- a/Engine/Source/Editor/UnrealEd/Private/UnrealEdMisc.cpp
+++ b/Engine/Source/Editor/UnrealEd/Private/UnrealEdMisc.cpp
@@ -836,7 +836,7 @@ void FUnrealEdMisc::OnExit()
 
 			return;
 		}
-		Handle.Close();
+		FPlatformProcess::CloseProc(Handle);
 	}
 }
 
diff --git a/Engine/Source/Runtime/Engine/Private/ShaderCompiler/ShaderCompiler.cpp b/Engine/Source/Runtime/Engine/Private/ShaderCompiler/ShaderCompiler.cpp
--- a/Engine/Source/Runtime/Engine/Private/ShaderCompiler/ShaderCompiler.cpp
+++ b/Engine/Source/Runtime/Engine/Private/ShaderCompiler/ShaderCompiler.cpp
@@ -415,7 +415,7 @@ struct FShaderCompileWorkerInfo
 		if(WorkerProcess.IsValid())
 		{
 			FPlatformProcess::TerminateProc(WorkerProcess);
-			WorkerProcess.Close();
+			FPlatformProcess::CloseProc(WorkerProcess);
 		}
 	}
 
@@ -2611,7 +2611,7 @@ FShaderCompileXGEThreadRunnable::~FShaderCompileXGEThreadRunnable()
 		// We still have a build in progress.
 		// Kill it...
 		FPlatformProcess::TerminateProc(BuildProcessHandle);
-		BuildProcessHandle.Close();
+		FPlatformProcess::CloseProc(BuildProcessHandle);
 	}
 
 	// Clean up any intermediate files/directories we've got left over.
@@ -2935,7 +2935,7 @@ int32 FShaderCompileXGEThreadRunnable::CompilingLoop()
 			}
 
 			ShaderBatchesInFlight.Empty();
-			BuildProcessHandle.Close();
+			FPlatformProcess::CloseProc(BuildProcessHandle);
 		}
 
 		bWorkRemaining |= ShaderBatchesInFlight.Num() > 0;
@@ -3026,7 +3026,7 @@ int32 FShaderCompileXGEThreadRunnable::CompilingLoop()
 			// Start up the build monitor process to monitor for engine crashes.
 			uint32 BuildMonitorProcessID;
 			FProcHandle BuildMonitorHandle = FPlatformProcess::CreateProc(*Manager->ShaderCompileWorkerName, *FString::Printf(TEXT("-xgemonitor %d %d"), Manager->ProcessId, BuildProcessID), true, false, false, &BuildMonitorProcessID, 0, nullptr, nullptr);
-			BuildMonitorHandle.Close();
+			FPlatformProcess::CloseProc(BuildMonitorHandle);
 
 			// Reset batch counters and switch directories
 			BatchIndexToFill = 0;
