UE4 - reduce ticking overhead and optimize the task graph

diff --git a/Engine/Source/Runtime/Core/Private/Async/TaskGraph.cpp b/Engine/Source/Runtime/Core/Private/Async/TaskGraph.cpp
--- a/Engine/Source/Runtime/Core/Private/Async/TaskGraph.cpp
+++ b/Engine/Source/Runtime/Core/Private/Async/TaskGraph.cpp
@@ -1204,10 +1204,10 @@ FTaskGraphInterface& FTaskGraphInterface::Get()
 
 // Statics and some implementations from FBaseGraphTask and FGraphEvent
 
-TLockFreeFixedSizeAllocator<FBaseGraphTask::SMALL_TASK_SIZE>& FBaseGraphTask::GetSmallTaskAllocator()
+static FBaseGraphTask::TSmallTaskAllocator TheSmallTaskAllocator;
+FBaseGraphTask::TSmallTaskAllocator& FBaseGraphTask::GetSmallTaskAllocator()
 {
-	static TLockFreeFixedSizeAllocator<FBaseGraphTask::SMALL_TASK_SIZE> TheAllocator;
-	return TheAllocator;
+	return TheSmallTaskAllocator;
 }
 
 #if !(UE_BUILD_SHIPPING || UE_BUILD_TEST)
@@ -1217,15 +1217,11 @@ void FBaseGraphTask::LogPossiblyInvalidSubsequentsTask(const TCHAR* TaskName)
 }
 #endif
 
-TLockFreeClassAllocator<FGraphEvent>& FGraphEvent::GetAllocator()
-{
-	static TLockFreeClassAllocator<FGraphEvent> TheAllocator;
-	return TheAllocator;
-}
+static TLockFreeClassAllocator_TLSCache<FGraphEvent> TheGraphEventAllocator;
 
 FGraphEventRef FGraphEvent::CreateGraphEvent()
 {
-	return GetAllocator().New();
+	return TheGraphEventAllocator.New();
 }
 
 void FGraphEvent::DispatchSubsequents(ENamedThreads::Type CurrentThreadIfKnown)
@@ -1263,7 +1259,7 @@ void FGraphEvent::DispatchSubsequents(TArray<FBaseGraphTask*>& NewTasks, ENamedT
 
 void FGraphEvent::Recycle(FGraphEvent* ToRecycle)
 {
-	GetAllocator().Free(ToRecycle);
+	TheGraphEventAllocator.Free(ToRecycle);
 }
 
 FGraphEvent::~FGraphEvent()
diff --git a/Engine/Source/Runtime/Core/Private/Containers/LockFreeList.cpp b/Engine/Source/Runtime/Core/Private/Containers/LockFreeList.cpp
--- a/Engine/Source/Runtime/Core/Private/Containers/LockFreeList.cpp
+++ b/Engine/Source/Runtime/Core/Private/Containers/LockFreeList.cpp
@@ -58,34 +58,36 @@ FLockFreeVoidPointerListBase128::FLinkAllocator::~FLinkAllocator()
 	FLockFreeVoidPointerListBase
 -----------------------------------------------------------------------------*/
 
-FLockFreeVoidPointerListBase::FLinkAllocator& FLockFreeVoidPointerListBase::FLinkAllocator::Get()
-{
-	static FLinkAllocator TheAllocator;
-	return TheAllocator;
-}
+FLockFreeVoidPointerListBase::FLinkAllocator FLockFreeVoidPointerListBase::FLinkAllocator::TheLinkAllocator;
 
 FLockFreeVoidPointerListBase::FLinkAllocator::FLinkAllocator()
-	: FreeLinks(NULL)
-	, SpecialClosedLink(new FLink())
+	: SpecialClosedLink(new FLink())
 {
-	ClosedLink()->LockCount.Increment();
+	static_assert(sizeof(FLink) >= sizeof(void*) && sizeof(FLink) % sizeof(void*) == 0, "Blocks in TLockFreeFixedSizeAllocator must be at least the size of a pointer.");
+	check(IsInGameThread());
+	TlsSlot = FPlatformTLS::AllocTlsSlot();
+	check(TlsSlot);
 }
-
 FLockFreeVoidPointerListBase::FLinkAllocator::~FLinkAllocator()
 {
-	// - Deliberately leak to allow LockFree to be used during shutdown
-#if	0
-	check(SpecialClosedLink); // double free?
-	check(!NumUsedLinks.GetValue());
-	while (FLink* ToDelete = FLink::Unlink(&FreeLinks))
+	FPlatformTLS::FreeTlsSlot(TlsSlot);
+	TlsSlot = 0;
+}
+
+void FLockFreeVoidPointerListBase::FLinkAllocator::PutFreeBundle(void **Bundle)
+{
+	FScopeLock Lock(&BundleArrayCriticalSection);
+	BundleArray.Push(Bundle);
+}
+void** FLockFreeVoidPointerListBase::FLinkAllocator::GetFreeBundle()
+{
+	FScopeLock Lock(&BundleArrayCriticalSection);
+	void** Result = nullptr;
+	if (BundleArray.Num())
 	{
-		delete ToDelete;
-		NumFreeLinks.Decrement();
+		Result = BundleArray.Pop(false);
 	}
-	check(!NumFreeLinks.GetValue());
-	delete SpecialClosedLink;
-	SpecialClosedLink = 0;
-#endif // 0
+	return Result;
 }
 
 #endif //USE_LOCKFREELIST_128
diff --git a/Engine/Source/Runtime/Core/Public/Async/TaskGraphInterfaces.h b/Engine/Source/Runtime/Core/Public/Async/TaskGraphInterfaces.h
--- a/Engine/Source/Runtime/Core/Public/Async/TaskGraphInterfaces.h
+++ b/Engine/Source/Runtime/Core/Public/Async/TaskGraphInterfaces.h
@@ -209,6 +209,14 @@ public:
  **/
 class FBaseGraphTask
 {
+public:
+	// Allocator for small tasks.
+	enum
+	{
+		/** Total size in bytes for a small task that will use the custom allocator **/
+		SMALL_TASK_SIZE = 256
+	};
+	typedef TLockFreeFixedSizeAllocator_TLSCache<SMALL_TASK_SIZE> TSmallTaskAllocator;
 protected:
 	/** 
 	 *	Constructor
@@ -233,11 +241,13 @@ protected:
 	/** 
 	 *	Indicates that the prerequisites are set up and that the task can be executed as soon as the prerequisites are finished.
 	 *	@param NumAlreadyFinishedPrequistes; the number of prerequisites that have not been set up because those tasks had already completed.
+	 *	@param bUnlock; if true, let the task execute if it can
 	 **/
-	void PrerequisitesComplete(ENamedThreads::Type CurrentThread, int32 NumAlreadyFinishedPrequistes)
+	void PrerequisitesComplete(ENamedThreads::Type CurrentThread, int32 NumAlreadyFinishedPrequistes, bool bUnlock = true)
 	{
 		checkThreadGraph(LifeStage.Increment() == int32(LS_PrequisitesSetup));
-		if (NumberOfPrerequistitesOutstanding.Subtract(NumAlreadyFinishedPrequistes + 1) == NumAlreadyFinishedPrequistes + 1) // the +1 is for the "lock" we set up in the constructor
+		int32 NumToSub = NumAlreadyFinishedPrequistes + (bUnlock ? 1 : 0); // the +1 is for the "lock" we set up in the constructor
+		if (NumberOfPrerequistitesOutstanding.Subtract(NumToSub) == NumToSub) 
 		{
 			QueueTask(CurrentThread);
 		}
@@ -253,14 +263,20 @@ protected:
 	static void CORE_API LogPossiblyInvalidSubsequentsTask(const TCHAR* TaskName);
 #endif
 
-	// Allocator for small tasks.
-	enum
-	{
-		/** Total size in bytes for a small task that will use the custom allocator **/
-		SMALL_TASK_SIZE = 256
-	};
 	/** Singleton to retrieve the small task allocator **/
-	static CORE_API TLockFreeFixedSizeAllocator<SMALL_TASK_SIZE>& GetSmallTaskAllocator();
+	static CORE_API TSmallTaskAllocator& GetSmallTaskAllocator();
+
+	/** 
+	 *	An indication that a prerequisite has been completed. Reduces the number of prerequisites by one and if no prerequisites are outstanding, it queues the task for execution.
+	 *	@param CurrentThread; provides the index of the thread we are running on. This is handy for submitting new taks. Can be ENamedThreads::AnyThread if the current thread is unknown.
+	 **/
+	void ConditionalQueueTask(ENamedThreads::Type CurrentThread)
+	{
+		if (NumberOfPrerequistitesOutstanding.Decrement()==0)
+		{
+			QueueTask(CurrentThread);
+		}
+	}
 
 private:
 	friend class FTaskThread;
@@ -277,18 +293,6 @@ private:
 
 	// API called from other parts of the system
 
-	/** 
-	 *	An indication that a prerequisite has been completed. Reduces the number of prerequisites by one and if no prerequisites are outstanding, it queues the task for execution.
-	 *	@param CurrentThread; provides the index of the thread we are running on. This is handy for submitting new taks. Can be ENamedThreads::AnyThread if the current thread is unknown.
-	 **/
-	void ConditionalQueueTask(ENamedThreads::Type CurrentThread)
-	{
-		if (NumberOfPrerequistitesOutstanding.Decrement()==0)
-		{
-			QueueTask(CurrentThread);
-		}
-	}
-
 	/** 
 	 *	Called by the system to execute this task after it has been removed from an internal queue.
 	 *	Just checks the life stage and passes off to the virtual ExecuteTask method.
@@ -397,13 +401,7 @@ public:
 	}
 private:
 	friend class TRefCountPtr<FGraphEvent>;
-	friend class TLockFreeClassAllocator<FGraphEvent>;
-
-	/**
-	 *	Internal function to retrieve the graph event allocation singleton
-	 *	@return a reference to the graph event allocation singleton
-	**/
-	static CORE_API TLockFreeClassAllocator<FGraphEvent>& GetAllocator();
+	friend class TLockFreeClassAllocator_TLSCache<FGraphEvent>;
 
 	/** 
 	 *	Internal function to call the destructor and recycle a graph event
@@ -520,75 +518,128 @@ public:
 	class FConstructor
 	{
 	public:
-		/** Passthrough internal task constructor and dispatch. Note! Generally speaking references will not pass through; use pointers */
+		/** Passthrough internal task constructor and dispatch. */
 		FGraphEventRef ConstructAndDispatchWhenReady()
 		{
 			new ((void *)&Owner->TaskStorage) TTask();
 			return Owner->Setup(Prerequisites, CurrentThreadIfKnown);
 		}
-		/** Passthrough internal task constructor and dispatch. Note! Generally speaking references will not pass through; use pointers */
 		template<typename T>
 		FGraphEventRef ConstructAndDispatchWhenReady(T&& Arg1)
 		{
 			new ((void *)&Owner->TaskStorage) TTask(Forward<T>(Arg1));
 			return Owner->Setup(Prerequisites, CurrentThreadIfKnown);
 		}
-		/** Passthrough internal task constructor and dispatch. Note! Generally speaking references will not pass through; use pointers */
 		template<typename T1,typename T2>
 		FGraphEventRef ConstructAndDispatchWhenReady(T1&& Arg1, T2&& Arg2)
 		{
 			new ((void *)&Owner->TaskStorage) TTask(Forward<T1>(Arg1), Forward<T2>(Arg2));
 			return Owner->Setup(Prerequisites, CurrentThreadIfKnown);
 		}
-		/** Passthrough internal task constructor and dispatch. Note! Generally speaking references will not pass through; use pointers */
 		template<typename T1,typename T2, typename T3>
 		FGraphEventRef ConstructAndDispatchWhenReady(T1&& Arg1, T2&& Arg2, T3&& Arg3)
 		{
 			new ((void *)&Owner->TaskStorage) TTask(Forward<T1>(Arg1), Forward<T2>(Arg2), Forward<T3>(Arg3));
 			return Owner->Setup(Prerequisites, CurrentThreadIfKnown);
 		}
-		/** Passthrough internal task constructor and dispatch. Note! Generally speaking references will not pass through; use pointers */
 		template<typename T1,typename T2, typename T3, typename T4>
 		FGraphEventRef ConstructAndDispatchWhenReady(T1&& Arg1, T2&& Arg2, T3&& Arg3, T4&& Arg4)
 		{
 			new ((void *)&Owner->TaskStorage) TTask(Forward<T1>(Arg1), Forward<T2>(Arg2), Forward<T3>(Arg3), Forward<T4>(Arg4));
 			return Owner->Setup(Prerequisites, CurrentThreadIfKnown);
 		}
-		/** Passthrough internal task constructor and dispatch. Note! Generally speaking references will not pass through; use pointers */
 		template<typename T1,typename T2, typename T3, typename T4, typename T5>
 		FGraphEventRef ConstructAndDispatchWhenReady(T1&& Arg1, T2&& Arg2, T3&& Arg3, T4&& Arg4, T5&& Arg5)
 		{
 			new ((void *)&Owner->TaskStorage) TTask(Forward<T1>(Arg1), Forward<T2>(Arg2), Forward<T3>(Arg3), Forward<T4>(Arg4), Forward<T5>(Arg5));
 			return Owner->Setup(Prerequisites, CurrentThreadIfKnown);
 		}
-		/** Passthrough internal task constructor and dispatch. Note! Generally speaking references will not pass through; use pointers */
 		template<typename T1, typename T2, typename T3, typename T4, typename T5, typename T6>
 		FGraphEventRef ConstructAndDispatchWhenReady(T1&& Arg1, T2&& Arg2, T3&& Arg3, T4&& Arg4, T5&& Arg5, T6&& Arg6)
 		{
 			new ((void *)&Owner->TaskStorage) TTask(Forward<T1>(Arg1), Forward<T2>(Arg2), Forward<T3>(Arg3), Forward<T4>(Arg4), Forward<T5>(Arg5), Forward<T6>(Arg6));
 			return Owner->Setup(Prerequisites, CurrentThreadIfKnown);
 		}
-		/** Passthrough internal task constructor and dispatch. Note! Generally speaking references will not pass through; use pointers */
 		template<typename T1, typename T2, typename T3, typename T4, typename T5, typename T6, typename T7>
 		FGraphEventRef ConstructAndDispatchWhenReady(T1&& Arg1, T2&& Arg2, T3&& Arg3, T4&& Arg4, T5&& Arg5, T6&& Arg6, T7&& Arg7)
 		{
 			new ((void *)&Owner->TaskStorage) TTask(Forward<T1>(Arg1), Forward<T2>(Arg2), Forward<T3>(Arg3), Forward<T4>(Arg4), Forward<T5>(Arg5), Forward<T6>(Arg6), Forward<T7>(Arg7));
 			return Owner->Setup(Prerequisites, CurrentThreadIfKnown);
 		}
-		/** Passthrough internal task constructor and dispatch. Note! Generally speaking references will not pass through; use pointers */
 		template<typename T1, typename T2, typename T3, typename T4, typename T5, typename T6, typename T7, typename T8>
 		FGraphEventRef ConstructAndDispatchWhenReady(T1&& Arg1, T2&& Arg2, T3&& Arg3, T4&& Arg4, T5&& Arg5, T6&& Arg6, T7&& Arg7, T8&& Arg8)
 		{
 			new ((void *)&Owner->TaskStorage) TTask(Forward<T1>(Arg1), Forward<T2>(Arg2), Forward<T3>(Arg3), Forward<T4>(Arg4), Forward<T5>(Arg5), Forward<T6>(Arg6), Forward<T7>(Arg7), Forward<T8>(Arg8));
 			return Owner->Setup(Prerequisites, CurrentThreadIfKnown);
 		}
-		/** Passthrough internal task constructor and dispatch. Note! Generally speaking references will not pass through; use pointers */
 		template<typename T1, typename T2, typename T3, typename T4, typename T5, typename T6, typename T7, typename T8, typename T9>
 		FGraphEventRef ConstructAndDispatchWhenReady(T1&& Arg1, T2&& Arg2, T3&& Arg3, T4&& Arg4, T5&& Arg5, T6&& Arg6, T7&& Arg7, T8&& Arg8, T9&& Arg9)
 		{
 			new ((void *)&Owner->TaskStorage) TTask(Forward<T1>(Arg1), Forward<T2>(Arg2), Forward<T3>(Arg3), Forward<T4>(Arg4), Forward<T5>(Arg5), Forward<T6>(Arg6), Forward<T7>(Arg7), Forward<T8>(Arg8), Forward<T9>(Arg9));
 			return Owner->Setup(Prerequisites, CurrentThreadIfKnown);
 		}
+
+		/** Passthrough internal task constructor and hold. */
+		TGraphTask* ConstructAndHold()
+		{
+			new ((void *)&Owner->TaskStorage) TTask();
+			return Owner->Hold(Prerequisites, CurrentThreadIfKnown);
+		}
+		template<typename T>
+		TGraphTask* ConstructAndHold(T&& Arg1)
+		{
+			new ((void *)&Owner->TaskStorage) TTask(Forward<T>(Arg1));
+			return Owner->Hold(Prerequisites, CurrentThreadIfKnown);
+		}
+		template<typename T1,typename T2>
+		TGraphTask* ConstructAndHold(T1&& Arg1, T2&& Arg2)
+		{
+			new ((void *)&Owner->TaskStorage) TTask(Forward<T1>(Arg1), Forward<T2>(Arg2));
+			return Owner->Hold(Prerequisites, CurrentThreadIfKnown);
+		}
+		template<typename T1,typename T2, typename T3>
+		TGraphTask* ConstructAndHold(T1&& Arg1, T2&& Arg2, T3&& Arg3)
+		{
+			new ((void *)&Owner->TaskStorage) TTask(Forward<T1>(Arg1), Forward<T2>(Arg2), Forward<T3>(Arg3));
+			return Owner->Hold(Prerequisites, CurrentThreadIfKnown);
+		}
+		template<typename T1,typename T2, typename T3, typename T4>
+		TGraphTask* ConstructAndHold(T1&& Arg1, T2&& Arg2, T3&& Arg3, T4&& Arg4)
+		{
+			new ((void *)&Owner->TaskStorage) TTask(Forward<T1>(Arg1), Forward<T2>(Arg2), Forward<T3>(Arg3), Forward<T4>(Arg4));
+			return Owner->Hold(Prerequisites, CurrentThreadIfKnown);
+		}
+		template<typename T1,typename T2, typename T3, typename T4, typename T5>
+		TGraphTask* ConstructAndHold(T1&& Arg1, T2&& Arg2, T3&& Arg3, T4&& Arg4, T5&& Arg5)
+		{
+			new ((void *)&Owner->TaskStorage) TTask(Forward<T1>(Arg1), Forward<T2>(Arg2), Forward<T3>(Arg3), Forward<T4>(Arg4), Forward<T5>(Arg5));
+			return Owner->Hold(Prerequisites, CurrentThreadIfKnown);
+		}
+		template<typename T1, typename T2, typename T3, typename T4, typename T5, typename T6>
+		TGraphTask* ConstructAndHold(T1&& Arg1, T2&& Arg2, T3&& Arg3, T4&& Arg4, T5&& Arg5, T6&& Arg6)
+		{
+			new ((void *)&Owner->TaskStorage) TTask(Forward<T1>(Arg1), Forward<T2>(Arg2), Forward<T3>(Arg3), Forward<T4>(Arg4), Forward<T5>(Arg5), Forward<T6>(Arg6));
+			return Owner->Hold(Prerequisites, CurrentThreadIfKnown);
+		}
+		template<typename T1, typename T2, typename T3, typename T4, typename T5, typename T6, typename T7>
+		TGraphTask* ConstructAndHold(T1&& Arg1, T2&& Arg2, T3&& Arg3, T4&& Arg4, T5&& Arg5, T6&& Arg6, T7&& Arg7)
+		{
+			new ((void *)&Owner->TaskStorage) TTask(Forward<T1>(Arg1), Forward<T2>(Arg2), Forward<T3>(Arg3), Forward<T4>(Arg4), Forward<T5>(Arg5), Forward<T6>(Arg6), Forward<T7>(Arg7));
+			return Owner->Hold(Prerequisites, CurrentThreadIfKnown);
+		}
+		template<typename T1, typename T2, typename T3, typename T4, typename T5, typename T6, typename T7, typename T8>
+		TGraphTask* ConstructAndHold(T1&& Arg1, T2&& Arg2, T3&& Arg3, T4&& Arg4, T5&& Arg5, T6&& Arg6, T7&& Arg7, T8&& Arg8)
+		{
+			new ((void *)&Owner->TaskStorage) TTask(Forward<T1>(Arg1), Forward<T2>(Arg2), Forward<T3>(Arg3), Forward<T4>(Arg4), Forward<T5>(Arg5), Forward<T6>(Arg6), Forward<T7>(Arg7), Forward<T8>(Arg8));
+			return Owner->Hold(Prerequisites, CurrentThreadIfKnown);
+		}
+		template<typename T1, typename T2, typename T3, typename T4, typename T5, typename T6, typename T7, typename T8, typename T9>
+		TGraphTask* ConstructAndHold(T1&& Arg1, T2&& Arg2, T3&& Arg3, T4&& Arg4, T5&& Arg5, T6&& Arg6, T7&& Arg7, T8&& Arg8, T9&& Arg9)
+		{
+			new ((void *)&Owner->TaskStorage) TTask(Forward<T1>(Arg1), Forward<T2>(Arg2), Forward<T3>(Arg3), Forward<T4>(Arg4), Forward<T5>(Arg5), Forward<T6>(Arg6), Forward<T7>(Arg7), Forward<T8>(Arg8), Forward<T9>(Arg9));
+			return Owner->Hold(Prerequisites, CurrentThreadIfKnown);
+		}
+
 	private:
 		friend class TGraphTask;
 
@@ -626,12 +677,23 @@ public:
 	**/
 	static FConstructor CreateTask(const FGraphEventArray* Prerequisites = NULL, ENamedThreads::Type CurrentThreadIfKnown = ENamedThreads::AnyThread)
 	{
+		int32 NumPrereq = Prerequisites ? Prerequisites->Num() : 0;
 		if (sizeof(TGraphTask) <= FBaseGraphTask::SMALL_TASK_SIZE)
 		{
 			void *Mem = FBaseGraphTask::GetSmallTaskAllocator().Allocate();
-			return FConstructor(new (Mem) TGraphTask(TTask::GetSubsequentsMode() == ESubsequentsMode::FireAndForget ? NULL : FGraphEvent::CreateGraphEvent(), Prerequisites ? Prerequisites->Num() : 0), Prerequisites, CurrentThreadIfKnown);
+			return FConstructor(new (Mem) TGraphTask(TTask::GetSubsequentsMode() == ESubsequentsMode::FireAndForget ? NULL : FGraphEvent::CreateGraphEvent(), NumPrereq), Prerequisites, CurrentThreadIfKnown);
 		}
-		return FConstructor(new TGraphTask(TTask::GetSubsequentsMode() == ESubsequentsMode::FireAndForget ? NULL : FGraphEvent::CreateGraphEvent(), Prerequisites ? Prerequisites->Num() : 0), Prerequisites, CurrentThreadIfKnown);
+		return FConstructor(new TGraphTask(TTask::GetSubsequentsMode() == ESubsequentsMode::FireAndForget ? NULL : FGraphEvent::CreateGraphEvent(), NumPrereq), Prerequisites, CurrentThreadIfKnown);
+	}
+
+	void Unlock(ENamedThreads::Type CurrentThreadIfKnown = ENamedThreads::AnyThread)
+	{
+		ConditionalQueueTask(CurrentThreadIfKnown);
+	}
+
+	FGraphEventRef GetCompletionEvent()
+	{
+		return Subsequents;
 	}
 
 private:
@@ -665,7 +727,8 @@ private:
 		{
 			FScopeCycleCounter Scope(Task.GetStatId(), true); 
 			Task.DoTask(CurrentThread, Subsequents);
-#if !(UE_BUILD_SHIPPING || UE_BUILD_TEST)
+#if (UE_BUILD_DEBUG || UE_BUILD_DEBUGGAME) && 0 // enable this to find nasty overwrite bugs
+			
 			if (TTask::GetSubsequentsMode() == ESubsequentsMode::TrackSubsequents)
 			{
 				if (!Subsequents->CheckSubsequents())
@@ -682,7 +745,7 @@ private:
 #endif
 
 			Task.~TTask();
-			checkSlow(ENamedThreads::GetThreadIndex(CurrentThread) <= ENamedThreads::RenderThread || FMemStack::Get().IsEmpty()); // you must mark and pop memstacks if you use them in tasks! Named threads are excepted.
+			checkThreadGraph(ENamedThreads::GetThreadIndex(CurrentThread) <= ENamedThreads::RenderThread || FMemStack::Get().IsEmpty()); // you must mark and pop memstacks if you use them in tasks! Named threads are excepted.
 		}
 		
 		TaskConstructed = false;
@@ -730,14 +793,14 @@ private:
 	 *	Call from FConstructor to complete the setup process
 	 *	@param Prerequisites; the list of FGraphEvents that must be completed prior to this task executing.
 	 *	@param CurrentThreadIfKnown; provides the index of the thread we are running on. Can be ENamedThreads::AnyThread if the current thread is unknown.
-	 *	@return A new graph event which represents the completion of this task.
+	 *	@param bUnlock; if true, task can execute right now if possible
 	 * 
 	 *	Create the completed event
 	 *	Set the thread to execute on based on the embedded task
 	 *	Attempt to add myself as a subsequent to each prerequisite
 	 *	Tell the base task that I am ready to start as soon as my prerequisites are ready.
 	 **/
-	FGraphEventRef Setup(const FGraphEventArray* Prerequisites = NULL, ENamedThreads::Type CurrentThreadIfKnown = ENamedThreads::AnyThread)
+	void SetupPrereqs(const FGraphEventArray* Prerequisites, ENamedThreads::Type CurrentThreadIfKnown, bool bUnlock)
 	{
 		checkThreadGraph(!TaskConstructed);
 		TaskConstructed = true;
@@ -755,11 +818,44 @@ private:
 				}
 			}
 		}
+		PrerequisitesComplete(CurrentThreadIfKnown, AlreadyCompletedPrerequisites, bUnlock);
+	}
+
+	/** 
+	 *	Call from FConstructor to complete the setup process
+	 *	@param Prerequisites; the list of FGraphEvents that must be completed prior to this task executing.
+	 *	@param CurrentThreadIfKnown; provides the index of the thread we are running on. Can be ENamedThreads::AnyThread if the current thread is unknown.
+	 *	@return A new graph event which represents the completion of this task.
+	 * 
+	 *	Create the completed event
+	 *	Set the thread to execute on based on the embedded task
+	 *	Attempt to add myself as a subsequent to each prerequisite
+	 *	Tell the base task that I am ready to start as soon as my prerequisites are ready.
+	 **/
+	FGraphEventRef Setup(const FGraphEventArray* Prerequisites = NULL, ENamedThreads::Type CurrentThreadIfKnown = ENamedThreads::AnyThread)
+	{
 		FGraphEventRef ReturnedEventRef = Subsequents; // very important so that this doesn't get destroyed before we return
-		PrerequisitesComplete(CurrentThreadIfKnown, AlreadyCompletedPrerequisites); // can execute now if all of the prereqs are done
+		SetupPrereqs(Prerequisites, CurrentThreadIfKnown, true);
 		return ReturnedEventRef;
 	}
 
+	/** 
+	 *	Call from FConstructor to complete the setup process, but doesn't allow the task to dispatch yet
+	 *	@param Prerequisites; the list of FGraphEvents that must be completed prior to this task executing.
+	 *	@param CurrentThreadIfKnown; provides the index of the thread we are running on. Can be ENamedThreads::AnyThread if the current thread is unknown.
+	 *	@return a pointer to the task
+	 * 
+	 *	Create the completed event
+	 *	Set the thread to execute on based on the embedded task
+	 *	Attempt to add myself as a subsequent to each prerequisite
+	 *	Tell the base task that I am ready to start as soon as my prerequisites are ready.
+	 **/
+	TGraphTask* Hold(const FGraphEventArray* Prerequisites = NULL, ENamedThreads::Type CurrentThreadIfKnown = ENamedThreads::AnyThread)
+	{
+		SetupPrereqs(Prerequisites, CurrentThreadIfKnown, false);
+		return this;
+	}
+
 	/** 
 	 *	Factory to create a gather task which assumes the given subsequent list from some other tasks.
 	 *	This is used to support "WaitFor" during a task execution.
diff --git a/Engine/Source/Runtime/Core/Public/Containers/LockFreeFixedSizeAllocator.h b/Engine/Source/Runtime/Core/Public/Containers/LockFreeFixedSizeAllocator.h
--- a/Engine/Source/Runtime/Core/Public/Containers/LockFreeFixedSizeAllocator.h
+++ b/Engine/Source/Runtime/Core/Public/Containers/LockFreeFixedSizeAllocator.h
@@ -4,13 +4,13 @@
 
 #include "LockFreeList.h"
 
-
 /**
  * Thread safe, lock free pooling allocator of fixed size blocks that
  * never returns free space until program shutdown.
+ * alignment isn't handled, assumes FMemory::Malloc will work
  */
 template<int32 SIZE, typename TTrackingCounter = FNoopCounter>
-class TLockFreeFixedSizeAllocator	// alignment isn't handled, assumes FMemory::Malloc will work
+class TLockFreeFixedSizeAllocator
 {
 public:
 
@@ -94,6 +94,169 @@ private:
 	TTrackingCounter NumFree;
 };
 
+/**
+ * Thread safe, lock free pooling allocator of fixed size blocks that
+ * never returns free space, even at shutdown
+ * alignment isn't handled, assumes FMemory::Malloc will work
+ */
+template<int32 SIZE, typename TTrackingCounter = FNoopCounter>
+class TLockFreeFixedSizeAllocator_TLSCache	
+{
+	enum
+	{
+		NUM_PER_BUNDLE=256,
+	};
+public:
+
+	TLockFreeFixedSizeAllocator_TLSCache()
+	{
+		static_assert(SIZE >= sizeof(void*) && SIZE % sizeof(void*) == 0, "Blocks in TLockFreeFixedSizeAllocator must be at least the size of a pointer.");
+		check(IsInGameThread());
+		TlsSlot = FPlatformTLS::AllocTlsSlot();
+		check(TlsSlot);
+	}
+	/** Destructor, leaks all of the memory **/
+	~TLockFreeFixedSizeAllocator_TLSCache()
+	{
+		FPlatformTLS::FreeTlsSlot(TlsSlot);
+		TlsSlot = 0;
+	}
+
+	/**
+	 * Allocates a memory block of size SIZE.
+	 *
+	 * @return Pointer to the allocated memory.
+	 * @see Free
+	 */
+	void* Allocate()
+	{
+		FThreadLocalCache& TLS = GetTLS();
+
+		if (!TLS.PartialBundle)
+		{
+			if (TLS.FullBundle)
+			{
+				TLS.PartialBundle = TLS.FullBundle;
+				TLS.FullBundle = nullptr;
+			}
+			else
+			{
+				TLS.PartialBundle = GlobalFreeListBundles.Pop();
+				if (!TLS.PartialBundle)
+				{
+					TLS.PartialBundle = (void**)FMemory::Malloc(SIZE * NUM_PER_BUNDLE);
+					void **Next = TLS.PartialBundle;
+					for (int32 Index = 0; Index < NUM_PER_BUNDLE - 1; Index++)
+					{
+						void* NextNext = (void*)(((uint8*)Next) + SIZE);
+						*Next = NextNext;
+						Next = (void**)NextNext;
+					}
+					*Next = nullptr;
+					NumFree.Add(NUM_PER_BUNDLE);
+				}
+			}
+			TLS.NumPartial = NUM_PER_BUNDLE;
+		}
+		NumUsed.Increment();
+		NumFree.Decrement();
+		void* Result = (void*)TLS.PartialBundle;
+		TLS.PartialBundle = (void**)*TLS.PartialBundle;
+		TLS.NumPartial--;
+		check(TLS.NumPartial >= 0 && ((!!TLS.NumPartial) == (!!TLS.PartialBundle)));
+		return Result;
+	}
+
+	/**
+	 * Puts a memory block previously obtained from Allocate() back on the free list for future use.
+	 *
+	 * @param Item The item to free.
+	 * @see Allocate
+	 */
+	void Free(void *Item)
+	{
+		NumUsed.Decrement();
+		NumFree.Increment();
+		FThreadLocalCache& TLS = GetTLS();
+		if (TLS.NumPartial >= NUM_PER_BUNDLE)
+		{
+			if (TLS.FullBundle)
+			{
+				GlobalFreeListBundles.Push(TLS.FullBundle);
+				//TLS.FullBundle = nullptr;
+			}
+			TLS.FullBundle = TLS.PartialBundle;
+			TLS.PartialBundle = nullptr;
+			TLS.NumPartial = 0;
+		}
+		*(void**)Item = (void*)TLS.PartialBundle;
+		TLS.PartialBundle = (void**)Item;
+		TLS.NumPartial++;
+	}
+
+	/**
+	 * Gets the number of allocated memory blocks that are currently in use.
+	 *
+	 * @return Number of used memory blocks.
+	 * @see GetNumFree
+	 */
+	const TTrackingCounter& GetNumUsed() const
+	{
+		return NumUsed;
+	}
+
+	/**
+	 * Gets the number of allocated memory blocks that are currently unused.
+	 *
+	 * @return Number of unused memory blocks.
+	 * @see GetNumUsed
+	 */
+	const TTrackingCounter& GetNumFree() const
+	{
+		return NumFree;
+	}
+
+private:
+
+	/** struct for the TLS cache. */
+	struct FThreadLocalCache
+	{
+		void **FullBundle;
+		void **PartialBundle;
+		int32 NumPartial;
+
+		FThreadLocalCache()
+			: FullBundle(nullptr)
+			, PartialBundle(nullptr)
+			, NumPartial(0)
+		{
+		}
+	};
+
+	FThreadLocalCache& GetTLS()
+	{
+		checkSlow(TlsSlot);
+		FThreadLocalCache* TLS = (FThreadLocalCache*)FPlatformTLS::GetTlsValue(TlsSlot);
+		if (!TLS)
+		{
+			TLS = new FThreadLocalCache();
+			FPlatformTLS::SetTlsValue(TlsSlot, TLS);
+		}
+		return *TLS;
+	}
+
+	/** Slot for TLS struct. */
+	uint32 TlsSlot;
+
+	/** Lock free list of free memory blocks, these are all linked into a bundle of NUM_PER_BUNDLE. */
+	TLockFreePointerList<void*> GlobalFreeListBundles;
+
+	/** Total number of blocks outstanding and not in the free list. */
+	TTrackingCounter NumUsed; 
+
+	/** Total number of blocks in the free list. */
+	TTrackingCounter NumFree;
+};
 
 /**
  * Thread safe, lock free pooling allocator of memory for instances of T.
@@ -104,7 +267,6 @@ template<class T>
 class TLockFreeClassAllocator : private TLockFreeFixedSizeAllocator<sizeof(T)>
 {
 public:
-
 	/**
 	 * Returns a memory block of size sizeof(T).
 	 *
@@ -139,3 +301,47 @@ public:
 		TLockFreeFixedSizeAllocator<sizeof(T)>::Free(Item);
 	}
 };
+
+/**
+ * Thread safe, lock free pooling allocator of memory for instances of T.
+ *
+ * Never returns free space until program shutdown.
+ */
+template<class T>
+class TLockFreeClassAllocator_TLSCache : private TLockFreeFixedSizeAllocator_TLSCache<sizeof(T)>
+{
+public:
+	/**
+	 * Returns a memory block of size sizeof(T).
+	 *
+	 * @return Pointer to the allocated memory.
+	 * @see Free, New
+	 */
+	void* Allocate()
+	{
+		return TLockFreeFixedSizeAllocator_TLSCache<sizeof(T)>::Allocate();
+	}
+
+	/**
+	 * Returns a new T using the default constructor.
+	 *
+	 * @return Pointer to the new object.
+	 * @see Allocate, Free
+	 */
+	T* New()
+	{
+		return new (Allocate()) T();
+	}
+
+	/**
+	 * Calls a destructor on Item and returns the memory to the free list for recycling.
+	 *
+	 * @param Item The item whose memory to free.
+	 * @see Allocate, New
+	 */
+	void Free(T *Item)
+	{
+		Item->~T();
+		TLockFreeFixedSizeAllocator_TLSCache<sizeof(T)>::Free(Item);
+	}
+};
diff --git a/Engine/Source/Runtime/Core/Public/Containers/LockFreeVoidPointerListBase.h b/Engine/Source/Runtime/Core/Public/Containers/LockFreeVoidPointerListBase.h
--- a/Engine/Source/Runtime/Core/Public/Containers/LockFreeVoidPointerListBase.h
+++ b/Engine/Source/Runtime/Core/Public/Containers/LockFreeVoidPointerListBase.h
@@ -260,13 +260,27 @@ private:
 		/** Part of the solution to the ABA problem. Links that are free still need locking, but we need to make sure that links are only added to the free list when they are transitioning to the free list. */
 		FThreadSafeCounter	MarkedForDeath;
 
-		/** Construcor, everything is intialized to zero. */
+		/** Constructor, everything is initialized to zero. */
 		FLink()
 			: Next(nullptr)
 			, Item(nullptr)
 		{
 		}
 
+		/** Constructor, initialize the item and lock count to 1. */
+		FLink(void *NewItem)
+			: Next(nullptr)
+			, Item(NewItem)
+			, LockCount(1)
+		{
+#if CHECK_NON_CONCURRENT_ASSUMPTIONS
+			FPlatformMisc::MemoryBarrier();
+			checkLockFreePointerList(NewLink->Item == NewItem);
+			checkLockFreePointerList(!NewLink->Next);
+			checkLockFreePointerList(NewLink->LockCount.GetValue() == 1);
+#endif
+		}
+
 		/**	
 		 *	Link this node into the head of the given list.
 		 *	@param HeadPointer; the head of the list
@@ -515,13 +529,21 @@ private:
 		void Unlock(bool bShouldNeverCauseAFree = false);
 	};
 
-	/**	
-	 *	Class to allocate and recycle links.
-	 */
-	class FLinkAllocator
+public:
+
+	// new link allocator
+	class FLinkAllocator	
 	{
+		enum
+		{
+			NUM_PER_BUNDLE=1024,
+		};
 	public:
 
+		CORE_API FLinkAllocator();
+		/** Destructor, leaks all of the memory **/
+		CORE_API ~FLinkAllocator();
+
 		/**	
 		 *	Return a new link.
 		 *	@param NewItem "item" pointer of the new node.
@@ -529,40 +551,45 @@ private:
 		 */
 		FLink* AllocateLink(void *NewItem)
 		{
-			checkLockFreePointerList(NewItem); // we don't allow nullptr items
-			if (NumUsedLinks.Increment() % 8192 == 1)
-			{
-				UE_CLOG(0/*MONITOR_LINK_ALLOCATION*/,LogLockFreeList, Log, TEXT("Number of links %d"),NumUsedLinks.GetValue());
-			}
+			FThreadLocalCache& TLS = GetTLS();
 
-			FLink*  NewLink = nullptr;
-			NewLink = FLink::Unlink(&FreeLinks);
-			if (NewLink)
-			{
-				NumFreeLinks.Decrement();
-			}
-			else
+			if (!TLS.PartialBundle)
 			{
-				if (NumAllocatedLinks.Increment() % 10 == 1)
+				if (TLS.FullBundle)
 				{
-					UE_CLOG(0/*MONITOR_LINK_ALLOCATION*/,LogLockFreeList, Log, TEXT("Number of allocated links %d"),NumAllocatedLinks.GetValue());
+					TLS.PartialBundle = TLS.FullBundle;
+					TLS.FullBundle = nullptr;
 				}
-
-				NewLink = new FLink();
-				NewLink->LockCount.Increment();
+				else
+				{
+					TLS.PartialBundle = GetFreeBundle();
+					if (!TLS.PartialBundle)
+					{
+						TLS.PartialBundle = (void**)FMemory::Malloc(sizeof(FLink) * NUM_PER_BUNDLE);
+						void **Next = TLS.PartialBundle;
+						for (int32 Index = 0; Index < NUM_PER_BUNDLE - 1; Index++)
+						{
+							void* NextNext = (void*)(((uint8*)Next) + sizeof(FLink));
+							*Next = NextNext;
+							Next = (void**)NextNext;
+						}
+						*Next = nullptr;
+						NumFree.Add(NUM_PER_BUNDLE);
+					}
+				}
+				TLS.NumPartial = NUM_PER_BUNDLE;
 			}
-			checkLockFreePointerList(!NewLink->Item);
-			checkLockFreePointerList(!NewLink->Next);
-			checkLockFreePointerList(NewLink->LockCount.GetValue() >= 1);
-#if CHECK_NON_CONCURRENT_ASSUMPTIONS
-			void* ValueWas = FPlatformAtomics::InterlockedCompareExchangePointer(&NewLink->Item, NewItem, nullptr); // in theory this doesn't need to be interlocked
-			checkLockFreePointerList(ValueWas == nullptr);
-#else
-			NewLink->Item = NewItem;
-#endif
-			FPlatformMisc::MemoryBarrier();
+			NumUsed.Increment();
+			NumFree.Decrement();
+			void* Result = (void*)TLS.PartialBundle;
+			TLS.PartialBundle = (void**)*TLS.PartialBundle;
+			TLS.NumPartial--;
+			checkLockFreePointerList(TLS.NumPartial >= 0 && ((!!TLS.NumPartial) == (!!TLS.PartialBundle)));
+
+			FLink* NewLink = new (Result) FLink(NewItem);
 			return NewLink;
 		}
+
 		/**	
 		 *	Make a link available for recycling. 
 		 *	@param Link link to recycle
@@ -571,17 +598,53 @@ private:
 		void FreeLink(FLink* Link)
 		{
 			checkLockFreePointerList(Link != ClosedLink());  // very bad to recycle the special link
-			NumUsedLinks.Decrement();
-			Link->LockCount.Increment();
-			FPlatformMisc::MemoryBarrier();
-			Link->Link(&FreeLinks);
-			NumFreeLinks.Increment();
+			Link->~FLink(); // noop, but here for completeness
+			NumUsed.Decrement();
+			NumFree.Increment();
+			FThreadLocalCache& TLS = GetTLS();
+			if (TLS.NumPartial >= NUM_PER_BUNDLE)
+			{
+				if (TLS.FullBundle)
+				{
+					PutFreeBundle(TLS.FullBundle);
+					//TLS.FullBundle = nullptr;
+				}
+				TLS.FullBundle = TLS.PartialBundle;
+				TLS.PartialBundle = nullptr;
+				TLS.NumPartial = 0;
+			}
+			*(void**)Link = (void*)TLS.PartialBundle;
+			TLS.PartialBundle = (void**)Link;
+			TLS.NumPartial++;
 		}
+
+		/**
+		 * Gets the number of allocated memory blocks that are currently in use.
+		 *
+		 * @return Number of used memory blocks.
+		 * @see GetNumFree
+		 */
+		const FLockFreeListCounter& GetNumUsed() const
+		{
+			return NumUsed;
+		}
+
+		/**
+		 * Gets the number of allocated memory blocks that are currently unused.
+		 *
+		 * @return Number of unused memory blocks.
+		 * @see GetNumUsed
+		 */
+		const FLockFreeListCounter& GetNumFree() const
+		{
+			return NumFree;
+		}
+
 		/**	
 		 *	Return a pointer to the special closed link
 		 *	@return pointer to special closed link
 		 */
-		FLink* ClosedLink()
+		FORCEINLINE FLink* ClosedLink()
 		{
 			checkLockFreePointerList(nullptr != SpecialClosedLink);
 			return SpecialClosedLink;
@@ -590,36 +653,69 @@ private:
 		 *	Singleton access
 		 *	@return the singleton for the link allocator
 		 */
-		static CORE_API FLinkAllocator& Get();
+		FORCEINLINE static FLinkAllocator& Get()
+		{
+			return TheLinkAllocator;
+		}
 
 	private:
-		/**	
-		 *	Constructor zeros the free link list and creates the special closed link
-		 *	@return the singleton for the link allocator
-		 */
-		FLinkAllocator();
 
-		/**	
-		 *	Destructor, should only be called when there are no outstanding links. 
-		 *	Frees the elements of the free list and frees the special link
-		 *	@return the singleton for the link allocator
-		 */
-		~FLinkAllocator();
+		static CORE_API FLinkAllocator TheLinkAllocator;
 
-		/** Head to a list of free links that can be used for new allocations. */
-		FLink*				FreeLinks;
+		/** struct for the TLS cache. */
+		struct FThreadLocalCache
+		{
+			void **FullBundle;
+			void **PartialBundle;
+			int32 NumPartial;
+
+			FThreadLocalCache()
+				: FullBundle(nullptr)
+				, PartialBundle(nullptr)
+				, NumPartial(0)
+			{
+			}
+		};
+
+		FThreadLocalCache& GetTLS()
+		{
+			checkSlow(TlsSlot);
+			FThreadLocalCache* TLS = (FThreadLocalCache*)FPlatformTLS::GetTlsValue(TlsSlot);
+			if (!TLS)
+			{
+				TLS = new FThreadLocalCache();
+				FPlatformTLS::SetTlsValue(TlsSlot, TLS);
+			}
+			return *TLS;
+		}
+
+		/** put a set of 1024 links to a global store **/
+		CORE_API void PutFreeBundle(void **Bundle);
+		/** retrieve a set of 1024 links from a global store **/
+		CORE_API void** GetFreeBundle();
+
+		/** Slot for TLS struct. */
+		uint32 TlsSlot;
 
 		/** Pointer to the special closed link. It should never be recycled. */
 		FLink*				SpecialClosedLink;
 
-		/** Total number of links outstanding and not in the free list */
-		FLockFreeListCounter	NumUsedLinks; 
-		/** Total number of links in the free list */
-		FLockFreeListCounter	NumFreeLinks;
-		/** Total number of links allocated */
-		FLockFreeListCounter	NumAllocatedLinks; 
+		/** Total number of blocks outstanding and not in the free list. */
+		FLockFreeListCounter NumUsed; 
+
+		/** Total number of blocks in the free list. */
+		FLockFreeListCounter NumFree;
+
+		/** Critical section for the bundle array. */
+		FCriticalSection BundleArrayCriticalSection;
+
+		/** Array of bundles of 1024 links */
+		TArray<void**> BundleArray;
+
 	};
 
+private:
+
 	/** Head of the list */
 	MS_ALIGN(8) FLink*	Head;
 };
diff --git a/Engine/Source/Runtime/Engine/Classes/Engine/EngineBaseTypes.h b/Engine/Source/Runtime/Engine/Classes/Engine/EngineBaseTypes.h
--- a/Engine/Source/Runtime/Engine/Classes/Engine/EngineBaseTypes.h
+++ b/Engine/Source/Runtime/Engine/Classes/Engine/EngineBaseTypes.h
@@ -268,14 +268,7 @@ private:
 	 * Queues a tick function for execution from the game thread
 	 * @param TickContext - context to tick in
 	 */
-	void QueueTickFunction(const struct FTickContext& TickContext);
-
-	/**
-	 * Queues a tick function for execution, assuming parallel queuing
-	 * @param TickContext - context to tick in
-	 * @param StackForCycleDetection - stack to detect cycles
-	*/
-	void QueueTickFunctionParallel(const FTickContext& TickContext, TArray<FTickFunction*, TInlineAllocator<4> >& StackForCycleDetection);
+	void QueueTickFunction(class FTickTaskSequencer& TTS, const struct FTickContext& TickContext);
 
 	/** 
 	 * Abstract function actually execute the tick. 
diff --git a/Engine/Source/Runtime/Engine/Private/TickTaskManager.cpp b/Engine/Source/Runtime/Engine/Private/TickTaskManager.cpp
--- a/Engine/Source/Runtime/Engine/Private/TickTaskManager.cpp
+++ b/Engine/Source/Runtime/Engine/Private/TickTaskManager.cpp
@@ -16,11 +16,11 @@ DECLARE_CYCLE_STAT(TEXT("Post Queue Tick Task"),STAT_PostTickTask,STATGROUP_Game
 DECLARE_DWORD_COUNTER_STAT(TEXT("Ticks Queued"),STAT_TicksQueued,STATGROUP_Game);
 
 static TAutoConsoleVariable<int32> CVarLogTicks(
-	TEXT("LogTicks"),0,
+	TEXT("tick.LogTicks"),0,
 	TEXT("Used to logging of ticks."));
 
 static TAutoConsoleVariable<int32> CVarAllowAsyncComponentTicks(
-	TEXT("AllowAsyncComponentTicks"),
+	TEXT("tick.AllowAsyncComponentTicks"),
 	0,
 	TEXT("Used to control async component ticks."));
 
@@ -64,11 +64,68 @@ struct FTickContext
  */
 class FTickTaskSequencer
 {
+	/** Helper class define the task of ticking a component **/
+	class FTickFunctionTask
+	{
+		/** Actor to tick **/
+		FTickFunction*			Target;
+		/** tick context, here thread is desired execution thread **/
+		FTickContext			Context;
+		/** If true, log each tick **/
+		bool					bLogTick; 
+	public:
+		/** Constructor
+		 * @param InTarget - Function to tick
+		 * @param InContext - context to tick in, here thread is desired execution thread
+		**/
+		FTickFunctionTask(FTickFunction* InTarget, const FTickContext* InContext, bool InbLogTick)
+			: Target(InTarget)
+			, Context(*InContext)
+			, bLogTick(InbLogTick)
+		{
+		}
+		FORCEINLINE TStatId GetStatId() const
+		{
+			RETURN_QUICK_DECLARE_CYCLE_STAT(FTickFunctionTask, STATGROUP_TaskGraphTasks);
+		}
+		/** return the thread for this task **/
+		ENamedThreads::Type GetDesiredThread()
+		{
+			return Context.Thread;
+		}
+		static ESubsequentsMode::Type GetSubsequentsMode() 
+		{ 
+			return ESubsequentsMode::TrackSubsequents; 
+		}
+		/** 
+		 *	Actually execute the tick.
+		 *	@param	CurrentThread; the thread we are running on
+		 *	@param	MyCompletionGraphEvent; my completion event. Not always useful since at the end of DoWork, you can assume you are done and hence further tasks do not need you as a prerequisite. 
+		 *	However, MyCompletionGraphEvent can be useful for passing to other routines or when it is handy to set up subsequents before you actually do work.
+		 **/
+		void DoTask(ENamedThreads::Type CurrentThread, const FGraphEventRef& MyCompletionGraphEvent)
+		{
+			if (bLogTick)
+			{
+				UE_LOG(LogTick, Log, TEXT("tick %6d %2d %s"),GFrameCounter, (int32)CurrentThread, *Target->DiagnosticMessage());
+			}
+			Target->ExecuteTick(Context.DeltaSeconds, Context.TickType, CurrentThread, MyCompletionGraphEvent);
+			Target->CompletionHandle = NULL; // Allow the old completion handle to be recycled
+		}
+	};
+
+
 	/** Completion handles for each phase of ticks */
 	FGraphEventArray	TickCompletionEvents[TG_MAX];
 
-	/** Start event for each phase of ticks */
-	FGraphEventRef		TickGroupStartEvents[TG_MAX];
+	/** Held tasks for each tick group. */
+	TArray<TGraphTask<FTickFunctionTask>*> TickTasks[TG_MAX];
+
+	/** These are waited for at the end of the frame; they are not on the critical path, but they have to be done before we leave the frame. */
+	FGraphEventArray CleanupTasks;
+
+	/** we keep track of the last TG we have blocked for so when we do block, we know which TG's to wait for . */
+	ETickingGroup WaitForTickGroup;
 
 	/** If true, allow concurrent ticks **/
 	bool				bAllowConcurrentTicks; 
@@ -118,7 +175,11 @@ public:
 		{
 			UseContext.Thread = ENamedThreads::AnyThread;
 		}
-		TickFunction->CompletionHandle = TGraphTask<FTickFunctionTask>::CreateTask(Prerequisites, TickContext.Thread).ConstructAndDispatchWhenReady(TickFunction, &UseContext, bLogTicks);
+		QUICK_SCOPE_CYCLE_COUNTER(STAT_StartTickTask_ConstructAndDispatchWhenReady);
+
+		TGraphTask<FTickFunctionTask>* Task = TGraphTask<FTickFunctionTask>::CreateTask(Prerequisites, TickContext.Thread).ConstructAndHold(TickFunction, &UseContext, bLogTicks);
+		TickTasks[TickFunction->ActualTickGroup].Add(Task);
+		TickFunction->CompletionHandle = Task->GetCompletionEvent();
 	}
 
 	/** Add a completion handle to a tick group **/
@@ -127,12 +188,6 @@ public:
 		checkSlow(TickGroup >=0 && TickGroup < TG_MAX);
 		new (TickCompletionEvents[TickGroup]) FGraphEventRef(CompletionHandle);
 	}
-	/** Add a completion handle to a tick group; in this case we swap the references to save the cache miss on the handle target **/
-	FORCEINLINE void AddTickTaskCompletionSwap(ETickingGroup TickGroup, FGraphEventRef& CompletionHandle)
-	{
-		checkSlow(TickGroup >=0 && TickGroup < TG_MAX);
-		(new (TickCompletionEvents[TickGroup]) FGraphEventRef())->Swap(CompletionHandle);
-	}
 	/**
 	 * Start a component tick task and add the completion handle
 	 *
@@ -147,13 +202,6 @@ public:
 		AddTickTaskCompletion(TickFunction->ActualTickGroup, TickFunction->CompletionHandle);
 	}
 
-	/** return the start event for a given tick group **/
-	FORCEINLINE FGraphEventRef& GetTickGroupStartEvent(ETickingGroup TickGroup)
-	{
-		checkSlow(TickGroup >=0 && TickGroup < TG_MAX);
-		check(TickGroupStartEvents[TickGroup].GetReference());
-		return TickGroupStartEvents[TickGroup];
-	}
 	/** 
 	 * Release the queued ticks for a given tick group and process them. 
 	 * @param WorldTickGroup - tick group to release
@@ -166,18 +214,14 @@ public:
 			UE_LOG(LogTick, Log, TEXT("tick %6d ---------------------------------------- Release tick group %d"),GFrameCounter, (int32)WorldTickGroup);
 		}
 		checkSlow(WorldTickGroup >=0 && WorldTickGroup < TG_MAX);
-		check(TickGroupStartEvents[WorldTickGroup].GetReference()); // the start event should exist
 
 		if (SingleThreadedMode())
 		{
-			TickGroupStartEvents[WorldTickGroup]->DispatchSubsequents(ENamedThreads::GameThread); // start this tick group
-			TickGroupStartEvents[WorldTickGroup] = NULL;
+			DispatchTickGroupInner(ENamedThreads::GameThread, WorldTickGroup);
 		}
 		else
 		{
-			DECLARE_CYCLE_STAT(TEXT("FDelegateGraphTask.DispatchTickGroup"),
-				STAT_FDelegateGraphTask_DispatchTickGroup,
-				STATGROUP_TaskGraphTasks);
+			DECLARE_CYCLE_STAT(TEXT("FDelegateGraphTask.DispatchTickGroup"), STAT_FDelegateGraphTask_DispatchTickGroup, STATGROUP_TaskGraphTasks);
 
 			// dispatch the tick group on another thread, that way, the game thread can be processing ticks while ticks are being queued by another thread
 			QUICK_SCOPE_CYCLE_COUNTER(STAT_ReleaseTickGroup);
@@ -188,32 +232,38 @@ public:
 				ENamedThreads::GameThread);
 		}
 
-		if (TickCompletionEvents[WorldTickGroup].Num())
+		if (bBlockTillComplete || SingleThreadedMode())
 		{
-			if (bBlockTillComplete || SingleThreadedMode())
+			QUICK_SCOPE_CYCLE_COUNTER(STAT_ReleaseTickGroup_Block);
+			for (ETickingGroup Block = WaitForTickGroup; Block <= WorldTickGroup; Block = ETickingGroup(Block + 1))
 			{
-				QUICK_SCOPE_CYCLE_COUNTER(STAT_ReleaseTickGroup_Block);
-				FTaskGraphInterface::Get().WaitUntilTasksComplete(TickCompletionEvents[WorldTickGroup],ENamedThreads::GameThread);
-			}
-			else
-			{
-				// since this is used to soak up some async time for another task (physics), we should process whatever we have now
-				FTaskGraphInterface::Get().ProcessThreadUntilIdle(ENamedThreads::GameThread);
-				check(WorldTickGroup + 1 < TG_MAX); // you must block on the last tick group!
-				// now we need to make sure these tasks are completed with the next tick group
-				TickCompletionEvents[WorldTickGroup + 1] += TickCompletionEvents[WorldTickGroup];
-			}
-			TickCompletionEvents[WorldTickGroup].Reset();
-			// the last tick group can be used over and over for newly spawned things, so we add a new start event now
-			if (WorldTickGroup == TG_NewlySpawned)
-			{
-				check(bBlockTillComplete); // last tick group must block because of this reuse
-				TickGroupStartEvents[TG_NewlySpawned] = FGraphEvent::CreateGraphEvent(); 
+				if (TickCompletionEvents[Block].Num())
+				{
+					FTaskGraphInterface::Get().WaitUntilTasksComplete(TickCompletionEvents[Block], ENamedThreads::GameThread);
+					if (SingleThreadedMode() || WorldTickGroup == TG_NewlySpawned)
+					{
+						ResetTickGroupInner(Block);
+					}
+					else
+					{
+						DECLARE_CYCLE_STAT(TEXT("FDelegateGraphTask.ResetTickGroup"), STAT_FDelegateGraphTask_ResetTickGroup, STATGROUP_TaskGraphTasks);
+
+						CleanupTasks.Add(FDelegateGraphTask::CreateAndDispatchWhenReady(
+							FDelegateGraphTask::FDelegate::CreateRaw(this, &FTickTaskSequencer::ResetTickGroup, Block),
+							GET_STATID(STAT_FDelegateGraphTask_ResetTickGroup)));
+					}
+				}
 			}
+			WaitForTickGroup = ETickingGroup(WorldTickGroup + 1);
+		}
+		else
+		{
+			// since this is used to soak up some async time for another task (physics), we should process whatever we have now
+			FTaskGraphInterface::Get().ProcessThreadUntilIdle(ENamedThreads::GameThread);
+			check(WorldTickGroup + 1 < TG_MAX && WorldTickGroup != TG_NewlySpawned); // you must block on the last tick group! And we must block on newly spawned
 		}
 	}
 
-
 	/**
 	 * Resets the internal state of the object at the start of a frame
 	 */
@@ -238,9 +288,10 @@ public:
 		{
 			check(!TickCompletionEvents[Index].Num());  // we should not be adding to these outside of a ticking proper and they were already cleared after they were ticked
 			TickCompletionEvents[Index].Reset();
-			check(!TickGroupStartEvents[Index].GetReference()); // this should have been NULL'ed out after it was released
-			TickGroupStartEvents[Index] = FGraphEvent::CreateGraphEvent();
+			check(!TickTasks[Index].Num());  // we should not be adding to these outside of a ticking proper and they were already cleared after they were ticked
+			TickTasks[Index].Reset();
 		}
+		WaitForTickGroup = (ETickingGroup)0;
 	}
 	/**
 	 * Checks that everything is clean at the end of a frame
@@ -251,11 +302,12 @@ public:
 		{
 			UE_LOG(LogTick, Log, TEXT("tick %6d ---------------------------------------- End Frame"),GFrameCounter);
 		}
-		TickGroupStartEvents[TG_NewlySpawned] = NULL; // The last of these was speculative and will be orphaned, delete it here
+		FTaskGraphInterface::Get().WaitUntilTasksComplete(CleanupTasks, ENamedThreads::GameThread);
+		CleanupTasks.Reset();
 		for (int32 Index = 0; Index < TG_MAX; Index++)
 		{
 			check(!TickCompletionEvents[Index].Num());  // we should not be adding to these outside of a ticking proper and they were already cleared after they were ticked
-			check(!TickGroupStartEvents[Index].GetReference()); // this should have been NULL'ed out after it was released
+			check(!TickTasks[Index].Num()); 
 		}
 	}
 private:
@@ -266,64 +318,30 @@ private:
 	{
 	}
 
-	void DispatchTickGroup(ENamedThreads::Type CurrentThread, const FGraphEventRef& MyCompletionGraphEvent,ETickingGroup WorldTickGroup)
+	void ResetTickGroupInner(ETickingGroup WorldTickGroup)
 	{
-		if (ensure(TickGroupStartEvents[WorldTickGroup].GetReference()))
-		{
-			TickGroupStartEvents[WorldTickGroup]->DispatchSubsequents(CurrentThread); // start this tick group
-			TickGroupStartEvents[WorldTickGroup] = NULL;
-		}
+		TickCompletionEvents[WorldTickGroup].Reset();
+	}
+	void ResetTickGroup(ENamedThreads::Type CurrentThread, const FGraphEventRef& MyCompletionGraphEvent, ETickingGroup WorldTickGroup)
+	{
+		ResetTickGroupInner(WorldTickGroup);
 	}
 
-	/** Helper class define the task of ticking a component **/
-	class FTickFunctionTask
+	void DispatchTickGroupInner(ENamedThreads::Type CurrentThread, ETickingGroup WorldTickGroup)
 	{
-		/** Actor to tick **/
-		FTickFunction*			Target;
-		/** tick context, here thread is desired execution thread **/
-		FTickContext			Context;
-		/** If true, log each tick **/
-		bool					bLogTick; 
-	public:
-		/** Constructor
-		 * @param InTarget - Function to tick
-		 * @param InContext - context to tick in, here thread is desired execution thread
-		**/
-		FTickFunctionTask(FTickFunction* InTarget, const FTickContext* InContext, bool InbLogTick)
-			: Target(InTarget)
-			, Context(*InContext)
-			, bLogTick(InbLogTick)
-		{
-		}
-		FORCEINLINE TStatId GetStatId() const
-		{
-			RETURN_QUICK_DECLARE_CYCLE_STAT(FTickFunctionTask, STATGROUP_TaskGraphTasks);
-		}
-		/** return the thread for this task **/
-		ENamedThreads::Type GetDesiredThread()
-		{
-			return Context.Thread;
-		}
-		static ESubsequentsMode::Type GetSubsequentsMode() 
-		{ 
-			return ESubsequentsMode::TrackSubsequents; 
-		}
-		/** 
-		 *	Actually execute the tick.
-		 *	@param	CurrentThread; the thread we are running on
-		 *	@param	MyCompletionGraphEvent; my completion event. Not always useful since at the end of DoWork, you can assume you are done and hence further tasks do not need you as a prerequisite. 
-		 *	However, MyCompletionGraphEvent can be useful for passing to other routines or when it is handy to set up subsequents before you actually do work.
-		 **/
-		void DoTask(ENamedThreads::Type CurrentThread, const FGraphEventRef& MyCompletionGraphEvent)
+		TArray<TGraphTask<FTickFunctionTask>*>& TickArray = TickTasks[WorldTickGroup];
+		for (int32 Index = 0; Index < TickArray.Num(); Index++)
 		{
-			if (bLogTick)
-			{
-				UE_LOG(LogTick, Log, TEXT("tick %6d %2d %s"),GFrameCounter, (int32)CurrentThread, *Target->DiagnosticMessage());
-			}
-			Target->ExecuteTick(Context.DeltaSeconds, Context.TickType, CurrentThread, MyCompletionGraphEvent);
-			Target->CompletionHandle = NULL; // Allow the old completion handle to be recycled
+			TickArray[Index]->Unlock(CurrentThread);
 		}
-	};
+		TickArray.Reset();
+	}
+
+	void DispatchTickGroup(ENamedThreads::Type CurrentThread, const FGraphEventRef& MyCompletionGraphEvent, ETickingGroup WorldTickGroup)
+	{
+		DispatchTickGroupInner(CurrentThread, WorldTickGroup);
+	}
+
 };
 
 
@@ -366,9 +384,10 @@ public:
 	/* Queue all tick functions for execution */
 	void QueueAllTicks()
 	{
+		FTickTaskSequencer& TTS = FTickTaskSequencer::Get();
 		for (TSet<FTickFunction *>::TIterator It(AllEnabledTickFunctions); It; ++It)
 		{
-			(*It)->QueueTickFunction(Context);
+			(*It)->QueueTickFunction(TTS, Context);
 		}
 	}
 	/* Retrieve the tick set */
@@ -384,9 +403,10 @@ public:
 	{
 		Context.TickGroup = CurrentTickGroup;
 		int32 Num = 0;
+		FTickTaskSequencer& TTS = FTickTaskSequencer::Get();
 		for (TSet<FTickFunction *>::TIterator It(NewlySpawnedTickFunctions); It; ++It)
 		{
-			(*It)->QueueTickFunction(Context);
+			(*It)->QueueTickFunction(TTS, Context);
 			Num++;
 		}
 		NewlySpawnedTickFunctions.Empty();
@@ -527,19 +547,6 @@ private:
 	bool										bTickNewlySpawned;
 };
 
-/** Helper struct to hold completion items from parallel task. They are moved into a separate place for cache coherency **/
-struct FTickGroupCompletionItem
-{
-	/** Tick group to complete with **/
-	ETickingGroup	ActualTickGroup;
-	/** Completion event for the task **/
-	FGraphEventRef	CompletionEvent;
-};
-
-static TAutoConsoleVariable<int32> CVarAllowConcurrentQueue(
-	TEXT("AllowConcurrentTickQueue"),0,
-	TEXT("If true, queue ticks concurrently."));
-
 /** Class that aggregates the individual levels and deals with parallel tick setup **/
 class FTickTaskManager : public FTickTaskManagerInterface
 {
@@ -593,82 +600,9 @@ public:
 		}
 		INC_DWORD_STAT_BY(STAT_TicksQueued, TotalTickFunctions);
 
-		int32 NumWorkerThread = 0;
-		bool bConcurrentQueue;
-		if (FTickTaskSequencer::SingleThreadedMode())
-		{
-			// if we don't have many ticks or worker threads, don't bother parallelizing it
-			bConcurrentQueue = false;
-		}
-		else
-		{
-			NumWorkerThread = FTaskGraphInterface::Get().GetNumWorkerThreads();
-			if (TotalTickFunctions < NumWorkerThread * 20 || NumWorkerThread < 2)
-			{
-				// if we don't have many ticks or worker threads, don't bother parallelizing it
-				bConcurrentQueue = false;
-			}
-			else
-			{
-				// the parallel case doesn't ever seem to be faster, so it disabled by default.
-
-				bConcurrentQueue = !!CVarAllowConcurrentQueue.GetValueOnGameThread();
-			}
-		}
-
-		if (!bConcurrentQueue)
-		{
-			// single threaded case
-			for( int32 LevelIndex = 0; LevelIndex < LevelList.Num(); LevelIndex++ )
-			{
-				LevelList[LevelIndex]->QueueAllTicks();
-			}
-		}
-		else
+		for( int32 LevelIndex = 0; LevelIndex < LevelList.Num(); LevelIndex++ )
 		{
-			// parallel case
-
-			// build a list of all tick functions
-			AllTickFunctions.AddZeroed(TotalTickFunctions);
-			AllCompletionEvents.AddZeroed(TotalTickFunctions);
-
-			// split the array into parts for a bunch of concurrent tasks
-			int32 NumTasks = NumWorkerThread * 4; // artificially increase the granularity so if other things are using the working threads, we aren't left waiting with idle cores for one to finish
-			int32 Start = 0;
-			Context.Thread = ENamedThreads::AnyThread;
-
-			int32 NumPerTask = AllTickFunctions.Num() / NumTasks;
-			check(NumPerTask);
-			
-			check(NumTasks > 1); // assumption, below relating to the remainder
-			int32 NumThisTask = NumPerTask;
-			int32 NumTicksSoFar = 0;
-			int32 Task = 0;
-
-			for( int32 LevelIndex = 0; LevelIndex < LevelList.Num(); LevelIndex++ )
-			{
-				for (TSet<FTickFunction *>::TIterator TickIt(LevelList[LevelIndex]->GetTickSet()); TickIt; ++TickIt)
-				{
-					AllTickFunctions[Start + NumTicksSoFar] = *TickIt;
-					NumTicksSoFar++;
-					if (NumTicksSoFar == NumThisTask)
-					{
-						check(Task < NumTasks);
-						FGraphEventArray Setup;
-						new (Setup) FGraphEventRef(TGraphTask<FQueueTickTasks>::CreateTask(NULL,ENamedThreads::GameThread).ConstructAndDispatchWhenReady(AllTickFunctions.GetData() + Start, AllCompletionEvents.GetData() + Start, NumThisTask, &Context));
-						new (QueueTickTasks) FGraphEventRef(TGraphTask<FPostTickTasks>::CreateTask(&Setup,ENamedThreads::GameThread).ConstructAndDispatchWhenReady(AllCompletionEvents.GetData() + Start, NumThisTask));
-						Start += NumThisTask;
-						NumTicksSoFar = 0;
-
-						Task++;
-						if (Task + 1 == NumTasks)
-						{
-							// last task needs to take an remainder
-							NumThisTask = AllTickFunctions.Num() - Start;
-						}
-					}
-				}
-			}
+			LevelList[LevelIndex]->QueueAllTicks();
 		}
 	}
 
@@ -703,21 +637,6 @@ public:
 	*/
 	virtual void RunTickGroup(ETickingGroup Group, bool bBlockTillComplete ) override
 	{
-		if (QueueTickTasks.Num())
-		{
-			check(AllTickFunctions.Num() && AllCompletionEvents.Num());
-			check(!Group); // we only do the wait on the first tick group
-			SCOPE_CYCLE_COUNTER(STAT_QueueTicksWait);
-			FTaskGraphInterface::Get().WaitUntilTasksComplete(QueueTickTasks, ENamedThreads::GameThread);
-			QueueTickTasks.Reset();
-			AllTickFunctions.Reset();
-			AllCompletionEvents.Reset();
-		}
-		else
-		{
-			check(!AllTickFunctions.Num() && !AllCompletionEvents.Num());
-		}
-
 		check(Context.TickGroup == Group); // this should already be at the correct value, but we want to make sure things are happening in the right order
 		check(bTickNewlySpawned); // we should be in the middle of ticking
 		TickTaskSequencer.ReleaseTickGroup(Group, bBlockTillComplete);
@@ -817,87 +736,6 @@ private:
 		check(Level->TickTaskLevel);
 		return Level->TickTaskLevel;
 	}
-	/** Helper class define the task of queueing a subsection of ticks, these run on any thread **/
-	class FQueueTickTasks
-	{
-		FTickFunction**				FirstTickFunction;
-		FTickGroupCompletionItem*	FirstTickCompletion;
-		int32							NumTickFunctions;
-		const FTickContext*			Context;
-	public:
-		FQueueTickTasks(FTickFunction** InFirstTickFunction, FTickGroupCompletionItem* InFirstTickCompletion, int32 InNumTickFunctions, const FTickContext* InContext)
-			: FirstTickFunction(InFirstTickFunction)
-			, FirstTickCompletion(InFirstTickCompletion)
-			, NumTickFunctions(InNumTickFunctions)
-			, Context(InContext)
-		{
-		}
-		FORCEINLINE TStatId GetStatId() const
-		{
-			RETURN_QUICK_DECLARE_CYCLE_STAT(QueueTickTasks, STATGROUP_TaskGraphTasks);
-		}
-		static ENamedThreads::Type GetDesiredThread()
-		{
-			return ENamedThreads::AnyThread;
-		}
-		static ESubsequentsMode::Type GetSubsequentsMode() 
-		{ 
-			return ESubsequentsMode::TrackSubsequents; 
-		}
-		void DoTask(ENamedThreads::Type CurrentThread, const FGraphEventRef& MyCompletionGraphEvent)
-		{
-			SCOPE_CYCLE_COUNTER(STAT_QueueTickTask);
-			TArray<FTickFunction*, TInlineAllocator<4> > StackForCycleDetection;
-			FTickContext UseContext(*Context);
-			UseContext.Thread = CurrentThread;
-			for (int32 Index = 0; Index < NumTickFunctions; Index++)
-			{
-				FTickFunction* TickFunction = *FirstTickFunction++;
-				TickFunction->QueueTickFunctionParallel(UseContext, StackForCycleDetection);
-				FirstTickCompletion->ActualTickGroup = TickFunction->ActualTickGroup;
-				FirstTickCompletion->CompletionEvent = TickFunction->CompletionHandle;
-				FirstTickCompletion++;
-			}
-		}
-	};
-
-	/** Helper class define the task of posting the tick completion handles to the game thread...this must run on the game thread **/
-	class FPostTickTasks
-	{
-		FTickGroupCompletionItem*	FirstTickCompletion;
-		int32							NumTickFunctions;
-	public:
-		FPostTickTasks(FTickGroupCompletionItem* InFirstTickCompletion, int32 InNumTickFunctions)
-			: FirstTickCompletion(InFirstTickCompletion)
-			, NumTickFunctions(InNumTickFunctions)
-		{
-		}
-		FORCEINLINE TStatId GetStatId() const
-		{
-			RETURN_QUICK_DECLARE_CYCLE_STAT(PostTickTasks, STATGROUP_TaskGraphTasks);
-		}
-		static ENamedThreads::Type GetDesiredThread()
-		{
-			return ENamedThreads::GameThread;
-		}
-		static ESubsequentsMode::Type GetSubsequentsMode() 
-		{ 
-			return ESubsequentsMode::TrackSubsequents; 
-		}
-		void DoTask(ENamedThreads::Type CurrentThread, const FGraphEventRef& MyCompletionGraphEvent)
-		{
-			SCOPE_CYCLE_COUNTER(STAT_PostTickTask);
-			for (int32 Index = 0; Index < NumTickFunctions; Index++)
-			{
-				if (FirstTickCompletion->CompletionEvent.GetReference())
-				{
-					FTickTaskSequencer::Get().AddTickTaskCompletionSwap(FirstTickCompletion->ActualTickGroup, FirstTickCompletion->CompletionEvent);
-				}
-				FirstTickCompletion++;
-			}
-		}
-	};
-
 
 	/** Dumps all tick functions to output device */
 	virtual void DumpAllTickFunctions(FOutputDevice& Ar, UWorld* InWorld, bool bEnabled, bool bDisabled)
@@ -936,11 +774,6 @@ private:
 	/** true during the tick phase, when true, tick function adds also go to the newly spawned list. **/
 	bool										bTickNewlySpawned;
 
-	/** Used between start frame and tick group zero. There is an opportunity for gamethread to soak up some time here, so we don't wait until we run tick group 0 **/
-	FGraphEventArray							QueueTickTasks;
-	TArray<FTickFunction*> AllTickFunctions;
-	TArray<FTickGroupCompletionItem> AllCompletionEvents;
-
 };
 
 
@@ -1045,7 +878,7 @@ void FTickFunction::RemovePrerequisite(UObject* TargetObject, struct FTickFuncti
 	* Queues a tick function for execution from the game thread
 	* @param TickContext - context to tick in
 */
-void FTickFunction::QueueTickFunction(const struct FTickContext& TickContext)
+void FTickFunction::QueueTickFunction(FTickTaskSequencer& TTS, const struct FTickContext& TickContext)
 {
 	checkSlow(TickContext.Thread == ENamedThreads::GameThread); // we assume same thread here
 	check(bRegistered);
@@ -1069,7 +902,7 @@ void FTickFunction::QueueTickFunction(const struct FTickContext& TickContext)
 				else if (Prereq->bRegistered)
 				{
 					// recursive call to make sure my prerequisite is set up so I can use its completion handle
-					Prereq->QueueTickFunction(TickContext);
+					Prereq->QueueTickFunction(TTS, TickContext);
 					if (Prereq->TickQueuedGFrameCounter != GFrameCounter)
 					{
 						// this must be up the call stack, therefore this is a cycle
@@ -1095,115 +928,114 @@ void FTickFunction::QueueTickFunction(const struct FTickContext& TickContext)
 			}
 			ActualTickGroup = MyActualTickGroup;
 
-			// we don't need to add a tick group prerequisite if we already have a prerequisite in the correct tick group (in that case, the delay until the correct tick group is implicit)
-			if (!TaskPrerequisites.Num() || MaxPrerequisiteTickGroup < MyActualTickGroup)
-			{
-				TaskPrerequisites.Add(FTickTaskSequencer::Get().GetTickGroupStartEvent(MyActualTickGroup));
-			}
-			FTickTaskSequencer::Get().QueueTickTask(&TaskPrerequisites, this, TickContext);
+			TTS.QueueTickTask(&TaskPrerequisites, this, TickContext);
 		}
 		TickQueuedGFrameCounter = GFrameCounter;
 	}
 }
 
+
 /**
-	* Queues a tick function for execution, assuming parallel queuing
-	* @param TickContext - context to tick in
-	* @param StackForCycleDetection - stack to detect cycles
-*/
-void FTickFunction::QueueTickFunctionParallel(const FTickContext& TickContext, TArray<FTickFunction*, TInlineAllocator<4> >& StackForCycleDetection)
+ * Singleton to retrieve the global tick task manager
+ * @return Reference to the global tick task manager
+**/
+FTickTaskManagerInterface& FTickTaskManagerInterface::Get()
 {
-		
-	bool bProcessTick;
+	return FTickTaskManager::Get();
+}
+
 
-	int32 OldValue = *(volatile int32*)&TickVisitedGFrameCounter;
-	if (OldValue != GFrameCounter)
+struct FTestTickFunction : public FTickFunction
+{
+	FTestTickFunction()
+	{
+		TickGroup = TG_PrePhysics;
+		bTickEvenWhenPaused = true;
+		bStartWithTickEnabled;
+	}
+	virtual void ExecuteTick(float DeltaTime, ELevelTick TickType, ENamedThreads::Type CurrentThread, const FGraphEventRef& MyCompletionGraphEvent) override
+	{
+		QUICK_SCOPE_CYCLE_COUNTER(STAT_TestStatOverhead_FTestTickFunction);
+		QUICK_SCOPE_CYCLE_COUNTER(STAT_TestStatOverhead_FTestTickFunction_Inner);
+	}
+	/** Abstract function to describe this tick. Used to print messages about illegal cycles in the dependency graph **/
+	virtual FString DiagnosticMessage() override
 	{
-		OldValue = FPlatformAtomics::InterlockedCompareExchange(&TickVisitedGFrameCounter , GFrameCounter, OldValue);
+		return FString(TEXT("test"));
 	}
-	bProcessTick = OldValue != GFrameCounter;
+};
+
+static const int32 NumTestTickFunctions = 10000;
+static TArray<FTestTickFunction> TestTickFunctions;
+static TArray<FTestTickFunction*> IndirectTestTickFunctions;
 
-	if (bProcessTick)
+static void RemoveTestTickFunctions(const TArray<FString>& Args)
+{
+	if (TestTickFunctions.Num() || IndirectTestTickFunctions.Num())
 	{
-		check(bRegistered);
-		if (bTickEnabled && (!EnableParent || EnableParent->bTickEnabled))
+		UE_LOG(LogConsoleResponse, Display, TEXT("Removing Test Tick Functions."));
+		TestTickFunctions.Empty(NumTestTickFunctions);
+		for (int32 Index = 0; Index < IndirectTestTickFunctions.Num(); Index++)
 		{
-			ETickingGroup MaxPrerequisiteTickGroup =  ETickingGroup(0);
+			delete IndirectTestTickFunctions[Index];
+		}
+		IndirectTestTickFunctions.Empty(NumTestTickFunctions);
+	}
+}
 
-			FGraphEventArray TaskPrerequisites;
-			if (Prerequisites.Num())
-			{
-				StackForCycleDetection.Push(this);
-				for (int32 PrereqIndex = 0; PrereqIndex < Prerequisites.Num(); PrereqIndex++)
-				{
-					FTickFunction* Prereq = Prerequisites[PrereqIndex].Get();
-					if (!Prereq)
-					{
-						// stale prereq, delete it
-						Prerequisites.RemoveAtSwap(PrereqIndex--);
-					}
-					else if (Prereq->bRegistered)
-					{
-						if (StackForCycleDetection.Contains(Prereq))
-						{
-							UE_LOG(LogTick, Warning, TEXT("While processing prerequisites for %s, could use %s because it would form a cycle."),*DiagnosticMessage(), *Prereq->DiagnosticMessage());
-						}
-						else
-						{
-							// recursive call to make sure my prerequisite is set up so I can use its completion handle
-							Prereq->QueueTickFunctionParallel(TickContext, StackForCycleDetection);
-							if (!Prereq->CompletionHandle.GetReference())
-							{
-								//ok UE_LOG(LogTick, Warning, TEXT("While processing prerequisites for %s, could use %s because it is disabled."),*DiagnosticMessage(), *Prereq->DiagnosticMessage());
-							}
-							else
-							{
-								MaxPrerequisiteTickGroup =  FMath::Max<ETickingGroup>(MaxPrerequisiteTickGroup, Prereq->ActualTickGroup);
-								TaskPrerequisites.Add(Prereq->CompletionHandle);
-							}
-						}
-					}
-				}
-				StackForCycleDetection.Pop();
-			}
+static void AddTestTickFunctions(const TArray<FString>& Args)
+{
+	RemoveTestTickFunctions(Args);
+	ULevel* Level = GWorld->GetCurrentLevel();
+	UE_LOG(LogConsoleResponse, Display, TEXT("Adding 1000 ticks in a cache coherent fashion."));
 
-			// tick group is the max of the prerequisites, the current tick group, and the desired tick group
-			ETickingGroup MyActualTickGroup =  FMath::Max<ETickingGroup>(MaxPrerequisiteTickGroup, FMath::Max<ETickingGroup>(TickGroup,TickContext.TickGroup));
-			if (MyActualTickGroup == TG_DuringPhysics && TickGroup != TG_DuringPhysics)
-			{
-				MyActualTickGroup = ETickingGroup(MyActualTickGroup + 1); // if the tick was "promoted" to during async, but it was not designed for that, then it needs to go later
-			}
-			ActualTickGroup = MyActualTickGroup;
 
-			// we don't need to add a tick group prerequisite if we already have a prerequisite in the correct tick group (in that case, the delay until the correct tick group is implicit)
-			if (!TaskPrerequisites.Num() || MaxPrerequisiteTickGroup < MyActualTickGroup)
-			{
-				TaskPrerequisites.Add(FTickTaskSequencer::Get().GetTickGroupStartEvent(MyActualTickGroup));
-			}
-			FTickTaskSequencer::Get().StartTickTask(&TaskPrerequisites, this, TickContext);
-		}
-		FPlatformMisc::MemoryBarrier();
-		TickQueuedGFrameCounter = GFrameCounter;
+	TestTickFunctions.Reserve(NumTestTickFunctions);
+	for (int32 Index = 0; Index < NumTestTickFunctions; Index++)
+	{
+		(new (TestTickFunctions) FTestTickFunction())->RegisterTickFunction(Level);
 	}
-	else
+}
+
+static void AddIndirectTestTickFunctions(const TArray<FString>& Args)
+{
+	RemoveTestTickFunctions(Args);
+	ULevel* Level = GWorld->GetCurrentLevel();
+	UE_LOG(LogConsoleResponse, Display, TEXT("Adding 1000 ticks in a cache coherent fashion."));
+	TArray<FTestTickFunction*> Junk;
+	for (int32 Index = 0; Index < NumTestTickFunctions; Index++)
 	{
-		// if we are not going to process it, we need to at least wait until the other thread finishes it
-		volatile int32* TickQueuedGFrameCounterPtr = &TickQueuedGFrameCounter;
-		while (*TickQueuedGFrameCounterPtr != GFrameCounter)
+		for (int32 JunkIndex = 0; JunkIndex < 8; JunkIndex++)
 		{
-			FPlatformMisc::MemoryBarrier(); //spin
+			Junk.Add(new FTestTickFunction); // don't give the allocator an easy ride
 		}
+		FTestTickFunction* NewTick = new FTestTickFunction;
+		NewTick->RegisterTickFunction(Level);
+		IndirectTestTickFunctions.Add(NewTick);
+	}
+	for (int32 JunkIndex = 0; JunkIndex < 8; JunkIndex++)
+	{
+		delete Junk[JunkIndex];
 	}
 }
 
+static FAutoConsoleCommand RemoveTestTickFunctionsCmd(
+	TEXT("tick.RemoveTestTickFunctions"),
+	TEXT("Remove no-op ticks to test performance of ticking infrastructure."),
+	FConsoleCommandWithArgsDelegate::CreateStatic(&RemoveTestTickFunctions)
+	);
+
+static FAutoConsoleCommand AddTestTickFunctionsCmd(
+	TEXT("tick.AddTestTickFunctions"),
+	TEXT("Add no-op ticks to test performance of ticking infrastructure."),
+	FConsoleCommandWithArgsDelegate::CreateStatic(&AddTestTickFunctions)
+	);
+
+static FAutoConsoleCommand AddIndirectTestTickFunctionsCmd(
+	TEXT("tick.AddIndirectTestTickFunctions"),
+	TEXT("Add no-op ticks to test performance of ticking infrastructure."),
+	FConsoleCommandWithArgsDelegate::CreateStatic(&AddIndirectTestTickFunctions)
+	);
 
-/**
- * Singleton to retrieve the global tick task manager
- * @return Reference to the global tick task manager
-**/
-FTickTaskManagerInterface& FTickTaskManagerInterface::Get()
-{
-	return FTickTaskManager::Get();
-}
 
 
