fix for WITH_VEHICLE instead of WITH_PHYSX to allow vehicle physics to be compiled out

diff --git a/Engine/Source/Runtime/Engine/Private/Vehicles/WheeledVehicleMovementComponent4W.cpp b/Engine/Source/Runtime/Engine/Private/Vehicles/WheeledVehicleMovementComponent4W.cpp
--- a/Engine/Source/Runtime/Engine/Private/Vehicles/WheeledVehicleMovementComponent4W.cpp
+++ b/Engine/Source/Runtime/Engine/Private/Vehicles/WheeledVehicleMovementComponent4W.cpp
@@ -392,7 +392,7 @@ void BackwardsConvertCm2ToM2(float& val, float defaultValue)
 void UWheeledVehicleMovementComponent4W::Serialize(FArchive & Ar)
 {
 	Super::Serialize(Ar);
-#if WITH_PHYSX
+#if WITH_VEHICLE
 	if (Ar.IsLoading() && Ar.UE4Ver() < VER_UE4_VEHICLES_UNIT_CHANGE)
 	{
 		PxVehicleEngineData DefEngineData;
