UE4 - fix cis

diff --git a/Engine/Source/Runtime/Engine/Private/TickTaskManager.cpp b/Engine/Source/Runtime/Engine/Private/TickTaskManager.cpp
--- a/Engine/Source/Runtime/Engine/Private/TickTaskManager.cpp
+++ b/Engine/Source/Runtime/Engine/Private/TickTaskManager.cpp
@@ -951,7 +951,6 @@ struct FTestTickFunction : public FTickFunction
 	{
 		TickGroup = TG_PrePhysics;
 		bTickEvenWhenPaused = true;
-		bStartWithTickEnabled;
 	}
 	virtual void ExecuteTick(float DeltaTime, ELevelTick TickType, ENamedThreads::Type CurrentThread, const FGraphEventRef& MyCompletionGraphEvent) override
 	{
