Supporting multiple audio devices for multiple PIE and other potential use-cases.

- Created audio device manager which manages weak ref handles to audio devices
- Audio device handles are created with WorldContexts and accessed through UWorld objects
- All access to audio devices are now through audio device handles
- There is a main/default audio device for GEngine for use with the Editor and single worldcontext games

diff --git a/Engine/Plugins/Runtime/SoundMod/Source/SoundModImporter/Private/SoundModImporterFactory.cpp b/Engine/Plugins/Runtime/SoundMod/Source/SoundModImporter/Private/SoundModImporterFactory.cpp
--- a/Engine/Plugins/Runtime/SoundMod/Source/SoundModImporter/Private/SoundModImporterFactory.cpp
+++ b/Engine/Plugins/Runtime/SoundMod/Source/SoundModImporter/Private/SoundModImporterFactory.cpp
@@ -45,11 +45,11 @@ UObject* USoundModImporterFactory::FactoryCreateBinary
 
 	// TODO - Audio Threading. This needs to be sent to the audio device and wait on stopping the sounds
 	TArray<UAudioComponent*> ComponentsToRestart;
-	FAudioDevice* AudioDevice = GEngine->GetAudioDevice();
-	if (AudioDevice && ExistingSound)
+	FAudioDeviceManager* AudioDeviceManager = GEngine->GetAudioDeviceManager();
+	if (AudioDeviceManager && ExistingSound)
 	{
 		// TODO: Generalize the stop sounds function
-		//AudioDevice->StopSoundsForReimport(ExistingSound, ComponentsToRestart);
+		//AudioDeviceManager->StopSoundsForReimport(ExistingSound, ComponentsToRestart);
 	}
 
 	bool bUseExistingSettings = bSoundModFactorySuppressImportOverwriteDialog;
diff --git a/Engine/Source/Developer/MovieSceneTools/Private/AudioTrackEditor.cpp b/Engine/Source/Developer/MovieSceneTools/Private/AudioTrackEditor.cpp
--- a/Engine/Source/Developer/MovieSceneTools/Private/AudioTrackEditor.cpp
+++ b/Engine/Source/Developer/MovieSceneTools/Private/AudioTrackEditor.cpp
@@ -175,9 +175,13 @@ void FAudioThumbnail::GenerateWaveformPreview(TArray<uint8>& OutData, TRange<flo
 	{
 		// @todo Sequencer optimize - We might want to generate the data when we generate the texture
 		// and then discard the data afterwards, though that might be a perf hit traded for better memory usage
-		SoundWave->InitAudioResource(GEngine->GetAudioDevice()->GetRuntimeFormat(SoundWave));
-		FAsyncAudioDecompress TempDecompress(SoundWave);
-		TempDecompress.StartSynchronousTask();
+		FAudioDevice* AudioDevice = GEngine->GetMainAudioDevice();
+		if (AudioDevice)
+		{
+			SoundWave->InitAudioResource(AudioDevice->GetRuntimeFormat(SoundWave));
+			FAsyncAudioDecompress TempDecompress(SoundWave);
+			TempDecompress.StartSynchronousTask();
+		}
 	}
 
 	int16* LookupData = (int16*)SoundWave->RawPCMData;
diff --git a/Engine/Source/Editor/GraphEditor/Private/SoundCueGraphConnectionDrawingPolicy.cpp b/Engine/Source/Editor/GraphEditor/Private/SoundCueGraphConnectionDrawingPolicy.cpp
--- a/Engine/Source/Editor/GraphEditor/Private/SoundCueGraphConnectionDrawingPolicy.cpp
+++ b/Engine/Source/Editor/GraphEditor/Private/SoundCueGraphConnectionDrawingPolicy.cpp
@@ -36,14 +36,14 @@ void FSoundCueGraphConnectionDrawingPolicy::Draw(TMap<TSharedRef<SWidget>, FArra
 
 void FSoundCueGraphConnectionDrawingPolicy::BuildAudioFlowRoadmap()
 {
-	FAudioDevice* AudioDevice = GEngine->GetAudioDevice();
+	UAudioComponent* PreviewAudioComponent = GEditor->GetPreviewAudioComponent();
+	FAudioDevice* AudioDevice = PreviewAudioComponent->GetAudioDevice();
 
 	if (AudioDevice)
 	{
 		USoundCueGraph* SoundCueGraph = CastChecked<USoundCueGraph>(GraphObj);
 		USoundCue* SoundCue = SoundCueGraph->GetSoundCue();
 
-		UAudioComponent* PreviewAudioComponent = GEditor->GetPreviewAudioComponent();
 
 		if (PreviewAudioComponent && PreviewAudioComponent->IsPlaying() && PreviewAudioComponent->Sound == SoundCue)
 		{
diff --git a/Engine/Source/Editor/Kismet/Private/BlueprintEditor.cpp b/Engine/Source/Editor/Kismet/Private/BlueprintEditor.cpp
--- a/Engine/Source/Editor/Kismet/Private/BlueprintEditor.cpp
+++ b/Engine/Source/Editor/Kismet/Private/BlueprintEditor.cpp
@@ -7271,9 +7271,9 @@ void FBlueprintEditor::UpdatePreviewActor(UBlueprint* InBlueprint, bool bInForce
 			}
 
 			// Prevent any audio from playing as a result of spawning
-			if ( GEngine->AudioDevice )
+			if (FAudioDevice* AudioDevice = GEngine->GetMainAudioDevice())
 			{
-				GEngine->AudioDevice->Flush(PreviewScene.GetWorld());
+				AudioDevice->Flush(PreviewScene.GetWorld());
 			}
 
 			// Set the reference to the preview actor for component editing purposes
diff --git a/Engine/Source/Editor/UnrealEd/Private/Editor.cpp b/Engine/Source/Editor/UnrealEd/Private/Editor.cpp
--- a/Engine/Source/Editor/UnrealEd/Private/Editor.cpp
+++ b/Engine/Source/Editor/UnrealEd/Private/Editor.cpp
@@ -1539,22 +1539,23 @@ void UEditorEngine::Tick( float DeltaSeconds, bool bIdleMode )
 	IStreamingManager::Get().Tick(DeltaSeconds);
 
 	// Update Audio. This needs to occur after rendering as the rendering code updates the listener position.
-	if (GetAudioDevice())
+	FAudioDeviceManager* AudioDeviceManager = GEngine->GetAudioDeviceManager();
+	if (AudioDeviceManager)
 	{
 		UWorld* OldGWorld = NULL;
-		if( PlayWorld )
+		if (PlayWorld)
 		{
 			// Use the PlayWorld as the GWorld if we're using PIE.
-			OldGWorld = SetPlayInEditorWorld( PlayWorld );
+			OldGWorld = SetPlayInEditorWorld(PlayWorld);
 		}
 
 		// Update audio device.
-		GetAudioDevice()->Update( (!PlayWorld && bAudioIsRealtime) || ( PlayWorld && !PlayWorld->IsPaused() ) );
+		AudioDeviceManager->UpdateActiveAudioDevices((!PlayWorld && bAudioIsRealtime) || (PlayWorld && !PlayWorld->IsPaused()));
 
-		if( PlayWorld )
+		if (PlayWorld)
 		{
 			// Pop the world.
-			RestoreEditorWorld( OldGWorld );
+			RestoreEditorWorld(OldGWorld);
 		}
 	}
 
@@ -1922,9 +1923,9 @@ UAudioComponent* UEditorEngine::GetPreviewAudioComponent()
 
 UAudioComponent* UEditorEngine::ResetPreviewAudioComponent( USoundBase* Sound, USoundNode* SoundNode )
 {
-	if( GetAudioDevice() )
+	if (FAudioDevice* AudioDevice = GetMainAudioDevice())
 	{
-		if( PreviewAudioComponent)
+		if (PreviewAudioComponent)
 		{
 			PreviewAudioComponent->Stop();
 		}
@@ -1933,10 +1934,10 @@ UAudioComponent* UEditorEngine::ResetPreviewAudioComponent( USoundBase* Sound, U
 			PreviewSoundCue = NewObject<USoundCue>();
 			// Set world to NULL as it will most likely become invalid in the next PIE/Simulate session and the
 			// component will be left with invalid pointer.
-			PreviewAudioComponent = FAudioDevice::CreateComponent( PreviewSoundCue, NULL, NULL, false );
+			PreviewAudioComponent = FAudioDevice::CreateComponent(PreviewSoundCue, NULL, NULL, false);
 		}
 
-		check( PreviewAudioComponent );
+		check(PreviewAudioComponent);
 		// Mark as a preview component so the distance calculations can be ignored
 		PreviewAudioComponent->bPreviewComponent = true;
 
@@ -1944,7 +1945,7 @@ UAudioComponent* UEditorEngine::ResetPreviewAudioComponent( USoundBase* Sound, U
 		{
 			PreviewAudioComponent->Sound = Sound;
 		}
-		else if( SoundNode)
+		else if (SoundNode)
 		{
 			PreviewSoundCue->FirstNode = SoundNode;
 			PreviewAudioComponent->Sound = PreviewSoundCue;
diff --git a/Engine/Source/Editor/UnrealEd/Private/EditorServer.cpp b/Engine/Source/Editor/UnrealEd/Private/EditorServer.cpp
--- a/Engine/Source/Editor/UnrealEd/Private/EditorServer.cpp
+++ b/Engine/Source/Editor/UnrealEd/Private/EditorServer.cpp
@@ -1793,9 +1793,9 @@ void UEditorEngine::EditorDestroyWorld( FWorldContext & Context, const FText& Cl
 	CloseEditedWorldAssets(ContextWorld);
 
 	// Stop all audio and remove references 
-	if ( GetAudioDevice() )
+	if (FAudioDevice* AudioDevice = ContextWorld->GetAudioDevice())
 	{
-		GetAudioDevice()->Flush(ContextWorld);
+		AudioDevice->Flush(ContextWorld);
 	}
 
 	// Reset the editor transform to avoid loading the new world with an offset if loading a sublevel
diff --git a/Engine/Source/Editor/UnrealEd/Private/Factories/EditorFactories.cpp b/Engine/Source/Editor/UnrealEd/Private/Factories/EditorFactories.cpp
--- a/Engine/Source/Editor/UnrealEd/Private/Factories/EditorFactories.cpp
+++ b/Engine/Source/Editor/UnrealEd/Private/Factories/EditorFactories.cpp
@@ -1695,10 +1695,10 @@ UObject* USoundFactory::FactoryCreateBinary
 
 		// TODO - Audio Threading. This needs to be sent to the audio device and wait on stopping the sounds
 		TArray<UAudioComponent*> ComponentsToRestart;
-		FAudioDevice* AudioDevice = GEngine->GetAudioDevice();
-		if (AudioDevice && ExistingSound)
+		FAudioDeviceManager* AudioDeviceManager = GEngine->GetAudioDeviceManager();
+		if (AudioDeviceManager && ExistingSound)
 		{
-			AudioDevice->StopSoundsUsingResource(ExistingSound, ComponentsToRestart);
+			AudioDeviceManager->StopSoundsUsingResource(ExistingSound, ComponentsToRestart);
 		}
 
 		bool bUseExistingSettings = bSoundFactorySuppressImportOverwriteDialog;
@@ -2444,10 +2444,11 @@ UObject* USoundClassFactory::FactoryCreateNew( UClass* InClass, UObject* InParen
 {
 	USoundClass* SoundClass = NewObject<USoundClass>(InParent, InName, Flags);
 	
-	FAudioDevice* AudioDevice = GEngine ? GEngine->GetAudioDevice() : nullptr;
-	if (AudioDevice)
+
+	class FAudioDeviceManager* AudioDeviceManager = GEngine ? GEngine->GetAudioDeviceManager() : nullptr;
+	if (AudioDeviceManager)
 	{
-		AudioDevice->InitSoundClasses();
+		AudioDeviceManager->InitSoundClasses();
 	}
 
 	return( SoundClass );
diff --git a/Engine/Source/Editor/UnrealEd/Private/LevelEditorViewport.cpp b/Engine/Source/Editor/UnrealEd/Private/LevelEditorViewport.cpp
--- a/Engine/Source/Editor/UnrealEd/Private/LevelEditorViewport.cpp
+++ b/Engine/Source/Editor/UnrealEd/Private/LevelEditorViewport.cpp
@@ -3828,25 +3828,19 @@ void FLevelEditorViewportClient::DrawBrushDetails(const FSceneView* View, FPrimi
 	}
 }
 
-/**
- * Updates the audio listener for this viewport 
- *
- * @param View	The scene view to use when calculate the listener position
- */
-void FLevelEditorViewportClient::UpdateAudioListener( const FSceneView& View )
+void FLevelEditorViewportClient::UpdateAudioListener(const FSceneView& View)
 {
-	FAudioDevice* AudioDevice = GEditor->GetAudioDevice();
+	UWorld* World = GetWorld();
 
-	// AudioDevice may not exist. For example if we are in -nosound mode
-	if( AudioDevice )
+	if (World)
 	{
-		if( GetWorld() )
+		if (FAudioDevice* AudioDevice = World->GetAudioDevice())
 		{
 			FReverbSettings ReverbSettings;
 			FInteriorSettings InteriorSettings;
 			const FVector& ViewLocation = GetViewLocation();
 
-			class AAudioVolume* AudioVolume = GetWorld()->GetAudioSettings( ViewLocation, &ReverbSettings, &InteriorSettings );
+			class AAudioVolume* AudioVolume = World->GetAudioSettings(ViewLocation, &ReverbSettings, &InteriorSettings);
 
 			FMatrix CameraToWorld = View.ViewMatrices.ViewMatrix.InverseFast();
 			FVector ProjUp = CameraToWorld.TransformVector(FVector(0, 1000, 0));
@@ -3856,8 +3850,8 @@ void FLevelEditorViewportClient::UpdateAudioListener( const FSceneView& View )
 			ListenerTransform.SetTranslation(ViewLocation);
 			ListenerTransform.NormalizeRotation();
 
-			AudioDevice->SetListener( 0, ListenerTransform, 0.f, AudioVolume, InteriorSettings );
-			AudioDevice->SetReverbSettings( AudioVolume, ReverbSettings );
+			AudioDevice->SetListener(0, ListenerTransform, 0.f, AudioVolume, InteriorSettings);
+			AudioDevice->SetReverbSettings(AudioVolume, ReverbSettings);
 		}
 	}
 }
@@ -3900,10 +3894,9 @@ void FLevelEditorViewportClient::SetupViewForRendering( FSceneViewFamily& ViewFa
 	}
 
 	// Update the listener.
-	FAudioDevice* const AudioDevice = GEditor ? GEditor->GetAudioDevice() : NULL;
-	if( AudioDevice && bHasAudioFocus )
+	if (bHasAudioFocus)
 	{
-		UpdateAudioListener( View );
+		UpdateAudioListener(View);
 	}
 }
 
diff --git a/Engine/Source/Editor/UnrealEd/Private/PlayLevel.cpp b/Engine/Source/Editor/UnrealEd/Private/PlayLevel.cpp
--- a/Engine/Source/Editor/UnrealEd/Private/PlayLevel.cpp
+++ b/Engine/Source/Editor/UnrealEd/Private/PlayLevel.cpp
@@ -430,12 +430,12 @@ void UEditorEngine::TeardownPlaySession(FWorldContext &PieWorldContext)
 
 	
 	// Stop all audio and remove references to temp level.
-	if (GetAudioDevice())
+	if (FAudioDevice* AudioDevice = PlayWorld->GetAudioDevice())
 	{
-		GetAudioDevice()->Flush( PlayWorld );
-		GetAudioDevice()->ResetInterpolation();
-		GetAudioDevice()->OnEndPIE(bIsSimulatingInEditor);
-		GetAudioDevice()->TransientMasterVolume = 1.0f;
+		AudioDevice->Flush( PlayWorld );
+		AudioDevice->ResetInterpolation();
+		AudioDevice->OnEndPIE(bIsSimulatingInEditor);
+		AudioDevice->TransientMasterVolume = 1.0f;
 	}
 
 	// Clean up all streaming levels
@@ -2133,19 +2133,20 @@ void UEditorEngine::PlayInEditor( UWorld* InWorld, bool bInSimulateInEditor )
 	GEngine->ClearOnScreenDebugMessages();
 
 	// Flush all audio sources from the editor world
-	if( GetAudioDevice() )
+	FAudioDevice* AudioDevice = EditorWorld->GetAudioDevice();
+	if (AudioDevice)
 	{
-		GetAudioDevice()->Flush( EditorWorld );
-		GetAudioDevice()->ResetInterpolation();
-		GetAudioDevice()->OnBeginPIE(bInSimulateInEditor);
+		AudioDevice->Flush( EditorWorld );
+		AudioDevice->ResetInterpolation();
+		AudioDevice->OnBeginPIE(bInSimulateInEditor);
 	}
 	EditorWorld->bAllowAudioPlayback = false;
 
 	ULevelEditorPlaySettings* PlayInSettings = Cast<ULevelEditorPlaySettings>(ULevelEditorPlaySettings::StaticClass()->GetDefaultObject());
 
-	if (!PlayInSettings->EnableSound && GetAudioDevice())
+	if (!PlayInSettings->EnableSound && AudioDevice)
 	{
-		GetAudioDevice()->TransientMasterVolume = 0.0f;
+		AudioDevice->TransientMasterVolume = 0.0f;
 	}
 
 	if (!GEditor->bAllowMultiplePIEWorlds)
@@ -2566,11 +2567,6 @@ UGameInstance* UEditorEngine::CreatePIEGameInstance(int32 PIEInstance, bool bInS
 	GWorld = PlayWorld;
 	SetPlayInEditorWorld( PlayWorld );
 
-	if( GetAudioDevice() )
-	{
-		GetAudioDevice()->SetDefaultBaseSoundMix( PlayWorld->GetWorldSettings()->DefaultBaseSoundMix );
-	}
-
 #if PLATFORM_64BITS
 	const FString PlatformBitsString( TEXT( "64" ) );
 #else
diff --git a/Engine/Source/Runtime/Android/AndroidAudio/Private/AndroidAudioBuffer.cpp b/Engine/Source/Runtime/Android/AndroidAudio/Private/AndroidAudioBuffer.cpp
--- a/Engine/Source/Runtime/Android/AndroidAudio/Private/AndroidAudioBuffer.cpp
+++ b/Engine/Source/Runtime/Android/AndroidAudio/Private/AndroidAudioBuffer.cpp
@@ -15,8 +15,8 @@
  *
  * @param AudioDevice	audio device this sound buffer is going to be attached to.
  */
-FSLESSoundBuffer::FSLESSoundBuffer( FSLESAudioDevice* InAudioDevice ) :
-	AudioDevice(InAudioDevice),
+FSLESSoundBuffer::FSLESSoundBuffer( FSLESAudioDevice* InAudioDevice )
+:	FAudioBuffer(InAudioDevice),
 	AudioData(NULL),
 	DecompressionState( NULL ),
 	Format(SoundFormat_Invalid)
@@ -117,8 +117,11 @@ FSLESSoundBuffer* FSLESSoundBuffer::CreateNativeBuffer( FSLESAudioDevice* AudioD
 
 	// Create new buffer.
 	Buffer = new FSLESSoundBuffer( AudioDevice );
-		
-	AudioDevice->TrackResource(InWave, Buffer);
+	
+	FAudioDeviceManager* AudioDeviceManager = GEngine->GetAudioDeviceManager();
+	check(AudioDeviceManager != nullptr);
+
+	AudioDeviceManager->TrackResource(InWave, Buffer);
 
 	Buffer->NumChannels		= InWave->NumChannels;
 	Buffer->SampleRate		= InWave->SampleRate;
@@ -184,6 +187,7 @@ FSLESSoundBuffer* FSLESSoundBuffer::Init(  FSLESAudioDevice* AudioDevice ,USound
 		return( NULL );
 	}
 	
+	FAudioDeviceManager* AudioDeviceManager = GEngine->GetAudioDeviceManager();
 	FSLESSoundBuffer* Buffer = NULL;
 	
 	EDecompressionType DecompressionType = InWave->DecompressionType;
@@ -205,7 +209,7 @@ FSLESSoundBuffer* FSLESSoundBuffer::Init(  FSLESAudioDevice* AudioDevice ,USound
 		// Upload entire wav
 		if( InWave->ResourceID )
 		{
-			Buffer = static_cast<FSLESSoundBuffer*>(AudioDevice->WaveBufferMap.FindRef( InWave->ResourceID ));
+			Buffer = static_cast<FSLESSoundBuffer*>(AudioDeviceManager->WaveBufferMap.FindRef(InWave->ResourceID));
 		}
 
 		if( Buffer == NULL )
diff --git a/Engine/Source/Runtime/Android/AndroidAudio/Private/AndroidAudioSource.cpp b/Engine/Source/Runtime/Android/AndroidAudio/Private/AndroidAudioSource.cpp
--- a/Engine/Source/Runtime/Android/AndroidAudio/Private/AndroidAudioSource.cpp
+++ b/Engine/Source/Runtime/Android/AndroidAudio/Private/AndroidAudioSource.cpp
@@ -348,14 +348,24 @@ void FSLESSoundSource::Update( void )
 		return;
 	}
 	
-	float Volume = WaveInstance->Volume * WaveInstance->VolumeMultiplier;
-	if( SetStereoBleed() )
+	float Volume = 1.0f;
+
+	// If the audio device is muted, then mute our sound source
+	if (AudioDevice->bIsDeviceMuted)
+	{
+		Volume = 0.0f;
+	}
+	else
 	{
-		// Emulate the bleed to rear speakers followed by stereo fold down
-		Volume *= 1.25f;
+		float Volume = WaveInstance->Volume * WaveInstance->VolumeMultiplier;
+		if (SetStereoBleed())
+		{
+			// Emulate the bleed to rear speakers followed by stereo fold down
+			Volume *= 1.25f;
+		}
+		Volume *= FApp::GetVolumeMultiplier();
+		Volume = FMath::Clamp(Volume, 0.0f, MAX_VOLUME);
 	}
-	Volume *= FApp::GetVolumeMultiplier();
-	Volume = FMath::Clamp(Volume, 0.0f, MAX_VOLUME);
 	
 	const float Pitch = FMath::Clamp<float>(WaveInstance->Pitch, MIN_PITCH, MAX_PITCH);
 
diff --git a/Engine/Source/Runtime/Engine/Classes/Components/AudioComponent.h b/Engine/Source/Runtime/Engine/Classes/Components/AudioComponent.h
--- a/Engine/Source/Runtime/Engine/Classes/Components/AudioComponent.h
+++ b/Engine/Source/Runtime/Engine/Classes/Components/AudioComponent.h
@@ -286,6 +286,9 @@ public:
 	/** Collects the various attenuation shapes that may be applied to the sound played by the audio component for visualization in the editor or via the in game debug visualization. */
 	void CollectAttenuationShapesForVisualization(TMultiMap<EAttenuationShape::Type, FAttenuationSettings::AttenuationShapeDetails>& ShapeDetailsMap) const;
 
+	/** Returns the active audio device to use for this component based on whether or not the component is playing in a world. */
+	FAudioDevice* GetAudioDevice() const;
+
 private:
 	
 #if WITH_EDITORONLY_DATA
diff --git a/Engine/Source/Runtime/Engine/Classes/Engine/Engine.h b/Engine/Source/Runtime/Engine/Classes/Engine/Engine.h
--- a/Engine/Source/Runtime/Engine/Classes/Engine/Engine.h
+++ b/Engine/Source/Runtime/Engine/Classes/Engine/Engine.h
@@ -283,6 +283,9 @@ struct FWorldContext
 	/** Is this world context waiting for an online login to complete (for PIE) */
 	bool	bWaitingOnOnlineSubsystem;
 
+	/** Handle to this world context's audio device.*/
+	uint32 AudioDeviceHandle;
+
 	/**************************************************************/
 
 	/** Outside pointers to CurrentWorld that should be kept in sync if current world changes  */
@@ -303,18 +306,7 @@ struct FWorldContext
 	}
 
 	/** Set CurrentWorld and update external reference pointers to reflect this*/
-	void SetCurrentWorld(UWorld *World)
-	{
-		for(int32 idx=0; idx < ExternalReferences.Num(); ++idx)
-		{
-			if (ExternalReferences[idx] && *ExternalReferences[idx] == ThisCurrentWorld)
-			{
-				*ExternalReferences[idx] = World;
-			}
-		}
-
-		ThisCurrentWorld = World;
-	}
+	ENGINE_API void SetCurrentWorld(UWorld *World);
 
 	/** Collect FWorldContext references for garbage collection */
 	void AddReferencedObjects(FReferenceCollector& Collector, const UObject* ReferencingObject);
@@ -336,6 +328,7 @@ struct FWorldContext
 		, PIEInstance(INDEX_NONE)
 		, RunAsDedicated(false)
 		, bWaitingOnOnlineSubsystem(false)
+		, AudioDeviceHandle(INDEX_NONE)
 		, ThisCurrentWorld(nullptr)
 	{ }
 
@@ -1519,14 +1512,20 @@ public:
 	 * Restores the selected material color back to the user setting
 	 */
 	void RestoreSelectedMaterialColor();
+
+protected:
+
+	/** The audio device manager */
+	class FAudioDeviceManager* AudioDeviceManager;
+
+	/** Audio device handle to the main audio device. */
+	uint32 MainAudioDeviceHandle;
+
 public:
 
 	/** A collection of messages to display on-screen. */
 	TMap<int32, FScreenMessageString> ScreenMessages;
 
-	/** The audio device */
-	class FAudioDevice* AudioDevice;
-
 	/** Reference to the stereoscopic rendering interace, if any */
 	TSharedPtr< class IStereoRendering > StereoRenderingDevice;
 
@@ -1647,7 +1646,7 @@ public:
 
 	/** Called at shutdown, just before the exit purge.	 */
 	virtual void PreExit();
-	virtual void ShutdownAudioDevice();
+	virtual void ShutdownAudioDeviceManager();
 
 	/** Called at startup, in the middle of FEngineLoop::Init.	 */
 	void ParseCommandline();
@@ -2042,11 +2041,18 @@ public:
 	/** @return the GIsEditor flag setting */
 	bool IsEditor();
 
-	/** @return the audio device (will be None if sound is disabled) */
-	virtual class FAudioDevice* GetAudioDevice()
-	{
-		return AudioDevice;
-	}
+	/** @return the audio device manager of the UEngine, this allows the creation and management of multiple audio devices. */
+	class FAudioDeviceManager* GetAudioDeviceManager();
+
+	/** @return the main audio device handle used by the engine. */
+	uint32 GetAudioDeviceHandle() const;
+
+	/** @return the main audio device. */
+	class FAudioDevice* GetMainAudioDevice();
+
+	DEPRECATED(4.8, "GetAudioDevice is deprecated UEngine::GetMainAudioDevice instead.")
+	/** @return the main audio device. */
+	class FAudioDevice* GetAudioDevice();
 
 	/** @return whether we're currently running in split screen (more than one local player) */
 	bool IsSplitScreen(UWorld *InWorld);
@@ -2191,11 +2197,11 @@ public:
 protected:
 
 	/**
-	 *	Initialize the audio device
+	 *	Initialize the audio device manager
 	 *
 	 *	@return	true on success, false otherwise.
 	 */
-	virtual bool InitializeAudioDevice();
+	virtual bool InitializeAudioDeviceManager();
 
 	/**
 	 *	Detects and initializes any attached HMD devices
diff --git a/Engine/Source/Runtime/Engine/Classes/Engine/GameViewportClient.h b/Engine/Source/Runtime/Engine/Classes/Engine/GameViewportClient.h
--- a/Engine/Source/Runtime/Engine/Classes/Engine/GameViewportClient.h
+++ b/Engine/Source/Runtime/Engine/Classes/Engine/GameViewportClient.h
@@ -699,6 +699,9 @@ private:
 
 	/** Handle to the registered ShowCollisionOnSpawnedActors delegate */
 	FDelegateHandle ShowCollisionOnSpawnedActorsDelegateHandle;
+
+	/** Handle to the audio device created for this viewport. Each viewport (for multiple PIE) will have its own audio device. */
+	uint32 AudioDeviceHandle;
 };
 
 
diff --git a/Engine/Source/Runtime/Engine/Classes/Engine/World.h b/Engine/Source/Runtime/Engine/Classes/Engine/World.h
--- a/Engine/Source/Runtime/Engine/Classes/Engine/World.h
+++ b/Engine/Source/Runtime/Engine/Classes/Engine/World.h
@@ -932,6 +932,9 @@ public:
 	/** Pointer to the higest priority audio volumes, each volume has a reference to the next lower priority volume creating a linked list of prioritized audio volumes in descending order */
 	TAutoWeakObjectPtr<class AAudioVolume> HighestPriorityAudioVolume;
 
+	/** Handle to the active audio device for this world. */
+	uint32 AudioDeviceHandle;
+
 	/** Time in FPlatformTime::Seconds unbuilt time was last encountered. 0 means not yet.							*/
 	double LastTimeUnbuiltLightingWasEncountered;
 
@@ -2961,6 +2964,16 @@ public:
 	 */
 	class AAudioVolume* GetAudioSettings( const FVector& ViewLocation, struct FReverbSettings* OutReverbSettings, struct FInteriorSettings* OutInteriorSettings );
 
+	/** Sets the audio device handle to the active audio device for this world.*/
+	void SetAudioDeviceHandle(const uint32 InAudioDeviceHandle);
+
+	/**
+	* Returns the audio device associated with this world, or returns the main audio device if there is none.
+	*
+	* @return Audio device to use with this world.
+	*/
+	class FAudioDevice* GetAudioDevice();
+
 	/** Return the URL of this level on the local machine. */
 	virtual FString GetLocalURL() const;
 
diff --git a/Engine/Source/Runtime/Engine/Classes/Kismet/GameplayStatics.h b/Engine/Source/Runtime/Engine/Classes/Kismet/GameplayStatics.h
--- a/Engine/Source/Runtime/Engine/Classes/Kismet/GameplayStatics.h
+++ b/Engine/Source/Runtime/Engine/Classes/Kismet/GameplayStatics.h
@@ -275,9 +275,9 @@ class ENGINE_API UGameplayStatics : public UBlueprintFunctionLibrary
 	 * @param MaximumRange	The maximum distance away from Location that a listener can be
 	 * @note This will always return false if there is no audio device, or the audio device is disabled.
 	 */
-	UFUNCTION(BlueprintCallable, Category="Audio")
-	static bool AreAnyListenersWithinRange(FVector Location, float MaximumRange);
-
+	UFUNCTION(BlueprintCallable, Category = "Audio", meta = (WorldContext = "WorldContextObject"))
+	static bool AreAnyListenersWithinRange(UObject* WorldContextObject, FVector Location, float MaximumRange);
+	
 	/**
 	 * Plays a sound directly with no attenuation, perfect for UI sounds.
 	 *
@@ -342,20 +342,20 @@ class ENGINE_API UGameplayStatics : public UBlueprintFunctionLibrary
 
 	// --- Audio Functions ----------------------------
 	/** Set the sound mix of the audio system for special EQing **/
-	UFUNCTION(BlueprintCallable, Category="Audio")
-	static void SetBaseSoundMix(class USoundMix* InSoundMix);
+	UFUNCTION(BlueprintCallable, Category="Audio", meta=(WorldContext = "WorldContextObject"))
+	static void SetBaseSoundMix(UObject* WorldContextObject, class USoundMix* InSoundMix);
 
 	/** Push a sound mix modifier onto the audio system **/
-	UFUNCTION(BlueprintCallable, Category="Audio")
-	static void PushSoundMixModifier(class USoundMix* InSoundMixModifier);
+	UFUNCTION(BlueprintCallable, Category="Audio", meta=(WorldContext = "WorldContextObject"))
+	static void PushSoundMixModifier(UObject* WorldContextObject, class USoundMix* InSoundMixModifier);
 
 	/** Pop a sound mix modifier from the audio system **/
-	UFUNCTION(BlueprintCallable, Category="Audio")
-	static void PopSoundMixModifier(class USoundMix* InSoundMixModifier);
+	UFUNCTION(BlueprintCallable, Category="Audio", meta=(WorldContext = "WorldContextObject"))
+	static void PopSoundMixModifier(UObject* WorldContextObject, class USoundMix* InSoundMixModifier);
 
 	/** Clear all sound mix modifiers from the audio system **/
-	UFUNCTION(BlueprintCallable, Category="Audio")
-	static void ClearSoundMixModifiers();
+	UFUNCTION(BlueprintCallable, Category="Audio", meta=(WorldContext = "WorldContextObject"))
+	static void ClearSoundMixModifiers(UObject* WorldContextObject);
 
 	/** Activates a Reverb Effect without the need for a volume
 	 * @param ReverbEffect Reverb Effect to use
@@ -364,16 +364,16 @@ class ENGINE_API UGameplayStatics : public UBlueprintFunctionLibrary
 	 * @param Volume Volume level of Reverb Effect
 	 * @param FadeTime Time before Reverb Effect is fully active
 	 */
-	UFUNCTION(BlueprintCallable, Category="Audio", meta=(AdvancedDisplay = "2"))
-	static void ActivateReverbEffect(class UReverbEffect* ReverbEffect, FName TagName, float Priority = 0.f, float Volume = 0.5f, float FadeTime = 2.f);
+	UFUNCTION(BlueprintCallable, Category="Audio", meta=(WorldContext = "WorldContextObject", AdvancedDisplay = "2"))
+	static void ActivateReverbEffect(UObject* WorldContextObject, class UReverbEffect* ReverbEffect, FName TagName, float Priority = 0.f, float Volume = 0.5f, float FadeTime = 2.f);
 
 	/**
 	 * Deactivates a Reverb Effect not applied by a volume
 	 *
 	 * @param TagName Tag associated with Reverb Effect to remove
 	 */
-	UFUNCTION(BlueprintCallable, Category="Audio")
-	static void DeactivateReverbEffect(FName TagName);
+	UFUNCTION(BlueprintCallable, Category="Audio", meta=(WorldContext = "WorldContextObject"))
+	static void DeactivateReverbEffect(UObject* WorldContextObject, FName TagName);
 
 	// --- Decal functions ------------------------------
 
diff --git a/Engine/Source/Runtime/Engine/Classes/Sound/SoundBase.h b/Engine/Source/Runtime/Engine/Classes/Sound/SoundBase.h
--- a/Engine/Source/Runtime/Engine/Classes/Sound/SoundBase.h
+++ b/Engine/Source/Runtime/Engine/Classes/Sound/SoundBase.h
@@ -84,7 +84,7 @@ public:
 	 * @param Location				Location to check against
 	 * @param AttenuationSettings	Optional Attenuation override if not using settings from the sound
 	 */
-	bool IsAudibleSimple( const FVector Location, USoundAttenuation* AttenuationSettings = NULL );
+	bool IsAudibleSimple( class FAudioDevice* AudioDevice, const FVector Location, USoundAttenuation* AttenuationSettings = NULL );
 
 	/** 
 	 * Returns the farthest distance at which the sound could be heard
diff --git a/Engine/Source/Runtime/Engine/Private/Audio.cpp b/Engine/Source/Runtime/Engine/Private/Audio.cpp
--- a/Engine/Source/Runtime/Engine/Private/Audio.cpp
+++ b/Engine/Source/Runtime/Engine/Private/Audio.cpp
@@ -59,9 +59,9 @@ DEFINE_STAT(STAT_AudioFindNearestLocation);
 FSoundBuffer::~FSoundBuffer()
 {
 	// remove ourselves from the set of waves that are tracked by the audio device
-	if (ResourceID && GEngine && GEngine->AudioDevice)
+	if (ResourceID && GEngine && GEngine->GetAudioDeviceManager())
 	{
-		GEngine->AudioDevice->WaveBufferMap.Remove(ResourceID);
+		GEngine->GetAudioDeviceManager()->RemoveSoundBufferForResourceID(ResourceID);
 	}
 }
 
diff --git a/Engine/Source/Runtime/Engine/Private/AudioComponent.cpp b/Engine/Source/Runtime/Engine/Private/AudioComponent.cpp
--- a/Engine/Source/Runtime/Engine/Private/AudioComponent.cpp
+++ b/Engine/Source/Runtime/Engine/Private/AudioComponent.cpp
@@ -33,9 +33,9 @@ FString UAudioComponent::GetDetailedInfoInternal( void ) const
 {
 	FString Result;
 
-	if( Sound != NULL )
+	if (Sound != nullptr)
 	{
-		Result = Sound->GetPathName( NULL );
+		Result = Sound->GetPathName(nullptr);
 	}
 	else
 	{
@@ -63,7 +63,7 @@ void UAudioComponent::OnUnregister()
 	// called from AActor::ClearComponents when an actor gets destroyed which is not usually what we want for one-
 	// shot sounds.
 	AActor* Owner = GetOwner();
-	if( !Owner || bStopWhenOwnerDestroyed )
+	if (!Owner || bStopWhenOwnerDestroyed)
 	{
 		Stop();
 	}
@@ -93,9 +93,7 @@ void UAudioComponent::OnUpdateTransform(bool bSkipPhysicsMove)
 
 	if (bIsActive && !bPreviewComponent)
 	{
-		// TODO - Audio Threading. This call would be a task
-		FAudioDevice* AudioDevice = GEngine ? GEngine->GetAudioDevice() : NULL;
-		if (AudioDevice != NULL)
+		if (FAudioDevice * AudioDevice = GetAudioDevice())
 		{
 			FActiveSound* ActiveSound = AudioDevice->FindActiveSound(this);
 			if (ActiveSound)
@@ -116,7 +114,7 @@ void UAudioComponent::PlayInternal(const float StartTime, const float FadeInDura
 {
 	UWorld* World = GetWorld();
 
-	UE_LOG(LogAudio, Verbose, TEXT( "%g: Playing AudioComponent : '%s' with Sound: '%s'" ), World ? World->GetAudioTimeSeconds() : 0.0f, *GetFullName(), Sound ? *Sound->GetName() : TEXT( "NULL" ) );
+	UE_LOG(LogAudio, Verbose, TEXT("%g: Playing AudioComponent : '%s' with Sound: '%s'"), World ? World->GetAudioTimeSeconds() : 0.0f, *GetFullName(), Sound ? *Sound->GetName() : TEXT("nullptr"));
 
 	if (bIsActive)
 	{
@@ -129,8 +127,7 @@ void UAudioComponent::PlayInternal(const float StartTime, const float FadeInDura
 
 	if (Sound && (World == nullptr || World->bAllowAudioPlayback))
 	{
-		FAudioDevice* AudioDevice = GEngine ? GEngine->GetAudioDevice() : NULL;
-		if (AudioDevice != NULL)
+		if (FAudioDevice* AudioDevice = GetAudioDevice())
 		{
 			FActiveSound NewActiveSound;
 			NewActiveSound.AudioComponent = this;
@@ -138,8 +135,8 @@ void UAudioComponent::PlayInternal(const float StartTime, const float FadeInDura
 			NewActiveSound.Sound = Sound;
 			NewActiveSound.SoundClassOverride = SoundClassOverride;
 
-			NewActiveSound.VolumeMultiplier = (VolumeModulationMax + ( ( VolumeModulationMin - VolumeModulationMax ) * FMath::SRand() )) * VolumeMultiplier;
-			NewActiveSound.PitchMultiplier = (PitchModulationMax + ( ( PitchModulationMin - PitchModulationMax ) * FMath::SRand() )) * PitchMultiplier;
+			NewActiveSound.VolumeMultiplier = (VolumeModulationMax + ((VolumeModulationMin - VolumeModulationMax) * FMath::SRand())) * VolumeMultiplier;
+			NewActiveSound.PitchMultiplier = (PitchModulationMax + ((PitchModulationMin - PitchModulationMax) * FMath::SRand())) * PitchMultiplier;
 			NewActiveSound.HighFrequencyGainMultiplier = HighFrequencyGainMultiplier;
 
 			NewActiveSound.RequestedStartTime = FMath::Max(0.f, StartTime);
@@ -162,9 +159,9 @@ void UAudioComponent::PlayInternal(const float StartTime, const float FadeInDura
 				NewActiveSound.Transform = ComponentToWorld;
 			}
 
-			const FAttenuationSettings* AttenuationSettingsToApply = (bAllowSpatialization ? GetAttenuationSettingsToApply() : NULL);
+			const FAttenuationSettings* AttenuationSettingsToApply = (bAllowSpatialization ? GetAttenuationSettingsToApply() : nullptr);
 			NewActiveSound.bAllowSpatialization = bAllowSpatialization;
-			NewActiveSound.bHasAttenuationSettings = (AttenuationSettingsToApply != NULL);
+			NewActiveSound.bHasAttenuationSettings = (AttenuationSettingsToApply != nullptr);
 			if (NewActiveSound.bHasAttenuationSettings)
 			{
 				NewActiveSound.AttenuationSettings = *AttenuationSettingsToApply;
@@ -192,6 +189,25 @@ void UAudioComponent::PlayInternal(const float StartTime, const float FadeInDura
 	}
 }
 
+FAudioDevice* UAudioComponent::GetAudioDevice() const
+{
+	FAudioDevice* AudioDevice = nullptr;
+
+	if (GEngine)
+	{
+		UWorld* World = GetWorld();
+		if (World)
+		{
+			AudioDevice = World->GetAudioDevice();
+		}
+		else
+		{
+			AudioDevice = GEngine->GetMainAudioDevice();
+		}
+	}
+	return AudioDevice;
+}
+
 void UAudioComponent::FadeIn( float FadeInDuration, float FadeVolumeLevel, float StartTime )
 {
 	PlayInternal(StartTime, FadeInDuration, FadeVolumeLevel);
@@ -204,8 +220,7 @@ void UAudioComponent::FadeOut( float FadeOutDuration, float FadeVolumeLevel )
 		if (FadeOutDuration > 0.0f)
 		{
 			// TODO - Audio Threading. This call would be a task
-			FAudioDevice* AudioDevice = GEngine ? GEngine->GetAudioDevice() : NULL;
-			if (AudioDevice != NULL)
+			if (FAudioDevice* AudioDevice = GetAudioDevice())
 			{
 				FActiveSound* ActiveSound = AudioDevice->FindActiveSound(this);
 				if (ActiveSound)
@@ -228,8 +243,7 @@ void UAudioComponent::AdjustVolume( float AdjustVolumeDuration, float AdjustVolu
 	if (bIsActive)
 	{
 		// TODO - Audio Threading. This call would be a task
-		FAudioDevice* AudioDevice = GEngine ? GEngine->GetAudioDevice() : NULL;
-		if (AudioDevice != NULL)
+		if (FAudioDevice* AudioDevice = GetAudioDevice())
 		{
 			FActiveSound* ActiveSound = AudioDevice->FindActiveSound(this);
 			if (ActiveSound)
@@ -255,11 +269,10 @@ void UAudioComponent::Stop()
 {
 	if (bIsActive)
 	{
-		UE_LOG(LogAudio, Verbose, TEXT( "%g: Stopping AudioComponent : '%s' with Sound: '%s'" ), GetWorld() ? GetWorld()->GetAudioTimeSeconds() : 0.0f, *GetFullName(), Sound ? *Sound->GetName() : TEXT( "NULL" ) );
+		UE_LOG(LogAudio, Verbose, TEXT( "%g: Stopping AudioComponent : '%s' with Sound: '%s'" ), GetWorld() ? GetWorld()->GetAudioTimeSeconds() : 0.0f, *GetFullName(), Sound ? *Sound->GetName() : TEXT( "nullptr" ) );
 
 		// TODO - Audio Threading. This call would be a task
-		FAudioDevice* AudioDevice = GEngine ? GEngine->GetAudioDevice() : NULL;
-		if (AudioDevice != NULL)
+		if (FAudioDevice* AudioDevice = GetAudioDevice())
 		{
 			AudioDevice->StopActiveSound(this);
 		}
@@ -271,7 +284,7 @@ void UAudioComponent::PlaybackCompleted(bool bFailedToStart)
 	// Mark inactive before calling destroy to avoid recursion
 	bIsActive = false;
 
-	if (!bFailedToStart && GetWorld() != NULL && (OnAudioFinished.IsBound() || OnAudioFinishedNative.IsBound()))
+	if (!bFailedToStart && GetWorld() != nullptr && (OnAudioFinished.IsBound() || OnAudioFinishedNative.IsBound()))
 	{
 		INC_DWORD_STAT( STAT_AudioFinishedDelegatesCalled );
 		SCOPE_CYCLE_COUNTER( STAT_AudioFinishedDelegates );
@@ -281,7 +294,7 @@ void UAudioComponent::PlaybackCompleted(bool bFailedToStart)
 	}
 
 	// Auto destruction is handled via marking object for deletion.
-	if( bAutoDestroy )
+	if (bAutoDestroy)
 	{
 		DestroyComponent();
 	}
@@ -299,11 +312,11 @@ void UAudioComponent::UpdateSpriteTexture()
 	{
 		if (bAutoActivate)
 		{
-			SpriteComponent->SetSprite(LoadObject<UTexture2D>(NULL, TEXT("/Engine/EditorResources/AudioIcons/S_AudioComponent_AutoActivate.S_AudioComponent_AutoActivate")));
+			SpriteComponent->SetSprite(LoadObject<UTexture2D>(nullptr, TEXT("/Engine/EditorResources/AudioIcons/S_AudioComponent_AutoActivate.S_AudioComponent_AutoActivate")));
 		}
 		else
 		{
-			SpriteComponent->SetSprite(LoadObject<UTexture2D>(NULL, TEXT("/Engine/EditorResources/AudioIcons/S_AudioComponent.S_AudioComponent")));
+			SpriteComponent->SetSprite(LoadObject<UTexture2D>(nullptr, TEXT("/Engine/EditorResources/AudioIcons/S_AudioComponent.S_AudioComponent")));
 		}
 	}
 }
@@ -344,7 +357,7 @@ const FAttenuationSettings* UAudioComponent::GetAttenuationSettingsToApply() con
 	{
 		return Sound->GetAttenuationSettingsToApply();
 	}
-	return NULL;
+	return nullptr;
 }
 
 void UAudioComponent::CollectAttenuationShapesForVisualization(TMultiMap<EAttenuationShape::Type, FAttenuationSettings::AttenuationShapeDetails>& ShapeDetailsMap) const
@@ -375,7 +388,7 @@ void UAudioComponent::CollectAttenuationShapesForVisualization(TMultiMap<EAttenu
 
 void UAudioComponent::Activate(bool bReset)
 {
-	if (bReset || ShouldActivate()==true)
+	if (bReset || ShouldActivate() == true)
 	{
 		Play();
 	}
@@ -383,7 +396,7 @@ void UAudioComponent::Activate(bool bReset)
 
 void UAudioComponent::Deactivate()
 {
-	if (ShouldActivate()==false)
+	if (ShouldActivate() == false)
 	{
 		Stop();
 	}
@@ -391,15 +404,15 @@ void UAudioComponent::Deactivate()
 
 void UAudioComponent::SetFloatParameter( FName InName, float InFloat )
 {
-	if( InName != NAME_None )
+	if (InName != NAME_None)
 	{
 		bool bFound = false;
 
 		// First see if an entry for this name already exists
-		for( int32 i = 0; i < InstanceParameters.Num(); i++ )
+		for (int32 i = 0; i < InstanceParameters.Num(); i++)
 		{
-			FAudioComponentParam& P = InstanceParameters[ i ];
-			if( P.ParamName == InName )
+			FAudioComponentParam& P = InstanceParameters[i];
+			if (P.ParamName == InName)
 			{
 				P.FloatParam = InFloat;
 				bFound = true;
@@ -419,8 +432,7 @@ void UAudioComponent::SetFloatParameter( FName InName, float InFloat )
 		if (bIsActive)
 		{
 			// TODO - Audio Threading. This call would be a task
-			FAudioDevice* AudioDevice = GEngine ? GEngine->GetAudioDevice() : NULL;
-			if (AudioDevice != NULL)
+			if (FAudioDevice* AudioDevice = GetAudioDevice())
 			{
 				FActiveSound* ActiveSound = AudioDevice->FindActiveSound(this);
 				if (ActiveSound)
@@ -434,14 +446,14 @@ void UAudioComponent::SetFloatParameter( FName InName, float InFloat )
 
 void UAudioComponent::SetWaveParameter( FName InName, USoundWave* InWave )
 {
-	if( InName != NAME_None )
+	if (InName != NAME_None)
 	{
 		bool bFound = false;
 		// First see if an entry for this name already exists
-		for( int32 i = 0; i < InstanceParameters.Num(); i++ )
+		for (int32 i = 0; i < InstanceParameters.Num(); i++)
 		{
-			FAudioComponentParam& P = InstanceParameters[ i ];
-			if( P.ParamName == InName )
+			FAudioComponentParam& P = InstanceParameters[i];
+			if (P.ParamName == InName)
 			{
 				P.SoundWaveParam = InWave;
 				bFound = true;
@@ -453,16 +465,15 @@ void UAudioComponent::SetWaveParameter( FName InName, USoundWave* InWave )
 		if (!bFound)
 		{
 			const int32 NewParamIndex = InstanceParameters.AddZeroed();
-			InstanceParameters[ NewParamIndex ].ParamName = InName;
-			InstanceParameters[ NewParamIndex ].SoundWaveParam = InWave;
+			InstanceParameters[NewParamIndex].ParamName = InName;
+			InstanceParameters[NewParamIndex].SoundWaveParam = InWave;
 		}
 
 		// If we're active we need to push this value to the ActiveSound
 		if (bIsActive)
 		{
 			// TODO - Audio Threading. This call would be a task
-			FAudioDevice* AudioDevice = GEngine ? GEngine->GetAudioDevice() : NULL;
-			if (AudioDevice != NULL)
+			if (FAudioDevice* AudioDevice = GetAudioDevice())
 			{
 				FActiveSound* ActiveSound = AudioDevice->FindActiveSound(this);
 				if (ActiveSound)
@@ -476,14 +487,14 @@ void UAudioComponent::SetWaveParameter( FName InName, USoundWave* InWave )
 
 void UAudioComponent::SetBoolParameter( FName InName, bool InBool )
 {
-	if( InName != NAME_None )
+	if (InName != NAME_None)
 	{
 		bool bFound = false;
 		// First see if an entry for this name already exists
-		for( int32 i = 0; i < InstanceParameters.Num(); i++ )
+		for (int32 i = 0; i < InstanceParameters.Num(); i++)
 		{
-			FAudioComponentParam& P = InstanceParameters[ i ];
-			if( P.ParamName == InName )
+			FAudioComponentParam& P = InstanceParameters[i];
+			if (P.ParamName == InName)
 			{
 				P.BoolParam = InBool;
 				bFound = true;
@@ -503,8 +514,7 @@ void UAudioComponent::SetBoolParameter( FName InName, bool InBool )
 		if (bIsActive)
 		{
 			// TODO - Audio Threading. This call would be a task
-			FAudioDevice* AudioDevice = GEngine ? GEngine->GetAudioDevice() : NULL;
-			if (AudioDevice != NULL)
+			if (FAudioDevice* AudioDevice = GetAudioDevice())
 			{
 				FActiveSound* ActiveSound = AudioDevice->FindActiveSound(this);
 				if (ActiveSound)
@@ -519,14 +529,14 @@ void UAudioComponent::SetBoolParameter( FName InName, bool InBool )
 
 void UAudioComponent::SetIntParameter( FName InName, int32 InInt )
 {
-	if( InName != NAME_None )
+	if (InName != NAME_None)
 	{
 		bool bFound = false;
 		// First see if an entry for this name already exists
-		for( int32 i = 0; i < InstanceParameters.Num(); i++ )
+		for (int32 i = 0; i < InstanceParameters.Num(); i++)
 		{
-			FAudioComponentParam& P = InstanceParameters[ i ];
-			if( P.ParamName == InName )
+			FAudioComponentParam& P = InstanceParameters[i];
+			if (P.ParamName == InName)
 			{
 				P.IntParam = InInt;
 				bFound = true;
@@ -538,16 +548,15 @@ void UAudioComponent::SetIntParameter( FName InName, int32 InInt )
 		if (!bFound)
 		{
 			const int32 NewParamIndex = InstanceParameters.AddZeroed();
-			InstanceParameters[ NewParamIndex ].ParamName = InName;
-			InstanceParameters[ NewParamIndex ].IntParam = InInt;
+			InstanceParameters[NewParamIndex].ParamName = InName;
+			InstanceParameters[NewParamIndex].IntParam = InInt;
 		}
 
 		// If we're active we need to push this value to the ActiveSound
 		if (bIsActive)
 		{
 			// TODO - Audio Threading. This call would be a task
-			FAudioDevice* AudioDevice = GEngine ? GEngine->GetAudioDevice() : NULL;
-			if (AudioDevice != NULL)
+			if (FAudioDevice* AudioDevice = GetAudioDevice())
 			{
 				FActiveSound* ActiveSound = AudioDevice->FindActiveSound(this);
 				if (ActiveSound)
@@ -565,8 +574,7 @@ void UAudioComponent::SetVolumeMultiplier(float NewVolumeMultiplier)
 	VolumeModulationMin = VolumeModulationMax = 1.f;
 
 	// TODO - Audio Threading. This call would be a task
-	FAudioDevice* AudioDevice = GEngine ? GEngine->GetAudioDevice() : NULL;
-	if (AudioDevice != NULL)
+	if (FAudioDevice* AudioDevice = GetAudioDevice())
 	{
 		FActiveSound* ActiveSound = AudioDevice->FindActiveSound(this);
 		if (ActiveSound)
@@ -582,8 +590,7 @@ void UAudioComponent::SetPitchMultiplier(float NewPitchMultiplier)
 	PitchModulationMin = PitchModulationMax = 1.f;
 
 	// TODO - Audio Threading. This call would be a task
-	FAudioDevice* AudioDevice = GEngine ? GEngine->GetAudioDevice() : NULL;
-	if (AudioDevice != NULL)
+	if (FAudioDevice* AudioDevice = GetAudioDevice())
 	{
 		FActiveSound* ActiveSound = AudioDevice->FindActiveSound(this);
 		if (ActiveSound)
@@ -598,8 +605,7 @@ void UAudioComponent::SetUISound(bool bInIsUISound)
 	bIsUISound = bInIsUISound;
 
 	// TODO - Audio Threading. This call would be a task
-	FAudioDevice* AudioDevice = GEngine ? GEngine->GetAudioDevice() : NULL;
-	if (AudioDevice != NULL)
+	if (FAudioDevice* AudioDevice = GetAudioDevice())
 	{
 		FActiveSound* ActiveSound = AudioDevice->FindActiveSound(this);
 		if (ActiveSound)
@@ -615,8 +621,7 @@ void UAudioComponent::AdjustAttenuation(const FAttenuationSettings& InAttenuatio
 	AttenuationOverrides = InAttenuationSettings;
 
 	// TODO - Audio Threading. This call would be a task
-	FAudioDevice* AudioDevice = GEngine ? GEngine->GetAudioDevice() : NULL;
-	if (AudioDevice != NULL)
+	if (FAudioDevice* AudioDevice = GetAudioDevice())
 	{
 		FActiveSound* ActiveSound = AudioDevice->FindActiveSound(this);
 		if (ActiveSound)
diff --git a/Engine/Source/Runtime/Engine/Private/AudioDecompress.cpp b/Engine/Source/Runtime/Engine/Private/AudioDecompress.cpp
--- a/Engine/Source/Runtime/Engine/Private/AudioDecompress.cpp
+++ b/Engine/Source/Runtime/Engine/Private/AudioDecompress.cpp
@@ -284,9 +284,9 @@ FAsyncAudioDecompressWorker::FAsyncAudioDecompressWorker(USoundWave* InWave)
 	: Wave(InWave)
 	, AudioInfo(NULL)
 {
-	if (GEngine && GEngine->GetAudioDevice())
+	if (GEngine && GEngine->GetMainAudioDevice())
 	{
-		AudioInfo = GEngine->GetAudioDevice()->CreateCompressedAudioInfo(Wave);
+		AudioInfo = GEngine->GetMainAudioDevice()->CreateCompressedAudioInfo(Wave);
 	}
 }
 
diff --git a/Engine/Source/Runtime/Engine/Private/AudioDevice.cpp b/Engine/Source/Runtime/Engine/Private/AudioDevice.cpp
--- a/Engine/Source/Runtime/Engine/Private/AudioDevice.cpp
+++ b/Engine/Source/Runtime/Engine/Private/AudioDevice.cpp
@@ -32,6 +32,7 @@ FAudioDevice::FAudioDevice()
 	, Effects(NULL)
 	, CurrentAudioVolume(NULL)
 	, HighestPriorityReverb(NULL)
+	, bIsDeviceMuted(false)
 {
 }
 
@@ -135,15 +136,12 @@ void FAudioDevice::Teardown()
 	// let platform shutdown
 	TeardownHardware();
 
-	// Release any loaded buffers - this calls stop on any sources that need it
-	for (int32 Index = Buffers.Num() - 1; Index >= 0; Index--)
-	{
-		FreeBufferResource(Buffers[Index]);
-	}
+	// Note: we don't free audio buffers at this stage since they are managed in the audio device manager
 
 	// Must be after FreeBufferResource as that potentially stops sources
 	for (int32 Index = 0; Index < Sources.Num(); Index++)
 	{
+		Sources[Index]->Stop();
 		delete Sources[Index];
 	}
 	Sources.Empty();
@@ -158,7 +156,9 @@ void FAudioDevice::Suspend(bool bGameTicking)
 void FAudioDevice::CountBytes(FArchive& Ar)
 {
 	Sources.CountBytes(Ar);
-	Buffers.CountBytes(Ar);
+	// The buffers are stored on the audio device since they are shared amongst all audio devices
+	// Though we are going to count them when querying an individual audio device object about its bytes
+	GEngine->GetAudioDeviceManager()->Buffers.CountBytes(Ar);
 	FreeSources.CountBytes(Ar);
 	WaveInstanceSourceMap.CountBytes(Ar);
 	Ar.CountBytes(sizeof(FWaveInstance) * WaveInstanceSourceMap.Num(), sizeof(FWaveInstance) * WaveInstanceSourceMap.Num());
@@ -487,13 +487,11 @@ bool FAudioDevice::HandleListAudioComponentsCommand( const TCHAR* Cmd, FOutputDe
 	}
 	Ar.Logf( TEXT( "AudioComponent Total = %d" ), Count );
 
-	FAudioDevice* AudioDevice = GEngine->AudioDevice;
-
 	Ar.Logf( TEXT( "AudioDevice 0x%p has %d ActiveSounds" ),
-		AudioDevice, AudioDevice->ActiveSounds.Num());
-	for( int32 ASIndex = 0; ASIndex < AudioDevice->ActiveSounds.Num(); ASIndex++ )
+		this, ActiveSounds.Num());
+	for( int32 ASIndex = 0; ASIndex < ActiveSounds.Num(); ASIndex++ )
 	{
-		const FActiveSound* ActiveSound = AudioDevice->ActiveSounds[ASIndex];
+		const FActiveSound* ActiveSound = ActiveSounds[ASIndex];
 		UAudioComponent* AComp = ActiveSound->AudioComponent.Get();
 		if( AComp )
 		{
@@ -774,14 +772,6 @@ bool FAudioDevice::Exec( UWorld* InWorld, const TCHAR* Cmd, FOutputDevice& Ar )
 	return false;
 }
 
-void FAudioDevice::RemoveClass( USoundClass* SoundClass )
-{
-	if( SoundClass )
-	{
-		SoundClasses.Remove(SoundClass);
-	}
-}
-
 void FAudioDevice::InitSoundClasses( void )
 {
 	// Reset the maps of sound class properties
@@ -2022,25 +2012,39 @@ bool FAudioDevice::LocationIsAudible( FVector Location, float MaxDistance )
 	return( false );
 }
 
-UAudioComponent* FAudioDevice::CreateComponent( USoundBase* Sound, UWorld* World, AActor* Actor, bool bPlay, bool bStopWhenOwnerDestroyed, const FVector* Location, USoundAttenuation* AttenuationSettings )
+UAudioComponent* FAudioDevice::CreateComponent(USoundBase* Sound, UWorld* World, AActor* Actor, bool bPlay, bool bStopWhenOwnerDestroyed, const FVector* Location, USoundAttenuation* AttenuationSettings)
 {
 	UAudioComponent* AudioComponent = NULL;
 
 	if( Sound && GEngine && GEngine->UseSound() )
 	{
+		// Get the world's audio device if there is a world, otherwise get the main audio device
+		FAudioDevice* AudioDevice = nullptr;
+		if (World)
+		{
+			AudioDevice = World->GetAudioDevice();
+		}
+		else
+		{
+			AudioDevice = GEngine->GetMainAudioDevice();
+		}
+		check(AudioDevice != nullptr);
+
 		// Avoid creating component if we're trying to play a sound on an already destroyed actor.
-		if( Actor && Actor->IsPendingKill() )
+		if (Actor && Actor->IsPendingKill())
 		{
 			// Don't create component on destroyed actor.
 		}
 		// Either no actor or actor is still alive.
-		else if( Location && !Sound->IsAudibleSimple( *Location, AttenuationSettings ) )
+		else if (Location && !Sound->IsAudibleSimple(AudioDevice, *Location, AttenuationSettings))
 		{
 			// Don't create a sound component for short sounds that start out of range of any listener
 			UE_LOG(LogAudio, Log, TEXT( "AudioComponent not created for out of range Sound %s" ), *Sound->GetName() );
 		}
 		else
 		{
+
+
 			// Use actor as outer if we have one.
 			if( Actor )
 			{
@@ -2251,82 +2255,54 @@ void FAudioDevice::Precache(USoundWave* SoundWave, bool bSynchronous, bool bTrac
 	}
 }
 
-/** 
- * Links up the resource data indices for looking up and cleaning up
- */
-void FAudioDevice::TrackResource(USoundWave* Wave, FSoundBuffer* Buffer)
+void FAudioDevice::StopSourcesUsingBuffer(FSoundBuffer * SoundBuffer)
 {
-	// Allocate new resource ID and assign to USoundWave. A value of 0 (default) means not yet registered.
-	int32 ResourceID = NextResourceID++;
-	Buffer->ResourceID = ResourceID;
-	Wave->ResourceID = ResourceID;
-
-	Buffers.Add(Buffer);
-	WaveBufferMap.Add(ResourceID, Buffer);
-
-#if !(UE_BUILD_SHIPPING || UE_BUILD_TEST)
-	// Keep track of associated resource name.
-	Buffer->ResourceName = Wave->GetPathName();
-#endif
-}
-
-void FAudioDevice::FreeResource(USoundWave* SoundWave)
-{
-	// Find buffer for resident wavs
-	if (SoundWave->ResourceID)
-	{
-		// Find buffer associated with resource id.
-		FSoundBuffer* Buffer = WaveBufferMap.FindRef(SoundWave->ResourceID);
-		FreeBufferResource(Buffer);
-
-		SoundWave->ResourceID = 0;
-	}
-}
-
-void FAudioDevice::FreeBufferResource(FSoundBuffer* Buffer)
-{
-	if (Buffer)
+	if (SoundBuffer)
 	{
-		// Remove from buffers array.
-		Buffers.Remove(Buffer);
-
-		// See if it is being used by a sound source...
 		for (int32 SrcIndex = 0; SrcIndex < Sources.Num(); SrcIndex++)
 		{
 			FSoundSource* Src = Sources[SrcIndex];
-			if (Src && Src->Buffer == Buffer)
+			if (Src && Src->Buffer == SoundBuffer)
 			{
 				// Make sure the buffer is no longer referenced by anything
 				Src->Stop();
 				break;
 			}
 		}
-
-		// Delete it. This will automatically remove itself from the WaveBufferMap.
-		delete Buffer;
 	}
 }
 
-
-FSoundClassProperties* FAudioDevice::GetSoundClassCurrentProperties( USoundClass* InSoundClass )
+void FAudioDevice::RegisterSoundClass(USoundClass* InSoundClass)
 {
-	FSoundClassProperties* Properties = NULL;
 	if (InSoundClass)
 	{
-		Properties = SoundClasses.Find( InSoundClass );
-
-		// If the sound class wasn't already loaded get it in to the system.
-		if (!Properties)
+		// If the sound class wasn't already registered get it in to the system.
+		if (!SoundClasses.Contains(InSoundClass))
 		{
-			FSoundClassProperties& PropertiesRef = SoundClasses.Add( InSoundClass, FSoundClassProperties() );
-			PropertiesRef = InSoundClass->Properties;
-			Properties = &PropertiesRef;
+			SoundClasses.Add(InSoundClass, FSoundClassProperties());
 		}
 	}
+}
+
+void FAudioDevice::UnregisterSoundClass(USoundClass* SoundClass)
+{
+	if (SoundClass)
+	{
+		SoundClasses.Remove(SoundClass);
+	}
+}
 
-	return Properties;
+FSoundClassProperties* FAudioDevice::GetSoundClassCurrentProperties(class USoundClass* InSoundClass)
+{
+	if (InSoundClass)
+	{
+		FSoundClassProperties* Properties = SoundClasses.Find(InSoundClass);
+		return Properties;
+	}
+	return nullptr;
 }
 
+
 #if !UE_BUILD_SHIPPING
 // sort memory usage from large to small 
 struct FSortSoundBuffersBySize
@@ -2360,10 +2336,14 @@ bool FAudioDevice::HandleListSoundsCommand(const TCHAR* Cmd, FOutputDevice& Ar)
 
 	Ar.Logf(TEXT("Listing all sounds:"));
 
+	// Get audio device manager since thats where sound buffers are stored
+	FAudioDeviceManager* AudioDeviceManager = GEngine->GetAudioDeviceManager();
+	check(AudioDeviceManager != nullptr);
+
 	TArray<FSoundBuffer*> AllSounds;
-	for (int32 BufferIndex = 0; BufferIndex < Buffers.Num(); BufferIndex++)
+	for (int32 BufferIndex = 0; BufferIndex < AudioDeviceManager->Buffers.Num(); BufferIndex++)
 	{
-		AllSounds.Add(Buffers[BufferIndex]);
+		AllSounds.Add(AudioDeviceManager->Buffers[BufferIndex]);
 	}
 
 	// sort by name or size, depending on flag
diff --git a/Engine/Source/Runtime/Engine/Private/AudioDeviceManager.cpp b/Engine/Source/Runtime/Engine/Private/AudioDeviceManager.cpp
new file mode 100644
--- /dev/null
+++ b/Engine/Source/Runtime/Engine/Private/AudioDeviceManager.cpp
@@ -0,0 +1,375 @@
+// Copyright 1998-2015 Epic Games, Inc. All Rights Reserved.
+
+#include "EnginePrivate.h"
+#include "Audio.h"
+#include "AudioDevice.h"
+#include "AudioDeviceManager.h"
+#include "Sound/SoundWave.h"
+
+// Private consts for helping with index/generation determination in audio device manager
+static const uint32 AUDIO_DEVICE_HANDLE_INDEX_BITS		= 24;
+static const uint32 AUDIO_DEVICE_HANDLE_INDEX_MASK		= (1 << AUDIO_DEVICE_HANDLE_INDEX_BITS) - 1;
+static const uint32 AUDIO_DEVICE_HANDLE_GENERATION_BITS = 8;
+static const uint32 AUDIO_DEVICE_HANDLE_GENERATION_MASK = (1 << AUDIO_DEVICE_HANDLE_GENERATION_BITS) - 1;
+
+static const uint16 AUDIO_DEVICE_MINIMUM_FREE_AUDIO_DEVICE_INDICES = 32;
+
+// Invalid handle used for initializing audio device handles
+static const uint32 AUDIO_DEVICE_HANDLE_INVALID			= ~0u;
+
+/*-----------------------------------------------------------------------------
+FAudioDeviceManager implementation.
+-----------------------------------------------------------------------------*/
+
+FAudioDeviceManager::FAudioDeviceManager()
+	: AudioDeviceModule(nullptr)
+	, FreeIndicesSize(0)
+	, NumActiveAudioDevices(0)
+	, NextResourceID(1)
+{
+}
+
+FAudioDeviceManager::~FAudioDeviceManager()
+{
+	// Confirm that we freed all the audio devices
+	check(NumActiveAudioDevices == 0);
+
+	// Release any loaded buffers - this calls stop on any sources that need it
+	for (int32 Index = Buffers.Num() - 1; Index >= 0; Index--)
+	{
+		FreeBufferResource(Buffers[Index]);
+	}
+}
+
+void FAudioDeviceManager::RegisterAudioDeviceModule(IAudioDeviceModule* AudioDeviceModuleInput)
+{
+	// Currently can't have multiple audio device modules registered
+	check(AudioDeviceModule == nullptr);
+	AudioDeviceModule = AudioDeviceModuleInput;
+}
+
+FAudioDevice* FAudioDeviceManager::CreateAudioDevice(uint32& HandleOut)
+{
+	// If we don't have an audio device module, then we can't create new audio devices.
+	if (AudioDeviceModule == nullptr)
+	{
+		HandleOut = AUDIO_DEVICE_HANDLE_INVALID;
+		return nullptr;
+	}
+
+	// Create the new audio device and make sure it succeeded
+	FAudioDevice* NewAudioDevice = AudioDeviceModule->CreateAudioDevice();
+	if (NewAudioDevice == nullptr)
+	{
+		HandleOut = AUDIO_DEVICE_HANDLE_INVALID;
+		return nullptr;
+	}
+
+	if (!NewAudioDevice->Init())
+	{
+		delete NewAudioDevice;
+		HandleOut = AUDIO_DEVICE_HANDLE_INVALID;
+		return nullptr;
+	}
+
+	// Now generation a new audio device handle for the device and store the
+	// ptr to the new device in the array of audio devices.
+
+	uint32 AudioDeviceIndex(AUDIO_DEVICE_HANDLE_INVALID);
+
+	// First check to see if we should start recycling audio device indices, if not
+	// then we add a new entry to the Generation array and generate a new index
+	if (FreeIndicesSize > AUDIO_DEVICE_MINIMUM_FREE_AUDIO_DEVICE_INDICES)
+	{
+		FreeIndices.Dequeue(AudioDeviceIndex);
+		--FreeIndicesSize;
+		check(int32(AudioDeviceIndex) < Devices.Num());
+		check(Devices[AudioDeviceIndex] == nullptr);
+		Devices[AudioDeviceIndex] = NewAudioDevice;
+	}
+	else
+	{
+		// Add a zeroth generation entry in the Generation array, get a brand new
+		// index and append the created device to the end of the Devices array
+
+		Generations.Add(0);
+		AudioDeviceIndex = Generations.Num() - 1;
+		check(AudioDeviceIndex < (1 << AUDIO_DEVICE_HANDLE_INDEX_BITS));
+		Devices.Add(NewAudioDevice);
+	}
+
+	++NumActiveAudioDevices;
+	HandleOut = CreateHandle(AudioDeviceIndex, Generations[AudioDeviceIndex]);
+
+	// Store the handle on the audio device itself
+	NewAudioDevice->DeviceHandle = HandleOut;
+
+	return NewAudioDevice;
+}
+
+bool FAudioDeviceManager::IsValidAudioDeviceHandle(uint32 Handle) const
+{
+	if (AudioDeviceModule == nullptr || Handle == AUDIO_DEVICE_HANDLE_INVALID)
+	{
+		return false;
+	}
+
+	uint32 Index = GetIndex(Handle);
+	if (int32(Index) >= Generations.Num())
+	{
+		return false;
+	}
+
+	uint8 Generation = GetGeneration(Handle);
+	return Generations[Index] == Generation;
+}
+
+bool FAudioDeviceManager::ShutdownAudioDevice(uint32 Handle)
+{
+	if (!IsValidAudioDeviceHandle(Handle))
+	{
+		return false;
+	}
+
+	check(NumActiveAudioDevices > 0);
+	--NumActiveAudioDevices;
+
+	// If we only have one audio device left, then set the active
+	// audio device to be the main audio device
+	if (NumActiveAudioDevices == 1)
+	{
+		uint32 ActiveDeviceHandle = GEngine->GetAudioDeviceHandle();
+		SetActiveDevice(ActiveDeviceHandle);
+	}
+
+	uint32 Index = GetIndex(Handle);
+	uint8 Generation = GetGeneration(Handle);
+
+	check(int32(Index) < Generations.Num());
+
+	// Bump up the generation at the given index. This will invalidate
+	// the handle without needing to broadcast to everybody who might be using the handle
+	Generations[Index] = ++Generation;
+
+	// Make sure we have a non-null device ptr in the index slot, then delete it
+	FAudioDevice* AudioDevice = Devices[Index];
+	check(AudioDevice != nullptr);
+
+    // Tear down the audio device
+	AudioDevice->Teardown();
+
+	delete AudioDevice;
+
+	// Nullify the audio device slot for future audio device creations
+	Devices[Index] = nullptr;
+
+	// Add this index to the list of free indices
+	++FreeIndicesSize;
+	FreeIndices.Enqueue(Index);
+
+	return true;
+}
+
+class FAudioDevice* FAudioDeviceManager::GetAudioDevice(uint32 Handle)
+{
+	if (!IsValidAudioDeviceHandle(Handle))
+	{
+		return nullptr;
+	}
+
+	uint32 Index = GetIndex(Handle);
+	check(int32(Index) < Devices.Num());
+	FAudioDevice* AudioDevice = Devices[Index];
+	check(AudioDevice != nullptr);
+	return AudioDevice;
+}
+
+void FAudioDeviceManager::UpdateActiveAudioDevices(bool bGameTicking)
+{
+	for (FAudioDevice* AudioDevice : Devices)
+	{
+		if (AudioDevice)
+		{
+			AudioDevice->Update(bGameTicking);
+		}
+	}
+}
+
+void FAudioDeviceManager::AddReferencedObjects(FReferenceCollector& Collector)
+{
+	for (FAudioDevice* AudioDevice : Devices)
+	{
+		if (AudioDevice)
+		{
+			AudioDevice->AddReferencedObjects(Collector);
+		}
+	}
+}
+
+void FAudioDeviceManager::StopSoundsUsingWave(class USoundWave* InSoundWave)
+{
+	for (FAudioDevice* AudioDevice : Devices)
+	{
+		if (AudioDevice)
+		{
+			TArray<UAudioComponent*> StoppedComponents;
+			AudioDevice->StopSoundsUsingResource(InSoundWave, StoppedComponents);
+		}
+	}
+}
+
+void FAudioDeviceManager::StopSoundsUsingResource(class USoundWave* InSoundWave, TArray<UAudioComponent*>& StoppedComponents)
+{
+	for (FAudioDevice* AudioDevice : Devices)
+	{
+		if (AudioDevice)
+		{
+			TArray<UAudioComponent*> StoppedComponents;
+			AudioDevice->StopSoundsUsingResource(InSoundWave, StoppedComponents);
+		}
+	}
+}
+
+void FAudioDeviceManager::RegisterSoundClass(USoundClass* SoundClass)
+{
+	for (FAudioDevice* AudioDevice : Devices)
+	{
+		if (AudioDevice)
+		{
+			AudioDevice->RegisterSoundClass(SoundClass);
+		}
+	}
+}
+
+void FAudioDeviceManager::UnregisterSoundClass(USoundClass* SoundClass)
+{
+	for (FAudioDevice* AudioDevice : Devices)
+	{
+		if (AudioDevice)
+		{
+			AudioDevice->UnregisterSoundClass(SoundClass);
+		}
+	}
+}
+
+void FAudioDeviceManager::InitSoundClasses()
+{
+	for (FAudioDevice* AudioDevice : Devices)
+	{
+		if (AudioDevice)
+		{
+			AudioDevice->InitSoundClasses();
+		}
+	}
+}
+
+void FAudioDeviceManager::SetActiveDevice(uint32 InAudioDeviceHandle)
+{
+	for (FAudioDevice* AudioDevice : Devices)
+	{
+		if (AudioDevice)
+		{
+			if (AudioDevice->DeviceHandle == InAudioDeviceHandle)
+			{
+				AudioDevice->bIsDeviceMuted = false;
+			}
+			else
+			{
+				AudioDevice->bIsDeviceMuted = true;
+			}
+		}
+	}
+}
+
+uint8 FAudioDeviceManager::GetNumActiveAudioDevices() const
+{
+	return NumActiveAudioDevices;
+}
+
+uint32 FAudioDeviceManager::GetIndex(uint32 Handle) const
+{
+	return Handle & AUDIO_DEVICE_HANDLE_INDEX_MASK;
+}
+
+uint32 FAudioDeviceManager::GetGeneration(uint32 Handle) const
+{
+	return (Handle >> AUDIO_DEVICE_HANDLE_INDEX_BITS) & AUDIO_DEVICE_HANDLE_GENERATION_MASK;
+}
+
+uint32 FAudioDeviceManager::CreateHandle(uint32 DeviceIndex, uint8 Generation)
+{
+	return (DeviceIndex | (Generation << AUDIO_DEVICE_HANDLE_INDEX_BITS));
+}
+
+void FAudioDeviceManager::StopSourcesUsingBuffer(FSoundBuffer* SoundBuffer)
+{
+	for (FAudioDevice* AudioDevice : Devices)
+	{
+		if (AudioDevice)
+		{
+			AudioDevice->StopSourcesUsingBuffer(SoundBuffer);
+		}
+	}
+}
+
+void FAudioDeviceManager::TrackResource(USoundWave* SoundWave, FSoundBuffer* Buffer)
+{
+	// Allocate new resource ID and assign to USoundWave. A value of 0 (default) means not yet registered.
+	int32 ResourceID = NextResourceID++;
+	Buffer->ResourceID = ResourceID;
+	SoundWave->ResourceID = ResourceID;
+
+	Buffers.Add(Buffer);
+	WaveBufferMap.Add(ResourceID, Buffer);
+
+#if !(UE_BUILD_SHIPPING || UE_BUILD_TEST)
+	// Keep track of associated resource name.
+	Buffer->ResourceName = SoundWave->GetPathName();
+#endif
+}
+
+void FAudioDeviceManager::FreeResource(USoundWave* SoundWave)
+{
+	if (SoundWave->ResourceID)
+	{
+		FSoundBuffer* SoundBuffer = WaveBufferMap.FindRef(SoundWave->ResourceID);
+		FreeBufferResource(SoundBuffer);
+
+		SoundWave->ResourceID = 0;
+	}
+}
+
+void FAudioDeviceManager::FreeBufferResource(FSoundBuffer* SoundBuffer)
+{
+	if (SoundBuffer)
+	{
+		Buffers.Remove(SoundBuffer);
+
+		// Stop any sound sources on any audio device currently using this buffer before deleting
+		StopSourcesUsingBuffer(SoundBuffer);
+
+		delete SoundBuffer;
+		SoundBuffer = nullptr;
+	}
+}
+
+FSoundBuffer* FAudioDeviceManager::GetSoundBufferForResourceID(uint32 ResourceID)
+{
+	return WaveBufferMap.FindRef(ResourceID);
+}
+
+void FAudioDeviceManager::RemoveSoundBufferForResourceID(uint32 ResourceID)
+{
+	WaveBufferMap.Remove(ResourceID);
+}
+
+void FAudioDeviceManager::RemoveSoundMix(USoundMix* SoundMix)
+{
+	for (FAudioDevice* AudioDevice : Devices)
+	{
+		if (AudioDevice)
+		{
+			AudioDevice->RemoveSoundMix(SoundMix);
+		}
+	}
+}
diff --git a/Engine/Source/Runtime/Engine/Private/AudioVolume.cpp b/Engine/Source/Runtime/Engine/Private/AudioVolume.cpp
--- a/Engine/Source/Runtime/Engine/Private/AudioVolume.cpp
+++ b/Engine/Source/Runtime/Engine/Private/AudioVolume.cpp
@@ -146,10 +146,10 @@ AAudioVolume::AAudioVolume(const FObjectInitializer& ObjectInitializer)
 	bEnabled = true;
 }
 
-void AAudioVolume::GetLifetimeReplicatedProps( TArray< FLifetimeProperty > & OutLifetimeProps ) const
+void AAudioVolume::GetLifetimeReplicatedProps(TArray< FLifetimeProperty > & OutLifetimeProps) const
 {
-	Super::GetLifetimeReplicatedProps( OutLifetimeProps );
-	DOREPLIFETIME( AAudioVolume, bEnabled );
+	Super::GetLifetimeReplicatedProps(OutLifetimeProps);
+	DOREPLIFETIME(AAudioVolume, bEnabled);
 }
 
 void AAudioVolume::PostUnregisterAllComponents( void )
@@ -158,26 +158,27 @@ void AAudioVolume::PostUnregisterAllComponents( void )
 	Super::PostUnregisterAllComponents();
 
 	// World will be NULL during exit purge.
-	if( GetWorld() )
+	UWorld* World = GetWorld();
+	if (World)
 	{
-		AAudioVolume* CurrentVolume = GetWorld()->HighestPriorityAudioVolume;
+		AAudioVolume* CurrentVolume = World->HighestPriorityAudioVolume;
 		AAudioVolume* PreviousVolume = NULL;
 
 		// Iterate over linked list, removing this volume if found.
-		while( CurrentVolume )
+		while (CurrentVolume)
 		{
 			// Found.
-			if( CurrentVolume == this )
+			if (CurrentVolume == this)
 			{
 				// Remove from linked list.
-				if( PreviousVolume )
+				if (PreviousVolume)
 				{
 					PreviousVolume->NextLowerPriorityVolume = NextLowerPriorityVolume;
 				}
 				else
 				{
 					// Special case removal from first entry.
-					GetWorld()->HighestPriorityAudioVolume = NextLowerPriorityVolume;
+					World->HighestPriorityAudioVolume = NextLowerPriorityVolume;
 				}
 
 				break;
@@ -189,9 +190,9 @@ void AAudioVolume::PostUnregisterAllComponents( void )
 		}
 
 		// Reset next pointer to avoid dangling end bits and also for GC.
-		NextLowerPriorityVolume = NULL;
+		NextLowerPriorityVolume = nullptr;
 
-		if (FAudioDevice* AudioDevice = GEngine->GetAudioDevice())
+		if (FAudioDevice* AudioDevice = World->GetAudioDevice())
 		{
 			AudioDevice->InvalidateCachedInteriorVolumes();
 		}
@@ -203,20 +204,21 @@ void AAudioVolume::PostRegisterAllComponents()
 	// Route update to super first.
 	Super::PostRegisterAllComponents();
 
-	AAudioVolume* CurrentVolume = GetWorld()->HighestPriorityAudioVolume;
-	AAudioVolume* PreviousVolume = NULL;
+	UWorld* World = GetWorld();
+	AAudioVolume* CurrentVolume = World->HighestPriorityAudioVolume;
+	AAudioVolume* PreviousVolume = nullptr;
 
 	// Find where to insert in sorted linked list.
-	if( CurrentVolume )
+	if (CurrentVolume)
 	{
 		// Avoid double insertion!
-		while( CurrentVolume && CurrentVolume != this )
+		while (CurrentVolume && CurrentVolume != this)
 		{
 			// We use > instead of >= to be sure that we are not inserting twice in the case of multiple volumes having
 			// the same priority and the current one already having being inserted after one with the same priority.
-			if( Priority > CurrentVolume->Priority )
+			if (Priority > CurrentVolume->Priority)
 			{
-				if ( PreviousVolume )
+				if (PreviousVolume)
 				{
 					// Insert before current node by fixing up previous to point to current.
 					PreviousVolume->NextLowerPriorityVolume = this;
@@ -224,7 +226,7 @@ void AAudioVolume::PostRegisterAllComponents()
 				else
 				{
 					// Special case for insertion at the beginning.
-					GetWorld()->HighestPriorityAudioVolume = this;
+					World->HighestPriorityAudioVolume = this;
 				}
 
 				// Point to current volume, finalizing insertion.
@@ -238,21 +240,21 @@ void AAudioVolume::PostRegisterAllComponents()
 		}
 
 		// We're the lowest priority volume, insert at the end.
-		if( !CurrentVolume )
+		if (!CurrentVolume)
 		{
-			checkSlow( PreviousVolume );
+			checkSlow(PreviousVolume);
 			PreviousVolume->NextLowerPriorityVolume = this;
-			NextLowerPriorityVolume = NULL;
+			NextLowerPriorityVolume = nullptr;
 		}
 	}
 	else
 	{
 		// First volume in the world info.
-		GetWorld()->HighestPriorityAudioVolume = this;
-		NextLowerPriorityVolume	= NULL;
+		World->HighestPriorityAudioVolume = this;
+		NextLowerPriorityVolume	= nullptr;
 	}
 
-	if (FAudioDevice* AudioDevice = GEngine->GetAudioDevice())
+	if (FAudioDevice* AudioDevice = World->GetAudioDevice())
 	{
 		AudioDevice->InvalidateCachedInteriorVolumes();
 	}
diff --git a/Engine/Source/Runtime/Engine/Private/GameEngine.cpp b/Engine/Source/Runtime/Engine/Private/GameEngine.cpp
--- a/Engine/Source/Runtime/Engine/Private/GameEngine.cpp
+++ b/Engine/Source/Runtime/Engine/Private/GameEngine.cpp
@@ -950,13 +950,11 @@ void UGameEngine::Tick( float DeltaSeconds, bool bIdleMode )
 		}
 
 		// Update Audio. This needs to occur after rendering as the rendering code updates the listener position.
-		if( GetAudioDevice() )
+		if (FAudioDevice* AudioDevice = Context.World()->GetAudioDevice())
 		{
-			GetAudioDevice()->Update( !Context.World()->IsPaused() );
+			AudioDevice->Update(!Context.World()->IsPaused());
 		}
 
-	
-
 		if( GIsClient )
 		{
 			// IStreamingManager is updated outside of a world context. For now, assuming it needs to tick here, before possibly calling PostLoadMap. 
diff --git a/Engine/Source/Runtime/Engine/Private/GameViewportClient.cpp b/Engine/Source/Runtime/Engine/Private/GameViewportClient.cpp
--- a/Engine/Source/Runtime/Engine/Private/GameViewportClient.cpp
+++ b/Engine/Source/Runtime/Engine/Private/GameViewportClient.cpp
@@ -168,6 +168,7 @@ UGameViewportClient::~UGameViewportClient()
 		delete StatUnitData;
 		StatUnitData = NULL;
 	}
+
 }
 
 
@@ -179,6 +180,15 @@ void UGameViewportClient::PostInitProperties()
 
 void UGameViewportClient::BeginDestroy()
 {
+	if (GEngine)
+	{
+		class FAudioDeviceManager* AudioDeviceManager = GEngine->GetAudioDeviceManager();
+		if (AudioDeviceManager)
+		{
+			AudioDeviceManager->ShutdownAudioDevice(AudioDeviceHandle);
+		}
+	}
+
 	RemoveAllViewportWidgets();
 	Super::BeginDestroy();
 }
@@ -231,6 +241,27 @@ void UGameViewportClient::Init(struct FWorldContext& WorldContext, UGameInstance
 	// Create the cursor Widgets
 	UUserInterfaceSettings* UISettings = GetMutableDefault<UUserInterfaceSettings>(UUserInterfaceSettings::StaticClass());
 
+	if (GEngine)
+	{
+		FAudioDeviceManager* AudioDeviceManager = GEngine->GetAudioDeviceManager();
+		if (AudioDeviceManager)
+		{
+			FAudioDevice* AudioDevice = AudioDeviceManager->CreateAudioDevice(AudioDeviceHandle);
+			if (AudioDevice)
+			{
+				// Set the base mix of the new device based on the world settings of the world
+				AudioDevice->SetDefaultBaseSoundMix(World->GetWorldSettings()->DefaultBaseSoundMix);
+
+				// Set this audio device handle on the world context so future world's set onto the world context
+				// will pass the audio device handle to them and audio will play on the correct audio device
+				WorldContext.AudioDeviceHandle = AudioDeviceHandle;
+
+				// Set the world's audio device handle to use so that sounds which play in that world will use the correct audio device
+				World->SetAudioDeviceHandle(AudioDeviceHandle);
+			}
+		}
+	}
+	
 	AddCursor(EMouseCursor::Default, UISettings->DefaultCursor);
 	AddCursor(EMouseCursor::TextEditBeam, UISettings->TextEditBeamCursor);
 	AddCursor(EMouseCursor::Crosshairs, UISettings->CrosshairsCursor);
@@ -815,7 +846,8 @@ void UGameViewportClient::Draw(FViewport* InViewport, FCanvas* SceneCanvas)
 
 	TMap<ULocalPlayer*,FSceneView*> PlayerViewMap;
 
-	FAudioDevice* AudioDevice = GEngine->GetAudioDevice();
+	FAudioDevice* AudioDevice = GetWorld()->GetAudioDevice();
+
 	bool bReverbSettingsFound = false;
 	FReverbSettings ReverbSettings;
 	class AAudioVolume* AudioVolume = nullptr;
@@ -1280,14 +1312,13 @@ void UGameViewportClient::Precache()
 	if(!GIsEditor)
 	{
 		// Precache sounds...
-		FAudioDevice* AudioDevice = GEngine ? GEngine->GetAudioDevice() : NULL;
-		if (AudioDevice != NULL)
+		if (FAudioDevice* AudioDevice = GetWorld()->GetAudioDevice())
 		{
 			UE_LOG(LogPlayerManagement, Log, TEXT("Precaching sounds..."));
 			for(TObjectIterator<USoundWave> It;It;++It)
 			{
 				USoundWave* SoundWave = *It;
-				AudioDevice->Precache( SoundWave );
+				AudioDevice->Precache(SoundWave);
 			}
 			UE_LOG(LogPlayerManagement, Log, TEXT("Precaching sounds completed..."));
 		}
@@ -1337,6 +1368,11 @@ void UGameViewportClient::ReceivedFocus(FViewport* InViewport)
 	{
 		FSlateApplication::Get().SetGameIsFakingTouchEvents(true);
 	}
+
+	if (GEngine && GEngine->GetAudioDeviceManager())
+	{ 
+		GEngine->GetAudioDeviceManager()->SetActiveDevice(AudioDeviceHandle);
+	}
 }
 
 bool UGameViewportClient::IsFocused(FViewport* InViewport)
diff --git a/Engine/Source/Runtime/Engine/Private/GameplayStatics.cpp b/Engine/Source/Runtime/Engine/Private/GameplayStatics.cpp
--- a/Engine/Source/Runtime/Engine/Private/GameplayStatics.cpp
+++ b/Engine/Source/Runtime/Engine/Private/GameplayStatics.cpp
@@ -625,68 +625,77 @@ EPhysicalSurface UGameplayStatics::GetSurfaceType(const struct FHitResult& Hit)
 	return UPhysicalMaterial::DetermineSurfaceType( HitPhysMat );
 }
 
-bool UGameplayStatics::AreAnyListenersWithinRange(FVector Location, float MaximumRange)
+bool UGameplayStatics::AreAnyListenersWithinRange(UObject* WorldContextObject, FVector Location, float MaximumRange)
 {
-	if (GEngine && GEngine->UseSound())
+	if (!GEngine || !GEngine->UseSound())
 	{
-		return GEngine->GetAudioDevice()->LocationIsAudible(Location, MaximumRange);
+		return false;
+	}
+	
+	UWorld* ThisWorld = GEngine->GetWorldFromContextObject(WorldContextObject);
+	if (!ThisWorld)
+	{
+		return false;
 	}
+	
+	// If there is no valid world from the world context object then there certainly are no listeners
+	if (FAudioDevice* AudioDevice = ThisWorld->GetAudioDevice())
+	{
+		return AudioDevice->LocationIsAudible(Location, MaximumRange);
+	}	
 
 	return false;
 }
 
 void UGameplayStatics::PlaySound2D(UObject* WorldContextObject, class USoundBase* Sound, float VolumeMultiplier, float PitchMultiplier, float StartTime)
 {
-	if ( Sound == nullptr )
+	if (!Sound || !GEngine || !GEngine->UseSound())
 	{
 		return;
 	}
 
 	UWorld* ThisWorld = GEngine->GetWorldFromContextObject(WorldContextObject);
-	if ( ThisWorld == nullptr )
+	if (!ThisWorld || !ThisWorld->bAllowAudioPlayback || ThisWorld->GetNetMode() == NM_DedicatedServer)
 	{
 		return;
 	}
 
-	if ( GEngine && GEngine->UseSound() && ThisWorld->bAllowAudioPlayback && ThisWorld->GetNetMode() != NM_DedicatedServer )
+	// TODO - Audio Threading. This call would be a task call to dispatch to the audio thread
+	if (FAudioDevice* AudioDevice = ThisWorld->GetAudioDevice())
 	{
-		// TODO - Audio Threading. This call would be a task call to dispatch to the audio thread
-		if ( FAudioDevice* AudioDevice = GEngine->GetAudioDevice() )
-		{
-			FActiveSound NewActiveSound;
-			NewActiveSound.Sound = Sound;
+		FActiveSound NewActiveSound;
+		NewActiveSound.Sound = Sound;
 
-			NewActiveSound.VolumeMultiplier = VolumeMultiplier;
-			NewActiveSound.PitchMultiplier = PitchMultiplier;
+		NewActiveSound.VolumeMultiplier = VolumeMultiplier;
+		NewActiveSound.PitchMultiplier = PitchMultiplier;
 
-			NewActiveSound.RequestedStartTime = FMath::Max(0.f, StartTime);
+		NewActiveSound.RequestedStartTime = FMath::Max(0.f, StartTime);
 
-			NewActiveSound.bIsUISound = true;
+		NewActiveSound.bIsUISound = true;
 
-			AudioDevice->AddNewActiveSound(NewActiveSound);
-		}
+		AudioDevice->AddNewActiveSound(NewActiveSound);
 	}
 }
 
 void UGameplayStatics::PlaySoundAtLocation(UObject* WorldContextObject, class USoundBase* Sound, FVector Location, float VolumeMultiplier, float PitchMultiplier, float StartTime, class USoundAttenuation* AttenuationSettings)
 {
-	if ( Sound == nullptr )
+	if (!Sound || !GEngine || !GEngine->UseSound())
 	{
 		return;
 	}
 
 	UWorld* ThisWorld = GEngine->GetWorldFromContextObject(WorldContextObject);
-	if(ThisWorld == nullptr)
+	if (!ThisWorld || !ThisWorld->bAllowAudioPlayback || ThisWorld->GetNetMode() == NM_DedicatedServer)
 	{
 		return;
 	}
 
-	const bool bIsInGameWorld = ThisWorld->IsGameWorld();
-
-	if (GEngine->UseSound() && ThisWorld->bAllowAudioPlayback && ThisWorld->GetNetMode() != NM_DedicatedServer)
+	if (FAudioDevice* AudioDevice = ThisWorld->GetAudioDevice())
 	{
-		if( Sound->IsAudibleSimple( Location, AttenuationSettings ) )
+		if (Sound->IsAudibleSimple(AudioDevice, Location, AttenuationSettings))
 		{
+			const bool bIsInGameWorld = ThisWorld->IsGameWorld();
+
 			FActiveSound NewActiveSound;
 			NewActiveSound.World = ThisWorld;
 			NewActiveSound.Sound = Sound;
@@ -712,36 +721,41 @@ void UGameplayStatics::PlaySoundAtLocation(UObject* WorldContextObject, class US
 			}
 
 			// TODO - Audio Threading. This call would be a task call to dispatch to the audio thread
-			GEngine->GetAudioDevice()->AddNewActiveSound(NewActiveSound);
+			AudioDevice->AddNewActiveSound(NewActiveSound);
 		}
 		else
 		{
 			// Don't play a sound for short sounds that start out of range of any listener
-			UE_LOG(LogAudio, Log, TEXT( "Sound not played for out of range Sound %s" ), *Sound->GetName() );
+			UE_LOG(LogAudio, Log, TEXT("Sound not played for out of range Sound %s"), *Sound->GetName());
 		}
 	}
 }
 
 void UGameplayStatics::PlayDialogueAtLocation(UObject* WorldContextObject, class UDialogueWave* Dialogue, const FDialogueContext& Context, FVector Location, float VolumeMultiplier, float PitchMultiplier, float StartTime, class USoundAttenuation* AttenuationSettings)
 {
-	if ( Dialogue == NULL )
+	if (!Dialogue || !GEngine || !GEngine->UseSound())
 	{
 		return;
 	}
 
+	UWorld* ThisWorld = GEngine->GetWorldFromContextObject(WorldContextObject);
+	if (!ThisWorld || !ThisWorld->bAllowAudioPlayback)
+	{
+		return;
+	}
+	
 	USoundBase* Sound = Dialogue->GetWaveFromContext(Context);
-	if ( Sound == NULL )
+	if (!Sound)
 	{
 		return;
 	}
 
-	UWorld* ThisWorld = GEngine->GetWorldFromContextObject(WorldContextObject);
-	if (GEngine && GEngine->UseSound() && (ThisWorld != nullptr) && ThisWorld->bAllowAudioPlayback)
+	if (FAudioDevice* AudioDevice = ThisWorld->GetAudioDevice())
 	{
-		const bool bIsInGameWorld = ThisWorld->IsGameWorld();
-
-		if( Sound->IsAudibleSimple( Location, AttenuationSettings ) )
+		if (Sound->IsAudibleSimple(AudioDevice, Location, AttenuationSettings))
 		{
+			const bool bIsInGameWorld = ThisWorld->IsGameWorld();
+			
 			FActiveSound NewActiveSound;
 			NewActiveSound.World = ThisWorld;
 			NewActiveSound.Sound = Sound;
@@ -766,27 +780,27 @@ void UGameplayStatics::PlayDialogueAtLocation(UObject* WorldContextObject, class
 				NewActiveSound.AttenuationSettings = *AttenuationSettingsToApply;
 			}
 
-			// TODO - Audio Threading. This call would be a task call to dispatch to the audio thread
-			GEngine->GetAudioDevice()->AddNewActiveSound(NewActiveSound);
+			AudioDevice->AddNewActiveSound(NewActiveSound);
 		}
 		else
 		{
 			// Don't play a sound for short sounds that start out of range of any listener
-			UE_LOG(LogAudio, Log, TEXT( "Sound not played for out of range Sound %s" ), *Sound->GetName() );
+			UE_LOG(LogAudio, Log, TEXT("Sound not played for out of range Sound %s"), *Sound->GetName());
 		}
 	}
 }
 
 class UAudioComponent* UGameplayStatics::PlaySoundAttached(class USoundBase* Sound, class USceneComponent* AttachToComponent, FName AttachPointName, FVector Location, EAttachLocation::Type LocationType, bool bStopWhenAttachedToDestroyed, float VolumeMultiplier, float PitchMultiplier, float StartTime, class USoundAttenuation* AttenuationSettings)
 {
-	if ( Sound == NULL )
+	if (!Sound)
 	{
-		return NULL;
+		return nullptr;
 	}
-	if ( AttachToComponent == NULL )
+
+	if (!AttachToComponent)
 	{
 		UE_LOG(LogScript, Warning, TEXT("UGameplayStatics::PlaySoundAttached: NULL AttachComponent specified!"));
-		return NULL;
+		return nullptr;
 	}
 
 	// Location used to check whether to create a component if out of range
@@ -796,8 +810,8 @@ class UAudioComponent* UGameplayStatics::PlaySoundAttached(class USoundBase* Sou
 		TestLocation = AttachToComponent->GetRelativeTransform().TransformPosition(Location);
 	}
 
-	UAudioComponent* AudioComponent = FAudioDevice::CreateComponent( Sound, AttachToComponent->GetWorld(), AttachToComponent->GetOwner(), false, bStopWhenAttachedToDestroyed, &TestLocation, AttenuationSettings );
-	if(AudioComponent)
+	UAudioComponent* AudioComponent = FAudioDevice::CreateComponent(Sound, AttachToComponent->GetWorld(), AttachToComponent->GetOwner(), false, bStopWhenAttachedToDestroyed, &TestLocation, AttenuationSettings);
+	if (AudioComponent)
 	{
 		const bool bIsInGameWorld = AudioComponent->GetWorld()->IsGameWorld();
 
@@ -824,21 +838,21 @@ class UAudioComponent* UGameplayStatics::PlaySoundAttached(class USoundBase* Sou
 
 class UAudioComponent* UGameplayStatics::PlayDialogueAttached(class UDialogueWave* Dialogue, const FDialogueContext& Context, class USceneComponent* AttachToComponent, FName AttachPointName, FVector Location, EAttachLocation::Type LocationType, bool bStopWhenAttachedToDestroyed, float VolumeMultiplier, float PitchMultiplier, float StartTime, class USoundAttenuation* AttenuationSettings)
 {
-	if ( Dialogue == NULL )
+	if (!Dialogue)
 	{
-		return NULL;
+		return nullptr;
 	}
 
 	USoundBase* Sound = Dialogue->GetWaveFromContext(Context);
-	if ( Sound == NULL )
+	if (!Sound)
 	{
-		return NULL;
+		return nullptr;
 	}
 
-	if ( AttachToComponent == NULL )
+	if (!AttachToComponent)
 	{
 		UE_LOG(LogScript, Warning, TEXT("UGameplayStatics::PlaySoundAttached: NULL AttachComponent specified!"));
-		return NULL;
+		return nullptr;
 	}
 
 	// Location used to check whether to create a component if out of range
@@ -848,8 +862,8 @@ class UAudioComponent* UGameplayStatics::PlayDialogueAttached(class UDialogueWav
 		TestLocation = AttachToComponent->GetRelativeTransform().TransformPosition(Location);
 	}
 
-	UAudioComponent* AudioComponent = FAudioDevice::CreateComponent( Sound, AttachToComponent->GetWorld(), AttachToComponent->GetOwner(), false, bStopWhenAttachedToDestroyed, &TestLocation, AttenuationSettings );
-	if(AudioComponent)
+	UAudioComponent* AudioComponent = FAudioDevice::CreateComponent(Sound, AttachToComponent->GetWorld(), AttachToComponent->GetOwner(), false, bStopWhenAttachedToDestroyed, &TestLocation, AttenuationSettings);
+	if (AudioComponent)
 	{
 		const bool bIsGameWorld = AudioComponent->GetWorld()->IsGameWorld();
 
@@ -874,67 +888,115 @@ class UAudioComponent* UGameplayStatics::PlayDialogueAttached(class UDialogueWav
 	return AudioComponent;
 }
 
-void UGameplayStatics::SetBaseSoundMix(USoundMix* InSoundMix)
+void UGameplayStatics::SetBaseSoundMix(UObject* WorldContextObject, USoundMix* InSoundMix)
 {
-	if (InSoundMix)
+	if (!InSoundMix || !GEngine || !GEngine->UseSound())
 	{
-		FAudioDevice* AudioDevice = GEngine->GetAudioDevice();
-		if( AudioDevice != NULL )
-		{
-			AudioDevice->SetBaseSoundMix(InSoundMix);
-		}
+		return;
+	}
+
+	UWorld* ThisWorld = GEngine->GetWorldFromContextObject(WorldContextObject);
+	if (!ThisWorld || !ThisWorld->bAllowAudioPlayback)
+	{
+		return;
+	}
+
+	if (FAudioDevice* AudioDevice = ThisWorld->GetAudioDevice())
+	{
+		AudioDevice->SetBaseSoundMix(InSoundMix);
 	}
 }
 
-void UGameplayStatics::PushSoundMixModifier(USoundMix* InSoundMixModifier)
+void UGameplayStatics::PushSoundMixModifier(UObject* WorldContextObject, USoundMix* InSoundMixModifier)
 {
-	if (InSoundMixModifier)
+	if (!InSoundMixModifier || !GEngine || !GEngine->UseSound())
 	{
-		FAudioDevice* AudioDevice = GEngine->GetAudioDevice();
-		if( AudioDevice != NULL )
-		{
-			AudioDevice->PushSoundMixModifier(InSoundMixModifier);
-		}
+		return;
+	}
+
+	UWorld* ThisWorld = GEngine->GetWorldFromContextObject(WorldContextObject);
+	if (!ThisWorld || !ThisWorld->bAllowAudioPlayback)
+	{
+		return;
+	}
+
+	if (FAudioDevice* AudioDevice = ThisWorld->GetAudioDevice())
+	{
+		AudioDevice->PushSoundMixModifier(InSoundMixModifier);
 	}
 }
 
-void UGameplayStatics::PopSoundMixModifier(USoundMix* InSoundMixModifier)
+void UGameplayStatics::PopSoundMixModifier(UObject* WorldContextObject, USoundMix* InSoundMixModifier)
 {
-	if (InSoundMixModifier)
+	if (InSoundMixModifier == nullptr || GEngine == nullptr || !GEngine->UseSound())
 	{
-		FAudioDevice* AudioDevice = GEngine->GetAudioDevice();
-		if( AudioDevice != NULL )
-		{
-			AudioDevice->PopSoundMixModifier(InSoundMixModifier);
-		}
+		return;
+	}
+
+	UWorld* ThisWorld = GEngine->GetWorldFromContextObject(WorldContextObject);
+	if (ThisWorld == nullptr || !ThisWorld->bAllowAudioPlayback)
+	{
+		return;
+	}
+
+	if (FAudioDevice* AudioDevice = ThisWorld->GetAudioDevice())
+	{
+		AudioDevice->PopSoundMixModifier(InSoundMixModifier);
 	}
 }
 
-void UGameplayStatics::ClearSoundMixModifiers()
+void UGameplayStatics::ClearSoundMixModifiers(UObject* WorldContextObject)
 {
-	FAudioDevice* AudioDevice = GEngine->GetAudioDevice();
-	if( AudioDevice != NULL )
+	if (!GEngine || !GEngine->UseSound())
+	{
+		return;
+	}
+
+	UWorld* ThisWorld = GEngine->GetWorldFromContextObject(WorldContextObject);
+	if (!ThisWorld || !ThisWorld->bAllowAudioPlayback)
+	{
+		return;
+	}
+
+	if (FAudioDevice* AudioDevice = ThisWorld->GetAudioDevice())
 	{
 		AudioDevice->ClearSoundMixModifiers();
 	}
 }
 
-void UGameplayStatics::ActivateReverbEffect(class UReverbEffect* ReverbEffect, FName TagName, float Priority, float Volume, float FadeTime)
+void UGameplayStatics::ActivateReverbEffect(UObject* WorldContextObject, class UReverbEffect* ReverbEffect, FName TagName, float Priority, float Volume, float FadeTime)
 {
-	if (ReverbEffect)
+	if (ReverbEffect || !GEngine || !GEngine->UseSound())
 	{
-		FAudioDevice* AudioDevice = GEngine->GetAudioDevice();
-		if( AudioDevice != NULL )
-		{
-			AudioDevice->ActivateReverbEffect(ReverbEffect, TagName, Priority, Volume, FadeTime);
-		}
+		return;
+	}
+
+	UWorld* ThisWorld = GEngine->GetWorldFromContextObject(WorldContextObject);
+	if (!ThisWorld || !ThisWorld->bAllowAudioPlayback)
+	{
+		return;
+	}
+
+	if (FAudioDevice* AudioDevice = ThisWorld->GetAudioDevice())
+	{
+		AudioDevice->ActivateReverbEffect(ReverbEffect, TagName, Priority, Volume, FadeTime);
 	}
 }
 
-void UGameplayStatics::DeactivateReverbEffect(FName TagName)
+void UGameplayStatics::DeactivateReverbEffect(UObject* WorldContextObject, FName TagName)
 {
-	FAudioDevice* AudioDevice = GEngine->GetAudioDevice();
-	if( AudioDevice != NULL )
+	if (GEngine || !GEngine->UseSound())
+	{
+		return;
+	}
+
+	UWorld* ThisWorld = GEngine->GetWorldFromContextObject(WorldContextObject);
+	if (!ThisWorld || !ThisWorld->bAllowAudioPlayback)
+	{
+		return;
+	}
+
+	if (FAudioDevice* AudioDevice = ThisWorld->GetAudioDevice())
 	{
 		AudioDevice->DeactivateReverbEffect(TagName);
 	}
diff --git a/Engine/Source/Runtime/Engine/Private/Interpolation.cpp b/Engine/Source/Runtime/Engine/Private/Interpolation.cpp
--- a/Engine/Source/Runtime/Engine/Private/Interpolation.cpp
+++ b/Engine/Source/Runtime/Engine/Private/Interpolation.cpp
@@ -1146,7 +1146,7 @@ void AMatineeActor::StepInterp( float DeltaTime, bool bPreview )
 
 void AMatineeActor::DisableRadioFilterIfNeeded()
 {
-	FAudioDevice* AudioDevice = GEngine->GetAudioDevice();
+	FAudioDevice* AudioDevice = GEngine->GetMainAudioDevice();
 	if( AudioDevice )
 	{
 		AudioDevice->EnableRadioEffect( !bDisableRadioFilter );
@@ -1171,7 +1171,7 @@ void AMatineeActor::EnableCinematicMode(bool bEnable)
 
 void AMatineeActor::EnableRadioFilter()
 {
-	FAudioDevice* AudioDevice = GEngine->GetAudioDevice();
+	FAudioDevice* AudioDevice = GEngine->GetMainAudioDevice();
 	if( AudioDevice )
 	{
 		AudioDevice->EnableRadioEffect( true );
diff --git a/Engine/Source/Runtime/Engine/Private/LevelActor.cpp b/Engine/Source/Runtime/Engine/Private/LevelActor.cpp
--- a/Engine/Source/Runtime/Engine/Private/LevelActor.cpp
+++ b/Engine/Source/Runtime/Engine/Private/LevelActor.cpp
@@ -953,6 +953,29 @@ AAudioVolume* UWorld::GetAudioSettings( const FVector& ViewLocation, FReverbSett
 	return Volume;
 }
 
+void UWorld::SetAudioDeviceHandle(const uint32 InAudioDeviceHandle)
+{
+	AudioDeviceHandle = InAudioDeviceHandle;
+}
+
+FAudioDevice* UWorld::GetAudioDevice()
+{
+	FAudioDevice* AudioDevice = nullptr;
+	if (GEngine)
+	{
+		class FAudioDeviceManager* AudioDeviceManager = GEngine->GetAudioDeviceManager();
+		if (AudioDeviceManager != nullptr)
+		{
+			AudioDevice = AudioDeviceManager->GetAudioDevice(AudioDeviceHandle);
+			if (AudioDevice == nullptr)
+			{
+				AudioDevice = GEngine->GetMainAudioDevice();
+			}
+		}
+	}
+	return AudioDevice;
+}
+
 /**
  * Sets bMapNeedsLightingFullyRebuild to the specified value.  Marks the worldsettings package dirty if the value changed.
  *
diff --git a/Engine/Source/Runtime/Engine/Private/PlayerCameraManager.cpp b/Engine/Source/Runtime/Engine/Private/PlayerCameraManager.cpp
--- a/Engine/Source/Runtime/Engine/Private/PlayerCameraManager.cpp
+++ b/Engine/Source/Runtime/Engine/Private/PlayerCameraManager.cpp
@@ -597,17 +597,31 @@ void APlayerCameraManager::UpdateCameraLensEffects(const FTViewTarget& OutVT)
 
 void APlayerCameraManager::ApplyAudioFade()
 {
-	if (GEngine && GEngine->GetAudioDevice() )
+	if (GEngine)
 	{
-		GEngine->GetAudioDevice()->TransientMasterVolume = 1.f - FadeAmount;
+		UWorld* World = GetWorld();
+		if (World)
+		{
+			if (FAudioDevice* AudioDevice = World->GetAudioDevice())
+			{
+				AudioDevice->TransientMasterVolume = 1.0f - FadeAmount;
+			}
+		}
 	}
 }
 
 void APlayerCameraManager::StopAudioFade()
 {
-	if (GEngine && GEngine->GetAudioDevice())
+	if (GEngine)
 	{
-		GEngine->GetAudioDevice()->TransientMasterVolume = 1.f;
+		UWorld* World = GetWorld();
+		if (World)
+		{
+			if (FAudioDevice* AudioDevice = World->GetAudioDevice())
+			{
+				AudioDevice->TransientMasterVolume = 1.0f;
+			}
+		}
 	}
 }
 
diff --git a/Engine/Source/Runtime/Engine/Private/PreviewScene.cpp b/Engine/Source/Runtime/Engine/Private/PreviewScene.cpp
--- a/Engine/Source/Runtime/Engine/Private/PreviewScene.cpp
+++ b/Engine/Source/Runtime/Engine/Private/PreviewScene.cpp
@@ -48,9 +48,16 @@ FPreviewScene::FPreviewScene(FPreviewScene::ConstructionValues CVS)
 FPreviewScene::~FPreviewScene()
 {
 	// Stop any audio components playing in this scene
-	if( GEngine && GEngine->GetAudioDevice() )
+	if (GEngine)
 	{
-		GEngine->GetAudioDevice()->Flush( GetWorld(), false );
+		UWorld* World = GetWorld();
+		if (World)
+		{
+			if (FAudioDevice* AudioDevice = World->GetAudioDevice())
+			{
+				AudioDevice->Flush(GetWorld(), false);
+			}
+		}
 	}
 
 	// Remove all the attached components
diff --git a/Engine/Source/Runtime/Engine/Private/Slate/SlateSoundDevice.cpp b/Engine/Source/Runtime/Engine/Private/Slate/SlateSoundDevice.cpp
--- a/Engine/Source/Runtime/Engine/Private/Slate/SlateSoundDevice.cpp
+++ b/Engine/Source/Runtime/Engine/Private/Slate/SlateSoundDevice.cpp
@@ -11,7 +11,7 @@ DEFINE_LOG_CATEGORY_STATIC(LogSlateSoundDevice, Log, All);
 
 void FSlateSoundDevice::PlaySound(const FSlateSound& Sound, int32 UserIndex) const
 {
-	FAudioDevice* const AudioDevice = GEngine->GetAudioDevice();
+	FAudioDevice* const AudioDevice = GEngine->GetMainAudioDevice();
 	if(AudioDevice)
 	{
 		UObject* const Object = Sound.GetResourceObject();
diff --git a/Engine/Source/Runtime/Engine/Private/SoundBase.cpp b/Engine/Source/Runtime/Engine/Private/SoundBase.cpp
--- a/Engine/Source/Runtime/Engine/Private/SoundBase.cpp
+++ b/Engine/Source/Runtime/Engine/Private/SoundBase.cpp
@@ -41,23 +41,23 @@ float USoundBase::GetMaxAudibleDistance()
 	return 0.f;
 }
 
-bool USoundBase::IsAudibleSimple( const FVector Location, USoundAttenuation* InAttenuationSettings )
+bool USoundBase::IsAudibleSimple(class FAudioDevice* AudioDevice, const FVector Location, USoundAttenuation* InAttenuationSettings)
 {
 	// No audio device means no listeners to check against
-	if( !GEngine || !GEngine->GetAudioDevice() )
+	if (!AudioDevice)
 	{
-		return true ;
+		return false;
 	}
 
 	// Listener position could change before long sounds finish
-	if( GetDuration() > 1.0f )
+	if (GetDuration() > 1.0f)
 	{
-		return true ;
+		return true;
 	}
 
 	// Is this SourceActor within the MaxAudibleDistance of any of the listeners?
-	float MaxAudibleDistance = InAttenuationSettings  != nullptr ? InAttenuationSettings ->Attenuation.GetMaxDimension() : GetMaxAudibleDistance();
-	return GEngine->GetAudioDevice()->LocationIsAudible( Location, MaxAudibleDistance );
+	float MaxAudibleDistance = InAttenuationSettings != nullptr ? InAttenuationSettings->Attenuation.GetMaxDimension() : GetMaxAudibleDistance();
+	return AudioDevice->LocationIsAudible(Location, MaxAudibleDistance);
 }
 
 bool USoundBase::IsAudible( const FVector &SourceLocation, const FVector &ListenerLocation, AActor* SourceActor, bool& bIsOccluded, bool bCheckOcclusion )
diff --git a/Engine/Source/Runtime/Engine/Private/SoundClass.cpp b/Engine/Source/Runtime/Engine/Private/SoundClass.cpp
--- a/Engine/Source/Runtime/Engine/Private/SoundClass.cpp
+++ b/Engine/Source/Runtime/Engine/Private/SoundClass.cpp
@@ -33,7 +33,7 @@ void USoundClass::PostLoad()
 		if (ChildClasses[ChildIndex] != NULL && ChildClasses[ChildIndex]->GetLinkerUE4Version() < VER_UE4_SOUND_CLASS_GRAPH_EDITOR)
 		{
 			// first come, first served
-			if (ChildClasses[ChildIndex]->ParentClass == NULL)
+			if (ChildClasses[ChildIndex]->ParentClass == nullptr)
 			{
 				ChildClasses[ChildIndex]->ParentClass = this;
 			}
@@ -45,12 +45,13 @@ void USoundClass::PostLoad()
 			}
 		}
 	}
+	// Use the main/default audio device for storing and retrieving sound class properties
+	FAudioDeviceManager* AudioDeviceManager = (GEngine ? GEngine->GetAudioDeviceManager() : nullptr);
 
-	FAudioDevice* AudioDevice = (GEngine ? GEngine->GetAudioDevice() : NULL);
-	if (AudioDevice)
+	// Force the properties to be initialized for this SoundClass on all active audio devices
+	if (AudioDeviceManager)
 	{
-		// Force the properties to be initialized for this SoundClass
-		AudioDevice->GetSoundClassCurrentProperties(this);
+		AudioDeviceManager->RegisterSoundClass(this);
 	}
 }
 
@@ -182,13 +183,9 @@ void USoundClass::BeginDestroy()
 {
 	Super::BeginDestroy();
 
-	if (!GExitPurge && GEngine)
+	if (!GExitPurge && GEngine && GEngine->GetAudioDeviceManager())
 	{
-		FAudioDevice* AudioDevice = GEngine->GetAudioDevice();
-		if( AudioDevice )
-		{
-			AudioDevice->RemoveClass(this);
-		}
+		GEngine->GetAudioDeviceManager()->UnregisterSoundClass(this);
 	}
 }
 
diff --git a/Engine/Source/Runtime/Engine/Private/SoundMix.cpp b/Engine/Source/Runtime/Engine/Private/SoundMix.cpp
--- a/Engine/Source/Runtime/Engine/Private/SoundMix.cpp
+++ b/Engine/Source/Runtime/Engine/Private/SoundMix.cpp
@@ -28,10 +28,10 @@ void USoundMix::BeginDestroy()
 {
 	if (!GExitPurge && GEngine)
 	{
-		FAudioDevice* AudioDevice = GEngine->GetAudioDevice();
-		if( AudioDevice )
+		FAudioDeviceManager* AudioDeviceManager = GEngine->GetAudioDeviceManager();
+		if (AudioDeviceManager)
 		{
-			AudioDevice->RemoveSoundMix(this);
+			AudioDeviceManager->RemoveSoundMix(this);
 		}
 	}
 
diff --git a/Engine/Source/Runtime/Engine/Private/SoundWave.cpp b/Engine/Source/Runtime/Engine/Private/SoundWave.cpp
--- a/Engine/Source/Runtime/Engine/Private/SoundWave.cpp
+++ b/Engine/Source/Runtime/Engine/Private/SoundWave.cpp
@@ -71,12 +71,12 @@ SIZE_T USoundWave::GetResourceSize(EResourceSizeMode::Type Mode)
 		CalculatedResourceSize += RawPCMDataSize;
 	}
 
-	if (GEngine && GEngine->GetAudioDevice())
+	if (GEngine && GEngine->GetMainAudioDevice())
 	{
 		// Don't add compressed data to size of streaming sounds
 		if (!FPlatformProperties::SupportsAudioStreaming() || !IsStreaming())
 		{
-			CalculatedResourceSize += GetCompressedDataSize(GEngine->GetAudioDevice()->GetRuntimeFormat(this));
+			CalculatedResourceSize += GetCompressedDataSize(GEngine->GetMainAudioDevice()->GetRuntimeFormat(this));
 		}
 	}
 
@@ -320,7 +320,7 @@ void USoundWave::PostLoad()
 	// most likely cause us to run out of memory.
 	if( !GIsEditor && !IsTemplate( RF_ClassDefaultObject ) && GEngine )
 	{
-		FAudioDevice* AudioDevice = GEngine->GetAudioDevice();
+		FAudioDevice* AudioDevice = GEngine->GetMainAudioDevice();
 		if( AudioDevice && AudioDevice->bStartupSoundsPreCached)
 		{
 			// Upload the data to the hardware, but only if we've precached startup sounds already
@@ -440,12 +440,11 @@ void USoundWave::FreeResources()
 	if( GEngine && !GExitPurge )
 	{
 		// Notify the audio device to free the bulk data associated with this wave.
-		FAudioDevice* AudioDevice = GEngine->GetAudioDevice();
-		if (AudioDevice != NULL)
+		FAudioDeviceManager* AudioDeviceManager = GEngine->GetAudioDeviceManager();
+		if (AudioDeviceManager != NULL)
 		{
-			TArray<UAudioComponent*> StoppedComponents;
-			AudioDevice->StopSoundsUsingResource(this, StoppedComponents);
-			AudioDevice->FreeResource( this );
+			AudioDeviceManager->StopSoundsUsingWave(this);
+			AudioDeviceManager->FreeResource(this);
 		}
 	}
 
diff --git a/Engine/Source/Runtime/Engine/Private/UnrealEngine.cpp b/Engine/Source/Runtime/Engine/Private/UnrealEngine.cpp
--- a/Engine/Source/Runtime/Engine/Private/UnrealEngine.cpp
+++ b/Engine/Source/Runtime/Engine/Private/UnrealEngine.cpp
@@ -42,6 +42,7 @@
 #include "StatsData.h"
 #include "ScreenRendering.h"
 #include "RHIStaticStates.h"
+#include "AudioDeviceManager.h"
 #include "AudioDevice.h"
 #include "ActiveSound.h"
 #include "DeviceProfiles/DeviceProfileManager.h"
@@ -567,6 +568,26 @@ class FScreenSaverInhibitor : public FRunnable
 	FWorldContext
 -----------------------------------------------------------------------------*/
 
+void FWorldContext::SetCurrentWorld(UWorld *World)
+{
+	if (World != nullptr)
+	{
+		// Set the world's audio device handle so that audio components playing in the 
+		// world will use the correct audio device instance.
+		World->SetAudioDeviceHandle(AudioDeviceHandle);
+	}
+
+	for (int32 idx = 0; idx < ExternalReferences.Num(); ++idx)
+	{
+		if (ExternalReferences[idx] && *ExternalReferences[idx] == ThisCurrentWorld)
+		{
+			*ExternalReferences[idx] = World;
+		}
+	}
+
+	ThisCurrentWorld = World;
+}
+
 void FWorldContext::AddReferencedObjects(FReferenceCollector& Collector, const UObject* ReferencingObject)
 {
 	// TODO: This is awfully unsafe as anything in a WorldContext that changes may not be referenced
@@ -744,9 +765,6 @@ void UEngine::Init(IEngineLoop* InEngineLoop)
 
 	GNearClippingPlane = NearClipPlane;
 
-	// Initialize the audio device
-	InitializeAudioDevice();
-
 	if (GIsEditor)
 	{
 		// Create a WorldContext for the editor to use and create an initially empty world.
@@ -755,6 +773,9 @@ void UEngine::Init(IEngineLoop* InEngineLoop)
 		GWorld = InitialWorldContext.World();
 	}
 
+	// Initialize the audio device after a world context is setup
+	InitializeAudioDeviceManager();
+
 	if ( IsConsoleBuild() )
 	{
 		bUseConsoleInput = true;
@@ -921,12 +942,15 @@ void UEngine::OnExternalUIChange(bool bInIsOpening)
 	FSlateApplication::Get().ExternalUIChange(bInIsOpening);
 }
 
-void UEngine::ShutdownAudioDevice()
+void UEngine::ShutdownAudioDeviceManager()
 {
-	if (AudioDevice)
+	// Shutdown the main audio device in the UEEngine
+	if (AudioDeviceManager)
 	{
-		AudioDevice->Teardown();
-		AudioDevice = NULL;
+		AudioDeviceManager->ShutdownAudioDevice(MainAudioDeviceHandle);
+		check(AudioDeviceManager->GetNumActiveAudioDevices() == 0);
+		delete AudioDeviceManager;
+		AudioDeviceManager = NULL;
 	}
 }
 
@@ -1481,10 +1505,7 @@ void UEngine::FinishDestroy()
 	{
 		// shut down all subsystems.
 		GEngine = NULL;
-		if (AudioDevice)
-		{
-			AudioDevice->Teardown();
-		}
+		ShutdownAudioDeviceManager();
 
 		FURL::StaticExit();
 	}
@@ -1499,10 +1520,9 @@ void UEngine::Serialize(FArchive& Ar)
 	// count memory
 	if (Ar.IsCountingMemory())
 	{
-		if (AudioDevice)
-		{
-			AudioDevice->CountBytes(Ar);
-		}
+		// Only use the main audio device when counting memory
+		FAudioDevice* AudioDevice = GetMainAudioDevice();
+		AudioDevice->CountBytes(Ar);
 	}
 }
 
@@ -1510,10 +1530,10 @@ void UEngine::AddReferencedObjects(UObject* InThis, FReferenceCollector& Collect
 {
 	UEngine* This = CastChecked<UEngine>(InThis);
 
-	// track objects in the audio device
-	if (This->AudioDevice)
+	// track objects in all the active audio devices
+	if (This->AudioDeviceManager)
 	{
-		This->AudioDevice->AddReferencedObjects(Collector);
+		This->AudioDeviceManager->AddReferencedObjects(Collector);
 	}
 
 	// TODO: This is quite dangerous as FWorldContext::AddReferencedObjects could fail to be updated when something it
@@ -1600,14 +1620,38 @@ UFont* UEngine::GetAdditionalFont(int32 AdditionalFontIndex)
 	return GEngine->AdditionalFonts.IsValidIndex(AdditionalFontIndex) ? GEngine->AdditionalFonts[AdditionalFontIndex] : NULL;
 }
 
+class FAudioDeviceManager* UEngine::GetAudioDeviceManager()
+{
+	return AudioDeviceManager;
+}
+
+uint32 UEngine::GetAudioDeviceHandle() const
+{
+	return MainAudioDeviceHandle;
+}
+
+FAudioDevice* UEngine::GetMainAudioDevice()
+{
+	if (AudioDeviceManager != nullptr)
+	{
+		return AudioDeviceManager->GetAudioDevice(MainAudioDeviceHandle);
+	}
+	return nullptr;
+}
+
+FAudioDevice* UEngine::GetAudioDevice()
+{
+	return GetMainAudioDevice();
+}
+
 /**
  *	Initialize the audio device
  *
  *	@return	bool		true if successful, false if not
  */
-bool UEngine::InitializeAudioDevice()
+bool UEngine::InitializeAudioDeviceManager()
 {
-	if (AudioDevice == NULL)
+	if (AudioDeviceManager == nullptr)
 	{
 		// Initialize the audio device.
 		if (bUseSound == true)
@@ -1624,29 +1668,25 @@ bool UEngine::InitializeAudioDevice()
 				// did the module exist?
 				if (AudioDeviceModule)
 				{
-					// use the module object to create the audio device
-					AudioDevice = AudioDeviceModule->CreateAudioDevice();
-					if (AudioDevice)
-					{
-						// Attempt to initialize the device
-						if ( !AudioDevice->Init() )
-						{
-							// Failed to initialize the device. Delete it.
-							delete AudioDevice;
-							AudioDevice = NULL;
-						}
-					}
+					// Create the audio device manager and register the platform module to the device manager
+					AudioDeviceManager = new FAudioDeviceManager();
+					AudioDeviceManager->RegisterAudioDeviceModule(AudioDeviceModule);
+					FAudioDevice* NewAudioDevice = AudioDeviceManager->CreateAudioDevice(MainAudioDeviceHandle);
+
+					// Set the new default device as the active device
+					AudioDeviceManager->SetActiveDevice(MainAudioDeviceHandle);
+					return NewAudioDevice != nullptr;
 				}
 			}
 		}
 	}
 
-	return (AudioDevice != NULL);
+	return AudioDeviceManager != nullptr;
 }
 
 bool UEngine::UseSound() const
 {
-	return (bUseSound && AudioDevice);
+	return (bUseSound && AudioDeviceManager != nullptr);
 }
 /**
  * A fake stereo rendering device used to test stereo rendering without an attached device.
@@ -2144,7 +2184,17 @@ bool UEngine::Exec( UWorld* InWorld, const TCHAR* Cmd, FOutputDevice& Ar )
 		return true;
 	}
 
-	if (GetAudioDevice() && (GetAudioDevice()->Exec( InWorld, Cmd,Ar) == true))
+	FAudioDevice* AudioDevice = nullptr;
+	if (InWorld)
+	{
+		AudioDevice = InWorld->GetAudioDevice();
+	}
+	else
+	{
+		AudioDevice = GetMainAudioDevice();
+	}
+
+	if (AudioDevice && AudioDevice->Exec(InWorld, Cmd, Ar) == true)
 	{
 		return true;
 	}
@@ -8825,15 +8875,14 @@ bool UEngine::LoadMap( FWorldContext& WorldContext, FURL URL, class UPendingNetG
 			}
 		}
 
-		WorldContext.SetCurrentWorld(NULL);
-	}
+		// Stop all audio to remove references to current level.
+		if (FAudioDevice* AudioDevice = WorldContext.World()->GetAudioDevice())
+		{
+			AudioDevice->Flush(WorldContext.World());
+			AudioDevice->TransientMasterVolume = 1.0f;
+		}
 
-	// Stop all audio to remove references to current level.
-	if( GEngine && GEngine->GetAudioDevice() )
-	{
-		GEngine->GetAudioDevice()->Flush( WorldContext.World() );
-		// reset transient volume
-		GEngine->GetAudioDevice()->TransientMasterVolume = 1.0;
+		WorldContext.SetCurrentWorld(nullptr);
 	}
 
 	if (bCookSeparateSharedMPGameContent)
@@ -9099,9 +9148,9 @@ bool UEngine::LoadMap( FWorldContext& WorldContext, FURL URL, class UPendingNetG
 
 	WorldContext.World()->SetGameMode(URL);
 
-	if( GetAudioDevice() )
+	if (FAudioDevice* AudioDevice = WorldContext.World()->GetAudioDevice())
 	{
-		GetAudioDevice()->SetDefaultBaseSoundMix( WorldContext.World()->GetWorldSettings()->DefaultBaseSoundMix );
+		AudioDevice->SetDefaultBaseSoundMix(WorldContext.World()->GetWorldSettings()->DefaultBaseSoundMix);
 	}
 
 	// Listen for clients.
@@ -10124,9 +10173,10 @@ bool UEngine::CommitMapChange( FWorldContext &Context )
 		Context.World()->FlushLevelStreaming(EFlushLevelStreamingType::Visibility);
 		
 		// make sure any looping sounds, etc are stopped
-		if (GetAudioDevice() != NULL)
+		
+		if (FAudioDevice* AudioDevice = Context.World()->GetAudioDevice())
 		{
-			GetAudioDevice()->StopAllSounds();
+			AudioDevice->StopAllSounds();
 		}
 
 		// Remove all unloaded levels from memory and perform full purge.
@@ -11396,7 +11446,7 @@ bool UEngine::ToggleStatRaw(UWorld* World, FCommonViewportClient* ViewportClient
 #if !(UE_BUILD_SHIPPING || UE_BUILD_TEST)
 int32 UEngine::RenderStatReverb(UWorld* World, FViewport* Viewport, FCanvas* Canvas, int32 X, int32 Y, const FVector* ViewLocation, const FRotator* ViewRotation)
 {
-	if (AudioDevice)
+	if (FAudioDevice* AudioDevice = World->GetAudioDevice())
 	{
 		UReverbEffect* ReverbEffect = (AudioDevice->Effects ? AudioDevice->Effects->GetCurrentReverbEffect() : NULL);
 		FString TheString;
@@ -11464,7 +11514,7 @@ int32 UEngine::RenderStatReverb(UWorld* World, FViewport* Viewport, FCanvas* Can
 // SOUNDMIXES
 int32 UEngine::RenderStatSoundMixes(UWorld* World, FViewport* Viewport, FCanvas* Canvas, int32 X, int32 Y, const FVector* ViewLocation, const FRotator* ViewRotation)
 {
-	if (AudioDevice)
+	if (FAudioDevice* AudioDevice = World->GetAudioDevice())
 	{
 		Canvas->DrawShadowedString(X, Y, TEXT("Active Sound Mixes:"), GetSmallFont(), FColor::Green);
 		Y += 12;
@@ -11506,7 +11556,7 @@ int32 UEngine::RenderStatSoundWaves(UWorld* World, FViewport* Viewport, FCanvas*
 
 	TSet<FActiveSound*> ActiveSounds;
 
-	if (AudioDevice)
+	if (FAudioDevice* AudioDevice = World->GetAudioDevice())
 	{
 		TArray<FWaveInstance*> WaveInstances;
 		int32 FirstActiveIndex = AudioDevice->GetSortedActiveWaveInstances(WaveInstances, ESortedActiveWaveGetType::QueryOnly);
@@ -11558,7 +11608,7 @@ int32 UEngine::RenderStatSoundCues(UWorld* World, FViewport* Viewport, FCanvas*
 {
 	TSet<FActiveSound*> ActiveSounds;
 
-	if (AudioDevice)
+	if (FAudioDevice* AudioDevice = World->GetAudioDevice())
 	{
 		TArray<FWaveInstance*> WaveInstances;
 		int32 FirstActiveIndex = AudioDevice->GetSortedActiveWaveInstances(WaveInstances, ESortedActiveWaveGetType::QueryOnly);
@@ -11679,7 +11729,7 @@ int32 UEngine::RenderStatSounds(UWorld* World, FViewport* Viewport, FCanvas* Can
 	const FViewportClient::ESoundShowFlags::Type ShowSounds = Viewport->GetClient() ? Viewport->GetClient()->GetSoundShowFlags() : FViewportClient::ESoundShowFlags::Disabled;
 	const bool bDebug = ShowSounds & FViewportClient::ESoundShowFlags::Debug;
 
-	if (AudioDevice)
+	if (FAudioDevice* AudioDevice = World->GetAudioDevice())
 	{
 		// Refresh the wave instances inside audio components.
 		static TArray<FWaveInstance*> WaveInstances;
diff --git a/Engine/Source/Runtime/Engine/Private/World.cpp b/Engine/Source/Runtime/Engine/Private/World.cpp
--- a/Engine/Source/Runtime/Engine/Private/World.cpp
+++ b/Engine/Source/Runtime/Engine/Private/World.cpp
@@ -112,6 +112,7 @@ UWorld::UWorld( const FObjectInitializer& ObjectInitializer )
 ,	FXSystem(NULL)
 ,	TickTaskLevel(FTickTaskManagerInterface::Get().AllocateTickTaskLevel())
 ,   bIsBuilt(false)
+,   AudioDeviceHandle(INDEX_NONE)
 ,	FlushLevelStreamingType(EFlushLevelStreamingType::None)
 ,	NextTravelType(TRAVEL_Relative)
 {
@@ -4712,10 +4713,11 @@ UWorld* FSeamlessTravelHandler::Tick()
 			CurrentWorld->CleanupWorld(bSwitchedToDefaultMap);
 			CurrentWorld->RemoveFromRoot();
 			CurrentWorld->ClearFlags(RF_Standalone);
+
 			// Stop all audio to remove references to old world
-			if (GEngine != NULL && GEngine->GetAudioDevice() != NULL)
+			if (FAudioDevice* AudioDevice = CurrentWorld->GetAudioDevice())
 			{
-				GEngine->GetAudioDevice()->Flush(CurrentWorld);
+				AudioDevice->Flush(CurrentWorld);
 			}
 
 			// Copy the standby cheat status
@@ -4806,9 +4808,9 @@ UWorld* FSeamlessTravelHandler::Tick()
 			{
 				LoadedWorld->SetGameMode(PendingTravelURL);
 
-				if( GEngine && GEngine->GetAudioDevice() )
+				if (FAudioDevice* AudioDevice = LoadedWorld->GetAudioDevice())
 				{
-					GEngine->GetAudioDevice()->SetDefaultBaseSoundMix( LoadedWorld->GetWorldSettings()->DefaultBaseSoundMix );
+					AudioDevice->SetDefaultBaseSoundMix(LoadedWorld->GetWorldSettings()->DefaultBaseSoundMix);
 				}
 
 				// Copy cheat flags if the game info is present
diff --git a/Engine/Source/Runtime/Engine/Public/Audio.h b/Engine/Source/Runtime/Engine/Public/Audio.h
--- a/Engine/Source/Runtime/Engine/Public/Audio.h
+++ b/Engine/Source/Runtime/Engine/Public/Audio.h
@@ -278,10 +278,11 @@ inline uint32 GetTypeHash( FWaveInstance* A ) { return A->TypeHash; }
 class FSoundBuffer
 {
 public:
-	FSoundBuffer()
+	FSoundBuffer(class FAudioDevice * InAudioDevice)
 		: ResourceID(0)
 		, NumChannels(0)
 		, bAllocationInPermanentPool(false)
+		, AudioDevice(InAudioDevice)
 	{
 
 	}
@@ -325,6 +326,8 @@ public:
 	FString	ResourceName;
 	/** Whether memory for this buffer has been allocated from permanent pool. */
 	bool	bAllocationInPermanentPool;
+	/** Parent Audio Device used when creating the sound buffer. Used to remove tracking info on this sound buffer when its done. */
+	class FAudioDevice * AudioDevice;
 };
 
 /*-----------------------------------------------------------------------------
diff --git a/Engine/Source/Runtime/Engine/Public/AudioDevice.h b/Engine/Source/Runtime/Engine/Public/AudioDevice.h
--- a/Engine/Source/Runtime/Engine/Public/AudioDevice.h
+++ b/Engine/Source/Runtime/Engine/Public/AudioDevice.h
@@ -2,6 +2,8 @@
 
 #pragma once 
 
+#include "AudioDeviceManager.h"
+
 /** 
  * Debug state of the audio system
  */
@@ -295,24 +297,12 @@ public:
 	 */
 	virtual void SetMaxChannels(int32 InMaxChannels);
 
-	/** 
-	 * Links up the resource data indices for looking up and cleaning up
-	 */
-	void TrackResource(USoundWave* Wave, FSoundBuffer* Buffer);
-
 	/**
-	 * Frees the bulk resource data assocated with this SoundWave.
-	 *
-	 * @param	SoundWave	wave object to free associated bulk data
-	 */
-	void FreeResource(USoundWave* SoundWave);
-
-	/**
-	 * Frees the resources associated with this buffer
-	 *
-	 * @param	FSoundBuffer	Buffer to clean up
-	 */
-	void FreeBufferResource(FSoundBuffer* Buffer);
+	* Stops any sound sources which are using the given buffer.
+	*
+	* @param	FSoundBuffer	Buffer to check against
+	*/
+	void StopSourcesUsingBuffer(FSoundBuffer * SoundBuffer);
 
 	/**
 	 * Stops all game sounds (and possibly UI) sounds
@@ -385,22 +375,30 @@ public:
 	void GetSoundClassInfo( TMap<FName, FAudioClassInfo>& AudioClassInfos );
 
 	/**
-	 * Returns the properties of the requested sound class modified to reflect the current state of the mix system 
+	 * Registers a sound class with the audio device
 	 *
 	 * @param	SoundClassName	name of sound class to retrieve
 	 * @return	sound class properties if it exists
 	 */
-	FSoundClassProperties* GetSoundClassCurrentProperties( class USoundClass* InSoundClass );
+	void RegisterSoundClass( class USoundClass* InSoundClass );
 
 	/**
-	 * Checks to see if a coordinate is within a distance of any listener
-	 */
-	bool LocationIsAudible( FVector Location, float MaxDistance );
+	* Unregisters a sound class
+	*/
+	void UnregisterSoundClass(class USoundClass* SoundClass);
+
+	/**
+	* Gets the current properties of a sound class, if the sound class hasn't been registered, then it returns nullptr
+	*
+	* @param	SoundClassName	name of sound class to retrieve
+	* @return	sound class properties if it exists
+	*/
+	FSoundClassProperties* GetSoundClassCurrentProperties(class USoundClass* InSoundClass);
 
 	/**
-	 * Removes a sound class
+	 * Checks to see if a coordinate is within a distance of any listener
 	 */
-	void RemoveClass( class USoundClass* SoundClass );
+	bool LocationIsAudible( FVector Location, float MaxDistance );
 
 	/**
 	 * Sets the Sound Mix that should be active by default
@@ -748,11 +746,6 @@ public:
 	TArray<class FSoundSource*>				FreeSources;
 	TMap<struct FWaveInstance*, class FSoundSource*>	WaveInstanceSourceMap;
 
-	/** Array of all created buffers associated with this audio device */
-	TArray<class FSoundBuffer*>				Buffers;
-	/** Look up associating a USoundWave's resource ID with low level sound buffers	*/
-	TMap<int32, class FSoundBuffer*>			WaveBufferMap;
-
 	/** Current properties of all sound classes */
 	TMap<class USoundClass*, FSoundClassProperties>	SoundClasses;
 
@@ -777,6 +770,12 @@ public:
 	/** The activated reverb that currently has the highest priority */
 	const FActivatedReverb*								HighestPriorityReverb;
 
+	/** The handle for this audio device used in the audio device manager. */
+	uint32 DeviceHandle;
+
+	/** Whether the audio device is active (current audio device in-focus in PIE) */
+	uint32 bIsDeviceMuted:1;
+
 private:
 
 	TArray<struct FActiveSound*> ActiveSounds;
@@ -798,3 +797,5 @@ public:
 	/** Creates a new instance of the audio device implemented by the module. */
 	virtual class FAudioDevice* CreateAudioDevice() = 0;
 };
+
+
diff --git a/Engine/Source/Runtime/Engine/Public/AudioDeviceManager.h b/Engine/Source/Runtime/Engine/Public/AudioDeviceManager.h
new file mode 100644
--- /dev/null
+++ b/Engine/Source/Runtime/Engine/Public/AudioDeviceManager.h
@@ -0,0 +1,149 @@
+// Copyright 1998-2015 Epic Games, Inc. All Rights Reserved.
+
+#pragma once 
+
+/**
+* Class for managing multiple audio devices.
+*/
+class ENGINE_API FAudioDeviceManager
+{
+public:
+
+	/**
+	* Constructor
+	*/
+	FAudioDeviceManager();
+
+	/**
+	* Destructor
+	*/
+	~FAudioDeviceManager();
+
+	/**
+	* Registers the audio device module to the audio device manager. The audio device
+	* module is responsible for creating platform-dependent audio devices.
+	*/
+	void RegisterAudioDeviceModule(class IAudioDeviceModule * AudioDeviceModuleInput);
+
+	/**
+	* Creates and audio device instance internally and returns a
+	* handle to the audio device. Returns true on success.
+	*/
+	class FAudioDevice* CreateAudioDevice(uint32 & HandleOut);
+
+	/**
+	* Returns whether the audio device handle is valid (i.e. points to
+	* an actual audio device instance)
+	*/
+	bool IsValidAudioDeviceHandle(uint32 Handle) const;
+
+	/**
+	* Shutsdown the audio device associated with the handle. The handle
+	* will become invalid after the audio device is shut down.
+	*/
+	bool ShutdownAudioDevice(uint32 Handle);
+
+	/**
+	* Returns a ptr to the audio device associated with the handle. If the
+	* handle is invalid then a NULL device ptr will be returned.
+	*/
+	class FAudioDevice* GetAudioDevice(uint32 Handle);
+
+	/** Returns the current number of active audio devices. */
+	uint8 GetNumActiveAudioDevices() const;
+
+	/** Updates all active audio devices */
+	void UpdateActiveAudioDevices(bool bGameTicking);
+
+	/** Tracks objects in the active audio devices. */
+	void AddReferencedObjects(FReferenceCollector& Collector);
+
+	/** Stops any sounds on any device using the given soundwave. */
+	void StopSoundsUsingWave(class USoundWave* InSoundWave);
+
+	/** Stops sounds using the given resource on all audio devices. */
+	void StopSoundsUsingResource(class USoundWave* InSoundWave, TArray<UAudioComponent*>& StoppedComponents);
+
+	/** Registers the Sound Class for all active devices. */
+	void RegisterSoundClass(USoundClass* SoundClass);
+
+	/** Unregisters the Sound Class for all active devices. */
+	void UnregisterSoundClass(USoundClass* SoundClass);
+
+	/** Initializes the sound class for all active devices. */
+	void InitSoundClasses();
+
+	/** Sets which audio device is the active audio device. */
+	void SetActiveDevice(uint32 InAudioDeviceHandle);
+
+	/** Links up the resource data indices for looking up and cleaning up. */
+	void TrackResource(USoundWave* SoundWave, FSoundBuffer* Buffer);
+
+	/** Frees the given sound wave resource from the device manager */
+	void FreeResource(USoundWave* SoundWave);
+
+	/** Frees the sound buffer from the device manager. */
+	void FreeBufferResource(FSoundBuffer* SoundBuffer);
+
+	/** Stops using the given sound buffer. Called before freeing the buffer */
+	void StopSourcesUsingBuffer(FSoundBuffer* Buffer);
+
+	/** Retrieves the sound buffer for the given resource id */
+	FSoundBuffer* GetSoundBufferForResourceID(uint32 ResourceID);
+
+	/** Removes the sound buffer for the given resource id */
+	void RemoveSoundBufferForResourceID(uint32 ResourceID);
+
+	/** Removes sound mix from all audio devices */
+	void RemoveSoundMix(USoundMix* SoundMix);
+
+public:
+
+	/** Array of all created buffers */
+	TArray<class FSoundBuffer*>			Buffers;
+
+	/** Look up associating a USoundWave's resource ID with sound buffers	*/
+	TMap<int32, class FSoundBuffer*>	WaveBufferMap;
+
+
+private:
+
+	/** Returns index of the given handle */
+	uint32 GetIndex(uint32 Handle) const;
+
+	/** Returns the generation of the given handle */
+	uint32 GetGeneration(uint32 Handle) const;
+
+	/** Creates a handle given the index and a generation value. */
+	uint32 CreateHandle(uint32 DeviceIndex, uint8 Generation);
+
+	/** Array for generation counts of audio devices in indices */
+	TArray<uint8> Generations;
+
+    /** Audio device module which creates audio devices. */
+    IAudioDeviceModule* AudioDeviceModule;
+
+    /** Count of the number of free slots in the audio device array. */
+    uint32 FreeIndicesSize;
+
+    /** Number of actively created audio device instances. */
+    uint8 NumActiveAudioDevices;
+
+    /** Queue for free indices */
+    TQueue<uint32> FreeIndices;
+
+    /**
+	* Array of audio device pointers. If the device has been free, then
+	* the device ptr at the array index will be null.
+	*/
+	TArray<FAudioDevice* > Devices;
+
+    /** The audio device handle currently in focus */
+	int32 InFocusDeviceHandle;
+
+	/** Next resource ID to assign out to a wave/buffer */
+	int32 NextResourceID;
+};
+
+
+
diff --git a/Engine/Source/Runtime/HTML5/HTML5Audio/Private/HTML5AudioBuffer.cpp b/Engine/Source/Runtime/HTML5/HTML5Audio/Private/HTML5AudioBuffer.cpp
--- a/Engine/Source/Runtime/HTML5/HTML5Audio/Private/HTML5AudioBuffer.cpp
+++ b/Engine/Source/Runtime/HTML5/HTML5Audio/Private/HTML5AudioBuffer.cpp
@@ -16,8 +16,8 @@
  * @param AudioDevice	audio device this sound buffer is going to be attached to.
  */
 FALSoundBuffer::FALSoundBuffer( FALAudioDevice* InAudioDevice )
+:	FSoundBuffer(InAudioDevice)
 {
-	AudioDevice	= InAudioDevice;
 }
 
 /**
@@ -48,6 +48,8 @@ FALSoundBuffer* FALSoundBuffer::Init(  FALAudioDevice* AudioDevice ,USoundWave*
 		return NULL;
 	}
 
+	FAudioDeviceManager* AudioDeviceManager = GEngine->GetAudioDeviceManager();
+
 	FALSoundBuffer *Buffer = NULL;
 
 	switch (static_cast<EDecompressionType>(InWave->DecompressionType))
@@ -62,7 +64,7 @@ FALSoundBuffer* FALSoundBuffer::Init(  FALAudioDevice* AudioDevice ,USoundWave*
 	case DTYPE_Native:
 		if (InWave->ResourceID)
 		{
-			Buffer = static_cast<FALSoundBuffer*>(AudioDevice->WaveBufferMap.FindRef(InWave->ResourceID));
+			Buffer = static_cast<FALSoundBuffer*>(AudioDeviceManager->WaveBufferMap.FindRef(InWave->ResourceID));
 		}
 
 		if (!Buffer)
@@ -126,7 +128,10 @@ FALSoundBuffer* FALSoundBuffer::CreateNativeBuffer( FALAudioDevice* AudioDevice,
 
 		AudioDevice->alError( TEXT( "RegisterSound" ) );
 
-		AudioDevice->TrackResource(Wave, Buffer);
+		FAudioDeviceManager* AudioDeviceManager = GEngine->GetAudioDeviceManager();
+		check(AudioDeviceManager != nullptr);
+
+		AudioDeviceManager->TrackResource(Wave, Buffer);
 
 		Buffer->InternalFormat = AudioDevice->GetInternalFormat( Wave->NumChannels );
 		Buffer->NumChannels = Wave->NumChannels;
diff --git a/Engine/Source/Runtime/HTML5/HTML5Audio/Private/HTML5AudioDevice.cpp b/Engine/Source/Runtime/HTML5/HTML5Audio/Private/HTML5AudioDevice.cpp
--- a/Engine/Source/Runtime/HTML5/HTML5Audio/Private/HTML5AudioDevice.cpp
+++ b/Engine/Source/Runtime/HTML5/HTML5Audio/Private/HTML5AudioDevice.cpp
@@ -17,10 +17,6 @@
 
 DEFINE_LOG_CATEGORY(LogALAudio);
 
-ALCdevice* FALAudioDevice::HardwareDevice = NULL;
-ALCcontext*	FALAudioDevice::SoundContext = NULL;
-
-
 class FHTML5AudioDeviceModule : public IAudioDeviceModule
 {
 public:
@@ -36,6 +32,10 @@ IMPLEMENT_MODULE(FHTML5AudioDeviceModule, HTML5Audio );
 /*------------------------------------------------------------------------------------
 	UALAudioDevice constructor and UObject interface.
 ------------------------------------------------------------------------------------*/
+FALAudioDevice::FALAudioDevice()
+	, HardwareDevice(nullptr)
+	, SoundContext(nullptr)
+{}
 
 void FALAudioDevice::TeardownHardware( void )
 {
diff --git a/Engine/Source/Runtime/HTML5/HTML5Audio/Private/HTML5AudioSource.cpp b/Engine/Source/Runtime/HTML5/HTML5Audio/Private/HTML5AudioSource.cpp
--- a/Engine/Source/Runtime/HTML5/HTML5Audio/Private/HTML5AudioSource.cpp
+++ b/Engine/Source/Runtime/HTML5/HTML5Audio/Private/HTML5AudioSource.cpp
@@ -3,7 +3,7 @@
 /*------------------------------------------------------------------------------------
 	FALSoundSource.
 ------------------------------------------------------------------------------------*/
-
+ 
 #include "HTML5AudioDevice.h"
 #include "Engine.h"
 /**
@@ -75,18 +75,25 @@ void FALSoundSource::Update( void )
 		return;
 	}
 
-	float Volume = WaveInstance->Volume * WaveInstance->VolumeMultiplier;
-	if( SetStereoBleed() )
+	float Volume = 1.0f;
+
+	if (AudioDevice->bIsDeviceMuted)
 	{
-		// Emulate the bleed to rear speakers followed by stereo fold down
-		Volume *= 1.25f;
+		Volume = 0.0f;
+	}
+	else
+	{
+		Volume = WaveInstance->Volume * WaveInstance->VolumeMultiplier;
+		if( SetStereoBleed() )
+		{
+			// Emulate the bleed to rear speakers followed by stereo fold down
+			Volume *= 1.25f;
+		}
+		Volume *= FApp::GetVolumeMultiplier();
+		Volume	= FMath::Clamp(Volume, 0.0f, MAX_VOLUME );
 	}
-	Volume *= FApp::GetVolumeMultiplier();
-
- 			Volume	= 	FMath::Clamp(Volume, 0.0f, MAX_VOLUME ); 
 	float	Pitch	=	FMath::Clamp(WaveInstance->Pitch, MIN_PITCH, MAX_PITCH ); 
 
-
 	// Set whether to apply reverb
 	SetReverbApplied( true );
 
diff --git a/Engine/Source/Runtime/HTML5/HTML5Audio/Public/HTML5AudioDevice.h b/Engine/Source/Runtime/HTML5/HTML5Audio/Public/HTML5AudioDevice.h
--- a/Engine/Source/Runtime/HTML5/HTML5Audio/Public/HTML5AudioDevice.h
+++ b/Engine/Source/Runtime/HTML5/HTML5Audio/Public/HTML5AudioDevice.h
@@ -184,7 +184,6 @@ protected:
 class FALAudioDevice : public FAudioDevice
 {
 public: 
-	FALAudioDevice() {} 
 	virtual ~FALAudioDevice() {} 
 
 	virtual FName GetRuntimeFormat(USoundWave* SoundWave) override
@@ -235,8 +234,8 @@ protected:
 	// AL specific
 
 	/** Device/context used to play back sounds (static so it can be initialized early) */
-	static ALCdevice*							HardwareDevice;
-	static ALCcontext*							SoundContext;
+	ALCdevice*									HardwareDevice;
+	ALCcontext*									SoundContext;
 	void*										DLLHandle;
 	/** Formats for multichannel sounds */
 	ALenum										Surround40Format;
diff --git a/Engine/Source/Runtime/IOS/IOSAudio/Private/IOSAudioBuffer.cpp b/Engine/Source/Runtime/IOS/IOSAudio/Private/IOSAudioBuffer.cpp
--- a/Engine/Source/Runtime/IOS/IOSAudio/Private/IOSAudioBuffer.cpp
+++ b/Engine/Source/Runtime/IOS/IOSAudio/Private/IOSAudioBuffer.cpp
@@ -18,7 +18,7 @@
  ------------------------------------------------------------------------------------*/
 
 FIOSAudioSoundBuffer::FIOSAudioSoundBuffer(FIOSAudioDevice* InAudioDevice, ESoundFormat InSoundFormat):
-	FSoundBuffer(),
+	FSoundBuffer(InAudioDevice),
 	SoundFormat(InSoundFormat),
 	SampleData(NULL),
 	SampleRate(0),
@@ -101,7 +101,10 @@ FIOSAudioSoundBuffer* FIOSAudioSoundBuffer::CreateNativeBuffer(FIOSAudioDevice*
 	Buffer->SampleData = static_cast<int16*>(FMemory::Malloc(Buffer->BufferSize));
 	FMemory::Memcpy(Buffer->SampleData, WaveInfo.SampleDataStart, Buffer->BufferSize);
 	
-	IOSAudioDevice->TrackResource(InWave, Buffer);
+	FAudioDeviceManager* AudioDeviceManager = GEngine->GetAudioDeviceManager();
+	check(AudioDeviceManager != nullptr);
+
+	AudioDeviceManager->TrackResource(InWave, Buffer);
 	InWave->RemoveAudioResource();
 
 	return Buffer;
@@ -115,6 +118,8 @@ FIOSAudioSoundBuffer* FIOSAudioSoundBuffer::Init(FIOSAudioDevice* IOSAudioDevice
 		return NULL;
 	}
 
+	FAudioDeviceManager* AudioDeviceManager = GEngine->GetAudioDeviceManager();
+
 	FIOSAudioSoundBuffer *Buffer = NULL;
 
 	switch (static_cast<EDecompressionType>(InWave->DecompressionType))
@@ -129,7 +134,7 @@ FIOSAudioSoundBuffer* FIOSAudioSoundBuffer::Init(FIOSAudioDevice* IOSAudioDevice
 		case DTYPE_Native:
 			if (InWave->ResourceID)
 			{
-				Buffer = static_cast<FIOSAudioSoundBuffer*>(IOSAudioDevice->WaveBufferMap.FindRef(InWave->ResourceID));
+				Buffer = static_cast<FIOSAudioSoundBuffer*>(AudioDeviceManager->WaveBufferMap.FindRef(InWave->ResourceID));
 			}
 
 			if (!Buffer)
diff --git a/Engine/Source/Runtime/IOS/IOSAudio/Private/IOSAudioSource.cpp b/Engine/Source/Runtime/IOS/IOSAudio/Private/IOSAudioSource.cpp
--- a/Engine/Source/Runtime/IOS/IOSAudio/Private/IOSAudioSource.cpp
+++ b/Engine/Source/Runtime/IOS/IOSAudio/Private/IOSAudioSource.cpp
@@ -243,7 +243,13 @@ void FIOSAudioSoundSource::Update(void)
 		return;
 	}
 
-	AudioUnitParameterValue Volume = WaveInstance->Volume * WaveInstance->VolumeMultiplier;
+	AudioUnitParameterValue Volume = 0.0f;
+
+	if (!AudioDevice->bIsDeviceMuted)
+	{
+		Volume = WaveInstance->Volume * WaveInstance->VolumeMultiplier;
+	}
+
 	if (SetStereoBleed())
 	{
 		// Emulate the bleed to rear speakers followed by stereo fold down
diff --git a/Engine/Source/Runtime/Launch/Private/LaunchEngineLoop.cpp b/Engine/Source/Runtime/Launch/Private/LaunchEngineLoop.cpp
--- a/Engine/Source/Runtime/Launch/Private/LaunchEngineLoop.cpp
+++ b/Engine/Source/Runtime/Launch/Private/LaunchEngineLoop.cpp
@@ -2083,7 +2083,7 @@ void FEngineLoop::Exit()
 
 	if ( GEngine != NULL )
 	{
-		GEngine->ShutdownAudioDevice();
+		GEngine->ShutdownAudioDeviceManager();
 	}
 
 	if ( GEngine != NULL )
diff --git a/Engine/Source/Runtime/Linux/ALAudio/Private/ALAudioBuffer.cpp b/Engine/Source/Runtime/Linux/ALAudio/Private/ALAudioBuffer.cpp
--- a/Engine/Source/Runtime/Linux/ALAudio/Private/ALAudioBuffer.cpp
+++ b/Engine/Source/Runtime/Linux/ALAudio/Private/ALAudioBuffer.cpp
@@ -16,8 +16,8 @@
  * @param AudioDevice	audio device this sound buffer is going to be attached to.
  */
 FALSoundBuffer::FALSoundBuffer( FALAudioDevice* InAudioDevice )
+	: FSoundBuffer(InAudioDevice)
 {
-	AudioDevice	= InAudioDevice;
 }
 
 /**
@@ -48,6 +48,9 @@ FALSoundBuffer* FALSoundBuffer::Init(FALAudioDevice* AudioDevice, USoundWave* In
 		return nullptr;
 	}
 
+	FAudioDeviceManager* AudioDeviceManager = GEngine->GetAudioDeviceManager();
+	check(AudioDeviceManager != nullptr);
+
 	FALSoundBuffer *Buffer = nullptr;
 
 	switch (static_cast<EDecompressionType>(InWave->DecompressionType))
@@ -62,7 +65,7 @@ FALSoundBuffer* FALSoundBuffer::Init(FALAudioDevice* AudioDevice, USoundWave* In
 	case DTYPE_Native:
 		if (InWave->ResourceID)
 		{
-			Buffer = static_cast<FALSoundBuffer*>(AudioDevice->WaveBufferMap.FindRef(InWave->ResourceID));
+			Buffer = static_cast<FALSoundBuffer*>(AudioDeviceManager->WaveBufferMap.FindRef(InWave->ResourceID));
 		}
 
 		if (!Buffer)
@@ -141,7 +144,10 @@ FALSoundBuffer* FALSoundBuffer::CreateNativeBuffer( FALAudioDevice* AudioDevice,
 
 		AudioDevice->alError(TEXT("RegisterSound"));
 
-		AudioDevice->TrackResource(Wave, Buffer);
+		FAudioDeviceManager* AudioDeviceManager = GEngine->GetAudioDeviceManager();
+		check(AudioDeviceManager != nullptr);
+
+		AudioDeviceManager->TrackResource(Wave, Buffer);
 
 		Buffer->InternalFormat = AudioDevice->GetInternalFormat(Wave->NumChannels);
 		Buffer->NumChannels = Wave->NumChannels;
diff --git a/Engine/Source/Runtime/Linux/ALAudio/Private/ALAudioDevice.cpp b/Engine/Source/Runtime/Linux/ALAudio/Private/ALAudioDevice.cpp
--- a/Engine/Source/Runtime/Linux/ALAudio/Private/ALAudioDevice.cpp
+++ b/Engine/Source/Runtime/Linux/ALAudio/Private/ALAudioDevice.cpp
@@ -17,9 +17,6 @@
 
 DEFINE_LOG_CATEGORY(LogALAudio);
 
-ALCdevice* FALAudioDevice::HardwareDevice = NULL;
-ALCcontext*	FALAudioDevice::SoundContext = NULL;
-
 
 class FALAudioDeviceModule : public IAudioDeviceModule
 {
@@ -36,6 +33,10 @@ IMPLEMENT_MODULE(FALAudioDeviceModule, ALAudio );
 /*------------------------------------------------------------------------------------
 	UALAudioDevice constructor and UObject interface.
 ------------------------------------------------------------------------------------*/
+FALAudioDevice::FALAudioDevice()
+	, HardwareDevice(nullptr)
+	, SoundContext(nullptr)
+{}
 
 void FALAudioDevice::TeardownHardware( void )
 {
diff --git a/Engine/Source/Runtime/Linux/ALAudio/Private/ALAudioSource.cpp b/Engine/Source/Runtime/Linux/ALAudio/Private/ALAudioSource.cpp
--- a/Engine/Source/Runtime/Linux/ALAudio/Private/ALAudioSource.cpp
+++ b/Engine/Source/Runtime/Linux/ALAudio/Private/ALAudioSource.cpp
@@ -80,18 +80,26 @@ void FALSoundSource::Update( void )
 		return;
 	}
 
-	float Volume = WaveInstance->Volume * WaveInstance->VolumeMultiplier;
-	if( SetStereoBleed() )
+	float Volume = 1.0f;
+
+	if (AudioDevice->bIsDeviceMuted)
+	{
+		Volume = 0.0f;
+	}
+	else
 	{
-		// Emulate the bleed to rear speakers followed by stereo fold down
-		Volume *= 1.25f;
+		Volume = WaveInstance->Volume * WaveInstance->VolumeMultiplier;
+		if (SetStereoBleed())
+		{
+			// Emulate the bleed to rear speakers followed by stereo fold down
+			Volume *= 1.25f;
+		}
+		Volume *= FApp::GetVolumeMultiplier();
+		Volume = FMath::Clamp(Volume, 0.0f, MAX_VOLUME);
 	}
-	Volume *= FApp::GetVolumeMultiplier();
 
- 			Volume	= 	FMath::Clamp(Volume, 0.0f, MAX_VOLUME ); 
 	float	Pitch	=	FMath::Clamp(WaveInstance->Pitch, MIN_PITCH, MAX_PITCH ); 
 
-
 	// Set whether to apply reverb
 	SetReverbApplied( true );
 
diff --git a/Engine/Source/Runtime/Linux/ALAudio/Public/ALAudioDevice.h b/Engine/Source/Runtime/Linux/ALAudio/Public/ALAudioDevice.h
--- a/Engine/Source/Runtime/Linux/ALAudio/Public/ALAudioDevice.h
+++ b/Engine/Source/Runtime/Linux/ALAudio/Public/ALAudioDevice.h
@@ -183,7 +183,8 @@ protected:
 class FALAudioDevice : public FAudioDevice
 {
 public: 
-	FALAudioDevice() {} 
+	FALAudioDevice();
+
 	virtual ~FALAudioDevice() {} 
 
 	virtual FName GetRuntimeFormat(USoundWave* SoundWave) override
@@ -234,8 +235,8 @@ protected:
 	// AL specific
 
 	/** Device/context used to play back sounds (static so it can be initialized early) */
-	static ALCdevice*							HardwareDevice;
-	static ALCcontext*							SoundContext;
+	ALCdevice*									HardwareDevice;
+	ALCcontext*									SoundContext;
 	void*										DLLHandle;
 	/** Formats for multichannel sounds */
 	ALenum										Surround40Format;
diff --git a/Engine/Source/Runtime/Mac/CoreAudio/Private/CoreAudioBuffer.cpp b/Engine/Source/Runtime/Mac/CoreAudio/Private/CoreAudioBuffer.cpp
--- a/Engine/Source/Runtime/Mac/CoreAudio/Private/CoreAudioBuffer.cpp
+++ b/Engine/Source/Runtime/Mac/CoreAudio/Private/CoreAudioBuffer.cpp
@@ -25,8 +25,8 @@
  * @param AudioDevice	audio device this sound buffer is going to be attached to.
  */
 FCoreAudioSoundBuffer::FCoreAudioSoundBuffer( FAudioDevice* InAudioDevice, ESoundFormat InSoundFormat )
-:	AudioDevice( InAudioDevice ),
-	SoundFormat( InSoundFormat ),
+:   FSoundBuffer( InAudioDevice ),
+    SoundFormat( InSoundFormat ),
 	PCMData( NULL ),
 	PCMDataSize( 0 ),
 	DecompressionState( NULL ),
@@ -207,9 +207,12 @@ FCoreAudioSoundBuffer* FCoreAudioSoundBuffer::CreateProceduralBuffer( FCoreAudio
  */
 FCoreAudioSoundBuffer* FCoreAudioSoundBuffer::CreatePreviewBuffer( FCoreAudioDevice* CoreAudioDevice, USoundWave* Wave, FCoreAudioSoundBuffer* Buffer )
 {
-	if( Buffer )
+    FAudioDeviceManager* AudioDeviceManager = GEngine->GetAudioDeviceManager();
+    check(AudioDeviceManager != nullptr);
+    
+	if (Buffer)
 	{
-		CoreAudioDevice->FreeBufferResource( Buffer );
+		AudioDeviceManager->FreeBufferResource( Buffer );
 	}
 	
 	// Create new buffer.
@@ -226,7 +229,7 @@ FCoreAudioSoundBuffer* FCoreAudioSoundBuffer::CreatePreviewBuffer( FCoreAudioDev
 	
 	Buffer->InitAudioStreamBasicDescription( kAudioFormatLinearPCM, Wave, true );
 	
-	CoreAudioDevice->TrackResource( Wave, Buffer );
+	AudioDeviceManager->TrackResource( Wave, Buffer );
 	
 	return( Buffer );
 }
@@ -262,7 +265,10 @@ FCoreAudioSoundBuffer* FCoreAudioSoundBuffer::CreateNativeBuffer( FCoreAudioDevi
 	// Keep track of associated resource name.
 	Buffer->InitAudioStreamBasicDescription( kAudioFormatLinearPCM, Wave, true );
 	
-	CoreAudioDevice->TrackResource( Wave, Buffer );
+    FAudioDeviceManager* AudioDeviceManager = GEngine->GetAudioDeviceManager();
+    check(AudioDeviceManager != nullptr);
+    
+	AudioDeviceManager->TrackResource( Wave, Buffer );
 	
 	Wave->RemoveAudioResource();
 	
@@ -284,6 +290,8 @@ FCoreAudioSoundBuffer* FCoreAudioSoundBuffer::Init( FAudioDevice* AudioDevice, U
 		return NULL;
 	}
 
+    FAudioDeviceManager* AudioDeviceManager = GEngine->GetAudioDeviceManager();
+    
 	FCoreAudioDevice *CoreAudioDevice = ( FCoreAudioDevice *)AudioDevice;
 	FCoreAudioSoundBuffer *Buffer = NULL;
 
@@ -307,7 +315,7 @@ FCoreAudioSoundBuffer* FCoreAudioSoundBuffer::Init( FAudioDevice* AudioDevice, U
 			// Find the existing buffer if any
 			if( Wave->ResourceID )
 			{
-				Buffer = (FCoreAudioSoundBuffer*)CoreAudioDevice->WaveBufferMap.FindRef( Wave->ResourceID );
+				Buffer = (FCoreAudioSoundBuffer*)AudioDeviceManager->WaveBufferMap.FindRef( Wave->ResourceID );
 			}
 			
 			// Override with any new PCM data even if some already exists. 
@@ -332,7 +340,7 @@ FCoreAudioSoundBuffer* FCoreAudioSoundBuffer::Init( FAudioDevice* AudioDevice, U
 			// Upload entire wav to CoreAudio
 			if( Wave->ResourceID )
 			{
-				Buffer = (FCoreAudioSoundBuffer*)CoreAudioDevice->WaveBufferMap.FindRef( Wave->ResourceID );
+				Buffer = (FCoreAudioSoundBuffer*)AudioDeviceManager->WaveBufferMap.FindRef( Wave->ResourceID );
 			}
 
 			if( Buffer == NULL )
diff --git a/Engine/Source/Runtime/Mac/CoreAudio/Private/CoreAudioSource.cpp b/Engine/Source/Runtime/Mac/CoreAudio/Private/CoreAudioSource.cpp
--- a/Engine/Source/Runtime/Mac/CoreAudio/Private/CoreAudioSource.cpp
+++ b/Engine/Source/Runtime/Mac/CoreAudio/Private/CoreAudioSource.cpp
@@ -251,7 +251,12 @@ void FCoreAudioSoundSource::Update( void )
 		return;
 	}
 
-	float Volume = WaveInstance->GetActualVolume();
+	float Volume = 0.0f;
+	
+	if (!AudioDevice->bIsDeviceMuted)
+	{
+		Volume = WaveInstance->GetActualVolume();
+	}
 	
 	if( Buffer->NumChannels < 3 )
 	{
diff --git a/Engine/Source/Runtime/MovieSceneCoreTypes/Private/Instances/MovieSceneAudioTrackInstance.cpp b/Engine/Source/Runtime/MovieSceneCoreTypes/Private/Instances/MovieSceneAudioTrackInstance.cpp
--- a/Engine/Source/Runtime/MovieSceneCoreTypes/Private/Instances/MovieSceneAudioTrackInstance.cpp
+++ b/Engine/Source/Runtime/MovieSceneCoreTypes/Private/Instances/MovieSceneAudioTrackInstance.cpp
@@ -107,7 +107,8 @@ void FMovieSceneAudioTrackInstance::Update( float Position, float LastPosition,
 				UAudioComponent* Component = GetAudioComponent(Actors[ActorIndex], RowIndex);
 				if (Component->IsPlaying())
 				{
-					FActiveSound* ActiveSound = GEngine->GetAudioDevice()->FindActiveSound(Component);
+					FAudioDevice* AudioDevice = Component->GetAudioDevice();
+					FActiveSound* ActiveSound = AudioDevice->FindActiveSound(Component);
 					ActiveSound->bLocationDefined = true;
 					ActiveSound->Transform = Actors[ActorIndex]->GetTransform();
 				}
diff --git a/Engine/Source/Runtime/Online/OnlineSubsystemUtils/Private/OnlineSubsystemUtils.cpp b/Engine/Source/Runtime/Online/OnlineSubsystemUtils/Private/OnlineSubsystemUtils.cpp
--- a/Engine/Source/Runtime/Online/OnlineSubsystemUtils/Private/OnlineSubsystemUtils.cpp
+++ b/Engine/Source/Runtime/Online/OnlineSubsystemUtils/Private/OnlineSubsystemUtils.cpp
@@ -42,21 +42,19 @@ void FOnlineSubsystemUtilsModule::ShutdownModule()
 
 UAudioComponent* CreateVoiceAudioComponent(uint32 SampleRate)
 {
-	UAudioComponent* AudioComponent = NULL;
-	if (GEngine != NULL)
+	UAudioComponent* AudioComponent = nullptr;
+	if (GEngine != nullptr)
 	{
-		FAudioDevice* AudioDevice = GEngine->GetAudioDevice();
-		if (AudioDevice != NULL)
+		if (FAudioDevice* AudioDevice = GEngine->GetMainAudioDevice())
 		{
-			USoundWaveStreaming* SoundStreaming = NULL;
-			SoundStreaming = NewObject<USoundWaveStreaming>();
+			USoundWaveStreaming* SoundStreaming = NewObject<USoundWaveStreaming>();
 			SoundStreaming->SampleRate = SampleRate;
 			SoundStreaming->NumChannels = 1;
 			SoundStreaming->Duration = INDEFINITELY_LOOPING_DURATION;
 			SoundStreaming->SoundGroup = SOUNDGROUP_Voice;
 			SoundStreaming->bLooping = false;
 
-			AudioComponent = AudioDevice->CreateComponent(SoundStreaming, NULL, NULL, false);
+			AudioComponent = AudioDevice->CreateComponent(SoundStreaming, nullptr, nullptr, false);
 			if (AudioComponent)
 			{
 				AudioComponent->bIsUISound = true;
@@ -66,7 +64,7 @@ UAudioComponent* CreateVoiceAudioComponent(uint32 SampleRate)
 				const FStringAssetReference VoiPSoundClassName = GetDefault<UAudioSettings>()->VoiPSoundClass;
 				if (VoiPSoundClassName.IsValid())
 				{
-					AudioComponent->SoundClassOverride = LoadObject<USoundClass>(NULL, *VoiPSoundClassName.ToString());
+					AudioComponent->SoundClassOverride = LoadObject<USoundClass>(nullptr, *VoiPSoundClassName.ToString());
 				}
 			}
 			else
diff --git a/Engine/Source/Runtime/Windows/XAudio2/Private/XAudio2Buffer.cpp b/Engine/Source/Runtime/Windows/XAudio2/Private/XAudio2Buffer.cpp
--- a/Engine/Source/Runtime/Windows/XAudio2/Private/XAudio2Buffer.cpp
+++ b/Engine/Source/Runtime/Windows/XAudio2/Private/XAudio2Buffer.cpp
@@ -80,10 +80,10 @@ struct FXMAInfo
  * @param AudioDevice	audio device this sound buffer is going to be attached to.
  */
 FXAudio2SoundBuffer::FXAudio2SoundBuffer( FAudioDevice* InAudioDevice, ESoundFormat InSoundFormat )
-:	AudioDevice( InAudioDevice ),
-	SoundFormat( InSoundFormat ),
+:	SoundFormat( InSoundFormat ),
 	DecompressionState( NULL ),
-	bDynamicResource( false )
+	bDynamicResource( false ),
+	FSoundBuffer(InAudioDevice)
 {
 	PCM.PCMData = NULL;
 	PCM.PCMDataSize = 0;
@@ -412,9 +412,12 @@ FXAudio2SoundBuffer* FXAudio2SoundBuffer::CreateProceduralBuffer( FXAudio2Device
  */
 FXAudio2SoundBuffer* FXAudio2SoundBuffer::CreatePreviewBuffer( FXAudio2Device* XAudio2Device, USoundWave* Wave, FXAudio2SoundBuffer* Buffer )
 {
-	if( Buffer )
+	FAudioDeviceManager* AudioDeviceManager = GEngine->GetAudioDeviceManager();
+	check(AudioDeviceManager != nullptr);
+
+	if (Buffer)
 	{
-		XAudio2Device->FreeBufferResource( Buffer );
+		AudioDeviceManager->FreeBufferResource(Buffer);
 	}
 
 	// Create new buffer.
@@ -431,7 +434,7 @@ FXAudio2SoundBuffer* FXAudio2SoundBuffer::CreatePreviewBuffer( FXAudio2Device* X
 
 	Buffer->InitWaveFormatEx( WAVE_FORMAT_PCM, Wave, true );
 
-	XAudio2Device->TrackResource( Wave, Buffer );
+	AudioDeviceManager->TrackResource(Wave, Buffer);
 
 	return( Buffer );
 }
@@ -467,7 +470,9 @@ FXAudio2SoundBuffer* FXAudio2SoundBuffer::CreateNativeBuffer( FXAudio2Device* XA
 	// Keep track of associated resource name.
 	Buffer->InitWaveFormatEx( WAVE_FORMAT_PCM, Wave, true );
 
-	XAudio2Device->TrackResource( Wave, Buffer );
+	FAudioDeviceManager* AudioDeviceManager = GEngine->GetAudioDeviceManager();
+	check(AudioDeviceManager != nullptr);
+	AudioDeviceManager->TrackResource(Wave, Buffer);
 
 	Wave->RemoveAudioResource();
 
@@ -524,6 +529,8 @@ FXAudio2SoundBuffer* FXAudio2SoundBuffer::Init( FAudioDevice* AudioDevice, USoun
 		return( NULL );
 	}
 
+	FAudioDeviceManager* AudioDeviceManager = GEngine->GetAudioDeviceManager();
+
 	FXAudio2Device* XAudio2Device = ( FXAudio2Device* )AudioDevice;
 	FXAudio2SoundBuffer* Buffer = NULL;
 
@@ -550,7 +557,7 @@ FXAudio2SoundBuffer* FXAudio2SoundBuffer::Init( FAudioDevice* AudioDevice, USoun
 		// Find the existing buffer if any
 		if( Wave->ResourceID )
 		{
-			Buffer = (FXAudio2SoundBuffer*)XAudio2Device->WaveBufferMap.FindRef( Wave->ResourceID );
+			Buffer = (FXAudio2SoundBuffer*)AudioDeviceManager->GetSoundBufferForResourceID(Wave->ResourceID);
 		}
 
 		// Override with any new PCM data even if some already exists. 
@@ -576,7 +583,7 @@ FXAudio2SoundBuffer* FXAudio2SoundBuffer::Init( FAudioDevice* AudioDevice, USoun
 		// Upload entire wav to XAudio2
 		if( Wave->ResourceID )
 		{
-			Buffer = (FXAudio2SoundBuffer*)XAudio2Device->WaveBufferMap.FindRef( Wave->ResourceID );
+			Buffer = (FXAudio2SoundBuffer*)AudioDeviceManager->GetSoundBufferForResourceID(Wave->ResourceID);
 		}
 
 		if( Buffer == NULL )
diff --git a/Engine/Source/Runtime/Windows/XAudio2/Private/XAudio2Device.cpp b/Engine/Source/Runtime/Windows/XAudio2/Private/XAudio2Device.cpp
--- a/Engine/Source/Runtime/Windows/XAudio2/Private/XAudio2Device.cpp
+++ b/Engine/Source/Runtime/Windows/XAudio2/Private/XAudio2Device.cpp
@@ -41,16 +41,14 @@ public:
 IMPLEMENT_MODULE(FXAudio2DeviceModule, XAudio2);
 
 /*------------------------------------------------------------------------------------
-	Static variables from the early init
+Static variables from the early init
 ------------------------------------------------------------------------------------*/
 
 // The number of speakers producing sound (stereo or 5.1)
-int32 FXAudioDeviceProperties::NumSpeakers						= 0;
-IXAudio2* FXAudioDeviceProperties::XAudio2						= NULL;
-const float* FXAudioDeviceProperties::OutputMixMatrix			= NULL;
-IXAudio2MasteringVoice* FXAudioDeviceProperties::MasteringVoice	= NULL;
+int32 FXAudioDeviceProperties::NumSpeakers = 0;
+const float* FXAudioDeviceProperties::OutputMixMatrix = NULL;
 #if XAUDIO_SUPPORTS_DEVICE_DETAILS
-XAUDIO2_DEVICE_DETAILS FXAudioDeviceProperties::DeviceDetails	= { 0 };
+XAUDIO2_DEVICE_DETAILS FXAudioDeviceProperties::DeviceDetails = { 0 };
 #endif	//XAUDIO_SUPPORTS_DEVICE_DETAILS
 
 /*------------------------------------------------------------------------------------
@@ -68,6 +66,12 @@ bool FXAudio2Device::InitializeHardware()
 		return false;
 	}
 
+	// Create a new DeviceProperties object
+	DeviceProperties = new FXAudioDeviceProperties;
+	// null out the non-static data
+	DeviceProperties->XAudio2 = nullptr;
+	DeviceProperties->MasteringVoice = nullptr;
+
 	// Load ogg and vorbis dlls if they haven't been loaded yet
 	LoadVorbisLibraries();
 
@@ -87,29 +91,33 @@ bool FXAudio2Device::InitializeHardware()
 #else
 	uint32 Flags = 0;
 #endif
-	if( XAudio2Create( &FXAudioDeviceProperties::XAudio2, Flags, AUDIO_HWTHREAD ) != S_OK )
+
+	// Create a new XAudio2 device object instance
+	if (XAudio2Create(&DeviceProperties->XAudio2, Flags, AUDIO_HWTHREAD) != S_OK)
 	{
 		UE_LOG(LogInit, Log, TEXT( "Failed to create XAudio2 interface" ) );
 		return( false );
 	}
 
+	check(DeviceProperties->XAudio2 != nullptr);
+
 #if XAUDIO_SUPPORTS_DEVICE_DETAILS
 	UINT32 DeviceCount = 0;
 	ValidateAPICall(TEXT("GetDeviceCount"),
-		FXAudioDeviceProperties::XAudio2->GetDeviceCount( &DeviceCount ));
+		DeviceProperties->XAudio2->GetDeviceCount(&DeviceCount));
 	if( DeviceCount < 1 )
 	{
 		UE_LOG(LogInit, Log, TEXT( "No audio devices found!" ) );
-		FXAudioDeviceProperties::XAudio2 = NULL;
+		DeviceProperties->XAudio2 = nullptr;
 		return( false );		
 	}
 
 	// Get the details of the default device 0
 	if( !ValidateAPICall(TEXT("GetDeviceDetails"),
-			FXAudioDeviceProperties::XAudio2->GetDeviceDetails( 0, &FXAudioDeviceProperties::DeviceDetails )))
+		DeviceProperties->XAudio2->GetDeviceDetails(0, &FXAudioDeviceProperties::DeviceDetails)))
 	{
 		UE_LOG(LogInit, Log, TEXT( "Failed to get DeviceDetails for XAudio2" ) );
-		FXAudioDeviceProperties::XAudio2 = NULL;
+		DeviceProperties->XAudio2 = nullptr;
 		return( false );
 	}
 
@@ -117,7 +125,7 @@ bool FXAudio2Device::InitializeHardware()
 	XAUDIO2_DEBUG_CONFIGURATION DebugConfig = {0};
 	DebugConfig.TraceMask = XAUDIO2_LOG_WARNINGS | XAUDIO2_LOG_DETAIL;
 	DebugConfig.BreakMask = XAUDIO2_LOG_ERRORS;
-	FXAudioDeviceProperties::XAudio2->SetDebugConfiguration(&DebugConfig);
+	DeviceProperties->XAudio2->SetDebugConfiguration(&DebugConfig);
 #endif
 
 	FXAudioDeviceProperties::NumSpeakers = UE4_XAUDIO2_NUMCHANNELS;
@@ -140,16 +148,16 @@ bool FXAudio2Device::InitializeHardware()
 	if( !GetOutputMatrix( UE4_XAUDIO2_CHANNELMASK, FXAudioDeviceProperties::NumSpeakers ) )
 	{
 		UE_LOG(LogInit, Log, TEXT( "Unsupported speaker configuration for this number of channels" ) );
-		FXAudioDeviceProperties::XAudio2 = NULL;
+		DeviceProperties->XAudio2 = NULL;
 		return( false );
 	}
 
 	// Create the final output voice with either 2 or 6 channels
 	if (!ValidateAPICall(TEXT("CreateMasteringVoice"), 
-			FXAudioDeviceProperties::XAudio2->CreateMasteringVoice(&FXAudioDeviceProperties::MasteringVoice, FXAudioDeviceProperties::NumSpeakers, SampleRate, 0, 0, NULL)))
+		DeviceProperties->XAudio2->CreateMasteringVoice(&DeviceProperties->MasteringVoice, FXAudioDeviceProperties::NumSpeakers, SampleRate, 0, 0, NULL)))
 	{
 		UE_LOG(LogInit, Warning, TEXT( "Failed to create the mastering voice for XAudio2" ) );
-		FXAudioDeviceProperties::XAudio2 = NULL;
+		DeviceProperties->XAudio2 = nullptr;
 		return( false );
 	}
 #else	//XAUDIO_SUPPORTS_DEVICE_DETAILS
@@ -183,20 +191,26 @@ bool FXAudio2Device::InitializeHardware()
 
 void FXAudio2Device::TeardownHardware()
 {
-	// close hardware interfaces
-	if( FXAudioDeviceProperties::MasteringVoice )
+	if (DeviceProperties)
 	{
-		FXAudioDeviceProperties::MasteringVoice->DestroyVoice();
-		FXAudioDeviceProperties::MasteringVoice = NULL;
-	}
+		// close hardware interfaces
+		if (DeviceProperties->MasteringVoice)
+		{
+			DeviceProperties->MasteringVoice->DestroyVoice();
+			DeviceProperties->MasteringVoice = NULL;
+		}
 
-	if( FXAudioDeviceProperties::XAudio2 )
-	{
-		// Force the hardware to release all references
-		FXAudioDeviceProperties::XAudio2->Release();
-		FXAudioDeviceProperties::XAudio2 = NULL;
+		if (DeviceProperties->XAudio2)
+		{
+			// Force the hardware to release all references
+			DeviceProperties->XAudio2->Release();
+			DeviceProperties->XAudio2 = NULL;
+		}
+		delete DeviceProperties;
+		DeviceProperties = nullptr;
 	}
 
+
 #if PLATFORM_WINDOWS
 	if (bComInitialized)
 	{
diff --git a/Engine/Source/Runtime/Windows/XAudio2/Private/XAudio2Effects.cpp b/Engine/Source/Runtime/Windows/XAudio2/Private/XAudio2Effects.cpp
--- a/Engine/Source/Runtime/Windows/XAudio2/Private/XAudio2Effects.cpp
+++ b/Engine/Source/Runtime/Windows/XAudio2/Private/XAudio2Effects.cpp
@@ -328,14 +328,17 @@ bool FXAudio2EffectsManager::CreateEQPremasterVoices()
 		1, EQEffects
 	};
 
+	check(XAudio2Device->DeviceProperties != nullptr); 
+	check(XAudio2Device->DeviceProperties->XAudio2 != nullptr);
+
 	if( !XAudio2Device->ValidateAPICall( TEXT( "CreateSubmixVoice (EQPremaster)" ), 
-		FXAudioDeviceProperties::XAudio2->CreateSubmixVoice( &EQPremasterVoice, SPEAKER_COUNT, SampleRate, 0, STAGE_EQPREMASTER, NULL, &EQEffectChain ) ) )
+		XAudio2Device->DeviceProperties->XAudio2->CreateSubmixVoice(&EQPremasterVoice, SPEAKER_COUNT, SampleRate, 0, STAGE_EQPREMASTER, NULL, &EQEffectChain)))
 	{
 		return( false );
 	}
 
 	if( !XAudio2Device->ValidateAPICall( TEXT( "CreateSubmixVoice (DryPremaster)" ), 
-		FXAudioDeviceProperties::XAudio2->CreateSubmixVoice( &DryPremasterVoice, SPEAKER_COUNT, SampleRate, 0, STAGE_EQPREMASTER, NULL, NULL ) ) )
+		XAudio2Device->DeviceProperties->XAudio2->CreateSubmixVoice(&DryPremasterVoice, SPEAKER_COUNT, SampleRate, 0, STAGE_EQPREMASTER, NULL, NULL)))
 	{
 		return( false );
 	}
@@ -398,7 +401,7 @@ bool FXAudio2EffectsManager::CreateReverbVoice()
 	};
 
 	if( !XAudio2Device->ValidateAPICall( TEXT( "CreateSubmixVoice (Reverb)" ), 
-		FXAudioDeviceProperties::XAudio2->CreateSubmixVoice( &ReverbEffectVoice, 2, SampleRate, 0, STAGE_REVERB, &ReverbSends, &ReverbEffectChain ) ) )
+		XAudio2Device->DeviceProperties->XAudio2->CreateSubmixVoice(&ReverbEffectVoice, 2, SampleRate, 0, STAGE_REVERB, &ReverbSends, &ReverbEffectChain)))
 	{
 		return( false );
 	}
@@ -455,7 +458,7 @@ bool FXAudio2EffectsManager::CreateRadioVoice()
 	// Finally, create the submix voice that holds the radio effect. Sounds (source voices) 
 	// will be piped to this submix voice to receive radio distortion. 
 	if( !XAudio2Device->ValidateAPICall( TEXT( "CreateSubmixVoice (Radio)" ), 
-		FXAudioDeviceProperties::XAudio2->CreateSubmixVoice( &RadioEffectVoice, OutputChannelCount, SampleRate, 0, STAGE_RADIO, NULL, &RadioEffectChain ) ) )
+		XAudio2Device->DeviceProperties->XAudio2->CreateSubmixVoice(&RadioEffectVoice, OutputChannelCount, SampleRate, 0, STAGE_RADIO, NULL, &RadioEffectChain)))
 	{
 		return false;
 	}
diff --git a/Engine/Source/Runtime/Windows/XAudio2/Private/XAudio2Source.cpp b/Engine/Source/Runtime/Windows/XAudio2/Private/XAudio2Source.cpp
--- a/Engine/Source/Runtime/Windows/XAudio2/Private/XAudio2Source.cpp
+++ b/Engine/Source/Runtime/Windows/XAudio2/Private/XAudio2Source.cpp
@@ -310,12 +310,14 @@ bool FXAudio2SoundSource::CreateSource( void )
 #endif	//XAUDIO2_SUPPORTS_MUSIC
 	Flags |= XAUDIO2_VOICE_USEFILTER;
 
+	check(AudioDevice->DeviceProperties != nullptr);
+	check(AudioDevice->DeviceProperties->XAudio2 != nullptr);
 	// All sound formats start with the WAVEFORMATEX structure (PCM, XMA2, XWMA)
-	if( !AudioDevice->ValidateAPICall( TEXT( "CreateSourceVoice (source)" ), 
+	if (!AudioDevice->ValidateAPICall(TEXT("CreateSourceVoice (source)"),
 #if XAUDIO2_SUPPORTS_SENDLIST
-		FXAudioDeviceProperties::XAudio2->CreateSourceVoice( &Source, ( WAVEFORMATEX* )&XAudio2Buffer->PCM, Flags, MAX_PITCH, &SourceCallback, &SourceSendList, NULL ) ) )
+		AudioDevice->DeviceProperties->XAudio2->CreateSourceVoice( &Source, ( WAVEFORMATEX* )&XAudio2Buffer->PCM, Flags, MAX_PITCH, &SourceCallback, &SourceSendList, NULL ) ) )
 #else	//XAUDIO2_SUPPORTS_SENDLIST
-		FXAudioDeviceProperties::XAudio2->CreateSourceVoice( &Source, ( WAVEFORMATEX* )&XAudio2Buffer->PCM, Flags, MAX_PITCH, &SourceCallback ) ) )
+		AudioDevice->DeviceProperties->XAudio2->CreateSourceVoice(&Source, (WAVEFORMATEX*)&XAudio2Buffer->PCM, Flags, MAX_PITCH, &SourceCallback)))
 #endif	//XAUDIO2_SUPPORTS_SENDLIST
 	{
 		return( false );
@@ -330,28 +332,39 @@ bool FXAudio2SoundSource::CreateSource( void )
  * @param	WaveInstance	wave instace being primed for playback
  * @return	true if initialization was successful, false otherwise
  */
-bool FXAudio2SoundSource::Init( FWaveInstance* InWaveInstance )
+bool FXAudio2SoundSource::Init(FWaveInstance* InWaveInstance)
 {
 	if (InWaveInstance->OutputTarget != EAudioOutputTarget::Controller)
 	{
 		// Find matching buffer.
-		XAudio2Buffer = FXAudio2SoundBuffer::Init( AudioDevice, InWaveInstance->WaveData, InWaveInstance->StartTime > 0.f );
+		FAudioDevice* AudioDevice = nullptr;
+		if (InWaveInstance->ActiveSound->AudioComponent.IsValid())
+		{
+			AudioDevice = InWaveInstance->ActiveSound->AudioComponent->GetAudioDevice();
+		}
+		else
+		{
+			AudioDevice = GEngine->GetMainAudioDevice();
+		}
+		check(AudioDevice);
+
+		XAudio2Buffer = FXAudio2SoundBuffer::Init(AudioDevice, InWaveInstance->WaveData, InWaveInstance->StartTime > 0.f);
 		Buffer = XAudio2Buffer;
 
 		// Buffer failed to be created, or there was an error with the compressed data
-		if( Buffer && Buffer->NumChannels > 0 )
+		if (Buffer && Buffer->NumChannels > 0)
 		{
 			SCOPE_CYCLE_COUNTER( STAT_AudioSourceInitTime );
 
 			WaveInstance = InWaveInstance;
 
 			// Set whether to apply reverb
-			SetReverbApplied( Effects->ReverbEffectVoice != NULL );
+			SetReverbApplied(Effects->ReverbEffectVoice != nullptr);
 
 			// Create a new source
-			if( !CreateSource() )
+			if (!CreateSource())
 			{
-				return( false );
+				return false;
 			}
 
 			if (WaveInstance->StartTime > 0.f)
@@ -360,7 +373,7 @@ bool FXAudio2SoundSource::Init( FWaveInstance* InWaveInstance )
 			}
 
 			// Submit audio buffers
-			switch( XAudio2Buffer->SoundFormat )
+			switch (XAudio2Buffer->SoundFormat)
 			{
 			case SoundFormat_PCM:
 			case SoundFormat_PCMPreview:
@@ -385,12 +398,12 @@ bool FXAudio2SoundSource::Init( FWaveInstance* InWaveInstance )
 			Update();
 		
 			// Initialization succeeded.
-			return( true );
+			return true;
 		}
 	}
 
 	// Initialization failed.
-	return( false );
+	return false;
 }
 
 /**
@@ -398,7 +411,7 @@ bool FXAudio2SoundSource::Init( FWaveInstance* InWaveInstance )
  */
 void FXAudio2SoundSource::GetChannelVolumes( float ChannelVolumes[CHANNELOUT_COUNT], float AttenuatedVolume )
 {
-	if (FApp::GetVolumeMultiplier() == 0.0f)
+	if (FApp::GetVolumeMultiplier() == 0.0f || AudioDevice->bIsDeviceMuted)
 	{
 		for( int32 i = 0; i < CHANNELOUT_COUNT; i++ )
 		{
diff --git a/Engine/Source/Runtime/Windows/XAudio2/Private/XAudio2Support.h b/Engine/Source/Runtime/Windows/XAudio2/Private/XAudio2Support.h
--- a/Engine/Source/Runtime/Windows/XAudio2/Private/XAudio2Support.h
+++ b/Engine/Source/Runtime/Windows/XAudio2/Private/XAudio2Support.h
@@ -199,9 +199,6 @@ public:
 	virtual int32 GetCurrentChunkIndex() const override;
 	virtual int32 GetCurrentChunkOffset() const override;
 
-	/** Audio device this buffer is attached to	*/
-	FAudioDevice*				AudioDevice;
-
 	/** Format of the sound referenced by this buffer */
 	int32							SoundFormat;
 
@@ -487,9 +484,12 @@ public:
 /** Variables required for the early init */
 struct FXAudioDeviceProperties
 {
-	static int32							NumSpeakers;
-	static struct IXAudio2*					XAudio2;
-	static struct IXAudio2MasteringVoice*		MasteringVoice;
+	// These variables are non-static to support multiple audio device instances
+	struct IXAudio2*					XAudio2;
+	struct IXAudio2MasteringVoice*		MasteringVoice;
+
+	// These variables are static because they are common across all audio device instances
+	static int32						NumSpeakers;
 	static const float*					OutputMixMatrix;
 #if XAUDIO_SUPPORTS_DEVICE_DETAILS
 	static XAUDIO2_DEVICE_DETAILS		DeviceDetails;
diff --git a/Engine/Source/Runtime/Windows/XAudio2/Public/XAudio2Device.h b/Engine/Source/Runtime/Windows/XAudio2/Public/XAudio2Device.h
--- a/Engine/Source/Runtime/Windows/XAudio2/Public/XAudio2Device.h
+++ b/Engine/Source/Runtime/Windows/XAudio2/Public/XAudio2Device.h
@@ -16,6 +16,7 @@
 
 DECLARE_LOG_CATEGORY_EXTERN(LogXAudio2, Log, All);
 
+
 /*------------------------------------------------------------------------------------
 	Dependencies, helpers & forward declarations.
 ------------------------------------------------------------------------------------*/
@@ -171,6 +172,8 @@ protected:
 	friend class FXAudio2EffectsManager;
 
 private:
+	struct FXAudioDeviceProperties* DeviceProperties;
+
 #if PLATFORM_WINDOWS
 	// We need to keep track whether com was successfully initialized so we can clean 
 	// it up during shutdown
