#ue4 - Fix PhysicsVolumes not being unregistered from the global list when they are streamed out.

diff --git a/Engine/Source/Runtime/Engine/Classes/GameFramework/PhysicsVolume.h b/Engine/Source/Runtime/Engine/Classes/GameFramework/PhysicsVolume.h
--- a/Engine/Source/Runtime/Engine/Classes/GameFramework/PhysicsVolume.h
+++ b/Engine/Source/Runtime/Engine/Classes/GameFramework/PhysicsVolume.h
@@ -22,6 +22,7 @@ class ENGINE_API APhysicsVolume : public AVolume
 
 	virtual void PostInitializeComponents() override;
 	virtual void Destroyed() override;
+	virtual void EndPlay(const EEndPlayReason::Type EndPlayReason) override;
 
 	//======================================================================================
 	// Character Movement related properties
diff --git a/Engine/Source/Runtime/Engine/Private/PhysicsVolume.cpp b/Engine/Source/Runtime/Engine/Private/PhysicsVolume.cpp
--- a/Engine/Source/Runtime/Engine/Private/PhysicsVolume.cpp
+++ b/Engine/Source/Runtime/Engine/Private/PhysicsVolume.cpp
@@ -38,23 +38,27 @@ void APhysicsVolume::LoadedFromAnotherClass(const FName& OldClassName)
 void APhysicsVolume::PostInitializeComponents()
 {
 	Super::PostInitializeComponents();
+	GetWorld()->AddPhysicsVolume(this);
+}
 
+void APhysicsVolume::Destroyed()
+{
 	UWorld* MyWorld = GetWorld();
 	if (MyWorld)
 	{
-		MyWorld->AddPhysicsVolume(this);
+		MyWorld->RemovePhysicsVolume(this);
 	}
+	Super::Destroyed();
 }
 
-void APhysicsVolume::Destroyed()
+void APhysicsVolume::EndPlay(const EEndPlayReason::Type EndPlayReason)
 {
 	UWorld* MyWorld = GetWorld();
 	if (MyWorld)
 	{
 		MyWorld->RemovePhysicsVolume(this);
 	}
-
-	Super::Destroyed();
+	Super::EndPlay(EndPlayReason);
 }
 
 
