Fix non-unity

diff --git a/Engine/Source/Runtime/Engine/Private/LevelActor.cpp b/Engine/Source/Runtime/Engine/Private/LevelActor.cpp
--- a/Engine/Source/Runtime/Engine/Private/LevelActor.cpp
+++ b/Engine/Source/Runtime/Engine/Private/LevelActor.cpp
@@ -6,6 +6,7 @@
 #include "Net/UnrealNetwork.h"
 #include "Collision.h"
 #include "Engine/DemoNetDriver.h"
+#include "AudioDeviceManager.h"
 
 #if WITH_PHYSX
 	#include "PhysicsEngine/PhysXSupport.h"
