#ANIM: SkelMesh: Fix Master Pose Component to work with hiding bone : requested by licensee

diff --git a/Engine/Source/Runtime/Engine/Private/SkeletalRender.cpp b/Engine/Source/Runtime/Engine/Private/SkeletalRender.cpp
--- a/Engine/Source/Runtime/Engine/Private/SkeletalRender.cpp
+++ b/Engine/Source/Runtime/Engine/Private/SkeletalRender.cpp
@@ -196,6 +196,7 @@ void UpdateRefToLocalMatrices( TArray<FMatrix>& ReferenceToLocal, const USkinned
 {
 	const USkeletalMesh* const ThisMesh = InMeshComponent->SkeletalMesh;
 	const USkinnedMeshComponent* const MasterComp = InMeshComponent->MasterPoseComponent.Get();
+	const USkeletalMesh* const MasterCompMesh = MasterComp? MasterComp->SkeletalMesh : nullptr;
 	const FStaticLODModel& LOD = InSkeletalMeshResource->LODModels[LODIndex];
 
 	check( ThisMesh->RefBasesInvMatrix.Num() != 0 );
@@ -229,11 +230,20 @@ void UpdateRefToLocalMatrices( TArray<FMatrix>& ReferenceToLocal, const USkinned
 				if( bIsMasterCompValid )
 				{
 					// If valid, use matrix from parent component.
-					const int32 ParentBoneIndex = InMeshComponent->MasterBoneMap[ThisBoneIndex];
-					if ( MasterComp->GetSpaceBases().IsValidIndex(ParentBoneIndex) )
+					const int32 MasterBoneIndex = InMeshComponent->MasterBoneMap[ThisBoneIndex];
+					if ( MasterComp->GetSpaceBases().IsValidIndex(MasterBoneIndex) )
 					{
-						checkSlow(MasterComp->GetSpaceBases()[ParentBoneIndex].IsRotationNormalized());
-						ReferenceToLocal[ThisBoneIndex] = MasterComp->GetSpaceBases()[ParentBoneIndex].ToMatrixWithScale();
+						const int32 ParentIndex = MasterCompMesh->RefSkeleton.GetParentIndex(MasterBoneIndex);
+						bool bNeedToHideBone = MasterComp->BoneVisibilityStates[MasterBoneIndex] != BVS_Visible;
+						if (bNeedToHideBone && ParentIndex != INDEX_NONE)
+						{
+							ReferenceToLocal[ThisBoneIndex] = ReferenceToLocal[ParentIndex].ApplyScale(0.f);
+						}
+						else
+						{
+							checkSlow(MasterComp->GetSpaceBases()[MasterBoneIndex].IsRotationNormalized());
+							ReferenceToLocal[ThisBoneIndex] = MasterComp->GetSpaceBases()[MasterBoneIndex].ToMatrixWithScale();
+						}
 					}
 				}
 				else
