Core: Added support for convertible laptops

diff --git a/Engine/Source/Runtime/Core/Private/GenericPlatform/GenericPlatformMisc.cpp b/Engine/Source/Runtime/Core/Private/GenericPlatform/GenericPlatformMisc.cpp
--- a/Engine/Source/Runtime/Core/Private/GenericPlatform/GenericPlatformMisc.cpp
+++ b/Engine/Source/Runtime/Core/Private/GenericPlatform/GenericPlatformMisc.cpp
@@ -852,6 +852,11 @@ bool FGenericPlatformMisc::SetEpicAccountId( const FString& AccountId )
 	return FPlatformMisc::SetStoredValue( TEXT( "Epic Games" ), TEXT( "Unreal Engine/Identifiers" ), TEXT( "AccountId" ), AccountId );
 }
 
+EConvertibleLaptopMode FGenericPlatformMisc::GetConvertibleLaptopMode()
+{
+	return EConvertibleLaptopMode::NotSupported;
+}
+
 const TCHAR* FGenericPlatformMisc::GetEngineMode()
 {
 	return	
diff --git a/Engine/Source/Runtime/Core/Private/Misc/CallbackDevice.cpp b/Engine/Source/Runtime/Core/Private/Misc/CallbackDevice.cpp
--- a/Engine/Source/Runtime/Core/Private/Misc/CallbackDevice.cpp
+++ b/Engine/Source/Runtime/Core/Private/Misc/CallbackDevice.cpp
@@ -68,3 +68,4 @@ FCoreDelegates::FStatDisabled FCoreDelegates::StatDisabled;
 FCoreDelegates::FStatDisableAll FCoreDelegates::StatDisableAll;
 
 FCoreDelegates::FApplicationLicenseChange FCoreDelegates::ApplicationLicenseChange;
+FCoreDelegates::FPlatformChangedLaptopMode FCoreDelegates::PlatformChangedLaptopMode;
diff --git a/Engine/Source/Runtime/Core/Private/Windows/WindowsApplication.cpp b/Engine/Source/Runtime/Core/Private/Windows/WindowsApplication.cpp
--- a/Engine/Source/Runtime/Core/Private/Windows/WindowsApplication.cpp
+++ b/Engine/Source/Runtime/Core/Private/Windows/WindowsApplication.cpp
@@ -16,29 +16,26 @@
 
 // Allow Windows Platform types in the entire file.
 #include "AllowWindowsPlatformTypes.h"
-	#include "Ole2.h"
-	#include <shlobj.h>
-	#include <objbase.h>
-	#include <SetupApi.h>
-	#include <devguid.h>
-	#include <dwmapi.h>
+#include "Ole2.h"
+#include <shlobj.h>
+#include <objbase.h>
+#include <SetupApi.h>
+#include <devguid.h>
+#include <dwmapi.h>
+
 
 // This might not be defined by Windows when maintaining backwards-compatibility to pre-Vista builds
 #ifndef WM_MOUSEHWHEEL
 #define WM_MOUSEHWHEEL                  0x020E
 #endif
 
-// This might not be defined by Windows when maintaining backwards-compatibility to pre-Win8 builds
-#ifndef SM_CONVERTIBLESLATEMODE
-#define SM_CONVERTIBLESLATEMODE			0x2003
-#endif
-
 DEFINE_LOG_CATEGORY(LogWindowsDesktop);
 
 const FIntPoint FWindowsApplication::MinimizedWindowPosition(-32000,-32000);
 
 FWindowsApplication* WindowApplication = NULL;
 
+
 FWindowsApplication* FWindowsApplication::CreateWindowsApplication( const HINSTANCE InstanceHandle, const HICON IconHandle )
 {
 	WindowApplication = new FWindowsApplication( InstanceHandle, IconHandle );
@@ -1497,9 +1494,10 @@ int32 FWindowsApplication::ProcessDeferredMessage( const FDeferredWindowsMessage
 			break;
 
 		case WM_SETTINGCHANGE:
-			MessageHandler->OnConvertibleDeviceModeChanged((GetSystemMetrics(SM_CONVERTIBLESLATEMODE) == 0)
-				? EConvertibleLaptopModes::Tablet
-				: EConvertibleLaptopModes::Laptop);
+			if ((lParam != 0) && (FCString::Strcmp((LPCTSTR)lParam, TEXT("ConvertibleSlateMode")) == 0))
+			{
+				MessageHandler->OnConvertibleLaptopModeChanged();
+			}
 			break;
 	
 		case WM_NCACTIVATE:
@@ -1994,4 +1992,5 @@ TSharedRef<FTaskbarList> FTaskbarList::Create()
 	return TaskbarList;
 }
 
+
 #include "HideWindowsPlatformTypes.h"
diff --git a/Engine/Source/Runtime/Core/Private/Windows/WindowsPlatformMisc.cpp b/Engine/Source/Runtime/Core/Private/Windows/WindowsPlatformMisc.cpp
--- a/Engine/Source/Runtime/Core/Private/Windows/WindowsPlatformMisc.cpp
+++ b/Engine/Source/Runtime/Core/Private/Windows/WindowsPlatformMisc.cpp
@@ -9,6 +9,7 @@
 // Resource includes.
 #include "Runtime/Launch/Resources/Windows/Resource.h"
 #include "Runtime/Launch/Resources/Version.h"
+
 #include "AllowWindowsPlatformTypes.h"
 	#include <time.h>
 	#include <MMSystem.h>
@@ -19,6 +20,7 @@
 	#include <shellapi.h>
 	#include <Iphlpapi.h>
 #include "HideWindowsPlatformTypes.h"
+
 #include "MallocTBB.h"
 #include "ModuleManager.h"
 #include "MallocAnsi.h"
@@ -36,6 +38,13 @@
 #include <io.h>
 
 #include "AllowWindowsPlatformTypes.h"
+
+// This might not be defined by Windows when maintaining backwards-compatibility to pre-Win8 builds
+#ifndef SM_CONVERTIBLESLATEMODE
+#define SM_CONVERTIBLESLATEMODE			0x2003
+#endif
+
+
 int32 FWindowsOSVersionHelper::GetOSVersions( FString& out_OSVersionLabel, FString& out_OSSubVersionLabel )
 {
 	int32 ErrorCode = (int32)SUCCEEDED;
@@ -2335,3 +2344,19 @@ FString FWindowsPlatformMisc::GetOperatingSystemId()
 	QueryRegKey(HKEY_LOCAL_MACHINE, TEXT("Software\\Microsoft\\Cryptography"), TEXT("MachineGuid"), Result);
 	return Result;
 }
+
+
+EConvertibleLaptopMode FWindowsPlatformMisc::GetConvertibleLaptopMode()
+{
+	if (!VerifyWindowsMajorVersion(8))
+	{
+		return EConvertibleLaptopMode::NotSupported;
+	}
+
+	if (::GetSystemMetrics(SM_CONVERTIBLESLATEMODE) == 0)
+	{
+		return EConvertibleLaptopMode::Tablet;
+	}
+	
+	return EConvertibleLaptopMode::Laptop;
+}
diff --git a/Engine/Source/Runtime/Core/Public/GenericPlatform/GenericApplicationMessageHandler.h b/Engine/Source/Runtime/Core/Public/GenericPlatform/GenericApplicationMessageHandler.h
--- a/Engine/Source/Runtime/Core/Public/GenericPlatform/GenericApplicationMessageHandler.h
+++ b/Engine/Source/Runtime/Core/Public/GenericPlatform/GenericApplicationMessageHandler.h
@@ -1,12 +1,15 @@
 // Copyright 1998-2015 Epic Games, Inc. All Rights Reserved.
 
 #pragma once
+
 #include "HAL/Platform.h"
 #include "Templates/SharedPointer.h"
 #include "Misc/Optional.h"
 
+
 class FGenericWindow;
 
+
 namespace EMouseButtons
 {
 	enum Type
@@ -21,6 +24,7 @@ namespace EMouseButtons
 	};
 }
 
+
 namespace EControllerButtons
 {
 	enum Type
@@ -82,6 +86,7 @@ namespace EControllerButtons
 	};
 }
 
+
 namespace EWindowActivation
 {
 	enum Type
@@ -92,6 +97,7 @@ namespace EWindowActivation
 	};
 }
 
+
 namespace EWindowZone
 {
 	/**
@@ -121,6 +127,7 @@ namespace EWindowZone
 	};
 }
 
+
 namespace EWindowAction
 {
 	enum Type
@@ -132,6 +139,7 @@ namespace EWindowAction
 	};
 }
 
+
 namespace EDropEffect
 {
 	enum Type
@@ -143,6 +151,7 @@ namespace EDropEffect
 	};
 }
 
+
 namespace EGestureEvent
 {
 	enum Type
@@ -157,18 +166,6 @@ namespace EGestureEvent
 }
 
 
-/**
- * Enumerates the modes a convertible laptop can be in.
- */
-enum class EConvertibleLaptopModes
-{
-	/** Laptop arranged as a tablet. */
-	Tablet,
-
-	/** Laptop arranged as a laptop. */
-	Laptop
-};
-
 /** Defines the minimum and maximum dimensions that a window can take on. */
 struct FWindowSizeLimits
 {
@@ -347,7 +344,7 @@ public:
 		return false;
 	}
 
-	virtual bool OnConvertibleDeviceModeChanged(const EConvertibleLaptopModes NewMode)
+	virtual bool OnConvertibleLaptopModeChanged()
 	{
 		return false;
 	}
@@ -391,4 +388,4 @@ public:
 	{
 		return true;
 	}
-};
\ No newline at end of file
+};
diff --git a/Engine/Source/Runtime/Core/Public/GenericPlatform/GenericPlatformMisc.h b/Engine/Source/Runtime/Core/Public/GenericPlatform/GenericPlatformMisc.h
--- a/Engine/Source/Runtime/Core/Public/GenericPlatform/GenericPlatformMisc.h
+++ b/Engine/Source/Runtime/Core/Public/GenericPlatform/GenericPlatformMisc.h
@@ -108,6 +108,22 @@ namespace EBuildTargets
 }
 
 
+/**
+ * Enumerates the modes a convertible laptop can be in.
+ */
+enum class EConvertibleLaptopMode
+{
+	/** Not a convertible laptop. */
+	NotSupported,
+
+	/** Laptop arranged as a laptop. */
+	Laptop,
+
+	/** Laptop arranged as a tablet. */
+	Tablet
+};
+
+
 namespace EErrorReportMode
 {
 	/** 
@@ -791,6 +807,13 @@ public:
 	 */
 	static FString GetOperatingSystemId();
 
+	/**
+	 * Gets the current mode of convertible laptops, i.e. Laptop or Tablet.
+	 *
+	 * @return The laptop mode, or Unknown if not known, or NotSupported if not a convertible laptop.
+	 */
+	static EConvertibleLaptopMode GetConvertibleLaptopMode();
+
 	/** 
 	 * Get a string description of the mode the engine was running in.
 	 */
diff --git a/Engine/Source/Runtime/Core/Public/Misc/CallbackDevice.h b/Engine/Source/Runtime/Core/Public/Misc/CallbackDevice.h
--- a/Engine/Source/Runtime/Core/Public/Misc/CallbackDevice.h
+++ b/Engine/Source/Runtime/Core/Public/Misc/CallbackDevice.h
@@ -199,6 +199,10 @@ public:
 	DECLARE_MULTICAST_DELEGATE(FApplicationLicenseChange);
 	static FApplicationLicenseChange ApplicationLicenseChange;
 
+	/** Sent when the platform changed its laptop mode (for convertible laptops).*/
+	DECLARE_MULTICAST_DELEGATE_OneParam(FPlatformChangedLaptopMode, EConvertibleLaptopMode);
+	static FPlatformChangedLaptopMode PlatformChangedLaptopMode;
+
 private:
 
 	// Callbacks for hotfixes
diff --git a/Engine/Source/Runtime/Core/Public/Windows/WindowsPlatformMisc.h b/Engine/Source/Runtime/Core/Public/Windows/WindowsPlatformMisc.h
--- a/Engine/Source/Runtime/Core/Public/Windows/WindowsPlatformMisc.h
+++ b/Engine/Source/Runtime/Core/Public/Windows/WindowsPlatformMisc.h
@@ -243,6 +243,7 @@ struct CORE_API FWindowsPlatformMisc
 	 */
 	static FString GetOperatingSystemId();
 
+	static EConvertibleLaptopMode GetConvertibleLaptopMode();
 };
 
 
diff --git a/Engine/Source/Runtime/Engine/Classes/Components/PlatformEventsComponent.h b/Engine/Source/Runtime/Engine/Classes/Components/PlatformEventsComponent.h
new file mode 100644
--- /dev/null
+++ b/Engine/Source/Runtime/Engine/Classes/Components/PlatformEventsComponent.h
@@ -0,0 +1,71 @@
+// Copyright 1998-2015 Epic Games, Inc. All Rights Reserved.
+
+#pragma once
+
+#include "PlatformEventsComponent.generated.h"
+
+
+/**
+ * Component to handle receiving notifications from the OS about platform events.
+ */
+UCLASS(ClassGroup=Utility, HideCategories=(Activation, "Components|Activation", Collision), meta=(BlueprintSpawnableComponent))
+class ENGINE_API UPlatformEventsComponent
+	: public UActorComponent
+{
+	GENERATED_UCLASS_BODY()
+
+	/**
+	 * Check whether a convertible laptop is laptop mode.
+	 *
+	 * @return true if in laptop mode, false otherwise or if not a convertible laptop.
+	 * @see IsInTabletMode, SupportsConvertibleLaptops
+	 */
+	UFUNCTION(BlueprintCallable, Category="Components|Activation")
+	bool IsInLaptopMode();
+
+	/**
+	 * Check whether a convertible laptop is laptop mode.
+	 *
+	 * @return true if in tablet mode, false otherwise or if not a convertible laptop.
+	 * @see IsInLaptopMode, SupportsConvertibleLaptops
+	 */
+	UFUNCTION(BlueprintCallable, Category="Components|Activation")
+	bool IsInTabletMode();
+
+	/**
+	 * Check whether the platform supports convertible laptops.
+	 *
+	 * Note: This does not necessarily mean that the platform is a convertible laptop.
+	 * For example, convertible laptops running Windows 7 or older will return false,
+	 * and regular laptops running Windows 8 or newer will return true.
+	 *
+	 * @return true for convertible laptop platforms, false otherwise.
+	 * @see IsInLaptopMode, IsInTabletMode
+	 */
+	UFUNCTION(BlueprintCallable, Category="Components|Activation")
+	bool SupportsConvertibleLaptops();
+
+public:
+
+	DECLARE_DYNAMIC_MULTICAST_DELEGATE(FPlatformEventDelegate);
+
+	/** This is called when a convertible laptop changed into laptop mode. */
+	UPROPERTY(BlueprintAssignable)
+	FPlatformEventDelegate PlatformChangedToLaptopModeDelegate;  
+
+	/** This is called when a convertible laptop changed into tablet mode. */
+	UPROPERTY(BlueprintAssignable)
+	FPlatformEventDelegate PlatformChangedToTabletModeDelegate;  
+
+public:
+
+	// UActorComponent overrides
+
+	virtual void OnRegister() override;
+	virtual void OnUnregister() override;
+
+private:
+
+	/** Handles the FCoreDelegates.PlatformChangedLaptopMode delegate. */
+	void HandlePlatformChangedLaptopMode(EConvertibleLaptopMode NewMode);
+};
diff --git a/Engine/Source/Runtime/Engine/Private/Components/PlatformEventsComponent.cpp b/Engine/Source/Runtime/Engine/Private/Components/PlatformEventsComponent.cpp
new file mode 100644
--- /dev/null
+++ b/Engine/Source/Runtime/Engine/Private/Components/PlatformEventsComponent.cpp
@@ -0,0 +1,57 @@
+// Copyright 1998-2015 Epic Games, Inc. All Rights Reserved.
+
+#include "EnginePrivate.h"
+#include "Components/PlatformEventsComponent.h"
+#include "CallbackDevice.h"
+
+
+UPlatformEventsComponent::UPlatformEventsComponent(const FObjectInitializer& ObjectInitializer)
+	: Super(ObjectInitializer)
+{ }
+
+
+bool UPlatformEventsComponent::IsInLaptopMode()
+{
+	return (FPlatformMisc::GetConvertibleLaptopMode() == EConvertibleLaptopMode::Laptop);
+}
+
+
+bool UPlatformEventsComponent::IsInTabletMode()
+{
+	return (FPlatformMisc::GetConvertibleLaptopMode() == EConvertibleLaptopMode::Tablet);
+}
+
+
+bool UPlatformEventsComponent::SupportsConvertibleLaptops()
+{
+	return (FPlatformMisc::GetConvertibleLaptopMode() != EConvertibleLaptopMode::NotSupported);
+}
+
+
+void UPlatformEventsComponent::OnRegister()
+{
+	Super::OnRegister();
+
+	FCoreDelegates::PlatformChangedLaptopMode.AddUObject(this, &UPlatformEventsComponent::HandlePlatformChangedLaptopMode);
+}
+
+
+void UPlatformEventsComponent::OnUnregister()
+{
+	Super::OnUnregister();
+	
+ 	FCoreDelegates::PlatformChangedLaptopMode.RemoveAll(this);
+}
+
+
+void UPlatformEventsComponent::HandlePlatformChangedLaptopMode(EConvertibleLaptopMode NewMode)
+{
+	if (NewMode == EConvertibleLaptopMode::Laptop)
+	{
+		PlatformChangedToLaptopModeDelegate.Broadcast();
+	}
+	else if (NewMode == EConvertibleLaptopMode::Tablet)
+	{
+		PlatformChangedToTabletModeDelegate.Broadcast();
+	}
+}
diff --git a/Engine/Source/Runtime/Slate/Private/Framework/Application/SlateApplication.cpp b/Engine/Source/Runtime/Slate/Private/Framework/Application/SlateApplication.cpp
--- a/Engine/Source/Runtime/Slate/Private/Framework/Application/SlateApplication.cpp
+++ b/Engine/Source/Runtime/Slate/Private/Framework/Application/SlateApplication.cpp
@@ -5130,13 +5130,15 @@ void FSlateApplication::ProcessApplicationActivationEvent( bool InAppActivated )
 }
 
 
-bool FSlateApplication::OnConvertibleDeviceModeChanged(const EConvertibleLaptopModes NewMode)
+bool FSlateApplication::OnConvertibleLaptopModeChanged()
 {
+	EConvertibleLaptopMode NewMode = FPlatformMisc::GetConvertibleLaptopMode();
+
 	// Notify that we want the mobile experience when in tablet mode, otherwise use mouse and keyboard
 	if (!(FParse::Param(FCommandLine::Get(), TEXT("simmobile")) || FParse::Param(FCommandLine::Get(), TEXT("faketouches"))))
 	{
 		// Not sure what the correct long-term strategy is. Use bIsFakingTouch for now to get things going.
-		if (NewMode == EConvertibleLaptopModes::Tablet)
+		if (NewMode == EConvertibleLaptopMode::Tablet)
 		{
 			bIsFakingTouch = true;
 		}
@@ -5146,6 +5148,8 @@ bool FSlateApplication::OnConvertibleDeviceModeChanged(const EConvertibleLaptopM
 		}
 	}
 
+	FCoreDelegates::PlatformChangedLaptopMode.Broadcast(NewMode);
+
 	return true;
 }
 
diff --git a/Engine/Source/Runtime/Slate/Public/Framework/Application/SlateApplication.h b/Engine/Source/Runtime/Slate/Public/Framework/Application/SlateApplication.h
--- a/Engine/Source/Runtime/Slate/Public/Framework/Application/SlateApplication.h
+++ b/Engine/Source/Runtime/Slate/Public/Framework/Application/SlateApplication.h
@@ -1039,7 +1039,7 @@ public:
 	virtual void OnMovedWindow( const TSharedRef< FGenericWindow >& PlatformWindow, const int32 X, const int32 Y ) override;
 	virtual bool OnWindowActivationChanged( const TSharedRef< FGenericWindow >& PlatformWindow, const EWindowActivation::Type ActivationType ) override;
 	virtual bool OnApplicationActivationChanged( const bool IsActive ) override;
-	virtual bool OnConvertibleDeviceModeChanged(const EConvertibleLaptopModes NewMode) override;
+	virtual bool OnConvertibleLaptopModeChanged() override;
 	virtual EWindowZone::Type GetWindowZoneForPoint( const TSharedRef< FGenericWindow >& PlatformWindow, const int32 X, const int32 Y ) override;
 	virtual void OnWindowClose( const TSharedRef< FGenericWindow >& PlatformWindow ) override;
 	virtual EDropEffect::Type OnDragEnterText( const TSharedRef< FGenericWindow >& Window, const FString& Text ) override;
@@ -1483,5 +1483,4 @@ private:
 	// e.g. On windows the origin (coordinates X=0, Y=0) is the upper left of the primary monitor,
 	// but there could be another monitor on any of the sides.
 	FSlateRect VirtualDesktopRect;
-
 };
