Increase size of replay packet buffer + set max bunch size to same value (we don't need to fragment replay packets)

diff --git a/Engine/Source/Runtime/Engine/Private/DemoNetDriver.cpp b/Engine/Source/Runtime/Engine/Private/DemoNetDriver.cpp
--- a/Engine/Source/Runtime/Engine/Private/DemoNetDriver.cpp
+++ b/Engine/Source/Runtime/Engine/Private/DemoNetDriver.cpp
@@ -25,7 +25,7 @@ static TAutoConsoleVariable<float> CVarDemoRecordHz( TEXT( "demo.RecordHz" ), 10
 static TAutoConsoleVariable<float> CVarDemoTimeDilation( TEXT( "demo.TimeDilation" ), -1.0f, TEXT( "Override time dilation during demo playback (-1 = don't override)" ) );
 static TAutoConsoleVariable<float> CVarDemoSkipTime( TEXT( "demo.SkipTime" ), 0, TEXT( "Skip fixed amount of network replay time (in seconds)" ) );
 
-static const int32 MAX_DEMO_READ_WRITE_BUFFER = 1024 * 2;
+static const int32 MAX_DEMO_READ_WRITE_BUFFER = 1024 * 32;
 
 #define DEMO_CHECKSUMS 0		// When setting this to 1, this will invalidate all demos, you will need to re-record and playback
 
@@ -81,11 +81,12 @@ FString UDemoNetDriver::LowLevelGetNetworkNumber()
 enum ENetworkVersionHistory
 {
 	HISTORY_INITIAL				= 1,
-	HISTORY_SAVE_ABS_TIME_MS	= 2			// We now save the abs demo time in ms for each frame (solves accumulation errors)
+	HISTORY_SAVE_ABS_TIME_MS	= 2,			// We now save the abs demo time in ms for each frame (solves accumulation errors)
+	HISTORY_INCREASE_BUFFER		= 3				// Increased buffer size of packets, which invalidates old replays
 };
 
 static const uint32 NETWORK_DEMO_MAGIC				= 0x2CF5A13D;
-static const uint32 NETWORK_DEMO_VERSION			= HISTORY_SAVE_ABS_TIME_MS;
+static const uint32 NETWORK_DEMO_VERSION			= HISTORY_INCREASE_BUFFER;
 
 static const uint32 NETWORK_DEMO_METADATA_MAGIC		= 0x3D06B24E;
 static const uint32 NETWORK_DEMO_METADATA_VERSION	= 0;
@@ -1029,7 +1030,7 @@ void UDemoNetDriver::ReplayStreamingReady( bool bSuccess, bool bRecord )
 
 UDemoNetConnection::UDemoNetConnection( const FObjectInitializer& ObjectInitializer ) : Super( ObjectInitializer )
 {
-	MaxPacket = 512;
+	MaxPacket = MAX_DEMO_READ_WRITE_BUFFER;
 	InternalAck = true;
 }
 
@@ -1038,6 +1039,9 @@ void UDemoNetConnection::InitConnection( UNetDriver* InDriver, EConnectionState
 	// default implementation
 	Super::InitConnection( InDriver, InState, InURL, InConnectionSpeed );
 
+	MaxPacket = MAX_DEMO_READ_WRITE_BUFFER;
+	InternalAck = true;
+
 	InitSendBuffer();
 
 	// the driver must be a DemoRecording driver (GetDriver makes assumptions to avoid Cast'ing each time)
