#ue4 - CharacterMovementComponent runs normal CheckStillInWorld() code when UpdatedComponent is simulating physics.

- Fixes ragdolls not respecting KillZ or world bounds.

UDN: https://udn.unrealengine.com/questions/239015/killz-doesnt-work-for-ragdolls.html

diff --git a/Engine/Source/Runtime/Engine/Private/CharacterMovementComponent.cpp b/Engine/Source/Runtime/Engine/Private/CharacterMovementComponent.cpp
--- a/Engine/Source/Runtime/Engine/Private/CharacterMovementComponent.cpp
+++ b/Engine/Source/Runtime/Engine/Private/CharacterMovementComponent.cpp
@@ -847,28 +847,37 @@ void UCharacterMovementComponent::TickComponent(float DeltaTime, enum ELevelTick
 	SCOPE_CYCLE_COUNTER(STAT_CharacterMovementTick);
 
 	const FVector InputVector = ConsumeInputVector();
-	if (!HasValidData() || ShouldSkipUpdate(DeltaTime) || UpdatedComponent->IsSimulatingPhysics())
+	if (!HasValidData() || ShouldSkipUpdate(DeltaTime))
 	{
 		return;
 	}
 
-	if (AvoidanceLockTimer > 0.0f)
+	// See if we fell out of the world.
+	const bool bIsSimulatingPhysics = UpdatedComponent->IsSimulatingPhysics();
+	if (CharacterOwner->Role == ROLE_Authority)
 	{
-		AvoidanceLockTimer -= DeltaTime;
-	}
-
-	if (CharacterOwner->Role > ROLE_SimulatedProxy)
-	{
-		if (CharacterOwner->Role == ROLE_Authority)
+		if (!bCheatFlying || bIsSimulatingPhysics)
 		{
-			// Check we are still in the world, and stop simulating if not.
-			const bool bStillInWorld = (bCheatFlying || CharacterOwner->CheckStillInWorld());
-			if (!bStillInWorld || !HasValidData())
+			if (!CharacterOwner->CheckStillInWorld())
 			{
 				return;
 			}
 		}
+	}
+
+	// We don't update if simulating physics (eg ragdolls).
+	if (bIsSimulatingPhysics)
+	{
+		return;
+	}
 
+	if (AvoidanceLockTimer > 0.0f)
+	{
+		AvoidanceLockTimer -= DeltaTime;
+	}
+
+	if (CharacterOwner->Role > ROLE_SimulatedProxy)
+	{
 		// If we are a client we might have received an update from the server.
 		const bool bIsClient = (GetNetMode() == NM_Client && CharacterOwner->Role == ROLE_AutonomousProxy);
 		if (bIsClient)
