Auto-ack packagemap network guids for 100% reliable connections

diff --git a/Engine/Source/Runtime/Engine/Private/PackageMapClient.cpp b/Engine/Source/Runtime/Engine/Private/PackageMapClient.cpp
--- a/Engine/Source/Runtime/Engine/Private/PackageMapClient.cpp
+++ b/Engine/Source/Runtime/Engine/Private/PackageMapClient.cpp
@@ -936,6 +936,13 @@ void UPackageMapClient::NotifyBunchCommit( const int32 OutPacketId, const TArray
 		// (GUID information doesn't change, so updating to the latest expected sequence is unnecessary)
 		if ( ExpectedPacketIdRef == GUID_PACKET_NOT_ACKED )
 		{
+			if ( Connection->InternalAck )
+			{
+				// Auto ack now if the connection is 100% reliable
+				ExpectedPacketIdRef = GUID_PACKET_ACKED;
+				continue;
+			}
+
 			ExpectedPacketIdRef = OutPacketId;
 			check( !PendingAckGUIDs.Contains( ExportNetGUIDs[i] ) );	// If we hit this assert, this means the lists are out of sync
 			PendingAckGUIDs.AddUnique( ExportNetGUIDs[i] );
