Fix compile issues

diff --git a/Engine/Source/Editor/BlueprintGraph/Private/DelegateNodeHandlers.cpp b/Engine/Source/Editor/BlueprintGraph/Private/DelegateNodeHandlers.cpp
--- a/Engine/Source/Editor/BlueprintGraph/Private/DelegateNodeHandlers.cpp
+++ b/Engine/Source/Editor/BlueprintGraph/Private/DelegateNodeHandlers.cpp
@@ -74,7 +74,7 @@ struct FKCHandlerDelegateHelper
 			const UFunction* OrgSignature = DelegateNode->GetDelegateSignature();
 			if(const UEdGraphPin* DelegatePin = DelegateNode->GetDelegatePin())
 			{
-				const UFunction* PinSignature = DelegatePin->PinType.PinSubCategoryMemberReference.ResolveSimpleMemberReference<UFunction>();
+				const UFunction* PinSignature = FMemberReference::ResolveSimpleMemberReference<UFunction>(DelegatePin->PinType.PinSubCategoryMemberReference);
 				if (!OrgSignature || !PinSignature || !OrgSignature->IsSignatureCompatibleWith(PinSignature))
 				{
 					MessageLog.Error(*LOCTEXT("WrongDelegate", "Wrong Event Dispatcher. Refresh node @@").ToString(), DelegateNode);
@@ -259,11 +259,11 @@ void FKCHandler_CreateDelegate::RegisterNets(FKismetFunctionContext& Context, UE
 		if(NULL == OutDelegateTerm)
 		{
 			OutDelegateTerm = Context.CreateLocalTerminalFromPinAutoChooseScope(Net, Context.NetNameMap->MakeValidName(Net));
-			if (NULL == OutDelegateTerm->Type.PinSubCategoryMemberReference.ResolveSimpleMemberReference<UFunction>())
+			if (NULL == FMemberReference::ResolveSimpleMemberReference<UFunction>(OutDelegateTerm->Type.PinSubCategoryMemberReference))
 			{
-				OutDelegateTerm->Type.PinSubCategoryMemberReference.FillSimpleMemberReference<UFunction>(DelegateNode->GetDelegateSignature());
+				FMemberReference::FillSimpleMemberReference<UFunction>(DelegateNode->GetDelegateSignature(), OutDelegateTerm->Type.PinSubCategoryMemberReference);
 			}
-			if (NULL == OutDelegateTerm->Type.PinSubCategoryMemberReference.ResolveSimpleMemberReference<UFunction>())
+			if (NULL == FMemberReference::ResolveSimpleMemberReference<UFunction>(OutDelegateTerm->Type.PinSubCategoryMemberReference))
 			{
 				CompilerContext.MessageLog.Error(*LOCTEXT("UnconnectedDelegateSig", "Event Dispatcher has no signature @@").ToString(), OutPin);
 				return;
diff --git a/Engine/Source/Editor/BlueprintGraph/Private/EdGraphSchema_K2.cpp b/Engine/Source/Editor/BlueprintGraph/Private/EdGraphSchema_K2.cpp
--- a/Engine/Source/Editor/BlueprintGraph/Private/EdGraphSchema_K2.cpp
+++ b/Engine/Source/Editor/BlueprintGraph/Private/EdGraphSchema_K2.cpp
@@ -3078,12 +3078,12 @@ bool UEdGraphSchema_K2::ConvertPropertyToPinType(const UProperty* Property, /*ou
 	else if (const UMulticastDelegateProperty* MulticastDelegateProperty = Cast<const UMulticastDelegateProperty>(TestProperty))
 	{
 		TypeOut.PinCategory = PC_MCDelegate;
-		TypeOut.PinSubCategoryMemberReference.FillSimpleMemberReference<UFunction>(MulticastDelegateProperty->SignatureFunction);
+		FMemberReference::FillSimpleMemberReference<UFunction>(MulticastDelegateProperty->SignatureFunction, TypeOut.PinSubCategoryMemberReference);
 	}
 	else if (const UDelegateProperty* DelegateProperty = Cast<const UDelegateProperty>(TestProperty))
 	{
 		TypeOut.PinCategory = PC_Delegate;
-		TypeOut.PinSubCategoryMemberReference.FillSimpleMemberReference<UFunction>(DelegateProperty->SignatureFunction);
+		FMemberReference::FillSimpleMemberReference<UFunction>(DelegateProperty->SignatureFunction, TypeOut.PinSubCategoryMemberReference);
 	}
 	else
 	{
@@ -3469,7 +3469,7 @@ bool UEdGraphSchema_K2::ArePinTypesCompatible(const FEdGraphPinType& Output, con
 				return Func && (Func->HasAllFlags(RF_LoadCompleted) || !Func->HasAnyFlags(RF_NeedLoad | RF_WasLoaded));
 			};
 
-			const UFunction* OutFunction = Output.PinSubCategoryMemberReference.ResolveSimpleMemberReference<UFunction>();
+			const UFunction* OutFunction = FMemberReference::ResolveSimpleMemberReference<UFunction>(Output.PinSubCategoryMemberReference);
 			if (!CanUseFunction(OutFunction))
 			{
 				OutFunction = NULL;
@@ -3483,7 +3483,7 @@ bool UEdGraphSchema_K2::ArePinTypesCompatible(const FEdGraphPinType& Output, con
 					OutFunction = BPOwner->SkeletonGeneratedClass->FindFunctionByName(Output.PinSubCategoryMemberReference.MemberName);
 				}
 			}
-			const UFunction* InFunction = Input.PinSubCategoryMemberReference.ResolveSimpleMemberReference<UFunction>();
+			const UFunction* InFunction = FMemberReference::ResolveSimpleMemberReference<UFunction>(Input.PinSubCategoryMemberReference);
 			if (!CanUseFunction(InFunction))
 			{
 				InFunction = NULL;
diff --git a/Engine/Source/Editor/BlueprintGraph/Private/K2ActionMenuBuilder.cpp b/Engine/Source/Editor/BlueprintGraph/Private/K2ActionMenuBuilder.cpp
--- a/Engine/Source/Editor/BlueprintGraph/Private/K2ActionMenuBuilder.cpp
+++ b/Engine/Source/Editor/BlueprintGraph/Private/K2ActionMenuBuilder.cpp
@@ -1416,7 +1416,7 @@ void FK2ActionMenuBuilder::GetPinAllowedNodeTypes(FBlueprintGraphActionListBuild
 				TSharedPtr<FEdGraphSchemaAction_K2NewNode> Action = FK2ActionMenuBuilder::AddNewNodeAction(ContextMenuBuilder, FString(), NodeTemplate->GetNodeTitle(ENodeTitleType::ListView), NodeTemplate->GetTooltipText().ToString(), 0, NodeTemplate->GetKeywords());
 				Action->NodeTemplate = NodeTemplate;
 
-				UFunction* Signature = FromPin.PinType.PinSubCategoryMemberReference.ResolveSimpleMemberReference<UFunction>();
+				UFunction* Signature = FMemberReference::ResolveSimpleMemberReference<UFunction>(FromPin.PinType.PinSubCategoryMemberReference);
 				const bool bSupportsEventGraphs = (Blueprint && FBlueprintEditorUtils::DoesSupportEventGraphs(Blueprint));
 				const bool bAllowEvents = (GraphType == GT_Ubergraph);
 				if(Signature && bSupportsEventGraphs && bAllowEvents)
diff --git a/Engine/Source/Editor/BlueprintGraph/Private/K2Node_CreateDelegate.cpp b/Engine/Source/Editor/BlueprintGraph/Private/K2Node_CreateDelegate.cpp
--- a/Engine/Source/Editor/BlueprintGraph/Private/K2Node_CreateDelegate.cpp
+++ b/Engine/Source/Editor/BlueprintGraph/Private/K2Node_CreateDelegate.cpp
@@ -72,7 +72,7 @@ bool UK2Node_CreateDelegate::IsValid(FString* OutMsg, bool bDontUseSkeletalClass
 	{
 		const UEdGraphPin* OtherPin = DelegatePin->LinkedTo[PinIter];
 		const UFunction* OtherSignature = OtherPin ? 
-			OtherPin->PinType.PinSubCategoryMemberReference.ResolveSimpleMemberReference<UFunction>() : NULL;
+			FMemberReference::ResolveSimpleMemberReference<UFunction>(OtherPin->PinType.PinSubCategoryMemberReference) : NULL;
 		if(!OtherSignature || !Signature->IsSignatureCompatibleWith(OtherSignature))
 		{
 			if (OutMsg)
@@ -272,7 +272,7 @@ UFunction* UK2Node_CreateDelegate::GetDelegateSignature() const
 		if(UEdGraphPin* ResultPin = Pin->LinkedTo[0])
 		{
 			ensure(K2Schema->PC_Delegate == ResultPin->PinType.PinCategory);
-			return ResultPin->PinType.PinSubCategoryMemberReference.ResolveSimpleMemberReference<UFunction>();
+			return FMemberReference::ResolveSimpleMemberReference<UFunction>(ResultPin->PinType.PinSubCategoryMemberReference);
 		}
 	}
 	return NULL;
diff --git a/Engine/Source/Editor/BlueprintGraph/Private/K2Node_DelegateSet.cpp b/Engine/Source/Editor/BlueprintGraph/Private/K2Node_DelegateSet.cpp
--- a/Engine/Source/Editor/BlueprintGraph/Private/K2Node_DelegateSet.cpp
+++ b/Engine/Source/Editor/BlueprintGraph/Private/K2Node_DelegateSet.cpp
@@ -30,7 +30,7 @@ public:
 			// Create a term to store the locally created delegate that we'll use to add to the MC delegate
 			FBPTerminal* DelegateTerm = Context.CreateLocalTerminal();
 			DelegateTerm->Type.PinCategory = CompilerContext.GetSchema()->PC_Delegate;
-			DelegateTerm->Type.PinSubCategoryMemberReference.FillSimpleMemberReference<UFunction>(DelegateNode->GetDelegateSignature());
+			FMemberReference::FillSimpleMemberReference<UFunction>(DelegateNode->GetDelegateSignature(), DelegateTerm->Type.PinSubCategoryMemberReference);
 			DelegateTerm->Source = Node;
 			DelegateTerm->Name = Context.NetNameMap->MakeValidName(Node) + TEXT("_TempBindingDelegate");
 			LocalDelegateMap.Add(Node, DelegateTerm);
diff --git a/Engine/Source/Editor/BlueprintGraph/Private/K2Node_Event.cpp b/Engine/Source/Editor/BlueprintGraph/Private/K2Node_Event.cpp
--- a/Engine/Source/Editor/BlueprintGraph/Private/K2Node_Event.cpp
+++ b/Engine/Source/Editor/BlueprintGraph/Private/K2Node_Event.cpp
@@ -179,7 +179,7 @@ void UK2Node_Event::UpdateDelegatePin(bool bSilent)
 	UEdGraphPin* Pin = FindPinChecked(DelegateOutputName);
 	checkSlow(EGPD_Output == Pin->Direction);
 
-	const UObject* OldSignature = Pin->PinType.PinSubCategoryMemberReference.ResolveSimpleMemberReference<UFunction>();
+	const UObject* OldSignature = FMemberReference::ResolveSimpleMemberReference<UFunction>(Pin->PinType.PinSubCategoryMemberReference);
 	if (!OldSignature)
 	{
 		OldSignature = Pin->PinType.PinSubCategoryObject.Get();
@@ -198,7 +198,7 @@ void UK2Node_Event::UpdateDelegatePin(bool bSilent)
 	}
 
 	Pin->PinType.PinSubCategoryObject = NULL;
-	Pin->PinType.PinSubCategoryMemberReference.FillSimpleMemberReference<UFunction>(NewSignature);
+	FMemberReference::FillSimpleMemberReference<UFunction>(NewSignature, Pin->PinType.PinSubCategoryMemberReference);
 
 	if ((OldSignature != NewSignature) && !bSilent)
 	{
diff --git a/Engine/Source/Editor/BlueprintGraph/Private/K2Node_MCDelegate.cpp b/Engine/Source/Editor/BlueprintGraph/Private/K2Node_MCDelegate.cpp
--- a/Engine/Source/Editor/BlueprintGraph/Private/K2Node_MCDelegate.cpp
+++ b/Engine/Source/Editor/BlueprintGraph/Private/K2Node_MCDelegate.cpp
@@ -275,7 +275,7 @@ void UK2Node_AddDelegate::AllocateDefaultPins()
 
 	if(UEdGraphPin* DelegatePin = CreatePin(EGPD_Input, K2Schema->PC_Delegate, TEXT(""), NULL, false, true, FK2Node_BaseMCDelegateHelper::DelegatePinName, true))
 	{
-		DelegatePin->PinType.PinSubCategoryMemberReference.FillSimpleMemberReference<UFunction>(GetDelegateSignature());
+		FMemberReference::FillSimpleMemberReference<UFunction>(GetDelegateSignature(), DelegatePin->PinType.PinSubCategoryMemberReference);
 		DelegatePin->PinFriendlyName = NSLOCTEXT("K2Node", "PinFriendlyDelegatetName", "Event");
 	}
 }
@@ -343,7 +343,7 @@ void UK2Node_RemoveDelegate::AllocateDefaultPins()
 
 	if(UEdGraphPin* DelegatePin = CreatePin(EGPD_Input, K2Schema->PC_Delegate, TEXT(""), NULL, false, true, FK2Node_BaseMCDelegateHelper::DelegatePinName, true))
 	{
-		DelegatePin->PinType.PinSubCategoryMemberReference.FillSimpleMemberReference<UFunction>(GetDelegateSignature());
+		FMemberReference::FillSimpleMemberReference<UFunction>(GetDelegateSignature(), DelegatePin->PinType.PinSubCategoryMemberReference);
 		DelegatePin->PinFriendlyName = NSLOCTEXT("K2Node", "PinFriendlyDelegatetName", "Event");
 	}
 }
diff --git a/Engine/Source/Editor/KismetCompiler/Private/KismetCompilerMisc.cpp b/Engine/Source/Editor/KismetCompiler/Private/KismetCompilerMisc.cpp
--- a/Engine/Source/Editor/KismetCompiler/Private/KismetCompilerMisc.cpp
+++ b/Engine/Source/Editor/KismetCompiler/Private/KismetCompilerMisc.cpp
@@ -158,7 +158,7 @@ bool FKismetCompilerUtilities::IsTypeCompatibleWithProperty(UEdGraphPin* SourceP
 	}
 	else if (PinCategory == Schema->PC_Delegate)
 	{
-		const UFunction* SignatureFunction = Type.PinSubCategoryMemberReference.ResolveSimpleMemberReference<UFunction>();
+		const UFunction* SignatureFunction = FMemberReference::ResolveSimpleMemberReference<UFunction>(Type.PinSubCategoryMemberReference);
 		const UDelegateProperty* PropertyDelegate = Cast<const UDelegateProperty>(TestProperty);
 		bTypeMismatch = !(SignatureFunction 
 			&& PropertyDelegate 
@@ -851,7 +851,7 @@ UProperty* FKismetCompilerUtilities::CreatePropertyOnScope(UStruct* Scope, const
 	}
 	else if (Type.PinCategory == Schema->PC_Delegate)
 	{
-		if (UFunction* SignatureFunction = Type.PinSubCategoryMemberReference.ResolveSimpleMemberReference<UFunction>())
+		if (UFunction* SignatureFunction = FMemberReference::ResolveSimpleMemberReference<UFunction>(Type.PinSubCategoryMemberReference))
 		{
 			UDelegateProperty* NewPropertyDelegate = NewObject<UDelegateProperty>(PropertyScope, ValidatedPropertyName, ObjectFlags);
 			NewPropertyDelegate->SignatureFunction = SignatureFunction;
@@ -860,7 +860,7 @@ UProperty* FKismetCompilerUtilities::CreatePropertyOnScope(UStruct* Scope, const
 	}
 	else if (Type.PinCategory == Schema->PC_MCDelegate)
 	{
-		UFunction* const SignatureFunction = Type.PinSubCategoryMemberReference.ResolveSimpleMemberReference<UFunction>();
+		UFunction* const SignatureFunction = FMemberReference::ResolveSimpleMemberReference<UFunction>(Type.PinSubCategoryMemberReference);
 		UMulticastDelegateProperty* NewPropertyDelegate = NewObject<UMulticastDelegateProperty>(PropertyScope, ValidatedPropertyName, ObjectFlags);
 		NewPropertyDelegate->SignatureFunction = SignatureFunction;
 		NewProperty = NewPropertyDelegate;
diff --git a/Engine/Source/Runtime/CoreUObject/Public/UObject/Class.h b/Engine/Source/Runtime/CoreUObject/Public/UObject/Class.h
--- a/Engine/Source/Runtime/CoreUObject/Public/UObject/Class.h
+++ b/Engine/Source/Runtime/CoreUObject/Public/UObject/Class.h
@@ -1906,7 +1906,6 @@ public:
 	bool ReplaceNativeFunction(FName InName, Native InPointer, bool bAddToFunctionRemapTable);
 #endif
 
-#if WITH_EDITOR
 	/**
 	 * If there are potentially multiple versions of this class (e.g. blueprint generated classes), this function will return the authoritative version, which should be used for references
 	 *
@@ -1916,7 +1915,6 @@ public:
 	{
 		return this;
 	}
-#endif
 
 	/**
 	 * Add a native function to the internal native function table
diff --git a/Engine/Source/Runtime/Engine/Classes/Components/ActorComponent.h b/Engine/Source/Runtime/Engine/Classes/Components/ActorComponent.h
--- a/Engine/Source/Runtime/Engine/Classes/Components/ActorComponent.h
+++ b/Engine/Source/Runtime/Engine/Classes/Components/ActorComponent.h
@@ -146,9 +146,7 @@ public:
 	void DetermineUCSModifiedProperties();
 	void GetUCSModifiedProperties(TSet<const UProperty*>& ModifiedProperties) const;
 
-#if WITH_EDITOR
 	bool IsEditableWhenInherited() const;
-#endif
 
 	bool HasBeenCreated() const { return bHasBeenCreated; }
 	bool HasBeenInitialized() const { return bHasBeenInitialized; }
diff --git a/Engine/Source/Runtime/Engine/Classes/EdGraph/EdGraphPin.h b/Engine/Source/Runtime/Engine/Classes/EdGraph/EdGraphPin.h
--- a/Engine/Source/Runtime/Engine/Classes/EdGraph/EdGraphPin.h
+++ b/Engine/Source/Runtime/Engine/Classes/EdGraph/EdGraphPin.h
@@ -3,9 +3,38 @@
 #pragma once
 
 #include "EdGraph/EdGraphNode.h"
-#include "Engine/MemberReference.h"
 #include "EdGraphPin.generated.h"
 
+USTRUCT()
+struct FSimpleMemberReference
+{
+	GENERATED_USTRUCT_BODY()
+
+	/** Class that this member is defined in. */
+	UPROPERTY()
+	TSubclassOf<class UObject> MemberParentClass;
+
+	/** Name of the member */
+	UPROPERTY()
+	FName MemberName;
+
+	/** The Guid of the member */
+	UPROPERTY()
+	FGuid MemberGuid;
+
+	void Reset()
+	{
+		operator=(FSimpleMemberReference());
+	}
+
+	bool operator==(const FSimpleMemberReference& Other) const
+	{
+		return (MemberParentClass == Other.MemberParentClass)
+			&& (MemberName == Other.MemberName)
+			&& (MemberGuid == Other.MemberGuid);
+	}
+};
+
 FORCEINLINE FArchive& operator<<(FArchive& Ar, FSimpleMemberReference& Data)
 {
 	Ar << Data.MemberParentClass;
diff --git a/Engine/Source/Runtime/Engine/Classes/Engine/MemberReference.h b/Engine/Source/Runtime/Engine/Classes/Engine/MemberReference.h
--- a/Engine/Source/Runtime/Engine/Classes/Engine/MemberReference.h
+++ b/Engine/Source/Runtime/Engine/Classes/Engine/MemberReference.h
@@ -2,6 +2,7 @@
 
 #pragma once
 
+#include "Engine/Blueprint.h"
 #include "MemberReference.generated.h"
 
 /** Helper struct to allow us to redirect properties and functions through renames and additionally between classes if necessary */
@@ -95,6 +96,7 @@ public:
 		bSelfContext = bIsConsideredSelfContext;
 		bWasDeprecated = false;
 
+#if WITH_EDITOR
 		if (MemberParentClass != nullptr)
 		{
 			MemberParentClass = MemberParentClass->GetAuthoritativeClass();
@@ -105,16 +107,19 @@ public:
 		{
 			UBlueprint::GetGuidFromClassByFieldName<TFieldType>(InField->GetOwnerClass(), InField->GetFName(), MemberGuid);
 		}
+#endif
 	}
 
 	template<class TFieldType>
 	void SetFromField(const UField* InField, UClass* SelfScope)
 	{
 		FGuid FieldGuid;
+#if WITH_EDITOR
 		if (InField->GetOwnerClass())
 		{
 			UBlueprint::GetGuidFromClassByFieldName<TFieldType>(InField->GetOwnerClass(), InField->GetFName(), FieldGuid);
 		}
+#endif
 
 		SetGivenSelfScope(InField->GetFName(), FieldGuid, InField->GetOwnerClass(), SelfScope);
 	}
@@ -125,7 +130,9 @@ public:
 	{
 		if ((MemberParentClass != NULL) && (SelfScope != NULL))
 		{
+#if WITH_EDITOR
 			UBlueprint::GetGuidFromClassByFieldName<TFieldType>((MemberParentClass ? *MemberParentClass : SelfScope), MemberName, MemberGuid);
+#endif
 			SetGivenSelfScope(MemberName, MemberGuid, MemberParentClass, SelfScope);
 		}
 		else
@@ -183,12 +190,14 @@ public:
 		return !MemberScope.IsEmpty();
 	}
 private:
+#if WITH_EDITOR
 	/**
 	 * Refreshes a local variable reference name if it has changed
 	 *
 	 * @param InSelfScope		Scope to lookup the variable in, to see if it has changed
 	 */
 	ENGINE_API FName RefreshLocalVariableName(UClass* InSelfScope) const;
+#endif
 
 protected:
 	/** Only intended for backwards compat! */
@@ -244,6 +253,7 @@ public:
 			// Find in target scope
 			ReturnField = FindField<TFieldType>(MemberScopeStruct, MemberName);
 
+#if WITH_EDITOR
 			if(ReturnField == NULL)
 			{
 				// If the property was not found, refresh the local variable name and try again
@@ -253,11 +263,13 @@ public:
 					ReturnField = FindField<TFieldType>(MemberScopeStruct, MemberName);
 				}
 			}
+#endif
 		}
 		else
 		{
 			// Look for remapped member
 			UClass* TargetScope = bSelfContext ? SelfScope : (UClass*)MemberParentClass;
+#if WITH_EDITOR
 			if( TargetScope != NULL &&  !GIsSavingPackage )
 			{
 				ReturnField = Cast<TFieldType>(FindRemappedField(TargetScope, MemberName, true));
@@ -282,11 +294,14 @@ public:
 					}
 				}	
 			}
-			else if(TargetScope != NULL)
+			else
+#endif
+				if(TargetScope != NULL)
 			{
 				// Find in target scope
 				ReturnField = FindField<TFieldType>(TargetScope, MemberName);
 
+#if WITH_EDITOR
 				// If we have a GUID find the reference variable and make sure the name is up to date and find the field again
 				// For now only variable references will have valid GUIDs.  Will have to deal with finding other names subsequently
 				if (ReturnField == NULL && MemberGuid.IsValid())
@@ -298,6 +313,7 @@ public:
 						ReturnField = FindField<TFieldType>(TargetScope, MemberName);
 					}
 				}
+#endif
 			}
 			else
 			{
@@ -326,6 +342,7 @@ public:
 		return ResolveMember<TFieldType>(SelfScope->SkeletonGeneratedClass);
 	}
 
+#if WITH_EDITOR
 	/**
 	 * Searches the field redirect map for the specified named field in the scope, and returns the remapped field if found
 	 *
@@ -364,65 +381,30 @@ protected:
 	 * @return	Whether or not a remap was found in the specified scope
 	 */
 	static bool FindReplacementFieldName(UClass* Class, FName FieldName, FFieldRemapInfo& RemapInfo);
+#endif
 
-};
-
-USTRUCT()
-struct FSimpleMemberReference
-{
-	GENERATED_USTRUCT_BODY()
-
-	FSimpleMemberReference()
-	{
-	}
-
-	/** Class that this member is defined in. */
-	UPROPERTY()
-	TSubclassOf<class UObject> MemberParentClass;
-
-	/** Name of the member */
-	UPROPERTY()
-	FName MemberName;
-
-	/** The Guid of the member */
-	UPROPERTY()
-	FGuid MemberGuid;
-
-	void Reset()
-	{
-		operator=(FSimpleMemberReference());
-	}
-
-	bool operator==(const FSimpleMemberReference& Other) const
-	{
-		return (MemberParentClass == Other.MemberParentClass)
-			&& (MemberName == Other.MemberName)
-			&& (MemberGuid == Other.MemberGuid);
-	}
-
+public:
 	template<class TFieldType>
-	void FillSimpleMemberReference(const UField* InField)
+	static void FillSimpleMemberReference(const UField* InField, FSimpleMemberReference& OutReference)
 	{
+		OutReference.Reset();
+
 		if (InField)
 		{
 			FMemberReference TempMemberReference;
 			TempMemberReference.SetFromField<TFieldType>(InField, false);
 
-			MemberName = TempMemberReference.GetMemberName();
-			MemberParentClass = TempMemberReference.GetMemberParentClass();
-			MemberGuid = TempMemberReference.GetMemberGuid();
-		}
-		else
-		{
-			Reset();
+			OutReference.MemberName = TempMemberReference.MemberName;
+			OutReference.MemberParentClass = TempMemberReference.MemberParentClass;
+			OutReference.MemberGuid = TempMemberReference.MemberGuid;
 		}
 	}
 
 	template<class TFieldType>
-	TFieldType* ResolveSimpleMemberReference() const
+	static TFieldType* ResolveSimpleMemberReference(const FSimpleMemberReference& Reference)
 	{
 		FMemberReference TempMemberReference;
-		TempMemberReference.SetDirect(MemberName, MemberGuid, MemberParentClass, false);
+		TempMemberReference.SetDirect(Reference.MemberName, Reference.MemberGuid, Reference.MemberParentClass, false);
 		return TempMemberReference.ResolveMember<TFieldType>((UClass*)NULL);
 	}
-};
\ No newline at end of file
+};
diff --git a/Engine/Source/Runtime/Engine/Private/ActorComponent.cpp b/Engine/Source/Runtime/Engine/Private/ActorComponent.cpp
--- a/Engine/Source/Runtime/Engine/Private/ActorComponent.cpp
+++ b/Engine/Source/Runtime/Engine/Private/ActorComponent.cpp
@@ -1430,24 +1430,25 @@ void UActorComponent::OnRep_IsActive()
 	SetComponentTickEnabled(bIsActive);
 }
 
-#if WITH_EDITOR
 bool UActorComponent::IsEditableWhenInherited() const
 {
 	bool bCanEdit = bEditableWhenInherited;
 	if (bCanEdit)
 	{
+#if WITH_EDITOR
 		if (CreationMethod == EComponentCreationMethod::Native && !IsTemplate())
 		{
 			bCanEdit = FComponentEditorUtils::CanEditNativeComponent(this);
 		}
-		else if (CreationMethod == EComponentCreationMethod::UserConstructionScript)
+		else
+#endif
+			if (CreationMethod == EComponentCreationMethod::UserConstructionScript)
 		{
 			bCanEdit = false;
 		}
 	}
 	return bCanEdit;
 }
-#endif
 
 void UActorComponent::DetermineUCSModifiedProperties()
 {
@@ -1484,7 +1485,7 @@ void UActorComponent::DetermineUCSModifiedProperties()
 				if (!Property->Identical( DataPtr, DefaultValue))
 				{
 					UCSModifiedProperties.Add(FSimpleMemberReference());
-					UCSModifiedProperties.Last().FillSimpleMemberReference<UProperty>(Property);
+					FMemberReference::FillSimpleMemberReference<UProperty>(Property, UCSModifiedProperties.Last());
 					break;
 				}
 			}
@@ -1496,7 +1497,7 @@ void UActorComponent::GetUCSModifiedProperties(TSet<const UProperty*>& ModifiedP
 {
 	for (const FSimpleMemberReference& MemberReference : UCSModifiedProperties)
 	{
-		ModifiedProperties.Add(MemberReference.ResolveSimpleMemberReference<UProperty>());
+		ModifiedProperties.Add(FMemberReference::ResolveSimpleMemberReference<UProperty>(MemberReference));
 	}
 }
 
diff --git a/Engine/Source/Runtime/Engine/Private/MemberReference.cpp b/Engine/Source/Runtime/Engine/Private/MemberReference.cpp
--- a/Engine/Source/Runtime/Engine/Private/MemberReference.cpp
+++ b/Engine/Source/Runtime/Engine/Private/MemberReference.cpp
@@ -3,7 +3,10 @@
 #include "EnginePrivate.h"
 
 #include "Engine/MemberReference.h"
+
+#if WITH_EDITOR
 #include "Kismet2/BlueprintEditorUtils.h"
+#endif
 
 //////////////////////////////////////////////////////////////////////////
 // FMemberReference
@@ -11,7 +14,11 @@
 void FMemberReference::SetExternalMember(FName InMemberName, TSubclassOf<class UObject> InMemberParentClass)
 {
 	MemberName = InMemberName;
+#if WITH_EDITOR
 	MemberParentClass = (InMemberParentClass != nullptr) ? InMemberParentClass->GetAuthoritativeClass() : nullptr;
+#else
+	MemberParentClass = *InMemberParentClass;
+#endif
 	MemberScope.Empty();
 	bSelfContext = false;
 	bWasDeprecated = false;
@@ -84,6 +91,7 @@ void FMemberReference::InvalidateScope()
 	}
 }
 
+#if WITH_EDITOR
 FName FMemberReference::RefreshLocalVariableName(UClass* InSelfScope) const
 {
 	TArray<UBlueprint*> Blueprints;
@@ -317,4 +325,6 @@ UField* FMemberReference::FindRemappedField(UClass* InitialScope, FName InitialN
 	}
 
 	return NULL;
-}
\ No newline at end of file
+}
+
+#endif
\ No newline at end of file
