Engine: Moved actor components into Components folder

diff --git a/Engine/Source/Runtime/Engine/Private/Components/ActorComponent.cpp b/Engine/Source/Runtime/Engine/Private/Components/ActorComponent.cpp
new file mode 100644
--- /dev/null
+++ b/Engine/Source/Runtime/Engine/Private/Components/ActorComponent.cpp
@@ -0,0 +1,1505 @@
+// Copyright 1998-2015 Epic Games, Inc. All Rights Reserved.
+// ActorComponent.cpp: Actor component implementation.
+
+#include "EnginePrivate.h"
+#include "Engine/AssetUserData.h"
+#include "Engine/LevelStreamingPersistent.h"
+#include "Net/UnrealNetwork.h"
+#include "MessageLog.h"
+#include "UObjectToken.h"
+#include "MapErrors.h"
+#include "ComponentReregisterContext.h"
+#include "ComponentRecreateRenderStateContext.h"
+#include "Engine/SimpleConstructionScript.h"
+
+#define LOCTEXT_NAMESPACE "ActorComponent"
+
+DEFINE_LOG_CATEGORY(LogActorComponent);
+
+/** Enable to log out all render state create, destroy and updatetransform events */
+#define LOG_RENDER_STATE 0
+
+/** Static var indicating activity of reregister context */
+int32 FGlobalComponentReregisterContext::ActiveGlobalReregisterContextCount = 0;
+
+FGlobalComponentReregisterContext::FGlobalComponentReregisterContext()
+{
+	ActiveGlobalReregisterContextCount++;
+
+	// wait until resources are released
+	FlushRenderingCommands();
+
+	// Detach all actor components.
+	for(auto* Component : TObjectRange<UActorComponent>())
+	{
+		new(ComponentContexts) FComponentReregisterContext(Component);
+	}
+}
+
+FGlobalComponentReregisterContext::FGlobalComponentReregisterContext(const TArray<UClass*>& ExcludeComponents)
+{
+	ActiveGlobalReregisterContextCount++;
+
+	// wait until resources are released
+	FlushRenderingCommands();
+
+	// Detach only actor components that are not in the excluded list
+	for (auto* Component : TObjectRange<UActorComponent>())
+	{
+		bool bShouldReregister=true;
+		for (UClass* ExcludeClass : ExcludeComponents)
+		{
+			if( ExcludeClass &&
+				Component->IsA(ExcludeClass) )
+			{
+				bShouldReregister = false;
+				break;
+			}
+		}
+		if( bShouldReregister )
+		{
+			new(ComponentContexts) FComponentReregisterContext(Component);		
+		}
+	}
+}
+
+FGlobalComponentReregisterContext::FGlobalComponentReregisterContext(const TArray<AActor*>& InParentActors)
+{
+	ActiveGlobalReregisterContextCount++;
+
+	// wait until resources are released
+	FlushRenderingCommands();
+
+	// Detach only actor components that are children of the actors list provided
+	for (auto* Component : TObjectRange<UActorComponent>())
+	{
+		bool bShouldReregister=false;
+		UPrimitiveComponent* PrimitiveComponent = Cast<UPrimitiveComponent>(Component);
+		if (PrimitiveComponent && PrimitiveComponent->LODParentPrimitive.Get())
+		{
+			UPrimitiveComponent* LODParentPrimitive = PrimitiveComponent->LODParentPrimitive.Get();
+			AActor* ParentActor = Cast<AActor>(LODParentPrimitive->GetOuter());
+			if (ParentActor && InParentActors.Contains(ParentActor))
+			{
+				bShouldReregister = true;
+			}
+		}
+		if( bShouldReregister )
+		{
+			new(ComponentContexts) FComponentReregisterContext(Component);		
+		}
+	}
+}
+
+FGlobalComponentReregisterContext::~FGlobalComponentReregisterContext()
+{
+	check(ActiveGlobalReregisterContextCount > 0);
+	// We empty the array now, to ensure that the FComponentReregisterContext destructors are called while ActiveGlobalReregisterContextCount still indicates activity
+	ComponentContexts.Empty();
+	ActiveGlobalReregisterContextCount--;
+}
+
+FGlobalComponentRecreateRenderStateContext::FGlobalComponentRecreateRenderStateContext()
+{
+	// wait until resources are released
+	FlushRenderingCommands();
+
+	// recreate render state for all components.
+	for (auto* Component : TObjectRange<UActorComponent>())
+	{
+		new(ComponentContexts) FComponentRecreateRenderStateContext(Component);
+	}
+}
+
+FGlobalComponentRecreateRenderStateContext::~FGlobalComponentRecreateRenderStateContext()
+{
+	ComponentContexts.Empty();
+}
+
+
+UActorComponent::UActorComponent(const FObjectInitializer& ObjectInitializer /*= FObjectInitializer::Get()*/)
+	: Super(ObjectInitializer)
+{
+	PrimaryComponentTick.TickGroup = TG_DuringPhysics;
+	PrimaryComponentTick.bStartWithTickEnabled = true;
+	PrimaryComponentTick.bCanEverTick = false;
+
+	CreationMethod = EComponentCreationMethod::Native;
+
+	bAutoRegister = true;
+	bNetAddressable = false;
+	bEditableWhenInherited = true;
+}
+
+void UActorComponent::PostInitProperties()
+{
+	Super::PostInitProperties();
+
+	AActor* Owner = GetOwner();
+	if (Owner)
+	{
+		Owner->AddOwnedComponent(this);
+	}
+}
+
+void UActorComponent::PostLoad()
+{
+	Super::PostLoad();
+
+	if (GetLinkerUE4Version() < VER_UE4_ACTOR_COMPONENT_CREATION_METHOD)
+	{
+		if (bCreatedByConstructionScript_DEPRECATED)
+		{
+				CreationMethod = EComponentCreationMethod::SimpleConstructionScript;
+		}
+		else if (bInstanceComponent_DEPRECATED)
+		{
+			CreationMethod = EComponentCreationMethod::Instance;
+		}
+
+		if (CreationMethod == EComponentCreationMethod::SimpleConstructionScript)
+		{
+			UBlueprintGeneratedClass* Class = CastChecked<UBlueprintGeneratedClass>(GetOuter()->GetClass());
+			while (Class)
+			{
+				USimpleConstructionScript* SCS = Class->SimpleConstructionScript;
+				if (SCS != nullptr && SCS->FindSCSNode(GetFName()))
+				{
+					break;
+				}
+				else
+				{
+					Class = Cast<UBlueprintGeneratedClass>(Class->GetSuperClass());
+					if (Class == nullptr)
+					{
+						CreationMethod = EComponentCreationMethod::UserConstructionScript;
+					}
+				}
+			}
+		}
+	}
+
+	if (GetLinkerUE4Version() < VER_UE4_TRACK_UCS_MODIFIED_PROPERTIES)
+	{
+		DetermineUCSModifiedProperties();
+	}
+}
+
+bool UActorComponent::Rename( const TCHAR* InName, UObject* NewOuter, ERenameFlags Flags )
+{
+	bRoutedPostRename = false;
+
+	const FName OldName = GetFName();
+	const UObject* OldOuter = GetOuter();
+	
+	const bool bRenameSuccessful = Super::Rename(InName, NewOuter, Flags);
+	
+	const bool bMoved = (OldName != GetFName()) || (OldOuter != GetOuter());
+	if (!bRoutedPostRename && ((Flags & REN_Test) == 0) && bMoved)
+	{
+		UE_LOG(LogActorComponent, Fatal, TEXT("%s failed to route PostRename.  Please call Super::PostRename() in your <className>::PostRename() function. "), *GetFullName() );
+	}
+
+	return bRenameSuccessful;
+}
+
+void UActorComponent::PostRename(UObject* OldOuter, const FName OldName)
+{
+	Super::PostRename(OldOuter, OldName);
+
+	if (OldOuter != GetOuter())
+	{
+		AActor* Owner = GetOwner();
+		AActor* OldOwner = (OldOuter->IsA<AActor>() ? CastChecked<AActor>(OldOuter) : OldOuter->GetTypedOuter<AActor>());
+
+		if (Owner != OldOwner)
+		{
+			if (OldOwner)
+			{
+				OldOwner->RemoveOwnedComponent(this);
+			}
+			if (Owner)
+			{
+				Owner->AddOwnedComponent(this);
+			}
+		}
+	}
+
+	bRoutedPostRename = true;
+}
+
+bool UActorComponent::IsCreatedByConstructionScript() const
+{
+	return ((CreationMethod == EComponentCreationMethod::SimpleConstructionScript) || (CreationMethod == EComponentCreationMethod::UserConstructionScript));
+}
+
+#if WITH_EDITOR
+void UActorComponent::CheckForErrors()
+{
+	AActor* Owner = GetOwner();
+
+	if (Owner != NULL && GetClass()->HasAnyClassFlags(CLASS_Deprecated))
+	{
+		FFormatNamedArguments Arguments;
+		Arguments.Add(TEXT("ComponentName"), FText::FromString(GetName()));
+		Arguments.Add(TEXT("OwnerName"), FText::FromString(Owner->GetName()));
+		FMessageLog("MapCheck").Warning()
+			->AddToken(FUObjectToken::Create(Owner))
+			->AddToken(FTextToken::Create(FText::Format( LOCTEXT( "MapCheck_Message_DeprecatedClass", "{ComponentName}::{OwnerName} is obsolete and must be removed (Class is deprecated)" ), Arguments ) ) )
+			->AddToken(FMapErrorToken::Create(FMapErrors::DeprecatedClass));
+	}
+
+	if (Owner != NULL && GetClass()->HasAnyClassFlags(CLASS_Abstract))
+	{
+		FFormatNamedArguments Arguments;
+		Arguments.Add(TEXT("ComponentName"), FText::FromString(GetName()));
+		Arguments.Add(TEXT("OwnerName"), FText::FromString(Owner->GetName()));
+		FMessageLog("MapCheck").Warning()
+			->AddToken(FUObjectToken::Create(Owner))
+			->AddToken(FTextToken::Create(FText::Format( LOCTEXT( "MapCheck_Message_AbstractClass", "{ComponentName}::{OwnerName} is obsolete and must be removed (Class is abstract)" ), Arguments ) ))
+			->AddToken(FMapErrorToken::Create(FMapErrors::AbstractClass));
+	}
+}
+#endif
+
+bool UActorComponent::IsOwnerSelected() const
+{
+	AActor* Owner = GetOwner();
+	return Owner && Owner->IsSelected();
+}
+
+AActor* UActorComponent::GetOwner() const
+{
+	// walk up outer chain to find an Actor
+	UObject* Obj = GetOuter();
+	while(Obj != NULL)
+	{
+		AActor* Actor = Cast<AActor>(Obj);
+		if(Actor != NULL)
+		{
+			return Actor;
+		}
+
+		Obj = Obj->GetOuter();
+	}
+
+	return NULL;
+}
+
+UWorld* UActorComponent::GetWorld() const
+{
+	UWorld* ComponentWorld = World;
+	if (ComponentWorld == NULL)
+	{
+		// If we don't have a world yet, it may be because we haven't gotten registered yet, but we can try to look at our owner
+		AActor* Owner = GetOwner();
+		if (Owner && !Owner->HasAnyFlags(RF_ClassDefaultObject))
+		{
+			ComponentWorld = Owner->GetWorld();
+		}
+
+		if( ComponentWorld == NULL )
+		{
+			// As a fallback check the outer of this component for a world. In some cases components are spawned directly in the world
+			return Cast<UWorld>( GetOuter() );
+		}
+	}
+	return ComponentWorld;
+}
+
+bool UActorComponent::ComponentHasTag(FName Tag) const
+{
+	return (Tag != NAME_None) && ComponentTags.Contains(Tag);
+}
+
+
+ENetMode UActorComponent::GetNetMode() const
+{
+	AActor *Owner = GetOwner();
+	return Owner ? Owner->GetNetMode() : NM_Standalone;
+}
+
+FSceneInterface* UActorComponent::GetScene() const
+{
+	return (World ? World->Scene : NULL);
+}
+
+ULevel* UActorComponent::GetComponentLevel() const
+{
+	// For model components Level is outer object
+	ULevel* ComponentLevel = Cast<ULevel>( GetOuter() );
+	
+	if (ComponentLevel == NULL)
+	{
+		AActor* Actor = GetOwner();
+		if (Actor)
+		{
+			ComponentLevel = CastChecked<ULevel>( Actor->GetOuter() );
+		}
+	}
+	
+	return ComponentLevel;
+}
+
+bool UActorComponent::ComponentIsInLevel(const ULevel *TestLevel) const
+{
+	return (GetComponentLevel() == TestLevel);
+}
+
+
+bool UActorComponent::ComponentIsInPersistentLevel(bool bIncludeLevelStreamingPersistent) const
+{
+	ULevel* MyLevel = GetComponentLevel();
+	AActor* MyActor = GetOwner();
+	UWorld* MyWorld = GetWorld();
+
+	if(MyActor == NULL || MyLevel == NULL || MyWorld == NULL)
+	{
+		return false;
+	}
+
+	return ( (MyLevel == MyWorld->PersistentLevel) || ( bIncludeLevelStreamingPersistent && MyWorld->StreamingLevels.Num() > 0 &&
+														Cast<ULevelStreamingPersistent>(MyWorld->StreamingLevels[0]) != NULL &&
+														MyWorld->StreamingLevels[0]->GetLoadedLevel() == MyLevel ) );
+}
+
+FString UActorComponent::GetReadableName() const
+{
+	FString Result = GetNameSafe(GetOwner()) + TEXT(".") + GetName();
+	UObject const *Add = AdditionalStatObject();
+	if (Add)
+	{
+		Result += TEXT(" ");
+		Add->AppendName(Result);
+	}
+	return Result;
+}
+
+void UActorComponent::BeginDestroy()
+{
+	if (bHasBegunPlay)
+	{
+		EndPlay(EEndPlayReason::Destroyed);
+	}
+
+	// Ensure that we call UninitializeComponent before we destroy this component
+	if (bHasBeenInitialized)
+	{
+		UninitializeComponent();
+	}
+
+	ExecuteUnregisterEvents();
+
+	// Ensure that we call OnComponentDestroyed before we destroy this component
+	if (bHasBeenCreated)
+	{
+		OnComponentDestroyed();
+	}
+
+	World = NULL;
+
+	// Remove from the parent's OwnedComponents list
+	AActor* Owner = GetOwner();
+	if (Owner)
+	{
+		Owner->RemoveOwnedComponent(this);
+	}
+
+	Super::BeginDestroy();
+}
+
+
+bool UActorComponent::NeedsLoadForClient() const
+{
+	check(GetOuter());
+	return (GetOuter()->NeedsLoadForClient() && Super::NeedsLoadForClient());
+}
+
+
+bool UActorComponent::NeedsLoadForServer() const
+{
+	check(GetOuter());
+	return (GetOuter()->NeedsLoadForServer() && Super::NeedsLoadForServer());
+}
+
+int32 UActorComponent::GetFunctionCallspace( UFunction* Function, void* Parameters, FFrame* Stack )
+{
+	AActor* Owner = GetOwner();
+	if (Owner == NULL)
+	{
+		return FunctionCallspace::Local;
+	}
+	return Owner->GetFunctionCallspace(Function, Parameters, Stack);
+}
+
+bool UActorComponent::CallRemoteFunction( UFunction* Function, void* Parameters, FOutParmRec* OutParms, FFrame* Stack )
+{
+	AActor* Owner = GetOwner();
+	if (Owner == NULL)
+	{
+		return false;
+	}
+	
+	UNetDriver* NetDriver = Owner->GetNetDriver();
+	if (NetDriver)
+	{
+		NetDriver->ProcessRemoteFunction(Owner, Function, Parameters, OutParms, Stack, this);
+		return true;
+	}
+
+	return false;
+}
+
+/** FComponentReregisterContexts for components which have had PreEditChange called but not PostEditChange. */
+static TMap<UActorComponent*,FComponentReregisterContext*> EditReregisterContexts;
+
+#if WITH_EDITOR
+bool UActorComponent::Modify( bool bAlwaysMarkDirty/*=true*/ )
+{
+	// If this is a construction script component we don't store them in the transaction buffer.  Instead, mark
+	// the Actor as modified so that we store of the transaction annotation that has the component properties stashed
+	if (IsCreatedByConstructionScript() && GetOwner())
+	{
+		return GetOwner()->Modify(bAlwaysMarkDirty);
+	}
+
+	return Super::Modify(bAlwaysMarkDirty);
+}
+
+void UActorComponent::PreEditChange(UProperty* PropertyThatWillChange)
+{
+	Super::PreEditChange(PropertyThatWillChange);
+
+	if(IsRegistered())
+	{
+		// The component or its outer could be pending kill when calling PreEditChange when applying a transaction.
+		// Don't do do a full recreate in this situation, and instead simply detach.
+		if( !IsPendingKill() )
+		{
+			check(!EditReregisterContexts.Find(this));
+			EditReregisterContexts.Add(this,new FComponentReregisterContext(this));
+		}
+		else
+		{
+			ExecuteUnregisterEvents();
+			World = NULL;
+		}
+	}
+
+	// Flush rendering commands to ensure the rendering thread processes the component detachment before it is modified.
+	FlushRenderingCommands();
+}
+
+void UActorComponent::PostEditUndo()
+{
+	// Objects marked pending kill don't call PostEditChange() from UObject::PostEditUndo(),
+	// so they can leave an EditReregisterContexts entry around if they are deleted by an undo action.
+	if( IsPendingKill() )
+	{
+		// The reregister context won't bother attaching components that are 'pending kill'. 
+		FComponentReregisterContext* ReregisterContext = EditReregisterContexts.FindRef(this);
+		if(ReregisterContext)
+		{
+			delete ReregisterContext;
+			EditReregisterContexts.Remove(this);
+		}
+	}
+	else
+	{
+		//Let the component be properly registered, after it was restored.
+		AActor* Owner = GetOwner();
+		if (Owner)
+		{
+			Owner->AddOwnedComponent(this);
+		}
+	}
+	Super::PostEditUndo();
+}
+
+void UActorComponent::ConsolidatedPostEditChange()
+{
+	FComponentReregisterContext* ReregisterContext = EditReregisterContexts.FindRef(this);
+	if(ReregisterContext)
+	{
+		delete ReregisterContext;
+		EditReregisterContexts.Remove(this);
+
+		AActor* Owner = GetOwner();
+		if (Owner && !Owner->IsTemplate())
+		{
+			Owner->RerunConstructionScripts();
+		}
+	}
+
+	// The component or its outer could be pending kill when calling PostEditChange when applying a transaction.
+	// Don't do do a full recreate in this situation, and instead simply detach.
+	if( IsPendingKill() )
+	{
+		// @todo UE4 james should this call UnregisterComponent instead to remove itself from the RegisteredComponents array on the owner?
+		ExecuteUnregisterEvents();
+		World = NULL;
+	}
+}
+
+void UActorComponent::PostEditChangeProperty(FPropertyChangedEvent& PropertyChangedEvent)
+{
+	ConsolidatedPostEditChange();
+
+	Super::PostEditChangeProperty(PropertyChangedEvent);
+}
+
+void UActorComponent::PostEditChangeChainProperty(FPropertyChangedChainEvent& PropertyChangedEvent)
+{
+	ConsolidatedPostEditChange();
+
+	Super::PostEditChangeChainProperty(PropertyChangedEvent);
+}
+
+
+#endif // WITH_EDITOR
+
+void UActorComponent::OnRegister()
+{
+	AActor* Owner = GetOwner();
+	checkf(!HasAnyFlags(RF_Unreachable), TEXT("%s"), *GetDetailedInfo());
+	checkf(!GetOuter()->IsTemplate(), TEXT("'%s' (%s)"), *GetOuter()->GetFullName(), *GetDetailedInfo());
+	checkf(!IsTemplate(), TEXT("'%s' (%s)"), *GetOuter()->GetFullName(), *GetDetailedInfo() );
+	checkf(World, TEXT("OnRegister: %s to %s"), *GetDetailedInfo(), Owner ? *Owner->GetFullName() : TEXT("*** No Owner ***") );
+	checkf(!bRegistered, TEXT("OnRegister: %s to %s"), *GetDetailedInfo(), Owner ? *Owner->GetFullName() : TEXT("*** No Owner ***") );
+	checkf(!IsPendingKill(), TEXT("OnRegister: %s to %s"), *GetDetailedInfo(), Owner ? *Owner->GetFullName() : TEXT("*** No Owner ***") );
+
+	bRegistered = true;
+
+	UpdateComponentToWorld();
+
+	if (bAutoActivate)
+	{
+		Activate(true);
+	}
+}
+
+void UActorComponent::OnUnregister()
+{
+	check(bRegistered);
+	bRegistered = false;
+}
+
+void UActorComponent::InitializeComponent()
+{
+	check(bRegistered);
+	check(!bHasBeenInitialized);
+
+	bHasBeenInitialized = true;
+}
+
+void UActorComponent::UninitializeComponent()
+{
+	check(bHasBeenInitialized);
+
+	bHasBeenInitialized = false;
+}
+
+void UActorComponent::BeginPlay()
+{
+	check(bRegistered);
+	check(!bHasBegunPlay);
+
+	ReceiveBeginPlay();
+
+	bHasBegunPlay = true;
+}
+
+void UActorComponent::EndPlay(const EEndPlayReason::Type EndPlayReason)
+{
+	check(bHasBegunPlay);
+
+	// If we're already pending kill blueprints don't get to be notified
+	if (!HasAnyFlags(RF_BeginDestroyed))
+	{
+		ReceiveEndPlay(EndPlayReason);
+	}
+
+	bHasBegunPlay = false;
+}
+
+FActorComponentInstanceData* UActorComponent::GetComponentInstanceData() const
+{
+	FActorComponentInstanceData* InstanceData = new FActorComponentInstanceData(this);
+
+	if (!InstanceData->ContainsSavedProperties())
+	{
+		delete InstanceData;
+		InstanceData = nullptr;
+	}
+
+	return InstanceData;
+}
+
+FName UActorComponent::GetComponentInstanceDataType() const
+{
+	static const FName ActorComponentInstanceDataTypeName(TEXT("ActorComponentInstanceData"));
+	return ActorComponentInstanceDataTypeName;
+}
+
+void FActorComponentTickFunction::ExecuteTick(float DeltaTime, enum ELevelTick TickType, ENamedThreads::Type CurrentThread, const FGraphEventRef& MyCompletionGraphEvent)
+{
+	if (Target && !Target->HasAnyFlags(RF_PendingKill | RF_Unreachable))
+	{
+		FScopeCycleCounterUObject ComponentScope(Target);
+		FScopeCycleCounterUObject AdditionalScope(Target->AdditionalStatObject());
+	    checkSlow(Target && (!EnableParent || Target->IsPendingKill() || ((FActorTickFunction*)EnableParent)->Target == Target->GetOwner())); // components that get renamed into other outers will have this wrong and hence will not necessarily tick after their actor, or use their actor as an enable parent
+	    Target->ConditionalTickComponent(DeltaTime, TickType, *this);	
+	}
+}
+
+FString FActorComponentTickFunction::DiagnosticMessage()
+{
+	return Target->GetFullName() + TEXT("[TickComponent]");
+}
+
+bool UActorComponent::SetupActorComponentTickFunction(struct FTickFunction* TickFunction)
+{
+	AActor* Owner = GetOwner();
+	if(TickFunction->bCanEverTick && !IsTemplate() && (!Owner || !Owner->IsTemplate()))
+	{
+		ULevel* ComponentLevel = (Owner ? Owner->GetLevel() : GetWorld()->PersistentLevel);
+		TickFunction->SetTickFunctionEnable(TickFunction->bStartWithTickEnabled);
+		TickFunction->RegisterTickFunction(ComponentLevel);
+		return true;
+	}
+	return false;
+}
+
+void UActorComponent::SetComponentTickEnabled(bool bEnabled)
+{
+	if (!IsTemplate() && PrimaryComponentTick.bCanEverTick)
+	{
+		PrimaryComponentTick.SetTickFunctionEnable(bEnabled);
+	}
+}
+
+void UActorComponent::SetComponentTickEnabledAsync(bool bEnabled)
+{
+	if (!IsTemplate() && PrimaryComponentTick.bCanEverTick)
+	{
+		DECLARE_CYCLE_STAT(TEXT("FSimpleDelegateGraphTask.SetComponentTickEnabledAsync"),
+			STAT_FSimpleDelegateGraphTask_SetComponentTickEnabledAsync,
+			STATGROUP_TaskGraphTasks);
+
+		FSimpleDelegateGraphTask::CreateAndDispatchWhenReady(
+			FSimpleDelegateGraphTask::FDelegate::CreateUObject(this, &UActorComponent::SetComponentTickEnabled, bEnabled),
+			GET_STATID(STAT_FSimpleDelegateGraphTask_SetComponentTickEnabledAsync), NULL, ENamedThreads::GameThread
+		);
+	}
+}
+
+bool UActorComponent::IsComponentTickEnabled() const
+{
+	return PrimaryComponentTick.IsTickFunctionEnabled();
+}
+
+static UActorComponent* GTestRegisterComponentTickFunctions = NULL;
+
+void UActorComponent::RegisterComponentTickFunctions(bool bRegister)
+{
+	if(bRegister)
+	{
+		if (SetupActorComponentTickFunction(&PrimaryComponentTick))
+		{
+			PrimaryComponentTick.Target = this;
+		}
+	}
+	else
+	{
+		if(PrimaryComponentTick.IsTickFunctionRegistered())
+		{
+			PrimaryComponentTick.UnRegisterTickFunction();
+		}
+	}
+
+	GTestRegisterComponentTickFunctions = this; // we will verify the super call chain is intact. Don't not copy paste this to a derived class!
+}
+
+void UActorComponent::RegisterAllComponentTickFunctions(bool bRegister)
+{
+	check(GTestRegisterComponentTickFunctions == NULL);
+	// Components don't have tick functions until they are registered with the world
+	if (bRegistered)
+	{
+		RegisterComponentTickFunctions(bRegister);
+		checkf(GTestRegisterComponentTickFunctions == this, TEXT("Failed to route component RegisterTickFunctions (%s)"), *GetFullName());
+		GTestRegisterComponentTickFunctions = NULL;
+	}
+}
+
+void UActorComponent::TickComponent(float DeltaTime, enum ELevelTick TickType, FActorComponentTickFunction *ThisTickFunction)
+{
+	check(bRegistered);
+
+	ReceiveTick(DeltaTime);
+}
+
+void UActorComponent::RegisterComponentWithWorld(UWorld* InWorld)
+{
+	checkf(!HasAnyFlags(RF_Unreachable), TEXT("%s"), *GetFullName());
+
+	if(IsPendingKill())
+	{
+		UE_LOG(LogActorComponent, Log, TEXT("RegisterComponentWithWorld: (%s) Trying to register component with IsPendingKill() == true. Aborting."), *GetPathName());
+		return;
+	}
+
+	// If the component was already registered, do nothing
+	if(IsRegistered())
+	{
+		UE_LOG(LogActorComponent, Log, TEXT("RegisterComponentWithWorld: (%s) Already registered. Aborting."), *GetPathName());
+		return;
+	}
+
+	if(InWorld == NULL)
+	{
+		//UE_LOG(LogActorComponent, Log, TEXT("RegisterComponentWithWorld: (%s) NULL InWorld specified. Aborting."), *GetPathName());
+		return;
+	}
+
+	// If not registered, should not have a scene
+	checkf(World == NULL, TEXT("%s"), *GetFullName());
+
+	// Follow outer chain to see if we can find an Actor
+	AActor* Owner = GetOwner();
+
+	checkSlow(Owner == NULL || Owner->OwnsComponent(this));
+
+	if ((Owner != NULL) && Owner->GetClass()->HasAnyClassFlags(CLASS_NewerVersionExists))
+	{
+		UE_LOG(LogActorComponent, Log, TEXT("RegisterComponentWithWorld: Owner belongs to a DEADCLASS"));
+		return;
+	}
+
+#if !(UE_BUILD_SHIPPING || UE_BUILD_TEST)
+	// Can only register with an Actor if we are created within one
+	if(Owner != NULL)
+	{
+		checkf(!Owner->HasAnyFlags(RF_Unreachable), TEXT("%s"), *GetFullName());
+		// can happen with undo because the owner will be restored "next"
+		//checkf(!Owner->IsPendingKill(), TEXT("%s"), *GetFullName());
+
+		if(InWorld != Owner->GetWorld())
+		{
+			// The only time you should specify a scene that is not GetOwner()->GetWorld() is when you don't have an Actor
+			UE_LOG(LogActorComponent, Log, TEXT("RegisterComponentWithWorld: (%s) Specifying a world, but an Owner Actor found, and InWorld is not GetOwner()->GetWorld()"), *GetPathName());
+		}
+	}
+#endif // !(UE_BUILD_SHIPPING || UE_BUILD_TEST)
+
+	World = InWorld;
+
+	ExecuteRegisterEvents();
+	RegisterAllComponentTickFunctions(true);
+
+	if (Owner == nullptr || Owner->IsActorInitialized())
+	{
+		if (!bHasBeenInitialized && bWantsInitializeComponent)
+		{
+			InitializeComponent();
+		}
+	}
+
+	if (Owner && Owner->HasActorBegunPlay())
+	{
+		if (bWantsBeginPlay)
+		{
+			BeginPlay();
+		}
+	}
+
+	// If this is a blueprint created component and it has component children they can miss getting registered in some scenarios
+	if (IsCreatedByConstructionScript())
+	{
+		TArray<UObject*> Children;
+		GetObjectsWithOuter(this, Children, true, RF_PendingKill);
+
+		for (UObject* Child : Children)
+		{
+			UActorComponent* ChildComponent = Cast<UActorComponent>(Child);
+			if (ChildComponent && !ChildComponent->IsRegistered())
+			{
+				ChildComponent->RegisterComponentWithWorld(InWorld);
+			}
+		}
+
+	}
+}
+
+void UActorComponent::RegisterComponent()
+{
+	if (ensure(GetOwner() && GetOwner()->GetWorld()))
+	{
+		RegisterComponentWithWorld(GetOwner()->GetWorld());
+	}
+}
+
+void UActorComponent::UnregisterComponent()
+{
+	// Do nothing if not registered
+	if(!IsRegistered())
+	{
+		UE_LOG(LogActorComponent, Log, TEXT("UnregisterComponent: (%s) Not registered. Aborting."), *GetPathName());
+		return;
+	}
+
+	// If registered, should have a world
+	checkf(World != NULL, TEXT("%s"), *GetFullName());
+
+	// Notify the texture streaming system
+	const UPrimitiveComponent* Primitive = Cast<const UPrimitiveComponent>(this);
+	if ( Primitive )
+	{
+		IStreamingManager::Get().NotifyPrimitiveDetached( Primitive );
+	}
+
+	AActor* Owner = GetOwner(); // Get Owner while we are still registered
+
+	RegisterAllComponentTickFunctions(false);
+	ExecuteUnregisterEvents();
+
+	World = NULL;
+}
+
+void UActorComponent::DestroyComponent(bool bPromoteChildren/*= false*/)
+{
+	if (bHasBegunPlay)
+	{
+		EndPlay(EEndPlayReason::Destroyed);
+	}
+
+	// Ensure that we call UninitializeComponent before we destroy this component
+	if (bHasBeenInitialized)
+	{
+		UninitializeComponent();
+	}
+
+	// Unregister if registered
+	if(IsRegistered())
+	{
+		UnregisterComponent();
+	}
+
+	// Then remove from Components array, if we have an Actor
+	AActor* Owner = GetOwner();
+	if(Owner != NULL)
+	{
+		if (IsCreatedByConstructionScript())
+		{
+			Owner->BlueprintCreatedComponents.Remove(this);
+		}
+		else
+		{
+			Owner->RemoveInstanceComponent(this);
+		}
+		Owner->RemoveOwnedComponent(this);
+		if (Owner->GetRootComponent() == this)
+		{
+			Owner->SetRootComponent(NULL);
+		}
+	}
+
+	// Tell the component it is being destroyed
+	OnComponentDestroyed();
+
+	// Finally mark pending kill, to NULL out any other refs
+	MarkPendingKill();
+}
+
+void UActorComponent::OnComponentCreated()
+{
+	ensure(!bHasBeenCreated);
+	bHasBeenCreated = true;
+}
+
+void UActorComponent::OnComponentDestroyed()
+{
+	// @TODO: Would be nice to ensure(bHasBeenCreated), but there are still many places where components are created without calling OnComponentCreated
+	bHasBeenCreated = false;
+}
+
+void UActorComponent::K2_DestroyComponent(UObject* Object)
+{
+	AActor* Owner = GetOwner();
+	if (Owner == NULL || Owner == Object)
+	{
+		DestroyComponent();
+	}
+	else
+	{
+		// TODO: Put in Message Log
+		UE_LOG(LogActorComponent, Error, TEXT("May not destroy component %s owned by %s."), *GetFullName(), *Owner->GetFullName());
+	}
+}
+
+void UActorComponent::CreateRenderState_Concurrent()
+{
+	check(IsRegistered());
+	check(World->Scene);
+	check(!bRenderStateCreated);
+	bRenderStateCreated = true;
+
+	bRenderStateDirty = false;
+	bRenderTransformDirty = false;
+	bRenderDynamicDataDirty = false;
+
+#if LOG_RENDER_STATE
+	UE_LOG(LogActorComponent, Log, TEXT("CreateRenderState_Concurrent: %s"), *GetPathName());
+#endif
+}
+
+void UActorComponent::SendRenderTransform_Concurrent()
+{
+	check(bRenderStateCreated);
+	bRenderTransformDirty = false;
+
+#if LOG_RENDER_STATE
+	UE_LOG(LogActorComponent, Log, TEXT("SendRenderTransform_Concurrent: %s"), *GetPathName());
+#endif
+}
+
+void UActorComponent::SendRenderDynamicData_Concurrent()
+{
+	check(bRenderStateCreated);
+	bRenderDynamicDataDirty = false;
+
+#if LOG_RENDER_STATE
+	UE_LOG(LogActorComponent, Log, TEXT("SendRenderDynamicData_Concurrent: %s"), *GetPathName());
+#endif
+}
+
+void UActorComponent::DestroyRenderState_Concurrent()
+{
+	check(bRenderStateCreated);
+	bRenderStateCreated = false;
+
+#if LOG_RENDER_STATE
+	UE_LOG(LogActorComponent, Log, TEXT("DestroyRenderState_Concurrent: %s"), *GetPathName());
+#endif
+}
+
+void UActorComponent::CreatePhysicsState()
+{
+	check(IsRegistered());
+	check(ShouldCreatePhysicsState());
+	check(World->GetPhysicsScene());
+	check(!bPhysicsStateCreated);
+	bPhysicsStateCreated = true;
+}
+
+void UActorComponent::DestroyPhysicsState()
+{
+	ensure(bPhysicsStateCreated);
+	bPhysicsStateCreated = false;
+}
+
+
+void UActorComponent::ExecuteRegisterEvents()
+{
+	if(!bRegistered)
+	{
+		OnRegister();
+		checkf(bRegistered, TEXT("Failed to route OnRegister (%s)"), *GetFullName());
+	}
+
+	if(FApp::CanEverRender() && !bRenderStateCreated && World->Scene)
+	{
+		CreateRenderState_Concurrent();
+		checkf(bRenderStateCreated, TEXT("Failed to route CreateRenderState_Concurrent (%s)"), *GetFullName());
+	}
+
+	if(!bPhysicsStateCreated && World->GetPhysicsScene() && ShouldCreatePhysicsState())
+	{
+		CreatePhysicsState();
+		checkf(bPhysicsStateCreated, TEXT("Failed to route CreatePhysicsState (%s)"), *GetFullName());
+	}
+}
+
+
+void UActorComponent::ExecuteUnregisterEvents()
+{
+	if(bPhysicsStateCreated)
+	{
+		check(bRegistered); // should not have physics state unless we are registered
+		DestroyPhysicsState();
+		checkf(!bPhysicsStateCreated, TEXT("Failed to route DestroyPhysicsState (%s)"), *GetFullName());
+		checkf(!HasValidPhysicsState(), TEXT("Failed to destroy physics state (%s)"), *GetFullName());
+	}
+
+	if(bRenderStateCreated)
+	{
+		check(bRegistered);
+		DestroyRenderState_Concurrent();
+		checkf(!bRenderStateCreated, TEXT("Failed to route DestroyRenderState_Concurrent (%s)"), *GetFullName());
+	}
+
+	if(bRegistered)
+	{
+		OnUnregister();
+		checkf(!bRegistered, TEXT("Failed to route OnUnregister (%s)"), *GetFullName());
+	}
+}
+
+void UActorComponent::ReregisterComponent()
+{
+	if(!IsRegistered())
+	{
+		UE_LOG(LogActorComponent, Log, TEXT("ReregisterComponent: (%s) Not currently registered. Aborting."), *GetPathName());
+		return;
+	}
+
+	FComponentReregisterContext(this);
+}
+
+void UActorComponent::RecreateRenderState_Concurrent()
+{
+	if(bRenderStateCreated)
+	{
+		check(IsRegistered()); // Should never have render state unless registered
+		DestroyRenderState_Concurrent();
+		checkf(!bRenderStateCreated, TEXT("Failed to route DestroyRenderState_Concurrent (%s)"), *GetFullName());
+	}
+
+	if(IsRegistered() && World->Scene)
+	{
+		CreateRenderState_Concurrent();
+		checkf(bRenderStateCreated, TEXT("Failed to route CreateRenderState_Concurrent (%s)"), *GetFullName());
+	}
+}
+
+void UActorComponent::RecreatePhysicsState()
+{
+	if(bPhysicsStateCreated)
+	{
+		check(IsRegistered()); // Should never have physics state unless registered
+		DestroyPhysicsState();
+		checkf(!bPhysicsStateCreated, TEXT("Failed to route DestroyPhysicsState (%s)"), *GetFullName());
+		checkf(!HasValidPhysicsState(), TEXT("Failed to destroy physics state (%s)"), *GetFullName());
+	}
+
+	if (IsRegistered() && World->GetPhysicsScene() && ShouldCreatePhysicsState())
+	{
+		CreatePhysicsState();
+		checkf(bPhysicsStateCreated, TEXT("Failed to route CreatePhysicsState (%s)"), *GetFullName());
+	}
+}
+
+void UActorComponent::ConditionalTickComponent(float DeltaTime, enum ELevelTick TickType, FActorComponentTickFunction &ThisTickFunction)
+{
+	if(bRegistered && !IsPendingKill())
+	{
+		//@optimization, I imagine this is all unnecessary in a shipping game with no editor
+		AActor* Owner = GetOwner();
+		if (TickType != LEVELTICK_ViewportsOnly || 
+			(bTickInEditor && TickType == LEVELTICK_ViewportsOnly) ||
+			(Owner && Owner->ShouldTickIfViewportsOnly())
+			)
+		{
+			const float TimeDilation = (Owner ? Owner->CustomTimeDilation : 1.f);
+			TickComponent(DeltaTime * TimeDilation, TickType, &ThisTickFunction);
+		}
+	}
+}
+
+void UActorComponent::SetTickGroup(ETickingGroup NewTickGroup)
+{
+	PrimaryComponentTick.TickGroup = NewTickGroup;
+}
+
+
+void UActorComponent::AddTickPrerequisiteActor(AActor* PrerequisiteActor)
+{
+	if (PrimaryComponentTick.bCanEverTick && PrerequisiteActor && PrerequisiteActor->PrimaryActorTick.bCanEverTick)
+	{
+		PrimaryComponentTick.AddPrerequisite(PrerequisiteActor, PrerequisiteActor->PrimaryActorTick);
+	}
+}
+
+void UActorComponent::AddTickPrerequisiteComponent(UActorComponent* PrerequisiteComponent)
+{
+	if (PrimaryComponentTick.bCanEverTick && PrerequisiteComponent && PrerequisiteComponent->PrimaryComponentTick.bCanEverTick)
+	{
+		PrimaryComponentTick.AddPrerequisite(PrerequisiteComponent, PrerequisiteComponent->PrimaryComponentTick);
+	}
+}
+
+void UActorComponent::RemoveTickPrerequisiteActor(AActor* PrerequisiteActor)
+{
+	if (PrerequisiteActor)
+	{
+		PrimaryComponentTick.RemovePrerequisite(PrerequisiteActor, PrerequisiteActor->PrimaryActorTick);
+	}
+}
+
+void UActorComponent::RemoveTickPrerequisiteComponent(UActorComponent* PrerequisiteComponent)
+{
+	if (PrerequisiteComponent)
+	{
+		PrimaryComponentTick.RemovePrerequisite(PrerequisiteComponent, PrerequisiteComponent->PrimaryComponentTick);
+	}
+}
+
+void UActorComponent::DoDeferredRenderUpdates_Concurrent()
+{
+	checkf(!HasAnyFlags(RF_Unreachable), TEXT("%s"), *GetFullName());
+	checkf(!IsTemplate(), TEXT("%s"), *GetFullName());
+	checkf(!IsPendingKill(), TEXT("%s"), *GetFullName());
+
+	if(!IsRegistered())
+	{
+		UE_LOG(LogActorComponent, Log, TEXT("UpdateComponent: (%s) Not registered, Aborting."), *GetPathName());
+		return;
+	}
+
+	if(bRenderStateDirty)
+	{
+		SCOPE_CYCLE_COUNTER(STAT_PostTickComponentRecreate);
+		RecreateRenderState_Concurrent();
+		checkf(!bRenderStateDirty, TEXT("Failed to route CreateRenderState_Concurrent (%s)"), *GetFullName());
+	}
+	else
+	{
+		SCOPE_CYCLE_COUNTER(STAT_PostTickComponentLW);
+		if(bRenderTransformDirty)
+		{
+			// Update the component's transform if the actor has been moved since it was last updated.
+			SendRenderTransform_Concurrent();
+		}
+
+		if(bRenderDynamicDataDirty)
+		{
+			SendRenderDynamicData_Concurrent();
+		}
+	}
+}
+
+
+void UActorComponent::MarkRenderStateDirty()
+{
+	// If registered and has a render state to make as dirty
+	if((!bRenderStateDirty || !GetWorld()) && IsRegistered() && bRenderStateCreated)
+	{
+		// Flag as dirty
+		bRenderStateDirty = true;
+		MarkForNeededEndOfFrameRecreate();
+	}
+}
+
+
+void UActorComponent::MarkRenderTransformDirty()
+{
+	if (IsRegistered() && bRenderStateCreated)
+	{
+		bRenderTransformDirty = true;
+		MarkForNeededEndOfFrameUpdate();
+	}
+}
+
+
+void UActorComponent::MarkRenderDynamicDataDirty()
+{
+	// If registered and has a render state to make as dirty
+	if(IsRegistered() && bRenderStateCreated)
+	{
+		// Flag as dirty
+		bRenderDynamicDataDirty = true;
+		MarkForNeededEndOfFrameUpdate();
+	}
+}
+
+void UActorComponent::MarkForNeededEndOfFrameUpdate()
+{
+	if (bNeverNeedsRenderUpdate)
+	{
+		return;
+	}
+
+	if(GetWorld())
+	{
+		GetWorld()->MarkActorComponentForNeededEndOfFrameUpdate(this, RequiresGameThreadEndOfFrameUpdates());
+	}
+	else if (!HasAnyFlags(RF_Unreachable))
+	{
+		// we don't have a world, do it right now.
+		DoDeferredRenderUpdates_Concurrent();
+	}
+}
+
+void UActorComponent::MarkForNeededEndOfFrameRecreate()
+{
+	if (bNeverNeedsRenderUpdate)
+	{
+		return;
+	}
+
+	if (GetWorld())
+	{
+		// by convention, recreates are always done on the gamethread
+		GetWorld()->MarkActorComponentForNeededEndOfFrameUpdate(this, true);
+	}
+	else if (!HasAnyFlags(RF_Unreachable))
+	{
+		// we don't have a world, do it right now.
+		DoDeferredRenderUpdates_Concurrent();
+	}
+}
+
+bool UActorComponent::RequiresGameThreadEndOfFrameUpdates() const
+{
+	return false;
+}
+
+void UActorComponent::Activate(bool bReset)
+{
+	if (bReset || ShouldActivate()==true)
+	{
+		SetComponentTickEnabled(true);
+		bIsActive = true;
+	}
+}
+
+void UActorComponent::Deactivate()
+{
+	if (ShouldActivate()==false)
+	{
+		SetComponentTickEnabled(false);
+		bIsActive = false;
+	}
+}
+
+bool UActorComponent::ShouldActivate() const
+{
+	// if not active, should activate
+	return !bIsActive;
+}
+
+void UActorComponent::SetActive(bool bNewActive, bool bReset)
+{
+	// if it wants to activate
+	if (bNewActive)
+	{
+		// make sure to check if it should activate
+		Activate(bReset);	
+	}
+	// otherwise, make sure it shouldn't activate
+	else 
+	{
+		Deactivate();
+	}
+}
+
+void UActorComponent::ToggleActive()
+{
+	SetActive(!bIsActive);
+}
+
+bool UActorComponent::IsActive() const
+{
+	return bIsActive;
+}
+
+void UActorComponent::SetTickableWhenPaused(bool bTickableWhenPaused)
+{
+	PrimaryComponentTick.bTickEvenWhenPaused = bTickableWhenPaused;
+}
+
+bool UActorComponent::IsOwnerRunningUserConstructionScript() const
+{
+	const AActor* Owner = GetOwner();
+	return Owner != nullptr && Owner->bRunningUserConstructionScript;
+}
+
+void UActorComponent::AddAssetUserData(UAssetUserData* InUserData)
+{
+	if (InUserData != NULL)
+	{
+		UAssetUserData* ExistingData = GetAssetUserDataOfClass(InUserData->GetClass());
+		if (ExistingData != NULL)
+		{
+			AssetUserData.Remove(ExistingData);
+		}
+		AssetUserData.Add(InUserData);
+	}
+}
+
+UAssetUserData* UActorComponent::GetAssetUserDataOfClass(TSubclassOf<UAssetUserData> InUserDataClass)
+{
+	for (int32 DataIdx = 0; DataIdx < AssetUserData.Num(); DataIdx++)
+	{
+		UAssetUserData* Datum = AssetUserData[DataIdx];
+		if (Datum != NULL && Datum->IsA(InUserDataClass))
+		{
+			return Datum;
+		}
+	}
+	return NULL;
+}
+
+void UActorComponent::RemoveUserDataOfClass(TSubclassOf<UAssetUserData> InUserDataClass)
+{
+	for (int32 DataIdx = 0; DataIdx < AssetUserData.Num(); DataIdx++)
+	{
+		UAssetUserData* Datum = AssetUserData[DataIdx];
+		if (Datum != NULL && Datum->IsA(InUserDataClass))
+		{
+			AssetUserData.RemoveAt(DataIdx);
+			return;
+		}
+	}
+}
+
+void UActorComponent::SetNetAddressable()
+{
+	bNetAddressable = true;
+}
+
+bool UActorComponent::IsNameStableForNetworking() const
+{
+	/** 
+	 * IsNameStableForNetworking means a component can be referred to its path name (relative to owning AActor*) over the network
+	 *
+	 * Components are net addressable if:
+	 *	-They are Default Subobjects (created in C++ constructor)
+	 *	-They were loaded directly from a package (placed in map actors)
+	 *	-They were explicitly set to bNetAddressable (blueprint components created by SCS)
+	 */
+
+	return bNetAddressable || Super::IsNameStableForNetworking();
+}
+
+bool UActorComponent::IsSupportedForNetworking() const
+{
+	return IsNameStableForNetworking() || GetIsReplicated();
+}
+
+void UActorComponent::SetIsReplicated(bool ShouldReplicate)
+{
+	check(GetComponentClassCanReplicate()); // Only certain component classes can replicate!
+	bReplicates = ShouldReplicate;
+
+	AActor* Owner = GetOwner();
+	if (Owner)
+	{
+		Owner->UpdateReplicatedComponent( this );
+	}
+}
+
+bool UActorComponent::GetIsReplicated() const
+{
+	return bReplicates;
+}
+
+bool UActorComponent::ReplicateSubobjects(class UActorChannel *Channel, class FOutBunch *Bunch, FReplicationFlags *RepFlags)
+{
+	return false;
+}
+
+bool UActorComponent::GetComponentClassCanReplicate() const
+{
+	return true;
+}
+
+ENetRole UActorComponent::GetOwnerRole() const
+{
+	AActor *Actor = GetOwner();
+	if (!Actor)
+	{
+		return ROLE_None;
+	}
+	return Actor->Role;
+}
+
+bool UActorComponent::IsNetSimulating() const
+{
+	return GetIsReplicated() && GetOwnerRole() != ROLE_Authority;
+}
+
+void UActorComponent::GetLifetimeReplicatedProps( TArray< FLifetimeProperty > & OutLifetimeProps ) const
+{
+	DOREPLIFETIME( UActorComponent, bIsActive );
+	DOREPLIFETIME( UActorComponent, bReplicates );
+}
+
+void UActorComponent::OnRep_IsActive()
+{
+	SetComponentTickEnabled(bIsActive);
+}
+
+bool UActorComponent::IsEditableWhenInherited() const
+{
+	bool bCanEdit = bEditableWhenInherited;
+	if (bCanEdit)
+	{
+#if WITH_EDITOR
+		if (CreationMethod == EComponentCreationMethod::Native && !IsTemplate())
+		{
+			bCanEdit = FComponentEditorUtils::CanEditNativeComponent(this);
+		}
+		else
+#endif
+			if (CreationMethod == EComponentCreationMethod::UserConstructionScript)
+		{
+			bCanEdit = false;
+		}
+	}
+	return bCanEdit;
+}
+
+void UActorComponent::DetermineUCSModifiedProperties()
+{
+	class FComponentPropertySkipper : public FArchive
+	{
+	public:
+		FComponentPropertySkipper()
+			: FArchive()
+		{
+			ArIsSaving = true;
+		}
+
+		virtual bool ShouldSkipProperty(const UProperty* InProperty) const override
+		{
+			return (    InProperty->HasAnyPropertyFlags(CPF_Transient | CPF_ContainsInstancedReference | CPF_InstancedReference)
+					|| !InProperty->HasAnyPropertyFlags(CPF_Edit | CPF_Interp));
+		}
+	} PropertySkipper;
+
+	UClass* ComponentClass = GetClass();
+	UObject* ComponentArchetype = GetArchetype();
+
+	UCSModifiedProperties.Empty();
+
+	for (TFieldIterator<UProperty> It(ComponentClass); It; ++It)
+	{
+		UProperty* Property = *It;
+		if( Property->ShouldSerializeValue(PropertySkipper) )
+		{
+			for( int32 Idx=0; Idx<Property->ArrayDim; Idx++ )
+			{
+				uint8* DataPtr      = Property->ContainerPtrToValuePtr           <uint8>((uint8*)this, Idx);
+				uint8* DefaultValue = Property->ContainerPtrToValuePtrForDefaults<uint8>(ComponentClass, (uint8*)ComponentArchetype, Idx);
+				if (!Property->Identical( DataPtr, DefaultValue))
+				{
+					UCSModifiedProperties.Add(FSimpleMemberReference());
+					FMemberReference::FillSimpleMemberReference<UProperty>(Property, UCSModifiedProperties.Last());
+					break;
+				}
+			}
+		}
+	}
+}
+
+void UActorComponent::GetUCSModifiedProperties(TSet<const UProperty*>& ModifiedProperties) const
+{
+	for (const FSimpleMemberReference& MemberReference : UCSModifiedProperties)
+	{
+		ModifiedProperties.Add(FMemberReference::ResolveSimpleMemberReference<UProperty>(MemberReference));
+	}
+}
+
+
+#undef LOCTEXT_NAMESPACE
diff --git a/Engine/Source/Runtime/Engine/Private/Components/ArrowComponent.cpp b/Engine/Source/Runtime/Engine/Private/Components/ArrowComponent.cpp
new file mode 100644
--- /dev/null
+++ b/Engine/Source/Runtime/Engine/Private/Components/ArrowComponent.cpp
@@ -0,0 +1,350 @@
+// Copyright 1998-2015 Epic Games, Inc. All Rights Reserved.
+
+#include "EnginePrivate.h"
+#include "DynamicMeshBuilder.h"
+#include "LocalVertexFactory.h"
+
+#define DEFAULT_SCREEN_SIZE	(0.0025f)
+#define ARROW_SCALE			(80.0f)
+#define ARROW_RADIUS_FACTOR	(0.03f)
+#define ARROW_HEAD_FACTOR	(0.2f)
+#define ARROW_HEAD_ANGLE	(20.f)
+
+#if WITH_EDITORONLY_DATA
+float UArrowComponent::EditorScale = 1.0f;
+#endif
+
+/** Vertex Buffer */
+class FArrowVertexBuffer : public FVertexBuffer
+{
+public:
+	TArray<FDynamicMeshVertex> Vertices;
+
+	virtual void InitRHI() override
+	{
+		FRHIResourceCreateInfo CreateInfo;
+		VertexBufferRHI = RHICreateVertexBuffer(Vertices.Num() * sizeof(FDynamicMeshVertex), BUF_Static, CreateInfo);
+
+		// Copy the vertex data into the vertex buffer.
+		void* VertexBufferData = RHILockVertexBuffer(VertexBufferRHI, 0, Vertices.Num() * sizeof(FDynamicMeshVertex), RLM_WriteOnly);
+		FMemory::Memcpy(VertexBufferData, Vertices.GetData(), Vertices.Num() * sizeof(FDynamicMeshVertex));
+		RHIUnlockVertexBuffer(VertexBufferRHI);
+	}
+
+};
+
+/** Index Buffer */
+class FArrowIndexBuffer : public FIndexBuffer
+{
+public:
+	TArray<int32> Indices;
+
+	virtual void InitRHI() override
+	{
+		FRHIResourceCreateInfo CreateInfo;
+		IndexBufferRHI = RHICreateIndexBuffer(sizeof(int32), Indices.Num() * sizeof(int32), BUF_Static, CreateInfo);
+
+		// Write the indices to the index buffer.
+		void* Buffer = RHILockIndexBuffer(IndexBufferRHI, 0, Indices.Num() * sizeof(int32), RLM_WriteOnly);
+		FMemory::Memcpy(Buffer, Indices.GetData(), Indices.Num() * sizeof(int32));
+		RHIUnlockIndexBuffer(IndexBufferRHI);
+	}
+};
+
+/** Vertex Factory */
+class FArrowVertexFactory : public FLocalVertexFactory
+{
+public:
+
+	FArrowVertexFactory()
+	{}
+
+
+	/** Initialization */
+	void Init(const FArrowVertexBuffer* VertexBuffer)
+	{
+		if (IsInRenderingThread())
+		{
+			// Initialize the vertex factory's stream components.
+			DataType NewData;
+			NewData.PositionComponent = STRUCTMEMBER_VERTEXSTREAMCOMPONENT(VertexBuffer, FDynamicMeshVertex, Position, VET_Float3);
+			NewData.TextureCoordinates.Add(
+				FVertexStreamComponent(VertexBuffer, STRUCT_OFFSET(FDynamicMeshVertex, TextureCoordinate), sizeof(FDynamicMeshVertex), VET_Float2)
+				);
+			NewData.TangentBasisComponents[0] = STRUCTMEMBER_VERTEXSTREAMCOMPONENT(VertexBuffer, FDynamicMeshVertex, TangentX, VET_PackedNormal);
+			NewData.TangentBasisComponents[1] = STRUCTMEMBER_VERTEXSTREAMCOMPONENT(VertexBuffer, FDynamicMeshVertex, TangentZ, VET_PackedNormal);
+			SetData(NewData);
+		}
+		else
+		{
+			ENQUEUE_UNIQUE_RENDER_COMMAND_TWOPARAMETER(
+				InitArrowVertexFactory,
+				FArrowVertexFactory*, VertexFactory, this,
+				const FArrowVertexBuffer*, VertexBuffer, VertexBuffer,
+				{
+				// Initialize the vertex factory's stream components.
+				DataType NewData;
+				NewData.PositionComponent = STRUCTMEMBER_VERTEXSTREAMCOMPONENT(VertexBuffer, FDynamicMeshVertex, Position, VET_Float3);
+				NewData.TextureCoordinates.Add(
+					FVertexStreamComponent(VertexBuffer, STRUCT_OFFSET(FDynamicMeshVertex, TextureCoordinate), sizeof(FDynamicMeshVertex), VET_Float2)
+					);
+				NewData.TangentBasisComponents[0] = STRUCTMEMBER_VERTEXSTREAMCOMPONENT(VertexBuffer, FDynamicMeshVertex, TangentX, VET_PackedNormal);
+				NewData.TangentBasisComponents[1] = STRUCTMEMBER_VERTEXSTREAMCOMPONENT(VertexBuffer, FDynamicMeshVertex, TangentZ, VET_PackedNormal);
+				VertexFactory->SetData(NewData);
+			});
+		}
+	}
+};
+
+/** Represents a UArrowComponent to the scene manager. */
+class FArrowSceneProxy : public FPrimitiveSceneProxy
+{
+public:
+
+	FArrowSceneProxy(UArrowComponent* Component)
+		: FPrimitiveSceneProxy(Component)
+		, ArrowColor(Component->ArrowColor)
+		, ArrowSize(Component->ArrowSize)
+		, bIsScreenSizeScaled(Component->bIsScreenSizeScaled)
+		, ScreenSize(Component->ScreenSize)
+#if WITH_EDITORONLY_DATA
+		, bLightAttachment(Component->bLightAttachment)
+		, bTreatAsASprite(Component->bTreatAsASprite)
+		, bUseInEditorScaling(Component->bUseInEditorScaling)
+		, EditorScale(Component->EditorScale)
+#endif
+	{
+		bWillEverBeLit = false;
+#if WITH_EDITOR
+		// If in the editor, extract the sprite category from the component
+		if ( GIsEditor )
+		{
+			SpriteCategoryIndex = GEngine->GetSpriteCategoryIndex( Component->SpriteInfo.Category );
+		}
+#endif	//WITH_EDITOR
+
+		const float HeadAngle = FMath::DegreesToRadians(ARROW_HEAD_ANGLE);
+		const float TotalLength = ArrowSize * ARROW_SCALE;
+		const float HeadLength = TotalLength * ARROW_HEAD_FACTOR;
+		const float ShaftRadius = TotalLength * ARROW_RADIUS_FACTOR;
+		const float ShaftLength = (TotalLength - HeadLength) * 1.1f; // 10% overlap between shaft and head
+		const FVector ShaftCenter = FVector(0.5f * ShaftLength, 0, 0);
+
+		BuildConeVerts(HeadAngle, HeadAngle, -HeadLength, TotalLength, 32, VertexBuffer.Vertices, IndexBuffer.Indices);
+		BuildCylinderVerts(ShaftCenter, FVector(0,0,1), FVector(0,1,0), FVector(1,0,0), ShaftRadius, 0.5f * ShaftLength, 16, VertexBuffer.Vertices, IndexBuffer.Indices);
+
+
+		// Init vertex factory
+		VertexFactory.Init(&VertexBuffer);
+
+		// Enqueue initialization of render resource
+		BeginInitResource(&VertexBuffer);
+		BeginInitResource(&IndexBuffer);
+		BeginInitResource(&VertexFactory);
+	}
+
+	virtual ~FArrowSceneProxy()
+	{
+		VertexBuffer.ReleaseResource();
+		IndexBuffer.ReleaseResource();
+		VertexFactory.ReleaseResource();
+	}
+
+	// FPrimitiveSceneProxy interface.
+	
+	virtual void GetDynamicMeshElements(const TArray<const FSceneView*>& Views, const FSceneViewFamily& ViewFamily, uint32 VisibilityMap, FMeshElementCollector& Collector) const override
+	{
+		QUICK_SCOPE_CYCLE_COUNTER( STAT_ArrowSceneProxy_DrawDynamicElements );
+
+		FMatrix EffectiveLocalToWorld;
+#if WITH_EDITOR
+		if (bLightAttachment)
+		{
+			EffectiveLocalToWorld = GetLocalToWorld().GetMatrixWithoutScale();
+		}
+		else
+#endif	//WITH_EDITOR
+		{
+			EffectiveLocalToWorld = GetLocalToWorld();
+		}
+
+		auto ArrowMaterialRenderProxy = new FColoredMaterialRenderProxy(
+			GEngine->ArrowMaterial->GetRenderProxy(IsSelected(), IsHovered()),
+			ArrowColor,
+			"GizmoColor"
+			);
+
+		Collector.RegisterOneFrameMaterialProxy(ArrowMaterialRenderProxy);
+
+		for (int32 ViewIndex = 0; ViewIndex < Views.Num(); ViewIndex++)
+		{
+			if (VisibilityMap & (1 << ViewIndex))
+			{
+				const FSceneView* View = Views[ViewIndex];
+
+				// Calculate the view-dependent scaling factor.
+				float ViewScale = 1.0f;
+				if (bIsScreenSizeScaled && (View->ViewMatrices.ProjMatrix.M[3][3] != 1.0f))
+				{
+					const float ZoomFactor = FMath::Min<float>(View->ViewMatrices.ProjMatrix.M[0][0], View->ViewMatrices.ProjMatrix.M[1][1]);
+					if (ZoomFactor != 0.0f)
+					{
+						const float Radius = View->WorldToScreen(Origin).W * (ScreenSize / ZoomFactor);
+						if (Radius < 1.0f && Radius > 0)
+						{
+							ViewScale *= Radius;
+						}
+					}
+				}
+
+		#if WITH_EDITORONLY_DATA
+				ViewScale *= EditorScale;
+		#endif
+
+				// Draw the mesh.
+				FMeshBatch& Mesh = Collector.AllocateMesh();
+				FMeshBatchElement& BatchElement = Mesh.Elements[0];
+				BatchElement.IndexBuffer = &IndexBuffer;
+				Mesh.bWireframe = false;
+				Mesh.VertexFactory = &VertexFactory;
+				Mesh.MaterialRenderProxy = ArrowMaterialRenderProxy;
+				BatchElement.PrimitiveUniformBuffer = CreatePrimitiveUniformBufferImmediate(FScaleMatrix(ViewScale) * EffectiveLocalToWorld, GetBounds(), GetLocalBounds(), true, UseEditorDepthTest());
+				BatchElement.FirstIndex = 0;
+				BatchElement.NumPrimitives = IndexBuffer.Indices.Num() / 3;
+				BatchElement.MinVertexIndex = 0;
+				BatchElement.MaxVertexIndex = VertexBuffer.Vertices.Num() - 1;
+				Mesh.ReverseCulling = IsLocalToWorldDeterminantNegative();
+				Mesh.Type = PT_TriangleList;
+				Mesh.DepthPriorityGroup = SDPG_World;
+				Mesh.bCanApplyViewModeOverrides = false;
+				Collector.AddMesh(ViewIndex, Mesh);
+			}
+		}
+	}
+
+	virtual FPrimitiveViewRelevance GetViewRelevance(const FSceneView* View)
+	{
+		FPrimitiveViewRelevance Result;
+		Result.bDrawRelevance = IsShown(View) && (View->Family->EngineShowFlags.BillboardSprites);
+		Result.bDynamicRelevance = true;
+#if WITH_EDITOR
+		if (bTreatAsASprite)
+		{
+			if ( GIsEditor && SpriteCategoryIndex != INDEX_NONE && SpriteCategoryIndex < View->SpriteCategoryVisibility.Num() && !View->SpriteCategoryVisibility[ SpriteCategoryIndex ] )
+			{
+				Result.bDrawRelevance = false;
+			}
+		}
+#endif
+		Result.bShadowRelevance = IsShadowCast(View);
+		Result.bEditorPrimitiveRelevance = UseEditorCompositing(View);
+		return Result;
+	}
+
+	virtual void OnTransformChanged() override
+	{
+		Origin = GetLocalToWorld().GetOrigin();
+	}
+
+	virtual uint32 GetMemoryFootprint( void ) const { return( sizeof( *this ) + GetAllocatedSize() ); }
+	uint32 GetAllocatedSize( void ) const { return( FPrimitiveSceneProxy::GetAllocatedSize() ); }
+
+private:
+	FArrowVertexBuffer VertexBuffer;
+	FArrowIndexBuffer IndexBuffer;
+	FArrowVertexFactory VertexFactory;
+
+	FVector Origin;
+	FColor ArrowColor;
+	float ArrowSize;
+	bool bIsScreenSizeScaled;
+	float ScreenSize;
+#if WITH_EDITORONLY_DATA
+	bool bLightAttachment;
+	bool bTreatAsASprite;
+	int32 SpriteCategoryIndex;
+	bool bUseInEditorScaling;
+	float EditorScale;
+#endif // #if WITH_EDITORONLY_DATA
+};
+
+UArrowComponent::UArrowComponent(const FObjectInitializer& ObjectInitializer)
+	: Super(ObjectInitializer)
+{
+	// Structure to hold one-time initialization
+	struct FConstructorStatics
+	{
+		FName ID_Misc;
+		FText NAME_Misc;
+		FConstructorStatics()
+			: ID_Misc(TEXT("Misc"))
+			, NAME_Misc(NSLOCTEXT( "SpriteCategory", "Misc", "Misc" ))
+		{
+		}
+	};
+	static FConstructorStatics ConstructorStatics;
+
+	SetCollisionProfileName(UCollisionProfile::NoCollision_ProfileName);
+	ArrowColor = FColor(255, 0, 0, 255);
+
+	ArrowSize = 1.0f;
+	bHiddenInGame = true;
+	bUseEditorCompositing = true;
+	bGenerateOverlapEvents = false;
+	bIsScreenSizeScaled = false;
+	ScreenSize = DEFAULT_SCREEN_SIZE;
+#if WITH_EDITORONLY_DATA
+	SpriteInfo.Category = ConstructorStatics.ID_Misc;
+	SpriteInfo.DisplayName = ConstructorStatics.NAME_Misc;
+	bLightAttachment = false;
+	bUseInEditorScaling = true;
+#endif // WITH_EDITORONLY_DATA
+}
+
+FPrimitiveSceneProxy* UArrowComponent::CreateSceneProxy()
+{
+	return new FArrowSceneProxy(this);
+}
+
+#if WITH_EDITOR
+bool UArrowComponent::ComponentIsTouchingSelectionBox(const FBox& InSelBBox, const FEngineShowFlags& ShowFlags, const bool bConsiderOnlyBSP, const bool bMustEncompassEntireComponent) const
+{
+	// Arrow components not treated as 'selectable' in editor
+	return false;
+}
+
+bool UArrowComponent::ComponentIsTouchingSelectionFrustum(const FConvexVolume& InFrustum, const FEngineShowFlags& ShowFlags, const bool bConsiderOnlyBSP, const bool bMustEncompassEntireComponent) const
+{
+	// Arrow components not treated as 'selectable' in editor
+	return false;
+}
+#endif
+
+
+FBoxSphereBounds UArrowComponent::CalcBounds(const FTransform& LocalToWorld) const
+{
+	return FBoxSphereBounds(FBox(FVector(0,-ARROW_SCALE,-ARROW_SCALE),FVector(ArrowSize * ARROW_SCALE * 3.0f,ARROW_SCALE,ARROW_SCALE))).TransformBy(LocalToWorld);
+}
+
+void UArrowComponent::SetArrowColor_DEPRECATED(FColor NewColor)
+{
+	ArrowColor = NewColor;
+	MarkRenderStateDirty();
+}
+
+void UArrowComponent::SetArrowColor_New(FLinearColor NewColor)
+{
+	ArrowColor = NewColor;
+	MarkRenderStateDirty();
+}
+
+#if WITH_EDITORONLY_DATA
+void UArrowComponent::SetEditorScale(float InEditorScale)
+{
+	EditorScale = InEditorScale;
+	for(TObjectIterator<UArrowComponent> It; It; ++It)
+	{
+		It->MarkRenderStateDirty();
+	}
+}
+#endif
\ No newline at end of file
diff --git a/Engine/Source/Runtime/Engine/Private/Components/AudioComponent.cpp b/Engine/Source/Runtime/Engine/Private/Components/AudioComponent.cpp
new file mode 100644
--- /dev/null
+++ b/Engine/Source/Runtime/Engine/Private/Components/AudioComponent.cpp
@@ -0,0 +1,632 @@
+// Copyright 1998-2015 Epic Games, Inc. All Rights Reserved.
+
+#include "EnginePrivate.h"
+#include "SoundDefinitions.h"
+#include "Sound/SoundNodeAttenuation.h"
+#include "Sound/SoundCue.h"
+#include "SubtitleManager.h"
+
+/*-----------------------------------------------------------------------------
+UAudioComponent implementation.
+-----------------------------------------------------------------------------*/
+UAudioComponent::UAudioComponent(const FObjectInitializer& ObjectInitializer)
+	: Super(ObjectInitializer)
+{
+	bAutoDestroy = false;
+	bAutoActivate = true;
+	bAllowSpatialization = true;
+	bStopWhenOwnerDestroyed = true;
+	bNeverNeedsRenderUpdate = true;
+#if WITH_EDITORONLY_DATA
+	bVisualizeComponent = true;
+#endif
+	VolumeMultiplier = 1.f;
+	PitchMultiplier = 1.f;
+	VolumeModulationMin = 1.f;
+	VolumeModulationMax = 1.f;
+	PitchModulationMin = 1.f;
+	PitchModulationMax = 1.f;
+	HighFrequencyGainMultiplier = 1.0f;
+}
+
+FString UAudioComponent::GetDetailedInfoInternal( void ) const
+{
+	FString Result;
+
+	if (Sound != nullptr)
+	{
+		Result = Sound->GetPathName(nullptr);
+	}
+	else
+	{
+		Result = TEXT( "No_Sound" );
+	}
+
+	return Result;
+}
+
+#if WITH_EDITORONLY_DATA
+void UAudioComponent::OnRegister()
+{
+	Super::OnRegister();
+
+	UpdateSpriteTexture();
+}
+#endif
+
+void UAudioComponent::OnUnregister()
+{
+	// Route OnUnregister event.
+	Super::OnUnregister();
+
+	// Don't stop audio and clean up component if owner has been destroyed (default behaviour). This function gets
+	// called from AActor::ClearComponents when an actor gets destroyed which is not usually what we want for one-
+	// shot sounds.
+	AActor* Owner = GetOwner();
+	if (!Owner || bStopWhenOwnerDestroyed)
+	{
+		Stop();
+	}
+}
+
+void UAudioComponent::SetSound( USoundBase* NewSound )
+{
+	const bool bPlay = IsPlaying();
+
+	// If this is an auto destroy component we need to prevent it from being auto-destroyed since we're really just restarting it
+	const bool bWasAutoDestroy = bAutoDestroy;
+	bAutoDestroy = false;
+	Stop();
+	bAutoDestroy = bWasAutoDestroy;
+
+	Sound = NewSound;
+
+	if (bPlay)
+	{
+		Play();
+	}
+}
+
+void UAudioComponent::OnUpdateTransform(bool bSkipPhysicsMove)
+{
+	Super::OnUpdateTransform(bSkipPhysicsMove);
+
+	if (bIsActive && !bPreviewComponent)
+	{
+		if (FAudioDevice * AudioDevice = GetAudioDevice())
+		{
+			FActiveSound* ActiveSound = AudioDevice->FindActiveSound(this);
+			if (ActiveSound)
+			{
+				FScopeCycleCounterUObject ComponentScope(ActiveSound->Sound);
+				ActiveSound->Transform = ComponentToWorld;
+			}
+		}
+	}
+};
+
+void UAudioComponent::Play(float StartTime)
+{
+	PlayInternal(StartTime);
+}
+
+void UAudioComponent::PlayInternal(const float StartTime, const float FadeInDuration, const float FadeVolumeLevel)
+{
+	UWorld* World = GetWorld();
+
+	UE_LOG(LogAudio, Verbose, TEXT("%g: Playing AudioComponent : '%s' with Sound: '%s'"), World ? World->GetAudioTimeSeconds() : 0.0f, *GetFullName(), Sound ? *Sound->GetName() : TEXT("nullptr"));
+
+	if (bIsActive)
+	{
+		// If this is an auto destroy component we need to prevent it from being auto-destroyed since we're really just restarting it
+		bool bCurrentAutoDestroy = bAutoDestroy;
+		bAutoDestroy = false;
+		Stop();
+		bAutoDestroy = bCurrentAutoDestroy;
+	}
+
+	if (Sound && (World == nullptr || World->bAllowAudioPlayback))
+	{
+		if (FAudioDevice* AudioDevice = GetAudioDevice())
+		{
+			FActiveSound NewActiveSound;
+			NewActiveSound.AudioComponent = this;
+			NewActiveSound.World = GetWorld();
+			NewActiveSound.Sound = Sound;
+			NewActiveSound.SoundClassOverride = SoundClassOverride;
+
+			NewActiveSound.VolumeMultiplier = (VolumeModulationMax + ((VolumeModulationMin - VolumeModulationMax) * FMath::SRand())) * VolumeMultiplier;
+			NewActiveSound.PitchMultiplier = (PitchModulationMax + ((PitchModulationMin - PitchModulationMax) * FMath::SRand())) * PitchMultiplier;
+			NewActiveSound.HighFrequencyGainMultiplier = HighFrequencyGainMultiplier;
+
+			NewActiveSound.RequestedStartTime = FMath::Max(0.f, StartTime);
+			NewActiveSound.OcclusionCheckInterval = OcclusionCheckInterval;
+			NewActiveSound.SubtitlePriority = SubtitlePriority;
+
+			NewActiveSound.bShouldRemainActiveIfDropped = bShouldRemainActiveIfDropped;
+			NewActiveSound.bHandleSubtitles = (!bSuppressSubtitles || OnQueueSubtitles.IsBound());
+			NewActiveSound.bIgnoreForFlushing = bIgnoreForFlushing;
+
+			NewActiveSound.bIsUISound = bIsUISound;
+			NewActiveSound.bIsMusic = bIsMusic;
+			NewActiveSound.bAlwaysPlay = bAlwaysPlay;
+			NewActiveSound.bReverb = bReverb;
+			NewActiveSound.bCenterChannelOnly = bCenterChannelOnly;
+
+			NewActiveSound.bLocationDefined = !bPreviewComponent;
+			if (NewActiveSound.bLocationDefined)
+			{
+				NewActiveSound.Transform = ComponentToWorld;
+			}
+
+			const FAttenuationSettings* AttenuationSettingsToApply = (bAllowSpatialization ? GetAttenuationSettingsToApply() : nullptr);
+			NewActiveSound.bAllowSpatialization = bAllowSpatialization;
+			NewActiveSound.bHasAttenuationSettings = (AttenuationSettingsToApply != nullptr);
+			if (NewActiveSound.bHasAttenuationSettings)
+			{
+				NewActiveSound.AttenuationSettings = *AttenuationSettingsToApply;
+			}
+
+			NewActiveSound.InstanceParameters = InstanceParameters;
+
+			NewActiveSound.TargetAdjustVolumeMultiplier = FadeVolumeLevel;
+
+			if (FadeInDuration > 0.0f)
+			{
+				NewActiveSound.CurrentAdjustVolumeMultiplier = 0.f;
+				NewActiveSound.TargetAdjustVolumeStopTime = FadeInDuration;
+			}
+			else
+			{
+				NewActiveSound.CurrentAdjustVolumeMultiplier = FadeVolumeLevel;
+			}
+
+			// TODO - Audio Threading. This call would be a task call to dispatch to the audio thread
+			AudioDevice->AddNewActiveSound(NewActiveSound);
+
+			bIsActive = true;
+		}
+	}
+}
+
+FAudioDevice* UAudioComponent::GetAudioDevice() const
+{
+	FAudioDevice* AudioDevice = nullptr;
+
+	if (GEngine)
+	{
+		UWorld* World = GetWorld();
+		if (World)
+		{
+			AudioDevice = World->GetAudioDevice();
+		}
+		else
+		{
+			AudioDevice = GEngine->GetMainAudioDevice();
+		}
+	}
+	return AudioDevice;
+}
+
+void UAudioComponent::FadeIn( float FadeInDuration, float FadeVolumeLevel, float StartTime )
+{
+	PlayInternal(StartTime, FadeInDuration, FadeVolumeLevel);
+}
+
+void UAudioComponent::FadeOut( float FadeOutDuration, float FadeVolumeLevel )
+{
+	if (bIsActive)
+	{
+		if (FadeOutDuration > 0.0f)
+		{
+			// TODO - Audio Threading. This call would be a task
+			if (FAudioDevice* AudioDevice = GetAudioDevice())
+			{
+				FActiveSound* ActiveSound = AudioDevice->FindActiveSound(this);
+				if (ActiveSound)
+				{
+					ActiveSound->TargetAdjustVolumeMultiplier = FadeVolumeLevel;
+					ActiveSound->TargetAdjustVolumeStopTime = ActiveSound->PlaybackTime + FadeOutDuration;
+					ActiveSound->bFadingOut = true;
+				}
+			}
+		}
+		else
+		{
+			Stop();
+		}
+	}
+}
+
+void UAudioComponent::AdjustVolume( float AdjustVolumeDuration, float AdjustVolumeLevel )
+{
+	if (bIsActive)
+	{
+		// TODO - Audio Threading. This call would be a task
+		if (FAudioDevice* AudioDevice = GetAudioDevice())
+		{
+			FActiveSound* ActiveSound = AudioDevice->FindActiveSound(this);
+			if (ActiveSound)
+			{
+				ActiveSound->bFadingOut = false;
+				ActiveSound->TargetAdjustVolumeMultiplier = AdjustVolumeLevel;
+
+				if( AdjustVolumeDuration > 0.0f )
+				{
+					ActiveSound->TargetAdjustVolumeStopTime = ActiveSound->PlaybackTime + AdjustVolumeDuration;
+				}
+				else
+				{
+					ActiveSound->CurrentAdjustVolumeMultiplier = AdjustVolumeLevel;
+					ActiveSound->TargetAdjustVolumeStopTime = -1.0f;
+				}
+			}
+		}
+	}
+}
+
+void UAudioComponent::Stop()
+{
+	if (bIsActive)
+	{
+		UE_LOG(LogAudio, Verbose, TEXT( "%g: Stopping AudioComponent : '%s' with Sound: '%s'" ), GetWorld() ? GetWorld()->GetAudioTimeSeconds() : 0.0f, *GetFullName(), Sound ? *Sound->GetName() : TEXT( "nullptr" ) );
+
+		// TODO - Audio Threading. This call would be a task
+		if (FAudioDevice* AudioDevice = GetAudioDevice())
+		{
+			AudioDevice->StopActiveSound(this);
+		}
+	}
+}
+
+void UAudioComponent::PlaybackCompleted(bool bFailedToStart)
+{
+	// Mark inactive before calling destroy to avoid recursion
+	bIsActive = false;
+
+	if (!bFailedToStart && GetWorld() != nullptr && (OnAudioFinished.IsBound() || OnAudioFinishedNative.IsBound()))
+	{
+		INC_DWORD_STAT( STAT_AudioFinishedDelegatesCalled );
+		SCOPE_CYCLE_COUNTER( STAT_AudioFinishedDelegates );
+
+		OnAudioFinished.Broadcast();
+		OnAudioFinishedNative.Broadcast(this);
+	}
+
+	// Auto destruction is handled via marking object for deletion.
+	if (bAutoDestroy)
+	{
+		DestroyComponent();
+	}
+}
+
+bool UAudioComponent::IsPlaying() const
+{
+	return bIsActive;
+}
+
+#if WITH_EDITORONLY_DATA
+void UAudioComponent::UpdateSpriteTexture()
+{
+	if (SpriteComponent)
+	{
+		if (bAutoActivate)
+		{
+			SpriteComponent->SetSprite(LoadObject<UTexture2D>(nullptr, TEXT("/Engine/EditorResources/AudioIcons/S_AudioComponent_AutoActivate.S_AudioComponent_AutoActivate")));
+		}
+		else
+		{
+			SpriteComponent->SetSprite(LoadObject<UTexture2D>(nullptr, TEXT("/Engine/EditorResources/AudioIcons/S_AudioComponent.S_AudioComponent")));
+		}
+	}
+}
+#endif
+
+#if WITH_EDITOR
+void UAudioComponent::PostEditChangeProperty(FPropertyChangedEvent& PropertyChangedEvent)
+{
+	if (bIsActive)
+	{
+		// If this is an auto destroy component we need to prevent it from being auto-destroyed since we're really just restarting it
+		const bool bWasAutoDestroy = bAutoDestroy;
+		bAutoDestroy = false;
+		Stop();
+		bAutoDestroy = bWasAutoDestroy;
+		Play();
+	}
+
+#if WITH_EDITORONLY_DATA
+	UpdateSpriteTexture();
+#endif
+
+	Super::PostEditChangeProperty(PropertyChangedEvent);
+}
+#endif // WITH_EDITOR
+
+const FAttenuationSettings* UAudioComponent::GetAttenuationSettingsToApply() const
+{
+	if (bOverrideAttenuation)
+	{
+		return &AttenuationOverrides;
+	}
+	else if (AttenuationSettings)
+	{
+		return &AttenuationSettings->Attenuation;
+	}
+	else if (Sound)
+	{
+		return Sound->GetAttenuationSettingsToApply();
+	}
+	return nullptr;
+}
+
+void UAudioComponent::CollectAttenuationShapesForVisualization(TMultiMap<EAttenuationShape::Type, FAttenuationSettings::AttenuationShapeDetails>& ShapeDetailsMap) const
+{
+	const FAttenuationSettings *AttenuationSettingsToApply = GetAttenuationSettingsToApply();
+
+	if (AttenuationSettingsToApply)
+	{
+		AttenuationSettingsToApply->CollectAttenuationShapesForVisualization(ShapeDetailsMap);
+	}
+
+	// For sound cues we'll dig in and see if we can find any attenuation sound nodes that will affect the settings
+	USoundCue* SoundCue = Cast<USoundCue>(Sound);
+	if (SoundCue)
+	{
+		TArray<USoundNodeAttenuation*> AttenuationNodes;
+		SoundCue->RecursiveFindAttenuation( SoundCue->FirstNode, AttenuationNodes );
+		for (int32 NodeIndex = 0; NodeIndex < AttenuationNodes.Num(); ++NodeIndex)
+		{
+			AttenuationSettingsToApply = AttenuationNodes[NodeIndex]->GetAttenuationSettingsToApply();
+			if (AttenuationSettingsToApply)
+			{
+				AttenuationSettingsToApply->CollectAttenuationShapesForVisualization(ShapeDetailsMap);
+			}
+		}
+	}
+}
+
+void UAudioComponent::Activate(bool bReset)
+{
+	if (bReset || ShouldActivate() == true)
+	{
+		Play();
+	}
+}
+
+void UAudioComponent::Deactivate()
+{
+	if (ShouldActivate() == false)
+	{
+		Stop();
+	}
+}
+
+void UAudioComponent::SetFloatParameter( FName InName, float InFloat )
+{
+	if (InName != NAME_None)
+	{
+		bool bFound = false;
+
+		// First see if an entry for this name already exists
+		for (int32 i = 0; i < InstanceParameters.Num(); i++)
+		{
+			FAudioComponentParam& P = InstanceParameters[i];
+			if (P.ParamName == InName)
+			{
+				P.FloatParam = InFloat;
+				bFound = true;
+				break;
+			}
+		}
+
+		// We didn't find one, so create a new one.
+		if (!bFound)
+		{
+			const int32 NewParamIndex = InstanceParameters.AddZeroed();
+			InstanceParameters[ NewParamIndex ].ParamName = InName;
+			InstanceParameters[ NewParamIndex ].FloatParam = InFloat;
+		}
+
+		// If we're active we need to push this value to the ActiveSound
+		if (bIsActive)
+		{
+			// TODO - Audio Threading. This call would be a task
+			if (FAudioDevice* AudioDevice = GetAudioDevice())
+			{
+				FActiveSound* ActiveSound = AudioDevice->FindActiveSound(this);
+				if (ActiveSound)
+				{
+					ActiveSound->SetFloatParameter(InName, InFloat);
+				}
+			}
+		}
+	}
+}
+
+void UAudioComponent::SetWaveParameter( FName InName, USoundWave* InWave )
+{
+	if (InName != NAME_None)
+	{
+		bool bFound = false;
+		// First see if an entry for this name already exists
+		for (int32 i = 0; i < InstanceParameters.Num(); i++)
+		{
+			FAudioComponentParam& P = InstanceParameters[i];
+			if (P.ParamName == InName)
+			{
+				P.SoundWaveParam = InWave;
+				bFound = true;
+				break;
+			}
+		}
+
+		// We didn't find one, so create a new one.
+		if (!bFound)
+		{
+			const int32 NewParamIndex = InstanceParameters.AddZeroed();
+			InstanceParameters[NewParamIndex].ParamName = InName;
+			InstanceParameters[NewParamIndex].SoundWaveParam = InWave;
+		}
+
+		// If we're active we need to push this value to the ActiveSound
+		if (bIsActive)
+		{
+			// TODO - Audio Threading. This call would be a task
+			if (FAudioDevice* AudioDevice = GetAudioDevice())
+			{
+				FActiveSound* ActiveSound = AudioDevice->FindActiveSound(this);
+				if (ActiveSound)
+				{
+					ActiveSound->SetWaveParameter(InName, InWave);
+				}
+			}
+		}
+	}
+}
+
+void UAudioComponent::SetBoolParameter( FName InName, bool InBool )
+{
+	if (InName != NAME_None)
+	{
+		bool bFound = false;
+		// First see if an entry for this name already exists
+		for (int32 i = 0; i < InstanceParameters.Num(); i++)
+		{
+			FAudioComponentParam& P = InstanceParameters[i];
+			if (P.ParamName == InName)
+			{
+				P.BoolParam = InBool;
+				bFound = true;
+				break;
+			}
+		}
+
+		// We didn't find one, so create a new one.
+		if (!bFound)
+		{
+			const int32 NewParamIndex = InstanceParameters.AddZeroed();
+			InstanceParameters[ NewParamIndex ].ParamName = InName;
+			InstanceParameters[ NewParamIndex ].BoolParam = InBool;
+		}
+
+		// If we're active we need to push this value to the ActiveSound
+		if (bIsActive)
+		{
+			// TODO - Audio Threading. This call would be a task
+			if (FAudioDevice* AudioDevice = GetAudioDevice())
+			{
+				FActiveSound* ActiveSound = AudioDevice->FindActiveSound(this);
+				if (ActiveSound)
+				{
+					ActiveSound->SetBoolParameter(InName, InBool);
+				}
+			}
+		}
+	}
+}
+
+
+void UAudioComponent::SetIntParameter( FName InName, int32 InInt )
+{
+	if (InName != NAME_None)
+	{
+		bool bFound = false;
+		// First see if an entry for this name already exists
+		for (int32 i = 0; i < InstanceParameters.Num(); i++)
+		{
+			FAudioComponentParam& P = InstanceParameters[i];
+			if (P.ParamName == InName)
+			{
+				P.IntParam = InInt;
+				bFound = true;
+				break;
+			}
+		}
+
+		// We didn't find one, so create a new one.
+		if (!bFound)
+		{
+			const int32 NewParamIndex = InstanceParameters.AddZeroed();
+			InstanceParameters[NewParamIndex].ParamName = InName;
+			InstanceParameters[NewParamIndex].IntParam = InInt;
+		}
+
+		// If we're active we need to push this value to the ActiveSound
+		if (bIsActive)
+		{
+			// TODO - Audio Threading. This call would be a task
+			if (FAudioDevice* AudioDevice = GetAudioDevice())
+			{
+				FActiveSound* ActiveSound = AudioDevice->FindActiveSound(this);
+				if (ActiveSound)
+				{
+					ActiveSound->SetIntParameter(InName, InInt);
+				}
+			}
+		}
+	}
+}
+
+void UAudioComponent::SetVolumeMultiplier(float NewVolumeMultiplier)
+{
+	VolumeMultiplier = NewVolumeMultiplier;
+	VolumeModulationMin = VolumeModulationMax = 1.f;
+
+	// TODO - Audio Threading. This call would be a task
+	if (FAudioDevice* AudioDevice = GetAudioDevice())
+	{
+		FActiveSound* ActiveSound = AudioDevice->FindActiveSound(this);
+		if (ActiveSound)
+		{
+			ActiveSound->VolumeMultiplier = NewVolumeMultiplier;
+		}
+	}
+}
+
+void UAudioComponent::SetPitchMultiplier(float NewPitchMultiplier)
+{
+	PitchMultiplier = NewPitchMultiplier;
+	PitchModulationMin = PitchModulationMax = 1.f;
+
+	// TODO - Audio Threading. This call would be a task
+	if (FAudioDevice* AudioDevice = GetAudioDevice())
+	{
+		FActiveSound* ActiveSound = AudioDevice->FindActiveSound(this);
+		if (ActiveSound)
+		{
+			ActiveSound->PitchMultiplier = NewPitchMultiplier;
+		}
+	}
+}
+
+void UAudioComponent::SetUISound(bool bInIsUISound)
+{
+	bIsUISound = bInIsUISound;
+
+	// TODO - Audio Threading. This call would be a task
+	if (FAudioDevice* AudioDevice = GetAudioDevice())
+	{
+		FActiveSound* ActiveSound = AudioDevice->FindActiveSound(this);
+		if (ActiveSound)
+		{
+			ActiveSound->bIsUISound = bIsUISound;
+		}
+	}
+}
+
+void UAudioComponent::AdjustAttenuation(const FAttenuationSettings& InAttenuationSettings)
+{
+	bOverrideAttenuation = true;
+	AttenuationOverrides = InAttenuationSettings;
+
+	// TODO - Audio Threading. This call would be a task
+	if (FAudioDevice* AudioDevice = GetAudioDevice())
+	{
+		FActiveSound* ActiveSound = AudioDevice->FindActiveSound(this);
+		if (ActiveSound)
+		{
+			ActiveSound->AttenuationSettings = AttenuationOverrides;
+		}
+	}
+}
diff --git a/Engine/Source/Runtime/Engine/Private/Components/BoxComponent.cpp b/Engine/Source/Runtime/Engine/Private/Components/BoxComponent.cpp
new file mode 100644
--- /dev/null
+++ b/Engine/Source/Runtime/Engine/Private/Components/BoxComponent.cpp
@@ -0,0 +1,153 @@
+// Copyright 1998-2015 Epic Games, Inc. All Rights Reserved.
+
+
+#include "EnginePrivate.h"
+#include "Components/BoxComponent.h"
+
+UBoxComponent::UBoxComponent(const FObjectInitializer& ObjectInitializer)
+	: Super(ObjectInitializer)
+{
+	BoxExtent = FVector(32.0f, 32.0f, 32.0f);
+
+	bUseEditorCompositing = true;
+}
+
+
+void UBoxComponent::SetBoxExtent(FVector NewBoxExtent, bool bUpdateOverlaps)
+{
+	BoxExtent = NewBoxExtent;
+	MarkRenderStateDirty();
+
+	// do this if already created
+	// otherwise, it hasn't been really created yet
+	if (bPhysicsStateCreated)
+	{
+		DestroyPhysicsState();
+		UpdateBodySetup();
+		CreatePhysicsState();
+
+		if ( bUpdateOverlaps && IsCollisionEnabled() && GetOwner() )
+		{
+			UpdateOverlaps();
+		}
+	}
+}
+
+void UBoxComponent::UpdateBodySetup()
+{
+	if(ShapeBodySetup == NULL || ShapeBodySetup->IsPendingKill())
+	{
+		ShapeBodySetup = NewObject<UBodySetup>(this);
+		ShapeBodySetup->CollisionTraceFlag = CTF_UseSimpleAsComplex;
+		ShapeBodySetup->AggGeom.BoxElems.Add(FKBoxElem());
+	}
+
+	check(ShapeBodySetup->AggGeom.BoxElems.Num() == 1);
+	FKBoxElem* se = ShapeBodySetup->AggGeom.BoxElems.GetData();
+
+	// @todo UE4 do we allow this now?
+	// check for malformed values
+	if( BoxExtent.X < KINDA_SMALL_NUMBER )
+	{
+		BoxExtent.X = 1.0f;
+	}
+
+	if( BoxExtent.Y < KINDA_SMALL_NUMBER )
+	{
+		BoxExtent.Y = 1.0f;
+	}
+
+	if( BoxExtent.Z < KINDA_SMALL_NUMBER )
+	{
+		BoxExtent.Z = 1.0f;
+	}
+
+	// now set the PhysX data values
+	se->SetTransform( FTransform::Identity );
+	se->X = BoxExtent.X*2;
+	se->Y = BoxExtent.Y*2;
+	se->Z = BoxExtent.Z*2;
+}
+
+bool UBoxComponent::IsZeroExtent() const
+{
+	return BoxExtent.IsZero();
+}
+
+FBoxSphereBounds UBoxComponent::CalcBounds(const FTransform& LocalToWorld) const 
+{
+	return FBoxSphereBounds( FBox( -BoxExtent, BoxExtent ) ).TransformBy(LocalToWorld);
+}
+
+
+
+FPrimitiveSceneProxy* UBoxComponent::CreateSceneProxy()
+{
+	/** Represents a UCapsuleComponent to the scene manager. */
+	class FBoxSceneProxy : public FPrimitiveSceneProxy
+	{
+	public:
+		FBoxSceneProxy(const UBoxComponent* InComponent)
+			:	FPrimitiveSceneProxy(InComponent)
+			,	bDrawOnlyIfSelected( InComponent->bDrawOnlyIfSelected )
+			,   BoxExtents( InComponent->BoxExtent )
+			,	BoxColor( InComponent->ShapeColor )
+		{
+			bWillEverBeLit = false;
+		}
+
+		virtual void GetDynamicMeshElements(const TArray<const FSceneView*>& Views, const FSceneViewFamily& ViewFamily, uint32 VisibilityMap, FMeshElementCollector& Collector) const override
+		{
+			QUICK_SCOPE_CYCLE_COUNTER( STAT_BoxSceneProxy_GetDynamicMeshElements );
+
+			const FMatrix& LocalToWorld = GetLocalToWorld();
+			
+			for (int32 ViewIndex = 0; ViewIndex < Views.Num(); ViewIndex++)
+			{
+				if (VisibilityMap & (1 << ViewIndex))
+				{
+					const FSceneView* View = Views[ViewIndex];
+
+					const FLinearColor DrawColor = GetViewSelectionColor(BoxColor, *View, IsSelected(), IsHovered(), false, IsIndividuallySelected() );
+
+					FPrimitiveDrawInterface* PDI = Collector.GetPDI(ViewIndex);
+					DrawOrientedWireBox(PDI, LocalToWorld.GetOrigin(), LocalToWorld.GetScaledAxis( EAxis::X ), LocalToWorld.GetScaledAxis( EAxis::Y ), LocalToWorld.GetScaledAxis( EAxis::Z ), BoxExtents, DrawColor, SDPG_World);
+				}
+			}
+		}
+
+		virtual FPrimitiveViewRelevance GetViewRelevance(const FSceneView* View) override
+		{
+			const bool bVisible = !bDrawOnlyIfSelected || IsSelected();
+
+			FPrimitiveViewRelevance Result;
+			Result.bDrawRelevance = IsShown(View) && bVisible;
+			Result.bDynamicRelevance = true;
+			Result.bShadowRelevance = IsShadowCast(View);
+			Result.bEditorPrimitiveRelevance = UseEditorCompositing(View);
+			return Result;
+		}
+		virtual uint32 GetMemoryFootprint( void ) const override { return( sizeof( *this ) + GetAllocatedSize() ); }
+		uint32 GetAllocatedSize( void ) const { return( FPrimitiveSceneProxy::GetAllocatedSize() ); }
+
+	private:
+		const uint32	bDrawOnlyIfSelected:1;
+		const FVector	BoxExtents;
+		const FColor	BoxColor;
+	};
+
+	return new FBoxSceneProxy( this );
+}
+
+
+FCollisionShape UBoxComponent::GetCollisionShape(float Inflation) const
+{
+	FVector Extent = GetScaledBoxExtent() + Inflation;
+	if (Inflation < 0.f)
+	{
+		// Don't shrink below zero size.
+		Extent = Extent.ComponentMax(FVector::ZeroVector);
+	}
+
+	return FCollisionShape::MakeBox(Extent);
+}
diff --git a/Engine/Source/Runtime/Engine/Private/Components/BrushComponent.cpp b/Engine/Source/Runtime/Engine/Private/Components/BrushComponent.cpp
new file mode 100644
--- /dev/null
+++ b/Engine/Source/Runtime/Engine/Private/Components/BrushComponent.cpp
@@ -0,0 +1,777 @@
+// Copyright 1998-2015 Epic Games, Inc. All Rights Reserved.
+
+/*=============================================================================
+	BrushComponent.cpp: Unreal brush component implementation
+=============================================================================*/
+
+#include "EnginePrivate.h"
+#include "Components/BrushComponent.h"
+#include "Engine/Polys.h"
+#include "Model.h"
+#include "LevelUtils.h"
+#if WITH_EDITOR
+#include "Collision.h"
+#include "ShowFlags.h"
+#include "ConvexVolume.h"
+#endif
+#include "DebuggingDefines.h"
+#include "ActorEditorUtils.h"
+
+DEFINE_LOG_CATEGORY_STATIC(LogBrushComponent, Log, All);
+
+#if WITH_EDITORONLY_DATA
+struct FModelWireVertex
+{
+	FVector Position;
+	FPackedNormal TangentX;
+	FPackedNormal TangentZ;
+	FVector2D UV;
+};
+
+class FModelWireVertexBuffer : public FVertexBuffer
+{
+public:
+
+	/** Initialization constructor. */
+	FModelWireVertexBuffer(UModel* InModel):
+		NumVertices(0)
+	{
+#if WITH_EDITOR
+		Polys.Append(InModel->Polys->Element);
+		for(int32 PolyIndex = 0;PolyIndex < InModel->Polys->Element.Num();PolyIndex++)
+		{
+			NumVertices += InModel->Polys->Element[PolyIndex].Vertices.Num();
+		}
+#endif
+	}
+
+	// FRenderResource interface.
+	virtual void InitRHI() override
+	{
+		if(NumVertices)
+		{
+			FRHIResourceCreateInfo CreateInfo;
+			VertexBufferRHI = RHICreateVertexBuffer(NumVertices * sizeof(FModelWireVertex),BUF_Static, CreateInfo);
+
+			FModelWireVertex* DestVertex = (FModelWireVertex*)RHILockVertexBuffer(VertexBufferRHI,0,NumVertices * sizeof(FModelWireVertex),RLM_WriteOnly);
+			for(int32 PolyIndex = 0;PolyIndex < Polys.Num();PolyIndex++)
+			{
+				FPoly& Poly = Polys[PolyIndex];
+				for(int32 VertexIndex = 0;VertexIndex < Poly.Vertices.Num();VertexIndex++)
+				{
+					DestVertex->Position = Poly.Vertices[VertexIndex];
+					DestVertex->TangentX = FVector(1,0,0);
+					DestVertex->TangentZ = FVector(0,0,1);
+					// TangentZ.w contains the sign of the tangent basis determinant. Assume +1
+					DestVertex->TangentZ.Vector.W = 255;
+					DestVertex->UV.X	 = 0.0f;
+					DestVertex->UV.Y	 = 0.0f;
+					DestVertex++;
+				}
+			}
+			RHIUnlockVertexBuffer(VertexBufferRHI);
+		}
+	}
+
+	// Accessors.
+	uint32 GetNumVertices() const { return NumVertices; }
+
+private:
+	TArray<FPoly> Polys;
+	uint32 NumVertices;
+};
+
+class FModelWireIndexBuffer : public FIndexBuffer
+{
+public:
+
+	/** Initialization constructor. */
+	FModelWireIndexBuffer(UModel* InModel):
+		NumEdges(0)
+	{
+#if WITH_EDITOR
+		Polys.Append(InModel->Polys->Element);
+		for(int32 PolyIndex = 0;PolyIndex < InModel->Polys->Element.Num();PolyIndex++)
+		{
+			NumEdges += InModel->Polys->Element[PolyIndex].Vertices.Num();
+		}
+#endif
+	}
+
+	// FRenderResource interface.
+	virtual void InitRHI() override
+	{
+		if(NumEdges)
+		{
+			FRHIResourceCreateInfo CreateInfo;
+			IndexBufferRHI = RHICreateIndexBuffer(sizeof(uint16),NumEdges * 2 * sizeof(uint16),BUF_Static, CreateInfo);
+
+			uint16* DestIndex = (uint16*)RHILockIndexBuffer(IndexBufferRHI,0,NumEdges * 2 * sizeof(uint16),RLM_WriteOnly);
+			uint16 BaseIndex = 0;
+			for(int32 PolyIndex = 0;PolyIndex < Polys.Num();PolyIndex++)
+			{
+				FPoly&	Poly = Polys[PolyIndex];
+				for(int32 VertexIndex = 0;VertexIndex < Poly.Vertices.Num();VertexIndex++)
+				{
+					*DestIndex++ = BaseIndex + VertexIndex;
+					*DestIndex++ = BaseIndex + ((VertexIndex + 1) % Poly.Vertices.Num());
+				}
+				BaseIndex += Poly.Vertices.Num();
+			}
+			RHIUnlockIndexBuffer(IndexBufferRHI);
+		}
+	}
+
+	// Accessors.
+	uint32 GetNumEdges() const { return NumEdges; }
+
+private:
+	TArray<FPoly> Polys;
+	uint32 NumEdges;
+};
+#endif // WITH_EDITORONLY_DATA
+
+class FBrushSceneProxy : public FPrimitiveSceneProxy
+{
+public:
+	FBrushSceneProxy(UBrushComponent* Component, ABrush* Owner):
+		FPrimitiveSceneProxy(Component),
+#if WITH_EDITORONLY_DATA
+		WireIndexBuffer(Component->Brush),
+		WireVertexBuffer(Component->Brush),
+#endif
+		bVolume(false),
+		bBuilder(false),
+		bSolidWhenSelected(false),
+		bInManipulation(false),
+		BrushColor(GEngine->C_BrushWire),
+		BodySetup(Component->BrushBodySetup),
+		CollisionResponse(Component->GetCollisionResponseToChannels())
+	{
+		bWillEverBeLit = false;
+
+		if(Owner)
+		{
+			// If the editor is in a state where drawing the brush wireframe isn't desired, bail out.
+			if( !GEngine->ShouldDrawBrushWireframe( Owner ) )
+			{
+				return;
+			}
+
+			// Determine the type of brush this is.
+			bVolume = Owner->IsVolumeBrush();
+			bBuilder = FActorEditorUtils::IsABuilderBrush( Owner );
+			BrushColor = Owner->GetWireColor();
+			bSolidWhenSelected = Owner->bSolidWhenSelected;
+			bInManipulation = Owner->bInManipulation;
+
+			// Builder brushes should be unaffected by level coloration, so if this is a builder brush, use
+			// the brush color as the level color.
+			if ( bBuilder )
+			{
+				LevelColor = BrushColor;
+			}
+			else
+			{
+				// Try to find a color for level coloration.
+				ULevel* Level = Owner->GetLevel();
+				ULevelStreaming* LevelStreaming = FLevelUtils::FindStreamingLevel( Level );
+				if ( LevelStreaming )
+				{
+					LevelColor = LevelStreaming->LevelColor;
+				}
+			}
+		}
+
+		bUseEditorDepthTest = !bInManipulation;
+
+		// Get a color for property coloration.
+		FColor NewPropertyColor;
+		GEngine->GetPropertyColorationColor( (UObject*)Component, NewPropertyColor );
+		PropertyColor = NewPropertyColor;
+
+#if WITH_EDITORONLY_DATA
+		ENQUEUE_UNIQUE_RENDER_COMMAND_TWOPARAMETER(
+			InitBrushVertexFactory,
+			FLocalVertexFactory*,VertexFactory,&VertexFactory,
+			FVertexBuffer*,WireVertexBuffer,&WireVertexBuffer,
+			{
+				FLocalVertexFactory::DataType Data;
+				Data.PositionComponent = STRUCTMEMBER_VERTEXSTREAMCOMPONENT(WireVertexBuffer,FModelWireVertex,Position,VET_Float3);
+				Data.TangentBasisComponents[0] = STRUCTMEMBER_VERTEXSTREAMCOMPONENT(WireVertexBuffer,FModelWireVertex,TangentX,VET_PackedNormal);
+				Data.TangentBasisComponents[1] = STRUCTMEMBER_VERTEXSTREAMCOMPONENT(WireVertexBuffer,FModelWireVertex,TangentZ,VET_PackedNormal);
+				Data.TextureCoordinates.Add( STRUCTMEMBER_VERTEXSTREAMCOMPONENT(WireVertexBuffer,FModelWireVertex,UV,VET_Float2) );
+				VertexFactory->SetData(Data);
+			});
+#endif
+	}
+
+	virtual ~FBrushSceneProxy()
+	{
+#if WITH_EDITORONLY_DATA
+		VertexFactory.ReleaseResource();
+		WireIndexBuffer.ReleaseResource();
+		WireVertexBuffer.ReleaseResource();
+#endif
+
+	}
+
+	bool IsCollisionView(const FEngineShowFlags& EngineShowFlags, bool & bDrawCollision) const
+	{
+		const bool bInCollisionView = EngineShowFlags.CollisionVisibility || EngineShowFlags.CollisionPawn;
+		if (bInCollisionView && IsCollisionEnabled())
+		{
+			bDrawCollision = EngineShowFlags.CollisionPawn && (CollisionResponse.GetResponse(ECC_Pawn) != ECR_Ignore);
+			bDrawCollision |= EngineShowFlags.CollisionVisibility && (CollisionResponse.GetResponse(ECC_Visibility) != ECR_Ignore);
+		}
+		else
+		{
+			bDrawCollision = false;
+		}
+
+		return bInCollisionView;
+	}
+
+	virtual void GetDynamicMeshElements(const TArray<const FSceneView*>& Views, const FSceneViewFamily& ViewFamily, uint32 VisibilityMap, FMeshElementCollector& Collector) const override
+	{
+		QUICK_SCOPE_CYCLE_COUNTER( STAT_BrushSceneProxy_GetDynamicMeshElements );
+
+		if( AllowDebugViewmodes() )
+		{
+			for (int32 ViewIndex = 0; ViewIndex < Views.Num(); ViewIndex++)
+			{
+				if (VisibilityMap & (1 << ViewIndex))
+				{
+					const FSceneView* View = Views[ViewIndex];
+
+					bool bDrawCollision = false;
+					const bool bInCollisionView = IsCollisionView(ViewFamily.EngineShowFlags, bDrawCollision);
+
+					// Draw solid if 'solid when selected' and selected, or we are in a 'collision view'
+					const bool bDrawSolid = ((bSolidWhenSelected && IsSelected()) || (bInCollisionView && bDrawCollision));
+					// Don't draw wireframe if in a collision view mode and not drawing solid
+					const bool bDrawWireframe = !bInCollisionView;
+
+					// Choose color to draw it
+					FLinearColor DrawColor = BrushColor;
+					// In a collision view mode
+					if(bInCollisionView)
+					{
+						DrawColor = BrushColor;
+					}
+					else if(View->Family->EngineShowFlags.PropertyColoration)
+					{
+						DrawColor = PropertyColor;
+					}
+					else if(View->Family->EngineShowFlags.LevelColoration)
+					{
+						DrawColor = LevelColor;
+					}
+
+
+					// SOLID
+					if(bDrawSolid)
+					{
+						if(BodySetup != NULL)
+						{
+							auto SolidMaterialInstance = new FColoredMaterialRenderProxy(
+								GEngine->ShadedLevelColorationUnlitMaterial->GetRenderProxy(IsSelected(), IsHovered()),
+								DrawColor
+								);
+
+							Collector.RegisterOneFrameMaterialProxy(SolidMaterialInstance);
+
+							FTransform GeomTransform(GetLocalToWorld());
+							BodySetup->AggGeom.GetAggGeom(GeomTransform, DrawColor, /*Material=*/SolidMaterialInstance, false, /*bSolid=*/ true, UseEditorDepthTest(), ViewIndex, Collector);
+						}
+					}
+					// WIREFRAME
+					else if(bDrawWireframe)
+					{
+						// If we have the editor data (Wire Buffers) use those for wireframe
+#if WITH_EDITOR
+						if(WireIndexBuffer.GetNumEdges() && WireVertexBuffer.GetNumVertices())
+						{
+							auto WireframeMaterial = new FColoredMaterialRenderProxy(
+								GEngine->LevelColorationUnlitMaterial->GetRenderProxy(IsSelected(), IsHovered()),
+								GetViewSelectionColor(DrawColor, *View, !(GIsEditor && (View->Family->EngineShowFlags.Selection)) || IsSelected(), IsHovered(), false, IsIndividuallySelected() )
+								);
+
+							Collector.RegisterOneFrameMaterialProxy(WireframeMaterial);
+
+							FMeshBatch& Mesh = Collector.AllocateMesh();
+							FMeshBatchElement& BatchElement = Mesh.Elements[0];
+							BatchElement.IndexBuffer = &WireIndexBuffer;
+							Mesh.VertexFactory = &VertexFactory;
+							Mesh.MaterialRenderProxy = WireframeMaterial;
+							BatchElement.PrimitiveUniformBufferResource = &GetUniformBuffer();
+							BatchElement.FirstIndex = 0;
+							BatchElement.NumPrimitives = WireIndexBuffer.GetNumEdges();
+							BatchElement.MinVertexIndex = 0;
+							BatchElement.MaxVertexIndex = WireVertexBuffer.GetNumVertices() - 1;
+							Mesh.CastShadow = false;
+							Mesh.Type = PT_LineList;
+							Mesh.DepthPriorityGroup = IsSelected() ? SDPG_Foreground : SDPG_World;
+							Collector.AddMesh(ViewIndex, Mesh);
+						}
+						else 
+#endif
+						if(BodySetup != NULL)
+							// If not, use the body setup for wireframe
+						{
+							FTransform GeomTransform(GetLocalToWorld());
+							BodySetup->AggGeom.GetAggGeom(GeomTransform, GetSelectionColor(DrawColor, IsSelected(), IsHovered()), /* Material=*/ NULL, false, /* bSolid=*/ false, UseEditorDepthTest(), ViewIndex, Collector);
+						}
+
+					}
+				}
+			}
+		}
+	}
+
+	virtual FPrimitiveViewRelevance GetViewRelevance(const FSceneView* View)
+	{
+		bool bVisible = false;
+
+
+		// We render volumes in collision view. In game, always, in editor, if the EngineShowFlags.Volumes option is on.
+		if(bSolidWhenSelected && IsSelected())
+		{
+			FPrimitiveViewRelevance Result;
+			Result.bDrawRelevance = true;
+			Result.bDynamicRelevance = true;
+			return Result;
+		}
+
+		const bool bInCollisionView = (View->Family->EngineShowFlags.CollisionVisibility || View->Family->EngineShowFlags.CollisionPawn);
+
+		if(IsShown(View))
+		{
+			bool bNeverShow = false;
+
+			if( GIsEditor )
+			{
+				const bool bShowBuilderBrush = View->Family->EngineShowFlags.BuilderBrush != 0;
+
+				// Only render builder brush and only if the show flags indicate that we should render builder brushes.
+				if( bBuilder && (!bShowBuilderBrush) )
+				{
+					bNeverShow = true;
+				}
+			}
+
+			if(bNeverShow == false)
+			{
+				const bool bBSPVisible = View->Family->EngineShowFlags.BSP;
+				const bool bBrushesVisible = View->Family->EngineShowFlags.Brushes;
+
+				if ( !bVolume ) // EngineShowFlags.Collision does not apply to volumes
+				{
+					if( (bBSPVisible && bBrushesVisible) )
+					{
+						bVisible = true;
+					}
+				}
+
+				// See if we should be visible because we are in a 'collision view' and have collision enabled
+				if (bInCollisionView && IsCollisionEnabled())
+				{
+					bVisible = true;
+				}
+
+				// Always show the build brush and any brushes that are selected in the editor.
+				if( GIsEditor )
+				{
+					if( bBuilder || IsSelected() )
+					{
+						bVisible = true;
+					}
+				}
+
+				if ( bVolume )
+				{
+					const bool bVolumesVisible = View->Family->EngineShowFlags.Volumes;
+					if(!GIsEditor || View->bIsGameView || bVolumesVisible)
+					{
+						bVisible = true;
+					}
+				}		
+			}
+		}
+		
+		FPrimitiveViewRelevance Result;
+		Result.bDrawRelevance = bVisible;
+		Result.bDynamicRelevance = true;
+		Result.bShadowRelevance = IsShadowCast(View);
+		if(bInManipulation)
+		{
+			Result.bEditorNoDepthTestPrimitiveRelevance = true;
+		}
+
+		// Don't render on top in 'collision view' modes
+		if(!bInCollisionView)
+		{
+			Result.bEditorPrimitiveRelevance = true;
+		}
+
+		return Result;
+	}
+
+	virtual void CreateRenderThreadResources()
+	{
+#if WITH_EDITORONLY_DATA
+		VertexFactory.InitResource();
+		WireIndexBuffer.InitResource();
+		WireVertexBuffer.InitResource();
+#endif
+
+	}
+
+	virtual uint32 GetMemoryFootprint( void ) const { return( sizeof( *this ) + GetAllocatedSize() ); }
+	uint32 GetAllocatedSize( void ) const { return( FPrimitiveSceneProxy::GetAllocatedSize() ); }
+
+private:
+#if WITH_EDITORONLY_DATA
+	FLocalVertexFactory VertexFactory;
+	FModelWireIndexBuffer WireIndexBuffer;
+	FModelWireVertexBuffer WireVertexBuffer;
+#endif
+
+	uint32 bVolume : 1;
+	uint32 bBuilder : 1;	
+	uint32 bSolidWhenSelected : 1;
+	uint32 bInManipulation : 1;
+
+	FColor BrushColor;
+	FLinearColor LevelColor;
+	FColor PropertyColor;
+
+	/** Collision Response of this component**/
+	UBodySetup* BodySetup;
+	FCollisionResponseContainer CollisionResponse;
+};
+
+UBrushComponent::UBrushComponent(const FObjectInitializer& ObjectInitializer)
+	: Super(ObjectInitializer)
+{
+	bHiddenInGame = true;
+	AlwaysLoadOnClient = false;
+	AlwaysLoadOnServer = false;
+	bUseAsOccluder = true;
+	bRequiresCustomLocation = true;
+	bUseEditorCompositing = true;
+	bCanEverAffectNavigation = true;
+}
+
+FPrimitiveSceneProxy* UBrushComponent::CreateSceneProxy()
+{
+	FPrimitiveSceneProxy* Proxy = NULL;
+	
+	if (Brush != NULL)	
+	{
+		// Check to make sure that we want to draw this brushed based on editor settings.
+		ABrush*	Owner = Cast<ABrush>(GetOwner());
+		if(Owner)
+		{
+			// If the editor is in a state where drawing the brush wireframe isn't desired, bail out.
+			if( GEngine->ShouldDrawBrushWireframe( Owner ) )
+			{
+				Proxy = new FBrushSceneProxy(this, Owner);
+			}
+		}
+		else
+		{
+			Proxy = new FBrushSceneProxy(this, Owner);
+		}
+	}
+
+	return Proxy;
+}
+
+
+FBoxSphereBounds UBrushComponent::CalcBounds(const FTransform& LocalToWorld) const
+{
+#if WITH_EDITOR
+	if(Brush && Brush->Polys && Brush->Polys->Element.Num())
+	{
+		TArray<FVector> Points;
+		for( int32 i=0; i<Brush->Polys->Element.Num(); i++ )
+		{
+			for( int32 j=0; j<Brush->Polys->Element[i].Vertices.Num(); j++ )
+			{
+				Points.Add(Brush->Polys->Element[i].Vertices[j]);
+			}
+		}
+		return FBoxSphereBounds( Points.GetData(), Points.Num() ).TransformBy(LocalToWorld);
+	}
+	else 
+#endif // WITH_EDITOR
+	if ((BrushBodySetup != NULL) && (BrushBodySetup->AggGeom.GetElementCount() > 0))
+	{
+		FBoxSphereBounds NewBounds;
+		BrushBodySetup->AggGeom.CalcBoxSphereBounds(NewBounds, LocalToWorld);
+		return NewBounds;
+	}
+	else
+	{
+		return FBoxSphereBounds(LocalToWorld.GetLocation(), FVector::ZeroVector, 0.f);
+	}
+}
+
+FVector UBrushComponent::GetCustomLocation() const
+{
+	const FVector LocationWithPivot = ComponentToWorld.GetLocation();
+	const FVector LocationNoPivot = LocationWithPivot + ComponentToWorld.TransformVector(PrePivot);
+
+	return LocationNoPivot;
+}
+
+FTransform UBrushComponent::CalcNewComponentToWorld(const FTransform& NewRelativeTransform, const USceneComponent* Parent) const
+{
+	FTransform CompToWorld = Super::CalcNewComponentToWorld(NewRelativeTransform, Parent);
+
+	const FVector LocationNoPivot = CompToWorld.GetLocation();
+	const FVector LocationWithPivot = LocationNoPivot + CompToWorld.TransformVector(-PrePivot);
+
+	CompToWorld.SetLocation(LocationWithPivot);
+
+	return CompToWorld;
+}
+
+
+void UBrushComponent::GetUsedMaterials( TArray<UMaterialInterface*>& OutMaterials ) const
+{
+#if WITH_EDITOR
+	// Get the material from each polygon making up the brush.
+	if( Brush && Brush->Polys )
+	{
+		UPolys* Polys = Brush->Polys;
+		if( Polys )
+		{
+			for( int32 ElementIdx = 0; ElementIdx < Polys->Element.Num(); ++ElementIdx )
+			{
+				OutMaterials.Add( Polys->Element[ ElementIdx ].Material );
+			}
+		}
+	}
+#endif // WITH_EDITOR
+}
+
+void UBrushComponent::PostLoad()
+{
+	Super::PostLoad();
+
+	// Stop existing BrushComponents from generating mirrored collision mesh
+	if ((GetLinkerUE4Version() < VER_UE4_NO_MIRROR_BRUSH_MODEL_COLLISION) && (BrushBodySetup != NULL))
+	{
+		BrushBodySetup->bGenerateMirroredCollision = false;
+	}
+}
+
+
+SIZE_T UBrushComponent::GetResourceSize(EResourceSizeMode::Type Mode)
+{
+	SIZE_T ResSize = Super::GetResourceSize(Mode);
+
+	// Count the bodysetup we own as well for 'inclusive' stats
+	if((Mode == EResourceSizeMode::Inclusive) && (BrushBodySetup != NULL))
+	{
+		ResSize += BrushBodySetup->GetResourceSize(Mode);
+	}
+
+	return ResSize;
+}
+
+uint8 UBrushComponent::GetStaticDepthPriorityGroup() const
+{
+	ABrush* BrushOwner = Cast<ABrush>(GetOwner());
+
+	// Draw selected and builder brushes in the foreground DPG.
+	if(BrushOwner && (IsOwnerSelected() || FActorEditorUtils::IsABuilderBrush(BrushOwner)))
+	{
+		return SDPG_Foreground;
+	}
+	else
+	{
+		return DepthPriorityGroup;
+	}
+}
+
+#if WITH_EDITOR
+static bool IsComponentTypeShown(AActor* Actor, const FEngineShowFlags& ShowFlags)
+{
+	if (Actor != nullptr)
+	{
+		return (Actor->IsA(AVolume::StaticClass())) ? ShowFlags.Volumes : ShowFlags.BSP;
+	}
+
+	return false;
+}
+
+bool UBrushComponent::ComponentIsTouchingSelectionBox(const FBox& InSelBBox, const FEngineShowFlags& ShowFlags, const bool bConsiderOnlyBSP, const bool bMustEncompassEntireComponent) const
+{
+	if (!IsComponentTypeShown(GetOwner(), ShowFlags))
+	{
+		return false;
+	}
+
+	if (Brush != nullptr && Brush->Polys != nullptr)
+	{
+		TArray<FVector> Vertices;
+
+		if (!bMustEncompassEntireComponent)
+		{
+			const bool bMustContainVertex = true;
+			if (bMustContainVertex)
+			{
+				// Test if any poly vertex intersects the box
+				for (const auto& Poly : Brush->Polys->Element)
+				{
+					for (const auto& Vertex : Poly.Vertices)
+					{
+						const FVector Location = ComponentToWorld.TransformPosition(Vertex);
+						const bool bLocationIntersected = FMath::PointBoxIntersection(Location, InSelBBox);
+
+						// If the selection box doesn't have to encompass the entire component and any poly vertex intersects with the selection
+						// box, this component qualifies
+						if (bLocationIntersected)
+						{
+							return true;
+						}
+					}
+				}
+			}
+			else
+			{
+				// Alternative method, which can be enabled by setting bMustContainVertex = false:
+				// Test if any poly edge intersects the box
+				for (const auto& Poly : Brush->Polys->Element)
+				{
+					const int32 NumVerts = Poly.Vertices.Num();
+					if (NumVerts > 0)
+					{
+						FVector StartVert = ComponentToWorld.TransformPosition(Poly.Vertices[NumVerts - 1]);
+						for (int32 Index = 0; Index < NumVerts; ++Index)
+						{
+							const FVector EndVert = ComponentToWorld.TransformPosition(Poly.Vertices[Index]);
+
+							if (FMath::LineBoxIntersection(InSelBBox, StartVert, EndVert, EndVert - StartVert))
+							{
+								return true;
+							}
+
+							StartVert = EndVert;
+						}
+					}
+				}
+			}
+
+			// No poly intersected with the bounding box
+			return false;
+		}
+		else
+		{
+			for (const auto& Poly : Brush->Polys->Element)
+			{
+				// The component must be entirely within the bounding box...
+				for (const auto& Vertex : Poly.Vertices)
+				{
+					const FVector Location = ComponentToWorld.TransformPosition(Vertex);
+					const bool bLocationIntersected = FMath::PointBoxIntersection(Location, InSelBBox);
+
+					// If the selection box has to encompass the entire component and a poly vertex didn't intersect with the selection
+					// box, this component does not qualify
+					if (!bLocationIntersected)
+					{
+						return false;
+					}
+				}
+
+				// If the selection box has to encompass all of the component and none of the component's verts failed the intersection test, this component
+				// is considered touching
+				return true;
+			}
+
+			// All points lay within the selection box
+			return true;
+		}
+	}
+
+	return false;
+}
+
+
+bool UBrushComponent::ComponentIsTouchingSelectionFrustum(const FConvexVolume& InFrustum, const FEngineShowFlags& ShowFlags, const bool bConsiderOnlyBSP, const bool bMustEncompassEntireComponent) const
+{
+	if (!IsComponentTypeShown(GetOwner(), ShowFlags))
+	{
+		return false;
+	}
+
+	if (Brush != nullptr && Brush->Polys != nullptr)
+	{
+		TArray<FVector> Vertices;
+
+		for (const auto& Poly : Brush->Polys->Element)
+		{
+			for (const auto& Vertex : Poly.Vertices)
+			{
+				const FVector Location = ComponentToWorld.TransformPosition(Vertex);
+				const bool bIntersect = InFrustum.IntersectSphere(Location, 0.0f);
+
+				if (bIntersect && !bMustEncompassEntireComponent)
+				{
+					// If we intersected a vertex and we don't require the box to encompass the entire component
+					// then the actor should be selected and we can stop checking
+					return true;
+				}
+				else if (!bIntersect && bMustEncompassEntireComponent)
+				{
+					// If we didn't intersect a vertex but we require the box to encompass the entire component
+					// then this test failed and we can stop checking
+					return false;
+				}
+			}
+		}
+
+		// If the selection box has to encompass all of the component and none of the component's verts failed the intersection test, this component
+		// is considered touching
+		return true;
+	}
+
+	return false;
+}
+#endif
+
+
+void UBrushComponent::BuildSimpleBrushCollision()
+{
+	AActor* Owner = GetOwner();
+	if(!Owner)
+	{
+		UE_LOG(LogBrushComponent, Warning, TEXT("BuildSimpleBrushCollision: BrushComponent with no Owner!") );
+		return;
+	}
+
+	if(BrushBodySetup == NULL)
+	{
+		BrushBodySetup = NewObject<UBodySetup>(this);
+		check(BrushBodySetup);
+	}
+
+	// No complex collision, so use the simple for that
+	BrushBodySetup->CollisionTraceFlag = CTF_UseSimpleAsComplex;
+	// Don't need mirrored version of collision data
+	BrushBodySetup->bGenerateMirroredCollision = false;
+
+#if WITH_EDITOR
+	// Convert collision model into convex hulls.
+	BrushBodySetup->CreateFromModel( Brush, true );
+	RecreatePhysicsState();
+#endif // WITH_EDITOR
+
+	MarkPackageDirty();
+}
+
+
+
