Engine: Moved actor components into Components folder

diff --git a/Engine/Source/Runtime/Engine/Private/Components/CapsuleComponent.cpp b/Engine/Source/Runtime/Engine/Private/Components/CapsuleComponent.cpp
new file mode 100644
--- /dev/null
+++ b/Engine/Source/Runtime/Engine/Private/Components/CapsuleComponent.cpp
@@ -0,0 +1,197 @@
+// Copyright 1998-2015 Epic Games, Inc. All Rights Reserved.
+
+
+#include "EnginePrivate.h"
+#include "Components/CapsuleComponent.h"
+
+
+UCapsuleComponent::UCapsuleComponent(const FObjectInitializer& ObjectInitializer)
+	: Super(ObjectInitializer)
+{
+	ShapeColor = FColor(223, 149, 157, 255);
+
+	CapsuleRadius = 22.0f;
+	CapsuleHalfHeight = 44.0f;
+	bUseEditorCompositing = true;
+}
+
+
+
+FPrimitiveSceneProxy* UCapsuleComponent::CreateSceneProxy()
+{
+	/** Represents a UCapsuleComponent to the scene manager. */
+	class FDrawCylinderSceneProxy : public FPrimitiveSceneProxy
+	{
+	public:
+		FDrawCylinderSceneProxy(const UCapsuleComponent* InComponent)
+			:	FPrimitiveSceneProxy(InComponent)
+			,	bDrawOnlyIfSelected( InComponent->bDrawOnlyIfSelected )
+			,	CapsuleRadius( InComponent->CapsuleRadius )
+			,	CapsuleHalfHeight( InComponent->CapsuleHalfHeight )
+			,	ShapeColor( InComponent->ShapeColor )
+		{
+			bWillEverBeLit = false;
+		}
+
+		virtual void GetDynamicMeshElements(const TArray<const FSceneView*>& Views, const FSceneViewFamily& ViewFamily, uint32 VisibilityMap, FMeshElementCollector& Collector) const override
+		{
+			QUICK_SCOPE_CYCLE_COUNTER( STAT_GetDynamicMeshElements_DrawDynamicElements );
+
+		
+			const FMatrix& LocalToWorld = GetLocalToWorld();
+			const int32 CapsuleSides =  FMath::Clamp<int32>(CapsuleRadius/4.f, 16, 64);
+
+			for (int32 ViewIndex = 0; ViewIndex < Views.Num(); ViewIndex++)
+			{
+
+				if (VisibilityMap & (1 << ViewIndex))
+				{
+					const FSceneView* View = Views[ViewIndex];
+					const FLinearColor DrawCapsuleColor = GetViewSelectionColor(ShapeColor, *View, IsSelected(), IsHovered(), false, IsIndividuallySelected() );
+
+					FPrimitiveDrawInterface* PDI = Collector.GetPDI(ViewIndex);
+					DrawWireCapsule( PDI, LocalToWorld.GetOrigin(), LocalToWorld.GetScaledAxis( EAxis::X ), LocalToWorld.GetScaledAxis( EAxis::Y ), LocalToWorld.GetScaledAxis( EAxis::Z ), DrawCapsuleColor, CapsuleRadius, CapsuleHalfHeight, CapsuleSides, SDPG_World );
+				}
+			}
+		}
+
+		virtual FPrimitiveViewRelevance GetViewRelevance(const FSceneView* View) override
+		{
+			const bool bVisible = !bDrawOnlyIfSelected || IsSelected();
+			FPrimitiveViewRelevance Result;
+			Result.bDrawRelevance = IsShown(View) && bVisible;
+			Result.bDynamicRelevance = true;
+			Result.bShadowRelevance = IsShadowCast(View);
+			Result.bEditorPrimitiveRelevance = UseEditorCompositing(View);
+			return Result;
+		}
+		virtual uint32 GetMemoryFootprint( void ) const override { return( sizeof( *this ) + GetAllocatedSize() ); }
+		uint32 GetAllocatedSize( void ) const { return( FPrimitiveSceneProxy::GetAllocatedSize() ); }
+
+	private:
+		const uint32	bDrawOnlyIfSelected:1;
+		const float		CapsuleRadius;
+		const float		CapsuleHalfHeight;
+		const FColor	ShapeColor;
+	};
+
+	return new FDrawCylinderSceneProxy( this );
+}
+
+
+FBoxSphereBounds UCapsuleComponent::CalcBounds(const FTransform& LocalToWorld) const
+{
+	FVector BoxPoint = FVector(CapsuleRadius,CapsuleRadius,CapsuleHalfHeight);
+	return FBoxSphereBounds(FVector::ZeroVector, BoxPoint, BoxPoint.Size()).TransformBy(LocalToWorld);
+}
+
+void UCapsuleComponent::CalcBoundingCylinder(float& CylinderRadius, float& CylinderHalfHeight) const 
+{
+	const float Scale = ComponentToWorld.GetMaximumAxisScale();
+	const float CapsuleEndCapCenter = FMath::Max(CapsuleHalfHeight - CapsuleRadius, 0.f);
+	const FVector ZAxis = ComponentToWorld.TransformVectorNoScale(FVector(0.f, 0.f, CapsuleEndCapCenter * Scale));
+	
+	const float ScaledRadius = CapsuleRadius * Scale;
+	
+	CylinderRadius		= ScaledRadius + FMath::Sqrt(FMath::Square(ZAxis.X) + FMath::Square(ZAxis.Y));
+	CylinderHalfHeight	= ScaledRadius + ZAxis.Z;
+}
+
+void UCapsuleComponent::Serialize(FArchive& Ar)
+{
+	Super::Serialize(Ar);
+
+	if (Ar.IsLoading() && (Ar.UE4Ver() < VER_UE4_AFTER_CAPSULE_HALF_HEIGHT_CHANGE))
+	{
+		if ((CapsuleHeight_DEPRECATED != 0.0f) || (Ar.UE4Ver() < VER_UE4_BLUEPRINT_VARS_NOT_READ_ONLY))
+		{
+			CapsuleHalfHeight = CapsuleHeight_DEPRECATED;
+			CapsuleHeight_DEPRECATED = 0.0f;
+		}
+	}
+
+	CapsuleHalfHeight = FMath::Max(CapsuleHalfHeight, CapsuleRadius);
+}
+
+#if WITH_EDITOR
+void UCapsuleComponent::PostEditChangeProperty(FPropertyChangedEvent& PropertyChangedEvent)
+{
+	const FName PropertyName = PropertyChangedEvent.Property ? PropertyChangedEvent.Property->GetFName() : NAME_None;
+
+	// We only want to modify the property that was changed at this point
+	// things like propagation from CDO to instances don't work correctly if changing one property causes a different property to change
+	if (PropertyName == GET_MEMBER_NAME_CHECKED(UCapsuleComponent, CapsuleHalfHeight))
+	{
+		CapsuleHalfHeight = FMath::Max(CapsuleHalfHeight, CapsuleRadius);
+	}
+	else if (PropertyName == GET_MEMBER_NAME_CHECKED(UCapsuleComponent, CapsuleRadius))
+	{
+		CapsuleRadius = FMath::Min(CapsuleHalfHeight, CapsuleRadius);
+	}
+
+	Super::PostEditChangeProperty(PropertyChangedEvent);
+}
+#endif	// WITH_EDITOR
+
+void UCapsuleComponent::SetCapsuleSize(float NewRadius, float NewHalfHeight, bool bUpdateOverlaps)
+{
+	CapsuleHalfHeight = FMath::Max(NewHalfHeight, NewRadius);
+	CapsuleRadius = NewRadius;
+	MarkRenderStateDirty();
+
+	// do this if already created
+	// otherwise, it hasn't been really created yet
+	if (bPhysicsStateCreated)
+	{
+		DestroyPhysicsState();
+		UpdateBodySetup();
+		CreatePhysicsState();
+
+		if ( bUpdateOverlaps && IsCollisionEnabled() && GetOwner() )
+		{
+			UpdateOverlaps();
+		}
+	}
+}
+
+void UCapsuleComponent::UpdateBodySetup()
+{
+	if (ShapeBodySetup == NULL || ShapeBodySetup->IsPendingKill())
+	{
+		ShapeBodySetup = NewObject<UBodySetup>(this);
+		ShapeBodySetup->CollisionTraceFlag = CTF_UseSimpleAsComplex;
+		ShapeBodySetup->AggGeom.SphylElems.Add(FKSphylElem());
+	}
+
+	check(ShapeBodySetup->AggGeom.SphylElems.Num() == 1);
+	FKSphylElem* SE = ShapeBodySetup->AggGeom.SphylElems.GetData();
+	
+	SE->SetTransform(FTransform::Identity);
+	SE->Radius = CapsuleRadius;
+	SE->Length = 2 * FMath::Max(CapsuleHalfHeight - CapsuleRadius, 1.f);	//SphylElem uses height from center of capsule spheres, but UCapsuleComponent uses halfHeight from end of the sphere
+}
+
+bool UCapsuleComponent::IsZeroExtent() const
+{
+	return (CapsuleRadius == 0.f) && (CapsuleHalfHeight == 0.f);
+}
+
+
+FCollisionShape UCapsuleComponent::GetCollisionShape(float Inflation) const
+{
+	const float Radius = FMath::Max(0.f, GetScaledCapsuleRadius() + Inflation);
+	const float HalfHeight = FMath::Max(0.f, GetScaledCapsuleHalfHeight() + Inflation);
+	return FCollisionShape::MakeCapsule(Radius, HalfHeight);
+}
+
+bool UCapsuleComponent::AreSymmetricRotations(const FQuat& A, const FQuat& B, const FVector& Scale3D) const
+{
+	if (Scale3D.X != Scale3D.Y)
+	{
+		return false;
+	}
+
+	const FVector AUp = A.GetAxisZ();
+	const FVector BUp = B.GetAxisZ();
+	return AUp.Equals(BUp);
+}
