Engine: Moved actor components into Components folder

diff --git a/Engine/Source/Runtime/Engine/Private/Components/CharacterMovementComponent.cpp b/Engine/Source/Runtime/Engine/Private/Components/CharacterMovementComponent.cpp
new file mode 100644
--- /dev/null
+++ b/Engine/Source/Runtime/Engine/Private/Components/CharacterMovementComponent.cpp
@@ -0,0 +1,7321 @@
+// Copyright 1998-2015 Epic Games, Inc. All Rights Reserved.
+
+/*=============================================================================
+	Movement.cpp: Character movement implementation
+
+=============================================================================*/
+
+#include "EnginePrivate.h"
+#include "GameFramework/PhysicsVolume.h"
+#include "GameFramework/GameNetworkManager.h"
+#include "GameFramework/Character.h"
+#include "GameFramework/CharacterMovementComponent.h"
+#include "GameFramework/GameState.h"
+#include "Components/PrimitiveComponent.h"
+#include "Animation/AnimMontage.h"
+#include "PhysicsEngine/DestructibleActor.h"
+
+// @todo this is here only due to circular dependency to AIModule. To be removed
+#include "Navigation/PathFollowingComponent.h"
+#include "AI/Navigation/AvoidanceManager.h"
+#include "Components/CapsuleComponent.h"
+#include "Components/BrushComponent.h"
+#include "Components/DestructibleComponent.h"
+
+DEFINE_LOG_CATEGORY_STATIC(LogCharacterMovement, Log, All);
+DEFINE_LOG_CATEGORY_STATIC(LogNavMeshMovement, Log, All);
+
+/**
+ * Character stats
+ */
+DECLARE_STATS_GROUP(TEXT("Character"), STATGROUP_Character, STATCAT_Advanced);
+
+DECLARE_CYCLE_STAT(TEXT("Char Movement Tick"), STAT_CharacterMovementTick, STATGROUP_Character);
+DECLARE_CYCLE_STAT(TEXT("Char Movement Authority Time"), STAT_CharacterMovementAuthority, STATGROUP_Character);
+DECLARE_CYCLE_STAT(TEXT("Char Movement Simulated Time"), STAT_CharacterMovementSimulated, STATGROUP_Character);
+DECLARE_CYCLE_STAT(TEXT("Char Physics Interation"), STAT_CharPhysicsInteraction, STATGROUP_Character);
+DECLARE_CYCLE_STAT(TEXT("StepUp"), STAT_CharStepUp, STATGROUP_Character);
+DECLARE_CYCLE_STAT(TEXT("Char Update Acceleration"), STAT_CharUpdateAcceleration, STATGROUP_Character);
+DECLARE_CYCLE_STAT(TEXT("Char MoveUpdateDelegate"), STAT_CharMoveUpdateDelegate, STATGROUP_Character);
+DECLARE_CYCLE_STAT(TEXT("PhysWalking"), STAT_CharPhysWalking, STATGROUP_Character);
+DECLARE_CYCLE_STAT(TEXT("PhysFalling"), STAT_CharPhysFalling, STATGROUP_Character);
+DECLARE_CYCLE_STAT(TEXT("PhysNavWalking"), STAT_CharPhysNavWalking, STATGROUP_Character);
+DECLARE_CYCLE_STAT(TEXT("NavProjectPoint"), STAT_CharNavProjectPoint, STATGROUP_Character);
+DECLARE_CYCLE_STAT(TEXT("NavProjectLocation"), STAT_CharNavProjectLocation, STATGROUP_Character);
+
+// MAGIC NUMBERS
+const float MAX_STEP_SIDE_Z = 0.08f;	// maximum z value for the normal on the vertical side of steps
+const float SWIMBOBSPEED = -80.f;
+const float VERTICAL_SLOPE_NORMAL_Z = 0.001f; // Slope is vertical if Abs(Normal.Z) <= this threshold. Accounts for precision problems that sometimes angle normals slightly off horizontal for vertical surface.
+
+const float UCharacterMovementComponent::MIN_TICK_TIME = 0.0002f;
+const float UCharacterMovementComponent::MIN_FLOOR_DIST = 1.9f;
+const float UCharacterMovementComponent::MAX_FLOOR_DIST = 2.4f;
+const float UCharacterMovementComponent::BRAKE_TO_STOP_VELOCITY = 10.f;
+const float UCharacterMovementComponent::SWEEP_EDGE_REJECT_DISTANCE = 0.15f;
+
+// CVars
+static TAutoConsoleVariable<int32> CVarNetEnableMoveCombining(
+	TEXT("p.NetEnableMoveCombining"),
+	1,
+	TEXT("Whether to enable move combining on the client to reduce bandwidth by combining similar moves.\n")
+	TEXT("0: Disable, 1: Enable"),
+	ECVF_Default);
+
+static TAutoConsoleVariable<float> CVarNetProxyShrinkRadius(
+	TEXT("p.NetProxyShrinkRadius"),
+	0.01f,
+	TEXT("Shrink simulated proxy capsule radius by this amount, to account for network rounding that may cause encroachment.\n")
+	TEXT("Changing this value at runtime may require the proxy to re-join for correct behavior.\n")
+	TEXT("<= 0: disabled, > 0: shrink by this amount."),
+	ECVF_Default);
+
+static TAutoConsoleVariable<float> CVarNetProxyShrinkHalfHeight(
+	TEXT("p.NetProxyShrinkHalfHeight"),
+	0.01f,
+	TEXT("Shrink simulated proxy capsule half height by this amount, to account for network rounding that may cause encroachment.\n")
+	TEXT("Changing this value at runtime may require the proxy to re-join for correct behavior.\n")
+	TEXT("<= 0: disabled, > 0: shrink by this amount."),
+	ECVF_Default);
+
+#if !(UE_BUILD_SHIPPING || UE_BUILD_TEST)
+static TAutoConsoleVariable<int32> CVarNetShowCorrections(
+	TEXT("p.NetShowCorrections"),
+	0,
+	TEXT("Whether to draw client position corrections (red is incorrect, green is corrected).\n")
+	TEXT("0: Disable, 1: Enable"),
+	ECVF_Cheat);
+
+static TAutoConsoleVariable<float> CVarNetCorrectionLifetime(
+	TEXT("p.NetCorrectionLifetime"),
+	4.f,
+	TEXT("How long a visualized network correction persists.\n")
+	TEXT("Time in seconds each visualized network correction persists."),
+	ECVF_Cheat);
+#endif // !(UE_BUILD_SHIPPING || UE_BUILD_TEST)
+
+// Version that does not use inverse sqrt estimate, for higher precision.
+FORCEINLINE FVector GetSafeNormalPrecise(const FVector& V)
+{
+	const float VSq = V.SizeSquared();
+	if (VSq < SMALL_NUMBER)
+	{
+		return FVector::ZeroVector;
+	}
+	else
+	{
+		return V * (1.f / FMath::Sqrt(VSq));
+	}
+}
+
+// Version that does not use inverse sqrt estimate, for higher precision.
+FORCEINLINE FVector GetClampedToMaxSizePrecise(const FVector& V, float MaxSize)
+{
+	if (MaxSize < KINDA_SMALL_NUMBER)
+	{
+		return FVector::ZeroVector;
+	}
+
+	const float VSq = V.SizeSquared();
+	if (VSq > FMath::Square(MaxSize))
+	{
+		return V * (MaxSize / FMath::Sqrt(VSq));
+	}
+	else
+	{
+		return V;
+	}
+}
+
+
+void FFindFloorResult::SetFromSweep(const FHitResult& InHit, const float InSweepFloorDist, const bool bIsWalkableFloor)
+{
+	bBlockingHit = InHit.IsValidBlockingHit();
+	bWalkableFloor = bIsWalkableFloor;
+	bLineTrace = false;
+	FloorDist = InSweepFloorDist;
+	LineDist = 0.f;
+	HitResult = InHit;
+}
+
+void FFindFloorResult::SetFromLineTrace(const FHitResult& InHit, const float InSweepFloorDist, const float InLineDist, const bool bIsWalkableFloor)
+{
+	// We require a sweep that hit if we are going to use a line result.
+	check(HitResult.bBlockingHit);
+	if (HitResult.bBlockingHit && InHit.bBlockingHit)
+	{
+		// Override most of the sweep result with the line result, but save some values
+		FHitResult OldHit(HitResult);
+		HitResult = InHit;
+
+		// Restore some of the old values. We want the new normals and hit actor, however.
+		HitResult.Time = OldHit.Time;
+		HitResult.ImpactPoint = OldHit.ImpactPoint;
+		HitResult.Location = OldHit.Location;
+		HitResult.TraceStart = OldHit.TraceStart;
+		HitResult.TraceEnd = OldHit.TraceEnd;
+
+		bLineTrace = true;
+		LineDist = InLineDist;
+		bWalkableFloor = bIsWalkableFloor;
+	}
+}
+
+void FCharacterMovementComponentPreClothTickFunction::ExecuteTick(float DeltaTime, enum ELevelTick TickType, ENamedThreads::Type CurrentThread, const FGraphEventRef& MyCompletionGraphEvent)
+{
+	if ( (TickType == LEVELTICK_All) && Target && !Target->HasAnyFlags(RF_PendingKill | RF_Unreachable))
+	{
+		FScopeCycleCounterUObject ComponentScope(Target);
+		FScopeCycleCounterUObject AdditionalScope(Target->AdditionalStatObject());
+		Target->PreClothTick(DeltaTime, *this);	
+	}
+}
+
+FString FCharacterMovementComponentPreClothTickFunction::DiagnosticMessage()
+{
+	return Target->GetFullName() + TEXT("[UCharacterMovementComponent::PreClothTick]");
+}
+
+UCharacterMovementComponent::UCharacterMovementComponent(const FObjectInitializer& ObjectInitializer)
+	: Super(ObjectInitializer)
+{
+	PreClothComponentTick.bCanEverTick = true;
+	PreClothComponentTick.bStartWithTickEnabled = false;
+	PreClothComponentTick.TickGroup = TG_PreCloth;
+
+	GravityScale = 1.f;
+	GroundFriction = 8.0f;
+	JumpZVelocity = 420.0f;
+	JumpOffJumpZFactor = 0.5f;
+	RotationRate = FRotator(0.f, 360.0f, 0.0f);
+	SetWalkableFloorZ(0.71f);
+
+	MaxStepHeight = 45.0f;
+	PerchRadiusThreshold = 0.0f;
+	PerchAdditionalHeight = 40.f;
+
+	MaxFlySpeed = 600.0f;
+	MaxWalkSpeed = 600.0f;
+	MaxSwimSpeed = 300.0f;
+	MaxCustomMovementSpeed = MaxWalkSpeed;
+	
+	MaxSimulationTimeStep = 0.05f;
+	MaxSimulationIterations = 8;
+
+	CrouchedSpeedMultiplier_DEPRECATED = 0.5f;
+	MaxWalkSpeedCrouched = MaxWalkSpeed * CrouchedSpeedMultiplier_DEPRECATED;
+	MaxOutOfWaterStepHeight = 40.0f;
+	OutofWaterZ = 420.0f;
+	AirControl = 0.05f;
+	AirControlBoostMultiplier = 2.f;
+	AirControlBoostVelocityThreshold = 25.f;
+	FallingLateralFriction = 0.f;
+	MaxAcceleration = 2048.0f;
+	BrakingDecelerationWalking = MaxAcceleration;
+	BrakingDecelerationFalling = 0.f;
+	BrakingDecelerationFlying = 0.f;
+	BrakingDecelerationSwimming = 0.f;
+	LedgeCheckThreshold = 4.0f;
+	JumpOutOfWaterPitch = 11.25f;
+	UpperImpactNormalScale_DEPRECATED = 0.5f;
+	
+	Mass = 100.0f;
+	bJustTeleported = true;
+	CrouchedHalfHeight = 40.0f;
+	Buoyancy = 1.0f;
+	PendingLaunchVelocity = FVector::ZeroVector;
+	DefaultWaterMovementMode = MOVE_Swimming;
+	DefaultLandMovementMode = MOVE_Walking;
+	GroundMovementMode = MOVE_Walking;
+	bForceNextFloorCheck = true;
+	bForceBraking_DEPRECATED = false;
+	bShrinkProxyCapsule = true;
+	bCanWalkOffLedges = true;
+	bCanWalkOffLedgesWhenCrouching = false;
+	bWantsToLeaveNavWalking = false;
+
+	bEnablePhysicsInteraction = true;
+	StandingDownwardForceScale = 1.0f;
+	InitialPushForceFactor = 500.0f;
+	PushForceFactor = 750000.0f;
+	PushForcePointZOffsetFactor = -0.75f;
+	bPushForceScaledToMass = false;
+	bScalePushForceToVelocity = true;
+	
+	TouchForceFactor = 1.0f;
+	bTouchForceScaledToMass = true;
+	MinTouchForce = -1.0f;
+	MaxTouchForce = 250.0f;
+	RepulsionForce = 2.5f;
+
+	bUseControllerDesiredRotation = false;
+
+	bMaintainHorizontalGroundVelocity = true;
+	bImpartBaseVelocityX = true;
+	bImpartBaseVelocityY = true;
+	bImpartBaseVelocityZ = true;
+	bImpartBaseAngularVelocity = true;
+	bAlwaysCheckFloor = true;
+
+	// default character can jump, walk, and swim
+	NavAgentProps.bCanJump = true;
+	NavAgentProps.bCanWalk = true;
+	NavAgentProps.bCanSwim = true;
+	ResetMoveState();
+
+	ClientPredictionData = NULL;
+	ServerPredictionData = NULL;
+
+	// This should be greater than tolerated player timeout * 2.
+	MinTimeBetweenTimeStampResets = 4.f * 60.f; 
+
+	bEnableScopedMovementUpdates = true;
+
+	bRequestedMoveUseAcceleration = true;
+	bUseRVOAvoidance = false;
+	bUseRVOPostProcess = false;
+	AvoidanceLockVelocity = FVector::ZeroVector;
+	AvoidanceLockTimer = 0.0f;
+	AvoidanceGroup.bGroup0 = true;
+	GroupsToAvoid.Packed = 0xFFFFFFFF;
+	GroupsToIgnore.Packed = 0;
+	AvoidanceConsiderationRadius = 500.0f;
+
+	OldBaseQuat = FQuat::Identity;
+	OldBaseLocation = FVector::ZeroVector;
+
+	NavMeshProjectionInterval = 0.1f;
+	NavMeshProjectionInterpSpeed = 12.f;
+	NavMeshProjectionHeightScaleUp = 0.67f;
+	NavMeshProjectionHeightScaleDown = 1.0f;
+}
+
+void UCharacterMovementComponent::PostLoad()
+{
+	Super::PostLoad();
+
+	const int32 LinkerUE4Ver = GetLinkerUE4Version();
+
+	if (LinkerUE4Ver < VER_UE4_CHARACTER_MOVEMENT_DECELERATION)
+	{
+		BrakingDecelerationWalking = MaxAcceleration;
+	}
+
+	if (LinkerUE4Ver < VER_UE4_CHARACTER_BRAKING_REFACTOR)
+	{
+		// This bool used to apply walking braking in flying and swimming modes.
+		if (bForceBraking_DEPRECATED)
+		{
+			BrakingDecelerationFlying = BrakingDecelerationWalking;
+			BrakingDecelerationSwimming = BrakingDecelerationWalking;
+		}
+	}
+
+	if (LinkerUE4Ver < VER_UE4_CHARACTER_MOVEMENT_WALKABLE_FLOOR_REFACTOR)
+	{
+		// Compute the walkable floor angle, since we have never done so yet.
+		UCharacterMovementComponent::SetWalkableFloorZ(WalkableFloorZ);
+	}
+
+	if (LinkerUE4Ver < VER_UE4_DEPRECATED_MOVEMENTCOMPONENT_MODIFIED_SPEEDS)
+	{
+		MaxWalkSpeedCrouched = MaxWalkSpeed * CrouchedSpeedMultiplier_DEPRECATED;
+		MaxCustomMovementSpeed = MaxWalkSpeed;
+	}
+}
+
+
+#if WITH_EDITOR
+void UCharacterMovementComponent::PostEditChangeProperty(FPropertyChangedEvent& PropertyChangedEvent)
+{
+	Super::PostEditChangeProperty(PropertyChangedEvent);
+
+	const UProperty* PropertyThatChanged = PropertyChangedEvent.MemberProperty;
+	if (PropertyThatChanged && PropertyThatChanged->GetFName() == GET_MEMBER_NAME_CHECKED(UCharacterMovementComponent, WalkableFloorAngle))
+	{
+		// Compute WalkableFloorZ from the Angle.
+		SetWalkableFloorAngle(WalkableFloorAngle);
+	}
+}
+#endif // WITH_EDITOR
+
+
+void UCharacterMovementComponent::OnRegister()
+{
+	Super::OnRegister();
+
+#if WITH_EDITOR
+	// Compute WalkableFloorZ from the WalkableFloorAngle.
+	// This is only to respond to changes propagated by PostEditChangeProperty, so it's only done in the editor.
+	SetWalkableFloorAngle(WalkableFloorAngle);
+#endif
+}
+
+
+void UCharacterMovementComponent::BeginDestroy()
+{
+	if (ClientPredictionData)
+	{
+		delete ClientPredictionData;
+		ClientPredictionData = NULL;
+	}
+	
+	if (ServerPredictionData)
+	{
+		delete ServerPredictionData;
+		ServerPredictionData = NULL;
+	}
+	
+	Super::BeginDestroy();
+}
+
+void UCharacterMovementComponent::SetUpdatedComponent(USceneComponent* NewUpdatedComponent)
+{
+	if (NewUpdatedComponent)
+	{
+		const ACharacter* NewCharacterOwner = Cast<ACharacter>(NewUpdatedComponent->GetOwner());
+		if (NewCharacterOwner == NULL)
+		{
+			UE_LOG(LogCharacterMovement, Error, TEXT("%s owned by %s must update a component owned by a Character"), *GetName(), *GetNameSafe(NewUpdatedComponent->GetOwner()));
+			return;
+		}
+
+		// check that UpdatedComponent is a Capsule
+		if (Cast<UCapsuleComponent>(NewUpdatedComponent) == NULL)
+		{
+			UE_LOG(LogCharacterMovement, Error, TEXT("%s owned by %s must update a capsule component"), *GetName(), *GetNameSafe(NewUpdatedComponent->GetOwner()));
+			return;
+		}
+	}
+
+	if ( bMovementInProgress )
+	{
+		// failsafe to avoid crashes in CharacterMovement. 
+		bDeferUpdateMoveComponent = true;
+		DeferredUpdatedMoveComponent = NewUpdatedComponent;
+		return;
+	}
+	bDeferUpdateMoveComponent = false;
+	DeferredUpdatedMoveComponent = NULL;
+
+	UPrimitiveComponent* OldPrimitive = Cast<UPrimitiveComponent>(UpdatedComponent);
+	if (IsValid(OldPrimitive) && OldPrimitive->OnComponentBeginOverlap.IsBound())
+	{
+		OldPrimitive->OnComponentBeginOverlap.RemoveDynamic(this, &UCharacterMovementComponent::CapsuleTouched);
+	}
+	
+	Super::SetUpdatedComponent(NewUpdatedComponent);
+	CharacterOwner = Cast<ACharacter>(PawnOwner);
+
+	if (UpdatedComponent == NULL)
+	{
+		StopActiveMovement();
+	}
+
+	if (IsValid(UpdatedPrimitive) && bEnablePhysicsInteraction)
+	{
+		UpdatedPrimitive->OnComponentBeginOverlap.AddUniqueDynamic(this, &UCharacterMovementComponent::CapsuleTouched);
+	}
+
+	if (bUseRVOAvoidance)
+	{
+		UAvoidanceManager* AvoidanceManager = GetWorld()->GetAvoidanceManager();
+		if (AvoidanceManager)
+		{
+			AvoidanceManager->RegisterMovementComponent(this, AvoidanceWeight);
+		}
+	}
+}
+
+bool UCharacterMovementComponent::HasValidData() const
+{
+	return UpdatedComponent && IsValid(CharacterOwner);
+}
+
+FCollisionShape UCharacterMovementComponent::GetPawnCapsuleCollisionShape(const EShrinkCapsuleExtent ShrinkMode, const float CustomShrinkAmount) const
+{
+	FVector Extent = GetPawnCapsuleExtent(ShrinkMode, CustomShrinkAmount);
+	return FCollisionShape::MakeCapsule(Extent);
+}
+
+FVector UCharacterMovementComponent::GetPawnCapsuleExtent(const EShrinkCapsuleExtent ShrinkMode, const float CustomShrinkAmount) const
+{
+	check(CharacterOwner);
+
+	float Radius, HalfHeight;
+	CharacterOwner->GetCapsuleComponent()->GetScaledCapsuleSize(Radius, HalfHeight);
+	FVector CapsuleExtent(Radius, Radius, HalfHeight);
+
+	float RadiusEpsilon = 0.f;
+	float HeightEpsilon = 0.f;
+
+	switch(ShrinkMode)
+	{
+	case SHRINK_None:
+		return CapsuleExtent;
+
+	case SHRINK_RadiusCustom:
+		RadiusEpsilon = CustomShrinkAmount;
+		break;
+
+	case SHRINK_HeightCustom:
+		HeightEpsilon = CustomShrinkAmount;
+		break;
+		
+	case SHRINK_AllCustom:
+		RadiusEpsilon = CustomShrinkAmount;
+		HeightEpsilon = CustomShrinkAmount;
+		break;
+
+	default:
+		UE_LOG(LogCharacterMovement, Warning, TEXT("Unknown EShrinkCapsuleExtent in UCharacterMovementComponent::GetCapsuleExtent"));
+		break;
+	}
+
+	// Don't shrink to zero extent.
+	const float MinExtent = KINDA_SMALL_NUMBER * 10.f;
+	CapsuleExtent.X = FMath::Max(CapsuleExtent.X - RadiusEpsilon, MinExtent);
+	CapsuleExtent.Y = CapsuleExtent.X;
+	CapsuleExtent.Z = FMath::Max(CapsuleExtent.Z - HeightEpsilon, MinExtent);
+
+	return CapsuleExtent;
+}
+
+
+bool UCharacterMovementComponent::DoJump(bool bReplayingMoves)
+{
+	if ( CharacterOwner && CharacterOwner->CanJump() )
+	{
+		// Don't jump if we can't move up/down.
+		if (!bConstrainToPlane || FMath::Abs(PlaneConstraintNormal.Z) != 1.f)
+		{
+			Velocity.Z = JumpZVelocity;
+			SetMovementMode(MOVE_Falling);
+			return true;
+		}
+	}
+	
+	return false;
+}
+
+FVector UCharacterMovementComponent::GetImpartedMovementBaseVelocity() const
+{
+	FVector Result = FVector::ZeroVector;
+	if (CharacterOwner)
+	{
+		UPrimitiveComponent* MovementBase = CharacterOwner->GetMovementBase();
+		if (MovementBaseUtility::IsDynamicBase(MovementBase))
+		{
+			FVector BaseVelocity = MovementBaseUtility::GetMovementBaseVelocity(MovementBase, CharacterOwner->GetBasedMovement().BoneName);
+			
+			if (bImpartBaseAngularVelocity)
+			{
+				const FVector CharacterBasePosition = (UpdatedComponent->GetComponentLocation() - FVector(0.f, 0.f, CharacterOwner->GetCapsuleComponent()->GetScaledCapsuleHalfHeight()));
+				const FVector BaseTangentialVel = MovementBaseUtility::GetMovementBaseTangentialVelocity(MovementBase, CharacterOwner->GetBasedMovement().BoneName, CharacterBasePosition);
+				BaseVelocity += BaseTangentialVel;
+			}
+
+			if (bImpartBaseVelocityX)
+			{
+				Result.X = BaseVelocity.X;
+			}
+			if (bImpartBaseVelocityY)
+			{
+				Result.Y = BaseVelocity.Y;
+			}
+			if (bImpartBaseVelocityZ)
+			{
+				Result.Z = BaseVelocity.Z;
+			}
+		}
+	}
+	
+	return Result;
+}
+
+void UCharacterMovementComponent::Launch(FVector const& LaunchVel)
+{
+	PendingLaunchVelocity = LaunchVel;
+}
+
+bool UCharacterMovementComponent::HandlePendingLaunch()
+{
+	if (!PendingLaunchVelocity.IsZero() && HasValidData())
+	{
+		Velocity = PendingLaunchVelocity;
+		SetMovementMode(MOVE_Falling);
+		PendingLaunchVelocity = FVector::ZeroVector;
+		return true;
+	}
+
+	return false;
+}
+
+void UCharacterMovementComponent::JumpOff(AActor* MovementBaseActor)
+{
+	if ( !bPerformingJumpOff )
+	{
+		bPerformingJumpOff = true;
+		if ( CharacterOwner )
+		{
+			const float MaxSpeed = GetMaxSpeed() * 0.85f;
+			Velocity += MaxSpeed * GetBestDirectionOffActor(MovementBaseActor);
+			if ( Velocity.Size2D() > MaxSpeed )
+			{
+				Velocity = MaxSpeed * Velocity.GetSafeNormal();
+			}
+			Velocity.Z = JumpOffJumpZFactor * JumpZVelocity;
+			SetMovementMode(MOVE_Falling);
+		}
+		bPerformingJumpOff = false;
+	}
+}
+
+FVector UCharacterMovementComponent::GetBestDirectionOffActor(AActor* BaseActor) const
+{
+	// By default, just pick a random direction.  Derived character classes can choose to do more complex calculations,
+	// such as finding the shortest distance to move in based on the BaseActor's Bounding Volume.
+	const float RandAngle = FMath::DegreesToRadians(GetNetworkSafeRandomAngleDegrees());
+	return FVector(FMath::Cos(RandAngle), FMath::Sin(RandAngle), 0.5f).GetSafeNormal();
+}
+
+float UCharacterMovementComponent::GetNetworkSafeRandomAngleDegrees() const
+{
+	float Angle = FMath::SRand() * 360.f;
+
+	if (GetNetMode() > NM_Standalone)
+	{
+		// Networked game
+		// Get a timestamp that is relatively close between client and server (within ping).
+		FNetworkPredictionData_Server_Character const* ServerData = (HasPredictionData_Server() ? GetPredictionData_Server_Character() : NULL);
+		FNetworkPredictionData_Client_Character const* ClientData = (HasPredictionData_Client() ? GetPredictionData_Client_Character() : NULL);
+
+		float TimeStamp = Angle;
+		if (ServerData)
+		{
+			TimeStamp = ServerData->CurrentClientTimeStamp;
+		}
+		else if (ClientData)
+		{
+			TimeStamp = ClientData->CurrentTimeStamp;
+		}
+		
+		// Convert to degrees with a faster period.
+		const float PeriodMult = 8.0f;
+		Angle = TimeStamp * PeriodMult;
+		Angle = FMath::Fmod(Angle, 360.f);
+	}
+
+	return Angle;
+}
+
+
+void UCharacterMovementComponent::SetDefaultMovementMode()
+{
+	// check for water volume
+	if ( IsInWater() && CanEverSwim() )
+	{
+		SetMovementMode(DefaultWaterMovementMode);
+	}
+	else if ( !CharacterOwner || MovementMode != DefaultLandMovementMode )
+	{
+		SetMovementMode(DefaultLandMovementMode);
+
+		// Avoid 1-frame delay if trying to walk but walking fails at this location.
+		if (MovementMode == MOVE_Walking && GetMovementBase() == NULL)
+		{
+			SetMovementMode(MOVE_Falling);
+		}
+	}
+}
+
+void UCharacterMovementComponent::SetMovementMode(EMovementMode NewMovementMode, uint8 NewCustomMode)
+{
+	if (NewMovementMode != MOVE_Custom)
+	{
+		NewCustomMode = 0;
+	}
+
+	// Do nothing if nothing is changing.
+	if (MovementMode == NewMovementMode)
+	{
+		// Allow changes in custom sub-mode.
+		if ((NewMovementMode != MOVE_Custom) || (NewCustomMode == CustomMovementMode))
+		{
+			return;
+		}
+	}
+
+	const EMovementMode PrevMovementMode = MovementMode;
+	const uint8 PrevCustomMode = CustomMovementMode;
+
+	MovementMode = NewMovementMode;
+	CustomMovementMode = NewCustomMode;
+
+	// We allow setting movement mode before we have a component to update, in case this happens at startup.
+	if (!HasValidData())
+	{
+		return;
+	}
+	
+	// Handle change in movement mode
+	OnMovementModeChanged(PrevMovementMode, PrevCustomMode);
+
+	// @todo UE4 do we need to disable ragdoll physics here? Should this function do nothing if in ragdoll?
+}
+
+
+void UCharacterMovementComponent::OnMovementModeChanged(EMovementMode PreviousMovementMode, uint8 PreviousCustomMode)
+{
+	if (!HasValidData())
+	{
+		return;
+	}
+
+	// Update collision settings if needed
+	if (MovementMode == MOVE_NavWalking)
+	{
+		SetNavWalkingPhysics(true);
+		GroundMovementMode = MovementMode;
+		// Walking uses only XY velocity
+		Velocity.Z = 0.f;
+	}
+	else if (PreviousMovementMode == MOVE_NavWalking)
+	{
+		if (MovementMode == DefaultLandMovementMode || IsWalking())
+		{
+			const bool bCanSwitchMode = TryToLeaveNavWalking();
+			if (!bCanSwitchMode)
+			{
+				SetMovementMode(MOVE_NavWalking);
+				return;
+			}
+		}
+		else
+		{
+			SetNavWalkingPhysics(false);
+		}
+	}
+
+	// React to changes in the movement mode.
+	if (MovementMode == MOVE_Walking)
+	{
+		// Walking uses only XY velocity, and must be on a walkable floor, with a Base.
+		Velocity.Z = 0.f;
+		bCrouchMaintainsBaseLocation = true;
+		GroundMovementMode = MovementMode;
+
+		// make sure we update our new floor/base on initial entry of the walking physics
+		FindFloor(UpdatedComponent->GetComponentLocation(), CurrentFloor, false);
+		AdjustFloorHeight();
+		SetBase(CurrentFloor.HitResult.Component.Get(), CurrentFloor.HitResult.BoneName);
+	}
+	else
+	{
+		CurrentFloor.Clear();
+		bCrouchMaintainsBaseLocation = false;
+
+		if (MovementMode == MOVE_Falling)
+		{
+			Velocity += GetImpartedMovementBaseVelocity();
+			CharacterOwner->Falling();
+		}
+
+		SetBase(NULL);
+
+		if (MovementMode == MOVE_None)
+		{
+			// Kill velocity and clear queued up events
+			StopMovementKeepPathing();
+			CharacterOwner->ClearJumpInput();
+		}
+	}
+
+	CharacterOwner->OnMovementModeChanged(PreviousMovementMode, PreviousCustomMode);
+};
+
+
+uint8 UCharacterMovementComponent::PackNetworkMovementMode() const
+{
+	const uint32 GroundShift = FMath::FloorLog2(MOVE_MAX) + 1;
+	const uint8 CustomModeThr = 2 * (1 << GroundShift);
+
+	if (MovementMode != MOVE_Custom)
+	{
+		return uint8(MovementMode.GetValue()) | (GroundMovementMode.GetValue() << GroundShift);
+	}
+	else
+	{
+		return CustomMovementMode + CustomModeThr;
+	}
+}
+
+
+void UCharacterMovementComponent::UnpackNetworkMovementMode(const uint8 ReceivedMode, TEnumAsByte<EMovementMode>& OutMode, uint8& OutCustomMode, TEnumAsByte<EMovementMode>& OutGroundMode) const
+{
+	const uint32 GroundShift = FMath::FloorLog2(MOVE_MAX) + 1;
+	const uint8 CustomModeThr = 2 * (1 << GroundShift);
+
+	if (ReceivedMode < CustomModeThr)
+	{
+		const uint8 GroundMask = (1 << GroundShift) - 1;
+
+		OutMode = TEnumAsByte<EMovementMode>(ReceivedMode & GroundMask);
+		OutCustomMode = 0;
+		OutGroundMode = TEnumAsByte<EMovementMode>(ReceivedMode >> GroundShift);
+	}
+	else
+	{
+		OutMode = MOVE_Custom;
+		OutCustomMode = ReceivedMode - CustomModeThr;
+		OutGroundMode = DefaultLandMovementMode;
+	}
+}
+
+
+void UCharacterMovementComponent::ApplyNetworkMovementMode(const uint8 ReceivedMode)
+{
+	TEnumAsByte<EMovementMode> NetMovementMode(MOVE_None);
+	TEnumAsByte<EMovementMode> NetGroundMode(MOVE_None);
+	uint8 NetCustomMode(0);
+	UnpackNetworkMovementMode(ReceivedMode, NetMovementMode, NetCustomMode, NetGroundMode);
+
+	GroundMovementMode = NetGroundMode;
+	SetMovementMode(NetMovementMode, NetCustomMode);
+}
+
+// TODO: deprecated, remove
+void UCharacterMovementComponent::PerformAirControl(FVector Direction, float ZDiff)
+{
+	PerformAirControlForPathFollowing(Direction, ZDiff);
+}
+
+
+void UCharacterMovementComponent::PerformAirControlForPathFollowing(FVector Direction, float ZDiff)
+{
+	// use air control if low grav or above destination and falling towards it
+	if ( CharacterOwner && Velocity.Z < 0.f && (ZDiff < 0.f || GetGravityZ() > 0.9f * GetWorld()->GetDefaultGravityZ()))
+	{
+		if ( ZDiff > 0.f )
+		{
+			if ( ZDiff > 2.f * GetMaxJumpHeight() )
+			{
+				if (PathFollowingComp.IsValid())
+				{
+					PathFollowingComp->AbortMove(TEXT("missed jump"));
+				}
+			}
+		}
+		else
+		{
+			if ( (Velocity.X == 0.f) && (Velocity.Y == 0.f) )
+			{
+				Acceleration = FVector::ZeroVector;
+			}
+			else
+			{
+				float Dist2D = Direction.Size2D();
+				//Direction.Z = 0.f;
+				Acceleration = Direction.GetSafeNormal() * GetMaxAcceleration();
+
+				if ( (Dist2D < 0.5f * FMath::Abs(Direction.Z)) && ((Velocity | Direction) > 0.5f*FMath::Square(Dist2D)) )
+				{
+					Acceleration *= -1.f;
+				}
+
+				if ( Dist2D < 1.5f*CharacterOwner->GetCapsuleComponent()->GetScaledCapsuleRadius() )
+				{
+					Velocity.X = 0.f;
+					Velocity.Y = 0.f;
+					Acceleration = FVector::ZeroVector;
+				}
+				else if ( (Velocity | Direction) < 0.f )
+				{
+					float M = FMath::Max(0.f, 0.2f - GetWorld()->DeltaTimeSeconds);
+					Velocity.X *= M;
+					Velocity.Y *= M;
+				}
+			}
+		}
+	}
+}
+
+void UCharacterMovementComponent::TickComponent(float DeltaTime, enum ELevelTick TickType, FActorComponentTickFunction *ThisTickFunction)
+{
+	Super::TickComponent(DeltaTime, TickType, ThisTickFunction);
+
+	SCOPE_CYCLE_COUNTER(STAT_CharacterMovementTick);
+
+	const FVector InputVector = ConsumeInputVector();
+	if (!HasValidData() || ShouldSkipUpdate(DeltaTime))
+	{
+		return;
+	}
+
+	// See if we fell out of the world.
+	const bool bIsSimulatingPhysics = UpdatedComponent->IsSimulatingPhysics();
+	if (CharacterOwner->Role == ROLE_Authority)
+	{
+		if (!bCheatFlying || bIsSimulatingPhysics)
+		{
+			if (!CharacterOwner->CheckStillInWorld())
+			{
+				return;
+			}
+		}
+	}
+
+	// We don't update if simulating physics (eg ragdolls).
+	if (bIsSimulatingPhysics)
+	{
+		return;
+	}
+
+	if (AvoidanceLockTimer > 0.0f)
+	{
+		AvoidanceLockTimer -= DeltaTime;
+	}
+
+	if (CharacterOwner->Role > ROLE_SimulatedProxy)
+	{
+		// If we are a client we might have received an update from the server.
+		const bool bIsClient = (GetNetMode() == NM_Client && CharacterOwner->Role == ROLE_AutonomousProxy);
+		if (bIsClient)
+		{
+			ClientUpdatePositionAfterServerUpdate();
+		}
+
+		// Allow root motion to move characters that have no controller.
+		if( CharacterOwner->IsLocallyControlled() || (!CharacterOwner->Controller && bRunPhysicsWithNoController) || (!CharacterOwner->Controller && CharacterOwner->IsPlayingRootMotion()) )
+		{
+			{
+				SCOPE_CYCLE_COUNTER(STAT_CharUpdateAcceleration);
+
+				// We need to check the jump state before adjusting input acceleration, to minimize latency
+				// and to make sure acceleration respects our potentially new falling state.
+				CharacterOwner->CheckJumpInput(DeltaTime);
+
+				// apply input to acceleration
+				Acceleration = ScaleInputAcceleration(ConstrainInputAcceleration(InputVector));
+				AnalogInputModifier = ComputeAnalogInputModifier();
+			}
+
+			if (CharacterOwner->Role == ROLE_Authority)
+			{
+				PerformMovement(DeltaTime);
+			}
+			else if (bIsClient)
+			{
+				ReplicateMoveToServer(DeltaTime, Acceleration);
+			}
+		}
+		else if (CharacterOwner->GetRemoteRole() == ROLE_AutonomousProxy)
+		{
+			// Server ticking for remote client.
+			// Between net updates from the client we need to update position if based on another object,
+			// otherwise the object will move on intermediate frames and we won't follow it.
+			MaybeUpdateBasedMovement(DeltaTime);
+			SaveBaseLocation();
+		}
+	}
+	else if (CharacterOwner->Role == ROLE_SimulatedProxy)
+	{
+		AdjustProxyCapsuleSize();
+		SimulatedTick(DeltaTime);
+	}
+
+	UpdateDefaultAvoidance();
+
+	if (bEnablePhysicsInteraction)
+	{
+		SCOPE_CYCLE_COUNTER(STAT_CharPhysicsInteraction);
+
+		if (CurrentFloor.HitResult.IsValidBlockingHit())
+		{
+			// Apply downwards force when walking on top of physics objects
+			if (UPrimitiveComponent* BaseComp = CurrentFloor.HitResult.GetComponent())
+			{
+				if (StandingDownwardForceScale != 0.f && BaseComp->IsAnySimulatingPhysics())
+				{
+					const float GravZ = GetGravityZ();
+					const FVector ForceLocation = CurrentFloor.HitResult.ImpactPoint;
+					BaseComp->AddForceAtLocation(FVector(0.f, 0.f, GravZ * Mass * StandingDownwardForceScale), ForceLocation, CurrentFloor.HitResult.BoneName);
+				}
+			}
+		}
+
+		ApplyRepulsionForce(DeltaTime);
+	}
+}
+
+void UCharacterMovementComponent::PreClothTick(float DeltaTime, FCharacterMovementComponentPreClothTickFunction& ThisTickFunction)
+{
+	if(bDeferUpdateBasedMovement)
+	{
+		UpdateBasedMovement(DeltaTime);
+		SaveBaseLocation();
+
+		bDeferUpdateBasedMovement = false;
+	}
+}
+
+void UCharacterMovementComponent::AdjustProxyCapsuleSize()
+{
+	if (bShrinkProxyCapsule && CharacterOwner && CharacterOwner->Role == ROLE_SimulatedProxy)
+	{
+		bShrinkProxyCapsule = false;
+
+		float ShrinkRadius = FMath::Max(0.f, CVarNetProxyShrinkRadius.GetValueOnGameThread());
+		float ShrinkHalfHeight = FMath::Max(0.f, CVarNetProxyShrinkHalfHeight.GetValueOnGameThread());
+
+		if (ShrinkRadius == 0.f && ShrinkHalfHeight == 0.f)
+		{
+			return;
+		}
+
+		float Radius, HalfHeight;
+		CharacterOwner->GetCapsuleComponent()->GetUnscaledCapsuleSize(Radius, HalfHeight);
+		const float ComponentScale = CharacterOwner->GetCapsuleComponent()->GetShapeScale();
+
+		if (ComponentScale <= KINDA_SMALL_NUMBER)
+		{
+			return;
+		}
+
+		const float NewRadius = FMath::Max(0.f, Radius - ShrinkRadius / ComponentScale);
+		const float NewHalfHeight = FMath::Max(0.f, HalfHeight - ShrinkHalfHeight / ComponentScale);
+
+		if (NewRadius == 0.f || NewHalfHeight == 0.f)
+		{
+			UE_LOG(LogCharacterMovement, Warning, TEXT("Invalid attempt to shrink Proxy capsule for %s to zero dimension!"), *CharacterOwner->GetName());
+			return;
+		}
+
+		UE_LOG(LogCharacterMovement, Verbose, TEXT("Shrinking capsule for %s from (r=%.3f, h=%.3f) to (r=%.3f, h=%.3f)"), *CharacterOwner->GetName(),
+			Radius * ComponentScale, HalfHeight * ComponentScale, NewRadius * ComponentScale, NewHalfHeight * ComponentScale);
+
+		CharacterOwner->GetCapsuleComponent()->SetCapsuleSize(NewRadius, NewHalfHeight, true);
+		CharacterOwner->GetCapsuleComponent()->UpdateBounds();
+	}	
+}
+
+
+void UCharacterMovementComponent::SimulatedTick(float DeltaSeconds)
+{
+	SCOPE_CYCLE_COUNTER(STAT_CharacterMovement);
+	SCOPE_CYCLE_COUNTER(STAT_CharacterMovementSimulated);
+
+	// If we are playing a RootMotion AnimMontage.
+	if (CharacterOwner && CharacterOwner->IsPlayingNetworkedRootMotionMontage())
+	{
+		bWasSimulatingRootMotion = true;
+		UE_LOG(LogRootMotion, Verbose, TEXT("UCharacterMovementComponent::SimulatedTick"));
+
+		// Tick animations before physics.
+		if( CharacterOwner->GetMesh() )
+		{
+			TickCharacterPose(DeltaSeconds);
+
+			// Make sure animation didn't trigger an event that destroyed us
+			if (!HasValidData())
+			{
+				return;
+			}
+		}
+
+		if( RootMotionParams.bHasRootMotion )
+		{
+			const FRotator OldRotation = CharacterOwner->GetActorRotation();
+			const FVector OldLocation = CharacterOwner->GetActorLocation();
+			SimulateRootMotion(DeltaSeconds, RootMotionParams.RootMotionTransform);
+			// Root Motion has been used, clear
+			RootMotionParams.Clear();
+
+			// debug
+			if( false )
+			{
+				const FRotator NewRotation = CharacterOwner->GetActorRotation();
+				const FVector NewLocation = CharacterOwner->GetActorLocation();
+				DrawDebugCoordinateSystem(GetWorld(), CharacterOwner->GetMesh()->GetComponentLocation() + FVector(0,0,1), NewRotation, 50.f, false);
+				DrawDebugLine(GetWorld(), OldLocation, NewLocation, FColor::Red, true, 10.f);
+
+				UE_LOG(LogRootMotion, Log,  TEXT("UCharacterMovementComponent::SimulatedTick DeltaMovement Translation: %s, Rotation: %s, MovementBase: %s"),
+					*(NewLocation - OldLocation).ToCompactString(), *(NewRotation - OldRotation).GetNormalized().ToCompactString(), *GetNameSafe(CharacterOwner->GetMovementBase()) );
+			}
+		}
+
+		// then, once our position is up to date with our animation, 
+		// handle position correction if we have any pending updates received from the server.
+		if( CharacterOwner && (CharacterOwner->RootMotionRepMoves.Num() > 0) )
+		{
+			CharacterOwner->SimulatedRootMotionPositionFixup(DeltaSeconds);
+		}
+	}
+	else
+	{
+		// if we were simulating root motion, we've been ignoring regular ReplicatedMovement updates.
+		// If we're not simulating root motion anymore, force us to sync our movement properties.
+		// (Root Motion could leave Velocity out of sync w/ ReplicatedMovement)
+		if( bWasSimulatingRootMotion )
+		{
+			bWasSimulatingRootMotion = false;
+			if( CharacterOwner )
+			{
+				CharacterOwner->RootMotionRepMoves.Empty();
+				CharacterOwner->OnRep_ReplicatedMovement();
+				CharacterOwner->OnRep_ReplicatedBasedMovement();
+			}
+		}
+
+		if( UpdatedComponent->IsSimulatingPhysics() 
+			|| (CharacterOwner && CharacterOwner->IsMatineeControlled()) 
+			|| (CharacterOwner && CharacterOwner->IsPlayingRootMotion()))
+		{
+			PerformMovement(DeltaSeconds);
+		}
+		else
+		{
+			SimulateMovement(DeltaSeconds);
+		}
+	}
+
+	if( GetNetMode() == NM_Client )
+	{
+		SmoothClientPosition(DeltaSeconds);
+	}
+}
+
+void UCharacterMovementComponent::SimulateRootMotion(float DeltaSeconds, const FTransform& LocalRootMotionTransform)
+{
+	if( CharacterOwner && CharacterOwner->GetMesh() && (DeltaSeconds > 0.f) )
+	{
+		// Convert Local Space Root Motion to world space. Do it right before used by physics to make sure we use up to date transforms, as translation is relative to rotation.
+		const FTransform WorldSpaceRootMotionTransform = CharacterOwner->GetMesh()->ConvertLocalRootMotionToWorld(LocalRootMotionTransform);
+
+		// Compute root motion velocity to be used by physics
+		const FVector RootMotionVelocity = WorldSpaceRootMotionTransform.GetTranslation() / DeltaSeconds;
+		// Do not override Velocity.Z if falling.
+		Velocity = FVector(RootMotionVelocity.X, RootMotionVelocity.Y, (MovementMode == MOVE_Falling ? Velocity.Z : RootMotionVelocity.Z));
+
+		// Update replicated movement mode.
+		if (bNetworkMovementModeChanged)
+		{
+			bNetworkMovementModeChanged = false;
+			ApplyNetworkMovementMode(CharacterOwner->GetReplicatedMovementMode());
+		}
+
+		StartNewPhysics(DeltaSeconds, 0);
+		// fixme laurent - simulate movement seems to have step up issues? investigate as that would be cheaper to use.
+		// 		SimulateMovement(DeltaSeconds);
+
+		// Apply Root Motion rotation after movement is complete.
+		const FRotator RootMotionRotation = WorldSpaceRootMotionTransform.GetRotation().Rotator();
+		if( !RootMotionRotation.IsNearlyZero() )
+		{
+			const FRotator NewActorRotation = (CharacterOwner->GetActorRotation() + RootMotionRotation).GetNormalized();
+			MoveUpdatedComponent(FVector::ZeroVector, NewActorRotation, true);
+		}
+	}
+}
+
+
+void UCharacterMovementComponent::SimulateMovement(float DeltaSeconds)
+{
+	if (!HasValidData() || UpdatedComponent->Mobility != EComponentMobility::Movable || UpdatedComponent->IsSimulatingPhysics())
+	{
+		return;
+	}
+
+	const bool bIsSimulatedProxy = (CharacterOwner->Role == ROLE_SimulatedProxy);
+
+	// Workaround for replication not being updated initially
+	if (bIsSimulatedProxy &&
+		CharacterOwner->ReplicatedMovement.Location.IsZero() &&
+		CharacterOwner->ReplicatedMovement.Rotation.IsZero() &&
+		CharacterOwner->ReplicatedMovement.LinearVelocity.IsZero())
+	{
+		return;
+	}
+
+	// If base is not resolved on the client, we should not try to simulate at all
+	if (CharacterOwner->GetReplicatedBasedMovement().IsBaseUnresolved())
+	{
+		UE_LOG(LogCharacterMovement, Verbose, TEXT("Base for simulated character '%s' is not resolved on client, skipping SimulateMovement"), *CharacterOwner->GetName());
+		return;
+	}
+
+	FVector OldVelocity;
+	FVector OldLocation;
+
+	// Scoped updates can improve performance of multiple MoveComponent calls.
+	{
+		FScopedMovementUpdate ScopedMovementUpdate(UpdatedComponent, bEnableScopedMovementUpdates ? EScopedUpdate::DeferredUpdates : EScopedUpdate::ImmediateUpdates);
+
+		if (bIsSimulatedProxy)
+		{
+			// Handle network changes
+			if (bNetworkUpdateReceived)
+			{
+				bNetworkUpdateReceived = false;
+				if (bNetworkMovementModeChanged)
+				{
+					bNetworkMovementModeChanged = false;
+					ApplyNetworkMovementMode(CharacterOwner->GetReplicatedMovementMode());
+				}
+				else if (bJustTeleported)
+				{
+					// Make sure floor is current. We will continue using the replicated base, if there was one.
+					bJustTeleported = false;
+					UpdateFloorFromAdjustment();
+				}
+			}
+
+			HandlePendingLaunch();
+		}
+
+		if (MovementMode == MOVE_None)
+		{
+			return;
+		}
+
+		Acceleration = Velocity.GetSafeNormal();	// Not currently used for simulated movement
+		AnalogInputModifier = 1.0f;				// Not currently used for simulated movement
+
+		MaybeUpdateBasedMovement(DeltaSeconds);
+
+		// simulated pawns predict location
+		OldVelocity = Velocity;
+		OldLocation = UpdatedComponent->GetComponentLocation();
+		FStepDownResult StepDownResult;
+		MoveSmooth(Velocity, DeltaSeconds, &StepDownResult);
+
+		// consume path following requested velocity
+		bHasRequestedVelocity = false;
+
+		// if simulated gravity, find floor and check if falling
+		const bool bEnableFloorCheck = (!CharacterOwner->bSimGravityDisabled || !bIsSimulatedProxy);
+		if (bEnableFloorCheck && (IsMovingOnGround() || MovementMode == MOVE_Falling))
+		{
+			const FVector CollisionCenter = UpdatedComponent->GetComponentLocation();
+			if (StepDownResult.bComputedFloor)
+			{
+				CurrentFloor = StepDownResult.FloorResult;
+			}
+			else if (Velocity.Z <= 0.f)
+			{
+				FindFloor(CollisionCenter, CurrentFloor, Velocity.IsZero(), NULL);
+			}
+			else
+			{
+				CurrentFloor.Clear();
+			}
+
+			if (!CurrentFloor.IsWalkableFloor())
+			{
+				// No floor, must fall.
+				Velocity = NewFallVelocity(Velocity, FVector(0.f,0.f,GetGravityZ()), DeltaSeconds);
+				SetMovementMode(MOVE_Falling);
+			}
+			else
+			{
+				// Walkable floor
+				if (IsMovingOnGround())
+				{
+					AdjustFloorHeight();
+					SetBase(CurrentFloor.HitResult.Component.Get(), CurrentFloor.HitResult.BoneName);
+				}
+				else if (MovementMode == MOVE_Falling)
+				{
+					if (CurrentFloor.FloorDist <= MIN_FLOOR_DIST)
+					{
+						// Landed
+						SetMovementMode(MOVE_Walking);
+					}
+					else
+					{
+						// Continue falling.
+						Velocity = NewFallVelocity(Velocity, FVector(0.f,0.f,GetGravityZ()), DeltaSeconds);
+						CurrentFloor.Clear();
+					}
+				}
+			}
+		}
+
+		OnMovementUpdated(DeltaSeconds, OldLocation, OldVelocity);
+	} // End scoped movement update
+
+	// Call custom post-movement events. These happen after the scoped movement completes in case the events want to use the current state of overlaps etc.
+	CallMovementUpdateDelegate(DeltaSeconds, OldLocation, OldVelocity);
+
+	SaveBaseLocation();
+	UpdateComponentVelocity();
+	bJustTeleported = false;
+
+	LastUpdateLocation = UpdatedComponent ? UpdatedComponent->GetComponentLocation() : FVector::ZeroVector;
+}
+
+UPrimitiveComponent* UCharacterMovementComponent::GetMovementBase() const
+{
+	return CharacterOwner ? CharacterOwner->GetMovementBase() : NULL;
+}
+
+void UCharacterMovementComponent::SetBase( UPrimitiveComponent* NewBase, FName BoneName, bool bNotifyActor )
+{
+	if (CharacterOwner)
+	{
+		CharacterOwner->SetBase(NewBase, BoneName, bNotifyActor);
+	}
+}
+
+void UCharacterMovementComponent::MaybeUpdateBasedMovement(float DeltaSeconds)
+{
+	bDeferUpdateBasedMovement = false;
+
+	UPrimitiveComponent* MovementBase = CharacterOwner->GetMovementBase();
+	if (MovementBaseUtility::UseRelativeLocation(MovementBase))
+	{
+		const bool bBaseIsSimulatingPhysics = MovementBase->IsSimulatingPhysics();
+		
+		// Temporarily disabling deferred tick on skeletal mesh components that sim physics.
+		// We need to be consistent on when we read the bone locations for those, and while this reads
+		// the wrong location, the relative changes (which is what we care about) will be accurate.
+		const bool bAllowDefer = (bBaseIsSimulatingPhysics && !Cast<USkeletalMeshComponent>(MovementBase));
+		
+		if (!bBaseIsSimulatingPhysics || !bAllowDefer)
+		{
+			UpdateBasedMovement(DeltaSeconds);
+			PreClothComponentTick.SetTickFunctionEnable(false);
+		}
+		else
+		{
+			// defer movement base update until after physics
+			bDeferUpdateBasedMovement = true;
+			PreClothComponentTick.SetTickFunctionEnable(true);
+		}
+	}
+}
+
+// todo: deprecated, remove.
+void UCharacterMovementComponent::MaybeSaveBaseLocation()
+{
+	SaveBaseLocation();
+}
+
+// @todo UE4 - handle lift moving up and down through encroachment
+void UCharacterMovementComponent::UpdateBasedMovement(float DeltaSeconds)
+{
+	if (!HasValidData())
+	{
+		return;
+	}
+
+	const UPrimitiveComponent* MovementBase = CharacterOwner->GetMovementBase();
+	if (!MovementBaseUtility::UseRelativeLocation(MovementBase))
+	{
+		return;
+	}
+
+	if (!IsValid(MovementBase) || !IsValid(MovementBase->GetOwner()))
+	{
+		SetBase(NULL);
+		return;
+	}
+
+	// Ignore collision with bases during these movements.
+	TGuardValue<EMoveComponentFlags> ScopedFlagRestore(MoveComponentFlags, MoveComponentFlags | MOVECOMP_IgnoreBases);
+
+	FQuat DeltaQuat = FQuat::Identity;
+	FVector DeltaPosition = FVector::ZeroVector;
+
+	FQuat NewBaseQuat;
+	FVector NewBaseLocation;
+	if (!MovementBaseUtility::GetMovementBaseTransform(MovementBase, CharacterOwner->GetBasedMovement().BoneName, NewBaseLocation, NewBaseQuat))
+	{
+		return;
+	}
+
+	// Find change in rotation
+	const bool bRotationChanged = !OldBaseQuat.Equals(NewBaseQuat);
+	if (bRotationChanged)
+	{
+		DeltaQuat = NewBaseQuat * OldBaseQuat.Inverse();
+	}
+
+	// only if base moved
+	if (bRotationChanged || (OldBaseLocation != NewBaseLocation))
+	{
+		// Calculate new transform matrix of base actor (ignoring scale).
+		const FQuatRotationTranslationMatrix OldLocalToWorld(OldBaseQuat, OldBaseLocation);
+		const FQuatRotationTranslationMatrix NewLocalToWorld(NewBaseQuat, NewBaseLocation);
+
+		if( CharacterOwner->IsMatineeControlled() )
+		{
+			FRotationTranslationMatrix HardRelMatrix(CharacterOwner->GetBasedMovement().Rotation, CharacterOwner->GetBasedMovement().Location);
+			const FMatrix NewWorldTM = HardRelMatrix * NewLocalToWorld;
+			const FRotator NewWorldRot = bIgnoreBaseRotation ? CharacterOwner->GetActorRotation() : NewWorldTM.Rotator();
+			MoveUpdatedComponent( NewWorldTM.GetOrigin() - CharacterOwner->GetActorLocation(), NewWorldRot, true );
+		}
+		else
+		{
+			FQuat FinalQuat = CharacterOwner->GetActorQuat();
+			
+			if (bRotationChanged && !bIgnoreBaseRotation)
+			{
+				// Apply change in rotation and pipe through FaceRotation to maintain axis restrictions
+				const FQuat PawnOldQuat = CharacterOwner->GetActorQuat();
+				FinalQuat = DeltaQuat * FinalQuat;
+				CharacterOwner->FaceRotation(FinalQuat.Rotator(), 0.f);
+				FinalQuat = CharacterOwner->GetActorQuat();
+
+				// Pipe through ControlRotation, to affect camera.
+				if (CharacterOwner->Controller)
+				{
+					const FQuat PawnDeltaRotation = FinalQuat * PawnOldQuat.Inverse();
+					FRotator FinalRotation = FinalQuat.Rotator();
+					UpdateBasedRotation(FinalRotation, PawnDeltaRotation.Rotator());
+					FinalQuat = FinalRotation.Quaternion();
+				}
+			}
+
+			// We need to offset the base of the character here, not its origin, so offset by half height
+			float HalfHeight, Radius;
+			CharacterOwner->GetCapsuleComponent()->GetScaledCapsuleSize(Radius, HalfHeight);
+
+			FVector const BaseOffset(0.0f, 0.0f, HalfHeight);
+			FVector const LocalBasePos = OldLocalToWorld.InverseTransformPosition(CharacterOwner->GetActorLocation() - BaseOffset);
+			FVector const NewWorldPos = ConstrainLocationToPlane(NewLocalToWorld.TransformPosition(LocalBasePos) + BaseOffset);
+			DeltaPosition = ConstrainDirectionToPlane(NewWorldPos - CharacterOwner->GetActorLocation());
+
+			// move attached actor
+			if (bFastAttachedMove)
+			{
+				// we're trusting no other obstacle can prevent the move here
+				UpdatedComponent->SetWorldLocationAndRotation(NewWorldPos, FinalQuat.Rotator(), false);
+			}
+			else
+			{
+				FHitResult MoveOnBaseHit(1.f);
+				const FVector OldLocation = UpdatedComponent->GetComponentLocation();
+				MoveUpdatedComponent(DeltaPosition, FinalQuat.Rotator(), true, &MoveOnBaseHit);
+				if ((UpdatedComponent->GetComponentLocation() - (OldLocation + DeltaPosition)).IsNearlyZero() == false)
+				{
+					OnUnableToFollowBaseMove(DeltaPosition, OldLocation, MoveOnBaseHit);
+				}
+			}
+		}
+
+		if (MovementBase->IsSimulatingPhysics() && CharacterOwner->GetMesh())
+		{
+			CharacterOwner->GetMesh()->ApplyDeltaToAllPhysicsTransforms(DeltaPosition, DeltaQuat);
+		}
+	}
+}
+
+
+void UCharacterMovementComponent::OnUnableToFollowBaseMove(const FVector& DeltaPosition, const FVector& OldLocation, const FHitResult& MoveOnBaseHit)
+{
+	// no default implementation, left for subclasses to override.
+}
+
+
+void UCharacterMovementComponent::UpdateBasedRotation(FRotator &FinalRotation, const FRotator& ReducedRotation)
+{
+	AController* Controller = CharacterOwner ? CharacterOwner->Controller : NULL;
+	float ControllerRoll = 0.f;
+	if( Controller && !bIgnoreBaseRotation )
+	{
+		FRotator const ControllerRot = Controller->GetControlRotation();
+		ControllerRoll = ControllerRot.Roll;
+		Controller->SetControlRotation(ControllerRot + ReducedRotation);
+	}
+
+	// Remove roll
+	FinalRotation.Roll = 0.f;
+	if( Controller )
+	{
+		FinalRotation.Roll = CharacterOwner->GetActorRotation().Roll;
+		FRotator NewRotation = Controller->GetControlRotation();
+		NewRotation.Roll = ControllerRoll;
+		Controller->SetControlRotation(NewRotation);
+	}
+}
+
+void UCharacterMovementComponent::DisableMovement()
+{
+	if (CharacterOwner)
+	{
+		SetMovementMode(MOVE_None);
+	}
+	else
+	{
+		MovementMode = MOVE_None;
+		CustomMovementMode = 0;
+	}
+}
+
+void UCharacterMovementComponent::PerformMovement(float DeltaSeconds)
+{
+	SCOPE_CYCLE_COUNTER(STAT_CharacterMovement);
+	SCOPE_CYCLE_COUNTER(STAT_CharacterMovementAuthority);
+
+	if (!HasValidData())
+	{
+		return;
+	}
+	
+	// no movement if we can't move, or if currently doing physical simulation on UpdatedComponent
+	if (MovementMode == MOVE_None || UpdatedComponent->Mobility != EComponentMobility::Movable || UpdatedComponent->IsSimulatingPhysics())
+	{
+		return;
+	}
+
+	// Force floor update if we've moved outside of CharacterMovement since last update.
+	bForceNextFloorCheck |= (IsMovingOnGround() && UpdatedComponent->GetComponentLocation() != LastUpdateLocation);
+
+	FVector OldVelocity;
+	FVector OldLocation;
+
+	// Scoped updates can improve performance of multiple MoveComponent calls.
+	{
+		FScopedMovementUpdate ScopedMovementUpdate(UpdatedComponent, bEnableScopedMovementUpdates ? EScopedUpdate::DeferredUpdates : EScopedUpdate::ImmediateUpdates);
+
+		MaybeUpdateBasedMovement(DeltaSeconds);
+
+		OldVelocity = Velocity;
+		OldLocation = CharacterOwner->GetActorLocation();
+
+		ApplyAccumulatedForces(DeltaSeconds);
+
+		// Check for a change in crouch state. Players toggle crouch by changing bWantsToCrouch.
+		const bool bAllowedToCrouch = CanCrouchInCurrentState();
+		if ((!bAllowedToCrouch || !bWantsToCrouch) && IsCrouching())
+		{
+			UnCrouch(false);
+		}
+		else if (bWantsToCrouch && bAllowedToCrouch && !IsCrouching()) 
+		{
+			Crouch(false);
+		}
+
+		if (MovementMode == MOVE_NavWalking && bWantsToLeaveNavWalking)
+		{
+			TryToLeaveNavWalking();
+		}
+
+		// Character::LaunchCharacter() has been deferred until now.
+		HandlePendingLaunch();
+
+		// If using RootMotion, tick animations before running physics.
+		if( !CharacterOwner->bClientUpdating && CharacterOwner->IsPlayingRootMotion() && CharacterOwner->GetMesh() )
+		{
+			TickCharacterPose(DeltaSeconds);
+
+			// Make sure animation didn't trigger an event that destroyed us
+			if (!HasValidData())
+			{
+				return;
+			}
+
+			// For local human clients, save off root motion data so it can be used by movement networking code.
+			if( CharacterOwner->IsLocallyControlled() && (CharacterOwner->Role == ROLE_AutonomousProxy) && CharacterOwner->IsPlayingNetworkedRootMotionMontage() )
+			{
+				CharacterOwner->ClientRootMotionParams = RootMotionParams;
+			}
+		}
+
+		// if we're about to use root motion, convert it to world space first.
+		if( HasRootMotion() )
+		{
+			USkeletalMeshComponent * SkelMeshComp = CharacterOwner->GetMesh();
+			if( SkelMeshComp )
+			{
+				// Convert Local Space Root Motion to world space. Do it right before used by physics to make sure we use up to date transforms, as translation is relative to rotation.
+				RootMotionParams.Set( SkelMeshComp->ConvertLocalRootMotionToWorld(RootMotionParams.RootMotionTransform) );
+				UE_LOG(LogRootMotion, Log,  TEXT("PerformMovement WorldSpaceRootMotion Translation: %s, Rotation: %s, Actor Facing: %s"),
+					*RootMotionParams.RootMotionTransform.GetTranslation().ToCompactString(), *RootMotionParams.RootMotionTransform.GetRotation().Rotator().ToCompactString(), *CharacterOwner->GetActorRotation().Vector().ToCompactString());
+			}
+
+			// Then turn root motion to velocity to be used by various physics modes.
+			if( DeltaSeconds > 0.f )
+			{
+				const FVector RootMotionVelocity = RootMotionParams.RootMotionTransform.GetTranslation() / DeltaSeconds;
+				// Do not override Velocity.Z if in falling physics, we want to keep the effect of gravity.
+				Velocity = FVector(RootMotionVelocity.X, RootMotionVelocity.Y, (MovementMode == MOVE_Falling ? Velocity.Z : RootMotionVelocity.Z));
+			}
+		}
+
+		// NaN tracking
+		checkf(!Velocity.ContainsNaN(), TEXT("UCharacterMovementComponent::PerformMovement: Velocity contains NaN (%s: %s)\n%s"), *GetPathNameSafe(this), *GetPathNameSafe(GetOuter()), *Velocity.ToString());
+
+		// Clear jump input now, to allow movement events to trigger it for next update.
+		CharacterOwner->ClearJumpInput();
+
+		// change position
+		StartNewPhysics(DeltaSeconds, 0);
+
+		if (!HasValidData())
+		{
+			return;
+		}
+
+		// uncrouch if no longer allowed to be crouched
+		if (IsCrouching() && !CanCrouchInCurrentState())
+		{
+			UnCrouch(false);
+		}
+
+		if (!HasRootMotion() && !CharacterOwner->IsMatineeControlled())
+		{
+			PhysicsRotation(DeltaSeconds);
+		}
+
+		// Apply Root Motion rotation after movement is complete.
+		if( HasRootMotion() )
+		{
+			const FRotator OldActorRotation = CharacterOwner->GetActorRotation();
+			const FRotator RootMotionRotation = RootMotionParams.RootMotionTransform.GetRotation().Rotator();
+			if( !RootMotionRotation.IsNearlyZero() )
+			{
+				const FRotator NewActorRotation = (OldActorRotation + RootMotionRotation).GetNormalized();
+				MoveUpdatedComponent(FVector::ZeroVector, NewActorRotation, true);
+			}
+
+			// debug
+			if( false )
+			{
+				const FVector ResultingLocation = CharacterOwner->GetActorLocation();
+				const FRotator ResultingRotation = CharacterOwner->GetActorRotation();
+
+				// Show current position
+				DrawDebugCoordinateSystem(GetWorld(), CharacterOwner->GetMesh()->GetComponentLocation() + FVector(0,0,1), ResultingRotation, 50.f, false);
+
+				// Show resulting delta move.
+				DrawDebugLine(GetWorld(), OldLocation, ResultingLocation, FColor::Red, true, 10.f);
+
+				// Log details.
+				UE_LOG(LogRootMotion, Warning,  TEXT("PerformMovement Resulting DeltaMove Translation: %s, Rotation: %s, MovementBase: %s"),
+					*(ResultingLocation - OldLocation).ToCompactString(), *(ResultingRotation - OldActorRotation).GetNormalized().ToCompactString(), *GetNameSafe(CharacterOwner->GetMovementBase()) );
+
+				const FVector RMTranslation = RootMotionParams.RootMotionTransform.GetTranslation();
+				const FRotator RMRotation = RootMotionParams.RootMotionTransform.GetRotation().Rotator();
+				UE_LOG(LogRootMotion, Warning,  TEXT("PerformMovement Resulting DeltaError Translation: %s, Rotation: %s"),
+					*(ResultingLocation - OldLocation - RMTranslation).ToCompactString(), *(ResultingRotation - OldActorRotation - RMRotation).GetNormalized().ToCompactString() );
+			}
+
+			// Root Motion has been used, clear
+			RootMotionParams.Clear();
+		}
+
+		// consume path following requested velocity
+		bHasRequestedVelocity = false;
+
+		OnMovementUpdated(DeltaSeconds, OldLocation, OldVelocity);
+	} // End scoped movement update
+
+	// Call external post-movement events. These happen after the scoped movement completes in case the events want to use the current state of overlaps etc.
+	CallMovementUpdateDelegate(DeltaSeconds, OldLocation, OldVelocity);
+
+	SaveBaseLocation();
+	UpdateComponentVelocity();
+
+	LastUpdateLocation = UpdatedComponent ? UpdatedComponent->GetComponentLocation() : FVector::ZeroVector;
+}
+
+
+void UCharacterMovementComponent::CallMovementUpdateDelegate(float DeltaTime, const FVector& OldLocation, const FVector& OldVelocity)
+{
+	SCOPE_CYCLE_COUNTER(STAT_CharMoveUpdateDelegate);
+
+	// Update component velocity in case events want to read it
+	UpdateComponentVelocity();
+
+	// Delegate (for blueprints)
+	if (CharacterOwner)
+	{
+		CharacterOwner->OnCharacterMovementUpdated.Broadcast(DeltaTime, OldLocation, OldVelocity);
+	}
+}
+
+
+void UCharacterMovementComponent::OnMovementUpdated(float DeltaTime, const FVector& OldLocation, const FVector& OldVelocity)
+{
+	// empty base implementation, intended for derived classes to override.
+}
+
+
+void UCharacterMovementComponent::SaveBaseLocation()
+{
+	if (!HasValidData())
+	{
+		return;
+	}
+
+	const UPrimitiveComponent* MovementBase = CharacterOwner->GetMovementBase();
+	if (MovementBaseUtility::UseRelativeLocation(MovementBase) && !CharacterOwner->IsMatineeControlled())
+	{
+		// Read transforms into OldBaseLocation, OldBaseQuat
+		MovementBaseUtility::GetMovementBaseTransform(MovementBase, CharacterOwner->GetBasedMovement().BoneName, OldBaseLocation, OldBaseQuat);
+
+		// Location
+		const FVector RelativeLocation = CharacterOwner->GetActorLocation() - OldBaseLocation;
+
+		// Rotation
+		if (bIgnoreBaseRotation)
+		{
+			// Absolute rotation
+			CharacterOwner->SaveRelativeBasedMovement(RelativeLocation, CharacterOwner->GetActorRotation(), false);
+		}
+		else
+		{
+			// Relative rotation
+			const FRotator RelativeRotation = (FRotationMatrix(CharacterOwner->GetActorRotation()) * FQuatRotationTranslationMatrix(OldBaseQuat, FVector::ZeroVector).GetTransposed()).Rotator();
+			CharacterOwner->SaveRelativeBasedMovement(RelativeLocation, RelativeRotation, true);	
+		}
+	}
+}
+
+
+bool UCharacterMovementComponent::CanCrouchInCurrentState() const
+{
+	if (!CanEverCrouch())
+	{
+		return false;
+	}
+
+	return IsFalling() || IsMovingOnGround();
+}
+
+
+void UCharacterMovementComponent::Crouch(bool bClientSimulation)
+{
+	if (!HasValidData())
+	{
+		return;
+	}
+
+	// Do not perform if collision is already at desired size.
+	if( CharacterOwner->GetCapsuleComponent()->GetUnscaledCapsuleHalfHeight() == CrouchedHalfHeight )
+	{
+		return;
+	}
+
+	if (!CanCrouchInCurrentState())
+	{
+		return;
+	}
+
+	if (bClientSimulation && CharacterOwner->Role == ROLE_SimulatedProxy)
+	{
+		// restore collision size before crouching
+		ACharacter* DefaultCharacter = CharacterOwner->GetClass()->GetDefaultObject<ACharacter>();
+		CharacterOwner->GetCapsuleComponent()->SetCapsuleSize(DefaultCharacter->GetCapsuleComponent()->GetUnscaledCapsuleRadius(), DefaultCharacter->GetCapsuleComponent()->GetUnscaledCapsuleHalfHeight());
+		bShrinkProxyCapsule = true;
+	}
+
+	// Change collision size to crouching dimensions
+	const float ComponentScale = CharacterOwner->GetCapsuleComponent()->GetShapeScale();
+	const float OldUnscaledHalfHeight = CharacterOwner->GetCapsuleComponent()->GetUnscaledCapsuleHalfHeight();
+	CharacterOwner->GetCapsuleComponent()->SetCapsuleSize(CharacterOwner->GetCapsuleComponent()->GetUnscaledCapsuleRadius(), CrouchedHalfHeight);
+	float HalfHeightAdjust = (OldUnscaledHalfHeight - CrouchedHalfHeight);
+	float ScaledHalfHeightAdjust = HalfHeightAdjust * ComponentScale;
+
+	if( !bClientSimulation )
+	{
+		// Crouching to a larger height? (this is rare)
+		if (CrouchedHalfHeight > OldUnscaledHalfHeight)
+		{
+			static const FName NAME_CrouchTrace = FName(TEXT("CrouchTrace"));
+			FCollisionQueryParams CapsuleParams(NAME_CrouchTrace, false, CharacterOwner);
+			FCollisionResponseParams ResponseParam;
+			InitCollisionParams(CapsuleParams, ResponseParam);
+			const bool bEncroached = GetWorld()->OverlapBlockingTestByChannel(CharacterOwner->GetActorLocation() - FVector(0.f,0.f,ScaledHalfHeightAdjust), FQuat::Identity,
+				UpdatedComponent->GetCollisionObjectType(), GetPawnCapsuleCollisionShape(SHRINK_None), CapsuleParams, ResponseParam);
+
+			// If encroached, cancel
+			if( bEncroached )
+			{
+				CharacterOwner->GetCapsuleComponent()->SetCapsuleSize(CharacterOwner->GetCapsuleComponent()->GetUnscaledCapsuleRadius(), OldUnscaledHalfHeight);
+				return;
+			}
+		}
+
+		if (bCrouchMaintainsBaseLocation)
+		{
+			// Intentionally not using MoveUpdatedComponent, where a horizontal plane constraint would prevent the base of the capsule from staying at the same spot.
+			UpdatedComponent->MoveComponent(FVector(0.f, 0.f, -ScaledHalfHeightAdjust), CharacterOwner->GetActorRotation(), true);
+		}
+
+		CharacterOwner->bIsCrouched = true;
+	}
+
+	bForceNextFloorCheck = true;
+
+	// OnStartCrouch takes the change from the Default size, not the current one (though they are usually the same).
+	ACharacter* DefaultCharacter = CharacterOwner->GetClass()->GetDefaultObject<ACharacter>();
+	HalfHeightAdjust = (DefaultCharacter->GetCapsuleComponent()->GetUnscaledCapsuleHalfHeight() - CrouchedHalfHeight);
+	ScaledHalfHeightAdjust = HalfHeightAdjust * ComponentScale;
+
+	AdjustProxyCapsuleSize();
+	CharacterOwner->OnStartCrouch( HalfHeightAdjust, ScaledHalfHeightAdjust );
+}
+
+
+void UCharacterMovementComponent::UnCrouch(bool bClientSimulation)
+{
+	if (!HasValidData())
+	{
+		return;
+	}
+
+	ACharacter* DefaultCharacter = CharacterOwner->GetClass()->GetDefaultObject<ACharacter>();
+
+	// Do not perform if collision is already at desired size.
+	if( CharacterOwner->GetCapsuleComponent()->GetUnscaledCapsuleHalfHeight() == DefaultCharacter->GetCapsuleComponent()->GetUnscaledCapsuleHalfHeight() )
+	{
+		return;
+	}
+
+	const float CurrentCrouchedHalfHeight = CharacterOwner->GetCapsuleComponent()->GetScaledCapsuleHalfHeight();
+
+	const float ComponentScale = CharacterOwner->GetCapsuleComponent()->GetShapeScale();
+	const float OldUnscaledHalfHeight = CharacterOwner->GetCapsuleComponent()->GetUnscaledCapsuleHalfHeight();
+	const float HalfHeightAdjust = DefaultCharacter->GetCapsuleComponent()->GetUnscaledCapsuleHalfHeight() - OldUnscaledHalfHeight;
+	const float ScaledHalfHeightAdjust = HalfHeightAdjust * ComponentScale;
+	const FVector PawnLocation = CharacterOwner->GetActorLocation();
+
+	// Grow to uncrouched size.
+	check(CharacterOwner->GetCapsuleComponent());
+	bool bUpdateOverlaps = false;
+	CharacterOwner->GetCapsuleComponent()->SetCapsuleSize(DefaultCharacter->GetCapsuleComponent()->GetUnscaledCapsuleRadius(), DefaultCharacter->GetCapsuleComponent()->GetUnscaledCapsuleHalfHeight(), bUpdateOverlaps);
+	CharacterOwner->GetCapsuleComponent()->UpdateBounds(); // Force an update of the bounds with the new dimensions
+
+	if( !bClientSimulation )
+	{
+		// Try to stay in place and see if the larger capsule fits. We use a slightly taller capsule to avoid penetration.
+		static const FName NAME_CrouchTrace = FName(TEXT("CrouchTrace"));
+		const float SweepInflation = KINDA_SMALL_NUMBER * 10.f;
+		FCollisionQueryParams CapsuleParams(NAME_CrouchTrace, false, CharacterOwner);
+		FCollisionResponseParams ResponseParam;
+		InitCollisionParams(CapsuleParams, ResponseParam);
+		const FCollisionShape StandingCapsuleShape = GetPawnCapsuleCollisionShape(SHRINK_HeightCustom, -SweepInflation); // Shrink by negative amount, so actually grow it.
+		const ECollisionChannel CollisionChannel = UpdatedComponent->GetCollisionObjectType();
+		bool bEncroached = true;
+
+		if (!bCrouchMaintainsBaseLocation)
+		{
+			// Expand in place
+			bEncroached = GetWorld()->OverlapBlockingTestByChannel(PawnLocation, FQuat::Identity, CollisionChannel, StandingCapsuleShape, CapsuleParams, ResponseParam);
+		
+			if (bEncroached)
+			{
+				// Try adjusting capsule position to see if we can avoid encroachment.
+				if (ScaledHalfHeightAdjust > 0.f)
+				{
+					// Shrink to a short capsule, sweep down to base to find where that would hit something, and then try to stand up from there.
+					float PawnRadius, PawnHalfHeight;
+					CharacterOwner->GetCapsuleComponent()->GetScaledCapsuleSize(PawnRadius, PawnHalfHeight);
+					const float ShrinkHalfHeight = PawnHalfHeight - PawnRadius;
+					const float TraceDist = PawnHalfHeight - ShrinkHalfHeight;
+					const FVector Down = FVector(0.f, 0.f, -TraceDist);
+
+					FHitResult Hit(1.f);
+					const FCollisionShape ShortCapsuleShape = GetPawnCapsuleCollisionShape(SHRINK_HeightCustom, ShrinkHalfHeight);
+					const bool bBlockingHit = GetWorld()->SweepSingleByChannel(Hit, PawnLocation, PawnLocation + Down, FQuat::Identity, CollisionChannel, ShortCapsuleShape, CapsuleParams);
+					if (Hit.bStartPenetrating)
+					{
+						bEncroached = true;
+					}
+					else
+					{
+						// Compute where the base of the sweep ended up, and see if we can stand there
+						const float DistanceToBase = (Hit.Time * TraceDist) + ShortCapsuleShape.Capsule.HalfHeight;
+						const FVector NewLoc = FVector(PawnLocation.X, PawnLocation.Y, PawnLocation.Z - DistanceToBase + PawnHalfHeight + SweepInflation + MIN_FLOOR_DIST / 2.f);
+						bEncroached = GetWorld()->OverlapBlockingTestByChannel(NewLoc, FQuat::Identity, CollisionChannel, StandingCapsuleShape, CapsuleParams, ResponseParam);
+						if (!bEncroached)
+						{
+							// Intentionally not using MoveUpdatedComponent, where a horizontal plane constraint would prevent the base of the capsule from staying at the same spot.
+							UpdatedComponent->MoveComponent(NewLoc - PawnLocation, CharacterOwner->GetActorRotation(), false);
+						}
+					}
+				}
+			}
+		}
+		else
+		{
+			// Expand while keeping base location the same.
+			FVector StandingLocation = PawnLocation + FVector(0.f, 0.f, StandingCapsuleShape.GetCapsuleHalfHeight() - CurrentCrouchedHalfHeight);
+			bEncroached = GetWorld()->OverlapBlockingTestByChannel(StandingLocation, FQuat::Identity, CollisionChannel, StandingCapsuleShape, CapsuleParams, ResponseParam);
+
+			if (bEncroached)
+			{
+				if (IsMovingOnGround())
+				{
+					// Something might be just barely overhead, try moving down closer to the floor to avoid it.
+					const float MinFloorDist = KINDA_SMALL_NUMBER * 10.f;
+					if (CurrentFloor.bBlockingHit && CurrentFloor.FloorDist > MinFloorDist)
+					{
+						StandingLocation.Z -= CurrentFloor.FloorDist - MinFloorDist;
+						bEncroached = GetWorld()->OverlapBlockingTestByChannel(StandingLocation, FQuat::Identity, CollisionChannel, StandingCapsuleShape, CapsuleParams, ResponseParam);
+					}
+				}				
+			}
+
+			if (!bEncroached)
+			{
+				// Commit the change in location.
+				UpdatedComponent->MoveComponent(StandingLocation - PawnLocation, CharacterOwner->GetActorRotation(), false);
+				bForceNextFloorCheck = true;
+			}
+		}
+
+		// If still encroached then abort.
+		if (bEncroached)
+		{
+			CharacterOwner->GetCapsuleComponent()->SetCapsuleSize(CharacterOwner->GetCapsuleComponent()->GetUnscaledCapsuleRadius(), OldUnscaledHalfHeight, false);
+			CharacterOwner->GetCapsuleComponent()->UpdateBounds(); // Update bounds again back to old value
+			return;
+		}
+
+		CharacterOwner->bIsCrouched = false;
+	}	
+	else
+	{
+		bShrinkProxyCapsule = true;
+	}
+
+	// now call SetCapsuleSize() to cause touch/untouch events
+	bUpdateOverlaps = true;
+	CharacterOwner->GetCapsuleComponent()->SetCapsuleSize(DefaultCharacter->GetCapsuleComponent()->GetUnscaledCapsuleRadius(), DefaultCharacter->GetCapsuleComponent()->GetUnscaledCapsuleHalfHeight(), bUpdateOverlaps);
+
+	AdjustProxyCapsuleSize();
+	CharacterOwner->OnEndCrouch( HalfHeightAdjust, ScaledHalfHeightAdjust );
+}
+
+void UCharacterMovementComponent::StartNewPhysics(float deltaTime, int32 Iterations)
+{
+	if ( (deltaTime < MIN_TICK_TIME) || (Iterations >= MaxSimulationIterations) || !HasValidData())
+		return;
+
+	if (UpdatedComponent->IsSimulatingPhysics())
+	{
+		UE_LOG(LogCharacterMovement, Log, TEXT("UCharacterMovementComponent::StartNewPhysics: UpdateComponent (%s) is simulating physics - aborting."), *UpdatedComponent->GetPathName());
+		return;
+	}
+
+	bMovementInProgress = true;
+
+	switch ( MovementMode )
+	{
+	case MOVE_None:
+		break;
+	case MOVE_Walking:
+		PhysWalking(deltaTime, Iterations);
+		break;
+	case MOVE_NavWalking:
+		PhysNavWalking(deltaTime, Iterations);
+		break;
+	case MOVE_Falling:
+		PhysFalling(deltaTime, Iterations);
+		break;
+	case MOVE_Flying:
+		PhysFlying(deltaTime, Iterations);
+		break;
+	case MOVE_Swimming:
+		PhysSwimming(deltaTime, Iterations);
+		break;
+	case MOVE_Custom:
+		PhysCustom(deltaTime, Iterations);
+		break;
+	default:
+		UE_LOG(LogCharacterMovement, Warning, TEXT("%s has unsupported movement mode %d"), *CharacterOwner->GetName(), int32(MovementMode));
+		SetMovementMode(MOVE_None);
+		break;
+	}
+
+	bMovementInProgress = false;
+	if ( bDeferUpdateMoveComponent )
+	{
+		SetUpdatedComponent(DeferredUpdatedMoveComponent);
+	}
+}
+
+float UCharacterMovementComponent::GetGravityZ() const
+{
+	return Super::GetGravityZ() * GravityScale;
+}
+
+float UCharacterMovementComponent::GetMaxSpeed() const
+{
+	switch(MovementMode)
+	{
+	case MOVE_Walking:
+	case MOVE_NavWalking:
+		return IsCrouching() ? MaxWalkSpeedCrouched : MaxWalkSpeed;
+	case MOVE_Falling:
+		return MaxWalkSpeed;
+	case MOVE_Swimming:
+		return MaxSwimSpeed;
+	case MOVE_Flying:
+		return MaxFlySpeed;
+	case MOVE_Custom:
+		return MaxCustomMovementSpeed;
+	case MOVE_None:
+	default:
+		return 0.f;
+	}
+}
+
+
+bool UCharacterMovementComponent::ResolvePenetration(const FVector& Adjustment, const FHitResult& Hit, const FRotator& NewRotation)
+{
+	// If movement occurs, mark that we teleported, so we don't incorrectly adjust velocity based on a potentially very different movement than our movement direction.
+	bJustTeleported |= Super::ResolvePenetration(Adjustment, Hit, NewRotation);
+	return bJustTeleported;
+}
+
+
+float UCharacterMovementComponent::SlideAlongSurface(const FVector& Delta, float Time, const FVector& InNormal, FHitResult &Hit, bool bHandleImpact)
+{
+	if (!Hit.bBlockingHit)
+	{
+		return 0.f;
+	}
+
+	FVector Normal(InNormal);
+	if (IsMovingOnGround())
+	{
+		// We don't want to be pushed up an unwalkable surface.
+		if (Normal.Z > 0.f)
+		{
+			if (!IsWalkable(Hit))
+			{
+				Normal = Normal.GetSafeNormal2D();
+			}
+		}
+		else if (Normal.Z < -KINDA_SMALL_NUMBER)
+		{
+			// Don't push down into the floor when the impact is on the upper portion of the capsule.
+			if (CurrentFloor.FloorDist < MIN_FLOOR_DIST && CurrentFloor.bBlockingHit)
+			{
+				const FVector FloorNormal = CurrentFloor.HitResult.Normal;
+				const bool bFloorOpposedToMovement = (Delta | FloorNormal) < 0.f && (FloorNormal.Z < 1.f - DELTA);
+				if (bFloorOpposedToMovement)
+				{
+					Normal = FloorNormal;
+				}
+				
+				Normal = Normal.GetSafeNormal2D();
+			}
+		}
+	}
+
+	return Super::SlideAlongSurface(Delta, Time, Normal, Hit, bHandleImpact);
+}
+
+
+void UCharacterMovementComponent::TwoWallAdjust(FVector &Delta, const FHitResult& Hit, const FVector &OldHitNormal) const
+{
+	FVector InDelta = Delta;
+	Super::TwoWallAdjust(Delta, Hit, OldHitNormal);
+
+	if (IsMovingOnGround())
+	{
+		// Allow slides up walkable surfaces, but not unwalkable ones (treat those as vertical barriers).
+		if (Delta.Z > 0.f)
+		{
+			if ((Hit.Normal.Z >= WalkableFloorZ || IsWalkable(Hit)) && Hit.Normal.Z > KINDA_SMALL_NUMBER)
+			{
+				// Maintain horizontal velocity
+				const float Time = (1.f - Hit.Time);
+				const FVector ScaledDelta = Delta.GetSafeNormal() * InDelta.Size();
+				Delta = FVector(InDelta.X, InDelta.Y, ScaledDelta.Z / Hit.Normal.Z) * Time;
+			}
+			else
+			{
+				Delta.Z = 0.f;
+			}
+		}
+		else if (Delta.Z < 0.f)
+		{
+			// Don't push down into the floor.
+			if (CurrentFloor.FloorDist < MIN_FLOOR_DIST && CurrentFloor.bBlockingHit)
+			{
+				Delta.Z = 0.f;
+			}
+		}
+	}
+}
+
+
+FVector UCharacterMovementComponent::ComputeSlideVector(const FVector& Delta, const float Time, const FVector& Normal, const FHitResult& Hit) const
+{
+	FVector Result = Super::ComputeSlideVector(Delta, Time, Normal, Hit);
+
+	// prevent boosting up slopes
+	if (IsFalling())
+	{
+		Result = HandleSlopeBoosting(Result, Delta, Time, Normal, Hit);
+	}
+
+	return Result;
+}
+
+
+FVector UCharacterMovementComponent::HandleSlopeBoosting(const FVector& SlideResult, const FVector& Delta, const float Time, const FVector& Normal, const FHitResult& Hit) const
+{
+	FVector Result = SlideResult;
+
+	if (Result.Z > 0.f)
+	{
+		// Don't move any higher than we originally intended.
+		const float ZLimit = Delta.Z * Time;
+		if (Result.Z - ZLimit > KINDA_SMALL_NUMBER)
+		{
+			if (ZLimit > 0.f)
+			{
+				// Rescale the entire vector (not just the Z component) otherwise we change the direction and likely head right back into the impact.
+				const float UpPercent = ZLimit / Result.Z;
+				Result *= UpPercent;
+			}
+			else
+			{
+				// We were heading down but were going to deflect upwards. Just make the deflection horizontal.
+				Result = FVector::ZeroVector;
+			}
+
+			// Make remaining portion of original result horizontal and parallel to impact normal.
+			const FVector RemainderXY = (SlideResult - Result) * FVector(1.f, 1.f, 0.f);
+			const FVector NormalXY = Normal.GetSafeNormal2D();
+			const FVector Adjust = Super::ComputeSlideVector(RemainderXY, 1.f, NormalXY, Hit);
+			Result += Adjust;
+		}
+	}
+
+	return Result;
+}
+
+
+// TODO: deprecated, remove.
+FVector UCharacterMovementComponent::AdjustUpperHemisphereImpact(const FVector& Delta, const FHitResult& Hit) const
+{
+	const float ZScale = FMath::Clamp(1.f - (FMath::Abs(Hit.Normal.Z) * UpperImpactNormalScale_DEPRECATED), 0.f, 1.f);
+	return FVector(Delta.X, Delta.Y, Delta.Z * ZScale);
+}
+
+
+FVector UCharacterMovementComponent::NewFallVelocity(const FVector& InitialVelocity, const FVector& Gravity, float DeltaTime) const
+{
+	FVector Result = InitialVelocity;
+
+	if (!Gravity.IsZero())
+	{
+		// Apply gravity.
+		Result += Gravity * DeltaTime;
+
+		const FVector GravityDir = Gravity.GetSafeNormal();
+		const float TerminalLimit = FMath::Abs(GetPhysicsVolume()->TerminalVelocity);
+
+		// Don't exceed terminal velocity.
+		if ((Result | GravityDir) > TerminalLimit)
+		{
+			Result = FVector::PointPlaneProject(Result, FVector::ZeroVector, GravityDir) + GravityDir * TerminalLimit;
+		}
+	}
+
+	return Result;
+}
+
+
+float UCharacterMovementComponent::ImmersionDepth()
+{
+	float depth = 0.f;
+
+	if ( CharacterOwner && GetPhysicsVolume()->bWaterVolume )
+	{
+		const float CollisionHalfHeight = CharacterOwner->GetSimpleCollisionHalfHeight();
+
+		if ( (CollisionHalfHeight == 0.f) || (Buoyancy == 0.f) )
+		{
+			depth = 1.f;
+		}
+		else
+		{
+			UBrushComponent* VolumeBrushComp = GetPhysicsVolume()->GetBrushComponent();
+			FHitResult Hit(1.f);
+			if ( VolumeBrushComp )
+			{
+				const FVector TraceStart = CharacterOwner->GetActorLocation() + FVector(0.f,0.f,CollisionHalfHeight);
+				const FVector TraceEnd = CharacterOwner->GetActorLocation() - FVector(0.f,0.f,CollisionHalfHeight);
+
+				const static FName MovementComp_Character_ImmersionDepthName(TEXT("MovementComp_Character_ImmersionDepth"));
+				FCollisionQueryParams NewTraceParams( MovementComp_Character_ImmersionDepthName, true );
+
+				VolumeBrushComp->LineTraceComponent( Hit, TraceStart, TraceEnd, NewTraceParams );
+			}
+
+			depth = (Hit.Time == 1.f) ? 1.f : (1.f - Hit.Time);
+		}
+	}
+	return depth;
+}
+
+bool UCharacterMovementComponent::IsFlying() const
+{
+	if (!CharacterOwner || !UpdatedComponent)
+	{
+		return false;
+	}
+
+	return MovementMode == MOVE_Flying;
+}
+
+
+
+bool UCharacterMovementComponent::IsMovingOnGround() const
+{
+	if (!CharacterOwner || !UpdatedComponent)
+	{
+		return false;
+	}
+
+	return (MovementMode == MOVE_Walking) || (MovementMode == MOVE_NavWalking);
+}
+
+
+
+bool UCharacterMovementComponent::IsFalling() const
+{
+	if (!CharacterOwner || !UpdatedComponent)
+	{
+		return false;
+	}
+
+	return MovementMode == MOVE_Falling;
+}
+
+bool UCharacterMovementComponent::IsSwimming() const
+{
+	if (!CharacterOwner || !UpdatedComponent)
+	{
+		return false;
+	}
+
+	return MovementMode == MOVE_Swimming;
+}
+
+bool UCharacterMovementComponent::IsCrouching() const
+{
+	return CharacterOwner && CharacterOwner->bIsCrouched;
+}
+
+bool UCharacterMovementComponent::IsWalking() const
+{
+	return IsMovingOnGround();
+}
+
+void UCharacterMovementComponent::CalcVelocity(float DeltaTime, float Friction, bool bFluid, float BrakingDeceleration)
+{
+	// Do not update velocity when using root motion
+	if (!HasValidData() || HasRootMotion() || DeltaTime < MIN_TICK_TIME)
+	{
+		return;
+	}
+
+	Friction = FMath::Max(0.f, Friction);
+	const float MaxAccel = GetMaxAcceleration();
+	float MaxSpeed = GetMaxSpeed();
+	
+	// Check if path following requested movement
+	bool bZeroRequestedAcceleration = true;
+	FVector RequestedAcceleration = FVector::ZeroVector;
+	float RequestedSpeed = 0.0f;
+	if (ApplyRequestedMove(DeltaTime, MaxAccel, MaxSpeed, Friction, BrakingDeceleration, RequestedAcceleration, RequestedSpeed))
+	{
+		RequestedAcceleration = RequestedAcceleration.GetClampedToMaxSize(MaxAccel);
+		bZeroRequestedAcceleration = false;
+	}
+
+	if (bForceMaxAccel)
+	{
+		// Force acceleration at full speed.
+		// In consideration order for direction: Acceleration, then Velocity, then Pawn's rotation.
+		if (Acceleration.SizeSquared() > SMALL_NUMBER)
+		{
+			Acceleration = GetSafeNormalPrecise(Acceleration) * MaxAccel;
+		}
+		else 
+		{
+			Acceleration = MaxAccel * (Velocity.SizeSquared() < SMALL_NUMBER ? CharacterOwner->GetActorRotation().Vector() : GetSafeNormalPrecise(Velocity));
+		}
+
+		AnalogInputModifier = 1.f;
+	}
+
+	// Path following above didn't care about the analog modifier, but we do for everything else below, so get the fully modified value.
+	// Use max of requested speed and max speed if we modified the speed in ApplyRequestedMove above.
+	MaxSpeed = FMath::Max(RequestedSpeed, MaxSpeed * AnalogInputModifier);
+
+	// Apply braking or deceleration
+	const bool bZeroAcceleration = Acceleration.IsZero();
+	const bool bVelocityOverMax = IsExceedingMaxSpeed(MaxSpeed);
+	
+	// Only apply braking if there is no acceleration, or we are over our max speed and need to slow down to it.
+	if ((bZeroAcceleration && bZeroRequestedAcceleration) || bVelocityOverMax)
+	{
+		const FVector OldVelocity = Velocity;
+		ApplyVelocityBraking(DeltaTime, Friction, BrakingDeceleration);
+	
+		// Don't allow braking to lower us below max speed if we started above it.
+		if (bVelocityOverMax && Velocity.SizeSquared() < FMath::Square(MaxSpeed) && FVector::DotProduct(Acceleration, OldVelocity) > 0.0f)
+		{
+			Velocity = GetSafeNormalPrecise(OldVelocity) * MaxSpeed;
+		}
+	}
+	else if (!bZeroAcceleration)
+	{
+		// Friction affects our ability to change direction. This is only done for input acceleration, not path following.
+		const FVector AccelDir = GetSafeNormalPrecise(Acceleration);
+		const float VelSize = Velocity.Size();
+		Velocity = Velocity - (Velocity - AccelDir * VelSize) * FMath::Min(DeltaTime * Friction, 1.f);
+	}
+
+	// Apply fluid friction
+	if (bFluid)
+	{
+		Velocity = Velocity * (1.f - FMath::Min(Friction * DeltaTime, 1.f));
+	}
+
+	// Apply acceleration
+	const float NewMaxSpeed = (IsExceedingMaxSpeed(MaxSpeed)) ? Velocity.Size() : MaxSpeed;
+	Velocity += Acceleration * DeltaTime;
+	Velocity += RequestedAcceleration * DeltaTime;
+	Velocity = GetClampedToMaxSizePrecise(Velocity, NewMaxSpeed);
+
+	if (bUseRVOAvoidance)
+	{
+		CalcAvoidanceVelocity(DeltaTime);
+	}
+}
+
+
+bool UCharacterMovementComponent::ApplyRequestedMove(float DeltaTime, float MaxAccel, float MaxSpeed, float Friction, float BrakingDeceleration, FVector& OutAcceleration, float& OutRequestedSpeed)
+{
+	if (bHasRequestedVelocity)
+	{
+		const float RequestedSpeedSquared = RequestedVelocity.SizeSquared();
+		if (RequestedSpeedSquared < KINDA_SMALL_NUMBER)
+		{
+			return false;
+		}
+
+		// Compute requested speed from path following
+		float RequestedSpeed = FMath::Sqrt(RequestedSpeedSquared);
+		const FVector RequestedMoveDir = RequestedVelocity / RequestedSpeed;
+		RequestedSpeed = (bRequestedMoveWithMaxSpeed ? MaxSpeed : FMath::Min(MaxSpeed, RequestedSpeed));
+		
+		// Compute actual requested velocity
+		const FVector MoveVelocity = RequestedMoveDir * RequestedSpeed;
+		
+		// Compute acceleration. Use MaxAccel to limit speed increase, 1% buffer.
+		FVector NewAcceleration = FVector::ZeroVector;
+		const float CurrentSpeedSq = Velocity.SizeSquared();
+		if (bRequestedMoveUseAcceleration && CurrentSpeedSq < FMath::Square(RequestedSpeed * 1.01f))
+		{
+			// Turn in the same manner as with input acceleration.
+			const float VelSize = FMath::Sqrt(CurrentSpeedSq);
+			Velocity = Velocity - (Velocity - RequestedMoveDir * VelSize) * FMath::Min(DeltaTime * Friction, 1.f);
+
+			// How much do we need to accelerate to get to the new velocity?
+			NewAcceleration = ((MoveVelocity - Velocity) / DeltaTime);
+			NewAcceleration = NewAcceleration.GetClampedToMaxSize(MaxAccel);
+		}
+		else
+		{
+			// Just set velocity directly.
+			// If decelerating we do so instantly, so we don't slide through the destination if we can't brake fast enough.
+			Velocity = MoveVelocity;
+		}
+
+		// Copy to out params
+		OutRequestedSpeed = RequestedSpeed;
+		OutAcceleration = NewAcceleration;
+		return true;
+	}
+
+	return false;
+}
+
+void UCharacterMovementComponent::RequestDirectMove(const FVector& MoveVelocity, bool bForceMaxSpeed)
+{
+	if (MoveVelocity.SizeSquared() < KINDA_SMALL_NUMBER)
+	{
+		return;
+	}
+
+	if (IsFalling())
+	{
+		const FVector FallVelocity = MoveVelocity.GetClampedToMaxSize(GetMaxSpeed());
+		PerformAirControlForPathFollowing(FallVelocity, FallVelocity.Z);
+		return;
+	}
+
+	RequestedVelocity = MoveVelocity;
+	bHasRequestedVelocity = true;
+	bRequestedMoveWithMaxSpeed = bForceMaxSpeed;
+
+	if (IsMovingOnGround())
+	{
+		RequestedVelocity.Z = 0.0f;
+	}
+}
+
+bool UCharacterMovementComponent::CanStartPathFollowing() const
+{
+	if (!HasValidData() || HasRootMotion())
+	{
+		return false;
+	}
+
+	if (CharacterOwner)
+	{
+		if (CharacterOwner->GetRootComponent() && CharacterOwner->GetRootComponent()->IsSimulatingPhysics())
+		{
+			return false;
+		}
+		else if (CharacterOwner->IsMatineeControlled())
+		{
+			return false;
+		}
+	}
+
+	return Super::CanStartPathFollowing();
+}
+
+bool UCharacterMovementComponent::CanStopPathFollowing() const
+{
+	return !IsFalling();
+}
+
+void UCharacterMovementComponent::CalcAvoidanceVelocity(float DeltaTime)
+{
+	SCOPE_CYCLE_COUNTER(STAT_AI_ObstacleAvoidance);
+
+	UAvoidanceManager* AvoidanceManager = GetWorld()->GetAvoidanceManager();
+	if (AvoidanceWeight >= 1.0f || AvoidanceManager == NULL || GetCharacterOwner() == NULL)
+	{
+		return;
+	}
+
+	if (GetCharacterOwner()->Role != ROLE_Authority)
+	{
+		return;
+	}
+
+#if !(UE_BUILD_SHIPPING || UE_BUILD_TEST)
+	const bool bShowDebug = AvoidanceManager->IsDebugEnabled(AvoidanceUID);
+#endif
+
+	//Adjust velocity only if we're in "Walking" mode. We should also check if we're dazed, being knocked around, maybe off-navmesh, etc.
+	UCapsuleComponent *OurCapsule = GetCharacterOwner()->GetCapsuleComponent();
+	if (!Velocity.IsZero() && IsMovingOnGround() && OurCapsule)
+	{
+		//See if we're doing a locked avoidance move already, and if so, skip the testing and just do the move.
+		if (AvoidanceLockTimer > 0.0f)
+		{
+			Velocity = AvoidanceLockVelocity;
+#if !(UE_BUILD_SHIPPING || UE_BUILD_TEST)
+			if (bShowDebug)
+			{
+				DrawDebugLine(GetWorld(), GetActorFeetLocation(), GetActorFeetLocation() + Velocity, FColor::Blue, true, 0.5f, SDPG_MAX);
+			}
+#endif
+		}
+		else
+		{
+			FVector NewVelocity = AvoidanceManager->GetAvoidanceVelocityForComponent(this);
+			if (bUseRVOPostProcess)
+			{
+				PostProcessAvoidanceVelocity(NewVelocity);
+			}
+
+			if (!NewVelocity.Equals(Velocity))		//Really want to branch hint that this will probably not pass
+			{
+				//Had to divert course, lock this avoidance move in for a short time. This will make us a VO, so unlocked others will know to avoid us.
+				Velocity = NewVelocity;
+				SetAvoidanceVelocityLock(AvoidanceManager, AvoidanceManager->LockTimeAfterAvoid);
+#if !(UE_BUILD_SHIPPING || UE_BUILD_TEST)
+				if (bShowDebug)
+				{
+					DrawDebugLine(GetWorld(), GetActorFeetLocation(), GetActorFeetLocation() + Velocity, FColor::Red, true, 0.05f, SDPG_MAX, 10.0f);
+				}
+#endif
+			}
+			else
+			{
+				//Although we didn't divert course, our velocity for this frame is decided. We will not reciprocate anything further, so treat as a VO for the remainder of this frame.
+				SetAvoidanceVelocityLock(AvoidanceManager, AvoidanceManager->LockTimeAfterClean);	//10 ms of lock time should be adequate.
+#if !(UE_BUILD_SHIPPING || UE_BUILD_TEST)
+				if (bShowDebug)
+				{
+					//DrawDebugLine(GetWorld(), GetActorLocation(), GetActorLocation() + Velocity, FColor::Green, true, 0.05f, SDPG_MAX, 10.0f);
+				}
+#endif
+			}
+		}
+		//RickH - We might do better to do this later in our update
+		AvoidanceManager->UpdateRVO(this);
+
+		bWasAvoidanceUpdated = true;
+	}
+#if !(UE_BUILD_SHIPPING || UE_BUILD_TEST)
+	else if (bShowDebug)
+	{
+		DrawDebugLine(GetWorld(), GetActorFeetLocation(), GetActorFeetLocation() + Velocity, FColor::Yellow, true, 0.05f, SDPG_MAX);
+	}
+
+	if (bShowDebug)
+	{
+		FVector UpLine(0,0,500);
+		DrawDebugLine(GetWorld(), GetActorFeetLocation(), GetActorFeetLocation() + UpLine, (AvoidanceLockTimer > 0.01f) ? FColor::Red : FColor::Blue, true, 0.05f, SDPG_MAX, 5.0f);
+	}
+#endif
+}
+
+void UCharacterMovementComponent::PostProcessAvoidanceVelocity(FVector& NewVelocity)
+{
+	// empty in base class
+}
+
+void UCharacterMovementComponent::UpdateDefaultAvoidance()
+{
+	if (!bUseRVOAvoidance)
+	{
+		return;
+	}
+
+	SCOPE_CYCLE_COUNTER(STAT_AI_ObstacleAvoidance);
+
+	UAvoidanceManager* AvoidanceManager = GetWorld()->GetAvoidanceManager();
+	if (AvoidanceManager && !bWasAvoidanceUpdated)
+	{
+		if (UCapsuleComponent *OurCapsule = GetCharacterOwner()->GetCapsuleComponent())
+		{
+			AvoidanceManager->UpdateRVO(this);
+
+			//Consider this a clean move because we didn't even try to avoid.
+			SetAvoidanceVelocityLock(AvoidanceManager, AvoidanceManager->LockTimeAfterClean);
+		}
+	}
+
+	bWasAvoidanceUpdated = false;		//Reset for next frame
+}
+
+void UCharacterMovementComponent::SetRVOAvoidanceUID(int32 UID)
+{
+	AvoidanceUID = UID;
+}
+
+int32 UCharacterMovementComponent::GetRVOAvoidanceUID()
+{
+	return AvoidanceUID;
+}
+
+void UCharacterMovementComponent::SetRVOAvoidanceWeight(float Weight)
+{
+	AvoidanceWeight = Weight;
+}
+
+float UCharacterMovementComponent::GetRVOAvoidanceWeight()
+{
+	return AvoidanceWeight;
+}
+
+FVector UCharacterMovementComponent::GetRVOAvoidanceOrigin()
+{
+	return GetActorFeetLocation();
+}
+
+float UCharacterMovementComponent::GetRVOAvoidanceRadius()
+{
+	UCapsuleComponent* CapsuleComp = GetCharacterOwner()->GetCapsuleComponent();
+	return CapsuleComp ? CapsuleComp->GetScaledCapsuleRadius() : 0.0f;
+}
+
+float UCharacterMovementComponent::GetRVOAvoidanceConsiderationRadius()
+{
+	return AvoidanceConsiderationRadius;
+}
+
+float UCharacterMovementComponent::GetRVOAvoidanceHeight()
+{
+	UCapsuleComponent* CapsuleComp = GetCharacterOwner()->GetCapsuleComponent();
+	return CapsuleComp ? CapsuleComp->GetScaledCapsuleHalfHeight() : 0.0f;
+}
+
+FVector UCharacterMovementComponent::GetVelocityForRVOConsideration()
+{
+	return Velocity;
+}
+
+int32 UCharacterMovementComponent::GetAvoidanceGroupMask()
+{
+	return AvoidanceGroup.Packed;
+}
+
+int32 UCharacterMovementComponent::GetGroupsToAvoidMask()
+{
+	return GroupsToAvoid.Packed;
+}
+
+int32 UCharacterMovementComponent::GetGroupsToIgnoreMask()
+{
+	return GroupsToIgnore.Packed;
+}
+
+void UCharacterMovementComponent::SetAvoidanceVelocityLock(class UAvoidanceManager* Avoidance, float Duration)
+{
+	Avoidance->OverrideToMaxWeight(AvoidanceUID, Duration);
+	AvoidanceLockVelocity = Velocity;
+	AvoidanceLockTimer = Duration;
+}
+
+void UCharacterMovementComponent::NotifyBumpedPawn(APawn* BumpedPawn)
+{
+	Super::NotifyBumpedPawn(BumpedPawn);
+
+#if !(UE_BUILD_SHIPPING || UE_BUILD_TEST)
+	UAvoidanceManager* Avoidance = GetWorld()->GetAvoidanceManager();
+	const bool bShowDebug = Avoidance && Avoidance->IsDebugEnabled(AvoidanceUID);
+	if (bShowDebug)
+	{
+		DrawDebugLine(GetWorld(), GetActorFeetLocation(), GetActorFeetLocation() + FVector(0,0,500), (AvoidanceLockTimer > 0) ? FColor(255,64,64) : FColor(64,64,255), true, 2.0f, SDPG_MAX, 20.0f);
+	}
+#endif
+
+	// Unlock avoidance move. This mostly happens when two pawns who are locked into avoidance moves collide with each other.
+	AvoidanceLockTimer = 0.0f;
+}
+
+float UCharacterMovementComponent::GetMaxJumpHeight() const
+{
+	const float Gravity = GetGravityZ();
+	if (FMath::Abs(Gravity) > KINDA_SMALL_NUMBER)
+	{
+		return FMath::Square(JumpZVelocity) / (-2.f * Gravity);
+	}
+	else
+	{
+		return 0.f;
+	}
+}
+
+// TODO: deprecated, remove.
+float UCharacterMovementComponent::GetModifiedMaxAcceleration() const
+{
+	// Allow calling old deprecated function to maintain old behavior until it is removed.
+	PRAGMA_DISABLE_DEPRECATION_WARNINGS
+	return CharacterOwner ? MaxAcceleration * GetMaxSpeedModifier() : 0.f;
+	PRAGMA_ENABLE_DEPRECATION_WARNINGS
+}
+
+// TODO: deprecated, remove.
+float UCharacterMovementComponent::K2_GetModifiedMaxAcceleration() const
+{
+	// Allow calling old deprecated function to maintain old behavior until it is removed.
+	PRAGMA_DISABLE_DEPRECATION_WARNINGS
+	return GetModifiedMaxAcceleration();
+	PRAGMA_ENABLE_DEPRECATION_WARNINGS
+}
+
+
+float UCharacterMovementComponent::GetMaxAcceleration() const
+{
+	return MaxAcceleration;
+}
+
+FVector UCharacterMovementComponent::GetCurrentAcceleration() const
+{
+	return Acceleration;
+}
+
+void UCharacterMovementComponent::ApplyVelocityBraking(float DeltaTime, float Friction, float BrakingDeceleration)
+{
+	if (Velocity.IsZero() || !HasValidData() || HasRootMotion() || DeltaTime < MIN_TICK_TIME)
+	{
+		return;
+	}
+
+	Friction = FMath::Max(0.f, Friction);
+	BrakingDeceleration = FMath::Max(0.f, BrakingDeceleration);
+	const bool bZeroFriction = (Friction == 0.f);
+	const bool bZeroBraking = (BrakingDeceleration == 0.f);
+
+	if (bZeroFriction && bZeroBraking)
+	{
+		return;
+	}
+
+	const FVector OldVel = Velocity;
+
+	// subdivide braking to get reasonably consistent results at lower frame rates
+	// (important for packet loss situations w/ networking)
+	float RemainingTime = DeltaTime;
+	const float MaxTimeStep = (1.0f / 33.0f);
+
+	// Decelerate to brake to a stop
+	const FVector RevAccel = (bZeroBraking ? FVector::ZeroVector : (-BrakingDeceleration * GetSafeNormalPrecise(Velocity)));
+	while( RemainingTime >= MIN_TICK_TIME )
+	{
+		// Zero friction uses constant deceleration, so no need for iteration.
+		const float dt = ((RemainingTime > MaxTimeStep && !bZeroFriction) ? FMath::Min(MaxTimeStep, RemainingTime * 0.5f) : RemainingTime);
+		RemainingTime -= dt;
+
+		// apply friction and braking
+		Velocity = Velocity + ((-2.f * Friction) * Velocity + RevAccel) * dt ; 
+		
+		// Don't reverse direction
+		if ((Velocity | OldVel) <= 0.f)
+		{
+			Velocity = FVector::ZeroVector;
+			return;
+		}
+	}
+
+	// Clamp to zero if nearly zero, or if below min threshold and braking.
+	const float VSizeSq = Velocity.SizeSquared();
+	if (VSizeSq <= KINDA_SMALL_NUMBER || (!bZeroBraking && VSizeSq <= FMath::Square(BRAKE_TO_STOP_VELOCITY)))
+	{
+		Velocity = FVector::ZeroVector;
+	}
+}
+
+void UCharacterMovementComponent::PhysFlying(float deltaTime, int32 Iterations)
+{
+	if (deltaTime < MIN_TICK_TIME)
+	{
+		return;
+	}
+
+	if( !HasRootMotion() )
+	{
+		if( bCheatFlying && Acceleration.IsZero() )
+		{
+			Velocity = FVector::ZeroVector;
+		}
+		const float Friction = 0.5f * GetPhysicsVolume()->FluidFriction;
+		CalcVelocity(deltaTime, Friction, true, BrakingDecelerationFlying);
+	}
+
+	Iterations++;
+	bJustTeleported = false;
+
+	FVector OldLocation = CharacterOwner->GetActorLocation();
+	const FVector Adjusted = Velocity * deltaTime;
+	FHitResult Hit(1.f);
+	SafeMoveUpdatedComponent(Adjusted, CharacterOwner->GetActorRotation(), true, Hit);
+
+	if( Hit.Time < 1.f && CharacterOwner )
+	{
+		const FVector GravDir = FVector(0.f,0.f,-1.f);
+		const FVector VelDir = Velocity.GetSafeNormal();
+		const float UpDown = GravDir | VelDir;
+
+		bool bSteppedUp = false;
+		if( (FMath::Abs(Hit.ImpactNormal.Z) < 0.2f) && (UpDown < 0.5f) && (UpDown > -0.2f) && CanStepUp(Hit))
+		{
+			float stepZ = CharacterOwner->GetActorLocation().Z;
+			bSteppedUp = StepUp(GravDir, Adjusted * (1.f - Hit.Time), Hit);
+			if (bSteppedUp)
+			{
+				OldLocation.Z = CharacterOwner->GetActorLocation().Z + (OldLocation.Z - stepZ);
+			}
+		}
+		
+		if (!bSteppedUp)
+		{
+			//adjust and try again
+			HandleImpact(Hit, deltaTime, Adjusted);
+			SlideAlongSurface(Adjusted, (1.f - Hit.Time), Hit.Normal, Hit, true);
+		}
+	}
+	
+	if( !bJustTeleported && !HasRootMotion() && CharacterOwner )
+	{
+		Velocity = (CharacterOwner->GetActorLocation() - OldLocation) / deltaTime;
+	}
+}
+
+
+void UCharacterMovementComponent::PhysSwimming(float deltaTime, int32 Iterations)
+{
+	if (deltaTime < MIN_TICK_TIME)
+	{
+		return;
+	}
+
+	float NetFluidFriction  = 0.f;
+	float Depth = ImmersionDepth();
+	float NetBuoyancy = Buoyancy * Depth;
+	if( !HasRootMotion() && (Velocity.Z > 0.5f*GetMaxSpeed()) && (NetBuoyancy != 0.f) )
+	{
+		//damp positive Z out of water
+		Velocity.Z = Velocity.Z * Depth;
+	}
+	Iterations++;
+	FVector OldLocation = CharacterOwner->GetActorLocation();
+	bJustTeleported = false;
+	if( !HasRootMotion() )
+	{
+		const float Friction = 0.5f * GetPhysicsVolume()->FluidFriction * Depth;
+		CalcVelocity(deltaTime, Friction, true, BrakingDecelerationSwimming);
+		Velocity.Z += GetGravityZ() * deltaTime * (1.f - NetBuoyancy);
+	}
+
+	FVector Adjusted = Velocity * deltaTime;
+	FHitResult Hit(1.f);
+	float remainingTime = deltaTime * Swim(Adjusted, Hit);
+
+	//may have left water - if so, script might have set new physics mode
+	if ( !IsSwimming() )
+	{
+		StartNewPhysics(remainingTime, Iterations);
+		return;
+	}
+
+	if ( Hit.Time < 1.f && CharacterOwner)
+	{
+		const FVector GravDir = FVector(0.f,0.f,-1.f);
+		const FVector VelDir = Velocity.GetSafeNormal();
+		const float UpDown = GravDir | VelDir;
+
+		bool bSteppedUp = false;
+		if( (FMath::Abs(Hit.ImpactNormal.Z) < 0.2f) && (UpDown < 0.5f) && (UpDown > -0.2f) && CanStepUp(Hit))
+		{
+			float stepZ = CharacterOwner->GetActorLocation().Z;
+			const FVector RealVelocity = Velocity;
+			Velocity.Z = 1.f;	// HACK: since will be moving up, in case pawn leaves the water
+			bSteppedUp = StepUp(GravDir, Adjusted * (1.f - Hit.Time), Hit);
+			if (bSteppedUp)
+			{
+				//may have left water - if so, script might have set new physics mode
+				if (!IsSwimming())
+				{
+					StartNewPhysics(remainingTime, Iterations);
+					return;
+				}
+				OldLocation.Z = CharacterOwner->GetActorLocation().Z + (OldLocation.Z - stepZ);
+			}
+			Velocity = RealVelocity;
+		}
+
+		if (!bSteppedUp)
+		{
+			//adjust and try again
+			HandleImpact(Hit, deltaTime, Adjusted);
+			SlideAlongSurface(Adjusted, (1.f - Hit.Time), Hit.Normal, Hit, true);
+		}
+	}
+
+	if( !HasRootMotion() && !bJustTeleported && ((deltaTime - remainingTime) > KINDA_SMALL_NUMBER) && CharacterOwner )
+	{
+		bool bWaterJump = !GetPhysicsVolume()->bWaterVolume;
+		float velZ = Velocity.Z;
+		Velocity = (CharacterOwner->GetActorLocation() - OldLocation) / (deltaTime - remainingTime);
+		if (bWaterJump)
+		{
+			Velocity.Z = velZ;
+		}
+	}
+
+	if ( !GetPhysicsVolume()->bWaterVolume && IsSwimming() )
+	{
+		SetMovementMode(MOVE_Falling); //in case script didn't change it (w/ zone change)
+	}
+
+	//may have left water - if so, script might have set new physics mode
+	if ( !IsSwimming() )
+	{
+		StartNewPhysics(remainingTime, Iterations);
+	}
+}
+
+
+void UCharacterMovementComponent::StartSwimming(FVector OldLocation, FVector OldVelocity, float timeTick, float remainingTime, int32 Iterations)
+{
+	if (remainingTime < MIN_TICK_TIME || timeTick < MIN_TICK_TIME)
+	{
+		return;
+	}
+
+	if( !HasRootMotion() && !bJustTeleported )
+	{
+		Velocity = (CharacterOwner->GetActorLocation() - OldLocation)/timeTick; //actual average velocity
+		Velocity = 2.f*Velocity - OldVelocity; //end velocity has 2* accel of avg
+		Velocity = Velocity.GetClampedToMaxSize(GetPhysicsVolume()->TerminalVelocity);
+	}
+	const FVector End = FindWaterLine(CharacterOwner->GetActorLocation(), OldLocation);
+	float waterTime = 0.f;
+	if (End != CharacterOwner->GetActorLocation())
+	{	
+		const float ActualDist = (CharacterOwner->GetActorLocation() - OldLocation).Size();
+		if (ActualDist > KINDA_SMALL_NUMBER)
+		{
+			waterTime = timeTick * (End - CharacterOwner->GetActorLocation()).Size() / ActualDist;
+			remainingTime += waterTime;
+		}
+		MoveUpdatedComponent(End - CharacterOwner->GetActorLocation(), CharacterOwner->GetActorRotation(), true);
+	}
+	if ( !HasRootMotion() && CharacterOwner && (Velocity.Z > 2.f*SWIMBOBSPEED) && (Velocity.Z < 0.f)) //allow for falling out of water
+	{
+		Velocity.Z = SWIMBOBSPEED - Velocity.Size2D() * 0.7f; //smooth bobbing
+	}
+	if ( (remainingTime >= MIN_TICK_TIME) && (Iterations < MaxSimulationIterations) && CharacterOwner )
+	{
+		PhysSwimming(remainingTime, Iterations);
+	}
+}
+
+float UCharacterMovementComponent::Swim(FVector Delta, FHitResult &Hit)
+{
+	FVector Start = CharacterOwner->GetActorLocation();
+	float airTime = 0.f;
+	SafeMoveUpdatedComponent(Delta, CharacterOwner->GetActorRotation(), true, Hit);
+
+	if ( !GetPhysicsVolume()->bWaterVolume ) //then left water
+	{
+		const FVector End = FindWaterLine(Start,CharacterOwner->GetActorLocation());
+		const float DesiredDist = Delta.Size();
+		if (End != CharacterOwner->GetActorLocation() && DesiredDist > KINDA_SMALL_NUMBER)
+		{
+			airTime = (End - CharacterOwner->GetActorLocation()).Size() / DesiredDist;
+			if ( ((CharacterOwner->GetActorLocation() - Start) | (End - CharacterOwner->GetActorLocation())) > 0.f )
+			{
+				airTime = 0.f;
+			}
+			SafeMoveUpdatedComponent(End - CharacterOwner->GetActorLocation(), CharacterOwner->GetActorRotation(), true, Hit);
+		}
+	}
+	return airTime;
+}
+
+FVector UCharacterMovementComponent::FindWaterLine(FVector InWater, FVector OutofWater)
+{
+	FVector Result = OutofWater;
+
+	TArray<FHitResult> Hits;
+	static const FName NAME_FindWaterLine = FName(TEXT("FindWaterLine"));
+	GetWorld()->LineTraceMultiByChannel(Hits, OutofWater, InWater, UpdatedComponent->GetCollisionObjectType(), FCollisionQueryParams(NAME_FindWaterLine, true, CharacterOwner));
+
+	for( int32 HitIdx = 0; HitIdx < Hits.Num(); HitIdx++ )
+	{
+		const FHitResult& Check = Hits[HitIdx];
+		if ( !CharacterOwner->IsOwnedBy(Check.GetActor()) && !Check.Component.Get()->IsWorldGeometry() )
+		{
+			APhysicsVolume *W = Cast<APhysicsVolume>(Check.GetActor());
+			if ( W && W->bWaterVolume )
+			{
+				FVector Dir = (InWater - OutofWater).GetSafeNormal();
+				Result = Check.Location;
+				if ( W == GetPhysicsVolume() )
+					Result += 0.1f * Dir;
+				else
+					Result -= 0.1f * Dir;
+				break;
+			}
+		}
+	}
+
+	return Result;
+}
+
+void UCharacterMovementComponent::NotifyJumpApex() 
+{
+	if( CharacterOwner )
+	{
+		CharacterOwner->NotifyJumpApex();
+	}
+}
+
+
+FVector UCharacterMovementComponent::GetFallingLateralAcceleration(float DeltaTime)
+{
+	// No acceleration in Z
+	FVector FallAcceleration = FVector(Acceleration.X, Acceleration.Y, 0.f);
+
+	// bound acceleration, falling object has minimal ability to impact acceleration
+	if (!HasRootMotion() && FallAcceleration.SizeSquared2D() > 0.f)
+	{
+		FallAcceleration = GetAirControl(DeltaTime, AirControl, FallAcceleration);
+		FallAcceleration = FallAcceleration.GetClampedToMaxSize(GetMaxAcceleration());
+	}
+
+	return FallAcceleration;
+}
+
+
+FVector UCharacterMovementComponent::GetAirControl(float DeltaTime, float TickAirControl, const FVector& FallAcceleration)
+{
+	// Boost
+	if (TickAirControl != 0.f)
+	{
+		TickAirControl = BoostAirControl(DeltaTime, TickAirControl, FallAcceleration);
+	}
+
+	return TickAirControl * FallAcceleration;
+}
+
+
+float UCharacterMovementComponent::BoostAirControl(float DeltaTime, float TickAirControl, const FVector& FallAcceleration)
+{
+	// Allow a burst of initial acceleration
+	if (AirControlBoostMultiplier > 0.f && Velocity.SizeSquared2D() < FMath::Square(AirControlBoostVelocityThreshold))
+	{
+		TickAirControl = FMath::Min(1.f, AirControlBoostMultiplier * TickAirControl);
+	}
+
+	return TickAirControl;
+}
+
+
+void UCharacterMovementComponent::PhysFalling(float deltaTime, int32 Iterations)
+{
+	SCOPE_CYCLE_COUNTER(STAT_CharPhysFalling);
+
+	if (deltaTime < MIN_TICK_TIME)
+	{
+		return;
+	}
+
+	FVector FallAcceleration = GetFallingLateralAcceleration(deltaTime);
+	FallAcceleration.Z = 0.f;
+	const bool bHasAirControl = (FallAcceleration.SizeSquared2D() > 0.f);
+
+	float remainingTime = deltaTime;
+	while( (remainingTime >= MIN_TICK_TIME) && (Iterations < MaxSimulationIterations) )
+	{
+		Iterations++;
+		const float timeTick = GetSimulationTimeStep(remainingTime, Iterations);
+		remainingTime -= timeTick;
+		
+		const FVector OldLocation = CharacterOwner->GetActorLocation();
+		const FRotator PawnRotation = CharacterOwner->GetActorRotation();
+		bJustTeleported = false;
+
+		FVector OldVelocity = Velocity;
+		FVector VelocityNoAirControl = Velocity;
+
+		// Apply input
+		if (!HasRootMotion())
+		{
+			// Compute VelocityNoAirControl
+			if (bHasAirControl)
+			{
+				// Find velocity *without* acceleration.
+				TGuardValue<FVector> RestoreAcceleration(Acceleration, FVector::ZeroVector);
+				TGuardValue<FVector> RestoreVelocity(Velocity, Velocity);
+				Velocity.Z = 0.f;
+				CalcVelocity(timeTick, FallingLateralFriction, false, BrakingDecelerationFalling);
+				VelocityNoAirControl = FVector(Velocity.X, Velocity.Y, OldVelocity.Z);
+			}
+
+			// Compute Velocity
+			{
+				// Acceleration = FallAcceleration for CalcVelocity(), but we restore it after using it.
+				TGuardValue<FVector> RestoreAcceleration(Acceleration, FallAcceleration);
+				Velocity.Z = 0.f;
+				CalcVelocity(timeTick, FallingLateralFriction, false, BrakingDecelerationFalling);
+				Velocity.Z = OldVelocity.Z;
+			}
+
+			// Just copy Velocity to VelocityNoAirControl if they are the same (ie no acceleration).
+			if (!bHasAirControl)
+			{
+				VelocityNoAirControl = Velocity;
+			}
+		}
+
+		// Apply gravity
+		const FVector Gravity(0.f, 0.f, GetGravityZ());
+		Velocity = NewFallVelocity(Velocity, Gravity, timeTick);
+		VelocityNoAirControl = NewFallVelocity(VelocityNoAirControl, Gravity, timeTick);
+		const FVector AirControlAccel = (Velocity - VelocityNoAirControl) / timeTick;
+
+		if( bNotifyApex && CharacterOwner->Controller && (Velocity.Z <= 0.f) )
+		{
+			// Just passed jump apex since now going down
+			bNotifyApex = false;
+			NotifyJumpApex();
+		}
+
+
+		// Move
+		FHitResult Hit(1.f);
+		FVector Adjusted = 0.5f*(OldVelocity + Velocity) * timeTick;
+		SafeMoveUpdatedComponent( Adjusted, PawnRotation, true, Hit);
+		
+		if (!HasValidData())
+		{
+			return;
+		}
+		
+		float LastMoveTimeSlice = timeTick;
+		float subTimeTickRemaining = timeTick * (1.f - Hit.Time);
+		
+		if ( IsSwimming() ) //just entered water
+		{
+			remainingTime += subTimeTickRemaining;
+			StartSwimming(OldLocation, OldVelocity, timeTick, remainingTime, Iterations);
+			return;
+		}
+		else if ( Hit.bBlockingHit )
+		{
+			if (IsValidLandingSpot(UpdatedComponent->GetComponentLocation(), Hit))
+			{
+				remainingTime += subTimeTickRemaining;
+				ProcessLanded(Hit, remainingTime, Iterations);
+				return;
+			}
+			else
+			{
+				// Compute impact deflection based on final velocity, not integration step.
+				// This allows us to compute a new velocity from the deflected vector, and ensures the full gravity effect is included in the slide result.
+				Adjusted = Velocity * timeTick;
+
+				// See if we can convert a normally invalid landing spot (based on the hit result) to a usable one.
+				if (!Hit.bStartPenetrating && ShouldCheckForValidLandingSpot(timeTick, Adjusted, Hit))
+				{
+					const FVector PawnLocation = UpdatedComponent->GetComponentLocation();
+					FFindFloorResult FloorResult;
+					FindFloor(PawnLocation, FloorResult, false);
+					if (FloorResult.IsWalkableFloor() && IsValidLandingSpot(PawnLocation, FloorResult.HitResult))
+					{
+						remainingTime += subTimeTickRemaining;
+						ProcessLanded(FloorResult.HitResult, remainingTime, Iterations);
+						return;
+					}
+				}
+
+				HandleImpact(Hit, LastMoveTimeSlice, Adjusted);
+				
+				// If we've changed physics mode, abort.
+				if (!HasValidData() || !IsFalling())
+				{
+					return;
+				}
+
+				// Limit air control based on what we hit.
+				// We moved to the impact point using air control, but may want to deflect from there based on a limited air control acceleration.
+				if (bHasAirControl)
+				{
+					const bool bCheckLandingSpot = false; // we already checked above.
+					const FVector AirControlDeltaV = LimitAirControl(LastMoveTimeSlice, AirControlAccel, Hit, bCheckLandingSpot) * LastMoveTimeSlice;
+					Adjusted = (VelocityNoAirControl + AirControlDeltaV) * LastMoveTimeSlice;
+				}
+
+				const FVector OldHitNormal = Hit.Normal;
+				const FVector OldHitImpactNormal = Hit.ImpactNormal;				
+				FVector Delta = ComputeSlideVector(Adjusted, 1.f - Hit.Time, OldHitNormal, Hit);
+
+				// Compute velocity after deflection (only gravity component for RootMotion)
+				if (subTimeTickRemaining > KINDA_SMALL_NUMBER && !bJustTeleported)
+				{
+					const FVector NewVelocity = (Delta / subTimeTickRemaining);
+					Velocity = HasRootMotion() ? FVector(Velocity.X, Velocity.Y, NewVelocity.Z) : NewVelocity;
+				}
+
+				if (subTimeTickRemaining > KINDA_SMALL_NUMBER && (Delta | Adjusted) > 0.f)
+				{
+					// Move in deflected direction.
+					SafeMoveUpdatedComponent( Delta, PawnRotation, true, Hit);
+					
+					if (Hit.bBlockingHit)
+					{
+						// hit second wall
+						LastMoveTimeSlice = subTimeTickRemaining;
+						subTimeTickRemaining = subTimeTickRemaining * (1.f - Hit.Time);
+
+						if (IsValidLandingSpot(UpdatedComponent->GetComponentLocation(), Hit))
+						{
+							remainingTime += subTimeTickRemaining;
+							ProcessLanded(Hit, remainingTime, Iterations);
+							return;
+						}
+
+						HandleImpact(Hit, LastMoveTimeSlice, Delta);
+
+						// If we've changed physics mode, abort.
+						if (!HasValidData() || !IsFalling())
+						{
+							return;
+						}
+
+						// Act as if there was no air control on the last move when computing new deflection.
+						if (bHasAirControl && Hit.Normal.Z > VERTICAL_SLOPE_NORMAL_Z)
+						{
+							const FVector LastMoveNoAirControl = VelocityNoAirControl * LastMoveTimeSlice;
+							Delta = ComputeSlideVector(LastMoveNoAirControl, 1.f, OldHitNormal, Hit);
+						}
+
+						FVector PreTwoWallDelta = Delta;
+						TwoWallAdjust(Delta, Hit, OldHitNormal);
+
+						// Limit air control, but allow a slide along the second wall.
+						if (bHasAirControl)
+						{
+							const bool bCheckLandingSpot = false; // we already checked above.
+							const FVector AirControlDeltaV = LimitAirControl(subTimeTickRemaining, AirControlAccel, Hit, bCheckLandingSpot) * subTimeTickRemaining;
+
+							// Only allow if not back in to first wall
+							if (FVector::DotProduct(AirControlDeltaV, OldHitNormal) > 0.f)
+							{
+								Delta += (AirControlDeltaV * subTimeTickRemaining);
+							}
+						}
+
+						// Compute velocity after deflection (only gravity component for RootMotion)
+						if (subTimeTickRemaining > KINDA_SMALL_NUMBER && !bJustTeleported)
+						{
+							const FVector NewVelocity = (Delta / subTimeTickRemaining);
+							Velocity = HasRootMotion() ? FVector(Velocity.X, Velocity.Y, NewVelocity.Z) : NewVelocity;
+						}
+
+						// bDitch=true means that pawn is straddling two slopes, neither of which he can stand on
+						bool bDitch = ( (OldHitImpactNormal.Z > 0.f) && (Hit.ImpactNormal.Z > 0.f) && (FMath::Abs(Delta.Z) <= KINDA_SMALL_NUMBER) && ((Hit.ImpactNormal | OldHitImpactNormal) < 0.f) );
+						SafeMoveUpdatedComponent( Delta, PawnRotation, true, Hit);
+						if ( Hit.Time == 0 )
+						{
+							// if we are stuck then try to side step
+							FVector SideDelta = (OldHitNormal + Hit.ImpactNormal).GetSafeNormal2D();
+							if ( SideDelta.IsNearlyZero() )
+							{
+								SideDelta = FVector(OldHitNormal.Y, -OldHitNormal.X, 0).GetSafeNormal();
+							}
+							SafeMoveUpdatedComponent( SideDelta, PawnRotation, true, Hit);
+						}
+							
+						if ( bDitch || IsValidLandingSpot(UpdatedComponent->GetComponentLocation(), Hit) || Hit.Time == 0  )
+						{
+							remainingTime = 0.f;
+							ProcessLanded(Hit, remainingTime, Iterations);
+							return;
+						}
+						else if (GetPerchRadiusThreshold() > 0.f && Hit.Time == 1.f && OldHitImpactNormal.Z >= WalkableFloorZ)
+						{
+							// We might be in a virtual 'ditch' within our perch radius. This is rare.
+							const FVector PawnLocation = CharacterOwner->GetActorLocation();
+							const float ZMovedDist = FMath::Abs(PawnLocation.Z - OldLocation.Z);
+							const float MovedDist2DSq = (PawnLocation - OldLocation).SizeSquared2D();
+							if (ZMovedDist <= 0.2f * timeTick && MovedDist2DSq <= 4.f * timeTick)
+							{
+								Velocity.X += 0.25f * GetMaxSpeed() * (FMath::FRand() - 0.5f);
+								Velocity.Y += 0.25f * GetMaxSpeed() * (FMath::FRand() - 0.5f);
+								Velocity.Z = FMath::Max<float>(JumpZVelocity * 0.25f, 1.f);
+								Delta = Velocity * timeTick;
+								SafeMoveUpdatedComponent(Delta, PawnRotation, true, Hit);
+							}
+						}
+					}
+				}
+			}
+		}
+
+		if (Velocity.SizeSquared2D() <= KINDA_SMALL_NUMBER * 10.f)
+		{
+			Velocity.X = 0.f;
+			Velocity.Y = 0.f;
+		}
+	}
+}
+
+
+bool UCharacterMovementComponent::FindAirControlImpact(float DeltaTime, float AdditionalTime, const FVector& FallVelocity, const FVector& FallAcceleration, const FVector& Gravity, FHitResult& OutHitResult)
+{
+	// Test for slope to avoid using air control to climb walls.
+	FVector TestWalk = Velocity * DeltaTime;
+	if (AdditionalTime > 0.f)
+	{
+		const FVector PostGravityVelocity = NewFallVelocity(FallVelocity, Gravity, AdditionalTime);
+		TestWalk += ((FallAcceleration * AdditionalTime) + PostGravityVelocity) * AdditionalTime;
+	}
+	
+	if (!TestWalk.IsZero())
+	{
+		static const FName FallingTraceParamsTag = FName(TEXT("PhysFalling"));
+		FCollisionQueryParams CapsuleQuery(FallingTraceParamsTag, false, CharacterOwner);
+		FCollisionResponseParams ResponseParam;
+		InitCollisionParams(CapsuleQuery, ResponseParam);
+		const FVector CapsuleLocation = UpdatedComponent->GetComponentLocation();
+		const FCollisionShape CapsuleShape = GetPawnCapsuleCollisionShape(SHRINK_None);
+		
+		if (GetWorld()->SweepSingleByChannel(OutHitResult, CapsuleLocation, CapsuleLocation + TestWalk, FQuat::Identity, UpdatedComponent->GetCollisionObjectType(), CapsuleShape, CapsuleQuery, ResponseParam))
+		{
+			return true;
+		}
+	}
+
+	return false;
+}
+
+
+FVector UCharacterMovementComponent::LimitAirControl(float DeltaTime, const FVector& FallAcceleration, const FHitResult& HitResult, bool bCheckForValidLandingSpot)
+{
+	FVector Result(FallAcceleration);
+
+	if (HitResult.IsValidBlockingHit() && HitResult.Normal.Z > VERTICAL_SLOPE_NORMAL_Z)
+	{
+		if (!bCheckForValidLandingSpot || !IsValidLandingSpot(HitResult.Location, HitResult))
+		{
+			// If acceleration is into the wall, limit contribution.
+			if (FVector::DotProduct(FallAcceleration, HitResult.Normal) < 0.f)
+			{
+				// Allow movement parallel to the wall, but not into it because that may push us up.
+				const FVector Normal2D = HitResult.Normal.GetSafeNormal2D();
+				Result = FVector::VectorPlaneProject(FallAcceleration, Normal2D);
+			}
+		}
+	}
+	else if (HitResult.bStartPenetrating)
+	{
+		// Allow movement out of penetration.
+		return (FVector::DotProduct(Result, HitResult.Normal) > 0.f ? Result : FVector::ZeroVector);
+	}
+
+	return Result;
+}
+
+bool UCharacterMovementComponent::CheckLedgeDirection(const FVector& OldLocation, const FVector& SideStep, const FVector& GravDir)
+{
+	FVector SideDest = OldLocation + SideStep;
+	static const FName CheckLedgeDirectionName(TEXT("CheckLedgeDirection"));
+	FCollisionQueryParams CapsuleParams(CheckLedgeDirectionName, false, CharacterOwner);
+	FCollisionResponseParams ResponseParam;
+	InitCollisionParams(CapsuleParams, ResponseParam);
+	FCollisionShape CapsuleShape = GetPawnCapsuleCollisionShape(SHRINK_None);
+	const ECollisionChannel CollisionChannel = UpdatedComponent->GetCollisionObjectType();
+	FHitResult Result(1.f);
+	GetWorld()->SweepSingleByChannel(Result, OldLocation, SideDest, FQuat::Identity, CollisionChannel, CapsuleShape, CapsuleParams, ResponseParam);
+
+	if ( !Result.bBlockingHit || IsWalkable(Result) )
+	{
+		if ( !Result.bBlockingHit )
+		{
+			GetWorld()->SweepSingleByChannel(Result, SideDest, SideDest + GravDir * (MaxStepHeight + LedgeCheckThreshold), FQuat::Identity, CollisionChannel, CapsuleShape, CapsuleParams, ResponseParam);
+		}
+		if ( (Result.Time < 1.f) && IsWalkable(Result) )
+		{
+			return true;
+		}
+	}
+	return false;
+}
+
+
+FVector UCharacterMovementComponent::GetLedgeMove(const FVector& OldLocation, const FVector& Delta, const FVector& GravDir)
+{
+	// We have a ledge!
+	if (!HasValidData())
+	{
+		return FVector::ZeroVector;
+	}
+
+	// check which direction ledge goes
+	float DesiredDistSq = Delta.SizeSquared();
+
+	if ( DesiredDistSq > 0.f )
+	{
+		FVector SideDir(Delta.Y, -1.f * Delta.X, 0.f);
+		
+		// try left
+		if ( CheckLedgeDirection(OldLocation, SideDir, GravDir) )
+		{
+			return SideDir;
+		}
+
+		// try right
+		SideDir *= -1.f;
+		if ( CheckLedgeDirection(OldLocation, SideDir, GravDir) )
+		{
+			return SideDir;
+		}
+	}
+	return FVector::ZeroVector;
+}
+
+
+bool UCharacterMovementComponent::CanWalkOffLedges() const
+{
+	if (!bCanWalkOffLedgesWhenCrouching && IsCrouching())
+	{
+		return false;
+	}	
+
+	return bCanWalkOffLedges;
+}
+
+bool UCharacterMovementComponent::CheckFall(FHitResult &Hit, FVector Delta, FVector subLoc, float remainingTime, float timeTick, int32 Iterations, bool bMustJump)
+{
+	if (!HasValidData())
+	{
+		return false;
+	}
+
+	if (bMustJump || CanWalkOffLedges())
+	{
+		CharacterOwner->OnWalkingOffLedge();
+		if (IsMovingOnGround())
+		{
+			// If still walking, then fall. If not, assume the user set a different mode they want to keep.
+			StartFalling(Iterations, remainingTime, timeTick, Delta, subLoc);
+		}
+		return true;
+	}
+	return false;
+}
+
+void UCharacterMovementComponent::StartFalling(int32 Iterations, float remainingTime, float timeTick, const FVector& Delta, const FVector& subLoc)
+{
+	// start falling 
+	const float DesiredDist = Delta.Size();
+	const float ActualDist = (CharacterOwner->GetActorLocation() - subLoc).Size2D();
+	remainingTime = (DesiredDist < KINDA_SMALL_NUMBER) 
+					? 0.f
+					: remainingTime + timeTick * (1.f - FMath::Min(1.f,ActualDist/DesiredDist));
+
+	Velocity.Z = 0.f;			
+	if ( IsMovingOnGround() )
+	{
+		// This is to catch cases where the first frame of PIE is executed, and the
+		// level is not yet visible. In those cases, the player will fall out of the
+		// world... So, don't set MOVE_Falling straight away.
+		if ( !GIsEditor || (GetWorld()->HasBegunPlay() && (GetWorld()->GetTimeSeconds() >= 1.f)) )
+		{
+			SetMovementMode(MOVE_Falling); //default behavior if script didn't change physics
+		}
+		else
+		{
+			// Make sure that the floor check code continues processing during this delay.
+			bForceNextFloorCheck = true;
+		}
+	}
+	StartNewPhysics(remainingTime,Iterations);
+}
+
+
+void UCharacterMovementComponent::RevertMove(const FVector& OldLocation, UPrimitiveComponent* OldBase, const FVector& PreviousBaseLocation, const FFindFloorResult& OldFloor, bool bFailMove)
+{
+	//UE_LOG(LogCharacterMovement, Log, TEXT("RevertMove from %f %f %f to %f %f %f"), CharacterOwner->Location.X, CharacterOwner->Location.Y, CharacterOwner->Location.Z, OldLocation.X, OldLocation.Y, OldLocation.Z);
+	UpdatedComponent->SetWorldLocation(OldLocation, false);
+	
+	//UE_LOG(LogCharacterMovement, Log, TEXT("Now at %f %f %f"), CharacterOwner->Location.X, CharacterOwner->Location.Y, CharacterOwner->Location.Z);
+	bJustTeleported = false;
+	// if our previous base couldn't have moved or changed in any physics-affecting way, restore it
+	if (IsValid(OldBase) && 
+		(!MovementBaseUtility::IsDynamicBase(OldBase) ||
+		 (OldBase->Mobility == EComponentMobility::Static) ||
+		 (OldBase->GetComponentLocation() == PreviousBaseLocation)
+		)
+	   )
+	{
+		CurrentFloor = OldFloor;
+		SetBase(OldBase, OldFloor.HitResult.BoneName);
+	}
+	else
+	{
+		SetBase(NULL);
+	}
+
+	if ( bFailMove )
+	{
+		// end movement now
+		Velocity = FVector::ZeroVector;
+		Acceleration = FVector::ZeroVector;
+		//UE_LOG(LogCharacterMovement, Log, TEXT("%s FAILMOVE RevertMove"), *CharacterOwner->GetName());
+
+		if (PathFollowingComp.IsValid())
+		{
+			PathFollowingComp->AbortMove(TEXT("RevertMove"));
+		}
+	}
+}
+
+
+FVector UCharacterMovementComponent::ComputeGroundMovementDelta(const FVector& Delta, const FHitResult& RampHit, const bool bHitFromLineTrace) const
+{
+	const FVector FloorNormal = RampHit.ImpactNormal;
+	const FVector ContactNormal = RampHit.Normal;
+
+	if (FloorNormal.Z < (1.f - KINDA_SMALL_NUMBER) && FloorNormal.Z > KINDA_SMALL_NUMBER && ContactNormal.Z > KINDA_SMALL_NUMBER && !bHitFromLineTrace && IsWalkable(RampHit))
+	{
+		// Compute a vector that moves parallel to the surface, by projecting the horizontal movement direction onto the ramp.
+		const float FloorDotDelta = (FloorNormal | Delta);
+		FVector RampMovement(Delta.X, Delta.Y, -FloorDotDelta / FloorNormal.Z);
+		
+		if (bMaintainHorizontalGroundVelocity)
+		{
+			return RampMovement;
+		}
+		else
+		{
+			return RampMovement.GetSafeNormal() * Delta.Size();
+		}
+	}
+
+	return Delta;
+}
+
+void UCharacterMovementComponent::OnCharacterStuckInGeometry()
+{
+	UE_LOG(LogCharacterMovement, Log, TEXT("%s is stuck and failed to move!"), *CharacterOwner->GetName());
+
+	// Don't update velocity based on our (failed) change in position this update since we're stuck.
+	bJustTeleported = true;
+}
+
+void UCharacterMovementComponent::MoveAlongFloor(const FVector& InVelocity, float DeltaSeconds, FStepDownResult* OutStepDownResult)
+{
+	const FVector Delta = FVector(InVelocity.X, InVelocity.Y, 0.f) * DeltaSeconds;
+	
+	if (!CurrentFloor.IsWalkableFloor())
+	{
+		return;
+	}
+
+	// Move along the current floor
+	FHitResult Hit(1.f);
+	FVector RampVector = ComputeGroundMovementDelta(Delta, CurrentFloor.HitResult, CurrentFloor.bLineTrace);
+	SafeMoveUpdatedComponent(RampVector, CharacterOwner->GetActorRotation(), true, Hit);
+	float LastMoveTimeSlice = DeltaSeconds;
+	
+	if (Hit.bStartPenetrating)
+	{
+		// Allow this hit to be used as an impact we can deflect off, otherwise we do nothing the rest of the update and appear to hitch.
+		HandleImpact(Hit);
+		SlideAlongSurface(Delta, 1.f, Hit.Normal, Hit, true);
+
+		if (Hit.bStartPenetrating)
+		{
+			OnCharacterStuckInGeometry();
+		}
+	}
+	else if (Hit.IsValidBlockingHit())
+	{
+		// We impacted something (most likely another ramp, but possibly a barrier).
+		float PercentTimeApplied = Hit.Time;
+		if ((Hit.Time > 0.f) && (Hit.Normal.Z > KINDA_SMALL_NUMBER) && IsWalkable(Hit))
+		{
+			// Another walkable ramp.
+			const float InitialPercentRemaining = 1.f - PercentTimeApplied;
+			RampVector = ComputeGroundMovementDelta(Delta * InitialPercentRemaining, Hit, false);
+			LastMoveTimeSlice = InitialPercentRemaining * LastMoveTimeSlice;
+			SafeMoveUpdatedComponent(RampVector, CharacterOwner->GetActorRotation(), true, Hit);
+
+			const float SecondHitPercent = Hit.Time * InitialPercentRemaining;
+			PercentTimeApplied = FMath::Clamp(PercentTimeApplied + SecondHitPercent, 0.f, 1.f);
+		}
+
+		if (Hit.IsValidBlockingHit())
+		{
+			if (CanStepUp(Hit) || (CharacterOwner->GetMovementBase() != NULL && CharacterOwner->GetMovementBase()->GetOwner() == Hit.GetActor()))
+			{
+				// hit a barrier, try to step up
+				const FVector GravDir(0.f, 0.f, -1.f);
+				if (!StepUp(GravDir, Delta * (1.f - PercentTimeApplied), Hit, OutStepDownResult))
+				{
+					UE_LOG(LogCharacterMovement, Verbose, TEXT("- StepUp (ImpactNormal %s, Normal %s"), *Hit.ImpactNormal.ToString(), *Hit.Normal.ToString());
+					HandleImpact(Hit, LastMoveTimeSlice, RampVector);
+					SlideAlongSurface(Delta, 1.f - PercentTimeApplied, Hit.Normal, Hit, true);
+				}
+				else
+				{
+					// Don't recalculate velocity based on this height adjustment, if considering vertical adjustments.
+					UE_LOG(LogCharacterMovement, Verbose, TEXT("+ StepUp (ImpactNormal %s, Normal %s"), *Hit.ImpactNormal.ToString(), *Hit.Normal.ToString());
+					bJustTeleported |= !bMaintainHorizontalGroundVelocity;
+				}
+			}
+			else if ( Hit.Component.IsValid() && !Hit.Component.Get()->CanCharacterStepUp(CharacterOwner) )
+			{
+				HandleImpact(Hit, LastMoveTimeSlice, RampVector);
+				SlideAlongSurface(Delta, 1.f - PercentTimeApplied, Hit.Normal, Hit, true);
+			}
+		}
+	}
+}
+
+
+void UCharacterMovementComponent::MaintainHorizontalGroundVelocity()
+{
+	if (Velocity.Z != 0.f)
+	{
+		if (bMaintainHorizontalGroundVelocity)
+		{
+			// Ramp movement already maintained the velocity, so we just want to remove the vertical component.
+			Velocity.Z = 0.f;
+		}
+		else
+		{
+			// Rescale velocity to be horizontal but maintain magnitude of last update.
+			Velocity = Velocity.GetSafeNormal2D() * Velocity.Size();
+		}
+	}
+}
+
+
+void UCharacterMovementComponent::PhysWalking(float deltaTime, int32 Iterations)
+{
+	SCOPE_CYCLE_COUNTER(STAT_CharPhysWalking);
+
+
+	if (deltaTime < MIN_TICK_TIME)
+	{
+		return;
+	}
+
+	if (!CharacterOwner || (!CharacterOwner->Controller && !bRunPhysicsWithNoController && !HasRootMotion()))
+	{
+		Acceleration = FVector::ZeroVector;
+		Velocity = FVector::ZeroVector;
+		return;
+	}
+
+	if (!UpdatedComponent->IsCollisionEnabled())
+	{
+		SetMovementMode(MOVE_Walking);
+		return;
+	}
+
+	checkf(!Velocity.ContainsNaN(), TEXT("PhysWalking: Velocity contains NaN before Iteration (%s: %s)\n%s"), *GetPathNameSafe(this), *GetPathNameSafe(GetOuter()), *Velocity.ToString());
+	
+	bJustTeleported = false;
+	bool bCheckedFall = false;
+	bool bTriedLedgeMove = false;
+	float remainingTime = deltaTime;
+
+	// Perform the move
+	while ( (remainingTime >= MIN_TICK_TIME) && (Iterations < MaxSimulationIterations) && CharacterOwner && (CharacterOwner->Controller || bRunPhysicsWithNoController || HasRootMotion()) )
+	{
+		Iterations++;
+		bJustTeleported = false;
+		const float timeTick = GetSimulationTimeStep(remainingTime, Iterations);
+		remainingTime -= timeTick;
+
+		// Save current values
+		UPrimitiveComponent * const OldBase = GetMovementBase();
+		const FVector PreviousBaseLocation = (OldBase != NULL) ? OldBase->GetComponentLocation() : FVector::ZeroVector;
+		const FVector OldLocation = UpdatedComponent->GetComponentLocation();
+		const FFindFloorResult OldFloor = CurrentFloor;
+
+		// Ensure velocity is horizontal.
+		MaintainHorizontalGroundVelocity();
+		Velocity.Z = 0.f;
+		const FVector OldVelocity = Velocity;
+
+		// Apply acceleration
+		Acceleration.Z = 0.f;
+		if( !HasRootMotion() )
+		{
+			CalcVelocity(timeTick, GroundFriction, false, BrakingDecelerationWalking);
+		}
+		checkf(!Velocity.ContainsNaN(), TEXT("PhysWalking: Velocity contains NaN after CalcVelocity (%s: %s)\n%s"), *GetPathNameSafe(this), *GetPathNameSafe(GetOuter()), *Velocity.ToString());
+
+		// Compute move parameters
+		const FVector MoveVelocity = Velocity;
+		const FVector Delta = timeTick * MoveVelocity;
+		const bool bZeroDelta = Delta.IsNearlyZero();
+		FStepDownResult StepDownResult;
+
+		if ( bZeroDelta )
+		{
+			remainingTime = 0.f;
+		}
+		else
+		{
+			// try to move forward
+			MoveAlongFloor(MoveVelocity, timeTick, &StepDownResult);
+
+			if ( IsFalling() )
+			{
+				// pawn decided to jump up
+				const float DesiredDist = Delta.Size();
+				if (DesiredDist > KINDA_SMALL_NUMBER)
+				{
+					const float ActualDist = (CharacterOwner->GetActorLocation() - OldLocation).Size2D();
+					remainingTime += timeTick * (1.f - FMath::Min(1.f,ActualDist/DesiredDist));
+				}
+				StartNewPhysics(remainingTime,Iterations);
+				return;
+			}
+			else if ( IsSwimming() ) //just entered water
+			{
+				StartSwimming(OldLocation, OldVelocity, timeTick, remainingTime, Iterations);
+				return;
+			}
+		}
+
+		// Update floor.
+		// StepUp might have already done it for us.
+		if (StepDownResult.bComputedFloor)
+		{
+			CurrentFloor = StepDownResult.FloorResult;
+		}
+		else
+		{
+			FindFloor(UpdatedComponent->GetComponentLocation(), CurrentFloor, bZeroDelta, NULL);
+		}
+
+		// check for ledges here
+		const bool bCheckLedges = !CanWalkOffLedges();
+		if ( bCheckLedges && !CurrentFloor.IsWalkableFloor() )
+		{
+			// calculate possible alternate movement
+			const FVector GravDir = FVector(0.f,0.f,-1.f);
+			const FVector NewDelta = bTriedLedgeMove ? FVector::ZeroVector : GetLedgeMove(OldLocation, Delta, GravDir);
+			if ( !NewDelta.IsZero() )
+			{
+				// first revert this move
+				RevertMove(OldLocation, OldBase, PreviousBaseLocation, OldFloor, false);
+
+				// avoid repeated ledge moves if the first one fails
+				bTriedLedgeMove = true;
+
+				// Try new movement direction
+				Velocity = NewDelta/timeTick;
+				remainingTime += timeTick;
+				continue;
+			}
+			else
+			{
+				// see if it is OK to jump
+				// @todo collision : only thing that can be problem is that oldbase has world collision on
+				bool bMustJump = bZeroDelta || (OldBase == NULL || (!OldBase->IsCollisionEnabled() && MovementBaseUtility::IsDynamicBase(OldBase)));
+				if ( (bMustJump || !bCheckedFall) && CheckFall(CurrentFloor.HitResult, Delta, OldLocation, remainingTime, timeTick, Iterations, bMustJump) )
+				{
+					return;
+				}
+				bCheckedFall = true;
+
+				// revert this move
+				RevertMove(OldLocation, OldBase, PreviousBaseLocation, OldFloor, true);
+				remainingTime = 0.f;
+				break;
+			}
+		}
+		else
+		{
+			// Validate the floor check
+			if (CurrentFloor.IsWalkableFloor())
+			{
+				if (ShouldCatchAir(OldFloor, CurrentFloor))
+				{
+					CharacterOwner->OnWalkingOffLedge();
+					if (IsMovingOnGround())
+					{
+						// If still walking, then fall. If not, assume the user set a different mode they want to keep.
+						StartFalling(Iterations, remainingTime, timeTick, Delta, OldLocation);
+					}
+					return;
+				}
+
+				AdjustFloorHeight();
+				SetBase(CurrentFloor.HitResult.Component.Get(), CurrentFloor.HitResult.BoneName);
+			}
+			else if (CurrentFloor.HitResult.bStartPenetrating && remainingTime <= 0.f)
+			{
+				// The floor check failed because it started in penetration
+				// We do not want to try to move downward because the downward sweep failed, rather we'd like to try to pop out of the floor.
+				FHitResult Hit(CurrentFloor.HitResult);
+				Hit.TraceEnd = Hit.TraceStart + FVector(0.f, 0.f, MAX_FLOOR_DIST);
+				const FVector RequestedAdjustment = GetPenetrationAdjustment(Hit);
+				ResolvePenetration(RequestedAdjustment, Hit, CharacterOwner->GetActorRotation());
+			}
+
+			// check if just entered water
+			if ( IsSwimming() )
+			{
+				StartSwimming(OldLocation, Velocity, timeTick, remainingTime, Iterations);
+				return;
+			}
+
+			// See if we need to start falling.
+			if (!CurrentFloor.IsWalkableFloor() && !CurrentFloor.HitResult.bStartPenetrating)
+			{
+				const bool bMustJump = bJustTeleported || bZeroDelta || (OldBase == NULL || (!OldBase->IsCollisionEnabled() && MovementBaseUtility::IsDynamicBase(OldBase)));
+				if ((bMustJump || !bCheckedFall) && CheckFall(CurrentFloor.HitResult, Delta, OldLocation, remainingTime, timeTick, Iterations, bMustJump) )
+				{
+					return;
+				}
+				bCheckedFall = true;
+			}
+		}
+
+
+		// Allow overlap events and such to change physics state and velocity
+		if (IsMovingOnGround())
+		{
+			// Make velocity reflect actual move
+			if( !bJustTeleported && !HasRootMotion() && timeTick >= MIN_TICK_TIME)
+			{
+				Velocity = (CharacterOwner->GetActorLocation() - OldLocation) / timeTick;
+			}
+		}
+
+		// If we didn't move at all this iteration then abort (since future iterations will also be stuck).
+		if (CharacterOwner->GetActorLocation() == OldLocation)
+		{
+			remainingTime = 0.f;
+			break;
+		}	
+	}
+
+	if (IsMovingOnGround())
+	{
+		MaintainHorizontalGroundVelocity();
+	}
+}
+
+void UCharacterMovementComponent::PhysNavWalking(float deltaTime, int32 Iterations)
+{
+	SCOPE_CYCLE_COUNTER(STAT_CharPhysNavWalking);
+
+	if (deltaTime < MIN_TICK_TIME)
+	{
+		return;
+	}
+
+	if ((!CharacterOwner || !CharacterOwner->Controller) && !bRunPhysicsWithNoController && !HasRootMotion())
+	{
+		Acceleration = FVector::ZeroVector;
+		Velocity = FVector::ZeroVector;
+		return;
+	}
+
+	// Ensure velocity is horizontal.
+	MaintainHorizontalGroundVelocity();
+	checkf(!Velocity.ContainsNaN(), TEXT("PhysNavWalking: Velocity contains NaN before CalcVelocity (%s: %s)\n%s"), *GetPathNameSafe(this), *GetPathNameSafe(GetOuter()), *Velocity.ToString());
+
+	//bound acceleration
+	Acceleration.Z = 0.f;
+	if (!HasRootMotion())
+	{
+		CalcVelocity(deltaTime, GroundFriction, false, BrakingDecelerationWalking);
+		checkf(!Velocity.ContainsNaN(), TEXT("PhysNavWalking: Velocity contains NaN after CalcVelocity (%s: %s)\n%s"), *GetPathNameSafe(this), *GetPathNameSafe(GetOuter()), *Velocity.ToString());
+	}
+
+	Iterations++;
+
+	FVector DesiredMove = Velocity;
+	DesiredMove.Z = 0.f;
+
+	const FVector OldLocation = GetActorFeetLocation();
+	const FVector DeltaMove = DesiredMove * deltaTime;
+
+	FVector AdjustedDest = OldLocation + DeltaMove;
+	FNavLocation DestNavLocation;
+
+	bool bSameNavLocation = false;
+	if (CachedNavLocation.NodeRef != INVALID_NAVNODEREF)
+	{
+		if (bProjectNavMeshWalking)
+		{
+			bSameNavLocation = (OldLocation - CachedNavLocation.Location).SizeSquared2D() <= KINDA_SMALL_NUMBER;
+		}
+		else
+		{
+			bSameNavLocation = CachedNavLocation.Location.Equals(OldLocation);
+		}
+	}
+		
+
+	if (DeltaMove.IsNearlyZero() && bSameNavLocation)
+	{
+		DestNavLocation = CachedNavLocation;
+		UE_LOG(LogNavMeshMovement, VeryVerbose, TEXT("%s using cached navmesh location! (bProjectNavMeshWalking = %d)"), *GetNameSafe(CharacterOwner), bProjectNavMeshWalking);
+	}
+	else
+	{
+		SCOPE_CYCLE_COUNTER(STAT_CharNavProjectPoint);
+
+		// Start the trace from the Z location of the last valid trace.
+		// Otherwise if we are projecting our location to the underlying geometry and it's far above or below the navmesh,
+		// we'll follow that geometry's plane out of range of valid navigation.
+		if (bSameNavLocation && bProjectNavMeshWalking)
+		{
+			AdjustedDest.Z = CachedNavLocation.Location.Z;
+		}
+
+		// Find the point on the NavMesh
+		const bool bHasNavigationData = FindNavFloor(AdjustedDest, DestNavLocation);
+		if (!bHasNavigationData)
+		{
+			SetMovementMode(MOVE_Walking);
+			return;
+		}
+
+		CachedNavLocation = DestNavLocation;
+	}
+
+	if (DestNavLocation.NodeRef != INVALID_NAVNODEREF)
+	{
+		FVector NewLocation(AdjustedDest.X, AdjustedDest.Y, DestNavLocation.Location.Z);
+		if (bProjectNavMeshWalking)
+		{
+			SCOPE_CYCLE_COUNTER(STAT_CharNavProjectLocation);
+			const float TotalCapsuleHeight = CharacterOwner->GetCapsuleComponent()->GetScaledCapsuleHalfHeight() * 2.0f;
+			const float UpOffset = TotalCapsuleHeight * FMath::Max(0.f, NavMeshProjectionHeightScaleUp);
+			const float DownOffset = TotalCapsuleHeight * FMath::Max(0.f, NavMeshProjectionHeightScaleDown);
+			NewLocation = ProjectLocationFromNavMesh(deltaTime, OldLocation, NewLocation, UpOffset, DownOffset);
+		}
+
+		FVector AdjustedDelta = NewLocation - OldLocation;
+
+		if (!AdjustedDelta.IsNearlyZero())
+		{
+			const bool bSweep = UpdatedPrimitive ? UpdatedPrimitive->bGenerateOverlapEvents : false;
+			FHitResult HitResult;
+			SafeMoveUpdatedComponent(AdjustedDelta, CharacterOwner->GetActorRotation(), bSweep, HitResult);
+		}
+
+		// Update velocity to reflect actual move
+		if (!bJustTeleported && !HasRootMotion())
+		{
+			Velocity = (GetActorFeetLocation() - OldLocation) / deltaTime;
+			MaintainHorizontalGroundVelocity();
+		}
+
+		bJustTeleported = false;
+	}
+	else
+	{
+		StartFalling(Iterations, deltaTime, deltaTime, DeltaMove, OldLocation);
+	}
+}
+
+bool UCharacterMovementComponent::FindNavFloor(const FVector& TestLocation, FNavLocation& NavFloorLocation) const
+{
+	UNavigationSystem* NavSys = UNavigationSystem::GetCurrent(GetWorld());
+	if (NavSys == nullptr)
+	{
+		return false;
+	}
+
+	const ANavigationData* NavData = nullptr;
+	INavAgentInterface* MyNavAgent = Cast<INavAgentInterface>(CharacterOwner);
+	float SearchRadius = 0.0f;
+	float SearchHeight = 100.0f;
+	if (MyNavAgent)
+	{
+		const FNavAgentProperties& AgentProps = MyNavAgent->GetNavAgentPropertiesRef();
+		NavData = NavSys->GetNavDataForProps(AgentProps);
+		SearchRadius = AgentProps.AgentRadius * 2.0f;
+		SearchHeight = AgentProps.AgentHeight * AgentProps.NavWalkingSearchHeightScale;
+	}
+	if (NavData == nullptr)
+	{
+		NavData = NavSys->GetMainNavData();
+	}
+
+	const ARecastNavMesh* NavMeshData = Cast<const ARecastNavMesh>(NavData);
+	if (NavMeshData == nullptr)
+	{
+		return false;
+	}
+
+	NavData->ProjectPoint(TestLocation, NavFloorLocation, FVector(SearchRadius, SearchRadius, SearchHeight));
+	return true;
+}
+
+FVector UCharacterMovementComponent::ProjectLocationFromNavMesh(float DeltaSeconds, const FVector& CurrentFeetLocation, const FVector& TargetNavLocation, float UpOffset, float DownOffset)
+{
+	SCOPE_CYCLE_COUNTER(STAT_CharNavProjectLocation);
+
+	FVector NewLocation = TargetNavLocation;
+
+	const FVector RayCastOffsetUp(0.0f, 0.0f, UpOffset);
+	const FVector RayCastOffsetDown(0.0f, 0.0f, DownOffset);
+
+	NavMeshProjectionTimer -= DeltaSeconds;
+	if (NavMeshProjectionTimer <= 0.0f)
+	{
+		// We can skip this trace if we are checking at the same location as the last trace (ie, we haven't moved).
+		const bool bCachedLocationStillValid = (!bAlwaysCheckFloor &&
+												CachedProjectedNavMeshHitResult.bBlockingHit &&
+												CachedProjectedNavMeshHitResult.TraceStart == (TargetNavLocation + RayCastOffsetUp) &&
+												CachedProjectedNavMeshHitResult.TraceEnd == (TargetNavLocation - RayCastOffsetDown));
+
+		if (!bCachedLocationStillValid)
+		{
+			UE_LOG(LogNavMeshMovement, VeryVerbose, TEXT("ProjectLocationFromNavMesh(): %s interval: %.3f velocity: %s"), *GetNameSafe(CharacterOwner), NavMeshProjectionInterval, *Velocity.ToString());
+
+			// raycast to underlying mesh to allow us to more closely follow geometry
+			// we use static objects here as a best approximation to accept only objects that
+			// influence navmesh generation
+			static const FName ProjectLocationName(TEXT("NavProjectLocation"));
+			FCollisionQueryParams Params(ProjectLocationName, false);
+			FCollisionResponseParams ResponseParams(ECR_Ignore); // ignore everything
+			ResponseParams.CollisionResponse.SetResponse(ECC_WorldStatic, ECR_Block); // get blocked only by WorldStatic
+			GetWorld()->LineTraceSingleByChannel(CachedProjectedNavMeshHitResult, TargetNavLocation + RayCastOffsetUp, TargetNavLocation - RayCastOffsetDown, ECC_WorldStatic, Params, ResponseParams);
+
+			// discard result if we were already inside something
+			if (CachedProjectedNavMeshHitResult.bStartPenetrating)
+			{
+				CachedProjectedNavMeshHitResult.Reset();
+			}
+		}
+		else
+		{
+			UE_LOG(LogNavMeshMovement, VeryVerbose, TEXT("ProjectLocationFromNavMesh(): %s interval: %.3f velocity: %s [SKIP TRACE]"), *GetNameSafe(CharacterOwner), NavMeshProjectionInterval, *Velocity.ToString());
+		}
+
+		// Wrap around to maintain same relative offset to tick time changes.
+		// Prevents large framerate spikes from aligning multiple characters to the same frame (if they start staggered, they will now remain staggered).
+		float ModTime = 0.f;
+		if (NavMeshProjectionInterval > SMALL_NUMBER)
+		{
+			ModTime = FMath::Fmod(FMath::Abs(NavMeshProjectionTimer), NavMeshProjectionInterval);
+		}
+
+		NavMeshProjectionTimer = NavMeshProjectionInterval - ModTime;
+	}
+		
+	// project to last plane we found
+	if (CachedProjectedNavMeshHitResult.bBlockingHit)
+	{
+		FVector ProjectedPoint = FMath::LinePlaneIntersection(TargetNavLocation + RayCastOffsetUp, TargetNavLocation - RayCastOffsetDown, CachedProjectedNavMeshHitResult.Location, CachedProjectedNavMeshHitResult.Normal);
+
+		// Limit to not be too far above or below NavMesh location
+		ProjectedPoint.Z = FMath::Clamp(ProjectedPoint.Z, TargetNavLocation.Z - DownOffset, TargetNavLocation.Z + UpOffset);
+
+		// Interp for smoother updates (less "pop" when trace hits something new). 0 interp speed is instant.
+		const float InterpSpeed = FMath::Max(0.f, NavMeshProjectionInterpSpeed);
+		ProjectedPoint.Z = FMath::FInterpTo(CurrentFeetLocation.Z, ProjectedPoint.Z, DeltaSeconds, InterpSpeed);
+		ProjectedPoint.Z = FMath::Clamp(ProjectedPoint.Z, TargetNavLocation.Z - DownOffset, TargetNavLocation.Z + UpOffset);
+
+		// Final result
+		NewLocation.Z = ProjectedPoint.Z;
+	}
+
+	return NewLocation;
+}
+
+void UCharacterMovementComponent::PhysCustom(float deltaTime, int32 Iterations)
+{
+	if (CharacterOwner)
+	{
+		CharacterOwner->K2_UpdateCustomMovement(deltaTime);
+	}
+}
+
+
+bool UCharacterMovementComponent::ShouldCatchAir(const FFindFloorResult& OldFloor, const FFindFloorResult& NewFloor)
+{
+	return false;
+}
+
+void UCharacterMovementComponent::AdjustFloorHeight()
+{
+	// If we have a floor check that hasn't hit anything, don't adjust height.
+	if (!CurrentFloor.bBlockingHit)
+	{
+		return;
+	}
+
+	const float OldFloorDist = CurrentFloor.FloorDist;
+	if (CurrentFloor.bLineTrace && OldFloorDist < MIN_FLOOR_DIST)
+	{
+		// This would cause us to scale unwalkable walls
+		return;
+	}
+
+	// Move up or down to maintain floor height.
+	if (OldFloorDist < MIN_FLOOR_DIST || OldFloorDist > MAX_FLOOR_DIST)
+	{
+		FHitResult AdjustHit(1.f);
+		const float InitialZ = UpdatedComponent->GetComponentLocation().Z;
+		const float AvgFloorDist = (MIN_FLOOR_DIST + MAX_FLOOR_DIST) * 0.5f;
+		const float MoveDist = AvgFloorDist - OldFloorDist;
+		SafeMoveUpdatedComponent( FVector(0.f,0.f,MoveDist), CharacterOwner->GetActorRotation(), true, AdjustHit );
+		UE_LOG(LogCharacterMovement, VeryVerbose, TEXT("Adjust floor height %.3f (Hit = %d)"), MoveDist, AdjustHit.bBlockingHit);
+
+		if (!AdjustHit.IsValidBlockingHit())
+		{
+			CurrentFloor.FloorDist += MoveDist;
+		}
+		else if (MoveDist > 0.f)
+		{
+			const float CurrentZ = UpdatedComponent->GetComponentLocation().Z;
+			CurrentFloor.FloorDist += CurrentZ - InitialZ;
+		}
+		else
+		{
+			checkSlow(MoveDist < 0.f);
+			const float CurrentZ = UpdatedComponent->GetComponentLocation().Z;
+			CurrentFloor.FloorDist = CurrentZ - AdjustHit.Location.Z;
+			if (IsWalkable(AdjustHit))
+			{
+				CurrentFloor.SetFromSweep(AdjustHit, CurrentFloor.FloorDist, true);
+			}
+		}
+
+		// Don't recalculate velocity based on this height adjustment, if considering vertical adjustments.
+		// Also avoid it if we moved out of penetration
+		bJustTeleported |= !bMaintainHorizontalGroundVelocity || (OldFloorDist < 0.f);
+	}
+}
+
+
+void UCharacterMovementComponent::StopActiveMovement() 
+{ 
+	Super::StopActiveMovement();
+
+	Acceleration = FVector::ZeroVector; 
+	bHasRequestedVelocity = false;
+	RequestedVelocity = FVector::ZeroVector;
+}
+
+void UCharacterMovementComponent::ProcessLanded(const FHitResult& Hit, float remainingTime, int32 Iterations)
+{
+	if( CharacterOwner && CharacterOwner->ShouldNotifyLanded(Hit) )
+	{
+		CharacterOwner->Landed(Hit);
+	}
+	if( IsFalling() )
+	{
+		if (GroundMovementMode == MOVE_NavWalking)
+		{
+			// verify navmesh projection and current floor
+			// otherwise movement will be stuck in infinite loop:
+			// navwalking -> (no navmesh) -> falling -> (standing on something) -> navwalking -> ....
+
+			const FVector TestLocation = GetActorFeetLocation();
+			FNavLocation NavLocation;
+
+			const bool bHasNavigationData = FindNavFloor(TestLocation, NavLocation);
+			if (!bHasNavigationData || NavLocation.NodeRef == INVALID_NAVNODEREF)
+			{
+				GroundMovementMode = MOVE_Walking;
+				UE_LOG(LogNavMeshMovement, Verbose, TEXT("ProcessLanded(): %s tried to go to NavWalking but couldn't find NavMesh! Using Walking instead."), *GetNameSafe(CharacterOwner));
+			}
+		}
+
+		SetPostLandedPhysics(Hit);
+	}
+	if (PathFollowingComp.IsValid())
+	{
+		PathFollowingComp->OnLanded();
+	}
+
+	StartNewPhysics(remainingTime, Iterations);
+}
+
+void UCharacterMovementComponent::SetPostLandedPhysics(const FHitResult& Hit)
+{
+	if( CharacterOwner )
+	{
+		if ( GetPhysicsVolume()->bWaterVolume && CanEverSwim() )
+		{
+			SetMovementMode(MOVE_Swimming);
+		}
+		else
+		{
+			const FVector PreImpactAccel = Acceleration + (IsFalling() ? FVector(0.f, 0.f, GetGravityZ()) : FVector::ZeroVector);
+			const FVector PreImpactVelocity = Velocity;
+			SetMovementMode(GroundMovementMode);
+			ApplyImpactPhysicsForces(Hit, PreImpactAccel, PreImpactVelocity);
+		}
+	}
+}
+
+void UCharacterMovementComponent::SetNavWalkingPhysics(bool bEnable)
+{
+	if (UpdatedPrimitive)
+	{
+		if (bEnable)
+		{
+			UpdatedPrimitive->SetCollisionResponseToChannel(ECC_WorldStatic, ECR_Ignore);
+			UpdatedPrimitive->SetCollisionResponseToChannel(ECC_WorldDynamic, ECR_Ignore);
+			CachedProjectedNavMeshHitResult.Reset();
+
+			// Stagger timed updates so many different characters spawned at the same time don't update on the same frame.
+			// Initially we want an immediate update though, so set time to a negative randomized range.
+			NavMeshProjectionTimer = (NavMeshProjectionInterval > 0.f) ? FMath::FRandRange(-NavMeshProjectionInterval, 0.f) : 0.f;
+		}
+		else
+		{
+			UPrimitiveComponent* DefaultCapsule = nullptr;
+			if (CharacterOwner && CharacterOwner->GetCapsuleComponent() == UpdatedComponent)
+			{
+				ACharacter* DefaultCharacter = CharacterOwner->GetClass()->GetDefaultObject<ACharacter>();
+				DefaultCapsule = DefaultCharacter ? DefaultCharacter->GetCapsuleComponent() : nullptr;
+			}
+
+			if (DefaultCapsule)
+			{
+				UpdatedPrimitive->SetCollisionResponseToChannel(ECC_WorldStatic, DefaultCapsule->GetCollisionResponseToChannel(ECC_WorldStatic));
+				UpdatedPrimitive->SetCollisionResponseToChannel(ECC_WorldDynamic, DefaultCapsule->GetCollisionResponseToChannel(ECC_WorldDynamic));
+			}
+			else
+			{
+				UE_LOG(LogCharacterMovement, Warning, TEXT("Can't revert NavWalking collision settings for %s.%s"),
+					*GetNameSafe(CharacterOwner), *GetNameSafe(UpdatedComponent));
+			}
+		}
+	}
+}
+
+bool UCharacterMovementComponent::TryToLeaveNavWalking()
+{
+	SetNavWalkingPhysics(false);
+
+	bool bCanTeleport = true;
+	if (CharacterOwner)
+	{
+		FVector CollisionFreeLocation = CharacterOwner->GetActorLocation();
+		bCanTeleport = GetWorld()->FindTeleportSpot(CharacterOwner, CollisionFreeLocation, CharacterOwner->GetActorRotation());
+		if (bCanTeleport)
+		{
+			CharacterOwner->SetActorLocation(CollisionFreeLocation);
+		}
+		else
+		{
+			SetNavWalkingPhysics(true);
+		}
+	}
+
+	bWantsToLeaveNavWalking = !bCanTeleport;
+	return bCanTeleport;
+}
+
+void UCharacterMovementComponent::OnTeleported()
+{
+	bJustTeleported = true;
+	if (!HasValidData())
+	{
+		return;
+	}
+
+	// Find floor at current location
+	UpdateFloorFromAdjustment();
+	SaveBaseLocation();
+
+	// Validate it. We don't want to pop down to walking mode from very high off the ground, but we'd like to keep walking if possible.
+	UPrimitiveComponent* OldBase = CharacterOwner->GetMovementBase();
+	UPrimitiveComponent* NewBase = NULL;
+	
+	if (OldBase && CurrentFloor.IsWalkableFloor() && CurrentFloor.FloorDist <= MAX_FLOOR_DIST && Velocity.Z <= 0.f)
+	{
+		// Close enough to land or just keep walking.
+		NewBase = CurrentFloor.HitResult.Component.Get();
+	}
+	else
+	{
+		CurrentFloor.Clear();
+	}
+
+	// If we were walking but no longer have a valid base or floor, start falling.
+	if (!CurrentFloor.IsWalkableFloor() || (OldBase && !NewBase))
+	{
+		if (DefaultLandMovementMode == MOVE_Walking)
+		{
+			SetMovementMode(MOVE_Falling);
+		}
+		else
+		{
+			SetDefaultMovementMode();
+		}
+	}
+}
+
+
+/** Internal. */
+FVector CharAngularVelocity(FRotator const& OldRot, FRotator const& NewRot, float DeltaTime)
+{
+	FVector RetAngVel(0.f);
+
+	if (OldRot != NewRot)
+	{
+		float const InvDeltaTime = 1.f / DeltaTime;
+		FQuat const DeltaQRot = (NewRot - OldRot).Quaternion();
+
+		FVector Axis; 
+		float Angle;
+		DeltaQRot.ToAxisAndAngle(Axis, Angle);
+
+		RetAngVel = Axis * Angle * InvDeltaTime;
+		check(!RetAngVel.ContainsNaN());
+	}
+
+	return RetAngVel;
+}
+
+FRotator UCharacterMovementComponent::GetDeltaRotation(float DeltaTime)
+{
+	return (RotationRate * DeltaTime);
+}
+
+FRotator UCharacterMovementComponent::ComputeOrientToMovementRotation(const FRotator& CurrentRotation, float DeltaTime, FRotator& DeltaRotation)
+{
+	if (Acceleration.SizeSquared() < KINDA_SMALL_NUMBER)
+	{
+		// AI path following request can orient us in that direction (it's effectively an acceleration)
+		if (bHasRequestedVelocity && RequestedVelocity.SizeSquared() > KINDA_SMALL_NUMBER)
+		{
+			return RequestedVelocity.GetSafeNormal().Rotation();
+		}
+
+		// Don't change rotation if there is no acceleration.
+		return CurrentRotation;
+	}
+
+	// Rotate toward direction of acceleration.
+	return Acceleration.GetSafeNormal().Rotation();
+}
+
+void UCharacterMovementComponent::PhysicsRotation(float DeltaTime)
+{
+	if (!HasValidData() || (!CharacterOwner->Controller && !bRunPhysicsWithNoController))
+	{
+		return;
+	}
+
+	if (!(bOrientRotationToMovement || bUseControllerDesiredRotation))
+	{
+		return;
+	}
+
+	const FRotator CurrentRotation = CharacterOwner->GetActorRotation();
+	FRotator DeltaRot = GetDeltaRotation(DeltaTime);
+	FRotator DesiredRotation = CurrentRotation;
+
+	if (bOrientRotationToMovement)
+	{
+		DesiredRotation = ComputeOrientToMovementRotation(CurrentRotation, DeltaTime, DeltaRot);
+	}
+	else if (CharacterOwner->Controller && bUseControllerDesiredRotation)
+	{
+		DesiredRotation = CharacterOwner->Controller->GetDesiredRotation();
+	}
+	else
+	{
+		return;
+	}
+
+	// Always remain vertical when walking or falling.
+	if( IsMovingOnGround() || IsFalling() )
+	{
+		DesiredRotation.Pitch = 0;
+		DesiredRotation.Roll = 0;
+	}
+
+	if( CurrentRotation.GetDenormalized().Equals(DesiredRotation.GetDenormalized(), 0.01f) )
+	{
+		return;
+	}
+
+	// Accumulate a desired new rotation.
+	FRotator NewRotation = CurrentRotation;	
+
+	//YAW
+	if( DesiredRotation.Yaw != CurrentRotation.Yaw )
+	{
+		NewRotation.Yaw = FMath::FixedTurn(CurrentRotation.Yaw, DesiredRotation.Yaw, DeltaRot.Yaw);
+	}
+
+	// PITCH
+	if( DesiredRotation.Pitch != CurrentRotation.Pitch )
+	{
+		NewRotation.Pitch = FMath::FixedTurn(CurrentRotation.Pitch, DesiredRotation.Pitch, DeltaRot.Pitch);
+	}
+
+	// ROLL
+	if( DesiredRotation.Roll != CurrentRotation.Roll )
+	{
+		NewRotation.Roll = FMath::FixedTurn(CurrentRotation.Roll, DesiredRotation.Roll, DeltaRot.Roll);
+	}
+
+	//UpdatedComponent->AngularVelocity = CharAngularVelocity( CurrentRotation, NewRotation, deltaTime );
+
+	// Set the new rotation.
+	if( !NewRotation.Equals(CurrentRotation.GetDenormalized(), 0.01f) )
+	{
+		MoveUpdatedComponent( FVector::ZeroVector, NewRotation, true );
+	}
+}
+
+
+void UCharacterMovementComponent::PhysicsVolumeChanged( APhysicsVolume* NewVolume )
+{
+	if (!HasValidData())
+	{
+		return;
+	}
+	if ( NewVolume && NewVolume->bWaterVolume )
+	{
+		// just entered water
+		if ( !CanEverSwim() )
+		{
+			// AI needs to stop any current moves
+			if (PathFollowingComp.IsValid())
+			{
+				PathFollowingComp->AbortMove(TEXT("water"));
+			}			
+		}
+		else if ( !IsSwimming() )
+		{
+			SetMovementMode(MOVE_Swimming);
+		}
+	}
+	else if ( IsSwimming() )
+	{
+		// just left the water - check if should jump out
+		SetMovementMode(MOVE_Falling);
+		FVector JumpDir(0.f);
+		FVector WallNormal(0.f);
+		if ( Acceleration.Z > 0.f && ShouldJumpOutOfWater(JumpDir)
+			&& ((JumpDir | Acceleration) > 0.f) && CheckWaterJump(JumpDir, WallNormal) ) 
+		{
+			JumpOutOfWater(WallNormal);
+			Velocity.Z = OutofWaterZ; //set here so physics uses this for remainder of tick
+		}
+	}
+}
+
+
+bool UCharacterMovementComponent::ShouldJumpOutOfWater(FVector& JumpDir)
+{
+	AController* OwnerController = CharacterOwner->GetController();
+	if (OwnerController)
+	{
+		const FRotator ControllerRot = OwnerController->GetControlRotation();
+		if ( (Velocity.Z > 0.0f) && (ControllerRot.Pitch > JumpOutOfWaterPitch) )
+		{
+			// if Pawn is going up and looking up, then make him jump
+			JumpDir = ControllerRot.Vector();
+			return true;
+		}
+	}
+	
+	return false;
+}
+
+
+void UCharacterMovementComponent::JumpOutOfWater(FVector WallNormal) {}
+
+
+bool UCharacterMovementComponent::CheckWaterJump(FVector CheckPoint, FVector& WallNormal)
+{
+	if (!HasValidData())
+	{
+		return false;
+	}
+	// check if there is a wall directly in front of the swimming pawn
+	CheckPoint.Z = 0.f;
+	FVector CheckNorm = CheckPoint.GetSafeNormal();
+	float PawnCapsuleRadius, PawnCapsuleHalfHeight;
+	CharacterOwner->GetCapsuleComponent()->GetScaledCapsuleSize(PawnCapsuleRadius, PawnCapsuleHalfHeight);
+	CheckPoint = CharacterOwner->GetActorLocation() + 1.2f * PawnCapsuleRadius * CheckNorm;
+	FVector Extent(PawnCapsuleRadius, PawnCapsuleRadius, PawnCapsuleHalfHeight);
+	FHitResult HitInfo(1.f);
+	static const FName CheckWaterJumpName(TEXT("CheckWaterJump"));
+	FCollisionQueryParams CapsuleParams(CheckWaterJumpName, false, CharacterOwner);
+	FCollisionResponseParams ResponseParam;
+	InitCollisionParams(CapsuleParams, ResponseParam);
+	FCollisionShape CapsuleShape = GetPawnCapsuleCollisionShape(SHRINK_None);
+	const ECollisionChannel CollisionChannel = UpdatedComponent->GetCollisionObjectType();
+	bool bHit = GetWorld()->SweepSingleByChannel( HitInfo, CharacterOwner->GetActorLocation(), CheckPoint, FQuat::Identity, CollisionChannel, CapsuleShape, CapsuleParams, ResponseParam);
+	
+	if ( bHit && !Cast<APawn>(HitInfo.GetActor()) )
+	{
+		// hit a wall - check if it is low enough
+		WallNormal = -1.f * HitInfo.ImpactNormal;
+		FVector Start = CharacterOwner->GetActorLocation();
+		Start.Z += MaxOutOfWaterStepHeight;
+		CheckPoint = Start + 3.2f * PawnCapsuleRadius * WallNormal;
+		FCollisionQueryParams LineParams(CheckWaterJumpName, true, CharacterOwner);
+		FCollisionResponseParams LineResponseParam;
+		InitCollisionParams(LineParams, LineResponseParam);
+		bHit = GetWorld()->LineTraceSingleByChannel( HitInfo, Start, CheckPoint, CollisionChannel, LineParams, LineResponseParam );
+		// if no high obstruction, or it's a valid floor, then pawn can jump out of water
+		return !bHit || IsWalkable(HitInfo);
+	}
+	return false;
+}
+
+void UCharacterMovementComponent::AddImpulse( FVector Impulse, bool bVelocityChange )
+{
+	if (HasValidData() && !Impulse.IsZero())
+	{
+		// handle scaling by mass
+		FVector FinalImpulse = Impulse;
+		if ( !bVelocityChange )
+		{
+			if (Mass > SMALL_NUMBER)
+			{
+				FinalImpulse = FinalImpulse / Mass;
+			}
+			else
+			{
+				UE_LOG(LogCharacterMovement, Warning, TEXT("Attempt to apply impulse to zero or negative Mass in CharacterMovement"));
+			}
+		}
+
+		PendingImpulseToApply += FinalImpulse;
+	}
+}
+
+void UCharacterMovementComponent::AddForce( FVector Force )
+{
+	if (HasValidData() && !Force.IsZero())
+	{
+		if (Mass > SMALL_NUMBER)
+		{
+			PendingForceToApply += Force / Mass;
+		}
+		else
+		{
+			UE_LOG(LogCharacterMovement, Warning, TEXT("Attempt to apply force to zero or negative Mass in CharacterMovement"));
+		}
+	}
+}
+
+void UCharacterMovementComponent::MoveSmooth(const FVector& InVelocity, const float DeltaSeconds, FStepDownResult* OutStepDownResult)
+{
+	if (!HasValidData())
+	{
+		return;
+	}
+
+	// Custom movement mode.
+	// Custom movement may need an update even if there is zero velocity.
+	if (MovementMode == MOVE_Custom)
+	{
+		FScopedMovementUpdate ScopedMovementUpdate(UpdatedComponent, bEnableScopedMovementUpdates ? EScopedUpdate::DeferredUpdates : EScopedUpdate::ImmediateUpdates);
+		PhysCustom(DeltaSeconds, 0);
+		return;
+	}
+
+	FVector Delta = InVelocity * DeltaSeconds;
+	if (Delta.IsZero())
+	{
+		return;
+	}
+
+	FScopedMovementUpdate ScopedMovementUpdate(UpdatedComponent, bEnableScopedMovementUpdates ? EScopedUpdate::DeferredUpdates : EScopedUpdate::ImmediateUpdates);
+
+	if (IsMovingOnGround())
+	{
+		MoveAlongFloor(InVelocity, DeltaSeconds, OutStepDownResult);
+	}
+	else
+	{
+		FHitResult Hit(1.f);
+		SafeMoveUpdatedComponent(Delta, CharacterOwner->GetActorRotation(), true, Hit);
+	
+		if (Hit.IsValidBlockingHit())
+		{
+			bool bSteppedUp = false;
+
+			if (IsFlying())
+			{
+				if (CanStepUp(Hit))
+				{
+					OutStepDownResult = NULL; // No need for a floor when not walking.
+					if (FMath::Abs(Hit.ImpactNormal.Z) < 0.2f)
+					{
+						const FVector GravDir = FVector(0.f,0.f,-1.f);
+						const FVector DesiredDir = Delta.GetSafeNormal();
+						const float UpDown = GravDir | DesiredDir;
+						if ((UpDown < 0.5f) && (UpDown > -0.2f))
+						{
+							bSteppedUp = StepUp(GravDir, Delta * (1.f - Hit.Time), Hit, OutStepDownResult);
+						}			
+					}
+				}
+			}
+				
+			// If StepUp failed, try sliding.
+			if (!bSteppedUp)
+			{
+				SlideAlongSurface(Delta, 1.f - Hit.Time, Hit.Normal, Hit, false);
+			}
+		}
+	}
+}
+
+
+bool UCharacterMovementComponent::IsWalkable(const FHitResult& Hit) const
+{
+	if (!Hit.IsValidBlockingHit())
+	{
+		// No hit, or starting in penetration
+		return false;
+	}
+
+	// Never walk up vertical surfaces.
+	if (Hit.ImpactNormal.Z < KINDA_SMALL_NUMBER)
+	{
+		return false;
+	}
+
+	float TestWalkableZ = WalkableFloorZ;
+
+	// See if this component overrides the walkable floor z.
+	const UPrimitiveComponent* HitComponent = Hit.Component.Get();
+	if (HitComponent)
+	{
+		const FWalkableSlopeOverride& SlopeOverride = HitComponent->GetWalkableSlopeOverride();
+		TestWalkableZ = SlopeOverride.ModifyWalkableFloorZ(TestWalkableZ);
+	}
+
+	// Can't walk on this surface if it is too steep.
+	if (Hit.ImpactNormal.Z < TestWalkableZ)
+	{
+		return false;
+	}
+
+	return true;
+}
+
+void UCharacterMovementComponent::SetWalkableFloorAngle(float InWalkableFloorAngle)
+{
+	WalkableFloorAngle = FMath::Clamp(InWalkableFloorAngle, 0.f, 90.0f);
+	WalkableFloorZ = FMath::Cos(FMath::DegreesToRadians(WalkableFloorAngle));
+}
+
+void UCharacterMovementComponent::SetWalkableFloorZ(float InWalkableFloorZ)
+{
+	WalkableFloorZ = FMath::Clamp(InWalkableFloorZ, 0.f, 1.f);
+	WalkableFloorAngle = FMath::RadiansToDegrees(FMath::Acos(WalkableFloorZ));
+}
+
+float UCharacterMovementComponent::K2_GetWalkableFloorAngle() const
+{
+	return GetWalkableFloorAngle();
+}
+
+float UCharacterMovementComponent::K2_GetWalkableFloorZ() const
+{
+	return GetWalkableFloorZ();
+}
+
+
+
+bool UCharacterMovementComponent::IsWithinEdgeTolerance(const FVector& CapsuleLocation, const FVector& TestImpactPoint, const float CapsuleRadius) const
+{
+	const float DistFromCenterSq = (TestImpactPoint - CapsuleLocation).SizeSquared2D();
+	const float ReducedRadiusSq = FMath::Square(FMath::Max(KINDA_SMALL_NUMBER, CapsuleRadius - SWEEP_EDGE_REJECT_DISTANCE));
+	return DistFromCenterSq < ReducedRadiusSq;
+}
+
+
+void UCharacterMovementComponent::ComputeFloorDist(const FVector& CapsuleLocation, float LineDistance, float SweepDistance, FFindFloorResult& OutFloorResult, float SweepRadius, const FHitResult* DownwardSweepResult) const
+{
+	OutFloorResult.Clear();
+
+	// No collision, no floor...
+	if (!UpdatedComponent->IsCollisionEnabled())
+	{
+		return;
+	}
+
+	float PawnRadius, PawnHalfHeight;
+	CharacterOwner->GetCapsuleComponent()->GetScaledCapsuleSize(PawnRadius, PawnHalfHeight);
+
+	bool bSkipSweep = false;
+	if (DownwardSweepResult != NULL && DownwardSweepResult->IsValidBlockingHit())
+	{
+		// Only if the supplied sweep was vertical and downward.
+		if ((DownwardSweepResult->TraceStart.Z > DownwardSweepResult->TraceEnd.Z) &&
+			(DownwardSweepResult->TraceStart - DownwardSweepResult->TraceEnd).SizeSquared2D() <= KINDA_SMALL_NUMBER)
+		{
+			// Reject hits that are barely on the cusp of the radius of the capsule
+			if (IsWithinEdgeTolerance(DownwardSweepResult->Location, DownwardSweepResult->ImpactPoint, PawnRadius))
+			{
+				// Don't try a redundant sweep, regardless of whether this sweep is usable.
+				bSkipSweep = true;
+
+				const bool bIsWalkable = IsWalkable(*DownwardSweepResult);
+				const float FloorDist = (CapsuleLocation.Z - DownwardSweepResult->Location.Z);
+				OutFloorResult.SetFromSweep(*DownwardSweepResult, FloorDist, bIsWalkable);
+				
+				if (bIsWalkable)
+				{
+					// Use the supplied downward sweep as the floor hit result.			
+					return;
+				}
+			}
+		}
+	}
+
+	// We require the sweep distance to be >= the line distance, otherwise the HitResult can't be interpreted as the sweep result.
+	if (SweepDistance < LineDistance)
+	{
+		check(SweepDistance >= LineDistance);
+		return;
+	}
+
+	bool bBlockingHit = false;
+	FCollisionQueryParams QueryParams(NAME_None, false, CharacterOwner);
+	FCollisionResponseParams ResponseParam;
+	InitCollisionParams(QueryParams, ResponseParam);
+	const ECollisionChannel CollisionChannel = UpdatedComponent->GetCollisionObjectType();
+
+	// Sweep test
+	if (!bSkipSweep && SweepDistance > 0.f && SweepRadius > 0.f)
+	{
+		// Use a shorter height to avoid sweeps giving weird results if we start on a surface.
+		// This also allows us to adjust out of penetrations.
+		const float ShrinkScale = 0.9f;
+		const float ShrinkScaleOverlap = 0.1f;
+		float ShrinkHeight = (PawnHalfHeight - PawnRadius) * (1.f - ShrinkScale);
+		float TraceDist = SweepDistance + ShrinkHeight;
+
+		static const FName ComputeFloorDistName(TEXT("ComputeFloorDistSweep"));
+		QueryParams.TraceTag = ComputeFloorDistName;
+		FCollisionShape CapsuleShape = FCollisionShape::MakeCapsule(SweepRadius, PawnHalfHeight - ShrinkHeight);
+
+		FHitResult Hit(1.f);
+		bBlockingHit = FloorSweepTest(Hit, CapsuleLocation, CapsuleLocation + FVector(0.f,0.f,-TraceDist), CollisionChannel, CapsuleShape, QueryParams, ResponseParam);
+
+		if (bBlockingHit)
+		{
+			// Reject hits adjacent to us, we only care about hits on the bottom portion of our capsule.
+			// Check 2D distance to impact point, reject if within a tolerance from radius.
+			if (Hit.bStartPenetrating || !IsWithinEdgeTolerance(CapsuleLocation, Hit.ImpactPoint, CapsuleShape.Capsule.Radius))
+			{
+				// Use a capsule with a slightly smaller radius and shorter height to avoid the adjacent object.
+				ShrinkHeight = (PawnHalfHeight - PawnRadius) * (1.f - ShrinkScaleOverlap);
+				TraceDist = SweepDistance + ShrinkHeight;
+				CapsuleShape.Capsule.Radius = FMath::Max(0.f, CapsuleShape.Capsule.Radius - SWEEP_EDGE_REJECT_DISTANCE - KINDA_SMALL_NUMBER);
+				CapsuleShape.Capsule.HalfHeight = FMath::Max(PawnHalfHeight - ShrinkHeight, CapsuleShape.Capsule.Radius);
+				Hit.Reset(1.f, false);
+
+				bBlockingHit = FloorSweepTest(Hit, CapsuleLocation, CapsuleLocation + FVector(0.f,0.f,-TraceDist), CollisionChannel, CapsuleShape, QueryParams, ResponseParam);
+			}
+
+			// Reduce hit distance by ShrinkHeight because we shrank the capsule for the trace.
+			// We allow negative distances here, because this allows us to pull out of penetrations.
+			const float MaxPenetrationAdjust = FMath::Max(MAX_FLOOR_DIST, PawnRadius);
+			const float SweepResult = FMath::Max(-MaxPenetrationAdjust, Hit.Time * TraceDist - ShrinkHeight);
+
+			OutFloorResult.SetFromSweep(Hit, SweepResult, false);
+			if (Hit.IsValidBlockingHit() && IsWalkable(Hit))
+			{		
+				if (SweepResult <= SweepDistance)
+				{
+					// Hit within test distance.
+					OutFloorResult.bWalkableFloor = true;
+					return;
+				}
+			}
+		}
+	}
+
+	// Since we require a longer sweep than line trace, we don't want to run the line trace if the sweep missed everything.
+	// We do however want to try a line trace if the sweep was stuck in penetration.
+	if (!OutFloorResult.bBlockingHit && !OutFloorResult.HitResult.bStartPenetrating)
+	{
+		OutFloorResult.FloorDist = SweepDistance;
+		return;
+	}
+
+	// Line trace
+	if (LineDistance > 0.f)
+	{
+		const float ShrinkHeight = PawnHalfHeight;
+		const FVector LineTraceStart = CapsuleLocation;	
+		const float TraceDist = LineDistance + ShrinkHeight;
+		const FVector Down = FVector(0.f, 0.f, -TraceDist);
+
+		static const FName FloorLineTraceName = FName(TEXT("ComputeFloorDistLineTrace"));
+		QueryParams.TraceTag = FloorLineTraceName;
+
+		FHitResult Hit(1.f);
+		bBlockingHit = GetWorld()->LineTraceSingleByChannel(Hit, LineTraceStart, LineTraceStart + Down, CollisionChannel, QueryParams, ResponseParam);
+
+		if (bBlockingHit)
+		{
+			if (Hit.Time > 0.f)
+			{
+				// Reduce hit distance by ShrinkHeight because we started the trace higher than the base.
+				// We allow negative distances here, because this allows us to pull out of penetrations.
+				const float MaxPenetrationAdjust = FMath::Max(MAX_FLOOR_DIST, PawnRadius);
+				const float LineResult = FMath::Max(-MaxPenetrationAdjust, Hit.Time * TraceDist - ShrinkHeight);
+
+				OutFloorResult.bBlockingHit = true;
+				if (LineResult <= LineDistance && IsWalkable(Hit))
+				{
+					OutFloorResult.SetFromLineTrace(Hit, OutFloorResult.FloorDist, LineResult, true);
+					return;
+				}
+			}
+		}
+	}
+	
+	// No hits were acceptable.
+	OutFloorResult.bWalkableFloor = false;
+	OutFloorResult.FloorDist = SweepDistance;
+}
+
+
+void UCharacterMovementComponent::FindFloor(const FVector& CapsuleLocation, FFindFloorResult& OutFloorResult, bool bZeroDelta, const FHitResult* DownwardSweepResult) const
+{
+	// No collision, no floor...
+	if (!UpdatedComponent->IsCollisionEnabled())
+	{
+		OutFloorResult.Clear();
+		return;
+	}
+
+	// Increase height check slightly if walking, to prevent floor height adjustment from later invalidating the floor result.
+	const float HeightCheckAdjust = (IsMovingOnGround() ? MAX_FLOOR_DIST + KINDA_SMALL_NUMBER : -MAX_FLOOR_DIST);
+
+	float FloorSweepTraceDist = FMath::Max(MAX_FLOOR_DIST, MaxStepHeight + HeightCheckAdjust);
+	float FloorLineTraceDist = FloorSweepTraceDist;
+	bool bNeedToValidateFloor = true;
+	
+	// Sweep floor
+	if (FloorLineTraceDist > 0.f || FloorSweepTraceDist > 0.f)
+	{
+		UCharacterMovementComponent* MutableThis = const_cast<UCharacterMovementComponent*>(this);
+
+		if ( bAlwaysCheckFloor || !bZeroDelta || bForceNextFloorCheck || bJustTeleported )
+		{
+			MutableThis->bForceNextFloorCheck = false;
+			ComputeFloorDist(CapsuleLocation, FloorLineTraceDist, FloorSweepTraceDist, OutFloorResult, CharacterOwner->GetCapsuleComponent()->GetScaledCapsuleRadius(), DownwardSweepResult);
+		}
+		else
+		{
+			// Force floor check if base has collision disabled or if it does not block us.
+			UPrimitiveComponent* MovementBase = CharacterOwner->GetMovementBase();
+			const AActor* BaseActor = MovementBase ? MovementBase->GetOwner() : NULL;
+			const ECollisionChannel CollisionChannel = UpdatedComponent->GetCollisionObjectType();
+
+			if (MovementBase != NULL)
+			{
+				MutableThis->bForceNextFloorCheck = !MovementBase->IsCollisionEnabled()
+				|| MovementBase->GetCollisionResponseToChannel(CollisionChannel) != ECR_Block
+				|| MovementBaseUtility::IsDynamicBase(MovementBase);
+			}
+
+			const bool IsActorBasePendingKill = BaseActor && BaseActor->IsPendingKill();
+
+			if ( !bForceNextFloorCheck && !IsActorBasePendingKill && MovementBase )
+			{
+				//UE_LOG(LogCharacterMovement, Log, TEXT("%s SKIP check for floor"), *CharacterOwner->GetName());
+				OutFloorResult = CurrentFloor;
+				bNeedToValidateFloor = false;
+			}
+			else
+			{
+				MutableThis->bForceNextFloorCheck = false;
+				ComputeFloorDist(CapsuleLocation, FloorLineTraceDist, FloorSweepTraceDist, OutFloorResult, CharacterOwner->GetCapsuleComponent()->GetScaledCapsuleRadius(), DownwardSweepResult);
+			}
+		}
+	}
+
+	// OutFloorResult.HitResult is now the result of the vertical floor check.
+	// See if we should try to "perch" at this location.
+	if (bNeedToValidateFloor && OutFloorResult.bBlockingHit && !OutFloorResult.bLineTrace)
+	{
+		const bool bCheckRadius = true;
+		if (ShouldComputePerchResult(OutFloorResult.HitResult, bCheckRadius))
+		{
+			float MaxPerchFloorDist = FMath::Max(MAX_FLOOR_DIST, MaxStepHeight + HeightCheckAdjust);
+			if (IsMovingOnGround())
+			{
+				MaxPerchFloorDist += FMath::Max(0.f, PerchAdditionalHeight);
+			}
+
+			FFindFloorResult PerchFloorResult;
+			if (ComputePerchResult(GetValidPerchRadius(), OutFloorResult.HitResult, MaxPerchFloorDist, PerchFloorResult))
+			{
+				// Don't allow the floor distance adjustment to push us up too high, or we will move beyond the perch distance and fall next time.
+				const float AvgFloorDist = (MIN_FLOOR_DIST + MAX_FLOOR_DIST) * 0.5f;
+				const float MoveUpDist = (AvgFloorDist - OutFloorResult.FloorDist);
+				if (MoveUpDist + PerchFloorResult.FloorDist >= MaxPerchFloorDist)
+				{
+					OutFloorResult.FloorDist = AvgFloorDist;
+				}
+
+				// If the regular capsule is on an unwalkable surface but the perched one would allow us to stand, override the normal to be one that is walkable.
+				if (!OutFloorResult.bWalkableFloor)
+				{
+					OutFloorResult.SetFromLineTrace(PerchFloorResult.HitResult, OutFloorResult.FloorDist, FMath::Min(PerchFloorResult.FloorDist, PerchFloorResult.LineDist), true);
+				}
+			}
+			else
+			{
+				// We had no floor (or an invalid one because it was unwalkable), and couldn't perch here, so invalidate floor (which will cause us to start falling).
+				OutFloorResult.bWalkableFloor = false;
+			}
+		}
+	}
+}
+
+
+bool UCharacterMovementComponent::FloorSweepTest(
+	FHitResult& OutHit,
+	const FVector& Start,
+	const FVector& End,
+	ECollisionChannel TraceChannel,
+	const struct FCollisionShape& CollisionShape,
+	const struct FCollisionQueryParams& Params,
+	const struct FCollisionResponseParams& ResponseParam
+	) const
+{
+	bool bBlockingHit = false;
+
+	if (!bUseFlatBaseForFloorChecks)
+	{
+		bBlockingHit = GetWorld()->SweepSingleByChannel(OutHit, Start, End, FQuat::Identity, TraceChannel, CollisionShape, Params, ResponseParam);
+	}
+	else
+	{
+		// Test with a box that is enclosed by the capsule.
+		const float CapsuleRadius = CollisionShape.GetCapsuleRadius();
+		const float CapsuleHeight = CollisionShape.GetCapsuleHalfHeight();
+		const FCollisionShape BoxShape = FCollisionShape::MakeBox(FVector(CapsuleRadius * 0.707f, CapsuleRadius * 0.707f, CapsuleHeight));
+
+		// First test with the box rotated so the corners are along the major axes (ie rotated 45 degrees).
+		bBlockingHit = GetWorld()->SweepSingleByChannel(OutHit, Start, End, FQuat(FVector(0.f, 0.f, -1.f), PI * 0.25f), TraceChannel, BoxShape, Params, ResponseParam);
+
+		if (!bBlockingHit)
+		{
+			// Test again with the same box, not rotated.
+			OutHit.Reset(1.f, false);			
+			bBlockingHit = GetWorld()->SweepSingleByChannel(OutHit, Start, End, FQuat::Identity, TraceChannel, BoxShape, Params, ResponseParam);
+		}
+	}
+
+	return bBlockingHit;
+}
+
+
+bool UCharacterMovementComponent::IsValidLandingSpot(const FVector& CapsuleLocation, const FHitResult& Hit) const
+{
+	if (!Hit.bBlockingHit)
+	{
+		return false;
+	}
+
+	// Skip some checks if penetrating. Penetration will be handled by the FindFloor call (using a smaller capsule)
+	if (!Hit.bStartPenetrating)
+	{
+		// Reject unwalkable floor normals.
+		if (!IsWalkable(Hit))
+		{
+			return false;
+		}
+
+		float PawnRadius, PawnHalfHeight;
+		CharacterOwner->GetCapsuleComponent()->GetScaledCapsuleSize(PawnRadius, PawnHalfHeight);
+
+		// Reject hits that are above our lower hemisphere (can happen when sliding down a vertical surface).
+		const float LowerHemisphereZ = Hit.Location.Z - PawnHalfHeight + PawnRadius;
+		if (Hit.ImpactPoint.Z >= LowerHemisphereZ)
+		{
+			return false;
+		}
+
+		// Reject hits that are barely on the cusp of the radius of the capsule
+		if (!IsWithinEdgeTolerance(Hit.Location, Hit.ImpactPoint, PawnRadius))
+		{
+			return false;
+		}
+	}
+	else
+	{
+		// Penetrating
+		if (Hit.Normal.Z < KINDA_SMALL_NUMBER)
+		{
+			// Normal is nearly horizontal or downward, that's a penetration adjustment next to a vertical or overhanging wall. Don't pop to the floor.
+			return false;
+		}
+	}
+
+	FFindFloorResult FloorResult;
+	FindFloor(CapsuleLocation, FloorResult, false, &Hit);
+
+	if (!FloorResult.IsWalkableFloor())
+	{
+		return false;
+	}
+
+	return true;
+}
+
+
+bool UCharacterMovementComponent::ShouldCheckForValidLandingSpot(float DeltaTime, const FVector& Delta, const FHitResult& Hit) const
+{
+	// See if we hit an edge of a surface on the lower portion of the capsule.
+	// In this case the normal will not equal the impact normal, and a downward sweep may find a walkable surface on top of the edge.
+	if (Hit.Normal.Z > KINDA_SMALL_NUMBER && !Hit.Normal.Equals(Hit.ImpactNormal))
+	{
+		const FVector PawnLocation = UpdatedComponent->GetComponentLocation();
+		if (IsWithinEdgeTolerance(PawnLocation, Hit.ImpactPoint, CharacterOwner->GetCapsuleComponent()->GetScaledCapsuleRadius()))
+		{						
+			return true;
+		}
+	}
+
+	return false;
+}
+
+
+float UCharacterMovementComponent::GetPerchRadiusThreshold() const
+{
+	// Don't allow negative values.
+	return FMath::Max(0.f, PerchRadiusThreshold);
+}
+
+
+float UCharacterMovementComponent::GetValidPerchRadius() const
+{
+	if (CharacterOwner)
+	{
+		const float PawnRadius = CharacterOwner->GetCapsuleComponent()->GetScaledCapsuleRadius();
+		return FMath::Clamp(PawnRadius - GetPerchRadiusThreshold(), 0.1f, PawnRadius);
+	}
+	return 0.f;
+}
+
+
+bool UCharacterMovementComponent::ShouldComputePerchResult(const FHitResult& InHit, bool bCheckRadius) const
+{
+	if (!InHit.IsValidBlockingHit())
+	{
+		return false;
+	}
+
+	// Don't try to perch if the edge radius is very small.
+	if (GetPerchRadiusThreshold() <= SWEEP_EDGE_REJECT_DISTANCE)
+	{
+		return false;
+	}
+
+	if (bCheckRadius)
+	{
+		const float DistFromCenterSq = (InHit.ImpactPoint - InHit.Location).SizeSquared2D();
+		const float StandOnEdgeRadius = GetValidPerchRadius();
+		if (DistFromCenterSq <= FMath::Square(StandOnEdgeRadius))
+		{
+			// Already within perch radius.
+			return false;
+		}
+	}
+	
+	return true;
+}
+
+
+bool UCharacterMovementComponent::ComputePerchResult(const float TestRadius, const FHitResult& InHit, const float InMaxFloorDist, FFindFloorResult& OutPerchFloorResult) const
+{
+	if (InMaxFloorDist <= 0.f)
+	{
+		return 0.f;
+	}
+
+	// Sweep further than actual requested distance, because a reduced capsule radius means we could miss some hits that the normal radius would contact.
+	float PawnRadius, PawnHalfHeight;
+	CharacterOwner->GetCapsuleComponent()->GetScaledCapsuleSize(PawnRadius, PawnHalfHeight);
+
+	const float InHitAboveBase = FMath::Max(0.f, InHit.ImpactPoint.Z - (InHit.Location.Z - PawnHalfHeight));
+	const float PerchLineDist = FMath::Max(0.f, InMaxFloorDist - InHitAboveBase);
+	const float PerchSweepDist = FMath::Max(0.f, InMaxFloorDist);
+
+	const float ActualSweepDist = PerchSweepDist + PawnRadius;
+	ComputeFloorDist(InHit.Location, PerchLineDist, ActualSweepDist, OutPerchFloorResult, TestRadius);
+
+	if (!OutPerchFloorResult.IsWalkableFloor())
+	{
+		return false;
+	}
+	else if (InHitAboveBase + OutPerchFloorResult.FloorDist > InMaxFloorDist)
+	{
+		// Hit something past max distance
+		OutPerchFloorResult.bWalkableFloor = false;
+		return false;
+	}
+
+	return true;
+}
+
+
+bool UCharacterMovementComponent::CanStepUp(const FHitResult& Hit) const
+{
+	if (!Hit.IsValidBlockingHit() || !HasValidData() || MovementMode == MOVE_Falling)
+	{
+		return false;
+	}
+
+	// No component for "fake" hits when we are on a known good base.
+	const UPrimitiveComponent* HitComponent = Hit.Component.Get();
+	if (!HitComponent)
+	{
+		return true;
+	}
+
+	if (!HitComponent->CanCharacterStepUp(CharacterOwner))
+	{
+		return false;
+	}
+
+	// No actor for "fake" hits when we are on a known good base.
+	const AActor* HitActor = Hit.GetActor();
+	if (!HitActor)
+	{
+		 return true;
+	}
+
+	if (!HitActor->CanBeBaseForCharacter(CharacterOwner))
+	{
+		return false;
+	}
+
+	return true;
+}
+
+
+bool UCharacterMovementComponent::StepUp(const FVector& InGravDir, const FVector& Delta, const FHitResult &InHit, FStepDownResult* OutStepDownResult)
+{
+	SCOPE_CYCLE_COUNTER(STAT_CharStepUp);
+
+
+	if (!CanStepUp(InHit))
+	{
+		return false;
+	}
+
+	if (MaxStepHeight <= 0.f)
+	{
+		return false;
+	}
+
+	const FVector OldLocation = UpdatedComponent->GetComponentLocation();
+	float PawnRadius, PawnHalfHeight;
+	CharacterOwner->GetCapsuleComponent()->GetScaledCapsuleSize(PawnRadius, PawnHalfHeight);
+
+	// Don't bother stepping up if top of capsule is hitting something.
+	const float InitialImpactZ = InHit.ImpactPoint.Z;
+	if (InitialImpactZ > OldLocation.Z + (PawnHalfHeight - PawnRadius))
+	{
+		return false;
+	}
+
+	// Don't step up if the impact is below us
+	if (InitialImpactZ <= OldLocation.Z - PawnHalfHeight)
+	{
+		return false;
+	}
+
+	const FVector GravDir = InGravDir.GetSafeNormal();
+	if (GravDir.IsZero())
+	{
+		return false;
+	}
+
+	float StepTravelUpHeight = MaxStepHeight;
+	float StepTravelDownHeight = StepTravelUpHeight;
+	const float StepSideZ = -1.f * (InHit.ImpactNormal | GravDir);
+	float PawnInitialFloorBaseZ = OldLocation.Z - PawnHalfHeight;
+	float PawnFloorPointZ = PawnInitialFloorBaseZ;
+
+	if (IsMovingOnGround() && CurrentFloor.IsWalkableFloor())
+	{
+		// Since we float a variable amount off the floor, we need to enforce max step height off the actual point of impact with the floor.
+		const float FloorDist = FMath::Max(0.f, CurrentFloor.FloorDist);
+		PawnInitialFloorBaseZ -= FloorDist;
+		StepTravelUpHeight = FMath::Max(StepTravelUpHeight - FloorDist, 0.f);
+		StepTravelDownHeight = (MaxStepHeight + MAX_FLOOR_DIST*2.f);
+
+		const bool bHitVerticalFace = !IsWithinEdgeTolerance(InHit.Location, InHit.ImpactPoint, PawnRadius);
+		if (!CurrentFloor.bLineTrace && !bHitVerticalFace)
+		{
+			PawnFloorPointZ = CurrentFloor.HitResult.ImpactPoint.Z;
+		}
+		else
+		{
+			// Base floor point is the base of the capsule moved down by how far we are hovering over the surface we are hitting.
+			PawnFloorPointZ -= CurrentFloor.FloorDist;
+		}
+	}
+
+	// Scope our movement updates, and do not apply them until all intermediate moves are completed.
+	FScopedMovementUpdate ScopedStepUpMovement(UpdatedComponent, EScopedUpdate::DeferredUpdates);
+
+	// step up - treat as vertical wall
+	FHitResult SweepUpHit(1.f);
+	const FRotator PawnRotation = CharacterOwner->GetActorRotation();
+	SafeMoveUpdatedComponent(-GravDir * StepTravelUpHeight, PawnRotation, true, SweepUpHit);
+
+	// step fwd
+	FHitResult Hit(1.f);
+	SafeMoveUpdatedComponent( Delta, PawnRotation, true, Hit);
+
+	// If we hit something above us and also something ahead of us, we should notify about the upward hit as well.
+	// The forward hit will be handled later (in the bSteppedOver case below).
+	// In the case of hitting something above but not forward, we are not blocked from moving so we don't need the notification.
+	if (SweepUpHit.bBlockingHit && Hit.bBlockingHit)
+	{
+		HandleImpact(SweepUpHit);
+	}
+
+	// Check result of forward movement
+	if (Hit.bBlockingHit)
+	{
+		if (Hit.bStartPenetrating)
+		{
+			// Undo movement
+			ScopedStepUpMovement.RevertMove();
+			return false;
+		}
+
+		// pawn ran into a wall
+		HandleImpact(Hit);
+		if ( IsFalling() )
+		{
+			return true;
+		}
+
+		// adjust and try again
+		const float ForwardHitTime = Hit.Time;
+		const float ForwardSlideAmount = SlideAlongSurface(Delta, 1.f - Hit.Time, Hit.Normal, Hit, true);
+		
+		if (IsFalling())
+		{
+			ScopedStepUpMovement.RevertMove();
+			return false;
+		}
+
+		// If both the forward hit and the deflection got us nowhere, there is no point in this step up.
+		if (ForwardHitTime == 0.f && ForwardSlideAmount == 0.f)
+		{
+			ScopedStepUpMovement.RevertMove();
+			return false;
+		}
+	}
+	
+	// Step down
+	SafeMoveUpdatedComponent(GravDir * StepTravelDownHeight, CharacterOwner->GetActorRotation(), true, Hit);
+
+	// If step down was initially penetrating abort the step up
+	if (Hit.bStartPenetrating)
+	{
+		ScopedStepUpMovement.RevertMove();
+		return false;
+	}
+
+	FStepDownResult StepDownResult;
+	if (Hit.IsValidBlockingHit())
+	{	
+		// See if this step sequence would have allowed us to travel higher than our max step height allows.
+		const float DeltaZ = Hit.ImpactPoint.Z - PawnFloorPointZ;
+		if (DeltaZ > MaxStepHeight)
+		{
+			//UE_LOG(LogCharacterMovement, VeryVerbose, TEXT("- Reject StepUp (too high Height %.3f) up from floor base %f to %f"), DeltaZ, PawnInitialFloorBaseZ, NewLocation.Z);
+			ScopedStepUpMovement.RevertMove();
+			return false;
+		}
+
+		// Reject unwalkable surface normals here.
+		if (!IsWalkable(Hit))
+		{
+			// Reject if normal opposes movement direction
+			const bool bNormalTowardsMe = (Delta | Hit.ImpactNormal) < 0.f;
+			if (bNormalTowardsMe)
+			{
+				//UE_LOG(LogCharacterMovement, VeryVerbose, TEXT("- Reject StepUp (unwalkable normal %s opposed to movement)"), *Hit.ImpactNormal.ToString());
+				ScopedStepUpMovement.RevertMove();
+				return false;
+			}
+
+			// Also reject if we would end up being higher than our starting location by stepping down.
+			// It's fine to step down onto an unwalkable normal below us, we will just slide off. Rejecting those moves would prevent us from being able to walk off the edge.
+			if (Hit.Location.Z > OldLocation.Z)
+			{
+				//UE_LOG(LogCharacterMovement, VeryVerbose, TEXT("- Reject StepUp (unwalkable normal %s above old position)"), *Hit.ImpactNormal.ToString());
+				ScopedStepUpMovement.RevertMove();
+				return false;
+			}
+		}
+
+		// Reject moves where the downward sweep hit something very close to the edge of the capsule. This maintains consistency with FindFloor as well.
+		if (!IsWithinEdgeTolerance(Hit.Location, Hit.ImpactPoint, PawnRadius))
+		{
+			//UE_LOG(LogCharacterMovement, VeryVerbose, TEXT("- Reject StepUp (outside edge tolerance)"));
+			ScopedStepUpMovement.RevertMove();
+			return false;
+		}
+
+		// Don't step up onto invalid surfaces if traveling higher.
+		if (DeltaZ > 0.f && !CanStepUp(Hit))
+		{
+			//UE_LOG(LogCharacterMovement, VeryVerbose, TEXT("- Reject StepUp (up onto surface with !CanStepUp())"));
+			ScopedStepUpMovement.RevertMove();
+			return false;
+		}
+
+		// See if we can validate the floor as a result of this step down. In almost all cases this should succeed, and we can avoid computing the floor outside this method.
+		if (OutStepDownResult != NULL)
+		{
+			FindFloor(UpdatedComponent->GetComponentLocation(), StepDownResult.FloorResult, false, &Hit);
+
+			// Reject unwalkable normals if we end up higher than our initial height.
+			// It's fine to walk down onto an unwalkable surface, don't reject those moves.
+			if (Hit.Location.Z > OldLocation.Z)
+			{
+				// We should reject the floor result if we are trying to step up an actual step where we are not able to perch (this is rare).
+				// In those cases we should instead abort the step up and try to slide along the stair.
+				if (!StepDownResult.FloorResult.bBlockingHit && StepSideZ < MAX_STEP_SIDE_Z)
+				{
+					ScopedStepUpMovement.RevertMove();
+					return false;
+				}
+			}
+
+			StepDownResult.bComputedFloor = true;
+		}
+	}
+	
+	// Copy step down result.
+	if (OutStepDownResult != NULL)
+	{
+		*OutStepDownResult = StepDownResult;
+	}
+
+	// Don't recalculate velocity based on this height adjustment, if considering vertical adjustments.
+	bJustTeleported |= !bMaintainHorizontalGroundVelocity;
+
+	return true;
+}
+
+void UCharacterMovementComponent::HandleImpact(FHitResult const& Impact, float TimeSlice, const FVector& MoveDelta)
+{
+	if (CharacterOwner)
+	{
+		CharacterOwner->MoveBlockedBy(Impact);
+	}
+
+	if (PathFollowingComp.IsValid())
+	{	// Also notify path following!
+		PathFollowingComp->OnMoveBlockedBy(Impact);
+	}
+
+	APawn* OtherPawn = Cast<APawn>(Impact.GetActor());
+	if (OtherPawn)
+	{
+		NotifyBumpedPawn(OtherPawn);
+	}
+
+	if (bEnablePhysicsInteraction)
+	{
+		const FVector ForceAccel = Acceleration + (IsFalling() ? FVector(0.f, 0.f, GetGravityZ()) : FVector::ZeroVector);
+		ApplyImpactPhysicsForces(Impact, ForceAccel, Velocity);
+	}
+}
+
+void UCharacterMovementComponent::ApplyImpactPhysicsForces(const FHitResult& Impact, const FVector& ImpactAcceleration, const FVector& ImpactVelocity)
+{
+	if (bEnablePhysicsInteraction && Impact.bBlockingHit)
+	{
+		UPrimitiveComponent* ImpactComponent = Impact.GetComponent();
+		if (ImpactComponent != NULL && ImpactComponent->IsAnySimulatingPhysics())
+		{
+			FVector ForcePoint = Impact.ImpactPoint;
+
+			FBodyInstance* BI = ImpactComponent->GetBodyInstance(Impact.BoneName);
+
+			float BodyMass = 1.0f;
+
+			if (BI != NULL)
+			{
+				BodyMass = FMath::Max(BI->GetBodyMass(), 1.0f);
+
+				FBox Bounds = BI->GetBodyBounds();
+
+				FVector Center, Extents;
+				Bounds.GetCenterAndExtents(Center, Extents);
+
+				if (!Extents.IsNearlyZero())
+				{
+					ForcePoint.Z = Center.Z + Extents.Z * PushForcePointZOffsetFactor;
+				}
+			}
+
+			FVector Force = Impact.ImpactNormal * -1.0f;
+
+			float PushForceModificator = 1.0f;
+
+			const FVector ComponentVelocity = ImpactComponent->GetPhysicsLinearVelocity();
+			const FVector VirtualVelocity = ImpactAcceleration.IsZero() ? ImpactVelocity : ImpactAcceleration.GetSafeNormal() * GetMaxSpeed();
+
+			float Dot = 0.0f;
+
+			if (bScalePushForceToVelocity && !ComponentVelocity.IsNearlyZero())
+			{			
+				Dot = ComponentVelocity | VirtualVelocity;
+
+				if (Dot > 0.0f && Dot < 1.0f)
+				{
+					PushForceModificator *= Dot;
+				}
+			}
+
+			if (bPushForceScaledToMass)
+			{
+				PushForceModificator *= BodyMass;
+			}
+
+			Force *= PushForceModificator;
+
+			if (ComponentVelocity.IsNearlyZero())
+			{
+				Force *= InitialPushForceFactor;
+				ImpactComponent->AddImpulseAtLocation(Force, ForcePoint, Impact.BoneName);
+			}
+			else
+			{
+				Force *= PushForceFactor;
+				ImpactComponent->AddForceAtLocation(Force, ForcePoint, Impact.BoneName);
+			}
+		}
+	}
+}
+
+FString UCharacterMovementComponent::GetMovementName()
+{
+	if( CharacterOwner )
+	{
+		if ( CharacterOwner->GetRootComponent() && CharacterOwner->GetRootComponent()->IsSimulatingPhysics() )
+		{
+			return TEXT("Rigid Body");
+		}
+		else if ( CharacterOwner->IsMatineeControlled() )
+		{
+			return TEXT("Matinee");
+		}
+	}
+
+	// Using character movement
+	switch( MovementMode )
+	{
+		case MOVE_None:				return TEXT("NULL"); break;
+		case MOVE_Walking:			return TEXT("Walking"); break;
+		case MOVE_NavWalking:		return TEXT("NavWalking"); break;
+		case MOVE_Falling:			return TEXT("Falling"); break;
+		case MOVE_Swimming:			return TEXT("Swimming"); break;
+		case MOVE_Flying:			return TEXT("Flying"); break;
+		case MOVE_Custom:			return TEXT("Custom"); break;
+		default:					break;
+	}
+	return TEXT("Unknown");
+}
+
+void UCharacterMovementComponent::DisplayDebug(UCanvas* Canvas, const FDebugDisplayInfo& DebugDisplay, float& YL, float& YPos)
+{
+	if ( CharacterOwner == NULL )
+	{
+		return;
+	}
+
+	Canvas->SetDrawColor(255,255,255);
+	UFont* RenderFont = GEngine->GetSmallFont();
+	FString T = FString::Printf(TEXT("CHARACTER MOVEMENT Floor %s Crouched %i"), *CurrentFloor.HitResult.ImpactNormal.ToString(), IsCrouching());
+	YL = Canvas->DrawText(RenderFont, T, 4.0f, YPos );
+	YPos += YL;
+
+	T = FString::Printf(TEXT("Updated Component: %s"), *UpdatedComponent->GetName());
+	YL = Canvas->DrawText(RenderFont, T, 4.0f, YPos );
+	YPos += YL;
+
+	T = FString::Printf(TEXT("bForceMaxAccel: %i"), bForceMaxAccel);
+	YL = Canvas->DrawText(RenderFont, T, 4.0f, YPos );
+	YPos += YL;
+
+	APhysicsVolume * PhysicsVolume = GetPhysicsVolume();
+
+	const UPrimitiveComponent* BaseComponent = CharacterOwner->GetMovementBase();
+	const AActor* BaseActor = BaseComponent ? BaseComponent->GetOwner() : NULL;
+
+	T = FString::Printf(TEXT("%s In physicsvolume %s on base %s component %s gravity %f"), *GetMovementName(), (PhysicsVolume ? *PhysicsVolume->GetName() : TEXT("None")),
+		(BaseActor ? *BaseActor->GetName() : TEXT("None")), (BaseComponent ? *BaseComponent->GetName() : TEXT("None")), GetGravityZ());
+	YL = Canvas->DrawText(RenderFont, T, 4.0f, YPos );
+	YPos += YL;
+}
+
+
+void UCharacterMovementComponent::ForceReplicationUpdate()
+{
+	if (HasPredictionData_Server())
+	{
+		GetPredictionData_Server_Character()->LastUpdateTime = GetWorld()->TimeSeconds - 10.f;
+	}
+}
+
+
+FVector UCharacterMovementComponent::ConstrainInputAcceleration(const FVector& InputAcceleration) const
+{
+	FVector NewAccel = InputAcceleration;
+
+	// walking or falling pawns ignore up/down sliding
+	if (IsMovingOnGround() || IsFalling())
+	{
+		NewAccel.Z = 0.f;
+	}
+
+	return NewAccel;
+}
+
+
+FVector UCharacterMovementComponent::ScaleInputAcceleration(const FVector& InputAcceleration) const
+{
+	return GetMaxAcceleration() * GetClampedToMaxSizePrecise(InputAcceleration, 1.0f);
+}
+
+
+FVector UCharacterMovementComponent::RoundAcceleration(FVector InAccel) const
+{
+	// Match FVector_NetQuantize10 (1 decimal place of precision).
+	InAccel.X = FMath::RoundToFloat(InAccel.X * 10.f) / 10.f;
+	InAccel.Y = FMath::RoundToFloat(InAccel.Y * 10.f) / 10.f;
+	InAccel.Z = FMath::RoundToFloat(InAccel.Z * 10.f) / 10.f;
+	return InAccel;
+}
+
+
+float UCharacterMovementComponent::ComputeAnalogInputModifier() const
+{
+	const float MaxAccel = GetMaxAcceleration();
+	if (Acceleration.SizeSquared() > 0.f && MaxAccel > SMALL_NUMBER)
+	{
+		return FMath::Clamp(Acceleration.Size() / MaxAccel, 0.f, 1.f);
+	}
+
+	return 0.f;
+}
+
+float UCharacterMovementComponent::GetAnalogInputModifier() const
+{
+	return AnalogInputModifier;
+}
+
+
+float UCharacterMovementComponent::GetSimulationTimeStep(float RemainingTime, int32 Iterations) const
+{
+	if (RemainingTime > MaxSimulationTimeStep)
+	{
+		if (Iterations < MaxSimulationIterations)
+		{
+			// Subdivide moves to be no longer than MaxSimulationTimeStep seconds
+			RemainingTime = FMath::Min(MaxSimulationTimeStep, RemainingTime * 0.5f);
+		}
+		else
+		{
+			// If this is the last iteration, just use all the remaining time. This is usually better than cutting things short, as the simulation won't move far enough otherwise.
+			// Print a throttled warning.
+#if !(UE_BUILD_SHIPPING || UE_BUILD_TEST)
+			static uint32 s_WarningCount = 0;
+			if ((s_WarningCount++ < 100) || (GFrameCounter & 15) == 0)
+			{
+				UE_LOG(LogCharacterMovement, Warning, TEXT("GetSimulationTimeStep() - Max iterations %d hit while remaining time %.6f > MaxSimulationTimeStep (%.3f) for '%s'"), MaxSimulationIterations, RemainingTime, MaxSimulationTimeStep, *GetNameSafe(CharacterOwner));
+			}
+#endif
+		}
+	}
+
+	// no less than MIN_TICK_TIME (to avoid potential divide-by-zero during simulation).
+	return FMath::Max(MIN_TICK_TIME, RemainingTime);
+}
+
+
+void UCharacterMovementComponent::SmoothCorrection(const FVector& OldLocation)
+{
+	if (!HasValidData() || GetNetMode() != NM_Client)
+	{
+		return;
+	}
+
+	FNetworkPredictionData_Client_Character* ClientData = GetPredictionData_Client_Character();
+	if (ClientData && ClientData->bSmoothNetUpdates)
+	{
+		float DistSq = (OldLocation - CharacterOwner->GetActorLocation()).SizeSquared();
+		if (DistSq > FMath::Square(ClientData->MaxSmoothNetUpdateDist))
+		{
+			ClientData->MeshTranslationOffset = (DistSq > FMath::Square(ClientData->NoSmoothNetUpdateDist)) 
+				? FVector::ZeroVector 
+				: ClientData->MeshTranslationOffset + ClientData->MaxSmoothNetUpdateDist * (OldLocation - CharacterOwner->GetActorLocation()).GetSafeNormal();
+		}
+		else
+		{
+			ClientData->MeshTranslationOffset = ClientData->MeshTranslationOffset + OldLocation - CharacterOwner->GetActorLocation();	
+		}
+	}
+}
+
+
+void UCharacterMovementComponent::SmoothClientPosition(float DeltaSeconds)
+{
+	if (!HasValidData() || GetNetMode() != NM_Client)
+	{
+		return;
+	}
+
+	FNetworkPredictionData_Client_Character* ClientData = GetPredictionData_Client_Character();
+	if (ClientData && ClientData->bSmoothNetUpdates && CharacterOwner->GetMesh() && !CharacterOwner->GetMesh()->IsSimulatingPhysics())
+	{
+		// smooth interpolation of mesh translation to avoid popping of other client pawns unless under a low tick rate
+		if (DeltaSeconds < ClientData->SmoothNetUpdateTime)
+		{
+			ClientData->MeshTranslationOffset = (ClientData->MeshTranslationOffset * (1.f - DeltaSeconds / ClientData->SmoothNetUpdateTime));
+		}
+		else
+		{
+			ClientData->MeshTranslationOffset = FVector::ZeroVector;
+		}
+
+		if (IsMovingOnGround())
+		{
+			// don't smooth Z position if walking on ground
+			ClientData->MeshTranslationOffset.Z = 0;
+		}
+
+		const FVector NewRelTranslation = CharacterOwner->ActorToWorld().InverseTransformVectorNoScale(ClientData->MeshTranslationOffset + CharacterOwner->GetBaseTranslationOffset());
+		CharacterOwner->GetMesh()->SetRelativeLocation(NewRelTranslation);
+	}
+}
+
+
+bool UCharacterMovementComponent::ClientUpdatePositionAfterServerUpdate()
+{
+	if (!HasValidData())
+	{
+		return false;
+	}
+
+	FNetworkPredictionData_Client_Character* ClientData = GetPredictionData_Client_Character();
+	check(ClientData);
+
+	if (!ClientData->bUpdatePosition)
+	{
+		return false;
+	}
+
+	ClientData->bUpdatePosition = false;
+
+	// Don't do any network position updates on things running PHYS_RigidBody
+	if (CharacterOwner->GetRootComponent() && CharacterOwner->GetRootComponent()->IsSimulatingPhysics())
+	{
+		return false;
+	}
+
+	if (ClientData->SavedMoves.Num() == 0)
+	{
+		UE_LOG(LogNetPlayerMovement, VeryVerbose, TEXT("ClientUpdatePositionAfterServerUpdate No saved moves to replay"), ClientData->SavedMoves.Num());
+		return false;
+	}
+
+	// Save important values that might get affected by the replay.
+	const float SavedAnalogInputModifier = AnalogInputModifier;
+	const FRootMotionMovementParams BackupRootMotionParams = RootMotionParams;
+	const bool bRealJump = CharacterOwner->bPressedJump;
+	const bool bRealCrouch = bWantsToCrouch;
+	const bool bRealForceMaxAccel = bForceMaxAccel;
+	CharacterOwner->bClientWasFalling = (MovementMode == MOVE_Falling);
+	CharacterOwner->bClientUpdating = true;
+	bForceNextFloorCheck = true;
+
+	// Replay moves that have not yet been acked.
+	UE_LOG(LogNetPlayerMovement, VeryVerbose, TEXT("ClientUpdatePositionAfterServerUpdate Replaying Moves (%d)"), ClientData->SavedMoves.Num());
+	for (int32 i=0; i<ClientData->SavedMoves.Num(); i++)
+	{
+		const FSavedMovePtr& CurrentMove = ClientData->SavedMoves[i];
+		CurrentMove->PrepMoveFor(CharacterOwner);
+		MoveAutonomous(CurrentMove->TimeStamp, CurrentMove->DeltaTime, CurrentMove->GetCompressedFlags(), CurrentMove->Acceleration);
+		CurrentMove->PostUpdate(CharacterOwner, FSavedMove_Character::PostUpdate_Replay);
+	}
+
+	if (ClientData->PendingMove.IsValid())
+	{
+		ClientData->PendingMove->bForceNoCombine = true;
+	}
+
+	// Restore saved values.
+	AnalogInputModifier = SavedAnalogInputModifier;
+	RootMotionParams = BackupRootMotionParams;
+	CharacterOwner->bClientResimulateRootMotion = false;
+	CharacterOwner->bClientUpdating = false;
+	CharacterOwner->bPressedJump = bRealJump;
+	bWantsToCrouch = bRealCrouch;
+	bForceMaxAccel = bRealForceMaxAccel;
+	bForceNextFloorCheck = true;
+	
+	return (ClientData->SavedMoves.Num() > 0);
+}
+
+
+void UCharacterMovementComponent::ForcePositionUpdate(float DeltaTime)
+{
+	if (!HasValidData() || MovementMode == MOVE_None || UpdatedComponent->Mobility != EComponentMobility::Movable)
+	{
+		return;
+	}
+
+	check(CharacterOwner->Role == ROLE_Authority);
+	check(CharacterOwner->GetRemoteRole() == ROLE_AutonomousProxy);
+
+	if (!Velocity.IsZero())
+	{
+		PerformMovement(DeltaTime);
+	}
+}
+
+
+FNetworkPredictionData_Client* UCharacterMovementComponent::GetPredictionData_Client() const
+{
+	// Should only be called on client in network games
+	check(CharacterOwner != NULL);
+	check(CharacterOwner->Role < ROLE_Authority);
+	check(GetNetMode() == NM_Client);
+
+	if (!ClientPredictionData)
+	{
+		UCharacterMovementComponent* MutableThis = const_cast<UCharacterMovementComponent*>(this);
+		MutableThis->ClientPredictionData = new FNetworkPredictionData_Client_Character();
+	}
+
+	return ClientPredictionData;
+}
+
+FNetworkPredictionData_Server* UCharacterMovementComponent::GetPredictionData_Server() const
+{
+	// Should only be called on server in network games
+	check(CharacterOwner != NULL);
+	check(CharacterOwner->Role == ROLE_Authority);
+	check(GetNetMode() < NM_Client);
+
+	if (!ServerPredictionData)
+	{
+		UCharacterMovementComponent* MutableThis = const_cast<UCharacterMovementComponent*>(this);
+		MutableThis->ServerPredictionData = new FNetworkPredictionData_Server_Character();
+	}
+
+	return ServerPredictionData;
+}
+
+
+FNetworkPredictionData_Client_Character* UCharacterMovementComponent::GetPredictionData_Client_Character() const
+{
+	return static_cast<class FNetworkPredictionData_Client_Character*>(GetPredictionData_Client());
+}
+
+
+FNetworkPredictionData_Server_Character* UCharacterMovementComponent::GetPredictionData_Server_Character() const
+{
+	return static_cast<class FNetworkPredictionData_Server_Character*>(GetPredictionData_Server());
+}
+
+
+void UCharacterMovementComponent::ResetPredictionData_Client()
+{
+	if (ClientPredictionData)
+	{
+		delete ClientPredictionData;
+		ClientPredictionData = NULL;
+	}
+}
+
+void UCharacterMovementComponent::ResetPredictionData_Server()
+{
+	if (ServerPredictionData)
+	{
+		delete ServerPredictionData;
+		ServerPredictionData = NULL;
+	}	
+}
+
+float FNetworkPredictionData_Client_Character::UpdateTimeStampAndDeltaTime(float DeltaTime, class ACharacter & CharacterOwner, class UCharacterMovementComponent & CharacterMovementComponent)
+{
+	// Reset TimeStamp regularly to combat float accuracy decreasing over time.
+	if( CurrentTimeStamp > CharacterMovementComponent.MinTimeBetweenTimeStampResets )
+	{
+		UE_LOG(LogNetPlayerMovement, Log, TEXT("Resetting Client's TimeStamp %f"), CurrentTimeStamp);
+		CurrentTimeStamp = 0.f;
+
+		// Mark all buffered moves as having old time stamps, so we make sure to not resend them.
+		// That would confuse the server.
+		for(int32 MoveIndex=0; MoveIndex<SavedMoves.Num(); MoveIndex++)
+		{
+			const FSavedMovePtr& CurrentMove = SavedMoves[MoveIndex];
+			SavedMoves[MoveIndex]->bOldTimeStampBeforeReset = true;
+		}
+		// Do LastAckedMove as well. No need to do PendingMove as that move is part of the SavedMoves array.
+		if( LastAckedMove.IsValid() )
+		{
+			LastAckedMove->bOldTimeStampBeforeReset = true;
+		}
+	}
+
+	// Update Current TimeStamp.
+	CurrentTimeStamp += DeltaTime;
+	float ClientDeltaTime = DeltaTime;
+
+	// Server uses TimeStamps to derive DeltaTime which introduces some rounding errors.
+	// Make sure we do the same, so MoveAutonomous uses the same inputs and is deterministic!!
+	if( SavedMoves.Num() > 0 )
+	{
+		const FSavedMovePtr& PreviousMove = SavedMoves.Last();
+		if( !PreviousMove->bOldTimeStampBeforeReset )
+		{
+			// How server will calculate its deltatime to update physics.
+			const float ServerDeltaTime = CurrentTimeStamp - PreviousMove->TimeStamp;
+			// Have client always use the Server's DeltaTime. Otherwise our physics simulation will differ and we'll trigger too many position corrections and increase our network traffic.
+			ClientDeltaTime = ServerDeltaTime;
+		}
+	}
+
+	return CharacterOwner.CustomTimeDilation * FMath::Min(ClientDeltaTime, MaxResponseTime * CharacterOwner.GetWorldSettings()->GetEffectiveTimeDilation());
+}
+
+void UCharacterMovementComponent::ReplicateMoveToServer(float DeltaTime, const FVector& NewAcceleration)
+{
+	check(CharacterOwner != NULL);
+
+	// Can only start sending moves if our controllers are synced up over the network, otherwise we flood the reliable buffer.
+	APlayerController* PC = Cast<APlayerController>(CharacterOwner->GetController());
+	if (PC && PC->AcknowledgedPawn != CharacterOwner)
+	{
+		return;
+	}
+
+	FNetworkPredictionData_Client_Character* ClientData = GetPredictionData_Client_Character();
+	if (!ClientData)
+	{
+		return;
+	}
+	
+	// Update our delta time for physics simulation.
+	DeltaTime = ClientData->UpdateTimeStampAndDeltaTime(DeltaTime, *CharacterOwner, *this);
+
+	// Find the oldest (unacknowledged) important move (OldMove).
+	// Don't include the last move because it may be combined with the next new move.
+	// A saved move is interesting if it differs significantly from the last acknowledged move
+	FSavedMovePtr OldMove = NULL;
+	if( ClientData->LastAckedMove.IsValid() )
+	{
+		for (int32 i=0; i < ClientData->SavedMoves.Num()-1; i++)
+		{
+			const FSavedMovePtr& CurrentMove = ClientData->SavedMoves[i];
+			if (CurrentMove->IsImportantMove(ClientData->LastAckedMove))
+			{
+				OldMove = CurrentMove;
+				break;
+			}
+		}
+	}
+
+	// Get a SavedMove object to store the movement in.
+	FSavedMovePtr NewMove = ClientData->CreateSavedMove();
+	if (NewMove.IsValid() == false)
+	{
+		return;
+	}
+
+	NewMove->SetMoveFor(CharacterOwner, DeltaTime, NewAcceleration, *ClientData);
+
+	// see if the two moves could be combined
+	// do not combine moves which have different TimeStamps (before and after reset).
+	if( ClientData->PendingMove.IsValid() && !ClientData->PendingMove->bOldTimeStampBeforeReset && ClientData->PendingMove->CanCombineWith(NewMove, CharacterOwner, ClientData->MaxResponseTime * CharacterOwner->GetWorldSettings()->GetEffectiveTimeDilation()))
+	{
+		// Only combine and move back to the start location if we don't move back in to a spot that would make us collide with something new.
+		const FVector OldStartLocation = ClientData->PendingMove->GetRevertedLocation();
+		if (!OverlapTest(OldStartLocation, ClientData->PendingMove->StartRotation.Quaternion(), UpdatedComponent->GetCollisionObjectType(), GetPawnCapsuleCollisionShape(SHRINK_None), CharacterOwner))
+		{
+			FScopedMovementUpdate ScopedMovementUpdate(UpdatedComponent, EScopedUpdate::DeferredUpdates);
+			UE_LOG(LogNetPlayerMovement, VeryVerbose, TEXT("CombineMove: add delta %f + %f and revert from %f %f to %f %f"), DeltaTime, ClientData->PendingMove->DeltaTime, CharacterOwner->GetActorLocation().X, CharacterOwner->GetActorLocation().Y, OldStartLocation.X, OldStartLocation.Y);
+			
+			// to combine move, first revert pawn position to PendingMove start position, before playing combined move on client
+			const bool bNoCollisionCheck = true;
+			UpdatedComponent->SetWorldLocationAndRotation(OldStartLocation, ClientData->PendingMove->StartRotation, false);
+			Velocity = ClientData->PendingMove->StartVelocity;
+
+			SetBase(ClientData->PendingMove->StartBase.Get(), ClientData->PendingMove->StartBoneName);
+			CurrentFloor = ClientData->PendingMove->StartFloor;
+
+			// Now that we have reverted to the old position, prepare a new move from that position,
+			// using our current velocity, acceleration, and rotation, but applied over the combined time from the old and new move.
+
+			NewMove->DeltaTime += ClientData->PendingMove->DeltaTime;
+			
+			if (PC)
+			{
+				// We reverted position to that at the start of the pending move (above), however some code paths expect rotation to be set correctly
+				// before character movement occurs (via FaceRotation), so try that now. The bOrientRotationToMovement path happens later as part of PerformMovement() and PhysicsRotation().
+				CharacterOwner->FaceRotation(PC->GetControlRotation(), NewMove->DeltaTime);
+			}
+
+			SaveBaseLocation();
+			NewMove->SetInitialPosition(CharacterOwner);
+
+			// Remove pending move from move list. It would have to be the last move on the list.
+			if (ClientData->SavedMoves.Num() > 0 && ClientData->SavedMoves.Last() == ClientData->PendingMove)
+			{
+				ClientData->SavedMoves.Pop();
+			}
+			ClientData->FreeMove(ClientData->PendingMove);
+			ClientData->PendingMove = NULL;
+		}
+		else
+		{
+			//UE_LOG(LogNet, Log, TEXT("Not combining move, would collide at start location"));
+		}
+	}
+
+	// Perform the move locally
+	Acceleration = NewMove->Acceleration;
+	AnalogInputModifier = ComputeAnalogInputModifier(); // recompute since acceleration may have changed.
+	CharacterOwner->ClientRootMotionParams.Clear();
+	PerformMovement(NewMove->DeltaTime);
+
+	NewMove->PostUpdate(CharacterOwner, FSavedMove_Character::PostUpdate_Record);
+
+	// Add NewMove to the list
+	ClientData->SavedMoves.Push(NewMove);
+
+	const bool bCanDelayMove = (CVarNetEnableMoveCombining.GetValueOnGameThread() != 0) && CanDelaySendingMove(NewMove);
+
+	if (bCanDelayMove && ClientData->PendingMove.IsValid() == false)
+	{
+		// Decide whether to hold off on move
+		// send moves more frequently in small games where server isn't likely to be saturated
+		float NetMoveDelta;
+		UPlayer* Player = (PC ? PC->Player : NULL);
+
+		if (Player && (Player->CurrentNetSpeed > 10000) && (GetWorld()->GameState != NULL) && (GetWorld()->GameState->PlayerArray.Num() <= 10))
+		{
+			NetMoveDelta = 0.011f;
+		}
+		else if (Player && CharacterOwner->GetWorldSettings()->GameNetworkManagerClass) 
+		{
+			NetMoveDelta = FMath::Max(0.0222f,2 * GetDefault<AGameNetworkManager>(CharacterOwner->GetWorldSettings()->GameNetworkManagerClass)->MoveRepSize/Player->CurrentNetSpeed);
+		}
+		else
+		{
+			NetMoveDelta = 0.011f;
+		}
+
+		if ((GetWorld()->TimeSeconds - ClientData->ClientUpdateTime) * CharacterOwner->GetWorldSettings()->GetEffectiveTimeDilation() < NetMoveDelta)
+		{
+			// Delay sending this move.
+			ClientData->PendingMove = NewMove;
+			return;
+		}
+	}
+
+	ClientData->ClientUpdateTime = GetWorld()->TimeSeconds;
+
+	UE_LOG(LogNetPlayerMovement, Verbose, TEXT("Client ReplicateMove Time %f Acceleration %s Position %s DeltaTime %f"),
+		NewMove->TimeStamp, *NewMove->Acceleration.ToString(), *CharacterOwner->GetActorLocation().ToString(), DeltaTime);
+
+	// Send to the server
+	CallServerMove( NewMove.Get(), OldMove.Get() );
+	ClientData->PendingMove = NULL;
+}
+
+
+inline uint32 PackYawAndPitchTo32(const float Yaw, const float Pitch)
+{
+	const uint32 YawShort = FRotator::CompressAxisToShort(Yaw);
+	const uint32 PitchShort = FRotator::CompressAxisToShort(Pitch);
+	const uint32 Rotation32 = (YawShort << 16) | PitchShort;
+	return Rotation32;
+}
+
+
+void UCharacterMovementComponent::CallServerMove
+	(
+	const class FSavedMove_Character* NewMove,
+	const class FSavedMove_Character* OldMove
+	)
+{
+	check(NewMove != NULL);
+
+	// Compress rotation down to 5 bytes
+	const uint32 ClientYawPitchINT = PackYawAndPitchTo32(NewMove->SavedControlRotation.Yaw, NewMove->SavedControlRotation.Pitch);
+	const uint8 ClientRollBYTE = FRotator::CompressAxisToByte(NewMove->SavedControlRotation.Roll);
+
+	// Determine if we send absolute or relative location
+	UPrimitiveComponent* ClientMovementBase = NewMove->EndBase.Get();
+	const FName ClientBaseBone = NewMove->EndBoneName;
+	const FVector SendLocation = MovementBaseUtility::UseRelativeLocation(ClientMovementBase) ? NewMove->SavedRelativeLocation : NewMove->SavedLocation;
+
+	// send old move if it exists
+	if (OldMove)
+	{
+		ServerMoveOld(OldMove->TimeStamp, OldMove->Acceleration, OldMove->GetCompressedFlags());
+	}
+
+	FNetworkPredictionData_Client_Character* ClientData = GetPredictionData_Client_Character();
+	if (ClientData->PendingMove.IsValid())
+	{
+		const uint32 OldClientYawPitchINT = PackYawAndPitchTo32(ClientData->PendingMove->SavedControlRotation.Yaw, ClientData->PendingMove->SavedControlRotation.Pitch);
+
+		// send two moves simultaneously
+		ServerMoveDual
+			(
+			ClientData->PendingMove->TimeStamp,
+			ClientData->PendingMove->Acceleration,
+			ClientData->PendingMove->GetCompressedFlags(),
+			OldClientYawPitchINT,
+			NewMove->TimeStamp,
+			NewMove->Acceleration,
+			SendLocation,
+			NewMove->GetCompressedFlags(),
+			ClientRollBYTE,
+			ClientYawPitchINT,
+			ClientMovementBase,
+			ClientBaseBone,
+			PackNetworkMovementMode()
+			);
+	}
+	else
+	{
+		ServerMove
+			(
+			NewMove->TimeStamp,
+			NewMove->Acceleration,
+			SendLocation,
+			NewMove->GetCompressedFlags(),
+			ClientRollBYTE,
+			ClientYawPitchINT,
+			ClientMovementBase,
+			ClientBaseBone,
+			PackNetworkMovementMode()
+			);
+	}
+
+
+	APlayerController* PC = Cast<APlayerController>(CharacterOwner->GetController());
+	APlayerCameraManager* PlayerCameraManager = (PC ? PC->PlayerCameraManager : NULL);
+	if (PlayerCameraManager != NULL && PlayerCameraManager->bUseClientSideCameraUpdates)
+	{
+		PlayerCameraManager->bShouldSendClientSideCameraUpdate = true;
+	}
+}
+
+
+
+void UCharacterMovementComponent::ServerMoveOld_Implementation
+	(
+	float OldTimeStamp,
+	FVector_NetQuantize10 OldAccel,
+	uint8 OldMoveFlags
+	)
+{
+	if (!HasValidData() || !IsComponentTickEnabled())
+	{
+		return;
+	}
+
+	FNetworkPredictionData_Server_Character* ServerData = GetPredictionData_Server_Character();
+	check(ServerData);
+
+	if( !VerifyClientTimeStamp(OldTimeStamp, *ServerData) )
+	{
+		return;
+	}
+
+	UE_LOG(LogNetPlayerMovement, Log, TEXT("Recovered move from OldTimeStamp %f, DeltaTime: %f"), OldTimeStamp, OldTimeStamp - ServerData->CurrentClientTimeStamp);
+	const float MaxResponseTime = ServerData->MaxResponseTime * CharacterOwner->GetWorldSettings()->GetEffectiveTimeDilation();
+
+	MoveAutonomous(OldTimeStamp, FMath::Min(OldTimeStamp - ServerData->CurrentClientTimeStamp, MaxResponseTime), OldMoveFlags, OldAccel);
+
+	ServerData->CurrentClientTimeStamp = OldTimeStamp;
+}
+
+
+void UCharacterMovementComponent::ServerMoveDual_Implementation(
+	float TimeStamp0,
+	FVector_NetQuantize10 InAccel0,
+	uint8 PendingFlags,
+	uint32 View0,
+	float TimeStamp,
+	FVector_NetQuantize10 InAccel,
+	FVector_NetQuantize100 ClientLoc,
+	uint8 NewFlags,
+	uint8 ClientRoll,
+	uint32 View,
+	UPrimitiveComponent* ClientMovementBase,
+	FName ClientBaseBone,
+	uint8 ClientMovementMode)
+{
+	ServerMove_Implementation(TimeStamp0, InAccel0, FVector(1.f,2.f,3.f), PendingFlags, ClientRoll, View0, ClientMovementBase, ClientBaseBone, ClientMovementMode);
+	ServerMove_Implementation(TimeStamp, InAccel, ClientLoc, NewFlags, ClientRoll, View, ClientMovementBase, ClientBaseBone, ClientMovementMode);
+}
+
+bool UCharacterMovementComponent::VerifyClientTimeStamp(float TimeStamp, FNetworkPredictionData_Server_Character & ServerData)
+{
+	// Very large deltas happen around a TimeStamp reset.
+	const float DeltaTimeStamp = (TimeStamp - ServerData.CurrentClientTimeStamp);
+	if( FMath::Abs(DeltaTimeStamp) > (MinTimeBetweenTimeStampResets * 0.5f) )
+	{
+		// Client is resetting TimeStamp to increase accuracy.
+		if( DeltaTimeStamp < 0.f )
+		{
+			UE_LOG(LogNetPlayerMovement, Log, TEXT("TimeStamp reset detected. CurrentTimeStamp: %f, new TimeStamp: %f"), ServerData.CurrentClientTimeStamp, TimeStamp);
+			ServerData.CurrentClientTimeStamp = 0.f;
+			return true;
+		}
+		else
+		{
+			// We already reset the TimeStamp, but we just got an old outdated move before the switch.
+			// Just ignore it.
+			UE_LOG(LogNetPlayerMovement, Log, TEXT("TimeStamp expired. Before TimeStamp Reset. CurrentTimeStamp: %f, TimeStamp: %f"), ServerData.CurrentClientTimeStamp, TimeStamp);
+			return false;
+		}
+	}
+
+	// If TimeStamp is in the past, move is outdated, ignore it.
+	if( TimeStamp <= ServerData.CurrentClientTimeStamp )
+	{
+		UE_LOG(LogNetPlayerMovement, Log, TEXT("TimeStamp expired. %f, CurrentTimeStamp: %f"), TimeStamp, ServerData.CurrentClientTimeStamp);
+		return false;
+	}
+	
+	UE_LOG(LogNetPlayerMovement, VeryVerbose, TEXT("TimeStamp %f Accepted! CurrentTimeStamp: %f"), TimeStamp, ServerData.CurrentClientTimeStamp);
+	return true;
+}
+
+void UCharacterMovementComponent::ServerMove_Implementation(
+	float TimeStamp,
+	FVector_NetQuantize10 InAccel,
+	FVector_NetQuantize100 ClientLoc,
+	uint8 MoveFlags,
+	uint8 ClientRoll,
+	uint32 View,
+	UPrimitiveComponent* ClientMovementBase,
+	FName ClientBaseBoneName,
+	uint8 ClientMovementMode)
+{
+	if (!HasValidData() || !IsComponentTickEnabled())
+	{
+		return;
+	}	
+
+	FNetworkPredictionData_Server_Character* ServerData = GetPredictionData_Server_Character();
+	check(ServerData);
+
+	if( !VerifyClientTimeStamp(TimeStamp, *ServerData) )
+	{
+		return;
+	}
+
+	bool bServerReadyForClient = true;
+	APlayerController* PC = Cast<APlayerController>(CharacterOwner->GetController());
+	if (PC)
+	{
+		bServerReadyForClient = PC->NotifyServerReceivedClientData(CharacterOwner, TimeStamp);
+		if (!bServerReadyForClient)
+		{
+			InAccel = FVector::ZeroVector;
+		}
+	}
+
+	// View components
+	const uint16 ViewPitch = (View & 65535);
+	const uint16 ViewYaw = (View >> 16);
+	
+	const FVector Accel = InAccel;
+	// Save move parameters.
+	const float DeltaTime = ServerData->GetServerMoveDeltaTime(TimeStamp) * CharacterOwner->CustomTimeDilation;
+
+	ServerData->CurrentClientTimeStamp = TimeStamp;
+	ServerData->ServerTimeStamp = GetWorld()->TimeSeconds;
+	FRotator ViewRot;
+	ViewRot.Pitch = FRotator::DecompressAxisFromShort(ViewPitch);
+	ViewRot.Yaw = FRotator::DecompressAxisFromShort(ViewYaw);
+	ViewRot.Roll = FRotator::DecompressAxisFromByte(ClientRoll);
+
+	if (PC)
+	{
+		PC->SetControlRotation(ViewRot);
+	}
+
+	if (!bServerReadyForClient)
+	{
+		return;
+	}
+
+	// Perform actual movement
+	if ((CharacterOwner->GetWorldSettings()->Pauser == NULL) && (DeltaTime > 0.f))
+	{
+		if (PC)
+		{
+			PC->UpdateRotation(DeltaTime);
+		}
+
+		MoveAutonomous(TimeStamp, DeltaTime, MoveFlags, Accel);
+	}
+
+	UE_LOG(LogNetPlayerMovement, Verbose, TEXT("ServerMove Time %f Acceleration %s Position %s DeltaTime %f"),
+			TimeStamp, *Accel.ToString(), *CharacterOwner->GetActorLocation().ToString(), DeltaTime);
+
+	ServerMoveHandleClientError(TimeStamp, DeltaTime, Accel, ClientLoc, ClientMovementBase, ClientBaseBoneName, ClientMovementMode);
+}
+
+
+void UCharacterMovementComponent::ServerMoveHandleClientError(float TimeStamp, float DeltaTime, const FVector& Accel, const FVector& RelativeClientLoc, UPrimitiveComponent* ClientMovementBase, FName ClientBaseBoneName, uint8 ClientMovementMode)
+{
+	if (RelativeClientLoc == FVector(1.f,2.f,3.f)) // first part of double servermove
+	{
+		return;
+	}
+
+	FNetworkPredictionData_Server_Character* ServerData = GetPredictionData_Server_Character();
+	check(ServerData);
+
+	// Don't prevent more recent updates from being sent if received this frame.
+	// We're going to send out an update anyway, might as well be the most recent one.
+	APlayerController* PC = Cast<APlayerController>(CharacterOwner->GetController());
+	if( (ServerData->LastUpdateTime != GetWorld()->TimeSeconds) && GetDefault<AGameNetworkManager>()->WithinUpdateDelayBounds(PC, ServerData->LastUpdateTime))
+	{
+		return;
+	}
+
+	// Offset may be relative to base component
+	FVector ClientLoc = RelativeClientLoc;
+	if (MovementBaseUtility::UseRelativeLocation(ClientMovementBase))
+	{
+		FVector BaseLocation;
+		FQuat BaseRotation;
+		MovementBaseUtility::GetMovementBaseTransform(ClientMovementBase, ClientBaseBoneName, BaseLocation, BaseRotation);
+		ClientLoc += BaseLocation;
+	}
+
+	// Compute the client error from the server's position
+	const FVector LocDiff = CharacterOwner->GetActorLocation() - ClientLoc;
+	const uint8 CurrentPackedMovementMode = PackNetworkMovementMode();
+
+	// If client has accumulated a noticeable positional error, correct him.
+	if (GetDefault<AGameNetworkManager>()->ExceedsAllowablePositionError(LocDiff) || (CurrentPackedMovementMode != ClientMovementMode))
+	{
+		UPrimitiveComponent* MovementBase = CharacterOwner->GetMovementBase();
+		ServerData->PendingAdjustment.NewVel = Velocity;
+		ServerData->PendingAdjustment.NewBase = MovementBase;
+		ServerData->PendingAdjustment.NewBaseBoneName = CharacterOwner->GetBasedMovement().BoneName;
+		ServerData->PendingAdjustment.NewLoc = CharacterOwner->GetActorLocation();
+		ServerData->PendingAdjustment.NewRot = CharacterOwner->GetActorRotation();
+
+		ServerData->PendingAdjustment.bBaseRelativePosition = MovementBaseUtility::UseRelativeLocation(MovementBase);
+		if (ServerData->PendingAdjustment.bBaseRelativePosition)
+		{
+			// Relative location
+			ServerData->PendingAdjustment.NewLoc = CharacterOwner->GetBasedMovement().Location;
+			
+			// TODO: this could be a relative rotation, but all client corrections ignore rotation right now except the root motion one, which would need to be updated.
+			//ServerData->PendingAdjustment.NewRot = CharacterOwner->GetBasedMovement().Rotation;
+		}
+
+
+#if !(UE_BUILD_SHIPPING || UE_BUILD_TEST)
+		if (CVarNetShowCorrections.GetValueOnGameThread() != 0)
+		{
+			UE_LOG(LogNetPlayerMovement, Warning, TEXT("******** Client Error at %f is %f Accel %s LocDiff %s ClientLoc %s, ServerLoc: %s, Base: %s, Bone: %s"),
+				TimeStamp, LocDiff.Size(), *Accel.ToString(), *LocDiff.ToString(), *ClientLoc.ToString(), *CharacterOwner->GetActorLocation().ToString(), *GetNameSafe(MovementBase), *ServerData->PendingAdjustment.NewBaseBoneName.ToString());
+			const float DebugLifetime = CVarNetCorrectionLifetime.GetValueOnGameThread();
+			DrawDebugCapsule(GetWorld(), CharacterOwner->GetActorLocation(), CharacterOwner->GetSimpleCollisionHalfHeight(), CharacterOwner->GetSimpleCollisionRadius(), FQuat::Identity, FColor(100, 255, 100), true, DebugLifetime);
+			DrawDebugCapsule(GetWorld(), ClientLoc                    , CharacterOwner->GetSimpleCollisionHalfHeight(), CharacterOwner->GetSimpleCollisionRadius(), FQuat::Identity, FColor(255, 100, 100), true, DebugLifetime);
+		}
+#endif
+
+		ServerData->LastUpdateTime = GetWorld()->TimeSeconds;
+		ServerData->PendingAdjustment.DeltaTime = DeltaTime;
+		ServerData->PendingAdjustment.TimeStamp = TimeStamp;
+		ServerData->PendingAdjustment.bAckGoodMove = false;
+		ServerData->PendingAdjustment.MovementMode = CurrentPackedMovementMode;
+	}
+	else
+	{
+		if (GetDefault<AGameNetworkManager>()->ClientAuthorativePosition)
+		{
+			if (!LocDiff.IsZero() || ClientMovementMode != CurrentPackedMovementMode)
+			{
+				// Just set the position. On subsequent moves we will resolve initially overlapping conditions.
+				UpdatedComponent->SetWorldLocation(ClientLoc, false);
+
+				// Trust the client's movement mode.
+				ApplyNetworkMovementMode(ClientMovementMode);
+
+				// Update base and floor at new location.
+				SetBase(ClientMovementBase, ClientBaseBoneName);
+				UpdateFloorFromAdjustment();
+
+				// Even if base has not changed, we need to recompute the relative offsets (since we've moved).
+				SaveBaseLocation();
+			}
+		}
+
+		// acknowledge receipt of this successful servermove()
+		ServerData->PendingAdjustment.TimeStamp = TimeStamp;
+		ServerData->PendingAdjustment.bAckGoodMove = true;
+	}
+}
+
+
+bool UCharacterMovementComponent::ServerMove_Validate(float TimeStamp, FVector_NetQuantize10 InAccel, FVector_NetQuantize100 ClientLoc, uint8 MoveFlags, uint8 ClientRoll, uint32 View, UPrimitiveComponent* ClientMovementBase, FName ClientBaseBoneName, uint8 ClientMovementMode)
+{
+	return true;
+}
+
+bool UCharacterMovementComponent::ServerMoveDual_Validate(float TimeStamp0, FVector_NetQuantize10 InAccel0, uint8 PendingFlags, uint32 View0, float TimeStamp, FVector_NetQuantize10 InAccel, FVector_NetQuantize100 ClientLoc, uint8 NewFlags, uint8 ClientRoll, uint32 View, UPrimitiveComponent* ClientMovementBase, FName ClientBaseBoneName, uint8 ClientMovementMode)
+{
+	return true;
+}
+
+bool UCharacterMovementComponent::ServerMoveOld_Validate(float OldTimeStamp, FVector_NetQuantize10 OldAccel, uint8 OldMoveFlags)
+{
+	return true;
+}
+
+
+
+void UCharacterMovementComponent::MoveAutonomous
+	(
+	float ClientTimeStamp,
+	float DeltaTime,
+	uint8 CompressedFlags,
+	const FVector& NewAccel
+	)
+{
+	if (!HasValidData())
+	{
+		return;
+	}
+
+	UpdateFromCompressedFlags(CompressedFlags);
+	CharacterOwner->CheckJumpInput(DeltaTime);
+
+	Acceleration = ConstrainInputAcceleration(NewAccel);
+	AnalogInputModifier = ComputeAnalogInputModifier();
+	
+	PerformMovement(DeltaTime);
+
+	// If not playing root motion, tick animations after physics. We do this here to keep events, notifies, states and transitions in sync with client updates.
+	if( !CharacterOwner->bClientUpdating && !CharacterOwner->IsPlayingRootMotion() && CharacterOwner->GetMesh() )
+	{
+		TickCharacterPose(DeltaTime);
+		// TODO: SaveBaseLocation() in case tick moves us?
+	}
+}
+
+
+void UCharacterMovementComponent::UpdateFloorFromAdjustment()
+{
+	if (!HasValidData())
+	{
+		return;
+	}
+
+	// If walking, try to update the cached floor so it is current. This is necessary for UpdateBasedMovement() and MoveAlongFloor() to work properly.
+	// If base is now NULL, presumably we are no longer walking. If we had a valid floor but don't find one now, we'll likely start falling.
+	if (CharacterOwner->GetMovementBase())
+	{
+		FindFloor(UpdatedComponent->GetComponentLocation(), CurrentFloor, false);
+	}
+	else
+	{
+		CurrentFloor.Clear();
+	}
+
+	bForceNextFloorCheck = true;
+}
+
+
+void UCharacterMovementComponent::SendClientAdjustment()
+{
+	if (!HasValidData())
+	{
+		return;
+	}
+
+	FNetworkPredictionData_Server_Character* ServerData = GetPredictionData_Server_Character();
+	check(ServerData);
+
+	if (ServerData->PendingAdjustment.TimeStamp <= 0.f)
+	{
+		return;
+	}
+
+	if (ServerData->PendingAdjustment.bAckGoodMove == true)
+	{
+		// just notify client this move was received
+		ClientAckGoodMove(ServerData->PendingAdjustment.TimeStamp);
+	}
+	else
+	{
+		if( CharacterOwner->IsPlayingNetworkedRootMotionMontage() )
+		{
+			FRotator Rotation = ServerData->PendingAdjustment.NewRot.GetNormalized();
+			FVector_NetQuantizeNormal CompressedRotation(Rotation.Pitch / 180.f, Rotation.Yaw / 180.f, Rotation.Roll / 180.f);
+			ClientAdjustRootMotionPosition
+				(
+				ServerData->PendingAdjustment.TimeStamp,
+				CharacterOwner->GetRootMotionAnimMontageInstance()->GetPosition(),
+				ServerData->PendingAdjustment.NewLoc,
+				CompressedRotation,
+				ServerData->PendingAdjustment.NewVel.Z,
+				ServerData->PendingAdjustment.NewBase,
+				ServerData->PendingAdjustment.NewBaseBoneName,
+				ServerData->PendingAdjustment.NewBase != NULL,
+				ServerData->PendingAdjustment.bBaseRelativePosition,
+				PackNetworkMovementMode()
+				);
+		}
+		else if (ServerData->PendingAdjustment.NewVel.IsZero())
+		{
+			ClientVeryShortAdjustPosition
+				(
+				ServerData->PendingAdjustment.TimeStamp,
+				ServerData->PendingAdjustment.NewLoc,
+				ServerData->PendingAdjustment.NewBase,
+				ServerData->PendingAdjustment.NewBaseBoneName,
+				ServerData->PendingAdjustment.NewBase != NULL,
+				ServerData->PendingAdjustment.bBaseRelativePosition,
+				PackNetworkMovementMode()
+				);
+		}
+		else
+		{
+			ClientAdjustPosition
+				(
+				ServerData->PendingAdjustment.TimeStamp,
+				ServerData->PendingAdjustment.NewLoc,
+				ServerData->PendingAdjustment.NewVel,
+				ServerData->PendingAdjustment.NewBase,
+				ServerData->PendingAdjustment.NewBaseBoneName,
+				ServerData->PendingAdjustment.NewBase != NULL,
+				ServerData->PendingAdjustment.bBaseRelativePosition,
+				PackNetworkMovementMode()
+				);
+		}
+	}
+
+	ServerData->PendingAdjustment.TimeStamp = 0;
+	ServerData->PendingAdjustment.bAckGoodMove = false;
+}
+
+
+void UCharacterMovementComponent::ClientVeryShortAdjustPosition_Implementation
+	(
+	float TimeStamp,
+	FVector NewLoc,
+	UPrimitiveComponent* NewBase,
+	FName NewBaseBoneName,
+	bool bHasBase,
+	bool bBaseRelativePosition,
+	uint8 ServerMovementMode
+	)
+{
+	if (HasValidData())
+	{
+		ClientAdjustPosition(TimeStamp, NewLoc, FVector::ZeroVector, NewBase, NewBaseBoneName, bHasBase, bBaseRelativePosition, ServerMovementMode);
+	}
+}
+
+
+void UCharacterMovementComponent::ClientAdjustPosition_Implementation
+	(
+	float TimeStamp,
+	FVector NewLocation,
+	FVector NewVelocity,
+	UPrimitiveComponent* NewBase,
+	FName NewBaseBoneName,
+	bool bHasBase,
+	bool bBaseRelativePosition,
+	uint8 ServerMovementMode
+	)
+{
+	if (!HasValidData() || !IsComponentTickEnabled())
+	{
+		return;
+	}
+
+
+	FNetworkPredictionData_Client_Character* ClientData = GetPredictionData_Client_Character();
+	check(ClientData);
+	
+	// Make sure the base actor exists on this client.
+	const bool bUnresolvedBase = bHasBase && (NewBase == NULL);
+	if (bUnresolvedBase)
+	{
+		if (bBaseRelativePosition)
+		{
+			UE_LOG(LogNetPlayerMovement, Warning, TEXT("ClientAdjustPosition_Implementation could not resolve the new relative movement base actor, ignoring server correction!"));
+			return;
+		}
+		else
+		{
+			UE_LOG(LogNetPlayerMovement, Verbose, TEXT("ClientAdjustPosition_Implementation could not resolve the new absolute movement base actor, but WILL use the position!"));
+		}
+	}
+	
+	// Ack move if it has not expired.
+	int32 MoveIndex = ClientData->GetSavedMoveIndex(TimeStamp);
+	if( MoveIndex == INDEX_NONE )
+	{
+		if( ClientData->LastAckedMove.IsValid() )
+		{
+			UE_LOG(LogNetPlayerMovement, Log,  TEXT("ClientAdjustPosition_Implementation could not find Move for TimeStamp: %f, LastAckedTimeStamp: %f, CurrentTimeStamp: %f"), TimeStamp, ClientData->LastAckedMove->TimeStamp, ClientData->CurrentTimeStamp);
+		}
+		return;
+	}
+	ClientData->AckMove(MoveIndex);
+		
+	//  Received Location is relative to dynamic base
+	if (bBaseRelativePosition)
+	{
+		FVector BaseLocation;
+		FQuat BaseRotation;
+		MovementBaseUtility::GetMovementBaseTransform(NewBase, NewBaseBoneName, BaseLocation, BaseRotation); // TODO: error handling if returns false
+		NewLocation += BaseLocation;
+	}
+
+#if !(UE_BUILD_SHIPPING || UE_BUILD_TEST)
+	static const auto CVarShowCorrections = IConsoleManager::Get().FindTConsoleVariableDataInt(TEXT("p.NetShowCorrections"));
+	if (CVarShowCorrections && CVarShowCorrections->GetValueOnGameThread() != 0)
+	{
+		UE_LOG(LogNetPlayerMovement, Warning, TEXT("******** ClientAdjustPosition Time %f velocity %s position %s NewBase: %s NewBone: %s SavedMoves %d"), TimeStamp, *NewVelocity.ToString(), *NewLocation.ToString(), *GetNameSafe(NewBase), *NewBaseBoneName.ToString(), ClientData->SavedMoves.Num());
+		static const auto CVarCorrectionLifetime = IConsoleManager::Get().FindTConsoleVariableDataFloat(TEXT("p.NetCorrectionLifetime"));
+		const float DebugLifetime = CVarCorrectionLifetime ? CVarCorrectionLifetime->GetValueOnGameThread() : 1.f;
+		DrawDebugCapsule(GetWorld(), CharacterOwner->GetActorLocation()	, CharacterOwner->GetSimpleCollisionHalfHeight(), CharacterOwner->GetSimpleCollisionRadius(), FQuat::Identity, FColor(255, 100, 100), true, DebugLifetime);
+		DrawDebugCapsule(GetWorld(), NewLocation						, CharacterOwner->GetSimpleCollisionHalfHeight(), CharacterOwner->GetSimpleCollisionRadius(), FQuat::Identity, FColor(100, 255, 100), true, DebugLifetime);
+	}
+#endif //!(UE_BUILD_SHIPPING || UE_BUILD_TEST)
+
+	// Trust the server's positioning.
+	UpdatedComponent->SetWorldLocation(NewLocation, false);	
+	Velocity = NewVelocity;
+
+	// Trust the server's movement mode
+	UPrimitiveComponent* PreviousBase = CharacterOwner->GetMovementBase();
+	ApplyNetworkMovementMode(ServerMovementMode);
+
+	// Set base component
+	UPrimitiveComponent* FinalBase = NewBase;
+	FName FinalBaseBoneName = NewBaseBoneName;
+	if (bUnresolvedBase)
+	{
+		check(NewBase == NULL);
+		check(!bBaseRelativePosition);
+		
+		// We had an unresolved base from the server
+		// If walking, we'd like to continue walking if possible, to avoid falling for a frame, so try to find a base where we moved to.
+		if (PreviousBase)
+		{
+			FindFloor(UpdatedComponent->GetComponentLocation(), CurrentFloor, false);
+			FinalBase = CurrentFloor.HitResult.Component.Get();
+			FinalBaseBoneName = CurrentFloor.HitResult.BoneName;
+		}
+	}
+	SetBase(FinalBase, FinalBaseBoneName);
+
+	// Update floor at new location
+	UpdateFloorFromAdjustment();
+	bJustTeleported = true;
+
+	// Even if base has not changed, we need to recompute the relative offsets (since we've moved).
+	SaveBaseLocation();
+
+	UpdateComponentVelocity();
+	ClientData->bUpdatePosition = true;
+}
+
+void UCharacterMovementComponent::ClientAdjustRootMotionPosition_Implementation(
+	float TimeStamp,
+	float ServerMontageTrackPosition,
+	FVector ServerLoc,
+	FVector_NetQuantizeNormal ServerRotation,
+	float ServerVelZ,
+	UPrimitiveComponent * ServerBase,
+	FName ServerBaseBoneName,
+	bool bHasBase,
+	bool bBaseRelativePosition,
+	uint8 ServerMovementMode)
+{
+	if (!HasValidData() || !IsComponentTickEnabled())
+	{
+		return;
+	}
+
+	// Call ClientAdjustPosition first. This will Ack the move if it's not outdated.
+	ClientAdjustPosition(TimeStamp, ServerLoc, FVector(0.f, 0.f, ServerVelZ), ServerBase, ServerBaseBoneName, bHasBase, bBaseRelativePosition, ServerMovementMode);
+	
+	FNetworkPredictionData_Client_Character* ClientData = GetPredictionData_Client_Character();
+	check(ClientData);
+
+	// If this adjustment wasn't acknowledged (because outdated), then abort.
+	if( !ClientData->LastAckedMove.IsValid() || (ClientData->LastAckedMove->TimeStamp != TimeStamp) )
+	{
+		return;
+	}
+
+	// We're going to replay Root Motion. This is relative to the Pawn's rotation, so we need to reset that as well.
+	FRotator DecompressedRot(ServerRotation.X * 180.f, ServerRotation.Y * 180.f, ServerRotation.Z * 180.f);
+	CharacterOwner->SetActorRotation(DecompressedRot);
+	const FVector ServerLocation(ServerLoc);
+	UE_LOG(LogRootMotion, Log,  TEXT("ClientAdjustRootMotionPosition_Implementation TimeStamp: %f, ServerMontageTrackPosition: %f, ServerLocation: %s, ServerRotation: %s, ServerVelZ: %f, ServerBase: %s"),
+		TimeStamp, ServerMontageTrackPosition, *ServerLocation.ToCompactString(), *DecompressedRot.ToCompactString(), ServerVelZ, *GetNameSafe(ServerBase) );
+
+	// DEBUG - get some insight on where errors came from
+	if( false )
+	{
+		const FVector DeltaLocation = ServerLocation - ClientData->LastAckedMove->SavedLocation;
+		const FRotator DeltaRotation = (DecompressedRot - ClientData->LastAckedMove->SavedRotation).GetNormalized();
+		const float DeltaTrackPosition = (ServerMontageTrackPosition - ClientData->LastAckedMove->RootMotionTrackPosition);
+		const float DeltaVelZ = (ServerVelZ - ClientData->LastAckedMove->SavedVelocity.Z);
+
+		UE_LOG(LogRootMotion, Log,  TEXT("\tErrors DeltaLocation: %s, DeltaRotation: %s, DeltaTrackPosition: %f"),
+			*DeltaLocation.ToCompactString(), *DeltaRotation.ToCompactString(), DeltaTrackPosition );
+	}
+
+	// Server disagrees with Client on the Root Motion AnimMontage Track position.
+	if( CharacterOwner->bClientResimulateRootMotion || (ServerMontageTrackPosition != ClientData->LastAckedMove->RootMotionTrackPosition) )
+	{
+		UE_LOG(LogRootMotion, Warning,  TEXT("\tServer disagrees with Client's track position!! ServerTrackPosition: %f, ClientTrackPosition: %f, DeltaTrackPosition: %f. TimeStamp: %f"),
+			ServerMontageTrackPosition, ClientData->LastAckedMove->RootMotionTrackPosition, (ServerMontageTrackPosition - ClientData->LastAckedMove->RootMotionTrackPosition), TimeStamp);
+
+		// Not much we can do there unfortunately, just jump to server's track position.
+		FAnimMontageInstance * RootMotionMontageInstance = CharacterOwner->GetRootMotionAnimMontageInstance();
+		if( RootMotionMontageInstance )
+		{
+			RootMotionMontageInstance->SetPosition(ServerMontageTrackPosition);
+			CharacterOwner->bClientResimulateRootMotion = true;
+		}
+	}
+}
+
+void UCharacterMovementComponent::ClientAckGoodMove_Implementation(float TimeStamp)
+{
+	if (!HasValidData() || !IsComponentTickEnabled())
+	{
+		return;
+	}
+
+	FNetworkPredictionData_Client_Character* ClientData = GetPredictionData_Client_Character();
+	check(ClientData);
+
+
+	// Ack move if it has not expired.
+	int32 MoveIndex = ClientData->GetSavedMoveIndex(TimeStamp);
+	if( MoveIndex == INDEX_NONE )
+	{
+		if( ClientData->LastAckedMove.IsValid() )
+		{
+			UE_LOG(LogNetPlayerMovement, Warning, TEXT("ClientAckGoodMove_Implementation could not find Move for TimeStamp: %f, LastAckedTimeStamp: %f, CurrentTimeStamp: %f"), TimeStamp, ClientData->LastAckedMove->TimeStamp, ClientData->CurrentTimeStamp);
+		}
+		return;
+	}
+	ClientData->AckMove(MoveIndex);
+}
+
+void UCharacterMovementComponent::CapsuleTouched( AActor* Other, UPrimitiveComponent* OtherComp, int32 OtherBodyIndex, bool bFromSweep, const FHitResult& SweepResult )
+{
+	if (!bEnablePhysicsInteraction)
+	{
+		return;
+	}
+
+	if (OtherComp != NULL && OtherComp->IsAnySimulatingPhysics())
+	{
+		const FVector OtherLoc = OtherComp->GetComponentLocation();
+		const FVector Loc = UpdatedComponent->GetComponentLocation();
+		FVector ImpulseDir = FVector(OtherLoc.X - Loc.X, OtherLoc.Y - Loc.Y, 0.25f).GetSafeNormal();
+		ImpulseDir = (ImpulseDir + Velocity.GetSafeNormal2D()) * 0.5f;
+		ImpulseDir.Normalize();
+
+		FName BoneName = NAME_None;
+		if (OtherBodyIndex != INDEX_NONE)
+		{
+			BoneName = ((USkinnedMeshComponent*)OtherComp)->GetBoneName(OtherBodyIndex);
+		}
+
+		float TouchForceFactorModified = TouchForceFactor;
+
+		if ( bTouchForceScaledToMass )
+		{
+			FBodyInstance* BI = OtherComp->GetBodyInstance(BoneName);
+			TouchForceFactorModified *= BI ? BI->GetBodyMass() : 1.0f;
+		}
+
+		float ImpulseStrength = FMath::Clamp(Velocity.Size2D() * TouchForceFactorModified, 
+			MinTouchForce > 0.0f ? MinTouchForce : -FLT_MAX, 
+			MaxTouchForce > 0.0f ? MaxTouchForce : FLT_MAX);
+
+		FVector Impulse = ImpulseDir * ImpulseStrength;
+
+		OtherComp->AddImpulse(Impulse, BoneName);
+	}
+}
+
+void UCharacterMovementComponent::SetAvoidanceGroup(int32 GroupFlags)
+{
+	AvoidanceGroup.SetFlagsDirectly(GroupFlags);
+}
+
+void UCharacterMovementComponent::SetGroupsToAvoid(int32 GroupFlags)
+{
+	GroupsToAvoid.SetFlagsDirectly(GroupFlags);
+}
+
+void UCharacterMovementComponent::SetGroupsToIgnore(int32 GroupFlags)
+{
+	GroupsToIgnore.SetFlagsDirectly(GroupFlags);
+}
+
+void UCharacterMovementComponent::SetAvoidanceEnabled(bool bEnable)
+{
+	if (bUseRVOAvoidance != bEnable)
+	{
+		bUseRVOAvoidance = bEnable;
+
+		// this is a safety check - it's possible to not have CharacterOwner at this point if this function gets
+		// called too early
+		ensure(GetCharacterOwner());
+		if (GetCharacterOwner() != nullptr)
+		{
+			UAvoidanceManager* AvoidanceManager = GetWorld()->GetAvoidanceManager();
+			if (AvoidanceManager && bEnable && AvoidanceUID == 0)
+			{
+				AvoidanceManager->RegisterMovementComponent(this, AvoidanceWeight);
+			}
+		}
+	}
+}
+
+void UCharacterMovementComponent::ApplyRepulsionForce(float DeltaSeconds)
+{
+	if (UpdatedPrimitive && RepulsionForce > 0.0f)
+	{
+		FCollisionQueryParams QueryParams;
+		QueryParams.bReturnFaceIndex = false;
+		QueryParams.bReturnPhysicalMaterial = false;
+
+		const FCollisionShape CollisionShape = UpdatedPrimitive->GetCollisionShape();
+		const float CapsuleRadius = CollisionShape.GetCapsuleRadius();
+		const float CapsuleHalfHeight = CollisionShape.GetCapsuleHalfHeight();
+		const float RepulsionForceRadius = CapsuleRadius * 1.2f;
+		const float StopBodyDistance = 2.5f;
+
+		const TArray<FOverlapInfo>& Overlaps = UpdatedPrimitive->GetOverlapInfos();
+		const FVector MyLocation = UpdatedPrimitive->GetComponentLocation();
+
+		for (int32 i=0; i < Overlaps.Num(); i++)
+		{
+			const FOverlapInfo& Overlap = Overlaps[i];
+
+			UPrimitiveComponent* OverlapComp = Overlap.OverlapInfo.Component.Get();
+			if (!OverlapComp || OverlapComp->Mobility < EComponentMobility::Movable)
+			{ 
+				continue; 
+			}
+
+			FName BoneName = NAME_None;
+			if (Overlap.GetBodyIndex() != INDEX_NONE && OverlapComp->IsA(USkinnedMeshComponent::StaticClass()))
+			{
+				BoneName = ((USkinnedMeshComponent*)OverlapComp)->GetBoneName(Overlap.GetBodyIndex());
+			}
+
+			// Use the body instead of the component for cases where we have multi-body overlaps enabled
+			FBodyInstance* OverlapBody = OverlapComp->GetBodyInstance(BoneName);
+
+			if (!OverlapBody)
+			{
+				UE_LOG(LogCharacterMovement, Warning, TEXT("%s could not find overlap body for bone %s"), *GetName(), *BoneName.ToString());
+				continue;
+			}
+
+			// Early out if this is not a destructible and the body is not simulated
+			bool bIsCompDestructible = OverlapComp->IsA(UDestructibleComponent::StaticClass());
+			if (!bIsCompDestructible && !OverlapBody->IsInstanceSimulatingPhysics())
+			{
+				continue;
+			}
+
+			FTransform BodyTransform = OverlapBody->GetUnrealWorldTransform();
+
+			FVector BodyVelocity = OverlapBody->GetUnrealWorldVelocity();
+			FVector BodyLocation = BodyTransform.GetLocation();
+
+			// Trace to get the hit location on the capsule
+			FHitResult Hit;
+			bool bHasHit = UpdatedPrimitive->LineTraceComponent(Hit, BodyLocation,
+																FVector(MyLocation.X, MyLocation.Y, BodyLocation.Z),
+																QueryParams);
+
+			FVector HitLoc = Hit.ImpactPoint;
+			bool bIsPenetrating = Hit.bStartPenetrating || Hit.PenetrationDepth > 2.5f;
+
+			// If we didn't hit the capsule, we're inside the capsule
+			if(!bHasHit) 
+			{ 
+				HitLoc = BodyLocation; 
+				bIsPenetrating = true;
+			}
+
+			const float DistanceNow = (HitLoc - BodyLocation).SizeSquared2D();
+			const float DistanceLater = (HitLoc - (BodyLocation + BodyVelocity * DeltaSeconds)).SizeSquared2D();
+
+			if (BodyLocation.SizeSquared() > 0.1f && bHasHit && DistanceNow < StopBodyDistance && !bIsPenetrating)
+			{
+				OverlapBody->SetLinearVelocity(FVector(0.0f, 0.0f, 0.0f), false);
+			}
+			else if (DistanceLater <= DistanceNow || bIsPenetrating)
+			{
+				FVector ForceCenter(MyLocation.X, MyLocation.Y, bHasHit ? HitLoc.Z : MyLocation.Z);
+
+				if (!bHasHit)
+				{
+					ForceCenter.Z = FMath::Clamp(BodyLocation.Z, MyLocation.Z - CapsuleHalfHeight, MyLocation.Z + CapsuleHalfHeight);
+				}
+
+				OverlapBody->AddRadialForceToBody(ForceCenter, RepulsionForceRadius, RepulsionForce * Mass, ERadialImpulseFalloff::RIF_Constant);
+			}
+		}
+	}
+}
+
+void UCharacterMovementComponent::ApplyAccumulatedForces(float DeltaSeconds)
+{
+	if (PendingImpulseToApply.Z != 0.f || PendingForceToApply.Z != 0.f)
+	{
+		// check to see if applied momentum is enough to overcome gravity
+		if ( IsMovingOnGround() && (PendingImpulseToApply.Z + (PendingForceToApply.Z * DeltaSeconds) + (GetGravityZ() * DeltaSeconds) > SMALL_NUMBER))
+		{
+			SetMovementMode(MOVE_Falling);
+		}
+	}
+
+	Velocity += PendingImpulseToApply + (PendingForceToApply * DeltaSeconds);
+
+	PendingImpulseToApply = FVector::ZeroVector;
+	PendingForceToApply = FVector::ZeroVector;
+}
+
+void UCharacterMovementComponent::AddRadialForce(const FVector& Origin, float Radius, float Strength, enum ERadialImpulseFalloff Falloff)
+{
+	FVector Delta = UpdatedComponent->GetComponentLocation() - Origin;
+	const float DeltaMagnitude = Delta.Size();
+
+	// Do nothing if outside radius
+	if(DeltaMagnitude > Radius)
+	{
+		return;
+	}
+
+	Delta = Delta.GetSafeNormal();
+
+	float ForceMagnitude = Strength;
+	if (Falloff == RIF_Linear && Radius > 0.0f)
+	{
+		ForceMagnitude *= (1.0f - (DeltaMagnitude / Radius));
+	}
+
+	AddForce(Delta * ForceMagnitude);
+}
+ 
+void UCharacterMovementComponent::AddRadialImpulse(const FVector& Origin, float Radius, float Strength, enum ERadialImpulseFalloff Falloff, bool bVelChange)
+{
+	FVector Delta = UpdatedComponent->GetComponentLocation() - Origin;
+	const float DeltaMagnitude = Delta.Size();
+
+	// Do nothing if outside radius
+	if(DeltaMagnitude > Radius)
+	{
+		return;
+	}
+
+	Delta = Delta.GetSafeNormal();
+
+	float ImpulseMagnitude = Strength;
+	if (Falloff == RIF_Linear && Radius > 0.0f)
+	{
+		ImpulseMagnitude *= (1.0f - (DeltaMagnitude / Radius));
+	}
+
+	AddImpulse(Delta * ImpulseMagnitude, bVelChange);
+}
+
+void UCharacterMovementComponent::RegisterComponentTickFunctions(bool bRegister)
+{
+	Super::RegisterComponentTickFunctions(bRegister);
+
+	if (bRegister)
+	{
+		if (SetupActorComponentTickFunction(&PreClothComponentTick))
+		{
+			PreClothComponentTick.Target = this;
+		}
+	}
+	else
+	{
+		if(PreClothComponentTick.IsTickFunctionRegistered())
+		{
+			PreClothComponentTick.UnRegisterTickFunction();
+		}
+	}
+}
+
+void UCharacterMovementComponent::TickCharacterPose(float DeltaTime)
+{
+	check(CharacterOwner && CharacterOwner->GetMesh());
+
+	CharacterOwner->GetMesh()->TickPose(DeltaTime, true);
+
+	// Grab root motion now that we have ticked the pose
+	if (CharacterOwner->IsPlayingRootMotion())
+	{
+		FRootMotionMovementParams RootMotion = CharacterOwner->GetMesh()->ConsumeRootMotion();
+		if (RootMotion.bHasRootMotion)
+		{
+			RootMotionParams.Accumulate(RootMotion);
+		}
+	}
+}
+
+FNetworkPredictionData_Client_Character::FNetworkPredictionData_Client_Character()
+	: ClientUpdateTime(0.f)
+	, CurrentTimeStamp(0.f)
+	, PendingMove(NULL)
+	, LastAckedMove(NULL)
+	, MaxFreeMoveCount(32)
+	, MaxSavedMoveCount(96)
+	, bUpdatePosition(false)
+	, bSmoothNetUpdates(false)
+	, MeshTranslationOffset(ForceInitToZero)
+	, MaxSmoothNetUpdateDist(0.f)
+	, NoSmoothNetUpdateDist(0.f)
+	, SmoothNetUpdateTime(0.f)
+	, MaxResponseTime(0.f)
+{
+	bSmoothNetUpdates = true;
+	MaxSmoothNetUpdateDist = 84.0;
+	NoSmoothNetUpdateDist = 128.0;
+	SmoothNetUpdateTime = 0.100;
+
+	MaxResponseTime = 0.125f;
+}
+
+
+FNetworkPredictionData_Client_Character::~FNetworkPredictionData_Client_Character()
+{
+	SavedMoves.Empty();
+	FreeMoves.Empty();
+	PendingMove = NULL;
+	LastAckedMove = NULL;
+}
+
+
+FSavedMovePtr FNetworkPredictionData_Client_Character::CreateSavedMove()
+{
+	if (SavedMoves.Num() >= MaxSavedMoveCount)
+	{
+		UE_LOG(LogNetPlayerMovement, Warning, TEXT("CreateSavedMove: Hit limit of %d saved moves (timing out or very bad ping?)"), SavedMoves.Num());
+		// Free all saved moves
+		for (int32 i=0; i < SavedMoves.Num(); i++)
+		{
+			FreeMove(SavedMoves[i]);
+		}
+		SavedMoves.Reset();
+	}
+
+	if (FreeMoves.Num() == 0)
+	{
+		// No free moves, allocate a new one.
+		FSavedMovePtr NewMove = AllocateNewMove();
+		check(NewMove.IsValid());
+		NewMove->Clear();
+		return NewMove;
+	}
+	else
+	{
+		// Pull from the free pool
+		FSavedMovePtr FirstFree = FreeMoves.Pop();
+		FirstFree->Clear();
+		return FirstFree;
+	}
+}
+
+
+FSavedMovePtr FNetworkPredictionData_Client_Character::AllocateNewMove()
+{
+	return FSavedMovePtr(new FSavedMove_Character());
+}
+
+
+void FNetworkPredictionData_Client_Character::FreeMove(const FSavedMovePtr& Move)
+{
+	if (Move.IsValid())
+	{
+		// Only keep a pool of a limited number of moves.
+		if (FreeMoves.Num() < MaxFreeMoveCount)
+		{
+			FreeMoves.Push(Move);
+		}
+
+		// Shouldn't keep a reference to the move on the free list.
+		if (PendingMove == Move)
+		{
+			PendingMove = NULL;
+		}
+		if( LastAckedMove == Move )
+		{
+			LastAckedMove = NULL;
+		}
+	}
+}
+
+int32 FNetworkPredictionData_Client_Character::GetSavedMoveIndex(float TimeStamp) const
+{
+	if( SavedMoves.Num() > 0 )
+	{
+		// If LastAckedMove isn't using an old TimeStamp (before reset), we can prevent the iteration if incoming TimeStamp is outdated
+		if( LastAckedMove.IsValid() && !LastAckedMove->bOldTimeStampBeforeReset && (TimeStamp <= LastAckedMove->TimeStamp) )
+		{
+			return INDEX_NONE;
+		}
+
+		// Otherwise see if we can find this move.
+		for (int32 Index=0; Index<SavedMoves.Num(); Index++)
+		{
+			const FSavedMovePtr& CurrentMove = SavedMoves[Index];
+			if( CurrentMove->TimeStamp == TimeStamp )
+			{
+				return Index;
+			}
+		}
+	}
+	return INDEX_NONE;
+}
+
+void FNetworkPredictionData_Client_Character::AckMove(int32 AckedMoveIndex) 
+{
+	// It is important that we know the move exists before we go deleting outdated moves.
+	// Timestamps are not guaranteed to be increasing order all the time, since they can be reset!
+	if( AckedMoveIndex != INDEX_NONE )
+	{
+		// Keep reference to LastAckedMove
+		const FSavedMovePtr& AckedMove = SavedMoves[AckedMoveIndex];
+		UE_LOG(LogNetPlayerMovement, VeryVerbose, TEXT("AckedMove Index: %2d (%2d moves). TimeStamp: %f, CurrentTimeStamp: %f"), AckedMoveIndex, SavedMoves.Num(), AckedMove->TimeStamp, CurrentTimeStamp);
+		if( LastAckedMove.IsValid() )
+		{
+			FreeMove(LastAckedMove);
+		}
+		LastAckedMove = AckedMove;
+
+		// Free expired moves.
+		for(int32 MoveIndex=0; MoveIndex<AckedMoveIndex; MoveIndex++)
+		{
+			const FSavedMovePtr& Move = SavedMoves[MoveIndex];
+			FreeMove(Move);
+		}
+
+		// And finally cull all of those, so only the unacknowledged moves remain in SavedMoves.
+		SavedMoves.RemoveAt(0, AckedMoveIndex + 1);
+	}
+}
+
+FNetworkPredictionData_Server_Character::FNetworkPredictionData_Server_Character()
+	: PendingAdjustment()
+	, CurrentClientTimeStamp(0.f)
+	, LastUpdateTime(0.f)
+	, MaxResponseTime(0.125f)
+{
+}
+
+
+FNetworkPredictionData_Server_Character::~FNetworkPredictionData_Server_Character()
+{
+}
+
+
+float FNetworkPredictionData_Server_Character::GetServerMoveDeltaTime(float TimeStamp) const
+{
+	const float DeltaTime = FMath::Min(MaxResponseTime, TimeStamp - CurrentClientTimeStamp);
+	return DeltaTime;
+}
+
+
+FSavedMove_Character::FSavedMove_Character()
+{
+	AccelMagThreshold = 1.f;
+	AccelDotThreshold = 0.9f;
+	AccelDotThresholdCombine = 0.996f; // approx 5 degrees.
+}
+
+FSavedMove_Character::~FSavedMove_Character()
+{
+}
+
+void FSavedMove_Character::Clear()
+{
+	bPressedJump = false;
+	bWantsToCrouch = false;
+	bForceMaxAccel = false;
+	bForceNoCombine = false;
+	bOldTimeStampBeforeReset = false;
+
+	TimeStamp = 0.f;
+	DeltaTime = 0.f;
+	CustomTimeDilation = 1.0f;
+	JumpKeyHoldTime = 0.0f;
+	MovementMode = 0;
+
+	StartLocation = FVector::ZeroVector;
+	StartRelativeLocation = FVector::ZeroVector;
+	StartVelocity = FVector::ZeroVector;
+	StartFloor = FFindFloorResult();
+	StartRotation = FRotator::ZeroRotator;
+	StartControlRotation = FRotator::ZeroRotator;
+	StartBaseRotation = FQuat::Identity;
+	StartCapsuleRadius = 0.f;
+	StartCapsuleHalfHeight = 0.f;
+	StartBase = NULL;
+	StartBoneName = NAME_None;
+
+	SavedLocation = FVector::ZeroVector;
+	SavedRotation = FRotator::ZeroRotator;
+	SavedRelativeLocation = FVector::ZeroVector;
+	Acceleration = FVector::ZeroVector;
+	SavedControlRotation = FRotator::ZeroRotator;
+	EndBase = NULL;
+	EndBoneName = NAME_None;
+
+	RootMotionMontage = NULL;
+	RootMotionTrackPosition = 0.f;
+	RootMotionMovement.Clear();
+}
+
+
+void FSavedMove_Character::SetMoveFor(ACharacter* Character, float InDeltaTime, FVector const& NewAccel, class FNetworkPredictionData_Client_Character & ClientData)
+{
+	DeltaTime = InDeltaTime;
+	
+	SetInitialPosition(Character);
+
+	AccelMag = NewAccel.Size();
+	AccelNormal = (AccelMag > SMALL_NUMBER ? NewAccel / AccelMag : FVector::ZeroVector);
+	
+	// Round value, so that client and server match exactly (and so we can send with less bandwidth). This rounded value is copied back to the client in ReplicateMoveToServer.
+	// This is done after the AccelMag and AccelNormal are computed above, because those are only used client-side for combining move logic and need to remain accurate.
+	Acceleration = Character->GetCharacterMovement()->RoundAcceleration(NewAccel);
+
+	bPressedJump = Character->bPressedJump;
+	JumpKeyHoldTime = Character->JumpKeyHoldTime;
+	bWantsToCrouch = Character->GetCharacterMovement()->bWantsToCrouch;
+	bForceMaxAccel = Character->GetCharacterMovement()->bForceMaxAccel;
+	MovementMode = Character->GetCharacterMovement()->PackNetworkMovementMode();
+
+	TimeStamp = ClientData.CurrentTimeStamp;
+}
+
+void FSavedMove_Character::SetInitialPosition(ACharacter* Character)
+{
+	StartLocation = Character->GetActorLocation();
+	StartRotation = Character->GetActorRotation();
+	StartVelocity = Character->GetCharacterMovement()->Velocity;
+	UPrimitiveComponent* const MovementBase = Character->GetMovementBase();
+	StartBase = MovementBase;
+	StartBaseRotation = FQuat::Identity;
+	StartFloor = Character->GetCharacterMovement()->CurrentFloor;
+	CustomTimeDilation = Character->CustomTimeDilation;
+	StartBoneName = Character->GetBasedMovement().BoneName;
+
+	if (MovementBaseUtility::UseRelativeLocation(MovementBase))
+	{
+		StartRelativeLocation = Character->GetBasedMovement().Location;
+		FVector StartBaseLocation_Unused;
+		MovementBaseUtility::GetMovementBaseTransform(MovementBase, StartBoneName, StartBaseLocation_Unused, StartBaseRotation);
+	}
+
+	StartControlRotation = Character->GetControlRotation().Clamp();
+	Character->GetCapsuleComponent()->GetScaledCapsuleSize(StartCapsuleRadius, StartCapsuleHalfHeight);
+}
+
+void FSavedMove_Character::PostUpdate(ACharacter* Character, FSavedMove_Character::EPostUpdateMode PostUpdateMode)
+{
+	// Common code for both recording and after a replay.
+	{
+		SavedLocation = Character->GetActorLocation();
+		SavedRotation = Character->GetActorRotation();
+		UPrimitiveComponent* const MovementBase = Character->GetMovementBase();
+		EndBase = MovementBase;
+		EndBoneName = Character->GetBasedMovement().BoneName;
+		if (MovementBaseUtility::UseRelativeLocation(MovementBase))
+		{
+			SavedRelativeLocation = Character->GetBasedMovement().Location;
+		}
+
+		SavedControlRotation = Character->GetControlRotation().Clamp();
+	}
+
+	// Only save RootMotion params when initially recording
+	if (PostUpdateMode == PostUpdate_Record)
+	{
+		const FAnimMontageInstance* RootMotionMontageInstance = Character->GetRootMotionAnimMontageInstance();
+		if( RootMotionMontageInstance )
+		{
+			RootMotionMontage = RootMotionMontageInstance->Montage;
+			RootMotionTrackPosition = RootMotionMontageInstance->GetPosition();
+			RootMotionMovement = Character->ClientRootMotionParams;
+		}
+	}
+}
+
+bool FSavedMove_Character::IsImportantMove(const FSavedMovePtr& LastAckedMove) const
+{
+	// Check if any important movement flags have changed status.
+	if( (bPressedJump != LastAckedMove->bPressedJump) || (bWantsToCrouch != LastAckedMove->bWantsToCrouch) )
+	{
+		return true;
+	}
+
+	if (MovementMode != LastAckedMove->MovementMode)
+	{
+		return true;
+	}
+
+	// check if acceleration has changed significantly
+	if( Acceleration != LastAckedMove->Acceleration ) 
+	{
+		// Compare magnitude and orientation
+		if( (FMath::Abs(AccelMag - LastAckedMove->AccelMag) > AccelMagThreshold) || ((AccelNormal | LastAckedMove->AccelNormal) < AccelDotThreshold) )
+		{
+			return true;
+		}
+	}
+	return false;
+}
+
+FVector FSavedMove_Character::GetRevertedLocation() const
+{
+	const UPrimitiveComponent* MovementBase = StartBase.Get();
+	if (MovementBaseUtility::UseRelativeLocation(MovementBase))
+	{
+		FVector BaseLocation; FQuat BaseRotation;
+		MovementBaseUtility::GetMovementBaseTransform(MovementBase, StartBoneName, BaseLocation, BaseRotation);
+		return BaseLocation + StartRelativeLocation;
+	}
+
+	return StartLocation;
+}
+
+bool UCharacterMovementComponent::CanDelaySendingMove(const FSavedMovePtr& NewMove)
+{
+	return true;
+}
+
+bool FSavedMove_Character::CanCombineWith(const FSavedMovePtr& NewMove, ACharacter* Character, float MaxDelta) const
+{
+	if (bForceNoCombine || NewMove->bForceNoCombine)
+	{
+		return false;
+	}
+
+	// Cannot combine moves which contain root motion for now.
+	// @fixme laurent - we should be able to combine most of them though, but current scheme of resetting pawn location and resimulating forward doesn't work.
+	// as we don't want to tick montage twice (so we don't fire events twice). So we need to rearchitecture this so we tick only the second part of the move, and reuse the first part.
+	if( (RootMotionMontage != NULL) || (NewMove->RootMotionMontage != NULL) )
+	{
+		return false;
+	}
+
+	if (NewMove->Acceleration.IsZero())
+	{
+		return Acceleration.IsZero()
+			&& StartVelocity.IsZero()
+			&& NewMove->StartVelocity.IsZero()
+			&& !bPressedJump && !NewMove->bPressedJump
+			&& bWantsToCrouch == NewMove->bWantsToCrouch
+			&& StartBase == NewMove->StartBase
+			&& StartBoneName == NewMove->StartBoneName
+			&& MovementMode == NewMove->MovementMode
+			&& StartCapsuleRadius == NewMove->StartCapsuleRadius
+			&& StartCapsuleHalfHeight == NewMove->StartCapsuleHalfHeight
+			&& StartBaseRotation.Equals(NewMove->StartBaseRotation) // only if base hasn't rotated
+			&& (CustomTimeDilation == NewMove->CustomTimeDilation);
+	}
+	else
+	{
+		return (NewMove->DeltaTime + DeltaTime < MaxDelta)
+			&& !bPressedJump && !NewMove->bPressedJump
+			&& bWantsToCrouch == NewMove->bWantsToCrouch
+			&& StartBase == NewMove->StartBase
+			&& StartBoneName == NewMove->StartBoneName
+			&& MovementMode == NewMove->MovementMode
+			&& StartCapsuleRadius == NewMove->StartCapsuleRadius
+			&& StartCapsuleHalfHeight == NewMove->StartCapsuleHalfHeight
+			&& FVector::Coincident(AccelNormal, NewMove->AccelNormal, AccelDotThresholdCombine)
+			&& StartBaseRotation.Equals(NewMove->StartBaseRotation) // only if base hasn't rotated
+			&& (CustomTimeDilation == NewMove->CustomTimeDilation);
+	}
+}
+
+void FSavedMove_Character::PrepMoveFor(ACharacter* Character)
+{
+	if( RootMotionMontage != NULL )
+	{
+		// If we need to resimulate Root Motion, then do so.
+		if( Character->bClientResimulateRootMotion )
+		{
+			// Make sure RootMotion montage matches what we are playing now.
+			FAnimMontageInstance * RootMotionMontageInstance = Character->GetRootMotionAnimMontageInstance();
+			if( RootMotionMontageInstance && (RootMotionMontage == RootMotionMontageInstance->Montage) )
+			{
+				RootMotionMovement.Clear();
+				RootMotionTrackPosition = RootMotionMontageInstance->GetPosition();
+				RootMotionMontageInstance->SimulateAdvance(DeltaTime, RootMotionTrackPosition, RootMotionMovement);
+				RootMotionMontageInstance->SetPosition(RootMotionTrackPosition);
+			}
+		}
+		Character->GetCharacterMovement()->RootMotionParams = RootMotionMovement;
+	}
+
+	Character->GetCharacterMovement()->bForceMaxAccel = bForceMaxAccel;
+	Character->JumpKeyHoldTime = JumpKeyHoldTime;
+}
+
+
+uint8 FSavedMove_Character::GetCompressedFlags() const
+{
+	uint8 Result = 0;
+
+	if (bPressedJump)
+	{
+		Result |= FLAG_JumpPressed;
+	}
+
+	if (bWantsToCrouch)
+	{
+		Result |= FLAG_WantsToCrouch;
+	}
+
+	return Result;
+}
+
+void UCharacterMovementComponent::UpdateFromCompressedFlags(uint8 Flags)
+{
+	if (!CharacterOwner)
+	{
+		return;
+	}
+
+	const bool bWasJumping = CharacterOwner->bPressedJump;
+
+	CharacterOwner->bPressedJump = ((Flags & FSavedMove_Character::FLAG_JumpPressed) != 0);	
+	bWantsToCrouch = ((Flags & FSavedMove_Character::FLAG_WantsToCrouch) != 0);
+
+	// Reset JumpKeyHoldTime when player presses Jump key on server as well.
+	if (!bWasJumping && CharacterOwner->bPressedJump)
+	{
+		CharacterOwner->JumpKeyHoldTime = 0.0f;
+	}
+}
+
