Engine: Moved actor components into Components folder

diff --git a/Engine/Source/Runtime/Engine/Private/Components/DecalComponent.cpp b/Engine/Source/Runtime/Engine/Private/Components/DecalComponent.cpp
new file mode 100644
--- /dev/null
+++ b/Engine/Source/Runtime/Engine/Private/Components/DecalComponent.cpp
@@ -0,0 +1,137 @@
+// Copyright 1998-2015 Epic Games, Inc. All Rights Reserved.
+
+/*=============================================================================
+	DecalComponent.cpp: Decal component implementation.
+=============================================================================*/
+
+#include "EnginePrivate.h"
+#include "LevelUtils.h"
+#include "Components/DecalComponent.h"
+
+FDeferredDecalProxy::FDeferredDecalProxy(const UDecalComponent* InComponent)
+{
+	UMaterialInterface* EffectiveMaterial = UMaterial::GetDefaultMaterial(MD_DeferredDecal);
+
+	if(InComponent->DecalMaterial)
+	{
+		UMaterial* BaseMaterial = InComponent->DecalMaterial->GetMaterial();
+
+		if(BaseMaterial->MaterialDomain == MD_DeferredDecal)
+		{
+			EffectiveMaterial = InComponent->DecalMaterial;
+		}
+	}
+
+	Component = InComponent;
+	DecalMaterial = EffectiveMaterial;
+	ComponentTrans = InComponent->GetComponentToWorld();
+	DrawInGame = InComponent->ShouldRender();
+	bOwnerSelected = InComponent->IsOwnerSelected();
+	SortOrder = InComponent->SortOrder;
+}
+
+void FDeferredDecalProxy::SetTransform(const FTransform& InComponentToWorld)
+{
+	ComponentTrans = InComponentToWorld;
+}
+
+UDecalComponent::UDecalComponent(const FObjectInitializer& ObjectInitializer)
+	: Super(ObjectInitializer)
+{
+}
+
+void UDecalComponent::SetLifeSpan(const float LifeSpan)
+{
+	if (LifeSpan > 0.f)
+	{
+		GetWorld()->GetTimerManager().SetTimer(TimerHandle_DestroyDecalComponent, this, &UDecalComponent::LifeSpanCallback, LifeSpan, false);
+	}
+	else
+	{
+		GetWorld()->GetTimerManager().ClearTimer(TimerHandle_DestroyDecalComponent);
+	}
+}
+
+void UDecalComponent::LifeSpanCallback()
+{
+	DestroyComponent();
+}
+
+void UDecalComponent::SetSortOrder(int32 Value)
+{
+	SortOrder = Value;
+	MarkRenderStateDirty();
+}
+
+void UDecalComponent::SetDecalMaterial(class UMaterialInterface* NewDecalMaterial)
+{
+	DecalMaterial = NewDecalMaterial;
+
+	MarkRenderStateDirty();	
+}
+
+void UDecalComponent::PushSelectionToProxy()
+{
+	MarkRenderStateDirty();	
+}
+
+class UMaterialInterface* UDecalComponent::GetDecalMaterial() const
+{
+	return DecalMaterial;
+}
+
+class UMaterialInstanceDynamic* UDecalComponent::CreateDynamicMaterialInstance()
+{
+	// Create the MID
+	UMaterialInstanceDynamic* Instance = UMaterialInstanceDynamic::Create(DecalMaterial, this);
+
+	// Assign it, once parent is set
+	SetDecalMaterial(Instance);
+
+	return Instance;
+}
+
+void UDecalComponent::GetUsedMaterials( TArray<UMaterialInterface*>& OutMaterials ) const
+{
+	OutMaterials.Add( GetDecalMaterial() );
+}
+
+
+FDeferredDecalProxy* UDecalComponent::CreateSceneProxy()
+{
+	return new FDeferredDecalProxy(this);
+}
+
+FBoxSphereBounds UDecalComponent::CalcBounds(const FTransform& LocalToWorld) const
+{
+	return FBoxSphereBounds();
+}
+
+void UDecalComponent::CreateRenderState_Concurrent()
+{
+	Super::CreateRenderState_Concurrent();
+
+	// Mimics UPrimitiveComponent's visibility logic, although without the UPrimitiveCompoent visibility flags
+	if ( ShouldComponentAddToScene() && ShouldRender() )
+	{
+		World->Scene->AddDecal(this);
+	}
+}
+
+void UDecalComponent::SendRenderTransform_Concurrent()
+{	
+	//If Decal isn't hidden update its transform.
+	if ( ShouldComponentAddToScene() && ShouldRender() )
+	{
+		World->Scene->UpdateDecalTransform(this);
+	}
+
+	Super::SendRenderTransform_Concurrent();
+}
+
+void UDecalComponent::DestroyRenderState_Concurrent()
+{
+	Super::DestroyRenderState_Concurrent();
+	World->Scene->RemoveDecal(this);
+}
+
diff --git a/Engine/Source/Runtime/Engine/Private/Components/DestructibleComponent.cpp b/Engine/Source/Runtime/Engine/Private/Components/DestructibleComponent.cpp
new file mode 100644
--- /dev/null
+++ b/Engine/Source/Runtime/Engine/Private/Components/DestructibleComponent.cpp
@@ -0,0 +1,1367 @@
+// Copyright 1998-2015 Epic Games, Inc. All Rights Reserved.
+
+/*=============================================================================
+	DestructibleComponent.cpp: UDestructibleComponent methods.
+=============================================================================*/
+
+
+#include "EnginePrivate.h"
+#include "PhysicsPublic.h"
+#include "PhysicsEngine/DestructibleActor.h"
+#include "PhysicsEngine/PhysXSupport.h"
+#include "Collision/PhysXCollision.h"
+#include "ParticleDefinitions.h"
+#include "ObjectEditorUtils.h"
+#include "Particles/ParticleSystemComponent.h"
+#include "PhysicsEngine/DestructibleActor.h"
+#include "Engine/DestructibleMesh.h"
+#include "Components/DestructibleComponent.h"
+#include "NavigationSystemHelpers.h"
+#include "PhysicalMaterials/PhysicalMaterial.h"
+#include "Kismet/GameplayStatics.h"
+#include "GameFramework/DamageType.h"
+
+UDestructibleComponent::UDestructibleComponent(const FObjectInitializer& ObjectInitializer)
+	: Super(ObjectInitializer)
+#if WITH_PHYSX
+	, PhysxUserData(this)
+#endif
+{
+
+	PrimaryComponentTick.bCanEverTick = false;
+	PostPhysicsComponentTick.bCanEverTick = false;
+
+	bHasCustomNavigableGeometry = EHasCustomNavigableGeometry::EvenIfNotCollidable;
+
+	BodyInstance.SetUseAsyncScene(true);
+	static FName CollisionProfileName(TEXT("Destructible"));
+	SetCollisionProfileName(CollisionProfileName);
+
+	bAlwaysCreatePhysicsState = true;
+	bIsActive = true;
+	bMultiBodyOverlap = true;
+
+	LargeChunkThreshold = 25.f;
+
+	SetSpaceBaseDoubleBuffering(false);
+}
+
+void UDestructibleComponent::Serialize(FArchive& Ar)
+{
+	Super::Serialize(Ar);
+
+#if WITH_EDITORONLY_DATA
+	if(Ar.IsLoading())
+	{
+		// Copy our skeletal mesh value to our transient variable, so it appears in slate correctly.
+		this->DestructibleMesh = GetDestructibleMesh();
+	}
+#endif // WITH_EDITORONLY_DATA
+}
+
+#if WITH_EDITOR
+void UDestructibleComponent::PostEditChangeProperty( struct FPropertyChangedEvent& PropertyChangedEvent )
+{
+	static const FName NAME_DestructibleComponent = FName(TEXT("DestructibleComponent"));
+	static const FName NAME_DestructibleMesh = FName(TEXT("DestructibleMesh"));
+
+	Super::PostEditChangeProperty(PropertyChangedEvent);
+
+	if (PropertyChangedEvent.Property != NULL)
+	{
+		if ((FObjectEditorUtils::GetCategoryFName(PropertyChangedEvent.Property) == NAME_DestructibleComponent)
+		&&  (PropertyChangedEvent.Property->GetFName() == NAME_DestructibleMesh))
+		{
+			// If our transient mesh has changed, update our skeletal mesh.
+			SetSkeletalMesh( this->DestructibleMesh );
+		}
+	}
+}
+#endif // WITH_EDITOR
+
+FBoxSphereBounds UDestructibleComponent::CalcBounds(const FTransform& LocalToWorld) const
+{
+#if WITH_APEX
+	if( ApexDestructibleActor == NULL )
+	{
+		// Fallback if we don't have physics
+		return Super::CalcBounds(LocalToWorld);
+	}
+
+	const PxBounds3& PBounds = ApexDestructibleActor->getBounds();
+
+	return FBoxSphereBounds( FBox( P2UVector(PBounds.minimum), P2UVector(PBounds.maximum) ) );
+#else	// #if WITH_APEX
+	return Super::CalcBounds(LocalToWorld);
+#endif	// #if WITH_APEX
+}
+
+bool IsImpactDamageEnabled(const UDestructibleMesh* TheDestructibleMesh, int32 Level)
+{
+	if(TheDestructibleMesh->DefaultDestructibleParameters.DamageParameters.ImpactDamage == 0.f)
+	{
+		return false;
+	}
+
+	bool bEnableImpactDamage = false;
+	const FDestructibleDepthParameters& DepthParams = TheDestructibleMesh->DefaultDestructibleParameters.DepthParameters[Level];
+	const EImpactDamageOverride LevelOverride = DepthParams.ImpactDamageOverride; 
+
+	switch(LevelOverride)
+	{
+		case EImpactDamageOverride::IDO_On:
+		{
+			return true;
+		}
+
+		case EImpactDamageOverride::IDO_Off:
+		{
+			return false;
+		}
+
+		default:
+		{
+			//return default if we're within the default level
+		    return TheDestructibleMesh->DefaultDestructibleParameters.DamageParameters.DefaultImpactDamageDepth >= Level ? TheDestructibleMesh->DefaultDestructibleParameters.DamageParameters.bEnableImpactDamage : false;
+		}
+	}
+}
+
+void UDestructibleComponent::OnUpdateTransform(bool bSkipPhysicsMove)
+{
+	// We are handling the physics move below, so don't handle it at higher levels
+	Super::OnUpdateTransform(true);
+
+	if (SkeletalMesh == NULL)
+	{
+		return;
+	}
+
+	if (!bPhysicsStateCreated || bSkipPhysicsMove)
+	{
+		return;
+	}
+
+	const FTransform& CurrentLocalToWorld = ComponentToWorld;
+
+	if(CurrentLocalToWorld.ContainsNaN())
+	{
+		return;
+	}
+
+	// warn if it has non-uniform scale
+	const FVector& MeshScale3D = CurrentLocalToWorld.GetScale3D();
+#if !(UE_BUILD_SHIPPING || UE_BUILD_TEST)
+	if( !MeshScale3D.IsUniform() )
+	{
+		UE_LOG(LogPhysics, Log, TEXT("UDestructibleComponent::SendPhysicsTransform : Non-uniform scale factor (%s) can cause physics to mismatch for %s  SkelMesh: %s"), *MeshScale3D.ToString(), *GetFullName(), SkeletalMesh ? *SkeletalMesh->GetFullName() : TEXT("NULL"));
+	}
+#endif
+
+#if WITH_APEX
+	if (ApexDestructibleActor != NULL && !BodyInstance.bSimulatePhysics)
+	{
+		PxMat44 GlobalPose(PxMat33(U2PQuat(CurrentLocalToWorld.GetRotation())), U2PVector(CurrentLocalToWorld.GetTranslation()));
+		ApexDestructibleActor->setGlobalPose(GlobalPose);
+	}
+#endif // #if WITH_APEX
+}
+
+void UDestructibleComponent::CreatePhysicsState()
+{
+	// to avoid calling PrimitiveComponent, I'm just calling ActorComponent::CreatePhysicsState
+	// @todo lh - fix me based on the discussion with Bryan G
+	UActorComponent::CreatePhysicsState();
+	bPhysicsStateCreated = true;
+
+	// What we want to do with BodySetup is simply use it to store a PhysicalMaterial, and possibly some other relevant fields.  Set up pointers from the BodyInstance to the BodySetup and this component
+	UBodySetup* BodySetup = GetBodySetup();
+	BodyInstance.OwnerComponent	= this;
+	BodyInstance.BodySetup = BodySetup;
+	BodyInstance.InstanceBodyIndex = 0;
+
+#if WITH_APEX
+	if( SkeletalMesh == NULL )
+	{
+		return;
+	}
+
+	FPhysScene* PhysScene = World->GetPhysicsScene();
+	check(PhysScene);
+
+	if( GApexModuleDestructible == NULL )
+	{
+		UE_LOG(LogPhysics, Log, TEXT("UDestructibleComponent::CreatePhysicsState(): APEX must be enabled to init UDestructibleComponent physics.") );
+		return;
+	}
+
+	if( ApexDestructibleActor != NULL )
+	{
+		UE_LOG(LogPhysics, Log, TEXT("UDestructibleComponent::CreatePhysicsState(): NxDestructibleActor already created.") );
+		return;
+	}
+
+	UDestructibleMesh* TheDestructibleMesh = GetDestructibleMesh();
+	if( TheDestructibleMesh == NULL || TheDestructibleMesh->ApexDestructibleAsset == NULL)
+	{
+		UE_LOG(LogPhysics, Log, TEXT("UDestructibleComponent::CreatePhysicsState(): No DestructibleMesh or missing ApexDestructibleAsset.") );
+		return;
+	}
+
+	int32 ChunkCount = TheDestructibleMesh->ApexDestructibleAsset->getChunkCount();
+	// Ensure the chunks start off invisible.  RefreshBoneTransforms should make them visible.
+	for (int32 ChunkIndex = 0; ChunkIndex < ChunkCount; ++ChunkIndex)
+	{
+		SetChunkVisible(ChunkIndex, false);
+	}
+
+#if WITH_EDITOR
+	if (GIsEditor && !World->IsGameWorld())
+	{
+		// In the editor, only set the 0 chunk to be visible.
+		if (TheDestructibleMesh->ApexDestructibleAsset->getChunkCount() > 0)
+		{
+			SetChunkVisible(0, true);
+		}
+		return;
+	}
+#endif	// WITH_EDITOR
+
+	// Only create physics in the game
+	if( !World->IsGameWorld() )
+	{
+		return;
+	}
+
+	// Set template actor/body/shape properties
+
+	// Find the PhysicalMaterial we need to apply to the physics bodies.
+	UPhysicalMaterial* PhysMat = BodyInstance.GetSimplePhysicalMaterial();
+
+	// Get the default actor descriptor NxParameterized data from the asset
+	NxParameterized::Interface* ActorParams = TheDestructibleMesh->GetDestructibleActorDesc(PhysMat);
+
+	// Create PhysX transforms from ComponentToWorld
+	const PxMat44 GlobalPose(PxMat33(U2PQuat(ComponentToWorld.GetRotation())), U2PVector(ComponentToWorld.GetTranslation()));
+	const PxVec3 Scale = U2PVector(ComponentToWorld.GetScale3D());
+
+	// Set the transform in the actor descriptor
+	verify( NxParameterized::setParamMat44(*ActorParams,"globalPose",GlobalPose) );
+	verify( NxParameterized::setParamVec3(*ActorParams,"scale",Scale) );
+
+	// Set the (initially) dynamic flag in the actor descriptor
+	// See if we are 'static'
+	verify( NxParameterized::setParamBool(*ActorParams,"dynamic", BodyInstance.bSimulatePhysics != false) );
+
+	// Set the sleep velocity frame decay constant (was sleepVelocitySmoothingFactor) - a new feature that should help sleeping in large piles
+	verify( NxParameterized::setParamF32(*ActorParams,"sleepVelocityFrameDecayConstant", 20.0f) );
+
+	// Set up the shape desc template
+
+	// Get collision channel and response
+	PxFilterData PQueryFilterData, PSimFilterData;
+	uint8 MoveChannel = GetCollisionObjectType();
+	FCollisionResponseContainer CollResponse;
+	if(IsCollisionEnabled())
+	{
+		// Only enable a collision response if collision is enabled
+		CollResponse = GetCollisionResponseToChannels();
+
+		LargeChunkCollisionResponse.SetCollisionResponseContainer(CollResponse);
+		SmallChunkCollisionResponse.SetCollisionResponseContainer(CollResponse);
+		SmallChunkCollisionResponse.SetResponse(ECC_Pawn, ECR_Overlap);
+	}
+	else
+	{
+		// now since by default it will all block, if collision is disabled, we need to set to ignore
+		MoveChannel = ECC_WorldStatic;
+		CollResponse.SetAllChannels(ECR_Ignore);
+		LargeChunkCollisionResponse.SetAllChannels(ECR_Ignore);
+		SmallChunkCollisionResponse.SetAllChannels(ECR_Ignore);
+	}
+
+	const bool bEnableImpactDamage = IsImpactDamageEnabled(TheDestructibleMesh, 0);
+	const bool bEnableContactModification = TheDestructibleMesh->DefaultDestructibleParameters.DamageParameters.bCustomImpactResistance && TheDestructibleMesh->DefaultDestructibleParameters.DamageParameters.ImpactResistance > 0.f;
+
+	// Passing AssetInstanceID = 0 so we'll have self-collision
+	AActor* Owner = GetOwner();
+	CreateShapeFilterData(MoveChannel, GetUniqueID(), CollResponse, 0, 0, PQueryFilterData, PSimFilterData, BodyInstance.bUseCCD, bEnableImpactDamage, false, bEnableContactModification);
+
+	// Build filterData variations for complex and simple
+	PSimFilterData.word3 |= EPDF_SimpleCollision | EPDF_ComplexCollision;
+	PQueryFilterData.word3 |= EPDF_SimpleCollision | EPDF_ComplexCollision;
+
+	// Set the filterData in the shape descriptor
+	verify( NxParameterized::setParamU32(*ActorParams,"p3ShapeDescTemplate.simulationFilterData.word0", PSimFilterData.word0 ) );
+	verify( NxParameterized::setParamU32(*ActorParams,"p3ShapeDescTemplate.simulationFilterData.word1", PSimFilterData.word1 ) );
+	verify( NxParameterized::setParamU32(*ActorParams,"p3ShapeDescTemplate.simulationFilterData.word2", PSimFilterData.word2 ) );
+	verify( NxParameterized::setParamU32(*ActorParams,"p3ShapeDescTemplate.simulationFilterData.word3", PSimFilterData.word3 ) );
+	verify( NxParameterized::setParamU32(*ActorParams,"p3ShapeDescTemplate.queryFilterData.word0", PQueryFilterData.word0 ) );
+	verify( NxParameterized::setParamU32(*ActorParams,"p3ShapeDescTemplate.queryFilterData.word1", PQueryFilterData.word1 ) );
+	verify( NxParameterized::setParamU32(*ActorParams,"p3ShapeDescTemplate.queryFilterData.word2", PQueryFilterData.word2 ) );
+	verify( NxParameterized::setParamU32(*ActorParams,"p3ShapeDescTemplate.queryFilterData.word3", PQueryFilterData.word3 ) );
+
+	// Set the PhysX material in the shape descriptor
+	PxMaterial* PMaterial = PhysMat->GetPhysXMaterial();
+	verify( NxParameterized::setParamU64(*ActorParams,"p3ShapeDescTemplate.material", (physx::PxU64)PMaterial) );
+
+	// Set the rest depth to match the skin width in the shape descriptor
+	const physx::PxCookingParams& CookingParams = GApexSDK->getCookingInterface()->getParams();
+	verify( NxParameterized::setParamF32(*ActorParams,"p3ShapeDescTemplate.restOffset", -CookingParams.skinWidth) );
+
+	// Set the PhysX material in the actor descriptor
+	verify( NxParameterized::setParamBool(*ActorParams,"p3ActorDescTemplate.flags.eDISABLE_GRAVITY",false) );
+	verify( NxParameterized::setParamBool(*ActorParams,"p3ActorDescTemplate.flags.eVISUALIZATION",true) );
+
+	// Set the PxActor's and PxShape's userData fields to this component's body instance
+	verify( NxParameterized::setParamU64(*ActorParams,"p3ActorDescTemplate.userData", 0 ) );
+
+	// All shapes created by this DestructibleActor will have the userdata of the owning component.
+	// We need this, as in some cases APEX is moving shapes accross actors ( ex. FormExtended structures )
+	verify( NxParameterized::setParamU64(*ActorParams,"p3ShapeDescTemplate.userData", (PxU64)&PhysxUserData ) );
+
+	// Set up the body desc template in the actor descriptor
+	verify( NxParameterized::setParamF32(*ActorParams,"p3BodyDescTemplate.angularDamping", BodyInstance.AngularDamping ) );
+	verify( NxParameterized::setParamF32(*ActorParams,"p3BodyDescTemplate.linearDamping", BodyInstance.LinearDamping ) );
+	const PxTolerancesScale& PScale = GPhysXSDK->getTolerancesScale();
+	PxF32 SleepEnergyThreshold = 0.00005f*PScale.speed*PScale.speed;	// 1/1000 Default, since the speed scale is quite high
+	if (BodyInstance.SleepFamily == ESleepFamily::Sensitive)
+	{
+		SleepEnergyThreshold /= 20.0f;
+	}
+	verify( NxParameterized::setParamF32(*ActorParams,"p3BodyDescTemplate.sleepThreshold", SleepEnergyThreshold) );
+//	NxParameterized::setParamF32(*ActorParams,"bodyDescTemplate.sleepDamping", SleepDamping );
+	verify( NxParameterized::setParamF32(*ActorParams,"p3BodyDescTemplate.density", 0.001f*PhysMat->Density) );	// Convert from g/cm^3 to kg/cm^3
+	// Enable CCD if requested
+	verify( NxParameterized::setParamBool(*ActorParams,"p3BodyDescTemplate.flags.eENABLE_CCD", BodyInstance.bUseCCD != 0) );
+	// Ask the actor to create chunk events, for more efficient visibility updates
+	verify( NxParameterized::setParamBool(*ActorParams,"createChunkEvents", true) );
+
+	// Enable hard sleeping if requested
+	verify( NxParameterized::setParamBool(*ActorParams,"useHardSleeping", bEnableHardSleeping) );
+
+	
+
+	// Destructibles are always dynamic or kinematic, and therefore only go into one of the scenes
+	const uint32 SceneType = BodyInstance.UseAsyncScene() ? PST_Async : PST_Sync;
+	NxApexScene* ApexScene = PhysScene->GetApexScene(SceneType);
+	check(ApexScene);
+
+	// Create an APEX NxDestructibleActor from the Destructible asset and actor descriptor
+	ApexDestructibleActor = static_cast<NxDestructibleActor*>(TheDestructibleMesh->ApexDestructibleAsset->createApexActor(*ActorParams, *ApexScene));
+	check(ApexDestructibleActor);
+
+	// Make a backpointer to this component
+	PhysxUserData = FPhysxUserData(this);
+	ApexDestructibleActor->userData = &PhysxUserData;
+
+	ChunkInfos.Empty();
+	PhysxChunkUserData.Empty();
+
+	// Cache cooked collision data
+	// BRGTODO : cook in asset
+	ApexDestructibleActor->cacheModuleData();
+
+	// BRGTODO : Per-actor LOD setting
+//	ApexDestructibleActor->forcePhysicalLod( DestructibleActor->LOD );
+
+	// Start asleep if requested
+	PxRigidDynamic* PRootActor = ApexDestructibleActor->getChunkPhysXActor(0);
+
+
+	//  Put to sleep or wake up only if the component is physics-simulated
+	//TODO: APEX 1.3.3 defers adding actors to scene which means this doesn't work. Waiting on flag from them so we can do this properly.
+	/*if (PRootActor != NULL && BodyInstance.bSimulatePhysics)
+	{
+		SCOPED_SCENE_WRITE_LOCK(PRootActor->getScene());
+
+		PRootActor->setActorFlag(PxActorFlag::eDISABLE_GRAVITY, !BodyInstance.bEnableGravity);
+
+		// Sleep/wake up as appropriate
+		if (BodyInstance.bStartAwake)
+		{
+			PRootActor->wakeUp();
+		}
+		else
+		{
+			PRootActor->putToSleep();
+		}
+	}*/
+
+	UpdateBounds();
+#endif	// #if WITH_APEX
+}
+
+void UDestructibleComponent::DestroyPhysicsState()
+{
+#if WITH_APEX
+	if(ApexDestructibleActor != NULL)
+	{
+		if(UWorld * World = GetWorld())
+		{
+			if (FPhysScene * PhysScene = World->GetPhysicsScene())
+			{
+				const uint32 SceneType = BodyInstance.UseAsyncScene() ? PST_Async : PST_Sync;
+				PxScene * PScene = PhysScene->GetPhysXScene(SceneType);
+				SCOPED_SCENE_WRITE_LOCK(PScene);
+				ApexDestructibleActor->release();
+				//Deferring here is difficult because a call to CreatePhysicsState will re-use the same chunk info buffer for new chunks. Note this may be possible, but need to fix crash now
+			}
+		}
+		
+		ApexDestructibleActor = NULL;
+		
+		//Destructible component uses the BodyInstance in PrimitiveComponent in a very dangerous way. It assigns PxRigidDynamic to it as it needs it.
+		//Destructible PxRigidDynamic actors can be deleted from under us as PhysX sees fit.
+		//Ideally we wouldn't ever have a dangling pointer, but in practice this is hard to avoid.
+		//In theory anyone using BodyInstance on a PrimitiveComponent should be using functions like GetBodyInstance - in which case we properly fix up the dangling pointer
+		BodyInstance.RigidActorSync = NULL;
+		BodyInstance.RigidActorAsync = NULL;
+	}
+#endif	// #if WITH_APEX
+	Super::DestroyPhysicsState();
+}
+
+UBodySetup* UDestructibleComponent::GetBodySetup()
+{
+	if (SkeletalMesh != NULL)
+	{
+		UDestructibleMesh* TheDestructibleMesh = GetDestructibleMesh();
+
+		if (TheDestructibleMesh != NULL)
+		{
+			return TheDestructibleMesh->BodySetup;
+		}
+	}
+
+	return NULL;
+}
+
+bool UDestructibleComponent::CanEditSimulatePhysics()
+{
+	// if destructiblemeshcomponent, we will allow it always
+	return true;
+}
+
+void UDestructibleComponent::AddImpulse( FVector Impulse, FName BoneName /*= NAME_None*/, bool bVelChange /*= false*/ )
+{
+#if WITH_APEX
+	if ( ApexDestructibleActor )
+	{
+		int32 ChunkIdx = BoneIdxToChunkIdx(GetBoneIndex(BoneName));
+		PxRigidDynamic* PActor = ApexDestructibleActor->getChunkPhysXActor(ChunkIdx);
+
+		if ( PActor != NULL )
+		{
+			SCOPED_SCENE_WRITE_LOCK(PActor->getScene());
+
+			PActor->addForce(U2PVector(Impulse), bVelChange ? PxForceMode::eVELOCITY_CHANGE : PxForceMode::eIMPULSE);
+		}
+	}
+#endif
+}
+
+
+void UDestructibleComponent::AddImpulseAtLocation( FVector Impulse, FVector Position, FName BoneName /*= NAME_None*/ )
+{
+#if WITH_APEX
+	if (ApexDestructibleActor)
+	{
+		int32 ChunkIdx = NxModuleDestructibleConst::INVALID_CHUNK_INDEX;
+		if (BoneName != NAME_None)
+		{
+			ChunkIdx = BoneIdxToChunkIdx(GetBoneIndex(BoneName));
+		}
+
+		//ApexDestructibleActor->applyDamage(1.0f, Impulse.Size(), U2PVector(Position), U2PVector(Impulse.SafeNormal()), ChunkIdx);
+		for (int32 i=0; i < ChunkInfos.Num(); ++i)
+		{
+			if (ChunkInfos[i].ChunkIndex == ChunkIdx && ChunkInfos[i].Actor)
+			{
+				ChunkInfos[i].Actor->addForce(U2PVector(Impulse), PxForceMode::eIMPULSE);
+			}
+		}
+	}
+#endif
+}
+
+void UDestructibleComponent::AddForce( FVector Force, FName BoneName /*= NAME_None*/, bool bAccelChange /* = false */ )
+{
+#if WITH_APEX
+	int32 ChunkIdx = NxModuleDestructibleConst::INVALID_CHUNK_INDEX;
+	if (BoneName != NAME_None)
+	{
+		ChunkIdx = BoneIdxToChunkIdx(GetBoneIndex(BoneName));
+	}
+
+	//ApexDestructibleActor->applyDamage(1.0f, Impulse.Size(), U2PVector(Position), U2PVector(Impulse.SafeNormal()), ChunkIdx);
+	for (int32 i=0; i < ChunkInfos.Num(); ++i)
+	{
+		if (ChunkInfos[i].ChunkIndex == ChunkIdx && ChunkInfos[i].Actor)
+		{
+			ChunkInfos[i].Actor->addForce(U2PVector(Force), bAccelChange ? PxForceMode::eACCELERATION : PxForceMode::eFORCE);
+		}
+	}
+#endif
+}
+
+void UDestructibleComponent::AddForceAtLocation( FVector Force, FVector Location, FName BoneName /*= NAME_None*/ )
+{
+#if WITH_APEX
+	if (ApexDestructibleActor)
+	{
+		int32 ChunkIdx = NxModuleDestructibleConst::INVALID_CHUNK_INDEX;
+		if (BoneName != NAME_None)
+		{
+			ChunkIdx = BoneIdxToChunkIdx(GetBoneIndex(BoneName));
+		}
+
+		//ApexDestructibleActor->applyDamage(1.0f, Impulse.Size(), U2PVector(Location), U2PVector(Impulse.SafeNormal()), ChunkIdx);
+		for (int32 i=0; i < ChunkInfos.Num(); ++i)
+		{
+			if (ChunkInfos[i].ChunkIndex == ChunkIdx && ChunkInfos[i].Actor)
+			{
+				ChunkInfos[i].Actor->addForce(U2PVector(Force), PxForceMode::eFORCE);
+			}
+		}
+	}
+#endif
+}
+
+void UDestructibleComponent::AddRadialImpulse(FVector Origin, float Radius, float Strength, ERadialImpulseFalloff Falloff, bool bVelChange)
+{
+#if WITH_APEX
+	if (bIgnoreRadialImpulse)
+	{
+		return;
+	}
+
+	if (ApexDestructibleActor == NULL)
+	{
+		return;
+	}
+
+	PxRigidDynamic** PActorBuffer = NULL;
+	PxU32 PActorCount = 0;
+	if (ApexDestructibleActor->acquirePhysXActorBuffer(PActorBuffer, PActorCount, NxDestructiblePhysXActorQueryFlags::Dynamic))
+	{
+		PxScene* LockedScene = NULL;
+
+		while (PActorCount--)
+		{
+			PxRigidDynamic* PActor = *PActorBuffer++;
+			if (PActor != NULL)
+			{
+				if (!LockedScene)
+				{
+					LockedScene = PActor->getScene();
+					LockedScene->lockWrite();
+					LockedScene->lockRead();
+				}
+				AddRadialImpulseToPxRigidBody(*PActor, Origin, Radius, Strength, Falloff, bVelChange);
+			}
+		}
+
+		if (LockedScene)
+		{
+			LockedScene->unlockRead();
+			LockedScene->unlockWrite();
+			LockedScene = NULL;
+		}
+
+		ApexDestructibleActor->releasePhysXActorBuffer();
+	}
+#endif	// #if WITH_APEX
+}
+
+void UDestructibleComponent::AddRadialForce(FVector Origin, float Radius, float Strength, ERadialImpulseFalloff Falloff, bool bAccelChange /* = false */)
+{
+#if WITH_APEX
+	if(bIgnoreRadialForce)
+	{
+		return;
+	}
+
+	PxRigidDynamic** PActorBuffer = NULL;
+	PxU32 PActorCount = 0;
+	if (ApexDestructibleActor->acquirePhysXActorBuffer(PActorBuffer, PActorCount, NxDestructiblePhysXActorQueryFlags::Dynamic))
+	{
+		PxScene* LockedScene = NULL;
+		
+
+		while (PActorCount--)
+		{
+			PxRigidDynamic* PActor = *PActorBuffer++;
+			if (PActor != NULL)
+			{
+				if (!LockedScene)
+				{
+					LockedScene = PActor->getScene();
+					LockedScene->lockWrite();
+					LockedScene->lockRead();
+				}
+
+				AddRadialForceToPxRigidBody(*PActor, Origin, Radius, Strength, Falloff, bAccelChange);
+			}
+
+			if (LockedScene)
+			{
+				LockedScene->unlockRead();
+				LockedScene->unlockWrite();
+				LockedScene = NULL;
+			}
+		}
+		ApexDestructibleActor->releasePhysXActorBuffer();
+	}
+#endif	// #if WITH_APEX
+}
+
+void UDestructibleComponent::ReceiveComponentDamage(float DamageAmount, FDamageEvent const& DamageEvent, AController* EventInstigator, AActor* DamageCauser)
+{
+	UDamageType const* const DamageTypeCDO = DamageEvent.DamageTypeClass ? DamageEvent.DamageTypeClass->GetDefaultObject<UDamageType>() : GetDefault<UDamageType>();
+	if (DamageEvent.IsOfType(FPointDamageEvent::ClassID))
+	{
+		FPointDamageEvent const* const PointDamageEvent = (FPointDamageEvent*)(&DamageEvent);
+		ApplyDamage(DamageAmount, PointDamageEvent->HitInfo.ImpactPoint, PointDamageEvent->ShotDirection, DamageTypeCDO->DestructibleImpulse);
+	}
+	else if (DamageEvent.IsOfType(FRadialDamageEvent::ClassID))
+	{
+		FRadialDamageEvent const* const RadialDamageEvent = (FRadialDamageEvent*)(&DamageEvent);
+		ApplyRadiusDamage(DamageAmount, RadialDamageEvent->Origin, RadialDamageEvent->Params.OuterRadius, DamageTypeCDO->DestructibleImpulse, false);
+	}
+}
+
+#if WITH_APEX
+void UDestructibleComponent::SpawnFractureEffectsFromDamageEvent(const NxApexDamageEventReportData& InDamageEvent)
+{
+	// Use the component's fracture effects if the override is selected, otherwise use fracture effects from the asset
+	TArray<FFractureEffect>& UseFractureEffects = (bFractureEffectOverride || !SkeletalMesh) ? FractureEffects : CastChecked<UDestructibleMesh>(SkeletalMesh)->FractureEffects;
+
+	UDestructibleMesh* TheDestructibleMesh = GetDestructibleMesh();
+	if (!TheDestructibleMesh)
+	{
+		return;
+	}
+
+
+	// We keep track of the handled parent chunks here
+	TArray<int32> HandledParents;
+	for(physx::PxU32 eventN = 0; eventN < InDamageEvent.fractureEventListSize; ++eventN)
+	{
+		const NxApexChunkData& chunkData = InDamageEvent.fractureEventList[eventN];
+		if( chunkData.depth < (physx::PxU32)UseFractureEffects.Num() )
+		{
+			// We can get the root chunk here as well, so make sure that the parent index is 0, even for the root chunk
+			int32 ParentIdx = FMath::Max(TheDestructibleMesh->ApexDestructibleAsset->getChunkParentIndex(chunkData.index), 0);
+			
+			// We can test a number of flags - we'll play an effect if the chunk was destroyed
+			// As we only get the fractured event here for chunks that come free, we spawn fracture
+			// effects only once per unique parent
+			if((chunkData.flags & NxApexChunkFlag::FRACTURED) && !HandledParents.Contains(ParentIdx))
+			{
+				FVector Position = P2UVector(chunkData.worldBounds.getCenter());
+				FFractureEffect& FractureEffect = UseFractureEffects[chunkData.depth];
+				if( FractureEffect.Sound != NULL )
+				{
+					// Spawn sound
+					UGameplayStatics::PlaySoundAtLocation( this, FractureEffect.Sound, Position );
+				}
+				if( FractureEffect.ParticleSystem != NULL )
+				{
+					// Spawn particle system
+					UParticleSystemComponent* ParticleSystemComponent = UGameplayStatics::SpawnEmitterAtLocation( this, FractureEffect.ParticleSystem, Position );
+
+					// Disable shadows, since destructibles tend to generate a lot of these
+					if (ParticleSystemComponent != NULL)
+					{
+						ParticleSystemComponent->CastShadow = false;
+					}
+				}
+
+				HandledParents.Add(ParentIdx);
+			}			
+		}
+	}
+}
+
+void UDestructibleComponent::OnDamageEvent(const NxApexDamageEventReportData& InDamageEvent)
+{
+	FVector HitPosition = P2UVector(InDamageEvent.hitPosition);
+	FVector HitDirection = P2UVector(InDamageEvent.hitDirection);
+
+	OnComponentFracture.Broadcast(HitPosition, HitDirection);
+	if (ADestructibleActor * DestructibleActor = Cast<ADestructibleActor>(GetOwner()))
+	{
+		DestructibleActor->OnActorFracture.Broadcast(HitPosition, HitDirection);
+	}
+
+	SpawnFractureEffectsFromDamageEvent(InDamageEvent);
+
+	// After receiving damage, no longer receive decals.
+	if (bReceivesDecals)
+	{
+		bReceivesDecals = false;
+		MarkRenderStateDirty();
+	}
+}
+
+void UDestructibleComponent::OnVisibilityEvent(const NxApexChunkStateEventData & InVisibilityEvent)
+{
+	for (uint32 EventIndex = 0; EventIndex < InVisibilityEvent.stateEventListSize; ++EventIndex)
+	{
+		const NxDestructibleChunkEvent &  Event = InVisibilityEvent.stateEventList[EventIndex];
+		// Right now the only events are visibility changes.  So as an optimization we won't check for the event type.
+		//				if (Event.event & physx::NxDestructibleChunkEvent::VisibilityChanged)
+		const bool bVisible = (Event.event & physx::NxDestructibleChunkEvent::ChunkVisible) != 0;
+		SetChunkVisible(Event.chunkIndex, bVisible);
+	}
+}
+#endif // WITH_APEX
+
+bool UDestructibleComponent::IsFracturedOrInitiallyStatic() const
+{
+	bool bInitiallyStatic = false;
+	bool bFractured = false;
+#if WITH_APEX
+	if (ApexDestructibleActor == nullptr)
+	{
+		return false;
+	}
+
+	bInitiallyStatic = !ApexDestructibleActor->isInitiallyDynamic();
+
+	if (bInitiallyStatic == false)
+	{
+		//If we have only one chunk and its index is 0 we are NOT fractured. Otherwise we must have fractured
+		const physx::PxU32 VisibleChunkCount = ApexDestructibleActor->getNumVisibleChunks();
+		if (VisibleChunkCount == 1)
+		{
+			const physx::PxU16* VisibleChunks = ApexDestructibleActor->getVisibleChunks();
+			bFractured = *VisibleChunks != 0;
+		}
+		else
+		{
+			bFractured = true;
+		}
+	}
+#endif
+	return bFractured || bInitiallyStatic;
+}
+
+void UDestructibleComponent::RefreshBoneTransforms(FActorComponentTickFunction* TickFunction)
+{
+}
+
+void UDestructibleComponent::SetDestructibleMesh(class UDestructibleMesh* NewMesh)
+{
+#if WITH_APEX
+	ChunkInfos.Empty();
+	PhysxChunkUserData.Empty();
+#endif // WITH_APEX
+
+	Super::SetSkeletalMesh( NewMesh );
+
+#if WITH_EDITORONLY_DATA
+	// If the SkeletalMesh has changed, update our transient value too.
+	this->DestructibleMesh = GetDestructibleMesh();
+#endif // WITH_EDITORONLY_DATA
+	
+	RecreatePhysicsState();
+}
+
+class UDestructibleMesh * UDestructibleComponent::GetDestructibleMesh()
+{
+	return Cast<UDestructibleMesh>(SkeletalMesh);
+}
+
+void UDestructibleComponent::SetSkeletalMesh( USkeletalMesh* InSkelMesh )
+{
+	if(InSkelMesh != NULL && !InSkelMesh->IsA(UDestructibleMesh::StaticClass()))
+	{
+		// Issue warning and do nothing if this is not actually a UDestructibleMesh
+		UE_LOG(LogPhysics, Log, TEXT("UDestructibleComponent::SetSkeletalMesh(): Passed-in USkeletalMesh (%s) must be a UDestructibleMesh.  SkeletalMesh not set."), *InSkelMesh->GetPathName() );
+		return;
+	}
+
+	UDestructibleMesh* TheDestructibleMesh = static_cast<UDestructibleMesh*>(InSkelMesh);
+#if WITH_APEX
+	if(TheDestructibleMesh != NULL && TheDestructibleMesh->ApexDestructibleAsset == NULL)
+	{
+		UE_LOG(LogPhysics, Log, TEXT("UDestructibleComponent::SetSkeletalMesh(): Missing ApexDestructibleAsset on '%s'."), *InSkelMesh->GetPathName() );
+		return;
+	}
+
+	SetDestructibleMesh(TheDestructibleMesh);
+
+	if(TheDestructibleMesh != NULL)
+	{
+		// Resize the fracture effects array to the appropriate size
+		FractureEffects.AddZeroed(TheDestructibleMesh->ApexDestructibleAsset->getDepthCount());
+	}
+#else
+	SetDestructibleMesh(TheDestructibleMesh);
+#endif // WITH_APEX
+}
+
+FTransform UDestructibleComponent::GetSocketTransform(FName InSocketName, ERelativeTransformSpace TransformSpace) const
+{
+	FTransform ST = Super::GetSocketTransform(InSocketName, TransformSpace);
+
+	int32 BoneIdx = GetBoneIndex(InSocketName);
+
+	// As bones in a destructible might be scaled to 0 when hidden, we force a scale of 1 if we want the socket transform
+	if (BoneIdx > 0 && IsBoneHidden(BoneIdx))
+	{
+		ST.SetScale3D(FVector(1.0f, 1.0f, 1.0f));
+	}
+
+	return ST;
+}
+
+void UDestructibleComponent::SetChunkVisible( int32 ChunkIndex, bool bVisible )
+{
+	// Bone 0 is a dummy root bone
+	const int32 BoneIndex = ChunkIdxToBoneIdx(ChunkIndex);
+	bool bClearActorFromChunkInfo = false;
+
+	if( bVisible )
+	{
+		UnHideBone(BoneIndex);
+#if WITH_APEX
+		PxRigidDynamic* PActor = ApexDestructibleActor != NULL ? ApexDestructibleActor->getChunkPhysXActor(ChunkIndex) : NULL;
+
+		UDestructibleMesh* DMesh = GetDestructibleMesh();
+
+		if (PActor != NULL)
+		{
+			//When we unhide a chunk, it might be the first time - in this case we need to update the transforms (only active transforms get updated after this)
+			//TODO: We can probably move this code somewhere so that it's only done the first time a chunk is created in case chunks are hidden unhidden again and again? (Not sure if this ever happens)
+			const physx::PxMat44 ChunkPoseRT = ApexDestructibleActor->getChunkPose(ChunkIndex);	// Unscaled
+			const physx::PxTransform Transform(ChunkPoseRT);
+			SetChunkWorldRT(ChunkIndex, P2UQuat(Transform.q), P2UVector(Transform.p));
+
+			// If actor has already a chunk info and userdata, we just make sure it is valid and update the 
+			// physx actor if needed. We do NOT do this for FormExtended structures, as in this case, the shapes/actors
+			// are moved to the 1st structure object internally by APEX.
+			if(PActor->userData != NULL && !DMesh->DefaultDestructibleParameters.Flags.bFormExtendedStructures)
+			{
+				FDestructibleChunkInfo* CI = FPhysxUserData::Get<FDestructibleChunkInfo>(PActor->userData);
+				checkf(CI, TEXT("If a chunk actor has user data and it is not a DestructibleChunkInfo, something is messed up."));
+				//check(CI->OwningComponent == this);
+
+				if (CI->ChunkIndex != ChunkIndex)
+				{
+					// grab the old actor and clear its user data, as we steal the ChunkInfo here
+					if (CI->Actor && CI->Actor != PActor)
+					{
+						CI->Actor->userData = NULL;
+					}
+					CI->ChunkIndex = ChunkIndex;
+				}
+
+				CI->OwningComponent = this;
+				CI->Actor = PActor;
+			}
+			else if (PActor->userData == NULL)
+			{
+				// Setup the user data to have a proper chunk - actor mapping 
+				FDestructibleChunkInfo* CI = new FDestructibleChunkInfo();
+				int32 InfoIndex = ChunkInfos.Add(CI);
+				CI->Index = InfoIndex;
+				CI->ChunkIndex = ChunkIndex;
+				CI->OwningComponent = this;
+				CI->Actor = PActor;
+
+				PActor->setActorFlag(PxActorFlag::eDISABLE_GRAVITY, !BodyInstance.bEnableGravity);
+
+				FPhysxUserData * UserData = new FPhysxUserData(CI);
+				int32 UserDataIdx = PhysxChunkUserData.Add(UserData);
+				check(InfoIndex == UserDataIdx);
+
+				PActor->userData = UserData;
+
+
+
+				// Set collision response to non-root chunks
+				if (GetDestructibleMesh()->ApexDestructibleAsset->getChunkParentIndex(ChunkIndex) >= 0)
+				{
+					SetCollisionResponseForActor(PActor, ChunkIndex);
+				}
+			}
+		}
+		else
+		{
+			bClearActorFromChunkInfo = true;
+		}
+#endif // WITH_APEX
+	}
+	else
+	{
+		HideBone(BoneIndex, PBO_None);
+		bClearActorFromChunkInfo = true;
+	}
+
+#if WITH_APEX
+	if (bClearActorFromChunkInfo)
+	{
+		// Make sure we clear the physx actor pointer of the chunk info as it might (and probably will) be
+		// invalid from now on
+		for (int32 i=0; i < ChunkInfos.Num(); ++i)
+		{
+			if (ChunkInfos[i].ChunkIndex == ChunkIndex)
+			{
+				ChunkInfos[i].Actor = NULL;
+				break;
+			}
+		}
+	}
+#endif // WITH_APEX
+
+	// Mark the transform as dirty, so the bounds are updated and sent to the render thread
+	MarkRenderTransformDirty();
+
+	// New bone positions need to be sent to render thread
+	MarkRenderDynamicDataDirty();
+}
+
+#if WITH_APEX
+void UDestructibleComponent::UpdateDestructibleChunkTM(const TArray<const PxRigidActor*>& ActiveActors)
+{
+	//We want to consolidate the transforms so that we update each destructible component once by passing it an array of chunks to update.
+	//This helps avoid a lot of duplicated work like marking render dirty, computing inverse world component, etc...
+
+	TMap<UDestructibleComponent*, TArray<FUpdateChunksInfo> > ComponentUpdateMapping;
+	
+	//prepare map to update destructible components
+	TArray<PxShape*> Shapes;
+	for (const PxRigidActor* RigidActor : ActiveActors)
+	{
+		if (const FDestructibleChunkInfo* DestructibleChunkInfo = FPhysxUserData::Get<FDestructibleChunkInfo>(RigidActor->userData))
+		{
+			if (GApexModuleDestructible->owns(RigidActor) && DestructibleChunkInfo->OwningComponent.IsValid())
+			{
+				Shapes.AddUninitialized(RigidActor->getNbShapes());
+				int32 NumShapes = RigidActor->getShapes(Shapes.GetData(), Shapes.Num());
+				for (int32 ShapeIdx = 0; ShapeIdx < Shapes.Num(); ++ShapeIdx)
+				{
+					PxShape* Shape = Shapes[ShapeIdx];
+					int32 ChunkIndex;
+					if (NxDestructibleActor* DestructibleActor = GApexModuleDestructible->getDestructibleAndChunk(Shape, &ChunkIndex))
+					{
+						const physx::PxMat44 ChunkPoseRT = DestructibleActor->getChunkPose(ChunkIndex);
+						const physx::PxTransform Transform(ChunkPoseRT);
+						if (UDestructibleComponent* DestructibleComponent = Cast<UDestructibleComponent>(FPhysxUserData::Get<UPrimitiveComponent>(DestructibleActor->userData)))
+						{
+							if (DestructibleComponent->IsRegistered())
+							{
+								TArray<FUpdateChunksInfo>& UpdateInfos = ComponentUpdateMapping.FindOrAdd(DestructibleComponent);
+								FUpdateChunksInfo* UpdateInfo = new (UpdateInfos)FUpdateChunksInfo(ChunkIndex, P2UTransform(Transform));
+							}
+						}
+					}
+				}
+
+				Shapes.Empty(Shapes.Num());	//we want to keep largest capacity array to avoid reallocs
+			}
+		}
+	}
+	
+	//update each component
+	for (auto It = ComponentUpdateMapping.CreateIterator(); It; ++It)
+	{
+		UDestructibleComponent* DestructibleComponent = It.Key();
+		TArray<FUpdateChunksInfo>& UpdateInfos = It.Value();
+		if (DestructibleComponent->IsFracturedOrInitiallyStatic())
+		{
+			DestructibleComponent->SetChunksWorldTM(UpdateInfos);
+		}
+		else
+		{
+			//if we haven't fractured it must mean that we're simulating a destructible and so we should update our ComponentToWorld based on the single rigid body
+			DestructibleComponent->SyncComponentToRBPhysics();
+		}
+	}
+
+}
+
+#endif
+
+void UDestructibleComponent::SetChunksWorldTM(const TArray<FUpdateChunksInfo>& UpdateInfos)
+{
+	const FQuat InvRotation = ComponentToWorld.GetRotation().Inverse();
+
+	for (const FUpdateChunksInfo& UpdateInfo : UpdateInfos)
+	{
+		// Bone 0 is a dummy root bone
+		const int32 BoneIndex = ChunkIdxToBoneIdx(UpdateInfo.ChunkIndex);
+		const FVector WorldTranslation	= UpdateInfo.WorldTM.GetLocation();
+		const FQuat WorldRotation		= UpdateInfo.WorldTM.GetRotation();
+
+		const FQuat BoneRotation = InvRotation*WorldRotation;
+		const FVector BoneTranslation = InvRotation.RotateVector(WorldTranslation - ComponentToWorld.GetTranslation()) / ComponentToWorld.GetScale3D();
+
+		GetEditableSpaceBases()[BoneIndex] = FTransform(BoneRotation, BoneTranslation);
+	}
+
+	// Mark the transform as dirty, so the bounds are updated and sent to the render thread
+	MarkRenderTransformDirty();
+
+	// New bone positions need to be sent to render thread
+	MarkRenderDynamicDataDirty();
+
+	//Update bone visibilty and flip the editable space base buffer
+	FlipEditableSpaceBases();
+}
+
+void UDestructibleComponent::SetChunkWorldRT( int32 ChunkIndex, const FQuat& WorldRotation, const FVector& WorldTranslation )
+{
+	// Bone 0 is a dummy root bone
+	const int32 BoneIndex = ChunkIdxToBoneIdx(ChunkIndex);
+
+	// Mark the transform as dirty, so the bounds are updated and sent to the render thread
+	MarkRenderTransformDirty();
+
+	// New bone positions need to be sent to render thread
+	MarkRenderDynamicDataDirty();
+
+#if 0
+	// Scale is already applied to the ComponentToWorld transform, and is carried into the bones _locally_.
+	// So there is no need to set scale in the bone local transforms
+	const FTransform WorldRT(WorldRotation, WorldTranslation, ComponentToWorld.GetScale3D());
+	SpaceBases(BoneIndex) = WorldRT*ComponentToWorld.Inverse();
+#elif 1
+	// More optimal form of the above
+	const FQuat BoneRotation = ComponentToWorld.GetRotation().Inverse()*WorldRotation;
+	const FVector BoneTranslation = ComponentToWorld.GetRotation().Inverse().RotateVector(WorldTranslation - ComponentToWorld.GetTranslation())/ComponentToWorld.GetScale3D();
+	GetEditableSpaceBases()[BoneIndex] = FTransform(BoneRotation, BoneTranslation);
+#endif
+}
+
+void UDestructibleComponent::ApplyDamage(float DamageAmount, const FVector& HitLocation, const FVector& ImpulseDir, float ImpulseStrength)
+{
+#if WITH_APEX
+	if (ApexDestructibleActor != NULL)
+	{
+		const FVector& NormalizedImpactDir = ImpulseDir.GetSafeNormal();
+
+		// Transfer damage information to the APEX NxDestructibleActor interface
+		ApexDestructibleActor->applyDamage(DamageAmount, ImpulseStrength, U2PVector( HitLocation ), U2PVector( ImpulseDir ));
+	}
+#endif
+}
+
+void UDestructibleComponent::ApplyRadiusDamage(float BaseDamage, const FVector& HurtOrigin, float DamageRadius, float ImpulseStrength, bool bFullDamage)
+{
+#if WITH_APEX
+	if (ApexDestructibleActor != NULL)
+	{
+		// Transfer damage information to the APEX NxDestructibleActor interface
+		ApexDestructibleActor->applyRadiusDamage(BaseDamage, ImpulseStrength, U2PVector( HurtOrigin ), DamageRadius, bFullDamage ? false : true);
+	}
+#endif
+}
+
+bool UDestructibleComponent::DoCustomNavigableGeometryExport(FNavigableGeometryExport& GeomExport) const
+{
+#if WITH_APEX
+	if (ApexDestructibleActor == NULL)
+	{
+		return false;
+	}
+
+	NxDestructibleActor* DestrActor = const_cast<NxDestructibleActor*>(ApexDestructibleActor);
+
+	const FTransform ComponentToWorldNoScale(ComponentToWorld.GetRotation(), ComponentToWorld.GetTranslation(), FVector(1.f));
+	TArray<PxShape*> Shapes;
+	Shapes.AddUninitialized(8);
+	PxRigidDynamic** PActorBuffer = NULL;
+	PxU32 PActorCount = 0;
+	if (DestrActor->acquirePhysXActorBuffer(PActorBuffer, PActorCount
+		, NxDestructiblePhysXActorQueryFlags::Static 
+		| NxDestructiblePhysXActorQueryFlags::Dormant
+		| NxDestructiblePhysXActorQueryFlags::Dynamic))
+	{
+		uint32 ShapesExportedCount = 0;
+
+		while (PActorCount--)
+		{
+			const PxRigidDynamic* PActor = *PActorBuffer++;
+			if (PActor != NULL)
+			{
+				const FTransform PActorGlobalPose = P2UTransform(PActor->getGlobalPose());
+
+				const PxU32 ShapesCount = PActor->getNbShapes();
+				if (ShapesCount > PxU32(Shapes.Num()))
+				{
+					Shapes.AddUninitialized(ShapesCount - Shapes.Num());
+				}
+				const PxU32 RetrievedShapesCount = PActor->getShapes(Shapes.GetData(), Shapes.Num());
+				PxShape* const* ShapePtr = Shapes.GetData();
+				for (PxU32 ShapeIndex = 0; ShapeIndex < RetrievedShapesCount; ++ShapeIndex, ++ShapePtr)
+				{
+					if (*ShapePtr != NULL)
+					{
+						const PxTransform LocalPose = (*ShapePtr)->getLocalPose();
+						FTransform LocalToWorld = P2UTransform(LocalPose);
+						LocalToWorld.Accumulate(PActorGlobalPose);
+
+						switch((*ShapePtr)->getGeometryType())
+						{
+						case PxGeometryType::eCONVEXMESH:
+							{
+								PxConvexMeshGeometry Geometry;
+								if ((*ShapePtr)->getConvexMeshGeometry(Geometry))
+								{
+									++ShapesExportedCount;
+
+									// @todo address Geometry.scale not being used here
+									GeomExport.ExportPxConvexMesh(Geometry.convexMesh, LocalToWorld);
+								}
+							}
+							break;
+						case PxGeometryType::eTRIANGLEMESH:
+							{
+								// @todo address Geometry.scale not being used here
+								PxTriangleMeshGeometry Geometry;
+								if ((*ShapePtr)->getTriangleMeshGeometry(Geometry))
+								{
+									++ShapesExportedCount;
+
+									if ((Geometry.triangleMesh->getTriangleMeshFlags()) & PxTriangleMeshFlag::eHAS_16BIT_TRIANGLE_INDICES)
+									{
+										GeomExport.ExportPxTriMesh16Bit(Geometry.triangleMesh, LocalToWorld);
+									}
+									else
+									{
+										GeomExport.ExportPxTriMesh32Bit(Geometry.triangleMesh, LocalToWorld);
+									}
+								}
+							}
+						default:
+							{
+								UE_LOG(LogPhysics, Log, TEXT("UDestructibleComponent::DoCustomNavigableGeometryExport(): unhandled PxGeometryType, %d.")
+									, int32((*ShapePtr)->getGeometryType()));
+							}
+							break;
+						}
+					}
+				}
+			}
+		}
+		ApexDestructibleActor->releasePhysXActorBuffer();
+
+		INC_DWORD_STAT_BY(STAT_Navigation_DestructiblesShapesExported, ShapesExportedCount);
+	}
+#endif // WITH_APEX
+
+	// we don't want a regular geometry export
+	return false;
+}
+
+void UDestructibleComponent::Activate( bool bReset/*=false*/ )
+{
+	if (bReset || ShouldActivate()==true)
+	{
+		bIsActive = true;
+	}
+}
+
+void UDestructibleComponent::Deactivate()
+{
+	if (ShouldActivate()==false)
+	{
+		bIsActive = false;
+	}
+}
+
+bool UDestructibleComponent::ShouldUpdateTransform(bool bLODHasChanged) const
+{
+	// We do not want to update bone transforms before physics has finished
+	return false;
+}
+
+bool UDestructibleComponent::LineTraceComponent( FHitResult& OutHit, const FVector Start, const FVector End, const FCollisionQueryParams& Params )
+{
+	bool bHaveHit = false;
+#if WITH_APEX
+	if (ApexDestructibleActor != NULL)
+	{
+		PxF32 HitTime = 0.0f;
+		PxVec3 HitNormal;
+		
+		int32 ChunkIdx = ApexDestructibleActor->rayCast(HitTime, HitNormal, U2PVector(Start), U2PVector(End-Start), NxDestructibleActorRaycastFlags::AllChunks);
+
+		if (ChunkIdx != NxModuleDestructibleConst::INVALID_CHUNK_INDEX && HitTime <= 1.0f)
+		{
+			PxRigidDynamic* PActor = ApexDestructibleActor->getChunkPhysXActor(ChunkIdx);
+
+			if (PActor != NULL)
+			{
+				// Store body instance state
+				FFakeBodyInstanceState PrevState;
+				SetupFakeBodyInstance(PActor, ChunkIdx, &PrevState);
+				
+				bHaveHit = Super::LineTraceComponent(OutHit, Start, End, Params);
+
+				// Reset original body instance
+				ResetFakeBodyInstance(PrevState);
+			}
+		}
+	}
+#endif
+	return bHaveHit;
+}
+
+bool UDestructibleComponent::SweepComponent(FHitResult& OutHit, const FVector Start, const FVector End, const FCollisionShape &CollisionShape, bool bTraceComplex/*=false*/)
+{
+	bool bHaveHit = false;
+#if WITH_APEX
+	if (ApexDestructibleActor != NULL)
+	{
+		PxF32 HitTime = 0.0f;
+		PxVec3 HitNormal;
+
+		int32 ChunkIdx = ApexDestructibleActor->obbSweep(HitTime, HitNormal, U2PVector(Start), U2PVector(CollisionShape.GetExtent()), PxMat33::createIdentity(), U2PVector(End - Start), NxDestructibleActorRaycastFlags::AllChunks);
+
+		if (ChunkIdx != NxModuleDestructibleConst::INVALID_CHUNK_INDEX && HitTime <= 1.0f)
+		{
+			PxRigidDynamic* PActor = ApexDestructibleActor->getChunkPhysXActor(ChunkIdx);
+
+			if (PActor != NULL)
+			{
+				// Store body instance state
+				FFakeBodyInstanceState PrevState;
+				SetupFakeBodyInstance(PActor, ChunkIdx, &PrevState);
+
+				bHaveHit = Super::SweepComponent(OutHit, Start, End, CollisionShape, bTraceComplex);
+
+				// Reset original body instance
+				ResetFakeBodyInstance(PrevState);
+			}
+		}
+	}
+#endif
+	return bHaveHit;
+}
+
+#if WITH_APEX
+void UDestructibleComponent::SetupFakeBodyInstance( physx::PxRigidActor* NewRigidActor, int32 InstanceIdx, FFakeBodyInstanceState* PrevState )
+{
+	//This code is very dangerous, but at the moment I have no better solution:
+	//Destructible component assigns PxRigidDynamic to the BodyInstance as it needs it.
+	//Destructible PxRigidDynamic actors can be deleted from under us as PhysX sees fit.
+	//Ideally we wouldn't ever have a dangling pointer, but in practice this is hard to avoid.
+	//In theory anyone using BodyInstance on a PrimitiveComponent should be using functions like GetBodyInstance - in which case we properly fix up the dangling pointer
+
+	if (PrevState != NULL)
+	{
+		PrevState->ActorSync = BodyInstance.RigidActorSync;
+		PrevState->ActorAsync = BodyInstance.RigidActorAsync;
+		PrevState->InstanceIndex = BodyInstance.InstanceBodyIndex;
+	}
+
+	BodyInstance.RigidActorSync = BodyInstance.UseAsyncScene() ? NULL : NewRigidActor;
+	BodyInstance.RigidActorAsync = BodyInstance.UseAsyncScene() ? NewRigidActor : NULL;
+	BodyInstance.BodyAggregate = NULL;
+	BodyInstance.InstanceBodyIndex = InstanceIdx;
+}
+
+void UDestructibleComponent::ResetFakeBodyInstance( FFakeBodyInstanceState& PrevState )
+{
+	BodyInstance.RigidActorSync = PrevState.ActorSync;
+	BodyInstance.RigidActorAsync = PrevState.ActorAsync;
+	BodyInstance.InstanceBodyIndex = PrevState.InstanceIndex;
+}
+
+#endif
+
+void UDestructibleComponent::SetEnableGravity(bool bGravityEnabled)
+{
+	Super::SetEnableGravity(bGravityEnabled);
+	
+#if WITH_APEX
+	for (FDestructibleChunkInfo& ChunkInfo : ChunkInfos)
+	 {
+		physx::PxRigidDynamic* Actor = ChunkInfo.Actor;
+		if (Actor)
+		{
+			Actor->setActorFlag(PxActorFlag::eDISABLE_GRAVITY, !bGravityEnabled);
+		}
+	}
+#endif //WITH_APEX
+}
+
+FBodyInstance* UDestructibleComponent::GetBodyInstance( FName BoneName /*= NAME_None*/, bool) const
+{
+#if WITH_APEX
+	if (ApexDestructibleActor != NULL)
+	{
+		int32 BoneIdx = GetBoneIndex(BoneName);
+		PxRigidDynamic* PActor = ApexDestructibleActor->getChunkPhysXActor(BoneIdxToChunkIdx(BoneIdx));
+
+		const_cast<UDestructibleComponent*>(this)->SetupFakeBodyInstance(PActor, BoneIdx);
+	}
+#endif // WITH_APEX
+
+	return (FBodyInstance*)&BodyInstance;
+}
+
+bool UDestructibleComponent::IsAnySimulatingPhysics() const 
+{
+	return !!BodyInstance.bSimulatePhysics;
+}
+
+#if WITH_PHYSX
+
+bool UDestructibleComponent::IsChunkLarge(int32 ChunkIdx) const
+{
+#if WITH_APEX
+	check(ApexDestructibleActor);
+	physx::PxBounds3 Bounds = ApexDestructibleActor->getChunkBounds(ChunkIdx);
+	return Bounds.getExtents().maxElement() > LargeChunkThreshold;
+#else
+	return true;
+#endif // WITH_APEX
+}
+
+void UDestructibleComponent::SetCollisionResponseForActor(PxRigidDynamic* Actor, int32 ChunkIdx)
+{
+	// Get collision channel and response
+	PxFilterData PQueryFilterData, PSimFilterData;
+	uint8 MoveChannel = GetCollisionObjectType();
+	if(IsCollisionEnabled())
+	{
+		UDestructibleMesh* TheDestructibleMesh = GetDestructibleMesh();
+		AActor* Owner = GetOwner();
+		bool bLargeChunk = IsChunkLarge(ChunkIdx);
+		const FCollisionResponse& ColResponse = bLargeChunk ? LargeChunkCollisionResponse : SmallChunkCollisionResponse;
+#if WITH_APEX
+		physx::PxU32 SupportDepth = TheDestructibleMesh->ApexDestructibleAsset->getChunkDepth(ChunkIdx);
+
+		const bool bEnableImpactDamage = IsImpactDamageEnabled(TheDestructibleMesh, SupportDepth);
+		CreateShapeFilterData(MoveChannel, GetUniqueID(), ColResponse.GetResponseContainer(), 0, ChunkIdxToBoneIdx(ChunkIdx), PQueryFilterData, PSimFilterData, BodyInstance.bUseCCD, bEnableImpactDamage, false);
+		
+		PQueryFilterData.word3 |= EPDF_SimpleCollision | EPDF_ComplexCollision;
+
+		SCOPED_SCENE_WRITE_LOCK(Actor->getScene());
+
+		TArray<PxShape*> Shapes;
+		Shapes.AddUninitialized(Actor->getNbShapes());
+
+		int ShapeCount = Actor->getShapes(Shapes.GetData(), Shapes.Num());
+
+		for (int32 i=0; i < ShapeCount; ++i)
+		{
+			PxShape* Shape = Shapes[i];
+
+			Shape->setQueryFilterData(PQueryFilterData);
+			Shape->setSimulationFilterData(PSimFilterData);
+			Shape->setFlag(PxShapeFlag::eSCENE_QUERY_SHAPE, true); 
+			Shape->setFlag(PxShapeFlag::eSIMULATION_SHAPE, true); 
+			Shape->setFlag(PxShapeFlag::eVISUALIZATION, true);
+		}
+#endif
+	}
+}
+
+#endif // WITH_PHYSX
\ No newline at end of file
diff --git a/Engine/Source/Runtime/Engine/Private/Components/DirectionalLightComponent.cpp b/Engine/Source/Runtime/Engine/Private/Components/DirectionalLightComponent.cpp
new file mode 100644
--- /dev/null
+++ b/Engine/Source/Runtime/Engine/Private/Components/DirectionalLightComponent.cpp
@@ -0,0 +1,900 @@
+// Copyright 1998-2015 Epic Games, Inc. All Rights Reserved.
+
+/*=============================================================================
+	DirectionalLightComponent.cpp: DirectionalLightComponent implementation.
+=============================================================================*/
+
+#include "EnginePrivate.h"
+#include "Engine/DirectionalLight.h"
+#include "Components/DirectionalLightComponent.h"
+
+static float GMaxCSMRadiusToAllowPerObjectShadows = 8000;
+static FAutoConsoleVariableRef CVarMaxCSMRadiusToAllowPerObjectShadows(
+	TEXT("r.MaxCSMRadiusToAllowPerObjectShadows"),
+	GMaxCSMRadiusToAllowPerObjectShadows,
+	TEXT("Only stationary lights with a CSM radius smaller than this will create per object shadows for dynamic objects.")
+	);
+
+static TAutoConsoleVariable<float> CVarUnbuiltWholeSceneDynamicShadowRadius(
+	TEXT("r.Shadow.UnbuiltWholeSceneDynamicShadowRadius"),
+	200000.0f,
+	TEXT("WholeSceneDynamicShadowRadius to use when using CSM to preview unbuilt lighting from a directional light")
+	);
+
+static TAutoConsoleVariable<int32> CVarUnbuiltNumWholeSceneDynamicShadowCascades(
+	TEXT("r.Shadow.UnbuiltNumWholeSceneDynamicShadowCascades"),
+	4,
+	TEXT("DynamicShadowCascades to use when using CSM to preview unbuilt lighting from a directional light"),
+	ECVF_RenderThreadSafe);
+
+/**
+ * The directional light policy for TMeshLightingDrawingPolicy.
+ */
+class FDirectionalLightPolicy
+{
+public:
+	typedef FLightSceneInfo SceneInfoType;
+};
+
+/**
+ * The scene info for a directional light.
+ */
+class FDirectionalLightSceneProxy : public FLightSceneProxy
+{
+public:
+
+	/** Whether to occlude fog and atmosphere inscattering with screenspace blurred occlusion from this light. */
+	bool bEnableLightShaftOcclusion;
+
+	/** 
+	 * Controls how dark the occlusion masking is, a value of 1 results in no darkening term.
+	 */
+	float OcclusionMaskDarkness;
+
+	/** Everything closer to the camera than this distance will occlude light shafts. */
+	float OcclusionDepthRange;
+
+	/** 
+	 * Can be used to make light shafts come from somewhere other than the light's actual direction. 
+	 * Will only be used when non-zero.
+	 */
+	FVector LightShaftOverrideDirection;
+
+	/** 
+	 * Radius of the whole scene dynamic shadow centered on the viewer, which replaces the precomputed shadows based on distance from the camera.  
+	 * A Radius of 0 disables the dynamic shadow.
+	 */
+	float WholeSceneDynamicShadowRadius;
+
+	/** 
+	 * Number of cascades to split the view frustum into for the whole scene dynamic shadow.  
+	 * More cascades result in better shadow resolution and allow WholeSceneDynamicShadowRadius to be further, but add rendering cost.
+	 */
+	uint32 DynamicShadowCascades;
+
+	/** 
+	 * Exponent that is applied to the cascade transition distances as a fraction of WholeSceneDynamicShadowRadius.
+	 * An exponent of 1 means that cascade transitions will happen at a distance proportional to their resolution.
+	 * A value greater than 1 brings transitions closer to the camera.
+	 */
+	float CascadeDistributionExponent;
+
+	/** see UDirectionalLightComponent::CascadeTransitionFraction */
+	float CascadeTransitionFraction;
+
+	/** see UDirectionalLightComponent::ShadowDistanceFadeoutFraction */
+	float ShadowDistanceFadeoutFraction;
+
+	bool bUseInsetShadowsForMovableObjects;
+
+	/** If greater than WholeSceneDynamicShadowRadius, a cascade will be created to support ray traced distance field shadows covering up to this distance. */
+	float DistanceFieldShadowDistance;
+
+	/** Light source angle in degrees. */
+	float LightSourceAngle;
+
+	/** Initialization constructor. */
+	FDirectionalLightSceneProxy(const UDirectionalLightComponent* Component):
+		FLightSceneProxy(Component),
+		bEnableLightShaftOcclusion(Component->bEnableLightShaftOcclusion),
+		OcclusionMaskDarkness(Component->OcclusionMaskDarkness),
+		OcclusionDepthRange(Component->OcclusionDepthRange),
+		LightShaftOverrideDirection(Component->LightShaftOverrideDirection),
+		DynamicShadowCascades(Component->DynamicShadowCascades > 0 ? Component->DynamicShadowCascades : 0),
+		CascadeDistributionExponent(Component->CascadeDistributionExponent),
+		CascadeTransitionFraction(Component->CascadeTransitionFraction),
+		ShadowDistanceFadeoutFraction(Component->ShadowDistanceFadeoutFraction),
+		bUseInsetShadowsForMovableObjects(Component->bUseInsetShadowsForMovableObjects),
+		DistanceFieldShadowDistance(Component->bUseRayTracedDistanceFieldShadows ? Component->DistanceFieldShadowDistance : 0),
+		LightSourceAngle(Component->LightSourceAngle)
+	{
+		LightShaftOverrideDirection.Normalize();
+
+		if(Component->Mobility == EComponentMobility::Movable)
+		{
+			WholeSceneDynamicShadowRadius = Component->DynamicShadowDistanceMovableLight;
+		}
+		else
+		{
+			WholeSceneDynamicShadowRadius = Component->DynamicShadowDistanceStationaryLight;
+		}
+
+		const float FarCascadeSize = Component->FarShadowDistance - WholeSceneDynamicShadowRadius;
+
+		// 100 is just some number to avoid cascades of 0 size, the user still can still create many tiny cascades (very inefficient)
+		if(Component->FarShadowCascadeCount && FarCascadeSize > 100.0f)
+		{
+			FarShadowDistance = Component->FarShadowDistance;
+			FarShadowCascadeCount = Component->FarShadowCascadeCount;
+		}
+	}
+
+	void UpdateLightShaftOverrideDirection_GameThread(const UDirectionalLightComponent* Component)
+	{
+		FVector NewLightShaftOverrideDirection = Component->LightShaftOverrideDirection;
+		NewLightShaftOverrideDirection.Normalize();
+		ENQUEUE_UNIQUE_RENDER_COMMAND_TWOPARAMETER(
+			FUpdateLightShaftOverrideDirectionCommand,
+			FDirectionalLightSceneProxy*,Proxy,this,
+			FVector,NewLightShaftOverrideDirection,NewLightShaftOverrideDirection,
+		{
+			Proxy->UpdateLightShaftOverrideDirection_RenderThread(NewLightShaftOverrideDirection);
+		});
+	}
+
+	/** Accesses parameters needed for rendering the light. */
+	virtual void GetParameters(FVector4& LightPositionAndInvRadius, FVector4& LightColorAndFalloffExponent, FVector& NormalizedLightDirection, FVector2D& SpotAngles, float& LightSourceRadius, float& LightSourceLength, float& LightMinRoughness) const
+	{
+		LightPositionAndInvRadius = FVector4(0, 0, 0, 0);
+
+		LightColorAndFalloffExponent = FVector4(
+			GetColor().R,
+			GetColor().G,
+			GetColor().B,
+			0);
+
+		NormalizedLightDirection = -GetDirection();
+
+		SpotAngles = FVector2D(0, 0);
+		LightSourceRadius = 0.0f;
+		LightSourceLength = 0.0f;
+		LightMinRoughness = MinRoughness;
+	}
+
+	virtual float GetLightSourceAngle() const override
+	{
+		return LightSourceAngle;
+	}
+
+	virtual bool GetLightShaftOcclusionParameters(float& OutOcclusionMaskDarkness, float& OutOcclusionDepthRange) const override
+	{
+		OutOcclusionMaskDarkness = OcclusionMaskDarkness;
+		OutOcclusionDepthRange = OcclusionDepthRange;
+		return bEnableLightShaftOcclusion;
+	}
+
+	virtual FVector GetLightPositionForLightShafts(FVector ViewOrigin) const override
+	{
+		const FVector EffectiveDirection = LightShaftOverrideDirection.SizeSquared() > 0 ? LightShaftOverrideDirection : GetDirection();
+		return ViewOrigin - EffectiveDirection * WORLD_MAX;
+	}
+
+	// FLightSceneInfo interface.
+
+	/** Called when precomputed lighting has been determined to be invalid */
+	virtual void InvalidatePrecomputedLighting(bool bIsEditor) 
+	{
+		// If the light is unbuilt and is not already using CSM, setup CSM settings to preview unbuilt lighting
+		if (HasStaticShadowing())
+		{
+			CascadeDistributionExponent = 4;
+			DynamicShadowCascades = FMath::Max(0, CVarUnbuiltNumWholeSceneDynamicShadowCascades.GetValueOnAnyThread());
+			WholeSceneDynamicShadowRadius = CVarUnbuiltWholeSceneDynamicShadowRadius.GetValueOnAnyThread();
+		}
+	}
+
+	virtual bool ShouldCreatePerObjectShadowsForDynamicObjects() const 
+	{
+		return FLightSceneProxy::ShouldCreatePerObjectShadowsForDynamicObjects()
+			// Only create per-object shadows for dynamic objects if the CSM range is under some threshold
+			// When the CSM range is very small, CSM is just being used to provide high resolution / animating shadows near the player,
+			// But dynamic objects outside the CSM range would not have a shadow (or ones inside the range that cast a shadow out of the CSM area of influence).
+			&& WholeSceneDynamicShadowRadius < GMaxCSMRadiusToAllowPerObjectShadows
+			&& bUseInsetShadowsForMovableObjects;
+	}
+
+	/** Returns the number of view dependent shadows this light will create, not counting distance field shadow cascades. */
+	virtual uint32 GetNumViewDependentWholeSceneShadows(const FSceneView& View) const
+	{ 
+		uint32 TotalCascades = GetNumShadowMappedCascades(View.MaxShadowCascades) + FarShadowCascadeCount;
+
+		return TotalCascades;
+	}
+
+	/**
+	 * Sets up a projected shadow initializer that's dependent on the current view for shadows from the entire scene.
+	 * @param InCascadeIndex ShadowSplitIndex or INDEX_NONE for the the DistanceFieldCascade
+	 * @return True if the whole-scene projected shadow should be used.
+	 */
+	virtual bool GetViewDependentWholeSceneProjectedShadowInitializer(const FSceneView& View, int32 InCascadeIndex, FWholeSceneProjectedShadowInitializer& OutInitializer) const
+	{
+		const FMatrix& WorldToLight = GetWorldToLight();
+
+		const bool bRayTracedCascade = (InCascadeIndex == INDEX_NONE);
+
+		FSphere Bounds = FDirectionalLightSceneProxy::GetShadowSplitBounds(View, InCascadeIndex, &OutInitializer.CascadeSettings);
+
+		uint32 NumNearCascades = GetNumShadowMappedCascades(View.MaxShadowCascades);
+
+		// Last cascade is the ray traced shadow, if enabled
+		OutInitializer.CascadeSettings.ShadowSplitIndex = bRayTracedCascade ? NumNearCascades : InCascadeIndex;
+
+		const float ShadowExtent = Bounds.W / FMath::Sqrt(3.0f);
+		const FBoxSphereBounds SubjectBounds(Bounds.Center, FVector(ShadowExtent, ShadowExtent, ShadowExtent), Bounds.W);
+		OutInitializer.PreShadowTranslation = -Bounds.Center;
+		OutInitializer.WorldToLight = FInverseRotationMatrix(FVector(WorldToLight.M[0][0],WorldToLight.M[1][0],WorldToLight.M[2][0]).GetSafeNormal().Rotation());
+		OutInitializer.Scales = FVector(1.0f,1.0f / Bounds.W,1.0f / Bounds.W);
+		OutInitializer.FaceDirection = FVector(1,0,0);
+		OutInitializer.SubjectBounds = FBoxSphereBounds(FVector::ZeroVector,SubjectBounds.BoxExtent,SubjectBounds.SphereRadius);
+		OutInitializer.WAxis = FVector4(0,0,0,1);
+		OutInitializer.MinLightW = -HALF_WORLD_MAX;
+		// Reduce casting distance on a directional light
+		// This is necessary to improve floating point precision in several places, especially when deriving frustum verts from InvReceiverMatrix
+		OutInitializer.MaxDistanceToCastInLightW = HALF_WORLD_MAX / 32.0f;
+		OutInitializer.CascadeSettings.bRayTracedDistanceField = bRayTracedCascade;
+		OutInitializer.CascadeSettings.bFarShadowCascade = !bRayTracedCascade && OutInitializer.CascadeSettings.ShadowSplitIndex >= (int32)NumNearCascades;
+		return true;
+	}
+
+	// reflective shadow map for Light Propagation Volume
+	virtual bool GetViewDependentRsmWholeSceneProjectedShadowInitializer(
+		const class FSceneView& View, 
+		const FBox& LightPropagationVolumeBounds,
+		class FWholeSceneProjectedShadowInitializer& OutInitializer ) const
+	{
+		const FMatrix& WorldToLight = GetWorldToLight();
+
+		float LpvExtent = LightPropagationVolumeBounds.GetExtent().X; // LPV is a cube, so this should be valid
+
+		OutInitializer.PreShadowTranslation = -LightPropagationVolumeBounds.GetCenter();
+		OutInitializer.WorldToLight = FInverseRotationMatrix(FVector(WorldToLight.M[0][0],WorldToLight.M[1][0],WorldToLight.M[2][0]).GetSafeNormal().Rotation());
+		OutInitializer.Scales = FVector(1.0f,1.0f / LpvExtent,1.0f / LpvExtent);
+		OutInitializer.FaceDirection = FVector(1,0,0);
+		OutInitializer.SubjectBounds = FBoxSphereBounds( FVector::ZeroVector, LightPropagationVolumeBounds.GetExtent(), FMath::Sqrt( LpvExtent * LpvExtent * 3.0f ) );
+		OutInitializer.WAxis = FVector4(0,0,0,1);
+		OutInitializer.MinLightW = -HALF_WORLD_MAX;
+		// Reduce casting distance on a directional light
+		// This is necessary to improve floating point precision in several places, especially when deriving frustum verts from InvReceiverMatrix
+		OutInitializer.MaxDistanceToCastInLightW = HALF_WORLD_MAX / 32.0f;
+
+		// Compute the RSM bounds
+		{
+			FVector Centre = LightPropagationVolumeBounds.GetCenter();
+			FVector Extent = LightPropagationVolumeBounds.GetExtent();
+			FVector CascadeFrustumVerts[8];
+			CascadeFrustumVerts[0] = Centre + Extent * FVector( -1.0f,-1.0f, 1.0f ); // 0 Near Top    Right
+			CascadeFrustumVerts[1] = Centre + Extent * FVector( -1.0f,-1.0f,-1.0f ); // 1 Near Bottom Right
+			CascadeFrustumVerts[2] = Centre + Extent * FVector(  1.0f,-1.0f, 1.0f ); // 2 Near Top    Left
+			CascadeFrustumVerts[3] = Centre + Extent * FVector(  1.0f,-1.0f,-1.0f ); // 3 Near Bottom Left
+			CascadeFrustumVerts[4] = Centre + Extent * FVector( -1.0f, 1.0f, 1.0f ); // 4 Far  Top    Right
+			CascadeFrustumVerts[5] = Centre + Extent * FVector( -1.0f, 1.0f,-1.0f ); // 5 Far  Bottom Right
+			CascadeFrustumVerts[6] = Centre + Extent * FVector(  1.0f, 1.0f, 1.0f ); // 6 Far  Top    Left
+			CascadeFrustumVerts[7] = Centre + Extent * FVector(  1.0f, 1.0f,-1.0f ); // 7 Far  Bottom Left
+
+			FPlane Far;
+			FPlane Near;
+			const FVector LightDirection = -GetDirection();
+			ComputeShadowCullingVolume( CascadeFrustumVerts, LightDirection, OutInitializer.CascadeSettings.ShadowBoundsAccurate, Near, Far );
+		}
+		return true;
+	}
+
+	virtual FVector2D GetDirectionalLightDistanceFadeParameters(ERHIFeatureLevel::Type InFeatureLevel) const
+	{
+		float FarDistance = GetCSMMaxDistance();
+		{
+			if (ShouldCreateRayTracedCascade(InFeatureLevel))
+			{
+				FarDistance = GetDistanceFieldShadowDistance();
+			}
+			FarDistance = FMath::Max(FarDistance, FarShadowDistance);
+		}
+        
+		// The far distance for the dynamic to static fade is the range of the directional light.
+		// The near distance is placed at a depth of 90% of the light's range.
+		const float NearDistance = FarDistance - FarDistance * ShadowDistanceFadeoutFraction;
+		return FVector2D(NearDistance, 1.0f / FMath::Max<float>(FarDistance - NearDistance, KINDA_SMALL_NUMBER));
+	}
+
+	virtual bool GetPerObjectProjectedShadowInitializer(const FBoxSphereBounds& SubjectBounds,FPerObjectProjectedShadowInitializer& OutInitializer) const
+	{
+		OutInitializer.PreShadowTranslation = -SubjectBounds.Origin;
+		OutInitializer.WorldToLight = FInverseRotationMatrix(FVector(WorldToLight.M[0][0],WorldToLight.M[1][0],WorldToLight.M[2][0]).GetSafeNormal().Rotation());
+		OutInitializer.Scales = FVector(1.0f,1.0f / SubjectBounds.SphereRadius,1.0f / SubjectBounds.SphereRadius);
+		OutInitializer.FaceDirection = FVector(1,0,0);
+		OutInitializer.SubjectBounds = FBoxSphereBounds(FVector::ZeroVector,SubjectBounds.BoxExtent,SubjectBounds.SphereRadius);
+		OutInitializer.WAxis = FVector4(0,0,0,1);
+		OutInitializer.MinLightW = -HALF_WORLD_MAX;
+		// Reduce casting distance on a directional light
+		// This is necessary to improve floating point precision in several places, especially when deriving frustum verts from InvReceiverMatrix
+		OutInitializer.MaxDistanceToCastInLightW = HALF_WORLD_MAX / 32.0f;
+		return true;
+	}
+
+	virtual bool ShouldCreateRayTracedCascade(ERHIFeatureLevel::Type InFeatureLevel) const
+	{
+		//@todo - handle View.MaxShadowCascades properly
+		const uint32 NumCascades = GetNumShadowMappedCascades(10);
+		const float RaytracedShadowDistance = GetDistanceFieldShadowDistance();
+		const bool bCreateWithCSM = NumCascades > 0 && RaytracedShadowDistance > GetCSMMaxDistance();
+		const bool bCreateWithoutCSM = NumCascades == 0 && RaytracedShadowDistance > 0;
+		return DoesPlatformSupportDistanceFieldShadowing(GShaderPlatformForFeatureLevel[InFeatureLevel]) && (bCreateWithCSM || bCreateWithoutCSM);
+	}
+
+private:
+
+	uint32 GetNumShadowMappedCascades(uint32 MaxShadowCascades) const
+	{
+		uint32 NumCascades = WholeSceneDynamicShadowRadius > 0.0f ? DynamicShadowCascades : 0;
+		return FMath::Min<int32>(NumCascades, MaxShadowCascades);
+	}
+
+	float GetCSMMaxDistance() const
+	{
+		static const auto CVar = IConsoleManager::Get().FindTConsoleVariableDataFloat(TEXT("r.Shadow.DistanceScale"));
+		 
+		float Scale = FMath::Clamp(CVar->GetValueOnRenderThread(), 0.1f, 2.0f);
+		float Distance = WholeSceneDynamicShadowRadius * Scale;
+		return Distance;
+	}
+
+	float GetDistanceFieldShadowDistance() const
+	{
+		static const auto CVar = IConsoleManager::Get().FindTConsoleVariableDataInt(TEXT("r.GenerateMeshDistanceFields"));
+
+		if (CVar->GetValueOnRenderThread() == 0)
+		{
+			// Meshes must have distance fields generated
+			return 0;
+		}
+
+		return DistanceFieldShadowDistance;
+	}
+
+	float GetShadowTransitionScale() const
+	{
+		static const auto CVar = IConsoleManager::Get().FindTConsoleVariableDataFloat(TEXT("r.Shadow.CSM.TransitionScale"));
+		float Scale = FMath::Clamp(CVar->GetValueOnRenderThread(), 0.0f, 2.0f);
+		return Scale;
+	}
+
+	void UpdateLightShaftOverrideDirection_RenderThread(FVector NewLightShaftOverrideDirection)
+	{
+		LightShaftOverrideDirection = NewLightShaftOverrideDirection;
+	}
+
+	// Computes a shadow culling volume (convex hull) based on a set of 8 vertices and a light direction
+	void ComputeShadowCullingVolume(const FVector* CascadeFrustumVerts, const FVector& LightDirection, FConvexVolume& ConvexVolumeOut, FPlane& NearPlaneOut, FPlane& FarPlaneOut) const
+	{
+		// Pairs of plane indices from SubFrustumPlanes whose intersections
+		// form the edges of the frustum.
+		static const int32 AdjacentPlanePairs[12][2] =
+		{
+			{0,2}, {0,4}, {0,1}, {0,3},
+			{2,3}, {4,2}, {1,4}, {3,1},
+			{2,5}, {4,5}, {1,5}, {3,5}
+		};
+		// Maps a plane pair index to the index of the two frustum corners
+		// which form the end points of the plane intersection.
+		static const int32 LineVertexIndices[12][2] =
+		{
+			{0,1}, {1,3}, {3,2}, {2,0},
+			{0,4}, {1,5}, {3,7}, {2,6},
+			{4,5}, {5,7}, {7,6}, {6,4}
+		};
+
+		TArray<FPlane, TInlineAllocator<6>> Planes;
+
+		// Find the view frustum subsection planes which face away from the light and add them to the bounding volume
+		FPlane SubFrustumPlanes[6];
+		SubFrustumPlanes[0] = FPlane(CascadeFrustumVerts[3], CascadeFrustumVerts[2], CascadeFrustumVerts[0]); // Near
+		SubFrustumPlanes[1] = FPlane(CascadeFrustumVerts[7], CascadeFrustumVerts[6], CascadeFrustumVerts[2]); // Left
+		SubFrustumPlanes[2] = FPlane(CascadeFrustumVerts[0], CascadeFrustumVerts[4], CascadeFrustumVerts[5]); // Right
+		SubFrustumPlanes[3] = FPlane(CascadeFrustumVerts[2], CascadeFrustumVerts[6], CascadeFrustumVerts[4]); // Top
+		SubFrustumPlanes[4] = FPlane(CascadeFrustumVerts[5], CascadeFrustumVerts[7], CascadeFrustumVerts[3]); // Bottom
+		SubFrustumPlanes[5] = FPlane(CascadeFrustumVerts[4], CascadeFrustumVerts[6], CascadeFrustumVerts[7]); // Far
+
+		NearPlaneOut = SubFrustumPlanes[0];
+		FarPlaneOut = SubFrustumPlanes[5];
+
+		// Add the planes from the camera's frustum which form the back face of the frustum when in light space.
+		for (int32 i = 0; i < 6; i++)
+		{
+			FVector Normal(SubFrustumPlanes[i]);
+			float d = Normal | LightDirection;
+			if ( d < 0.0f )
+			{
+				Planes.Add(SubFrustumPlanes[i]);
+			}
+		}
+
+		// Now add the planes which form the silhouette edges of the camera frustum in light space.
+		for (int32 i = 0; i < 12; i++)
+		{
+			FVector NormalA(SubFrustumPlanes[AdjacentPlanePairs[i][0]]);
+			FVector NormalB(SubFrustumPlanes[AdjacentPlanePairs[i][1]]);
+
+			float DotA = NormalA | LightDirection;
+			float DotB = NormalB | LightDirection;
+
+			// If the signs of the dot product are different
+			if ( DotA * DotB < 0.0f )
+			{
+				// Planes are opposing, so this is an edge. 
+				// Extrude the plane along the light direction, and add it to the array.
+
+				FVector A = CascadeFrustumVerts[LineVertexIndices[i][0]];
+				FVector B = CascadeFrustumVerts[LineVertexIndices[i][1]];
+				// Scale the 3rd vector by the length of AB for precision
+				FVector C = A + LightDirection * (A - B).Size(); 
+
+				// Account for winding
+				if ( DotA >= 0.0f ) 
+				{
+					Planes.Add(FPlane(A, B, C));
+				}
+				else
+				{
+					Planes.Add(FPlane(B, A, C));
+				}
+			}
+		}
+		ConvexVolumeOut = FConvexVolume(Planes);
+	}
+
+	// Useful helper function to compute shadow map cascade distribution
+	// @param Exponent >=1, 1:linear, 2:each cascade gets 2x in size, ...
+	// @param CascadeIndex 0..CascadeCount
+	// @param CascadeCount >0
+	static float ComputeAccumulatedScale(float Exponent, int32 CascadeIndex, int32 CascadeCount)
+	{
+		if(CascadeIndex <= 0)
+		{
+			return 0.0f;
+		}
+
+		float CurrentScale = 1;
+		float TotalScale = 0.0f;
+		float Ret = 0.0f;
+
+		// lame implementation for simplicity, CascadeIndex is small anyway
+		for(int i = 0; i < CascadeCount; ++i)
+		{
+			if(i < CascadeIndex)
+			{
+				Ret += CurrentScale;
+			}
+			TotalScale += CurrentScale;
+			CurrentScale *= Exponent;
+		}
+
+		return Ret / TotalScale;
+	}
+	
+	// @param SplitIndex 0: near, 1:end of 1st cascade, ...
+	inline float GetSplitDistance(const FSceneView& View, uint32 SplitIndex, bool bDistanceFieldShadows) const
+	{
+		// near cascade means non far and non distance field cascade
+		uint32 NumNearCascades = GetNumShadowMappedCascades(View.MaxShadowCascades);
+
+		float CascadeDistanceWithoutFar = GetCSMMaxDistance();
+		float ShadowNear = View.NearClippingDistance;
+
+		// non near cascades are split differently for distance field shadow
+		if(SplitIndex > NumNearCascades)
+		{
+			if(bDistanceFieldShadows)
+			{
+				// there is only one distance field shadow cascade
+				check(SplitIndex == NumNearCascades + 1);
+				return DistanceFieldShadowDistance;
+			}
+			else
+			{
+				// the far cascades start at the after the near cascades
+				return CascadeDistanceWithoutFar + ComputeAccumulatedScale(CascadeDistributionExponent, SplitIndex - NumNearCascades, FarShadowCascadeCount) * (FarShadowDistance - CascadeDistanceWithoutFar);
+			}
+		}
+		else
+		{
+			return ShadowNear + ComputeAccumulatedScale(CascadeDistributionExponent, SplitIndex, NumNearCascades) * (CascadeDistanceWithoutFar - ShadowNear);
+		}
+	}
+
+	virtual FSphere GetShadowSplitBoundsDepthRange(const FSceneView& View, FVector ViewOrigin, float SplitNear, float SplitFar, FShadowCascadeSettings* OutCascadeSettings) const
+	{
+		const FMatrix ViewMatrix = View.ShadowViewMatrices.ViewMatrix;
+		const FMatrix ProjectionMatrix = View.ShadowViewMatrices.ProjMatrix;
+
+		const FVector CameraDirection = ViewMatrix.GetColumn(2);
+		const FVector LightDirection = -GetDirection();
+
+		// Get FOV and AspectRatio from the view's projection matrix.
+		float AspectRatio = ProjectionMatrix.M[1][1] / ProjectionMatrix.M[0][0];
+		float HalfFOV = View.ShadowViewMatrices.IsPerspectiveProjection() ? FMath::Atan(1.0f / ProjectionMatrix.M[0][0]) : PI/4.0f;
+
+		// Build the camera frustum for this cascade
+		const float StartHorizontalLength = SplitNear * FMath::Tan(HalfFOV);
+		const FVector StartCameraRightOffset = ViewMatrix.GetColumn(0) * StartHorizontalLength;
+		const float StartVerticalLength = StartHorizontalLength / AspectRatio;
+		const FVector StartCameraUpOffset = ViewMatrix.GetColumn(1) * StartVerticalLength;
+
+		const float EndHorizontalLength = SplitFar * FMath::Tan(HalfFOV);
+		const FVector EndCameraRightOffset = ViewMatrix.GetColumn(0) * EndHorizontalLength;
+		const float EndVerticalLength = EndHorizontalLength / AspectRatio;
+		const FVector EndCameraUpOffset = ViewMatrix.GetColumn(1) * EndVerticalLength;
+
+		// Get the 8 corners of the cascade's camera frustum, in world space
+		FVector CascadeFrustumVerts[8];
+		CascadeFrustumVerts[0] = ViewOrigin + CameraDirection * SplitNear + StartCameraRightOffset + StartCameraUpOffset; // 0 Near Top    Right
+		CascadeFrustumVerts[1] = ViewOrigin + CameraDirection * SplitNear + StartCameraRightOffset - StartCameraUpOffset; // 1 Near Bottom Right
+		CascadeFrustumVerts[2] = ViewOrigin + CameraDirection * SplitNear - StartCameraRightOffset + StartCameraUpOffset; // 2 Near Top    Left
+		CascadeFrustumVerts[3] = ViewOrigin + CameraDirection * SplitNear - StartCameraRightOffset - StartCameraUpOffset; // 3 Near Bottom Left
+		CascadeFrustumVerts[4] = ViewOrigin + CameraDirection * SplitFar  + EndCameraRightOffset   + EndCameraUpOffset;   // 4 Far  Top    Right
+		CascadeFrustumVerts[5] = ViewOrigin + CameraDirection * SplitFar  + EndCameraRightOffset   - EndCameraUpOffset;   // 5 Far  Bottom Right
+		CascadeFrustumVerts[6] = ViewOrigin + CameraDirection * SplitFar  - EndCameraRightOffset   + EndCameraUpOffset;   // 6 Far  Top    Left
+		CascadeFrustumVerts[7] = ViewOrigin + CameraDirection * SplitFar  - EndCameraRightOffset   - EndCameraUpOffset;   // 7 Far  Bottom Left
+
+		// Fit a bounding sphere around the world space camera cascade frustum.
+		// Compute the sphere ideal centre point given the FOV and near/far.
+		float TanHalfFOVx = FMath::Tan(HalfFOV);
+		float TanHalfFOVy = TanHalfFOVx / AspectRatio;
+		float FrustumLength = SplitFar - SplitNear;
+
+		float FarX = TanHalfFOVx * SplitFar;
+		float FarY = TanHalfFOVy * SplitFar;
+		float DiagonalASq = FarX * FarX + FarY * FarY;
+
+		float NearX = TanHalfFOVx * SplitNear;
+		float NearY = TanHalfFOVy * SplitNear;
+		float DiagonalBSq = NearX * NearX + NearY * NearY;
+
+		// Calculate the ideal bounding sphere for the subfrustum.
+		// Fx  = (Db^2 - da^2) / 2Fl + Fl / 2 
+		// (where Da is the far diagonal, and Db is the near, and Fl is the frustum length)
+		float OptimalOffset = (DiagonalBSq - DiagonalASq) / (2.0f * FrustumLength) + FrustumLength * 0.5f;
+		float CentreZ = SplitFar - OptimalOffset;
+		CentreZ = FMath::Clamp( CentreZ, SplitNear, SplitFar );
+		FSphere CascadeSphere(ViewOrigin + CameraDirection * CentreZ, 0);
+		for (int32 Index = 0; Index < 8; Index++)
+		{
+			CascadeSphere.W = FMath::Max(CascadeSphere.W, FVector::DistSquared(CascadeFrustumVerts[Index], CascadeSphere.Center));
+		}
+
+		// Don't allow the bounds to reach 0 (INF)
+		CascadeSphere.W = FMath::Max(FMath::Sqrt(CascadeSphere.W), 1.0f); 
+
+		if (OutCascadeSettings)
+		{
+			// this function is needed, since it's also called by the LPV code.
+			ComputeShadowCullingVolume(CascadeFrustumVerts, LightDirection, OutCascadeSettings->ShadowBoundsAccurate, OutCascadeSettings->NearFrustumPlane, OutCascadeSettings->FarFrustumPlane);
+		}
+
+		return CascadeSphere;
+	}
+
+	// @param InShadowSplitIndex cascade index or InShadowSplitIndex == INDEX_NONE for the distance field cascade
+	virtual FSphere GetShadowSplitBounds(const FSceneView& View, int32 InCascadeIndex, FShadowCascadeSettings* OutCascadeSettings) const
+	{
+		uint32 NumNearCascades = GetNumShadowMappedCascades(View.MaxShadowCascades);
+
+		const bool bHasRayTracedCascade = ShouldCreateRayTracedCascade(View.GetFeatureLevel());
+
+		// this checks for WholeSceneDynamicShadowRadius and DynamicShadowCascades
+		uint32 NumNearAndFarCascades = GetNumViewDependentWholeSceneShadows(View);
+
+		uint32 NumTotalCascades = FMath::Max(NumNearAndFarCascades, NumNearCascades + (bHasRayTracedCascade ? 1 : 0));
+
+		const bool bIsRayTracedCascade = InCascadeIndex == INDEX_NONE;
+		const uint32 ShadowSplitIndex = bIsRayTracedCascade ? NumNearCascades : InCascadeIndex;		// todo -1?
+
+		// Determine start and end distances to the current cascade's split planes
+		// Presence of the ray traced cascade does not change depth ranges for the shadow-mapped cascades
+		float SplitNear = GetSplitDistance(View, ShadowSplitIndex, bIsRayTracedCascade);
+		float SplitFar = GetSplitDistance(View, ShadowSplitIndex + 1, bIsRayTracedCascade);
+		float FadePlane = SplitFar;
+
+		float LocalCascadeTransitionFraction = CascadeTransitionFraction * GetShadowTransitionScale();
+
+		float FadeExtension = (SplitFar - SplitNear) * LocalCascadeTransitionFraction;
+
+		// Add the fade region to the end of the subfrustum, if this is not the last cascade.
+		if ((int32)ShadowSplitIndex < (int32)NumTotalCascades - 1)
+		{
+			SplitFar += FadeExtension;			
+		}
+
+		if(OutCascadeSettings)
+		{
+			OutCascadeSettings->SplitFarFadeRegion = FadeExtension;
+
+			OutCascadeSettings->SplitNearFadeRegion = 0.0f;
+
+			if(ShadowSplitIndex >= 1)
+			{
+				// only used to fade the translucency lighting volume
+
+				float BeforeSplitNear = GetSplitDistance(View, ShadowSplitIndex - 1, bIsRayTracedCascade);
+				float BeforeSplitFar = GetSplitDistance(View, ShadowSplitIndex, bIsRayTracedCascade);
+
+				OutCascadeSettings->SplitNearFadeRegion = (BeforeSplitFar - BeforeSplitNear) * LocalCascadeTransitionFraction;
+			}
+
+			// Pass out the split settings
+			OutCascadeSettings->SplitFar = SplitFar;
+			OutCascadeSettings->SplitNear = SplitNear;
+			OutCascadeSettings->FadePlaneOffset = FadePlane;
+			OutCascadeSettings->FadePlaneLength = SplitFar - FadePlane;
+		}
+
+		const FSphere CascadeSphere = FDirectionalLightSceneProxy::GetShadowSplitBoundsDepthRange(View, View.ViewMatrices.ViewOrigin, SplitNear, SplitFar, OutCascadeSettings);
+
+		return CascadeSphere;
+	}
+};
+
+UDirectionalLightComponent::UDirectionalLightComponent(const FObjectInitializer& ObjectInitializer)
+	: Super(ObjectInitializer)
+{
+#if WITH_EDITORONLY_DATA
+	if (!IsRunningCommandlet())
+	{
+		static ConstructorHelpers::FObjectFinder<UTexture2D> StaticTexture(TEXT("/Engine/EditorResources/LightIcons/S_LightDirectional"));
+		static ConstructorHelpers::FObjectFinder<UTexture2D> DynamicTexture(TEXT("/Engine/EditorResources/LightIcons/S_LightDirectionalMove"));
+
+		StaticEditorTexture = StaticTexture.Object;
+		StaticEditorTextureScale = 0.5f;
+		DynamicEditorTexture = DynamicTexture.Object;
+		DynamicEditorTextureScale = 0.5f;
+	}
+#endif
+
+	Intensity = 10;
+
+	bEnableLightShaftOcclusion = false;
+	OcclusionDepthRange = 100000.f;
+	OcclusionMaskDarkness = 0.05f;
+
+	WholeSceneDynamicShadowRadius_DEPRECATED = 20000.0f;
+	DynamicShadowDistanceMovableLight = 20000.0f;
+	DynamicShadowDistanceStationaryLight = 0.f;
+
+	DistanceFieldShadowDistance = 30000.0f;
+	FarShadowDistance = 300000.0f;
+	LightSourceAngle = 1;
+
+	DynamicShadowCascades = 3;
+	CascadeDistributionExponent = 3.0f;
+	CascadeTransitionFraction = 0.1f;
+	ShadowDistanceFadeoutFraction = 0.1f;
+	IndirectLightingIntensity = 1.0f;
+	CastTranslucentShadows = true;
+	bUseInsetShadowsForMovableObjects = true;
+}
+
+#if WITH_EDITOR
+/**
+ * Called after property has changed via e.g. property window or set command.
+ *
+ * @param	PropertyThatChanged	UProperty that has been changed, NULL if unknown
+ */
+void UDirectionalLightComponent::PostEditChangeProperty(FPropertyChangedEvent& PropertyChangedEvent)
+{
+	LightmassSettings.LightSourceAngle = FMath::Max(LightmassSettings.LightSourceAngle, 0.0f);
+	LightmassSettings.IndirectLightingSaturation = FMath::Max(LightmassSettings.IndirectLightingSaturation, 0.0f);
+	LightmassSettings.ShadowExponent = FMath::Clamp(LightmassSettings.ShadowExponent, .5f, 8.0f);
+
+	DynamicShadowDistanceMovableLight = FMath::Max(DynamicShadowDistanceMovableLight, 0.0f);
+	DynamicShadowDistanceStationaryLight = FMath::Max(DynamicShadowDistanceStationaryLight, 0.0f);
+
+	DynamicShadowCascades = FMath::Clamp(DynamicShadowCascades, 0, 10);
+	CascadeDistributionExponent = FMath::Clamp(CascadeDistributionExponent, .1f, 10.0f);
+	CascadeTransitionFraction = FMath::Clamp(CascadeTransitionFraction, 0.0f, 0.3f);
+	ShadowDistanceFadeoutFraction = FMath::Clamp(ShadowDistanceFadeoutFraction, 0.0f, 1.0f);
+	// max range is larger than UI
+	ShadowBias = FMath::Clamp(ShadowBias, 0.0f, 10.0f);
+
+	Super::PostEditChangeProperty(PropertyChangedEvent);
+}
+
+bool UDirectionalLightComponent::CanEditChange(const UProperty* InProperty) const
+{
+	if (InProperty)
+	{
+		FString PropertyName = InProperty->GetName();
+		
+		bool bShadowCascases = CastShadows
+			&& CastDynamicShadows 
+			&& ((DynamicShadowDistanceMovableLight > 0 && Mobility == EComponentMobility::Movable)
+			|| (DynamicShadowDistanceStationaryLight > 0 && Mobility == EComponentMobility::Stationary));
+
+		if (PropertyName == GET_MEMBER_NAME_STRING_CHECKED(UDirectionalLightComponent, bUseInsetShadowsForMovableObjects))
+		{
+			return CastShadows && CastDynamicShadows && DynamicShadowDistanceStationaryLight > 0 && Mobility == EComponentMobility::Stationary;
+		}
+
+		if (PropertyName == GET_MEMBER_NAME_STRING_CHECKED(UDirectionalLightComponent, DynamicShadowDistanceMovableLight))
+		{
+			return CastShadows && CastDynamicShadows;
+		}
+
+		if (PropertyName == GET_MEMBER_NAME_STRING_CHECKED(UDirectionalLightComponent, DynamicShadowCascades)
+			|| PropertyName == GET_MEMBER_NAME_STRING_CHECKED(UDirectionalLightComponent, CascadeDistributionExponent)
+			|| PropertyName == GET_MEMBER_NAME_STRING_CHECKED(UDirectionalLightComponent, CascadeTransitionFraction)
+			|| PropertyName == GET_MEMBER_NAME_STRING_CHECKED(UDirectionalLightComponent, ShadowDistanceFadeoutFraction)
+			|| PropertyName == GET_MEMBER_NAME_STRING_CHECKED(UDirectionalLightComponent, bUseInsetShadowsForMovableObjects)
+			|| PropertyName == GET_MEMBER_NAME_STRING_CHECKED(UDirectionalLightComponent, FarShadowCascadeCount))
+		{
+			return bShadowCascases;
+		}
+
+		if (PropertyName == GET_MEMBER_NAME_STRING_CHECKED(UDirectionalLightComponent, FarShadowDistance))
+		{
+			return bShadowCascases && FarShadowCascadeCount > 0;
+		}
+
+		if (PropertyName == GET_MEMBER_NAME_STRING_CHECKED(UDirectionalLightComponent, DistanceFieldShadowDistance)
+			|| PropertyName == GET_MEMBER_NAME_STRING_CHECKED(UDirectionalLightComponent, LightSourceAngle))
+		{
+			static const auto CVar = IConsoleManager::Get().FindTConsoleVariableDataInt(TEXT("r.GenerateMeshDistanceFields"));
+			bool bCanEdit = CastShadows && CastDynamicShadows && bUseRayTracedDistanceFieldShadows && Mobility != EComponentMobility::Static && CVar->GetValueOnGameThread() != 0;
+			return bCanEdit;
+		}
+
+		if (PropertyName == GET_MEMBER_NAME_STRING_CHECKED(UDirectionalLightComponent, OcclusionMaskDarkness)
+			|| PropertyName == GET_MEMBER_NAME_STRING_CHECKED(UDirectionalLightComponent, OcclusionDepthRange))
+		{
+			return bEnableLightShaftOcclusion;
+		}
+	}
+
+	return Super::CanEditChange(InProperty);
+}
+
+#endif // WITH_EDITOR
+
+FLightSceneProxy* UDirectionalLightComponent::CreateSceneProxy() const
+{
+	return new FDirectionalLightSceneProxy(this);
+}
+
+FVector4 UDirectionalLightComponent::GetLightPosition() const
+{
+	return FVector4(-GetDirection() * WORLD_MAX, 0 );
+}
+
+/**
+* @return ELightComponentType for the light component class 
+*/
+ELightComponentType UDirectionalLightComponent::GetLightType() const
+{
+	return LightType_Directional;
+}
+
+void UDirectionalLightComponent::SetDynamicShadowDistanceMovableLight(float NewValue)
+{
+	if (DynamicShadowDistanceMovableLight != NewValue)
+	{
+		DynamicShadowDistanceMovableLight = NewValue;
+		MarkRenderStateDirty();
+	}
+}
+
+void UDirectionalLightComponent::SetDynamicShadowDistanceStationaryLight(float NewValue)
+{
+	if (DynamicShadowDistanceStationaryLight != NewValue)
+	{
+		DynamicShadowDistanceStationaryLight = NewValue;
+		MarkRenderStateDirty();
+	}
+}
+
+void UDirectionalLightComponent::SetDynamicShadowCascades(int32 NewValue)
+{
+	if (DynamicShadowCascades != NewValue)
+	{
+		DynamicShadowCascades = NewValue;
+		MarkRenderStateDirty();
+	}
+}
+
+void UDirectionalLightComponent::SetCascadeDistributionExponent(float NewValue)
+{
+	if (CascadeDistributionExponent != NewValue)
+	{
+		CascadeDistributionExponent = NewValue;
+		MarkRenderStateDirty();
+	}
+}
+
+void UDirectionalLightComponent::SetCascadeTransitionFraction(float NewValue)
+{
+	if (CascadeTransitionFraction != NewValue)
+	{
+		CascadeTransitionFraction = NewValue;
+		MarkRenderStateDirty();
+	}
+}
+
+void UDirectionalLightComponent::SetShadowDistanceFadeoutFraction(float NewValue)
+{
+	if (ShadowDistanceFadeoutFraction != NewValue)
+	{
+		ShadowDistanceFadeoutFraction = NewValue;
+		MarkRenderStateDirty();
+	}
+}
+
+void UDirectionalLightComponent::SetEnableLightShaftOcclusion(bool bNewValue)
+{
+	if (AreDynamicDataChangesAllowed()
+		&& bEnableLightShaftOcclusion != bNewValue)
+	{
+		bEnableLightShaftOcclusion = bNewValue;
+		MarkRenderStateDirty();
+	}
+}
+
+void UDirectionalLightComponent::SetOcclusionMaskDarkness(float NewValue)
+{
+	if (AreDynamicDataChangesAllowed()
+		&& OcclusionMaskDarkness != NewValue)
+	{
+		OcclusionMaskDarkness = NewValue;
+		MarkRenderStateDirty();
+	}
+}
+
+void UDirectionalLightComponent::SetLightShaftOverrideDirection(FVector NewValue)
+{
+	if (AreDynamicDataChangesAllowed()
+		&& LightShaftOverrideDirection != NewValue)
+	{
+		LightShaftOverrideDirection = NewValue;
+		if (SceneProxy)
+		{
+			FDirectionalLightSceneProxy* DirectionalLightSceneProxy = (FDirectionalLightSceneProxy*)SceneProxy;
+			DirectionalLightSceneProxy->UpdateLightShaftOverrideDirection_GameThread(this);
+		}
+	}
+}
+
+void UDirectionalLightComponent::Serialize(FArchive& Ar)
+{
+	Super::Serialize(Ar);
+
+	if(Ar.UE4Ver() < VER_UE4_REMOVE_LIGHT_MOBILITY_CLASSES)
+	{
+		// If outer is a DirectionalLight, we use the ADirectionalLight::LoadedFromAnotherClass path
+		if( GetOuter() != NULL && !GetOuter()->IsA(ADirectionalLight::StaticClass()) )
+		{
+			if(Mobility == EComponentMobility::Movable)
+			{
+				DynamicShadowDistanceMovableLight = WholeSceneDynamicShadowRadius_DEPRECATED; 
+			}
+			else if(Mobility == EComponentMobility::Stationary)
+			{
+				DynamicShadowDistanceStationaryLight = WholeSceneDynamicShadowRadius_DEPRECATED;
+			}
+		}
+	}
+}
+
+void UDirectionalLightComponent::InvalidateLightingCacheDetailed(bool bInvalidateBuildEnqueuedLighting, bool bTranslationOnly)
+{
+	// Directional lights don't care about translation
+	if (!bTranslationOnly)
+	{
+		Super::InvalidateLightingCacheDetailed(bInvalidateBuildEnqueuedLighting, bTranslationOnly);
+	}
+}
\ No newline at end of file
diff --git a/Engine/Source/Runtime/Engine/Private/Components/DrawSphereComponent.cpp b/Engine/Source/Runtime/Engine/Private/Components/DrawSphereComponent.cpp
new file mode 100644
--- /dev/null
+++ b/Engine/Source/Runtime/Engine/Private/Components/DrawSphereComponent.cpp
@@ -0,0 +1,32 @@
+// Copyright 1998-2015 Epic Games, Inc. All Rights Reserved.
+
+#include "EnginePrivate.h"
+#include "Components/DrawSphereComponent.h"
+#if WITH_EDITOR
+#include "ShowFlags.h"
+#include "ConvexVolume.h"
+#endif
+
+UDrawSphereComponent::UDrawSphereComponent(const FObjectInitializer& ObjectInitializer)
+	: Super(ObjectInitializer)
+{
+	SetCollisionProfileName(UCollisionProfile::NoCollision_ProfileName);
+
+	bHiddenInGame = true;
+	bUseEditorCompositing = true;
+	bGenerateOverlapEvents = false;
+}
+
+#if WITH_EDITOR
+bool UDrawSphereComponent::ComponentIsTouchingSelectionBox(const FBox& InSelBBox, const FEngineShowFlags& ShowFlags, const bool bConsiderOnlyBSP, const bool bMustEncompassEntireComponent) const
+{
+	// Draw sphere components not treated as 'selectable' in editor
+	return false;
+}
+
+bool UDrawSphereComponent::ComponentIsTouchingSelectionFrustum(const FConvexVolume& InFrustum, const FEngineShowFlags& ShowFlags, const bool bConsiderOnlyBSP, const bool bMustEncompassEntireComponent) const
+{
+	// Draw sphere components not treated as 'selectable' in editor
+	return false;
+}
+#endif
diff --git a/Engine/Source/Runtime/Engine/Private/Components/HeightFogComponent.cpp b/Engine/Source/Runtime/Engine/Private/Components/HeightFogComponent.cpp
new file mode 100644
--- /dev/null
+++ b/Engine/Source/Runtime/Engine/Private/Components/HeightFogComponent.cpp
@@ -0,0 +1,204 @@
+// Copyright 1998-2015 Epic Games, Inc. All Rights Reserved.
+
+/*=============================================================================
+	HeightFogComponent.cpp: Height fog implementation.
+=============================================================================*/
+
+#include "EnginePrivate.h"
+#include "Components/ExponentialHeightFogComponent.h"
+#include "Engine/ExponentialHeightFog.h"
+#include "Net/UnrealNetwork.h"
+
+UExponentialHeightFogComponent::UExponentialHeightFogComponent(const FObjectInitializer& ObjectInitializer)
+	: Super(ObjectInitializer)
+{
+	FogInscatteringColor = FLinearColor(0.447f, 0.638f, 1.0f);
+
+	DirectionalInscatteringExponent = 4.0f;
+	DirectionalInscatteringStartDistance = 10000.0f;
+	DirectionalInscatteringColor = FLinearColor(0.25f, 0.25f, 0.125f);
+
+	FogDensity = 0.02f;
+	FogHeightFalloff = 0.2f;
+	FogMaxOpacity = 1.0f;
+	StartDistance = 0.0f;
+}
+
+void UExponentialHeightFogComponent::AddFogIfNeeded()
+{
+	if (ShouldComponentAddToScene() && ShouldRender() && IsRegistered() && FogDensity * 1000 > DELTA && FogMaxOpacity > DELTA
+		&& (GetOuter() == NULL || !GetOuter()->HasAnyFlags(RF_ClassDefaultObject)))
+	{
+		World->Scene->AddExponentialHeightFog(this);
+	}
+}
+
+void UExponentialHeightFogComponent::CreateRenderState_Concurrent()
+{
+	Super::CreateRenderState_Concurrent();
+	AddFogIfNeeded();
+}
+
+void UExponentialHeightFogComponent::SendRenderTransform_Concurrent()
+{
+	World->Scene->RemoveExponentialHeightFog(this);
+	AddFogIfNeeded();
+	Super::SendRenderTransform_Concurrent();
+}
+
+void UExponentialHeightFogComponent::DestroyRenderState_Concurrent()
+{
+	Super::DestroyRenderState_Concurrent();
+	World->Scene->RemoveExponentialHeightFog(this);
+}
+
+#if WITH_EDITOR
+void UExponentialHeightFogComponent::PostEditChangeProperty(FPropertyChangedEvent& PropertyChangedEvent)
+{
+	FogDensity = FMath::Clamp(FogDensity, 0.0f, 10.0f);
+	FogHeightFalloff = FMath::Clamp(FogHeightFalloff, 0.0f, 2.0f);
+	FogMaxOpacity = FMath::Clamp(FogMaxOpacity, 0.0f, 1.0f);
+	StartDistance = FMath::Clamp(StartDistance, 0.0f, (float)WORLD_MAX);
+
+	Super::PostEditChangeProperty(PropertyChangedEvent);
+}
+#endif // WITH_EDITOR
+
+void UExponentialHeightFogComponent::PostInterpChange(UProperty* PropertyThatChanged)
+{
+	Super::PostInterpChange(PropertyThatChanged);
+
+	MarkRenderStateDirty();
+}
+
+void UExponentialHeightFogComponent::SetFogDensity(float Value)
+{
+	if(FogDensity != Value)
+	{
+		FogDensity = Value;
+		MarkRenderStateDirty();
+	}
+}
+
+void UExponentialHeightFogComponent::SetFogInscatteringColor(FLinearColor Value)
+{
+	if(FogInscatteringColor != Value)
+	{
+		FogInscatteringColor = Value;
+		MarkRenderStateDirty();
+	}
+}
+
+void UExponentialHeightFogComponent::SetDirectionalInscatteringExponent(float Value)
+{
+	if(DirectionalInscatteringExponent != Value)
+	{
+		DirectionalInscatteringExponent = Value;
+		MarkRenderStateDirty();
+	}
+}
+
+void UExponentialHeightFogComponent::SetDirectionalInscatteringStartDistance(float Value)
+{
+	if(DirectionalInscatteringStartDistance != Value)
+	{
+		DirectionalInscatteringStartDistance = Value;
+		MarkRenderStateDirty();
+	}
+}
+
+void UExponentialHeightFogComponent::SetDirectionalInscatteringColor(FLinearColor Value)
+{
+	if(DirectionalInscatteringColor != Value)
+	{
+		DirectionalInscatteringColor = Value;
+		MarkRenderStateDirty();
+	}
+}
+
+void UExponentialHeightFogComponent::SetFogHeightFalloff(float Value)
+{
+	if(FogHeightFalloff != Value)
+	{
+		FogHeightFalloff = Value;
+		MarkRenderStateDirty();
+	}
+}
+
+void UExponentialHeightFogComponent::SetFogMaxOpacity(float Value)
+{
+	if(FogMaxOpacity != Value)
+	{
+		FogMaxOpacity = Value;
+		MarkRenderStateDirty();
+	}
+}
+
+void UExponentialHeightFogComponent::SetStartDistance(float Value)
+{
+	if(StartDistance != Value)
+	{
+		StartDistance = Value;
+		MarkRenderStateDirty();
+	}
+}
+
+//////////////////////////////////////////////////////////////////////////
+// AExponentialHeightFog
+
+AExponentialHeightFog::AExponentialHeightFog(const FObjectInitializer& ObjectInitializer)
+	: Super(ObjectInitializer)
+{
+	Component = CreateDefaultSubobject<UExponentialHeightFogComponent>(TEXT("HeightFogComponent0"));
+	RootComponent = Component;
+
+	bHidden = false;
+
+#if WITH_EDITORONLY_DATA
+	if (!IsRunningCommandlet() && (GetSpriteComponent() != NULL))
+	{
+		// Structure to hold one-time initialization
+		struct FConstructorStatics
+		{
+			ConstructorHelpers::FObjectFinderOptional<UTexture2D> FogTextureObject;
+			FName ID_Fog;
+			FText NAME_Fog;
+			FConstructorStatics()
+				: FogTextureObject(TEXT("/Engine/EditorResources/S_ExpoHeightFog"))
+				, ID_Fog(TEXT("Fog"))
+				, NAME_Fog(NSLOCTEXT("SpriteCategory", "Fog", "Fog"))
+			{
+			}
+		};
+		static FConstructorStatics ConstructorStatics;
+
+		GetSpriteComponent()->Sprite = ConstructorStatics.FogTextureObject.Get();
+		GetSpriteComponent()->RelativeScale3D = FVector(0.5f, 0.5f, 0.5f);
+		GetSpriteComponent()->SpriteInfo.Category = ConstructorStatics.ID_Fog;
+		GetSpriteComponent()->SpriteInfo.DisplayName = ConstructorStatics.NAME_Fog;
+		GetSpriteComponent()->AttachParent = Component;
+	}
+#endif // WITH_EDITORONLY_DATA
+}
+
+void AExponentialHeightFog::PostInitializeComponents()
+{
+	Super::PostInitializeComponents();
+
+	bEnabled = Component->bVisible;
+}
+
+void AExponentialHeightFog::GetLifetimeReplicatedProps( TArray< FLifetimeProperty > & OutLifetimeProps ) const
+{
+	Super::GetLifetimeReplicatedProps( OutLifetimeProps );
+	
+	DOREPLIFETIME( AExponentialHeightFog, bEnabled );
+}
+
+void AExponentialHeightFog::OnRep_bEnabled()
+{
+	Component->SetVisibility(bEnabled);
+}
+
+/** Returns Component subobject **/
+UExponentialHeightFogComponent* AExponentialHeightFog::GetComponent() const { return Component; }
