Engine: Moved actor components into Components folder

diff --git a/Engine/Source/Runtime/Engine/Private/Components/InputComponent.cpp b/Engine/Source/Runtime/Engine/Private/Components/InputComponent.cpp
new file mode 100644
--- /dev/null
+++ b/Engine/Source/Runtime/Engine/Private/Components/InputComponent.cpp
@@ -0,0 +1,261 @@
+// Copyright 1998-2015 Epic Games, Inc. All Rights Reserved.
+
+#include "EnginePrivate.h"
+
+
+/* FInputChord interface
+ *****************************************************************************/
+
+FInputChord::ERelationshipType FInputChord::GetRelationship( const FInputChord& OtherChord ) const
+{
+	ERelationshipType Relationship = None;
+
+	if (Key == OtherChord.Key)
+	{
+		if ((bAlt == OtherChord.bAlt) &&
+			(bCtrl == OtherChord.bCtrl) &&
+			(bShift == OtherChord.bShift) &&
+			(bCmd == OtherChord.bCmd))
+		{
+			Relationship = Same;
+		}
+		else if ((bAlt || !OtherChord.bAlt) &&
+				(bCtrl || !OtherChord.bCtrl) &&
+				(bShift || !OtherChord.bShift) &&
+				(bCmd || !OtherChord.bCmd))
+		{
+			Relationship = Masks;
+		}
+		else if ((!bAlt || OtherChord.bAlt) &&
+				(!bCtrl || OtherChord.bCtrl) &&
+				(!bShift || OtherChord.bShift) &&
+				(!bCmd || OtherChord.bCmd))
+		{
+			Relationship = Masked;
+		}
+	}
+
+	return Relationship;
+}
+
+
+/* UInputComponent interface
+ *****************************************************************************/
+
+UInputComponent::UInputComponent( const FObjectInitializer& ObjectInitializer )
+	: Super(ObjectInitializer)
+{
+	bBlockInput = false;
+}
+
+
+float UInputComponent::GetAxisValue( const FName AxisName ) const
+{
+	float AxisValue = 0.f;
+	bool bFound = false;
+
+	for (const FInputAxisBinding& AxisBinding : AxisBindings)
+	{
+		if (AxisBinding.AxisName == AxisName)
+		{
+			bFound = true;
+			AxisValue = AxisBinding.AxisValue;
+			break;
+		}
+	}
+
+	if (!bFound)
+	{
+		UE_LOG(LogPlayerController, Warning, TEXT("Request for value of axis '%s' returning 0 as it is not bound on this input component."), *AxisName.ToString());
+	}
+
+	return AxisValue;
+}
+
+
+float UInputComponent::GetAxisKeyValue( const FKey AxisKey ) const
+{
+	float AxisValue = 0.f;
+	bool bFound = false;
+
+	for (const FInputAxisKeyBinding& AxisBinding : AxisKeyBindings)
+	{
+		if (AxisBinding.AxisKey == AxisKey)
+		{
+			bFound = true;
+			AxisValue = AxisBinding.AxisValue;
+			break;
+		}
+	}
+
+	if (!bFound)
+	{
+		UE_LOG(LogPlayerController, Warning, TEXT("Request for value of axis key '%s' returning 0 as it is not bound on this input component."), *AxisKey.ToString());
+	}
+
+	return AxisValue;
+}
+
+
+FVector UInputComponent::GetVectorAxisValue( const FKey AxisKey ) const
+{
+	FVector AxisValue;
+	bool bFound = false;
+
+	for (const FInputVectorAxisBinding& AxisBinding : VectorAxisBindings)
+	{
+		if (AxisBinding.AxisKey == AxisKey)
+		{
+			AxisValue = AxisBinding.AxisValue;
+			bFound = true;
+			break;
+		}
+	}
+
+	if (!bFound)
+	{
+		UE_LOG(LogPlayerController, Warning, TEXT("Request for value of vector axis key '%s' returning 0 as it is not bound on this input component."), *AxisKey.ToString());
+	}
+
+	return AxisValue;
+}
+
+
+bool UInputComponent::HasBindings( ) const
+{
+	return ((ActionBindings.Num() > 0) ||
+			(AxisBindings.Num() > 0) ||
+			(AxisKeyBindings.Num() > 0) ||
+			(KeyBindings.Num() > 0) ||
+			(TouchBindings.Num() > 0) ||
+			(GestureBindings.Num() > 0) ||
+			(VectorAxisBindings.Num() > 0));
+}
+
+
+FInputActionBinding& UInputComponent::AddActionBinding( const FInputActionBinding& Binding )
+{
+	ActionBindings.Add(Binding);
+	if (Binding.KeyEvent == IE_Pressed || Binding.KeyEvent == IE_Released)
+	{
+		const EInputEvent PairedEvent = (Binding.KeyEvent == IE_Pressed ? IE_Released : IE_Pressed);
+		for (int32 BindingIndex = ActionBindings.Num() - 2; BindingIndex >= 0; --BindingIndex )
+		{
+			FInputActionBinding& ActionBinding = ActionBindings[BindingIndex];
+			if (ActionBinding.ActionName == Binding.ActionName)
+			{
+				// If we find a matching event that is already paired we know this is paired so mark it off and we're done
+				if (ActionBinding.bPaired)
+				{
+					ActionBindings.Last().bPaired = true;
+					break;
+				}
+				// Otherwise if this is a pair to the new one mark them both as paired
+				// Don't break as there could be two bound paired events
+				else if (ActionBinding.KeyEvent == PairedEvent)
+				{
+					ActionBinding.bPaired = true;
+					ActionBindings.Last().bPaired = true;
+				}
+			}
+		}
+	}
+
+	return ActionBindings.Last();
+}
+
+
+void UInputComponent::ClearActionBindings( )
+{
+	ActionBindings.Reset();
+}
+
+
+FInputActionBinding& UInputComponent::GetActionBinding( const int32 BindingIndex )
+{
+	return ActionBindings[BindingIndex];
+}
+
+
+int32 UInputComponent::GetNumActionBindings( ) const
+{
+	return ActionBindings.Num();
+}
+
+
+void UInputComponent::RemoveActionBinding( const int32 BindingIndex )
+{
+	if (BindingIndex >= 0 && BindingIndex < ActionBindings.Num())
+	{
+		const FInputActionBinding& BindingToRemove = ActionBindings[BindingIndex];
+
+		// Potentially need to clear some pairings
+		if (BindingToRemove.bPaired)
+		{
+			TArray<int32> IndicesToClear;
+			const EInputEvent PairedEvent = (BindingToRemove.KeyEvent == IE_Pressed ? IE_Released : IE_Pressed);
+
+			for (int32 ActionIndex = 0; ActionIndex < ActionBindings.Num(); ++ActionIndex)
+			{
+				if (ActionIndex != BindingIndex)
+				{
+					const FInputActionBinding& ActionBinding = ActionBindings[ActionIndex];
+					if (ActionBinding.ActionName == BindingToRemove.ActionName)
+					{
+						// If we find another of the same key event then the pairing is intact so we're done
+						if (ActionBinding.KeyEvent == BindingToRemove.KeyEvent)
+						{
+							IndicesToClear.Empty();
+							break;
+						}
+						// Otherwise we may need to clear the pairing so track the index
+						else if (ActionBinding.KeyEvent == PairedEvent)
+						{
+							IndicesToClear.Add(ActionIndex);
+						}
+					}
+				}
+			}
+
+			for (int32 ClearIndex = 0; ClearIndex < IndicesToClear.Num(); ++ClearIndex)
+			{
+				ActionBindings[IndicesToClear[ClearIndex]].bPaired = false;
+			}
+		}
+
+		ActionBindings.RemoveAt(BindingIndex);
+	}
+}
+
+void UInputComponent::ClearBindingValues()
+{
+	for (FInputAxisBinding& AxisBinding : AxisBindings)
+	{
+		AxisBinding.AxisValue = 0.f;
+	}
+	for (FInputAxisKeyBinding& AxisKeyBinding : AxisKeyBindings)
+	{
+		AxisKeyBinding.AxisValue = 0.f;
+	}
+	for (FInputVectorAxisBinding& VectorAxisBinding : VectorAxisBindings)
+	{
+		VectorAxisBinding.AxisValue = FVector::ZeroVector;
+	}
+	for (FInputGestureBinding& GestureBinding : GestureBindings)
+	{
+		GestureBinding.GestureValue = 0.f;
+	}
+}
+
+/* Deprecated functions (needed for Blueprints)
+ *****************************************************************************/
+
+bool UInputComponent::IsControllerKeyDown(FKey Key) const { return false; }
+bool UInputComponent::WasControllerKeyJustPressed(FKey Key) const { return false; }
+bool UInputComponent::WasControllerKeyJustReleased(FKey Key) const { return false; }
+float UInputComponent::GetControllerAnalogKeyState(FKey Key) const { return 0.f; }
+FVector UInputComponent::GetControllerVectorKeyState(FKey Key) const { return FVector(); }
+void UInputComponent::GetTouchState(int32 FingerIndex, float& LocationX, float& LocationY, bool& bIsCurrentlyPressed) const { }
+float UInputComponent::GetControllerKeyTimeDown(FKey Key) const { return 0.f; }
+void UInputComponent::GetControllerMouseDelta(float& DeltaX, float& DeltaY) const { }
+void UInputComponent::GetControllerAnalogStickState(EControllerAnalogStick::Type WhichStick, float& StickX, float& StickY) const { }
diff --git a/Engine/Source/Runtime/Engine/Private/Components/LightComponent.cpp b/Engine/Source/Runtime/Engine/Private/Components/LightComponent.cpp
new file mode 100644
--- /dev/null
+++ b/Engine/Source/Runtime/Engine/Private/Components/LightComponent.cpp
@@ -0,0 +1,1220 @@
+// Copyright 1998-2015 Epic Games, Inc. All Rights Reserved.
+
+/*=============================================================================
+	LightComponent.cpp: LightComponent implementation.
+=============================================================================*/
+
+#include "EnginePrivate.h"
+#if WITH_EDITOR
+#include "ObjectEditorUtils.h"
+#endif
+#include "MessageLog.h"
+#include "UObjectToken.h"
+#include "ComponentInstanceDataCache.h"
+#include "TargetPlatform.h"
+#include "ComponentReregisterContext.h"
+#include "Components/PointLightComponent.h"
+#include "Components/LightComponent.h"
+#include "Components/DirectionalLightComponent.h"
+
+void FStaticShadowDepthMap::InitRHI()
+{
+	if (ShadowMapSizeX > 0 && ShadowMapSizeY > 0 && GMaxRHIFeatureLevel >= ERHIFeatureLevel::SM4)
+	{
+		FRHIResourceCreateInfo CreateInfo;
+		FTexture2DRHIRef Texture2DRHI = RHICreateTexture2D(ShadowMapSizeX, ShadowMapSizeY, PF_R16F, 1, 1, 0, CreateInfo);
+		TextureRHI = Texture2DRHI;
+
+		uint32 DestStride = 0;
+		uint8* TextureData = (uint8*)RHILockTexture2D(Texture2DRHI, 0, RLM_WriteOnly, DestStride, false);
+		uint32 RowSize = ShadowMapSizeX * GPixelFormats[PF_R16F].BlockBytes;
+
+		for (int32 Y = 0; Y < ShadowMapSizeY; Y++)
+		{
+			FMemory::Memcpy(TextureData + DestStride * Y, ((uint8*)DepthSamples.GetData()) + RowSize * Y, RowSize);
+		}
+
+		RHIUnlockTexture2D(Texture2DRHI, 0, false);
+	}
+}
+
+void FStaticShadowDepthMap::Empty()
+{
+	DEC_DWORD_STAT_BY(STAT_PrecomputedShadowDepthMapMemory, DepthSamples.GetAllocatedSize());
+
+	ShadowMapSizeX = 0;
+	ShadowMapSizeY = 0;
+	DepthSamples.Empty();
+}
+
+void FStaticShadowDepthMap::InitializeAfterImport()
+{
+	INC_DWORD_STAT_BY(STAT_PrecomputedShadowDepthMapMemory, DepthSamples.GetAllocatedSize());
+	BeginInitResource(this);
+}
+
+FArchive& operator<<(FArchive& Ar, FStaticShadowDepthMap& ShadowMap)
+{
+	Ar << ShadowMap.WorldToLight;
+	Ar << ShadowMap.ShadowMapSizeX;
+	Ar << ShadowMap.ShadowMapSizeY;
+	Ar << ShadowMap.DepthSamples;
+
+	if (Ar.IsLoading())
+	{
+		INC_DWORD_STAT_BY(STAT_PrecomputedShadowDepthMapMemory, ShadowMap.DepthSamples.GetAllocatedSize());
+	}
+
+	return Ar;
+}
+
+void ULightComponentBase::SetCastShadows(bool bNewValue)
+{
+	if (AreDynamicDataChangesAllowed()
+		&& CastShadows != bNewValue)
+	{
+		CastShadows = bNewValue;
+		MarkRenderStateDirty();
+	}
+}
+
+void ULightComponentBase::Serialize(FArchive& Ar)
+{
+	Super::Serialize(Ar);
+
+	if (Ar.UE4Ver() < VER_UE4_INVERSE_SQUARED_LIGHTS_DEFAULT)
+	{
+		Intensity = Brightness_DEPRECATED;
+	}
+}
+
+/**
+ * Called after duplication & serialization and before PostLoad. Used to e.g. make sure GUIDs remains globally unique.
+ */
+void ULightComponentBase::PostDuplicate(bool bDuplicateForPIE)
+{
+	Super::PostDuplicate(bDuplicateForPIE);
+
+	if (!bDuplicateForPIE)
+	{
+		// Create new guids for light.
+		UpdateLightGUIDs();
+	}
+}
+
+
+#if WITH_EDITOR
+/**
+ * Called after importing property values for this object (paste, duplicate or .t3d import)
+ * Allow the object to perform any cleanup for properties which shouldn't be duplicated or
+ * are unsupported by the script serialization
+ */
+void ULightComponentBase::PostEditImport()
+{
+	Super::PostEditImport();
+	// Create new guids for light.
+	UpdateLightGUIDs();
+}
+
+void ULightComponentBase::UpdateLightSpriteTexture()
+{
+	if (SpriteComponent != NULL)
+	{
+		SpriteComponent->SetSprite(GetEditorSprite());
+
+		float SpriteScale = GetEditorSpriteScale();
+		SpriteComponent->SetRelativeScale3D(FVector(SpriteScale));
+	}
+}
+
+void ULightComponentBase::PostEditChangeProperty(FPropertyChangedEvent& PropertyChangedEvent)
+{
+	Super::PostEditChangeProperty(PropertyChangedEvent);
+
+	// Update sprite 
+	UpdateLightSpriteTexture();
+}
+
+#endif
+
+/**
+ * Validates light GUIDs and resets as appropriate.
+ */
+void ULightComponentBase::ValidateLightGUIDs()
+{
+	// Validate light guids.
+	if( !LightGuid.IsValid() )
+	{
+		LightGuid = FGuid::NewGuid();
+	}
+}
+
+void ULightComponentBase::UpdateLightGUIDs()
+{
+	LightGuid = FGuid::NewGuid();
+}
+
+bool ULightComponentBase::HasStaticLighting() const
+{
+	AActor* Owner = GetOwner();
+
+	return Owner && (Mobility == EComponentMobility::Static);
+}
+
+bool ULightComponentBase::HasStaticShadowing() const
+{
+	AActor* Owner = GetOwner();
+
+	return Owner && (Mobility != EComponentMobility::Movable);
+}
+
+#if WITH_EDITOR
+void ULightComponentBase::OnRegister()
+{
+	Super::OnRegister();
+
+	if (SpriteComponent)
+	{
+		SpriteComponent->SpriteInfo.Category = TEXT("Lighting");
+		SpriteComponent->SpriteInfo.DisplayName = NSLOCTEXT("SpriteCategory", "Lighting", "Lighting");
+
+		UpdateLightSpriteTexture();
+	}
+}
+#endif
+
+bool ULightComponentBase::ShouldCollideWhenPlacing() const
+{
+	return true;
+}
+
+FBoxSphereBounds ULightComponentBase::GetPlacementExtent() const
+{
+	FBoxSphereBounds NewBounds;
+	NewBounds.Origin = FVector::ZeroVector;
+	NewBounds.BoxExtent = FVector(25.0f, 25.0f, 25.0f);
+	NewBounds.SphereRadius = 12.5f;
+	return NewBounds;
+}
+
+
+FLightSceneProxy::FLightSceneProxy(const ULightComponent* InLightComponent)
+	: LightComponent(InLightComponent)
+	, IndirectLightingScale(InLightComponent->IndirectLightingIntensity)
+	, ShadowBias(InLightComponent->ShadowBias)
+	, ShadowSharpen(InLightComponent->ShadowSharpen)
+	, MinRoughness(InLightComponent->MinRoughness)
+	, LightGuid(InLightComponent->LightGuid)
+	, ShadowMapChannel(InLightComponent->ShadowMapChannel)
+	, PreviewShadowMapChannel(InLightComponent->PreviewShadowMapChannel)
+	, IESTexture(0)
+	, bStaticLighting(InLightComponent->HasStaticLighting())
+	, bStaticShadowing(InLightComponent->HasStaticShadowing())
+	, bCastDynamicShadow(InLightComponent->CastShadows && InLightComponent->CastDynamicShadows)
+	, bCastStaticShadow(InLightComponent->CastShadows && InLightComponent->CastStaticShadows)
+	, bCastTranslucentShadows(InLightComponent->CastTranslucentShadows)
+	, bCastShadowsFromCinematicObjectsOnly(InLightComponent->bCastShadowsFromCinematicObjectsOnly)
+	, bAffectTranslucentLighting(InLightComponent->bAffectTranslucentLighting)
+	, bUsedAsAtmosphereSunLight(InLightComponent->IsUsedAsAtmosphereSunLight())
+	, bAffectDynamicIndirectLighting(InLightComponent->bAffectDynamicIndirectLighting)
+	, bHasReflectiveShadowMap(InLightComponent->bAffectDynamicIndirectLighting && InLightComponent->GetLightType() == LightType_Directional)
+	, bUseRayTracedDistanceFieldShadows(InLightComponent->bUseRayTracedDistanceFieldShadows)
+	, RayStartOffsetDepthScale(InLightComponent->RayStartOffsetDepthScale)
+	, LightType(InLightComponent->GetLightType())	
+	, ComponentName(InLightComponent->GetOwner() ? InLightComponent->GetOwner()->GetFName() : InLightComponent->GetFName())
+	, LevelName(InLightComponent->GetOutermost()->GetFName())
+	, StatId(InLightComponent->GetStatID(true))
+	, FarShadowDistance(0)
+	, FarShadowCascadeCount(0)
+{
+	// Brightness in Lumens
+	float LightBrightness = InLightComponent->ComputeLightBrightness();
+
+	if(LightComponent->IESTexture)
+	{
+		UTextureLightProfile* IESTextureObject = Cast<UTextureLightProfile>(LightComponent->IESTexture);
+
+		if(IESTextureObject)
+		{
+			IESTexture = IESTextureObject;
+		}
+	}
+	Color = FLinearColor(InLightComponent->LightColor) * LightBrightness;
+
+	if(LightComponent->LightFunctionMaterial &&
+		LightComponent->LightFunctionMaterial->GetMaterial()->MaterialDomain == MD_LightFunction )
+	{
+		LightFunctionMaterial = LightComponent->LightFunctionMaterial->GetRenderProxy(false);
+	}
+	else
+	{
+		LightFunctionMaterial = NULL;
+	}
+
+	LightFunctionScale = LightComponent->LightFunctionScale;
+	LightFunctionFadeDistance = LightComponent->LightFunctionFadeDistance;
+	LightFunctionDisabledBrightness = LightComponent->DisabledBrightness;
+
+	StaticShadowDepthMap = &LightComponent->StaticShadowDepthMap;
+}
+
+bool FLightSceneProxy::ShouldCreatePerObjectShadowsForDynamicObjects() const
+{
+	// Only create per-object shadows for Stationary lights, which use static shadowing from the world and therefore need a way to integrate dynamic objects
+	return HasStaticShadowing() && !HasStaticLighting();
+}
+
+void FLightSceneProxy::SetTransform(const FMatrix& InLightToWorld,const FVector4& InPosition)
+{
+	LightToWorld = InLightToWorld;
+	WorldToLight = InLightToWorld.InverseFast();
+	Position = InPosition;
+}
+
+void FLightSceneProxy::SetColor(const FLinearColor& InColor)
+{
+	Color = InColor;
+}
+
+void FLightSceneProxy::ApplyWorldOffset(FVector InOffset)
+{
+	FMatrix NewLightToWorld = LightToWorld.ConcatTranslation(InOffset);
+	FVector4 NewPosition = Position + InOffset;
+	SetTransform(NewLightToWorld, NewPosition);
+}
+
+ULightComponentBase::ULightComponentBase(const FObjectInitializer& ObjectInitializer)
+	: Super(ObjectInitializer)
+{
+	Brightness_DEPRECATED = 3.1415926535897932f;
+	Intensity = 3.1415926535897932f;
+	LightColor = FColor::White;
+	bAffectsWorld = true;
+	CastShadows = true;
+	CastStaticShadows = true;
+	CastDynamicShadows = true;
+	bPrecomputedLightingIsValid = true;
+#if WITH_EDITORONLY_DATA
+	bVisualizeComponent = true;
+#endif
+}
+
+/**
+ * Updates/ resets light GUIDs.
+ */
+ULightComponent::ULightComponent(const FObjectInitializer& ObjectInitializer)
+	: Super(ObjectInitializer)
+{
+	ShadowMapChannel = INDEX_NONE;
+	PreviewShadowMapChannel = INDEX_NONE;
+	IndirectLightingIntensity = 1.0f;
+	ShadowBias = 0.5f;
+	ShadowSharpen = 0.0f;
+	bUseIESBrightness = false;
+	IESBrightnessScale = 1.0f;
+	IESTexture = NULL;
+
+	bAffectTranslucentLighting = true;
+	LightFunctionScale = FVector(1024.0f, 1024.0f, 1024.0f);
+
+	LightFunctionFadeDistance = 100000.0f;
+	DisabledBrightness = 0.5f;
+	MinRoughness = 0.08f;
+
+	bEnableLightShaftBloom = false;
+	BloomScale = .2f;
+	BloomThreshold = 0;
+	BloomTint = FColor::White;
+
+	RayStartOffsetDepthScale = .01f;
+}
+
+
+void ULightComponent::UpdateLightGUIDs()
+{
+	Super::UpdateLightGUIDs();
+	ShadowMapChannel = INDEX_NONE;
+}
+
+bool ULightComponent::AffectsPrimitive(const UPrimitiveComponent* Primitive) const
+{
+	// Check whether the light affects the primitive's bounding volume.
+	return AffectsBounds(Primitive->Bounds);
+}
+
+bool ULightComponent::AffectsBounds(const FBoxSphereBounds& Bounds) const
+{
+	return true;
+}
+
+bool ULightComponent::IsShadowCast(UPrimitiveComponent* Primitive) const
+{
+	if(Primitive->HasStaticLighting())
+	{
+		return CastShadows && CastStaticShadows;
+	}
+	else
+	{
+		return CastShadows && CastDynamicShadows;
+	}
+}
+
+float ULightComponent::ComputeLightBrightness() const
+{
+	float LightBrightness = Intensity;
+
+	if(IESTexture)
+	{
+		UTextureLightProfile* IESTextureObject = Cast<UTextureLightProfile>(IESTexture);
+
+		if(IESTextureObject)
+		{
+			if(bUseIESBrightness)
+			{
+				LightBrightness = IESTextureObject->Brightness * IESBrightnessScale;
+			}
+
+			LightBrightness *= IESTextureObject->TextureMultiplier;
+		}
+	}
+
+	return LightBrightness;
+}
+
+void ULightComponent::Serialize(FArchive& Ar)
+{
+	Super::Serialize(Ar);
+
+	if (Ar.UE4Ver() >= VER_UE4_STATIC_SHADOW_DEPTH_MAPS)
+	{
+		if (Ar.IsCooking() && !Ar.CookingTarget()->SupportsFeature(ETargetPlatformFeatures::HighQualityLightmaps))
+		{
+			// Toss lighting data only needed for high quality lightmaps
+			FStaticShadowDepthMap EmptyDepthMap;
+			Ar << EmptyDepthMap;
+		}
+		else
+		{
+			Ar << StaticShadowDepthMap;
+		}
+	}
+}
+
+/**
+ * Called after this UObject has been serialized
+ */
+void ULightComponent::PostLoad()
+{
+	Super::PostLoad();
+
+	if (LightFunctionMaterial && HasStaticLighting())
+	{
+		// Light functions can only be used on dynamic lights
+		LightFunctionMaterial = NULL;
+	}
+
+	PreviewShadowMapChannel = ShadowMapChannel;
+	Intensity = FMath::Max(0.0f, Intensity);
+
+	if (GetLinkerUE4Version() < VER_UE4_LIGHTCOMPONENT_USE_IES_TEXTURE_MULTIPLIER_ON_NON_IES_BRIGHTNESS)
+	{
+		if(IESTexture)
+		{
+			UTextureLightProfile* IESTextureObject = Cast<UTextureLightProfile>(IESTexture);
+
+			if(IESTextureObject)
+			{
+				Intensity /= IESTextureObject->TextureMultiplier; // Previous version didn't apply IES texture multiplier, so cancel out
+				IESBrightnessScale = FMath::Pow(IESBrightnessScale, 2.2f); // Previous version applied 2.2 gamma to brightness scale
+				IESBrightnessScale /= IESTextureObject->TextureMultiplier; // Previous version didn't apply IES texture multiplier, so cancel out
+			}
+		}
+	}
+
+	if (HasStaticShadowing() && !HasStaticLighting())
+	{
+		BeginInitResource(&StaticShadowDepthMap);
+	}
+}
+
+#if WITH_EDITOR
+void ULightComponent::PreEditUndo()
+{
+	// Directly call UActorComponent::PreEditChange to avoid ULightComponent::PreEditChange being called for transactions.
+	UActorComponent::PreEditChange(NULL);
+}
+
+void ULightComponent::PostEditUndo()
+{
+	// Directly call UActorComponent::PostEditChange to avoid ULightComponent::PostEditChange being called for transactions.
+	UActorComponent::PostEditChange();
+}
+
+bool ULightComponent::CanEditChange(const UProperty* InProperty) const
+{
+	if (InProperty)
+	{
+		FString PropertyName = InProperty->GetName();
+
+		if (PropertyName == GET_MEMBER_NAME_STRING_CHECKED(ULightComponent, bCastShadowsFromCinematicObjectsOnly))
+		{
+			return Mobility == EComponentMobility::Movable;
+		}
+
+		if (PropertyName == GET_MEMBER_NAME_STRING_CHECKED(ULightComponent, LightFunctionMaterial)
+			|| PropertyName == GET_MEMBER_NAME_STRING_CHECKED(ULightComponent, LightFunctionScale)
+			|| PropertyName == GET_MEMBER_NAME_STRING_CHECKED(ULightComponent, LightFunctionFadeDistance)
+			|| PropertyName == GET_MEMBER_NAME_STRING_CHECKED(ULightComponent, DisabledBrightness)
+			|| PropertyName == GET_MEMBER_NAME_STRING_CHECKED(ULightComponent, IESTexture)
+			|| PropertyName == GET_MEMBER_NAME_STRING_CHECKED(ULightComponent, bUseIESBrightness)
+			|| PropertyName == GET_MEMBER_NAME_STRING_CHECKED(ULightComponent, IESBrightnessScale))
+		{
+			if (Mobility == EComponentMobility::Static)
+			{
+				return false;
+			}
+		}
+		
+		if (PropertyName == GET_MEMBER_NAME_STRING_CHECKED(ULightComponent, bUseRayTracedDistanceFieldShadows)
+			|| PropertyName == GET_MEMBER_NAME_STRING_CHECKED(ULightComponent, RayStartOffsetDepthScale))
+		{
+			static const auto CVar = IConsoleManager::Get().FindTConsoleVariableDataInt(TEXT("r.GenerateMeshDistanceFields"));
+			bool bCanEdit = CastShadows && CastDynamicShadows && Mobility != EComponentMobility::Static && CVar->GetValueOnGameThread() != 0;
+			return bCanEdit;
+		}
+
+		if (PropertyName == GET_MEMBER_NAME_STRING_CHECKED(ULightComponent, LightFunctionScale)
+			|| PropertyName == GET_MEMBER_NAME_STRING_CHECKED(ULightComponent, LightFunctionFadeDistance)
+			|| PropertyName == GET_MEMBER_NAME_STRING_CHECKED(ULightComponent, DisabledBrightness))
+		{
+			return LightFunctionMaterial != NULL;
+		}
+
+		if (PropertyName == TEXT("LightmassSettings"))
+		{
+			return Mobility != EComponentMobility::Movable;
+		}
+
+		if (PropertyName == GET_MEMBER_NAME_STRING_CHECKED(ULightComponent, BloomScale)
+			|| PropertyName == GET_MEMBER_NAME_STRING_CHECKED(ULightComponent, BloomThreshold)
+			|| PropertyName == GET_MEMBER_NAME_STRING_CHECKED(ULightComponent, BloomTint))
+		{
+			return bEnableLightShaftBloom;
+		}
+	}
+
+	return Super::CanEditChange(InProperty);
+}
+
+void ULightComponent::PostEditChangeProperty(FPropertyChangedEvent& PropertyChangedEvent)
+{
+	UProperty* PropertyThatChanged = PropertyChangedEvent.MemberProperty;
+	const FString PropertyName = PropertyThatChanged ? PropertyThatChanged->GetName() : TEXT("");
+
+	Intensity = FMath::Max(0.0f, Intensity);
+
+	if (HasStaticLighting())
+	{
+		// Lightmapped lights must not have light functions
+		LightFunctionMaterial = NULL;
+	}
+
+	// Unbuild lighting because a property changed
+	// Exclude properties that don't affect built lighting
+	//@todo - make this inclusive instead of exclusive?
+	if( PropertyName != GET_MEMBER_NAME_STRING_CHECKED(ULightComponent, CastTranslucentShadows) &&
+		PropertyName != GET_MEMBER_NAME_STRING_CHECKED(ULightComponent, bCastShadowsFromCinematicObjectsOnly) &&
+		PropertyName != GET_MEMBER_NAME_STRING_CHECKED(ULightComponent, CastDynamicShadows) &&
+		PropertyName != GET_MEMBER_NAME_STRING_CHECKED(ULightComponent, bAffectTranslucentLighting) &&
+		PropertyName != GET_MEMBER_NAME_STRING_CHECKED(ULightComponent, MinRoughness) &&
+		PropertyName != GET_MEMBER_NAME_STRING_CHECKED(ULightComponent, LightFunctionMaterial) &&
+		PropertyName != GET_MEMBER_NAME_STRING_CHECKED(ULightComponent, LightFunctionScale) &&
+		PropertyName != GET_MEMBER_NAME_STRING_CHECKED(ULightComponent, LightFunctionFadeDistance) &&
+		PropertyName != GET_MEMBER_NAME_STRING_CHECKED(ULightComponent, DisabledBrightness) &&
+		PropertyName != GET_MEMBER_NAME_STRING_CHECKED(ULightComponent, ShadowBias) &&
+		PropertyName != GET_MEMBER_NAME_STRING_CHECKED(ULightComponent, ShadowSharpen) &&
+		PropertyName != GET_MEMBER_NAME_STRING_CHECKED(ULightComponent, bEnableLightShaftBloom) &&
+		PropertyName != GET_MEMBER_NAME_STRING_CHECKED(ULightComponent, BloomScale) &&
+		PropertyName != GET_MEMBER_NAME_STRING_CHECKED(ULightComponent, BloomThreshold) &&
+		PropertyName != GET_MEMBER_NAME_STRING_CHECKED(ULightComponent, BloomTint) &&
+		PropertyName != GET_MEMBER_NAME_STRING_CHECKED(ULightComponent, bUseRayTracedDistanceFieldShadows) &&
+		PropertyName != GET_MEMBER_NAME_STRING_CHECKED(ULightComponent, RayStartOffsetDepthScale) &&
+		PropertyName != GET_MEMBER_NAME_STRING_CHECKED(ULightComponent, bVisible) &&
+		// Point light properties that shouldn't unbuild lighting
+		PropertyName != GET_MEMBER_NAME_STRING_CHECKED(UPointLightComponent, SourceRadius) &&
+		PropertyName != GET_MEMBER_NAME_STRING_CHECKED(UPointLightComponent, SourceLength) &&
+		// Directional light properties that shouldn't unbuild lighting
+		PropertyName != GET_MEMBER_NAME_STRING_CHECKED(UDirectionalLightComponent, DynamicShadowDistanceMovableLight) &&
+		PropertyName != GET_MEMBER_NAME_STRING_CHECKED(UDirectionalLightComponent, DynamicShadowDistanceStationaryLight) &&
+		PropertyName != GET_MEMBER_NAME_STRING_CHECKED(UDirectionalLightComponent, DynamicShadowCascades) &&
+		PropertyName != GET_MEMBER_NAME_STRING_CHECKED(UDirectionalLightComponent, FarShadowDistance) &&
+		PropertyName != GET_MEMBER_NAME_STRING_CHECKED(UDirectionalLightComponent, FarShadowCascadeCount) &&
+		PropertyName != GET_MEMBER_NAME_STRING_CHECKED(UDirectionalLightComponent, CascadeDistributionExponent) &&
+		PropertyName != GET_MEMBER_NAME_STRING_CHECKED(UDirectionalLightComponent, CascadeTransitionFraction) &&
+		PropertyName != GET_MEMBER_NAME_STRING_CHECKED(UDirectionalLightComponent, ShadowDistanceFadeoutFraction) &&
+		PropertyName != GET_MEMBER_NAME_STRING_CHECKED(UDirectionalLightComponent, bUseInsetShadowsForMovableObjects) &&
+		PropertyName != GET_MEMBER_NAME_STRING_CHECKED(UDirectionalLightComponent, DistanceFieldShadowDistance) &&
+		PropertyName != GET_MEMBER_NAME_STRING_CHECKED(UDirectionalLightComponent, bEnableLightShaftOcclusion) &&
+		PropertyName != GET_MEMBER_NAME_STRING_CHECKED(UDirectionalLightComponent, OcclusionMaskDarkness) &&
+		PropertyName != GET_MEMBER_NAME_STRING_CHECKED(UDirectionalLightComponent, OcclusionDepthRange) &&
+		PropertyName != GET_MEMBER_NAME_STRING_CHECKED(UDirectionalLightComponent, LightShaftOverrideDirection) &&
+		// Properties that should only unbuild lighting for a Static light (can be changed dynamically on a Stationary light)
+		(PropertyName != GET_MEMBER_NAME_STRING_CHECKED(ULightComponent, Intensity) || Mobility == EComponentMobility::Static) &&
+		(PropertyName != GET_MEMBER_NAME_STRING_CHECKED(ULightComponent, LightColor) || Mobility == EComponentMobility::Static))
+	{
+		InvalidateLightingCache();
+	}
+
+	Super::PostEditChangeProperty(PropertyChangedEvent);
+}
+
+void ULightComponent::UpdateLightSpriteTexture()
+{
+	if( SpriteComponent != NULL )
+	{
+		if (HasStaticShadowing() &&
+			!HasStaticLighting() &&
+			bAffectsWorld &&
+			CastShadows &&
+			CastStaticShadows &&
+			PreviewShadowMapChannel == INDEX_NONE)
+		{
+			UTexture2D* SpriteTexture = NULL;
+			SpriteTexture = LoadObject<UTexture2D>(NULL, TEXT("/Engine/EditorResources/LightIcons/S_LightError.S_LightError"));
+			SpriteComponent->SetSprite(SpriteTexture);
+			SpriteComponent->SetRelativeScale3D(FVector(0.5f));
+		}
+		else
+		{
+			Super::UpdateLightSpriteTexture();
+		}
+	}
+}
+
+#endif // WITH_EDITOR
+
+void ULightComponent::BeginDestroy()
+{
+	Super::BeginDestroy();
+
+	if (HasStaticShadowing() && !HasStaticLighting())
+	{
+		BeginReleaseResource(&StaticShadowDepthMap);
+		StaticShadowDepthMap.Empty();
+	}
+
+	// Use a fence to keep track of when the rendering thread executes the release command
+	DestroyFence.BeginFence();
+}
+
+bool ULightComponent::IsReadyForFinishDestroy()
+{
+	// Don't allow the light component to be destroyed until its rendering resources have been released
+	return Super::IsReadyForFinishDestroy() && DestroyFence.IsFenceComplete();
+}
+
+void ULightComponent::OnRegister()
+{
+	Super::OnRegister();
+
+	// Update GUIDs on attachment if they are not valid.
+	ValidateLightGUIDs();
+}
+
+void ULightComponent::CreateRenderState_Concurrent()
+{
+	Super::CreateRenderState_Concurrent();
+
+	bool bHidden = false;
+#if WITH_EDITORONLY_DATA
+	bHidden = GetOwner() ? GetOwner()->bHiddenEdLevel : false;
+#endif // WITH_EDITORONLY_DATA
+
+	if(!ShouldComponentAddToScene())
+	{
+		bHidden = true;
+	}
+
+	if (bAffectsWorld)
+	{
+		if (bVisible && !bHidden)
+		{
+			// Add the light to the scene.
+			World->Scene->AddLight(this);
+		}
+		// Add invisible stationary lights to the scene in the editor
+		// Even invisible stationary lights consume a shadowmap channel so they must be included in the stationary light overlap preview
+		else if (GIsEditor 
+			&& !World->IsGameWorld()
+			&& CastShadows 
+			&& CastStaticShadows 
+			&& HasStaticShadowing()
+			&& !HasStaticLighting())
+		{
+			World->Scene->AddInvisibleLight(this);
+		}
+	}
+}
+
+void ULightComponent::SendRenderTransform_Concurrent()
+{
+	// Update the scene info's transform for this light.
+	World->Scene->UpdateLightTransform(this);
+	Super::SendRenderTransform_Concurrent();
+}
+
+void ULightComponent::DestroyRenderState_Concurrent()
+{
+	Super::DestroyRenderState_Concurrent();
+	World->Scene->RemoveLight(this);
+}
+
+/** Set brightness of the light */
+void ULightComponent::SetIntensity(float NewIntensity)
+{
+	// Can't set brightness on a static light
+	if (AreDynamicDataChangesAllowed()
+		&& Intensity != NewIntensity)
+	{
+		Intensity = NewIntensity;
+
+		// Use lightweight color and brightness update 
+		if( World && World->Scene )
+		{
+			//@todo - remove from scene if brightness or color becomes 0
+			World->Scene->UpdateLightColorAndBrightness( this );
+		}
+	}
+}
+
+void ULightComponent::SetIndirectLightingIntensity(float NewIntensity)
+{
+	// Can't set brightness on a static light
+	if (AreDynamicDataChangesAllowed()
+		&& IndirectLightingIntensity != NewIntensity)
+	{
+		IndirectLightingIntensity = NewIntensity;
+
+		// Use lightweight color and brightness update 
+		if( World && World->Scene )
+		{
+			//@todo - remove from scene if brightness or color becomes 0
+			World->Scene->UpdateLightColorAndBrightness( this );
+		}
+	}
+}
+
+/** Set color of the light */
+void ULightComponent::SetLightColor(FLinearColor NewLightColor)
+{
+	FColor NewColor(NewLightColor);
+
+	// Can't set color on a static light
+	if (AreDynamicDataChangesAllowed()
+		&& LightColor != NewColor)
+	{
+		LightColor	= NewColor;
+
+		// Use lightweight color and brightness update 
+		if( World && World->Scene )
+		{
+			//@todo - remove from scene if brightness or color becomes 0
+			World->Scene->UpdateLightColorAndBrightness( this );
+		}
+	}
+}
+
+void ULightComponent::SetLightFunctionMaterial(UMaterialInterface* NewLightFunctionMaterial)
+{
+	// Can't set light function on a static light
+	if (AreDynamicDataChangesAllowed()
+		&& NewLightFunctionMaterial != LightFunctionMaterial)
+	{
+		LightFunctionMaterial = NewLightFunctionMaterial;
+		MarkRenderStateDirty();
+	}
+}
+
+void ULightComponent::SetLightFunctionScale(FVector NewLightFunctionScale)
+{
+	if (AreDynamicDataChangesAllowed()
+		&& NewLightFunctionScale != LightFunctionScale)
+	{
+		LightFunctionScale = NewLightFunctionScale;
+		MarkRenderStateDirty();
+	}
+}
+
+void ULightComponent::SetLightFunctionFadeDistance(float NewLightFunctionFadeDistance)
+{
+	if (AreDynamicDataChangesAllowed()
+		&& NewLightFunctionFadeDistance != LightFunctionFadeDistance)
+	{
+		LightFunctionFadeDistance = NewLightFunctionFadeDistance;
+		MarkRenderStateDirty();
+	}
+}
+
+void ULightComponent::SetLightFunctionDisabledBrightness(float NewValue)
+{
+	if (AreDynamicDataChangesAllowed()
+		&& NewValue != DisabledBrightness)
+	{
+		DisabledBrightness = NewValue;
+		MarkRenderStateDirty();
+	}
+}
+
+void ULightComponent::SetAffectDynamicIndirectLighting(bool bNewValue)
+{
+	if (AreDynamicDataChangesAllowed()
+		&& bAffectDynamicIndirectLighting != bNewValue)
+	{
+		bAffectDynamicIndirectLighting = bNewValue;
+		MarkRenderStateDirty();
+	}
+}
+
+void ULightComponent::SetAffectTranslucentLighting(bool bNewValue)
+{
+	if (AreDynamicDataChangesAllowed()
+		&& bAffectTranslucentLighting != bNewValue)
+	{
+		bAffectTranslucentLighting = bNewValue;
+		MarkRenderStateDirty();
+	}
+}
+
+void ULightComponent::SetEnableLightShaftBloom(bool bNewValue)
+{
+	if (AreDynamicDataChangesAllowed()
+		&& bEnableLightShaftBloom != bNewValue)
+	{
+		bEnableLightShaftBloom = bNewValue;
+		MarkRenderStateDirty();
+	}
+}
+
+void ULightComponent::SetBloomScale(float NewValue)
+{
+	if (AreDynamicDataChangesAllowed()
+		&& BloomScale != NewValue)
+	{
+		BloomScale = NewValue;
+		MarkRenderStateDirty();
+	}
+}
+
+void ULightComponent::SetBloomThreshold(float NewValue)
+{
+	if (AreDynamicDataChangesAllowed()
+		&& BloomThreshold != NewValue)
+	{
+		BloomThreshold = NewValue;
+		MarkRenderStateDirty();
+	}
+}
+
+void ULightComponent::SetBloomTint(FColor NewValue)
+{
+	if (AreDynamicDataChangesAllowed()
+		&& BloomTint != NewValue)
+	{
+		BloomTint = NewValue;
+		MarkRenderStateDirty();
+	}
+}
+
+void ULightComponent::SetIESTexture(UTextureLightProfile* NewValue)
+{
+	if (AreDynamicDataChangesAllowed()
+		&& IESTexture != NewValue)
+	{
+		IESTexture = NewValue;
+		MarkRenderStateDirty();
+	}
+}
+
+// GetDirection
+FVector ULightComponent::GetDirection() const 
+{ 
+	return ComponentToWorld.GetUnitAxis( EAxis::X );
+}
+
+void ULightComponent::UpdateColorAndBrightness()
+{
+	if( World && World->Scene )
+	{
+		World->Scene->UpdateLightColorAndBrightness( this );
+	}
+}
+
+//
+//	ULightComponent::InvalidateLightingCache
+//
+
+void ULightComponent::InvalidateLightingCacheDetailed(bool bInvalidateBuildEnqueuedLighting, bool bTranslationOnly)
+{
+	InvalidateLightingCacheInner(true);
+
+	UWorld* World = GetWorld();
+	if (GIsEditor
+		&& World != NULL
+		&& HasStaticShadowing()
+		&& !HasStaticLighting())
+	{
+		ReassignStationaryLightChannels(World, false);
+	}
+}
+
+/** Invalidates the light's cached lighting with the option to recreate the light Guids. */
+void ULightComponent::InvalidateLightingCacheInner(bool bRecreateLightGuids)
+{
+	// Save the light state for transactions.
+	Modify();
+
+	// Detach the component from the scene for the duration of this function.
+	FComponentReregisterContext ReregisterContext(this);
+
+	// Block until the RT processes the unregister before modifying variables that it may need to access
+	FlushRenderingCommands();
+
+	StaticShadowDepthMap.Empty();
+	BeginReleaseResource(&StaticShadowDepthMap);
+
+	bPrecomputedLightingIsValid = false;
+
+	if (bRecreateLightGuids)
+	{
+		// Create new guids for light.
+		UpdateLightGUIDs();
+	}
+	else
+	{
+		ValidateLightGUIDs();
+		ShadowMapChannel = INDEX_NONE;
+	}
+	
+	// Send to render thread
+	MarkRenderStateDirty();
+}
+
+/** Used to store lightmap data during RerunConstructionScripts */
+class FPrecomputedLightInstanceData : public FSceneComponentInstanceData
+{
+public:
+	FPrecomputedLightInstanceData(const ULightComponent* SourceComponent)
+		: FSceneComponentInstanceData(SourceComponent)
+		, Transform(SourceComponent->ComponentToWorld)
+		, LightGuid(SourceComponent->LightGuid)
+		, ShadowMapChannel(SourceComponent->ShadowMapChannel)
+		, PreviewShadowMapChannel(SourceComponent->PreviewShadowMapChannel)
+		, bPrecomputedLightingIsValid(SourceComponent->bPrecomputedLightingIsValid)
+	{}
+
+	virtual void ApplyToComponent(UActorComponent* Component, const ECacheApplyPhase CacheApplyPhase) override
+	{
+		FSceneComponentInstanceData::ApplyToComponent(Component, CacheApplyPhase);
+		CastChecked<ULightComponent>(Component)->ApplyComponentInstanceData(this);
+	}
+
+	FTransform Transform;
+	FGuid LightGuid;
+	int32 ShadowMapChannel;
+	int32 PreviewShadowMapChannel;
+	bool bPrecomputedLightingIsValid;
+};
+
+FName ULightComponent::GetComponentInstanceDataType() const
+{
+	static const FName PrecomputedLightInstanceDataTypeName(TEXT("PrecomputedLightInstanceData"));
+	return PrecomputedLightInstanceDataTypeName;
+}
+
+FActorComponentInstanceData* ULightComponent::GetComponentInstanceData() const
+{
+	// Allocate new struct for holding light map data
+	return new FPrecomputedLightInstanceData(this);
+}
+
+void ULightComponent::ApplyComponentInstanceData(FPrecomputedLightInstanceData* LightMapData)
+{
+	check(LightMapData);
+
+	if (!LightMapData->Transform.Equals(ComponentToWorld))
+	{
+		return;
+	}
+
+	LightGuid = LightMapData->LightGuid;
+	ShadowMapChannel = LightMapData->ShadowMapChannel;
+	PreviewShadowMapChannel = LightMapData->PreviewShadowMapChannel;
+	bPrecomputedLightingIsValid = LightMapData->bPrecomputedLightingIsValid;
+
+	MarkRenderStateDirty();
+
+#if WITH_EDITOR
+	// Update the icon with the new state of PreviewShadowMapChannel
+	UpdateLightSpriteTexture();
+#endif
+}
+
+int32 ULightComponent::GetNumMaterials() const
+{
+	return 1;
+}
+
+UMaterialInterface* ULightComponent::GetMaterial(int32 ElementIndex) const
+{
+	if (ElementIndex == 0)
+	{
+		return LightFunctionMaterial;
+	}
+	else
+	{
+		return NULL;
+	}
+}
+
+void ULightComponent::SetMaterial(int32 ElementIndex, UMaterialInterface* InMaterial)
+{
+	if (ElementIndex == 0)
+	{
+		LightFunctionMaterial = InMaterial;
+		MarkRenderStateDirty();
+	}
+}
+
+/** 
+* This is called when property is modified by InterpPropertyTracks
+*
+* @param PropertyThatChanged	Property that changed
+*/
+void ULightComponent::PostInterpChange(UProperty* PropertyThatChanged)
+{
+	static FName LightColorName(TEXT("LightColor"));
+	static FName IntensityName(TEXT("Intensity"));
+	static FName BrightnessName(TEXT("Brightness"));
+	static FName IndirectLightingIntensityName(TEXT("IndirectLightingIntensity"));
+
+	FName PropertyName = PropertyThatChanged->GetFName();
+	if (PropertyName == LightColorName
+		|| PropertyName == IntensityName
+		|| PropertyName == BrightnessName
+		|| PropertyName == IndirectLightingIntensityName)
+	{
+		// Old brightness tracks will animate the deprecated value
+		if (PropertyName == BrightnessName)
+		{
+			Intensity = Brightness_DEPRECATED;
+		}
+
+		UpdateColorAndBrightness();
+	}
+	else
+	{
+		Super::PostInterpChange(PropertyThatChanged);
+	}
+}
+
+/** Stores a light and a channel it has been assigned to. */
+struct FLightAndChannel
+{
+	ULightComponent* Light;
+	int32 Channel;
+
+	FLightAndChannel(ULightComponent* InLight) :
+		Light(InLight),
+		Channel(INDEX_NONE)
+	{}
+};
+
+struct FCompareLightsByArrayCount
+{
+	FORCEINLINE bool operator()( const TArray<FLightAndChannel*>& A, const TArray<FLightAndChannel*>& B ) const 
+	{ 
+		// Sort by descending array count
+		return B.Num() < A.Num(); 
+	}
+};
+
+void ULightComponent::ReassignStationaryLightChannels(UWorld* TargetWorld, bool bAssignForLightingBuild)
+{
+	TMap<FLightAndChannel*, TArray<FLightAndChannel*> > LightToOverlapMap;
+
+	// Build an array of all static shadowing lights that need to be assigned
+	for(TObjectIterator<ULightComponent> LightIt(RF_ClassDefaultObject|RF_PendingKill); LightIt; ++LightIt)
+	{
+		ULightComponent* const LightComponent = *LightIt;
+
+		const bool bLightIsInWorld = LightComponent->GetOwner() && TargetWorld->ContainsActor(LightComponent->GetOwner()) && !LightComponent->GetOwner()->IsPendingKill();
+
+		if (bLightIsInWorld 
+			// Only operate on stationary light components (static shadowing only)
+			&& LightComponent->HasStaticShadowing()
+			&& !LightComponent->HasStaticLighting())
+		{
+			if (LightComponent->bAffectsWorld
+				&& LightComponent->CastShadows 
+				&& LightComponent->CastStaticShadows)
+			{
+				LightToOverlapMap.Add(new FLightAndChannel(LightComponent), TArray<FLightAndChannel*>());
+			}
+			else
+			{
+				// Reset the preview channel of stationary light components that shouldn't get a channel
+				// This is necessary to handle a light being newly disabled
+				LightComponent->PreviewShadowMapChannel = INDEX_NONE;
+
+#if WITH_EDITOR
+				LightComponent->UpdateLightSpriteTexture();
+#endif
+			}
+		}
+	}
+
+	// Build an array of overlapping lights
+	for (TMap<FLightAndChannel*, TArray<FLightAndChannel*> >::TIterator It(LightToOverlapMap); It; ++It)
+	{
+		ULightComponent* CurrentLight = It.Key()->Light;
+		TArray<FLightAndChannel*>& OverlappingLights = It.Value();
+
+		if (bAssignForLightingBuild)
+		{
+			// This should have happened during lighting invalidation at the beginning of the build anyway
+			CurrentLight->ShadowMapChannel = INDEX_NONE;
+		}
+
+		for (TMap<FLightAndChannel*, TArray<FLightAndChannel*> >::TIterator OtherIt(LightToOverlapMap); OtherIt; ++OtherIt)
+		{
+			ULightComponent* OtherLight = OtherIt.Key()->Light;
+
+			if (CurrentLight != OtherLight 
+				// Testing both directions because the spotlight <-> spotlight test is just cone vs bounding sphere
+				//@todo - more accurate spotlight <-> spotlight intersection
+				&& CurrentLight->AffectsBounds(FBoxSphereBounds(OtherLight->GetBoundingSphere()))
+				&& OtherLight->AffectsBounds(FBoxSphereBounds(CurrentLight->GetBoundingSphere())))
+			{
+				OverlappingLights.Add(OtherIt.Key());
+			}
+		}
+	}
+		
+	// Sort lights with the most overlapping lights first
+	LightToOverlapMap.ValueSort(FCompareLightsByArrayCount());
+
+	TMap<FLightAndChannel*, TArray<FLightAndChannel*> > SortedLightToOverlapMap;
+
+	// Add directional lights to the beginning so they always get channels
+	for (TMap<FLightAndChannel*, TArray<FLightAndChannel*> >::TIterator It(LightToOverlapMap); It; ++It)
+	{
+		FLightAndChannel* CurrentLight = It.Key();
+
+		if (CurrentLight->Light->GetLightType() == LightType_Directional)
+		{
+			SortedLightToOverlapMap.Add(It.Key(), It.Value());
+		}
+	}
+
+	// Add everything else, which has been sorted by descending overlaps
+	for (TMap<FLightAndChannel*, TArray<FLightAndChannel*> >::TIterator It(LightToOverlapMap); It; ++It)
+	{
+		FLightAndChannel* CurrentLight = It.Key();
+
+		if (CurrentLight->Light->GetLightType() != LightType_Directional)
+		{
+			SortedLightToOverlapMap.Add(It.Key(), It.Value());
+		}
+	}
+
+	// Go through lights and assign shadowmap channels
+	//@todo - retry with different ordering heuristics when it fails
+	for (TMap<FLightAndChannel*, TArray<FLightAndChannel*> >::TIterator It(SortedLightToOverlapMap); It; ++It)
+	{
+		bool bChannelUsed[4] = {0};
+		FLightAndChannel* CurrentLight = It.Key();
+		const TArray<FLightAndChannel*>& OverlappingLights = It.Value();
+
+		// Mark which channels have already been assigned to overlapping lights
+		for (int32 OverlappingIndex = 0; OverlappingIndex < OverlappingLights.Num(); OverlappingIndex++)
+		{
+			FLightAndChannel* OverlappingLight = OverlappingLights[OverlappingIndex];
+
+			if (OverlappingLight->Channel != INDEX_NONE)
+			{
+				bChannelUsed[OverlappingLight->Channel] = true;
+			}
+		}
+
+		// Use the lowest free channel
+		for (int32 ChannelIndex = 0; ChannelIndex < ARRAY_COUNT(bChannelUsed); ChannelIndex++)
+		{
+			if (!bChannelUsed[ChannelIndex])
+			{
+				CurrentLight->Channel = ChannelIndex;
+				break;
+			}
+		}
+	}
+
+	// Go through the assigned lights and update their render state and icon
+	for (TMap<FLightAndChannel*, TArray<FLightAndChannel*> >::TIterator It(SortedLightToOverlapMap); It; ++It)
+	{
+		FLightAndChannel* CurrentLight = It.Key();
+
+		if (CurrentLight->Light->PreviewShadowMapChannel != CurrentLight->Channel)
+		{
+			CurrentLight->Light->PreviewShadowMapChannel = CurrentLight->Channel;
+			CurrentLight->Light->MarkRenderStateDirty();
+		}
+
+#if WITH_EDITOR
+		CurrentLight->Light->UpdateLightSpriteTexture();
+#endif
+
+		if (bAssignForLightingBuild)
+		{
+			CurrentLight->Light->ShadowMapChannel = CurrentLight->Channel;
+
+			if (CurrentLight->Light->ShadowMapChannel == INDEX_NONE)
+			{
+				FMessageLog("LightingResults").Error()
+					->AddToken(FUObjectToken::Create(CurrentLight->Light->GetOwner()))
+					->AddToken(FTextToken::Create( NSLOCTEXT("Lightmass", "LightmassError_FailedToAllocateShadowmapChannel", "Severe performance loss: Failed to allocate shadowmap channel for stationary light due to overlap - light will fall back to dynamic shadows!") ) );
+			}
+		}
+
+		delete CurrentLight;
+	}
+}
+
+static void ToggleLight(const TArray<FString>& Args)
+{
+	for (TObjectIterator<ULightComponent> It; It; ++It)
+	{
+		ULightComponent* Light = *It;
+		if (Light->Mobility != EComponentMobility::Static)
+		{
+			FString LightName = (Light->GetOwner() ? Light->GetOwner()->GetFName() : Light->GetFName()).ToString();
+			for (int32 ArgIndex = 0; ArgIndex < Args.Num(); ++ArgIndex)
+			{
+				const FString& ToggleName = Args[ArgIndex];
+				if (LightName.Contains(ToggleName) )
+				{
+					Light->ToggleVisibility(/*bPropagateToChildren=*/ false);
+					UE_LOG(LogConsoleResponse,Display,TEXT("Now%svisible: %s"),
+						Light->IsVisible() ? TEXT("") : TEXT(" not "),
+						*Light->GetFullName()
+						);
+					break;
+				}
+			}
+		}
+	}
+}
+
+static FAutoConsoleCommand ToggleLightCmd(
+	TEXT("ToggleLight"),
+	TEXT("Toggles all lights whose name contains the specified string"),
+	FConsoleCommandWithArgsDelegate::CreateStatic(ToggleLight)
+	);
diff --git a/Engine/Source/Runtime/Engine/Private/Components/MaterialBillboardComponent.cpp b/Engine/Source/Runtime/Engine/Private/Components/MaterialBillboardComponent.cpp
new file mode 100644
--- /dev/null
+++ b/Engine/Source/Runtime/Engine/Private/Components/MaterialBillboardComponent.cpp
@@ -0,0 +1,308 @@
+// Copyright 1998-2015 Epic Games, Inc. All Rights Reserved.
+
+#include "EnginePrivate.h"
+#include "LevelUtils.h"
+#include "LocalVertexFactory.h"
+#include "Components/MaterialBillboardComponent.h"
+
+/** A material sprite vertex. */
+struct FMaterialSpriteVertex
+{
+	FVector Position;
+	FPackedNormal TangentX;
+	FPackedNormal TangentZ;
+	FColor Color;
+	FVector2D TexCoords;
+};
+
+/** A dummy vertex buffer used to give the FMaterialSpriteVertexFactory something to reference as a stream source. */
+class FMaterialSpriteVertexBuffer : public FVertexBuffer
+{
+public:
+
+	virtual void InitRHI() override
+	{
+		FRHIResourceCreateInfo CreateInfo;
+		VertexBufferRHI = RHICreateVertexBuffer(sizeof(FMaterialSpriteVertex),BUF_Static,CreateInfo);
+	}
+};
+static TGlobalResource<FMaterialSpriteVertexBuffer> GDummyMaterialSpriteVertexBuffer;
+
+/** The vertex factory used to draw material sprites. */
+class FMaterialSpriteVertexFactory : public FLocalVertexFactory
+{
+public:
+
+	FMaterialSpriteVertexFactory()
+	{
+		FLocalVertexFactory::DataType Data;
+
+		Data.PositionComponent = FVertexStreamComponent(
+			&GDummyMaterialSpriteVertexBuffer,
+			STRUCT_OFFSET(FMaterialSpriteVertex,Position),
+			sizeof(FMaterialSpriteVertex),
+			VET_Float3
+			);
+		Data.TangentBasisComponents[0] = FVertexStreamComponent(
+			&GDummyMaterialSpriteVertexBuffer,
+			STRUCT_OFFSET(FMaterialSpriteVertex,TangentX),
+			sizeof(FMaterialSpriteVertex),
+			VET_PackedNormal
+			);
+		Data.TangentBasisComponents[1] = FVertexStreamComponent(
+			&GDummyMaterialSpriteVertexBuffer,
+			STRUCT_OFFSET(FMaterialSpriteVertex,TangentZ),
+			sizeof(FMaterialSpriteVertex),
+			VET_PackedNormal
+			);
+		Data.ColorComponent = FVertexStreamComponent(
+			&GDummyMaterialSpriteVertexBuffer,
+			STRUCT_OFFSET(FMaterialSpriteVertex,Color),
+			sizeof(FMaterialSpriteVertex),
+			VET_Color
+			);
+		Data.TextureCoordinates.Empty();
+		Data.TextureCoordinates.Add(FVertexStreamComponent(
+			&GDummyMaterialSpriteVertexBuffer,
+			STRUCT_OFFSET(FMaterialSpriteVertex,TexCoords),
+			sizeof(FMaterialSpriteVertex),
+			VET_Float2
+			));
+
+		SetData(Data);
+	}
+};
+static TGlobalResource<FMaterialSpriteVertexFactory> GMaterialSpriteVertexFactory;
+
+class FMaterialSpriteVertexArray : public FOneFrameResource
+{
+public:
+	TArray<FMaterialSpriteVertex, TInlineAllocator<4> > Vertices;
+};
+
+/** Represents a sprite to the scene manager. */
+class FMaterialSpriteSceneProxy : public FPrimitiveSceneProxy
+{
+public:
+
+	/** Initialization constructor. */
+	FMaterialSpriteSceneProxy(const UMaterialBillboardComponent* InComponent)
+	: FPrimitiveSceneProxy(InComponent)
+	, Elements(InComponent->Elements)
+	, BaseColor(FLinearColor::White)
+	{
+		AActor* Owner = InComponent->GetOwner();
+		if (Owner)
+		{
+			// Level colorization
+			ULevel* Level = Owner->GetLevel();
+			ULevelStreaming* LevelStreaming = FLevelUtils::FindStreamingLevel( Level );
+			if ( LevelStreaming )
+			{
+				// Selection takes priority over level coloration.
+				LevelColor = LevelStreaming->LevelColor;
+			}
+		}
+
+		for (int32 ElementIndex = 0; ElementIndex < Elements.Num(); ElementIndex++)
+		{
+			UMaterialInterface* Material = Elements[ElementIndex].Material;
+			if (Material)
+			{
+				MaterialRelevance |= Material->GetRelevance(GetScene().GetFeatureLevel());
+			}
+		}
+
+		FColor NewPropertyColor;
+		GEngine->GetPropertyColorationColor( (UObject*)InComponent, NewPropertyColor );
+		PropertyColor = NewPropertyColor;
+	}
+
+	virtual void GetDynamicMeshElements(const TArray<const FSceneView*>& Views, const FSceneViewFamily& ViewFamily, uint32 VisibilityMap, FMeshElementCollector& Collector) const override
+	{
+		QUICK_SCOPE_CYCLE_COUNTER( STAT_MaterialSpriteSceneProxy_GetDynamicMeshElements );
+
+		for (int32 ViewIndex = 0; ViewIndex < Views.Num(); ViewIndex++)
+		{
+			if (VisibilityMap & (1 << ViewIndex))
+			{
+				const FSceneView* View = Views[ViewIndex];
+
+				const bool bIsWireframe = View->Family->EngineShowFlags.Wireframe;
+				// Determine the position of the source
+				const FVector SourcePosition = GetLocalToWorld().GetOrigin();
+				const FVector CameraToSource = View->ViewMatrices.ViewOrigin - SourcePosition;
+				const float DistanceToSource = CameraToSource.Size();
+
+				const FVector CameraUp      = -View->InvViewProjectionMatrix.TransformVector(FVector(1.0f,0.0f,0.0f));
+				const FVector CameraRight   = -View->InvViewProjectionMatrix.TransformVector(FVector(0.0f,1.0f,0.0f));
+				const FVector CameraForward = -View->InvViewProjectionMatrix.TransformVector(FVector(0.0f,0.0f,1.0f));
+				const FMatrix WorldToLocal = GetLocalToWorld().InverseFast();
+				const FVector LocalCameraUp = WorldToLocal.TransformVector(CameraUp);
+				const FVector LocalCameraRight = WorldToLocal.TransformVector(CameraRight);
+				const FVector LocalCameraForward = WorldToLocal.TransformVector(CameraForward);
+
+				// Draw the elements ordered so the last is on top of the first.
+				for(int32 ElementIndex = 0;ElementIndex < Elements.Num();++ElementIndex)
+				{
+					const FMaterialSpriteElement& Element = Elements[ElementIndex];
+
+					if(Element.Material)
+					{
+						// Evaluate the size of the sprite.
+						float SizeX = Element.BaseSizeX;
+						float SizeY = Element.BaseSizeY;
+						if(Element.DistanceToSizeCurve)
+						{
+							const float SizeFactor = Element.DistanceToSizeCurve->GetFloatValue(DistanceToSource);
+							SizeX *= SizeFactor;
+							SizeY *= SizeFactor;
+						}
+						
+						// Convert the size into world-space.
+						const float W = View->ViewProjectionMatrix.TransformPosition(SourcePosition).W;
+						const float AspectRatio = CameraRight.Size() / CameraUp.Size();
+						const float WorldSizeX = Element.bSizeIsInScreenSpace ? (SizeX               * W) : (SizeX / CameraRight.Size());
+						const float WorldSizeY = Element.bSizeIsInScreenSpace ? (SizeY * AspectRatio * W) : (SizeY / CameraUp.Size());
+			
+						// Evaluate the color/opacity of the sprite.
+						FLinearColor Color = BaseColor;
+						if(Element.DistanceToOpacityCurve)
+						{
+							Color.A *= Element.DistanceToOpacityCurve->GetFloatValue(DistanceToSource);
+						}
+
+						// Set up the sprite vertex attributes that are constant across the sprite.
+						FMaterialSpriteVertexArray& VertexArray = Collector.AllocateOneFrameResource<FMaterialSpriteVertexArray>();
+						VertexArray.Vertices.Empty(4);
+						VertexArray.Vertices.AddUninitialized(4);
+
+						for(uint32 VertexIndex = 0;VertexIndex < 4;++VertexIndex)
+						{
+							VertexArray.Vertices[VertexIndex].Color = Color;
+							VertexArray.Vertices[VertexIndex].TangentX = FPackedNormal(LocalCameraRight.GetSafeNormal());
+							VertexArray.Vertices[VertexIndex].TangentZ = FPackedNormal(-LocalCameraForward.GetSafeNormal());
+						}
+
+						// Set up the sprite vertex positions and texture coordinates.
+						VertexArray.Vertices[0].Position  = -WorldSizeX * LocalCameraRight + +WorldSizeY * LocalCameraUp;
+						VertexArray.Vertices[0].TexCoords = FVector2D(0,0);
+						VertexArray.Vertices[1].Position  = +WorldSizeX * LocalCameraRight + +WorldSizeY * LocalCameraUp;
+						VertexArray.Vertices[1].TexCoords = FVector2D(0,1);
+						VertexArray.Vertices[2].Position  = -WorldSizeX * LocalCameraRight + -WorldSizeY * LocalCameraUp;
+						VertexArray.Vertices[2].TexCoords = FVector2D(1,0);
+						VertexArray.Vertices[3].Position  = +WorldSizeX * LocalCameraRight + -WorldSizeY * LocalCameraUp;
+						VertexArray.Vertices[3].TexCoords = FVector2D(1,1);
+			
+						// Set up the FMeshElement.
+						FMeshBatch& Mesh = Collector.AllocateMesh();
+						Mesh.UseDynamicData      = true;
+						Mesh.DynamicVertexData   = VertexArray.Vertices.GetData();
+						Mesh.DynamicVertexStride = sizeof(FMaterialSpriteVertex);
+
+						Mesh.VertexFactory           = &GMaterialSpriteVertexFactory;
+						Mesh.MaterialRenderProxy     = Element.Material->GetRenderProxy((View->Family->EngineShowFlags.Selection) && IsSelected(),IsHovered());
+						Mesh.LCI                     = NULL;
+						Mesh.ReverseCulling          = IsLocalToWorldDeterminantNegative() ? true : false;
+						Mesh.CastShadow              = false;
+						Mesh.DepthPriorityGroup      = (ESceneDepthPriorityGroup)GetDepthPriorityGroup(View);
+						Mesh.Type                    = PT_TriangleStrip;
+						Mesh.bDisableBackfaceCulling = true;
+
+						// Set up the FMeshBatchElement.
+						FMeshBatchElement& BatchElement = Mesh.Elements[0];
+						BatchElement.IndexBuffer            = NULL;
+						BatchElement.DynamicIndexData       = NULL;
+						BatchElement.DynamicIndexStride     = 0;
+						BatchElement.FirstIndex             = 0;
+						BatchElement.MinVertexIndex         = 0;
+						BatchElement.MaxVertexIndex         = 3;
+						BatchElement.PrimitiveUniformBuffer = GetUniformBuffer();
+						BatchElement.NumPrimitives          = 2;
+
+						Mesh.bCanApplyViewModeOverrides = true;
+						Mesh.bUseWireframeSelectionColoring = IsSelected();
+
+						Collector.AddMesh(ViewIndex, Mesh);
+					}
+				}
+			}
+		}
+	}
+
+	virtual FPrimitiveViewRelevance GetViewRelevance(const FSceneView* View)
+	{
+		bool bVisible = View->Family->EngineShowFlags.BillboardSprites;
+		FPrimitiveViewRelevance Result;
+		Result.bDrawRelevance = IsShown(View);
+		Result.bDynamicRelevance = true;
+		Result.bShadowRelevance = IsShadowCast(View);
+		MaterialRelevance.SetPrimitiveViewRelevance(Result);
+		return Result;
+	}
+	virtual bool CanBeOccluded() const { return !MaterialRelevance.bDisableDepthTest; }
+	virtual uint32 GetMemoryFootprint() const { return sizeof(*this) + GetAllocatedSize(); }
+	uint32 GetAllocatedSize() const { return FPrimitiveSceneProxy::GetAllocatedSize(); }
+
+private:
+	TArray<FMaterialSpriteElement> Elements;
+	FMaterialRelevance MaterialRelevance;
+	FColor BaseColor;
+};
+
+UMaterialBillboardComponent::UMaterialBillboardComponent(const FObjectInitializer& ObjectInitializer)
+: Super(ObjectInitializer)
+{
+	SetCollisionProfileName(UCollisionProfile::NoCollision_ProfileName);
+}
+
+FPrimitiveSceneProxy* UMaterialBillboardComponent::CreateSceneProxy()
+{
+	return new FMaterialSpriteSceneProxy(this);
+}
+
+FBoxSphereBounds UMaterialBillboardComponent::CalcBounds(const FTransform& LocalToWorld) const
+{
+	const float BoundsSize = 32.0f;
+	return FBoxSphereBounds(LocalToWorld.GetLocation(),FVector(BoundsSize,BoundsSize,BoundsSize),FMath::Sqrt(3.0f * FMath::Square(BoundsSize)));
+}
+
+void UMaterialBillboardComponent::AddElement(
+	class UMaterialInterface* Material,
+	UCurveFloat* DistanceToOpacityCurve,
+	bool bSizeIsInScreenSpace,
+	float BaseSizeX,
+	float BaseSizeY,
+	UCurveFloat* DistanceToSizeCurve
+	)
+{
+	FMaterialSpriteElement* Element = new(Elements) FMaterialSpriteElement;
+	Element->Material = Material;
+	Element->DistanceToOpacityCurve = DistanceToOpacityCurve;
+	Element->bSizeIsInScreenSpace = bSizeIsInScreenSpace;
+	Element->BaseSizeX = BaseSizeX;
+	Element->BaseSizeY = BaseSizeY;
+	Element->DistanceToSizeCurve = DistanceToSizeCurve;
+
+	MarkRenderStateDirty();
+}
+
+UMaterialInterface* UMaterialBillboardComponent::GetMaterial(int32 Index) const
+{
+	UMaterialInterface* ResultMI = nullptr;
+	if (Elements.IsValidIndex(Index))
+	{
+		ResultMI = Elements[Index].Material;
+	}
+	return ResultMI;
+}
+
+void UMaterialBillboardComponent::SetMaterial(int32 ElementIndex, class UMaterialInterface* Material)
+{
+	if (Elements.IsValidIndex(ElementIndex))
+	{
+		Elements[ElementIndex].Material = Material;
+
+		MarkRenderStateDirty();
+	}
+}
\ No newline at end of file
diff --git a/Engine/Source/Runtime/Engine/Private/Components/ModelComponent.cpp b/Engine/Source/Runtime/Engine/Private/Components/ModelComponent.cpp
new file mode 100644
--- /dev/null
+++ b/Engine/Source/Runtime/Engine/Private/Components/ModelComponent.cpp
@@ -0,0 +1,589 @@
+// Copyright 1998-2015 Epic Games, Inc. All Rights Reserved.
+
+/*=============================================================================
+	ModelComponent.cpp: Model component implementation
+=============================================================================*/
+
+#include "EnginePrivate.h"
+#include "Components/ModelComponent.h"
+#include "Model.h"
+#include "LightMap.h"
+#include "ShadowMap.h"
+#include "Components/ModelComponent.h"
+#include "PhysicsEngine/PhysicsSettings.h"
+
+FModelElement::FModelElement(UModelComponent* InComponent,UMaterialInterface* InMaterial):
+	Component(InComponent),
+	Material(InMaterial),
+	IndexBuffer(NULL),
+	FirstIndex(0),
+	NumTriangles(0),
+	MinVertexIndex(0),
+	MaxVertexIndex(0),
+	BoundingBox(ForceInitToZero)
+{}
+
+FModelElement::FModelElement():
+	Component(NULL), 
+	Material(NULL), 
+	IndexBuffer(NULL), 
+	FirstIndex(0), 
+	NumTriangles(0), 
+	MinVertexIndex(0), 
+	MaxVertexIndex(0), 
+	BoundingBox(ForceInitToZero)
+{}
+
+FModelElement::~FModelElement()
+{}
+
+
+/**
+ * Serializer.
+ */
+FArchive& operator<<(FArchive& Ar,FModelElement& Element)
+{
+	Ar << Element.LightMap;
+	Ar << Element.ShadowMap;
+	
+	Ar << (UObject*&)Element.Component << (UObject*&)Element.Material << Element.Nodes;
+	Ar << Element.IrrelevantLights;
+	return Ar;
+}
+
+UModelComponent::UModelComponent(const FObjectInitializer& ObjectInitializer)
+	: Super(ObjectInitializer)
+{
+	CastShadow = true;
+	bUseAsOccluder = true;
+	Mobility = EComponentMobility::Static;
+	bGenerateOverlapEvents = false;
+
+	SetCollisionProfileName(UCollisionProfile::BlockAll_ProfileName);
+}
+
+#if WITH_EDITOR
+void UModelComponent::InitializeModelComponent(UModel* InModel, uint16 InComponentIndex, uint32 MaskedSurfaceFlags, const TArray<uint16>& InNodes)
+{
+	Model = InModel;
+	ComponentIndex = InComponentIndex;
+	Nodes = InNodes;
+
+	// Model components are transacted.
+	SetFlags(RF_Transactional);
+
+	GenerateElements(true);
+
+	CastShadow = true;
+	bUseAsOccluder = true;
+	Mobility = EComponentMobility::Static;
+	bGenerateOverlapEvents = false;
+
+	SetCollisionProfileName(UCollisionProfile::BlockAll_ProfileName);
+}
+#endif // WITH_EDITOR
+
+
+void UModelComponent::AddReferencedObjects(UObject* InThis, FReferenceCollector& Collector)
+{
+	UModelComponent* This = CastChecked<UModelComponent>(InThis);
+	Collector.AddReferencedObject( This->Model, This );
+	for( int32 ElementIndex=0; ElementIndex < This->Elements.Num(); ElementIndex++ )
+	{
+		FModelElement& Element = This->Elements[ElementIndex];
+		Collector.AddReferencedObject( Element.Component, This );
+		Collector.AddReferencedObject( Element.Material, This );
+		if(Element.LightMap != NULL)
+		{
+			Element.LightMap->AddReferencedObjects(Collector);
+		}
+		if(Element.ShadowMap != NULL)
+		{
+			Element.ShadowMap->AddReferencedObjects(Collector);
+		}
+	}
+	Super::AddReferencedObjects( This, Collector );
+}
+
+void UModelComponent::CommitSurfaces()
+{
+	TArray<int32> InvalidElements;
+
+	// Find nodes that are from surfaces which have been invalidated.
+	TMap<uint16,FModelElement*> InvalidNodes;
+	for(int32 ElementIndex = 0;ElementIndex < Elements.Num();ElementIndex++)
+	{
+		FModelElement& Element = Elements[ElementIndex];
+		TArray<uint16> NewNodes;
+		for(int32 NodeIndex = 0;NodeIndex < Element.Nodes.Num();NodeIndex++)
+		{
+			FBspNode& Node = Model->Nodes[Element.Nodes[NodeIndex]];
+			FBspSurf& Surf = Model->Surfs[Node.iSurf];
+			if(Surf.Material != Element.Material)
+			{
+				// This node's material changed, remove it from the element and put it on the invalid node list.
+				InvalidNodes.Add(Element.Nodes[NodeIndex],&Element);
+
+				// Mark the node's original element as being invalid.
+				InvalidElements.AddUnique(ElementIndex);
+			}
+			else
+			{
+				NewNodes.Add(Element.Nodes[NodeIndex]);
+			}
+		}
+		Element.Nodes = NewNodes;
+	}
+
+	// Reassign the invalid nodes to appropriate mesh elements.
+	for(TMap<uint16,FModelElement*>::TConstIterator It(InvalidNodes);It;++It)
+	{
+		FBspNode& Node = Model->Nodes[It.Key()];
+		FBspSurf& Surf = Model->Surfs[Node.iSurf];
+		FModelElement* OldElement = It.Value();
+
+		// Find an element which has the same material and lights as the invalid node.
+		FModelElement* NewElement = NULL;
+		for(int32 ElementIndex = 0;ElementIndex < Elements.Num();ElementIndex++)
+		{
+			FModelElement& Element = Elements[ElementIndex];
+			if(Element.Material != Surf.Material)
+				continue;
+			if(Element.LightMap != OldElement->LightMap)
+				continue;
+			if(Element.ShadowMap != OldElement->ShadowMap)
+				continue;
+			if(Element.IrrelevantLights != OldElement->IrrelevantLights)
+				continue;
+			// This element's material and lights match the node.
+			NewElement = &Element;
+		}
+
+		// If no matching element was found, create a new element.
+		if(!NewElement)
+		{
+			NewElement = new(Elements) FModelElement(this,Surf.Material);
+			NewElement->LightMap = OldElement->LightMap;
+			NewElement->ShadowMap = OldElement->ShadowMap;
+			NewElement->IrrelevantLights = OldElement->IrrelevantLights;
+		}
+
+		NewElement->Nodes.Add(It.Key());
+		InvalidElements.AddUnique(Elements.Num() - 1);
+	}
+
+	// Rebuild the render data for the elements which have changed.
+	BuildRenderData();
+
+	ShrinkElements();
+
+#if WITH_EDITOR
+	// Need to update collision data as well
+	InvalidateCollisionData();
+#endif	//WITH_EDITOR
+}
+
+void UModelComponent::ShrinkElements()
+{
+	// Find elements which have no nodes remaining.
+	for(int32 ElementIndex = 0;ElementIndex < Elements.Num();ElementIndex++)
+	{
+		FModelElement& Element = Elements[ElementIndex];
+		if(!Element.Nodes.Num())
+		{
+			// This element contains no nodes, remove it.
+			Elements.RemoveAt(ElementIndex--);
+		}
+	}
+}
+
+void UModelComponent::Serialize(FArchive& Ar)
+{
+	Super::Serialize(Ar);
+
+	Ar << Model;
+
+	if( Ar.IsLoading() && Ar.UE4Ver() < VER_UE4_REMOVE_ZONES_FROM_MODEL)
+	{
+		int32 DummyZoneIndex;
+		Ar << DummyZoneIndex << Elements;
+	}
+	else
+	{
+		Ar << Elements;
+	}
+
+	Ar << ComponentIndex << Nodes;
+}
+
+void UModelComponent::PostLoad()
+{
+	Super::PostLoad();
+
+	// Fix for old model components which weren't created with transactional flag.
+	SetFlags( RF_Transactional );
+
+	// BuildRenderData relies on the model having been post-loaded, so we ensure this by calling ConditionalPostLoad.
+	check(Model);
+	Model->ConditionalPostLoad();
+
+	// Initialize model elements' index buffers (required for generating ddc data).
+	BuildRenderData();
+
+	// Older content without a body setup
+	if(ModelBodySetup == NULL)
+	{
+		CreateModelBodySetup();
+		check(ModelBodySetup != NULL);
+		ModelBodySetup->CreatePhysicsMeshes(); // Want to do this in PostLoad before model vertex buffer is discarded
+	}
+
+	// Stop existing ModelComponents from generating mirrored collision mesh
+	if ((GetLinkerUE4Version() < VER_UE4_NO_MIRROR_BRUSH_MODEL_COLLISION) && (ModelBodySetup != NULL))
+	{
+		ModelBodySetup->bGenerateMirroredCollision = false;
+	}
+
+	ModelBodySetup->bDoubleSidedGeometry = true;	//saved content wants this to be true
+}
+
+#if WITH_EDITOR
+void UModelComponent::PostEditUndo()
+{
+	// Rebuild the component's render data after applying a transaction to it.
+	ULevel* Level = GetTypedOuter<ULevel>();
+	if(ensure(Level))
+	{
+		Level->InvalidateModelSurface();
+		Level->CommitModelSurfaces();
+	}
+	Super::PostEditUndo();
+}
+#endif // WITH_EDITOR
+
+SIZE_T UModelComponent::GetResourceSize(EResourceSizeMode::Type Mode)
+{
+	SIZE_T ResSize = Super::GetResourceSize(Mode);
+
+	// Count the bodysetup we own as well for 'inclusive' stats
+	if((Mode == EResourceSizeMode::Inclusive) && (ModelBodySetup != NULL))
+	{
+		ResSize += ModelBodySetup->GetResourceSize(Mode);
+	}
+
+	return ResSize;
+}
+
+
+void UModelComponent::GetUsedMaterials(TArray<UMaterialInterface*>& OutMaterials) const
+{
+	for( int32 ElementIndex = 0; ElementIndex < Elements.Num(); ++ElementIndex )
+	{
+		const FModelElement& Element = Elements[ ElementIndex ];
+		if( Element.Material )
+		{
+			OutMaterials.Add( Element.Material );
+		}
+	}
+}
+
+int32 UModelComponent::GetNumMaterials() const
+{
+	return Elements.Num();
+}
+
+
+UMaterialInterface* UModelComponent::GetMaterial(int32 MaterialIndex) const
+{
+	UMaterialInterface* Material = NULL;
+
+	if(MaterialIndex < Elements.Num())
+	{
+		return  Elements[MaterialIndex].Material;
+	}
+
+	return Material;
+}
+
+#if WITH_EDITOR
+void UModelComponent::SelectAllSurfaces()
+{
+	check(Model);
+	for( int32 NodeIndex=0; NodeIndex<Nodes.Num(); NodeIndex++ )
+	{
+		FBspNode& Node = Model->Nodes[Nodes[NodeIndex]];
+		FBspSurf& Surf = Model->Surfs[Node.iSurf];
+		Model->ModifySurf( Node.iSurf, 0 );
+		Surf.PolyFlags |= PF_Selected;
+	}
+}
+#endif // WITH_EDITOR
+
+void UModelComponent::GetStreamingTextureInfo(TArray<FStreamingTexturePrimitiveInfo>& OutStreamingTextures) const
+{
+	if( Model )
+	{
+		// Build a map from surface indices to node indices for nodes in this component.
+		TMultiMap<int32,int32> SurfToNodeMap;
+		for(int32 NodeIndex = 0;NodeIndex < Nodes.Num();NodeIndex++)
+		{
+			const FBspNode& Node = Model->Nodes[Nodes[NodeIndex]];
+			SurfToNodeMap.Add(Node.iSurf,Nodes[NodeIndex]);
+		}
+
+		TArray<int32> SurfaceNodes;
+		TArray<FVector> SurfaceVertices;
+		for(int32 SurfaceIndex = 0;SurfaceIndex < Model->Surfs.Num();SurfaceIndex++)
+		{
+			// Check if the surface has nodes in this component.
+			SurfaceNodes.Reset();
+			SurfToNodeMap.MultiFind(SurfaceIndex,SurfaceNodes);
+			if(SurfaceNodes.Num())
+			{
+				const FBspSurf& Surf = Model->Surfs[SurfaceIndex];
+
+				// Compute a bounding sphere for the surface's nodes.
+				SurfaceVertices.Reset();
+				for(int32 NodeIndex = 0;NodeIndex < SurfaceNodes.Num();NodeIndex++)
+				{
+					const FBspNode& Node = Model->Nodes[SurfaceNodes[NodeIndex]];
+					for(int32 VertexIndex = 0;VertexIndex < Node.NumVertices;VertexIndex++)
+					{
+						const FVector WorldVertex = ComponentToWorld.TransformPosition(Model->Points[Model->Verts[Node.iVertPool + VertexIndex].pVertex]);
+						SurfaceVertices.Add(WorldVertex);
+					}
+				}
+				const FSphere SurfaceBoundingSphere(SurfaceVertices.GetData(),SurfaceVertices.Num());
+
+				// Compute the surface's texture scaling factor.
+				const float BspTexelsPerNormalizedTexel = UModel::GetGlobalBSPTexelScale();
+				const float WorldUnitsPerBspTexel = FMath::Max( Model->Vectors[Surf.vTextureU].Size(), Model->Vectors[Surf.vTextureV].Size() );
+				const float TexelFactor = BspTexelsPerNormalizedTexel / WorldUnitsPerBspTexel;
+
+				// Determine the material applied to the surface.
+				UMaterialInterface* Material = Surf.Material;
+				if(!Material)
+				{
+					Material = UMaterial::GetDefaultMaterial(MD_Surface);
+				}
+
+				// Enumerate the textures used by the surface's material.
+				TArray<UTexture*> Textures;
+				
+				Material->GetUsedTextures(Textures, EMaterialQualityLevel::Num, false, GMaxRHIFeatureLevel, true);
+
+				// Add each texture to the output with the appropriate parameters.
+				for(int32 TextureIndex = 0;TextureIndex < Textures.Num();TextureIndex++)
+				{
+					FStreamingTexturePrimitiveInfo& StreamingTexture = *new(OutStreamingTextures) FStreamingTexturePrimitiveInfo;
+					StreamingTexture.Bounds = SurfaceBoundingSphere;
+					StreamingTexture.TexelFactor = TexelFactor;
+					StreamingTexture.Texture = Textures[TextureIndex];
+				}
+			}
+		}
+	}
+}
+
+#if WITH_EDITOR
+
+bool UModelComponent::GenerateElements(bool bBuildRenderData)
+{
+	Elements.Empty();
+	if (GIsEditor == false)
+	{
+		TMap<UMaterialInterface*,FModelElement*> MaterialToElementMap;
+
+		// Find the BSP nodes which are in this component's zone.
+		for(int32 NodeIndex = 0;NodeIndex < Nodes.Num();NodeIndex++)
+		{
+			FBspNode& Node = Model->Nodes[Nodes[NodeIndex]];
+			FBspSurf& Surf = Model->Surfs[Node.iSurf];
+
+			// Find an element with the same material as this node.
+			FModelElement* Element = MaterialToElementMap.FindRef(Surf.Material);
+			if(!Element)
+			{
+				// If there's no matching element, create a new element.
+				Element = MaterialToElementMap.Add(Surf.Material,new(Elements) FModelElement(this,Surf.Material));
+			}
+
+			// Add the node to the element.
+			Element->Nodes.Add(Nodes[NodeIndex]);
+		}
+	}
+	else
+	{
+		typedef TMap<UMaterialInterface*,FModelElement*> TMaterialElementMap;
+		typedef TMap<uint32,TMaterialElementMap*> TResolutionToMaterialElementMap;
+		TMap<FNodeGroup*,TResolutionToMaterialElementMap*> NodeGroupToResolutionToMaterialElementMap;
+
+		// Find the BSP nodes which are in this component's zone.
+		for (int32 NodeIndex = 0;NodeIndex < Nodes.Num();NodeIndex++)
+		{
+			FBspNode& Node = Model->Nodes[Nodes[NodeIndex]];
+			FBspSurf& Surf = Model->Surfs[Node.iSurf];
+
+			// Find the node group it belong to
+			FNodeGroup* FoundNodeGroup = NULL;
+			// find the NodeGroup that this node went into, and get all of its node
+			for (TMap<int32, FNodeGroup*>::TIterator It(Model->NodeGroups); It && (FoundNodeGroup == NULL); ++It)
+			{
+				FNodeGroup* NodeGroup = It.Value();
+				for (int32 InnerNodeIdx = 0; InnerNodeIdx < NodeGroup->Nodes.Num(); InnerNodeIdx++)
+				{
+					if (NodeGroup->Nodes[InnerNodeIdx] == Nodes[NodeIndex])
+					{
+						FoundNodeGroup = NodeGroup;
+						break;
+					}
+				}
+			}
+
+			TResolutionToMaterialElementMap* ResToMaterialElementMap = NodeGroupToResolutionToMaterialElementMap.FindRef(FoundNodeGroup);
+			if (!ResToMaterialElementMap)
+			{
+				TResolutionToMaterialElementMap* TempResToMaterialElementMap = new TResolutionToMaterialElementMap();
+				NodeGroupToResolutionToMaterialElementMap.Add(FoundNodeGroup, TempResToMaterialElementMap);
+				ResToMaterialElementMap = NodeGroupToResolutionToMaterialElementMap.FindRef(FoundNodeGroup);
+				check(ResToMaterialElementMap);
+			}
+
+			// Find an element with the same material as this node.
+			uint32 ResValue = FMath::TruncToInt(Surf.LightMapScale);
+			TMap<UMaterialInterface*,FModelElement*>* MaterialElementMap = ResToMaterialElementMap->FindRef(ResValue);
+			if (!MaterialElementMap)
+			{
+				TMap<UMaterialInterface*,FModelElement*>* TempMap = new TMap<UMaterialInterface*,FModelElement*>();
+				ResToMaterialElementMap->Add(ResValue, TempMap);
+				MaterialElementMap = ResToMaterialElementMap->FindRef(ResValue);
+				check(MaterialElementMap);
+			}
+
+			FModelElement* Element = MaterialElementMap->FindRef(Surf.Material);
+			if (!Element)
+			{
+				// If there's no matching element, create a new element.
+				Element = MaterialElementMap->Add(Surf.Material,new(Elements) FModelElement(this,Surf.Material));
+			}
+			check(Element);
+
+			// Add the node to the element.
+			Element->Nodes.Add(Nodes[NodeIndex]);
+		}
+	}
+
+	if (bBuildRenderData == true)
+	{
+		BuildRenderData();
+	}
+
+	return true;
+}
+#endif // WITH_EDITOR
+
+void UModelComponent::CopyElementsFrom(UModelComponent* SrcComponent)
+{
+	Elements.Empty();
+	for (int32 ElementIndex = 0; ElementIndex < SrcComponent->Elements.Num(); ++ElementIndex)
+	{
+		FModelElement& SrcElement = SrcComponent->Elements[ElementIndex];
+		FModelElement& DestElement = *new(Elements) FModelElement(SrcElement);
+		DestElement.Component = this;
+	}
+
+	if(ModelBodySetup && SrcComponent->ModelBodySetup)
+	{
+		ModelBodySetup->BodySetupGuid = SrcComponent->ModelBodySetup->BodySetupGuid;
+	}
+}
+
+void UModelComponent::CreateModelBodySetup()
+{
+	if(ModelBodySetup == NULL)
+	{
+		ModelBodySetup = NewObject<UBodySetup>(this);
+		check(ModelBodySetup);
+		ModelBodySetup->BodySetupGuid = FGuid::NewGuid();
+	}
+
+	ModelBodySetup->CollisionTraceFlag = CTF_UseComplexAsSimple;
+	ModelBodySetup->bGenerateMirroredCollision = false;
+	ModelBodySetup->bDoubleSidedGeometry = true;
+}
+
+#if WITH_EDITOR
+void UModelComponent::InvalidateCollisionData()
+{
+	// Make sure we have a body setup..
+	CreateModelBodySetup(); 
+	check(ModelBodySetup != NULL);
+
+	UE_LOG(LogPhysics, Log, TEXT("Invalidate ModelComponent: %s"), *GetPathName());
+
+	// Then give it a new GUID
+	ModelBodySetup->InvalidatePhysicsData();
+}
+#endif // WITH_EDITOR
+
+bool UModelComponent::GetPhysicsTriMeshData(struct FTriMeshCollisionData* CollisionData, bool InUseAllTriData)
+{
+	check(Model);
+	int32 nBadArea = 0;
+	const float AreaThreshold = UPhysicsSettings::Get()->TriangleMeshTriangleMinAreaThreshold;
+
+	const int32 NumVerts = Model->VertexBuffer.Vertices.Num();
+	CollisionData->Vertices.AddUninitialized(NumVerts);
+	for(int32 i=0; i<NumVerts; i++)
+	{
+		CollisionData->Vertices[i] = Model->VertexBuffer.Vertices[i].Position;
+	}
+
+	for(int32 ElementIndex = 0;ElementIndex < Elements.Num();ElementIndex++)
+	{
+		FModelElement& Element = Elements[ElementIndex];
+		FRawIndexBuffer16or32* IndexBuffer = (FRawIndexBuffer16or32*)Element.IndexBuffer; // @Bad to assume its always one of these?
+
+		for(uint32 TriIdx=0; TriIdx<Element.NumTriangles; TriIdx++)
+		{
+			FTriIndices Triangle;
+
+			Triangle.v0 = IndexBuffer->Indices[Element.FirstIndex + (TriIdx*3) + 0];
+			Triangle.v1 = IndexBuffer->Indices[Element.FirstIndex + (TriIdx*3) + 1];
+			Triangle.v2 = IndexBuffer->Indices[Element.FirstIndex + (TriIdx*3) + 2];
+
+			if (AreaThreshold >= 0.f)
+			{
+				const FVector V0 = Model->VertexBuffer.Vertices[Triangle.v0].Position;
+				const FVector V1 = Model->VertexBuffer.Vertices[Triangle.v1].Position;
+				const FVector V2 = Model->VertexBuffer.Vertices[Triangle.v2].Position;
+
+				const FVector V01 = (V1 - V0);
+				const FVector V02 = (V2 - V0);
+				const FVector Cross = FVector::CrossProduct(V01, V02);
+				const float Area = Cross.Size() * 0.5f;
+				if (Area <= AreaThreshold)
+				{
+					nBadArea++;
+					continue;
+				}
+			}
+
+			CollisionData->Indices.Add(Triangle);
+			CollisionData->MaterialIndices.Add(ElementIndex);
+		}
+	}
+
+	if (nBadArea > 0)
+	{
+		UE_LOG(LogPhysics, Log, TEXT("Cooking removed %d triangle%s with area <= %f (%s)"), nBadArea, (nBadArea > 1 ? TEXT("s") : TEXT("")), AreaThreshold, *GetPathName(GetOuter()));
+	}
+
+	CollisionData->bFlipNormals = true;
+	return true;
+}
+
+bool UModelComponent::ContainsPhysicsTriMeshData(bool InUseAllTriData) const
+{
+	return (Elements.Num() > 0);
+}
\ No newline at end of file
diff --git a/Engine/Source/Runtime/Engine/Private/Components/NavMeshRenderingComponent.cpp b/Engine/Source/Runtime/Engine/Private/Components/NavMeshRenderingComponent.cpp
new file mode 100644
--- /dev/null
+++ b/Engine/Source/Runtime/Engine/Private/Components/NavMeshRenderingComponent.cpp
@@ -0,0 +1,589 @@
+// Copyright 1998-2015 Epic Games, Inc. All Rights Reserved.
+
+/*=============================================================================
+	NavMeshRenderingComponent.cpp: A component that renders a nav mesh.
+ =============================================================================*/
+
+#include "EnginePrivate.h"
+#include "DebugRenderSceneProxy.h"
+#include "NavigationOctree.h"
+#include "AI/Navigation/RecastHelpers.h"
+#include "AI/Navigation/RecastNavMeshGenerator.h"
+#include "AI/Navigation/NavigationSystem.h"
+
+#if WITH_EDITOR
+#include "Editor.h"
+#include "EditorViewportClient.h"
+#endif
+
+#include "NavMeshRenderingHelpers.h"
+
+UNavMeshRenderingComponent::UNavMeshRenderingComponent(const FObjectInitializer& ObjectInitializer)
+	: Super(ObjectInitializer)
+{
+	SetCollisionProfileName(UCollisionProfile::NoCollision_ProfileName);
+	AlwaysLoadOnClient = false;
+	AlwaysLoadOnServer = false;
+	bSelectable = false;
+	bCollectNavigationData = false;
+}
+
+bool UNavMeshRenderingComponent::NeedsLoadForServer() const
+{
+	// This rendering component is not needed in dedicated server builds
+	return false;
+}
+
+bool UNavMeshRenderingComponent::IsNavigationShowFlagSet(const UWorld* World)
+{
+	bool bShowNavigation = false;
+
+	FWorldContext* WorldContext = GEngine->GetWorldContextFromWorld(World);
+
+#if WITH_EDITOR
+	if (GEditor && WorldContext && WorldContext->WorldType == EWorldType::Editor)
+	{
+		for (FEditorViewportClient* CurrentViewport : GEditor->AllViewportClients)
+		{
+			if (CurrentViewport && CurrentViewport->EngineShowFlags.Navigation)
+			{
+				bShowNavigation = true;
+				break;
+			}
+		}
+	}
+	else
+#endif //WITH_EDITOR
+	{
+		bShowNavigation = WorldContext && WorldContext->GameViewport && WorldContext->GameViewport->EngineShowFlags.Navigation;
+	}
+
+	return bShowNavigation;
+}
+
+void UNavMeshRenderingComponent::TimerFunction()
+{
+	const bool bShowNavigation = IsNavigationShowFlagSet(GetWorld());
+
+	if (bShowNavigation != !!bCollectNavigationData)
+	{
+		bCollectNavigationData = bShowNavigation;
+		MarkRenderStateDirty();
+	}
+}
+
+void UNavMeshRenderingComponent::OnRegister()
+{
+	Super::OnRegister();
+
+#if WITH_RECAST && !UE_BUILD_SHIPPING && !UE_BUILD_TEST
+	// it's a kind of HACK but there is no event or other information that show flag was changed by user => we have to check it periodically
+#if WITH_EDITOR
+	if (GEditor)
+	{
+		GEditor->GetTimerManager()->SetTimer(TimerHandle, FTimerDelegate::CreateUObject(this, &UNavMeshRenderingComponent::TimerFunction), 1, true);
+	}
+	else
+#endif //WITH_EDITOR
+	{
+		GetWorld()->GetTimerManager().SetTimer(TimerHandle, FTimerDelegate::CreateUObject(this, &UNavMeshRenderingComponent::TimerFunction), 1, true);
+	}
+#endif //WITH_RECAST && !UE_BUILD_SHIPPING && !UE_BUILD_TEST
+}
+
+void UNavMeshRenderingComponent::OnUnregister()
+{
+#if WITH_RECAST && !UE_BUILD_SHIPPING && !UE_BUILD_TEST
+	// it's a kind of HACK but there is no event or other information that show flag was changed by user => we have to check it periodically
+#if WITH_EDITOR
+	if (GEditor)
+	{
+		GEditor->GetTimerManager()->ClearTimer(TimerHandle);
+	}
+	else
+#endif //WITH_EDITOR
+	{
+		GetWorld()->GetTimerManager().ClearTimer(TimerHandle);
+	}
+#endif //WITH_RECAST && !UE_BUILD_SHIPPING && !UE_BUILD_TEST
+	Super::OnUnregister();
+}
+
+FPrimitiveSceneProxy* UNavMeshRenderingComponent::CreateSceneProxy()
+{
+#if WITH_RECAST && !UE_BUILD_SHIPPING && !UE_BUILD_TEST
+	FPrimitiveSceneProxy* SceneProxy = NULL;
+
+	const bool bShowNavigation = IsNavigationShowFlagSet(GetWorld());
+
+	bCollectNavigationData = bShowNavigation;
+
+	if (bCollectNavigationData && IsVisible())
+	{
+		FNavMeshSceneProxyData ProxyData;
+		ProxyData.Reset();
+		GatherData(ProxyData);
+
+		SceneProxy = new FRecastRenderingSceneProxy(this, &ProxyData);
+		ProxyData.Reset();
+	}
+	return SceneProxy;
+#else
+	return NULL;
+#endif //WITH_RECAST && !UE_BUILD_SHIPPING && !UE_BUILD_TEST
+}
+
+void UNavMeshRenderingComponent::CreateRenderState_Concurrent()
+{
+	Super::CreateRenderState_Concurrent();
+
+#if WITH_RECAST && !UE_BUILD_SHIPPING && !UE_BUILD_TEST
+	if (SceneProxy != NULL)
+	{
+		static_cast<FRecastRenderingSceneProxy*>(SceneProxy)->RegisterDebugDrawDelgate();
+	}
+#endif
+}
+
+void UNavMeshRenderingComponent::DestroyRenderState_Concurrent()
+{
+#if WITH_RECAST && !UE_BUILD_SHIPPING && !UE_BUILD_TEST
+	if (SceneProxy != NULL)
+	{
+		static_cast<FRecastRenderingSceneProxy*>(SceneProxy)->UnregisterDebugDrawDelgate();
+	}
+#endif
+
+	Super::DestroyRenderState_Concurrent();
+}
+
+FBoxSphereBounds UNavMeshRenderingComponent::CalcBounds(const FTransform& LocalToWorld) const
+{
+	FBox BoundingBox(0);
+#if WITH_RECAST
+	ARecastNavMesh* NavMesh = Cast<ARecastNavMesh>(GetOwner());
+	if (NavMesh)
+	{
+		BoundingBox = NavMesh->GetNavMeshBounds();
+		if (NavMesh->bDrawOctree)
+		{
+			const UNavigationSystem* NavSys = UNavigationSystem::GetCurrent(GetWorld());
+			const FNavigationOctree* NavOctree = NavSys ? NavSys->GetNavOctree() : NULL;
+			if (NavOctree)
+			{
+				for (FNavigationOctree::TConstIterator<> NodeIt(*NavOctree); NodeIt.HasPendingNodes(); NodeIt.Advance())
+				{
+					const FOctreeNodeContext& CorrentContext = NodeIt.GetCurrentContext();
+					BoundingBox += CorrentContext.Bounds.GetBox();
+				}
+			}
+		}
+	}
+#endif
+	return FBoxSphereBounds(BoundingBox);
+}
+
+void UNavMeshRenderingComponent::GatherData(FNavMeshSceneProxyData& CurrentData) const
+{
+	QUICK_SCOPE_CYCLE_COUNTER(STAT_NavMesh_GatherDebugDrawingGeometry);
+#if WITH_RECAST
+	const ARecastNavMesh* NavMesh = Cast<ARecastNavMesh>(GetOwner());
+
+	CurrentData.Reset();
+	
+	if (NavMesh && NavMesh->IsDrawingEnabled())
+	{
+		CurrentData.bEnableDrawing = true;
+		CurrentData.bNeedsNewData = false;
+
+		FHitProxyId HitProxyId = FHitProxyId();
+		CurrentData.bDrawPathCollidingGeometry = NavMesh->bDrawPathCollidingGeometry;
+
+		CurrentData.NavMeshDrawOffset.Z = NavMesh->DrawOffset;
+		CurrentData.NavMeshGeometry.bGatherPolyEdges = NavMesh->bDrawPolyEdges;
+		CurrentData.NavMeshGeometry.bGatherNavMeshEdges = NavMesh->bDrawNavMeshEdges;
+
+		const FNavDataConfig& NavConfig = NavMesh->GetConfig();
+		CurrentData.NavMeshColors[RECAST_DEFAULT_AREA] = NavConfig.Color.DWColor() > 0 ? NavConfig.Color : NavMeshRenderColor_RecastMesh;
+		for (uint8 i = 0; i < RECAST_DEFAULT_AREA; ++i)
+		{
+			CurrentData.NavMeshColors[i] = NavMesh->GetAreaIDColor(i);
+		}
+
+		// just a little trick to make sure navmeshes with different sized are not drawn with same offset
+		CurrentData.NavMeshDrawOffset.Z += NavMesh->GetConfig().AgentRadius / 10.f;
+
+		NavMesh->BeginBatchQuery();
+
+		if (NavMesh->bDrawOctree)
+		{
+			const UNavigationSystem* NavSys = UNavigationSystem::GetCurrent(GetWorld());
+			const FNavigationOctree* NavOctree = NavSys ? NavSys->GetNavOctree() : NULL;
+			if (NavOctree)
+			{
+				for (FNavigationOctree::TConstIterator<> NodeIt(*NavOctree); NodeIt.HasPendingNodes(); NodeIt.Advance())
+				{
+					const FNavigationOctree::FNode& CurrentNode = NodeIt.GetCurrentNode();
+					const FOctreeNodeContext& CorrentContext = NodeIt.GetCurrentContext();
+					CurrentData.OctreeBounds.Add(CorrentContext.Bounds);
+
+					FOREACH_OCTREE_CHILD_NODE(ChildRef)
+					{
+						if (CurrentNode.HasChild(ChildRef))
+						{
+							NodeIt.PushChild(ChildRef);
+						}
+					}
+				}
+			}
+		}
+
+		NavMesh->GetDebugGeometry(CurrentData.NavMeshGeometry);
+		const TArray<FVector>& MeshVerts = CurrentData.NavMeshGeometry.MeshVerts;
+
+		// @fixme, this is going to double up on lots of interior lines
+		if (NavMesh->bDrawTriangleEdges)
+		{
+			for (int32 AreaIdx = 0; AreaIdx < RECAST_MAX_AREAS; ++AreaIdx)
+			{
+				const TArray<int32>& MeshIndices = CurrentData.NavMeshGeometry.AreaIndices[AreaIdx];
+				for (int32 Idx=0; Idx<MeshIndices.Num(); Idx += 3)
+				{
+					CurrentData.ThickLineItems.Add(FNavMeshSceneProxyData::FDebugThickLine(MeshVerts[MeshIndices[Idx + 0]] + CurrentData.NavMeshDrawOffset, MeshVerts[MeshIndices[Idx + 1]] + CurrentData.NavMeshDrawOffset, NavMeshRenderColor_Recast_TriangleEdges, DefaultEdges_LineThickness));
+					CurrentData.ThickLineItems.Add(FNavMeshSceneProxyData::FDebugThickLine(MeshVerts[MeshIndices[Idx + 1]] + CurrentData.NavMeshDrawOffset, MeshVerts[MeshIndices[Idx + 2]] + CurrentData.NavMeshDrawOffset, NavMeshRenderColor_Recast_TriangleEdges, DefaultEdges_LineThickness));
+					CurrentData.ThickLineItems.Add(FNavMeshSceneProxyData::FDebugThickLine(MeshVerts[MeshIndices[Idx + 2]] + CurrentData.NavMeshDrawOffset, MeshVerts[MeshIndices[Idx + 0]] + CurrentData.NavMeshDrawOffset, NavMeshRenderColor_Recast_TriangleEdges, DefaultEdges_LineThickness));
+				}
+			}
+		}
+
+		// make lines for tile edges
+		if (NavMesh->bDrawPolyEdges)
+		{
+			const TArray<FVector>& TileEdgeVerts = CurrentData.NavMeshGeometry.PolyEdges;
+			for (int32 Idx=0; Idx < TileEdgeVerts.Num(); Idx += 2)
+			{
+				CurrentData.TileEdgeLines.Add(FDebugRenderSceneProxy::FDebugLine(TileEdgeVerts[Idx] + CurrentData.NavMeshDrawOffset, TileEdgeVerts[Idx + 1] + CurrentData.NavMeshDrawOffset, NavMeshRenderColor_Recast_TileEdges));
+			}
+		}
+
+		// make lines for navmesh edges
+		if (NavMesh->bDrawNavMeshEdges)
+		{
+			const FColor EdgesColor = DarkenColor(CurrentData.NavMeshColors[RECAST_DEFAULT_AREA]);
+			const TArray<FVector>& NavMeshEdgeVerts = CurrentData.NavMeshGeometry.NavMeshEdges;
+			for (int32 Idx=0; Idx < NavMeshEdgeVerts.Num(); Idx += 2)
+			{
+				CurrentData.NavMeshEdgeLines.Add( FDebugRenderSceneProxy::FDebugLine(NavMeshEdgeVerts[Idx] + CurrentData.NavMeshDrawOffset, NavMeshEdgeVerts[Idx+1] + CurrentData.NavMeshDrawOffset, EdgesColor));
+			}
+		}
+
+		// offset all navigation-link positions
+		if (!NavMesh->bDrawClusters)
+		{
+			for (int32 OffMeshLineIndex = 0; OffMeshLineIndex < CurrentData.NavMeshGeometry.OffMeshLinks.Num(); ++OffMeshLineIndex)
+			{
+				FRecastDebugGeometry::FOffMeshLink& Link = CurrentData.NavMeshGeometry.OffMeshLinks[OffMeshLineIndex];
+				const bool bLinkValid = (Link.ValidEnds & FRecastDebugGeometry::OMLE_Left) && (Link.ValidEnds & FRecastDebugGeometry::OMLE_Right);
+
+				if (NavMesh->bDrawFailedNavLinks || (NavMesh->bDrawNavLinks && bLinkValid))
+				{
+					const FVector V0 = Link.Left + CurrentData.NavMeshDrawOffset;
+					const FVector V1 = Link.Right + CurrentData.NavMeshDrawOffset;
+					const FColor LinkColor = ((Link.Direction && Link.ValidEnds) || (Link.ValidEnds & FRecastDebugGeometry::OMLE_Left)) ? DarkenColor(CurrentData.NavMeshColors[Link.AreaID]) : NavMeshRenderColor_OffMeshConnectionInvalid;
+
+					CacheArc(CurrentData.NavLinkLines, V0, V1, 0.4f, 4, LinkColor, LinkLines_LineThickness);
+
+					const FVector VOffset(0, 0, FVector::Dist(V0, V1) * 1.333f);
+					CacheArrowHead(CurrentData.NavLinkLines, V1, V0+VOffset, 30.f, LinkColor, LinkLines_LineThickness);
+					if (Link.Direction)
+					{
+						CacheArrowHead(CurrentData.NavLinkLines, V0, V1+VOffset, 30.f, LinkColor, LinkLines_LineThickness);
+					}
+
+					// if the connection as a whole is valid check if there are any of ends is invalid
+					if (LinkColor != NavMeshRenderColor_OffMeshConnectionInvalid)
+					{
+						if (Link.Direction && (Link.ValidEnds & FRecastDebugGeometry::OMLE_Left) == 0)
+						{
+							// left end invalid - mark it
+							DrawWireCylinder(CurrentData.NavLinkLines, V0, FVector(1,0,0), FVector(0,1,0), FVector(0,0,1), NavMeshRenderColor_OffMeshConnectionInvalid, Link.Radius, NavMesh->AgentMaxStepHeight, 16, 0, DefaultEdges_LineThickness);
+						}
+						if ((Link.ValidEnds & FRecastDebugGeometry::OMLE_Right) == 0)
+						{
+							DrawWireCylinder(CurrentData.NavLinkLines, V1, FVector(1,0,0), FVector(0,1,0), FVector(0,0,1), NavMeshRenderColor_OffMeshConnectionInvalid, Link.Radius, NavMesh->AgentMaxStepHeight, 16, 0, DefaultEdges_LineThickness);
+						}
+					}
+				}					
+			}
+		}
+
+		if (NavMesh->bDrawTileLabels || NavMesh->bDrawPolygonLabels || NavMesh->bDrawDefaultPolygonCost || NavMesh->bDrawTileBounds)
+		{
+			// calculate appropriate points for displaying debug labels
+			const int32 TilesCount = NavMesh->GetNavMeshTilesCount();
+			CurrentData.DebugLabels.Reserve(TilesCount);
+
+			for (int32 TileIndex = 0; TileIndex < TilesCount; ++TileIndex)
+			{
+				int32 X, Y, Layer;
+				if (NavMesh->GetNavMeshTileXY(TileIndex, X, Y, Layer))
+				{
+					const FBox TileBoundingBox = NavMesh->GetNavMeshTileBounds(TileIndex);
+					FVector TileLabelLocation = TileBoundingBox.GetCenter();
+					TileLabelLocation.Z = TileBoundingBox.Max.Z;
+
+					FNavLocation NavLocation(TileLabelLocation);
+					if (!NavMesh->ProjectPoint(TileLabelLocation, NavLocation, FVector(NavMesh->TileSizeUU/100, NavMesh->TileSizeUU/100, TileBoundingBox.Max.Z-TileBoundingBox.Min.Z)))
+					{
+						NavMesh->ProjectPoint(TileLabelLocation, NavLocation, FVector(NavMesh->TileSizeUU/2, NavMesh->TileSizeUU/2, TileBoundingBox.Max.Z-TileBoundingBox.Min.Z));
+					}
+
+					if (NavMesh->bDrawTileLabels)
+					{
+						CurrentData.DebugLabels.Add(FNavMeshSceneProxyData::FDebugText(
+							/*Location*/NavLocation.Location + CurrentData.NavMeshDrawOffset
+							, /*Text*/FString::Printf(TEXT("(%d,%d:%d)"), X, Y, Layer)
+							));
+					}
+
+					if (NavMesh->bDrawPolygonLabels || NavMesh->bDrawDefaultPolygonCost)
+					{
+						TArray<FNavPoly> Polys;
+						NavMesh->GetPolysInTile(TileIndex, Polys);
+
+						if (NavMesh->bDrawDefaultPolygonCost)
+						{
+							float DefaultCosts[RECAST_MAX_AREAS];
+							float FixedCosts[RECAST_MAX_AREAS];
+
+							NavMesh->GetDefaultQueryFilter()->GetAllAreaCosts(DefaultCosts, FixedCosts, RECAST_MAX_AREAS);
+
+							for(int k = 0; k < Polys.Num(); ++k)
+							{
+								uint32 AreaID = NavMesh->GetPolyAreaID(Polys[k].Ref);
+
+								CurrentData.DebugLabels.Add(FNavMeshSceneProxyData::FDebugText(
+									/*Location*/Polys[k].Center + CurrentData.NavMeshDrawOffset
+									, /*Text*/FString::Printf(TEXT("\\%.3f; %.3f\\"), DefaultCosts[AreaID], FixedCosts[AreaID])
+									));
+							}
+						}
+						else
+						{
+							for(int k = 0; k < Polys.Num(); ++k)
+							{
+								uint32 NavPolyIndex = 0;
+								uint32 NavTileIndex = 0;
+								NavMesh->GetPolyTileIndex(Polys[k].Ref, NavPolyIndex, NavTileIndex);
+
+								CurrentData.DebugLabels.Add(FNavMeshSceneProxyData::FDebugText(
+									/*Location*/Polys[k].Center + CurrentData.NavMeshDrawOffset
+									, /*Text*/FString::Printf(TEXT("[%X:%X]"), NavTileIndex, NavPolyIndex)
+									));
+							}
+						}
+					}							
+
+					if (NavMesh->bDrawTileBounds)
+					{
+						FBox TileBox = NavMesh->GetNavMeshTileBounds(TileIndex);
+
+						float DrawZ = (TileBox.Min.Z + TileBox.Max.Z) * 0.5f;		// @hack average
+						FVector LL(TileBox.Min.X, TileBox.Min.Y, DrawZ);
+						FVector UR(TileBox.Max.X, TileBox.Max.Y, DrawZ);
+						FVector UL(LL.X, UR.Y, DrawZ);
+						FVector LR(UR.X, LL.Y, DrawZ);
+						CurrentData.ThickLineItems.Add(FNavMeshSceneProxyData::FDebugThickLine(LL, UL, NavMeshRenderColor_TileBounds, DefaultEdges_LineThickness));
+						CurrentData.ThickLineItems.Add(FNavMeshSceneProxyData::FDebugThickLine(UL, UR, NavMeshRenderColor_TileBounds, DefaultEdges_LineThickness));
+						CurrentData.ThickLineItems.Add(FNavMeshSceneProxyData::FDebugThickLine(UR, LR, NavMeshRenderColor_TileBounds, DefaultEdges_LineThickness));
+						CurrentData.ThickLineItems.Add(FNavMeshSceneProxyData::FDebugThickLine(LR, LL, NavMeshRenderColor_TileBounds, DefaultEdges_LineThickness));
+					}
+				}
+			}
+		}
+
+		CurrentData.bSkipDistanceCheck = GIsEditor && (GEngine->GetDebugLocalPlayer() == NULL);
+		CurrentData.bDrawClusters = NavMesh->bDrawClusters;
+		NavMesh->FinishBatchQuery();
+
+		// Draw Mesh
+		if (NavMesh->bDrawClusters)
+		{
+			for (int32 Idx = 0; Idx < CurrentData.NavMeshGeometry.Clusters.Num(); ++Idx)
+			{
+				const TArray<int32>& MeshIndices = CurrentData.NavMeshGeometry.Clusters[Idx].MeshIndices;
+
+				if (MeshIndices.Num() == 0)
+				{
+					continue;
+				}
+
+				FNavMeshSceneProxyData::FDebugMeshData DebugMeshData;
+				DebugMeshData.ClusterColor = GetClusterColor(Idx);
+				for (int32 VertIdx=0; VertIdx < MeshVerts.Num(); ++VertIdx)
+				{
+					AddVertexHelper(DebugMeshData, MeshVerts[VertIdx] + CurrentData.NavMeshDrawOffset, DebugMeshData.ClusterColor);
+				}
+				for (int32 TriIdx=0; TriIdx < MeshIndices.Num(); TriIdx+=3)
+				{
+					AddTriangleHelper(DebugMeshData, MeshIndices[TriIdx], MeshIndices[TriIdx+1], MeshIndices[TriIdx+2]);
+				}
+
+				CurrentData.MeshBuilders.Add(DebugMeshData);
+			}
+		}
+		else if (NavMesh->bDrawNavMesh)
+		{			
+			for (int32 AreaType = 0; AreaType < RECAST_MAX_AREAS; ++AreaType)
+			{
+				const TArray<int32>& MeshIndices = CurrentData.NavMeshGeometry.AreaIndices[AreaType];
+
+				if (MeshIndices.Num() == 0)
+				{
+					continue;
+				}
+
+				FNavMeshSceneProxyData::FDebugMeshData DebugMeshData;
+				for (int32 VertIdx=0; VertIdx < MeshVerts.Num(); ++VertIdx)
+				{
+					AddVertexHelper(DebugMeshData, MeshVerts[VertIdx] + CurrentData.NavMeshDrawOffset, CurrentData.NavMeshColors[AreaType]);
+				}
+				for (int32 TriIdx=0; TriIdx < MeshIndices.Num(); TriIdx+=3)
+				{
+					AddTriangleHelper(DebugMeshData, MeshIndices[TriIdx], MeshIndices[TriIdx+1], MeshIndices[TriIdx+2]);
+				}
+
+				DebugMeshData.ClusterColor = CurrentData.NavMeshColors[AreaType];
+				CurrentData.MeshBuilders.Add(DebugMeshData);
+			}
+		}
+
+		if (NavMesh->bDrawPathCollidingGeometry)
+		{
+			// draw all geometry gathered in navoctree
+			const FNavigationOctree* NavOctree = NavMesh->GetWorld()->GetNavigationSystem()->GetNavOctree();
+
+			TArray<FVector> PathCollidingGeomVerts;
+			TArray <int32> PathCollidingGeomIndices;
+			for (FNavigationOctree::TConstIterator<> It(*NavOctree); It.HasPendingNodes(); It.Advance())
+			{
+				const FNavigationOctree::FNode& Node = It.GetCurrentNode();
+				for (FNavigationOctree::ElementConstIt ElementIt(Node.GetElementIt()); ElementIt; ElementIt++)
+				{
+					const FNavigationOctreeElement& Element = *ElementIt;
+					if (Element.ShouldUseGeometry(NavMesh->GetConfig()) && Element.Data->CollisionData.Num())
+					{
+						const FRecastGeometryCache CachedGeometry(Element.Data->CollisionData.GetData());
+						AppendGeometry(PathCollidingGeomVerts, PathCollidingGeomIndices, CachedGeometry.Verts, CachedGeometry.Header.NumVerts, CachedGeometry.Indices, CachedGeometry.Header.NumFaces);
+					}
+				}
+				FOREACH_OCTREE_CHILD_NODE(ChildRef)
+				{
+					if (Node.HasChild(ChildRef))
+					{
+						It.PushChild(ChildRef);
+					}
+				}
+			}
+			CurrentData.PathCollidingGeomIndices = PathCollidingGeomIndices;
+			for (const auto& Vertex : PathCollidingGeomVerts)
+			{
+				CurrentData.PathCollidingGeomVerts.Add(FDynamicMeshVertex(Vertex));
+			}
+
+		}
+
+		if (CurrentData.NavMeshGeometry.BuiltMeshIndices.Num() > 0)
+		{
+			FNavMeshSceneProxyData::FDebugMeshData DebugMeshData;
+			for (int32 VertIdx=0; VertIdx < MeshVerts.Num(); ++VertIdx)
+			{
+				AddVertexHelper(DebugMeshData, MeshVerts[VertIdx] + CurrentData.NavMeshDrawOffset, NavMeshRenderColor_RecastTileBeingRebuilt);
+			}
+			DebugMeshData.Indices.Append(CurrentData.NavMeshGeometry.BuiltMeshIndices);
+			DebugMeshData.ClusterColor = NavMeshRenderColor_RecastTileBeingRebuilt;
+			CurrentData.MeshBuilders.Add(DebugMeshData);
+
+			// updates should be requested by FRecastNavMeshGenerator::TickAsyncBuild after tiles were refreshed
+		}
+
+		if (NavMesh->bDrawClusters)
+		{
+			for (int i = 0; i < CurrentData.NavMeshGeometry.ClusterLinks.Num(); i++)
+			{
+				const FRecastDebugGeometry::FClusterLink& CLink = CurrentData.NavMeshGeometry.ClusterLinks[i];
+				const FVector V0 = CLink.FromCluster + CurrentData.NavMeshDrawOffset;
+				const FVector V1 = CLink.ToCluster + CurrentData.NavMeshDrawOffset + FVector(0,0,20.0f);
+
+				CacheArc(CurrentData.ClusterLinkLines, V0, V1, 0.4f, 4, FColor::Black, ClusterLinkLines_LineThickness);
+				const FVector VOffset(0, 0, FVector::Dist(V0, V1) * 1.333f);
+				CacheArrowHead(CurrentData.ClusterLinkLines, V1, V0+VOffset, 30.f, FColor::Black, ClusterLinkLines_LineThickness);
+			}
+		}
+
+		// cache segment links
+		if (NavMesh->bDrawNavLinks)
+		{
+			for (int32 iArea = 0; iArea < RECAST_MAX_AREAS; iArea++)
+			{
+				const TArray<int32>& Indices = CurrentData.NavMeshGeometry.OffMeshSegmentAreas[iArea];
+				FNavMeshSceneProxyData::FDebugMeshData DebugMeshData;
+				int32 VertBase = 0;
+
+				for (int32 i = 0; i < Indices.Num(); i++)
+				{
+					FRecastDebugGeometry::FOffMeshSegment& SegInfo = CurrentData.NavMeshGeometry.OffMeshSegments[Indices[i]];
+					const FVector A0 = SegInfo.LeftStart + CurrentData.NavMeshDrawOffset;
+					const FVector A1 = SegInfo.LeftEnd + CurrentData.NavMeshDrawOffset;
+					const FVector B0 = SegInfo.RightStart + CurrentData.NavMeshDrawOffset;
+					const FVector B1 = SegInfo.RightEnd + CurrentData.NavMeshDrawOffset;
+					const FVector Edge0 = B0 - A0;
+					const FVector Edge1 = B1 - A1;
+					const float Len0 = Edge0.Size();
+					const float Len1 = Edge1.Size();
+					const FColor SegColor = DarkenColor(CurrentData.NavMeshColors[SegInfo.AreaID]);
+					const FColor ColA = (SegInfo.ValidEnds & FRecastDebugGeometry::OMLE_Left) ? FColor::White : FColor::Black;
+					const FColor ColB = (SegInfo.ValidEnds & FRecastDebugGeometry::OMLE_Right) ? FColor::White : FColor::Black;
+
+					const int32 NumArcPoints = 8;
+					const float ArcPtsScale = 1.0f / NumArcPoints;
+
+					FVector Prev0 = EvalArc(A0, Edge0, Len0*0.25f, 0);
+					FVector Prev1 = EvalArc(A1, Edge1, Len1*0.25f, 0);
+					AddVertexHelper(DebugMeshData, Prev0, ColA);
+					AddVertexHelper(DebugMeshData, Prev1, ColA);
+					for (int32 j = 1; j <= NumArcPoints; j++)
+					{
+						const float u = j * ArcPtsScale;
+						FVector Pt0 = EvalArc(A0, Edge0, Len0*0.25f, u);
+						FVector Pt1 = EvalArc(A1, Edge1, Len1*0.25f, u);
+
+						AddVertexHelper(DebugMeshData, Pt0, (j == NumArcPoints) ? ColB : FColor::White);
+						AddVertexHelper(DebugMeshData, Pt1, (j == NumArcPoints) ? ColB : FColor::White);
+
+						AddTriangleHelper(DebugMeshData, VertBase+0, VertBase+2, VertBase+1);
+						AddTriangleHelper(DebugMeshData, VertBase+2, VertBase+3, VertBase+1);
+						AddTriangleHelper(DebugMeshData, VertBase+0, VertBase+1, VertBase+2);
+						AddTriangleHelper(DebugMeshData, VertBase+2, VertBase+1, VertBase+3);
+
+						VertBase += 2;
+						Prev0 = Pt0;
+						Prev1 = Pt1;
+					}
+					VertBase += 2;
+
+					DebugMeshData.ClusterColor = SegColor;
+				}
+
+				if (DebugMeshData.Indices.Num())
+				{
+					CurrentData.MeshBuilders.Add(DebugMeshData);
+				}
+			}
+		}
+
+		CurrentData.NavMeshGeometry.PolyEdges.Empty();
+		CurrentData.NavMeshGeometry.NavMeshEdges.Empty();
+	}
+#endif //WITH_RECAST
+}
