Engine: Moved actor components into Components folder

diff --git a/Engine/Source/Runtime/Engine/Private/Components/PointLightComponent.cpp b/Engine/Source/Runtime/Engine/Private/Components/PointLightComponent.cpp
new file mode 100644
--- /dev/null
+++ b/Engine/Source/Runtime/Engine/Private/Components/PointLightComponent.cpp
@@ -0,0 +1,315 @@
+// Copyright 1998-2015 Epic Games, Inc. All Rights Reserved.
+
+/*=============================================================================
+	PointLightComponent.cpp: PointLightComponent implementation.
+=============================================================================*/
+
+#include "EnginePrivate.h"
+#include "PointLightSceneProxy.h"
+
+static int32 GAllowPointLightCubemapShadows = 1;
+static FAutoConsoleVariableRef CVarAllowPointLightCubemapShadows(
+	TEXT("r.AllowPointLightCubemapShadows"),
+	GAllowPointLightCubemapShadows,
+	TEXT("When 0, will prevent point light cube map shadows from being used and the light will be unshadowed.")
+	);
+
+/**
+ * The point light policy for TMeshLightingDrawingPolicy.
+ */
+class FPointLightPolicy
+{
+public:
+	typedef TPointLightSceneProxy<FPointLightPolicy> SceneInfoType;
+};
+
+void FPointLightSceneProxyBase::UpdateRadius_GameThread(UPointLightComponent* Component)
+{
+	ENQUEUE_UNIQUE_RENDER_COMMAND_TWOPARAMETER(
+		UpdateRadius,
+		FPointLightSceneProxyBase*,LightSceneInfo,this,
+		float,ComponentRadius,Component->AttenuationRadius,
+	{
+		LightSceneInfo->UpdateRadius(ComponentRadius);
+	});
+}
+
+class FPointLightSceneProxy : public TPointLightSceneProxy<FPointLightPolicy>
+{
+public:
+
+	/** Accesses parameters needed for rendering the light. */
+	virtual void GetParameters(FVector4& LightPositionAndInvRadius, FVector4& LightColorAndFalloffExponent, FVector& NormalizedLightDirection, FVector2D& SpotAngles, float& LightSourceRadius, float& LightSourceLength, float& LightMinRoughness) const
+	{
+		LightPositionAndInvRadius = FVector4(
+			GetOrigin(),
+			InvRadius);
+
+		LightColorAndFalloffExponent = FVector4(
+			GetColor().R,
+			GetColor().G,
+			GetColor().B,
+			FalloffExponent);
+
+		NormalizedLightDirection = -GetDirection();
+		SpotAngles = FVector2D( -2.0f, 1.0f );
+		LightSourceRadius = SourceRadius;
+		LightSourceLength = SourceLength;
+		LightMinRoughness = MinRoughness;
+	}
+
+	virtual FSphere GetBoundingSphere() const
+	{
+		return FSphere(GetPosition(), GetRadius());
+	}
+
+	/**
+	 * Sets up a projected shadow initializer for shadows from the entire scene.
+	 * @return True if the whole-scene projected shadow should be used.
+	 */
+	virtual bool GetWholeSceneProjectedShadowInitializer(const FSceneViewFamily& ViewFamily, TArray<FWholeSceneProjectedShadowInitializer, TInlineAllocator<6> >& OutInitializers) const
+	{
+		if (ViewFamily.GetFeatureLevel() >= ERHIFeatureLevel::SM4
+			&& GAllowPointLightCubemapShadows != 0)
+		{
+			FWholeSceneProjectedShadowInitializer& OutInitializer = *new(OutInitializers) FWholeSceneProjectedShadowInitializer;
+			OutInitializer.PreShadowTranslation = -GetLightToWorld().GetOrigin();
+			OutInitializer.WorldToLight = GetWorldToLight().RemoveTranslation();
+			OutInitializer.Scales = FVector(1, 1, 1);
+			OutInitializer.FaceDirection = FVector(0,0,1);
+			OutInitializer.SubjectBounds = FBoxSphereBounds(FVector(0, 0, 0),FVector(Radius,Radius,Radius),Radius);
+			OutInitializer.WAxis = FVector4(0,0,1,0);
+			OutInitializer.MinLightW = 0.1f;
+			OutInitializer.MaxDistanceToCastInLightW = Radius;
+			OutInitializer.CascadeSettings.bOnePassPointLightShadow = true;
+			OutInitializer.CascadeSettings.bRayTracedDistanceField = UseRayTracedDistanceFieldShadows() && DoesPlatformSupportDistanceFieldShadowing(ViewFamily.GetShaderPlatform());
+			return true;
+		}
+		
+		return false;
+	}
+
+	FPointLightSceneProxy(const UPointLightComponent* Component)
+	:	TPointLightSceneProxy<FPointLightPolicy>(Component)
+	{}
+};
+
+UPointLightComponent::UPointLightComponent(const FObjectInitializer& ObjectInitializer)
+	: Super(ObjectInitializer)
+{
+#if WITH_EDITORONLY_DATA
+	if (!IsRunningCommandlet())
+	{
+		static ConstructorHelpers::FObjectFinder<UTexture2D> StaticTexture(TEXT("/Engine/EditorResources/LightIcons/S_LightPoint"));
+		static ConstructorHelpers::FObjectFinder<UTexture2D> DynamicTexture(TEXT("/Engine/EditorResources/LightIcons/S_LightPointMove"));
+
+		StaticEditorTexture = StaticTexture.Object;
+		StaticEditorTextureScale = 0.5f;
+		DynamicEditorTexture = DynamicTexture.Object;
+		DynamicEditorTextureScale = 0.5f;
+	}
+#endif
+
+	Intensity = 5000;
+	Radius_DEPRECATED = 1024.0f;
+	AttenuationRadius = 1000;
+	LightFalloffExponent = 8.0f;
+	SourceRadius = 0.0f;
+	SourceLength = 0.0f;
+	bUseInverseSquaredFalloff = true;
+}
+
+FLightSceneProxy* UPointLightComponent::CreateSceneProxy() const
+{
+	return new FPointLightSceneProxy(this);
+}
+
+void UPointLightComponent::SetAttenuationRadius(float NewRadius)
+{
+	// Only movable lights can change their radius at runtime
+	if (AreDynamicDataChangesAllowed(false)
+		&& NewRadius != AttenuationRadius)
+	{
+		AttenuationRadius = NewRadius;
+		PushRadiusToRenderThread();
+	}
+}
+
+void UPointLightComponent::SetLightFalloffExponent(float NewLightFalloffExponent)
+{
+	if (AreDynamicDataChangesAllowed()
+		&& NewLightFalloffExponent != LightFalloffExponent)
+	{
+		LightFalloffExponent = NewLightFalloffExponent;
+		MarkRenderStateDirty();
+	}
+}
+
+void UPointLightComponent::SetSourceRadius(float bNewValue)
+{
+	if (AreDynamicDataChangesAllowed()
+		&& SourceRadius != bNewValue)
+	{
+		SourceRadius = bNewValue;
+		MarkRenderStateDirty();
+	}
+}
+
+bool UPointLightComponent::AffectsBounds(const FBoxSphereBounds& Bounds) const
+{
+	if((Bounds.Origin - ComponentToWorld.GetLocation()).SizeSquared() > FMath::Square(AttenuationRadius + Bounds.SphereRadius))
+	{
+		return false;
+	}
+
+	if(!Super::AffectsBounds(Bounds))
+	{
+		return false;
+	}
+
+	return true;
+}
+
+void UPointLightComponent::SendRenderTransform_Concurrent()
+{
+	// Update the scene info's cached radius-dependent data.
+	if(SceneProxy)
+	{
+		((FPointLightSceneProxyBase*)SceneProxy)->UpdateRadius_GameThread(this);
+	}
+
+	Super::SendRenderTransform_Concurrent();
+}
+
+//
+FVector4 UPointLightComponent::GetLightPosition() const
+{
+	return FVector4(ComponentToWorld.GetLocation(),1);
+}
+
+/**
+* @return ELightComponentType for the light component class 
+*/
+ELightComponentType UPointLightComponent::GetLightType() const
+{
+	return LightType_Point;
+}
+
+//
+FBox UPointLightComponent::GetBoundingBox() const
+{
+	return FBox(GetComponentLocation() - FVector(AttenuationRadius,AttenuationRadius,AttenuationRadius),GetComponentLocation() + FVector(AttenuationRadius,AttenuationRadius,AttenuationRadius));
+}
+
+FSphere UPointLightComponent::GetBoundingSphere() const
+{
+	return FSphere(ComponentToWorld.GetLocation(), AttenuationRadius);
+}
+
+void UPointLightComponent::Serialize(FArchive& Ar)
+{
+	Super::Serialize(Ar);
+
+	if (Ar.UE4Ver() < VER_UE4_INVERSE_SQUARED_LIGHTS_DEFAULT)
+	{
+		bUseInverseSquaredFalloff = InverseSquaredFalloff_DEPRECATED;
+		AttenuationRadius = Radius_DEPRECATED;
+	}
+}
+
+void UPointLightComponent::PostLoad()
+{
+	Super::PostLoad();
+
+	if (GetLinkerUE4Version() < VER_UE4_POINT_LIGHT_SOURCE_RADIUS)
+	{
+		SourceRadius = SourceRadius_DEPRECATED;	
+	}
+}
+
+#if WITH_EDITOR
+
+bool UPointLightComponent::CanEditChange(const UProperty* InProperty) const
+{
+	if (InProperty)
+	{
+		FString PropertyName = InProperty->GetName();
+
+		if (PropertyName == GET_MEMBER_NAME_STRING_CHECKED(ULightComponent, bCastShadowsFromCinematicObjectsOnly) && bUseRayTracedDistanceFieldShadows)
+		{
+			return false;
+		}
+
+		if (FCString::Strcmp(*PropertyName, TEXT("LightFalloffExponent")) == 0)
+		{
+			return !bUseInverseSquaredFalloff;
+		}
+	}
+
+	return Super::CanEditChange(InProperty);
+}
+
+/**
+ * Called after property has changed via e.g. property window or set command.
+ *
+ * @param	PropertyThatChanged	UProperty that has been changed, NULL if unknown
+ */
+void UPointLightComponent::PostEditChangeProperty(FPropertyChangedEvent& PropertyChangedEvent)
+{
+	// Make sure exponent is > 0.
+	LightFalloffExponent = FMath::Max( (float) KINDA_SMALL_NUMBER, LightFalloffExponent );
+	SourceRadius = FMath::Max(0.0f, SourceRadius);
+	SourceLength = FMath::Max(0.0f, SourceLength);
+	Intensity = FMath::Max(0.0f, Intensity);
+	LightmassSettings.IndirectLightingSaturation = FMath::Max(LightmassSettings.IndirectLightingSaturation, 0.0f);
+	LightmassSettings.ShadowExponent = FMath::Clamp(LightmassSettings.ShadowExponent, .5f, 8.0f);
+
+	Super::PostEditChangeProperty(PropertyChangedEvent);
+}
+#endif // WITH_EDITOR
+
+void UPointLightComponent::PostInterpChange(UProperty* PropertyThatChanged)
+{
+	static FName RadiusName(TEXT("Radius"));
+	static FName AttenuationRadiusName(TEXT("AttenuationRadius"));
+	static FName LightFalloffExponentName(TEXT("LightFalloffExponent"));
+	FName PropertyName = PropertyThatChanged->GetFName();
+
+	if (PropertyName == RadiusName
+		|| PropertyName == AttenuationRadiusName)
+	{
+		// Old radius tracks will animate the deprecated value
+		if (PropertyName == RadiusName)
+		{
+			AttenuationRadius = Radius_DEPRECATED;
+		}
+
+		PushRadiusToRenderThread();
+	}
+	else if (PropertyName == LightFalloffExponentName)
+	{
+		MarkRenderStateDirty();
+	}
+	else
+	{
+		Super::PostInterpChange(PropertyThatChanged);
+	}
+}
+
+void UPointLightComponent::PushRadiusToRenderThread()
+{
+	if (CastShadows)
+	{
+		// Shadow casting lights need to recompute light interactions
+		// to determine which primitives to draw in shadow depth passes.
+		MarkRenderStateDirty();
+	}
+	else
+	{
+		if (SceneProxy)
+		{
+			((FPointLightSceneProxyBase*)SceneProxy)->UpdateRadius_GameThread(this);
+		}
+	}
+}
+
diff --git a/Engine/Source/Runtime/Engine/Private/Components/PoseableMeshComponent.cpp b/Engine/Source/Runtime/Engine/Private/Components/PoseableMeshComponent.cpp
new file mode 100644
--- /dev/null
+++ b/Engine/Source/Runtime/Engine/Private/Components/PoseableMeshComponent.cpp
@@ -0,0 +1,249 @@
+// Copyright 1998-2015 Epic Games, Inc. All Rights Reserved.
+
+/*=============================================================================
+	PoseableMeshComponent.cpp: UPoseableMeshComponent methods.
+=============================================================================*/
+
+#include "EnginePrivate.h"
+#include "AnimTree.h"
+#include "Animation/AnimInstance.h"
+#include "Components/PoseableMeshComponent.h"
+
+UPoseableMeshComponent::UPoseableMeshComponent(const FObjectInitializer& ObjectInitializer)
+	: Super(ObjectInitializer)
+{
+
+}
+
+bool UPoseableMeshComponent::AllocateTransformData()
+{
+	// Allocate transforms if not present.
+	if ( Super::AllocateTransformData() )
+	{
+		if( LocalAtoms.Num() != SkeletalMesh->RefSkeleton.GetNum() )
+		{
+			LocalAtoms = SkeletalMesh->RefSkeleton.GetRefBonePose();
+
+			TArray<FBoneIndexType> RequiredBoneIndexArray;
+			RequiredBoneIndexArray.AddUninitialized(LocalAtoms.Num());
+			for(int32 BoneIndex = 0; BoneIndex < LocalAtoms.Num(); ++BoneIndex)
+			{
+				RequiredBoneIndexArray[BoneIndex] = BoneIndex;
+			}
+
+			RequiredBones.InitializeTo(RequiredBoneIndexArray, *SkeletalMesh);
+		}
+
+		FillSpaceBases();
+		FlipEditableSpaceBases();
+
+		return true;
+	}
+
+	LocalAtoms.Empty();
+
+	return false;
+}
+
+void UPoseableMeshComponent::RefreshBoneTransforms(FActorComponentTickFunction* TickFunction)
+{
+	SCOPE_CYCLE_COUNTER(STAT_RefreshBoneTransforms);
+
+	// Can't do anything without a SkeletalMesh
+	if( !SkeletalMesh )
+	{
+		return;
+	}
+
+	// Do nothing more if no bones in skeleton.
+	if( GetNumSpaceBases() == 0 )
+	{
+		return;
+	}
+
+	// We need the mesh space bone transforms now for renderer to get delta from ref pose:
+	FillSpaceBases();
+	FlipEditableSpaceBases();
+
+	MarkRenderDynamicDataDirty();
+}
+
+void UPoseableMeshComponent::FillSpaceBases()
+{
+	SCOPE_CYCLE_COUNTER(STAT_SkelComposeTime);
+
+	if( !SkeletalMesh )
+	{
+		return;
+	}
+
+	// right now all this does is to convert to SpaceBases
+	check( SkeletalMesh->RefSkeleton.GetNum() == LocalAtoms.Num() );
+	check( SkeletalMesh->RefSkeleton.GetNum() == GetNumSpaceBases());
+	check( SkeletalMesh->RefSkeleton.GetNum() == BoneVisibilityStates.Num() );
+
+	const int32 NumBones = LocalAtoms.Num();
+
+#if (UE_BUILD_DEBUG || UE_BUILD_DEVELOPMENT)
+	/** Keep track of which bones have been processed for fast look up */
+	TArray<uint8> BoneProcessed;
+	BoneProcessed.AddZeroed(NumBones);
+#endif
+	// Build in 3 passes.
+	FTransform* LocalTransformsData = LocalAtoms.GetData(); 
+	FTransform* SpaceBasesData = GetEditableSpaceBases().GetData();
+	
+	GetEditableSpaceBases()[0] = LocalAtoms[0];
+#if (UE_BUILD_DEBUG || UE_BUILD_DEVELOPMENT)
+	BoneProcessed[0] = 1;
+#endif
+
+	for(int32 BoneIndex=1; BoneIndex<LocalAtoms.Num(); BoneIndex++)
+	{
+		FPlatformMisc::Prefetch(SpaceBasesData + BoneIndex);
+
+#if (UE_BUILD_DEBUG || UE_BUILD_DEVELOPMENT)
+		// Mark bone as processed
+		BoneProcessed[BoneIndex] = 1;
+#endif
+		// For all bones below the root, final component-space transform is relative transform * component-space transform of parent.
+		const int32 ParentIndex = SkeletalMesh->RefSkeleton.GetParentIndex(BoneIndex);
+		FPlatformMisc::Prefetch(SpaceBasesData + ParentIndex);
+
+#if (UE_BUILD_DEBUG || UE_BUILD_DEVELOPMENT)
+		// Check the precondition that Parents occur before Children in the RequiredBones array.
+		checkSlow(BoneProcessed[ParentIndex] == 1);
+#endif
+		FTransform::Multiply(SpaceBasesData + BoneIndex, LocalTransformsData + BoneIndex, SpaceBasesData + ParentIndex);
+
+		checkSlow(GetEditableSpaceBases()[BoneIndex].IsRotationNormalized());
+		checkSlow(!GetEditableSpaceBases()[BoneIndex].ContainsNaN());
+	}
+	bNeedToFlipSpaceBaseBuffers = true;
+}
+
+void UPoseableMeshComponent::SetBoneTransformByName(FName BoneName, const FTransform& InTransform, EBoneSpaces::Type BoneSpace)
+{
+	if( !SkeletalMesh || !RequiredBones.IsValid() )
+	{
+		return;
+	}
+
+	int32 BoneIndex = GetBoneIndex(BoneName);
+	if(BoneIndex >=0 && BoneIndex < LocalAtoms.Num())
+	{
+		LocalAtoms[BoneIndex] = InTransform;
+
+		// If we haven't requested local space we need to transform the position passed in
+		//if(BoneSpace != EBoneSpaces::LocalSpace)
+		{
+			if(BoneSpace == EBoneSpaces::WorldSpace)
+			{
+				LocalAtoms[BoneIndex].SetToRelativeTransform(GetComponentToWorld());
+			}
+
+			int32 ParentIndex = RequiredBones.GetParentBoneIndex(BoneIndex);
+			if(ParentIndex >=0)
+			{
+				FA2CSPose CSPose;
+				CSPose.AllocateLocalPoses(RequiredBones, LocalAtoms);
+
+				LocalAtoms[BoneIndex].SetToRelativeTransform(CSPose.GetComponentSpaceTransform(ParentIndex));
+			}
+
+			// Need to send new state to render thread
+			MarkRenderDynamicDataDirty();
+		}
+	}
+}
+
+void UPoseableMeshComponent::SetBoneLocationByName(FName BoneName, FVector InLocation, EBoneSpaces::Type BoneSpace)
+{
+	FTransform CurrentTransform = GetBoneTransformByName(BoneName, BoneSpace);
+	CurrentTransform.SetLocation(InLocation);
+	SetBoneTransformByName(BoneName, CurrentTransform, BoneSpace);
+}
+
+void UPoseableMeshComponent::SetBoneRotationByName(FName BoneName, FRotator InRotation, EBoneSpaces::Type BoneSpace)
+{
+	FTransform CurrentTransform = GetBoneTransformByName(BoneName, BoneSpace);
+	CurrentTransform.SetRotation(FQuat(InRotation));
+	SetBoneTransformByName(BoneName, CurrentTransform, BoneSpace);
+}
+
+void UPoseableMeshComponent::SetBoneScaleByName(FName BoneName, FVector InScale3D, EBoneSpaces::Type BoneSpace)
+{
+	FTransform CurrentTransform = GetBoneTransformByName(BoneName, BoneSpace);
+	CurrentTransform.SetScale3D(InScale3D);
+	SetBoneTransformByName(BoneName, CurrentTransform, BoneSpace);
+}
+
+FTransform UPoseableMeshComponent::GetBoneTransformByName(FName BoneName, EBoneSpaces::Type BoneSpace)
+{
+	if( !SkeletalMesh || !RequiredBones.IsValid() )
+	{
+		return FTransform();
+	}
+
+	int32 BoneIndex = GetBoneIndex(BoneName);
+	if( BoneIndex == INDEX_NONE)
+	{
+		FString Message = FString::Printf(TEXT("Invalid Bone Name '%s'"), *BoneName.ToString());
+		FFrame::KismetExecutionMessage(*Message, ELogVerbosity::Warning);
+		return FTransform();
+	}
+
+	/*if(BoneSpace == EBoneSpaces::LocalSpace)
+	{
+		return LocalAtoms[i];
+	}*/
+
+	FA2CSPose CSPose;
+	CSPose.AllocateLocalPoses(RequiredBones, LocalAtoms);
+
+	if(BoneSpace == EBoneSpaces::ComponentSpace)
+	{
+		return CSPose.GetComponentSpaceTransform(BoneIndex);
+	}
+	else
+	{
+		return CSPose.GetComponentSpaceTransform(BoneIndex) * ComponentToWorld;
+	}
+}
+
+FVector UPoseableMeshComponent::GetBoneLocationByName(FName BoneName, EBoneSpaces::Type BoneSpace)
+{
+	FTransform CurrentTransform = GetBoneTransformByName(BoneName, BoneSpace);
+	return CurrentTransform.GetLocation();
+}
+
+FRotator UPoseableMeshComponent::GetBoneRotationByName(FName BoneName, EBoneSpaces::Type BoneSpace)
+{
+	FTransform CurrentTransform = GetBoneTransformByName(BoneName, BoneSpace);
+	return FRotator(CurrentTransform.GetRotation());
+}
+
+FVector UPoseableMeshComponent::GetBoneScaleByName(FName BoneName, EBoneSpaces::Type BoneSpace)
+{
+	FTransform CurrentTransform = GetBoneTransformByName(BoneName, BoneSpace);
+	return CurrentTransform.GetScale3D();
+}
+
+void UPoseableMeshComponent::ResetBoneTransformByName(FName BoneName)
+{
+	if( !SkeletalMesh )
+	{
+		return;
+	}
+
+	const int32 BoneIndex = GetBoneIndex(BoneName);
+	if( BoneIndex != INDEX_NONE )
+	{
+		LocalAtoms[BoneIndex] = SkeletalMesh->RefSkeleton.GetRefBonePose()[BoneIndex];
+	}
+	else
+	{
+		FString Message = FString::Printf(TEXT("Invalid Bone Name '%s'"), *BoneName.ToString());
+		FFrame::KismetExecutionMessage(*Message, ELogVerbosity::Warning);
+	}
+}
\ No newline at end of file
