Engine: Moved actor components into Components folder

diff --git a/Engine/Source/Runtime/Engine/Private/Components/PrimitiveComponent.cpp b/Engine/Source/Runtime/Engine/Private/Components/PrimitiveComponent.cpp
new file mode 100644
--- /dev/null
+++ b/Engine/Source/Runtime/Engine/Private/Components/PrimitiveComponent.cpp
@@ -0,0 +1,2673 @@
+// Copyright 1998-2015 Epic Games, Inc. All Rights Reserved.
+
+/*=============================================================================
+	PrimitiveComponent.cpp: Primitive component implementation.
+=============================================================================*/
+
+#include "EnginePrivate.h"
+#include "GameFramework/PhysicsVolume.h"
+#include "PhysicsPublic.h"
+#include "LevelUtils.h"
+#if WITH_EDITOR
+#include "ShowFlags.h"
+#include "Collision.h"
+#include "ConvexVolume.h"
+#endif
+#if WITH_PHYSX
+#include "PhysicsEngine/PhysXSupport.h"
+#include "Collision/PhysXCollision.h"
+#endif // WITH_PHYSX
+#if WITH_BOX2D
+#include "PhysicsEngine/BodySetup2D.h"
+#endif
+#include "Collision/CollisionDebugDrawing.h"
+#include "MessageLog.h"
+#include "UObjectToken.h"
+#include "MapErrors.h"
+#include "CollisionDebugDrawingPublic.h"
+#include "GameFramework/CheatManager.h"
+#include "GameFramework/DamageType.h"
+
+#define LOCTEXT_NAMESPACE "PrimitiveComponent"
+
+
+//////////////////////////////////////////////////////////////////////////
+// Globals
+
+DEFINE_LOG_CATEGORY_STATIC(LogPrimitiveComponent, Log, All);
+
+static FAutoConsoleVariable CVarAllowCachedOverlaps(
+	TEXT("p.AllowCachedOverlaps"), 
+	1,
+	TEXT("Primitive Component physics\n")
+	TEXT("0: disable cached overlaps, 1: enable (default)"));
+
+static TAutoConsoleVariable<float> CVarInitialOverlapTolerance(
+	TEXT("p.InitialOverlapTolerance"),
+	0.0f,
+	TEXT("Tolerance for initial overlapping test in PrimitiveComponent movement.\n")
+	TEXT("Normals within this tolerance are ignored if moving out of the object.\n")
+	TEXT("Dot product of movement direction and surface normal."),
+	ECVF_Default);
+
+#if !(UE_BUILD_SHIPPING || UE_BUILD_TEST)
+TAutoConsoleVariable<int32> CVarShowInitialOverlaps(
+	TEXT("p.ShowInitialOverlaps"),
+	0,
+	TEXT("Show initial overlaps when moving a component, including estimated 'exit' direction.\n")
+	TEXT(" 0:off, otherwise on"),
+	ECVF_Cheat);
+#endif // !(UE_BUILD_SHIPPING || UE_BUILD_TEST)
+
+///////////////////////////////////////////////////////////////////////////////
+// PRIMITIVE COMPONENT
+///////////////////////////////////////////////////////////////////////////////
+
+int32 UPrimitiveComponent::CurrentTag = 2147483647 / 4;
+
+// 0 is reserved to mean invalid
+uint32 UPrimitiveComponent::NextComponentId = 1;
+
+UPrimitiveComponent::UPrimitiveComponent(const FObjectInitializer& ObjectInitializer /*= FObjectInitializer::Get()*/)
+	: Super(ObjectInitializer)
+{
+	PostPhysicsComponentTick.bCanEverTick = false;
+	PostPhysicsComponentTick.bStartWithTickEnabled = true;
+	PostPhysicsComponentTick.TickGroup = TG_PostPhysics;
+
+	LastRenderTime = -1000.0f;
+	BoundsScale = 1.0f;
+	MinDrawDistance = 0.0f;
+	DepthPriorityGroup = SDPG_World;
+	bAllowCullDistanceVolume = true;
+	bUseAsOccluder = false;
+	bReceivesDecals = true;
+	CastShadow = false;
+	bCastDynamicShadow = true;
+	bAffectDynamicIndirectLighting = true;
+	bAffectDistanceFieldLighting = true;
+	LpvBiasMultiplier = 1.0f;
+	bCastStaticShadow = true;
+	bCastVolumetricTranslucentShadow = false;
+	IndirectLightingCacheQuality = ILCQ_Point;
+	bSelectable = true;
+	AlwaysLoadOnClient = true;
+	AlwaysLoadOnServer = true;
+	SetCollisionProfileName(UCollisionProfile::BlockAll_ProfileName);
+	bAlwaysCreatePhysicsState = false;
+	bRenderInMainPass = true;
+	VisibilityId = -1;
+	CanBeCharacterBase_DEPRECATED = ECB_Yes;
+	CanCharacterStepUpOn = ECB_Yes;
+	ComponentId.PrimIDValue = NextComponentId;
+	check(IsInGameThread());
+	NextComponentId++;
+
+	bUseEditorCompositing = false;
+
+	bGenerateOverlapEvents = true;
+	bMultiBodyOverlap = false;
+	bCheckAsyncSceneOnMove = false;
+	bReturnMaterialOnMove = false;
+	bCanEverAffectNavigation = false;
+	bNavigationRelevant = false;
+
+	bCachedAllCollideableDescendantsRelative = false;
+	LastCheckedAllCollideableDescendantsTime = 0.f;
+}
+
+bool UPrimitiveComponent::UsesOnlyUnlitMaterials() const
+{
+	return false;
+}
+
+
+bool UPrimitiveComponent::GetLightMapResolution( int32& Width, int32& Height ) const
+{
+	Width	= 0;
+	Height	= 0;
+	return false;
+}
+
+
+void UPrimitiveComponent::GetLightAndShadowMapMemoryUsage( int32& LightMapMemoryUsage, int32& ShadowMapMemoryUsage ) const
+{
+	LightMapMemoryUsage		= 0;
+	ShadowMapMemoryUsage	= 0;
+	return;
+}
+
+void UPrimitiveComponent::InvalidateLightingCacheDetailed(bool bInvalidateBuildEnqueuedLighting, bool bTranslationOnly) 
+{
+	bHasCachedStaticLighting = false;
+	if (bInvalidateBuildEnqueuedLighting)
+	{
+		bStaticLightingBuildEnqueued = false;
+	}
+
+	// If a static lighting build has been enqueued for this primitive, don't stomp on its visibility ID.
+	if (bInvalidateBuildEnqueuedLighting || !bStaticLightingBuildEnqueued)
+	{
+		VisibilityId = INDEX_NONE;
+	}
+
+	Super::InvalidateLightingCacheDetailed(bInvalidateBuildEnqueuedLighting, bTranslationOnly);
+}
+
+bool UPrimitiveComponent::IsEditorOnly() const
+{
+	return (AlwaysLoadOnClient == false) && (AlwaysLoadOnServer == false);
+}
+
+bool UPrimitiveComponent::HasStaticLighting() const
+{
+	return ((Mobility == EComponentMobility::Static) || bLightAsIfStatic) && SupportsStaticLighting();
+}
+
+
+void UPrimitiveComponent::GetUsedTextures(TArray<UTexture*>& OutTextures, EMaterialQualityLevel::Type QualityLevel)
+{
+	// Get the used materials so we can get their textures
+	TArray<UMaterialInterface*> UsedMaterials;
+	GetUsedMaterials( UsedMaterials );
+
+	TArray<UTexture*> UsedTextures;
+	for( int32 MatIndex = 0; MatIndex < UsedMaterials.Num(); ++MatIndex )
+	{
+		// Ensure we don't have any NULL elements.
+		if( UsedMaterials[ MatIndex ] )
+		{
+			auto World = GetWorld();
+
+			UsedTextures.Reset();
+			UsedMaterials[MatIndex]->GetUsedTextures(UsedTextures, QualityLevel, false, World ? World->FeatureLevel : GMaxRHIFeatureLevel, false);
+
+			for( int32 TextureIndex=0; TextureIndex<UsedTextures.Num(); TextureIndex++ )
+			{
+				OutTextures.AddUnique( UsedTextures[TextureIndex] );
+			}
+		}
+	}
+}
+
+/** 
+* Abstract function actually execute the tick. 
+* @param DeltaTime - frame time to advance, in seconds
+* @param TickType - kind of tick for this frame
+* @param CurrentThread - thread we are executing on, useful to pass along as new tasks are created
+* @param MyCompletionGraphEvent - completion event for this task. Useful for holding the completetion of this task until certain child tasks are complete.
+**/
+void FPrimitiveComponentPostPhysicsTickFunction::ExecuteTick(float DeltaTime, enum ELevelTick TickType, ENamedThreads::Type CurrentThread, const FGraphEventRef& MyCompletionGraphEvent)
+{
+	if ( (TickType == LEVELTICK_All) && Target && !Target->HasAnyFlags(RF_PendingKill | RF_Unreachable))
+	{
+		FScopeCycleCounterUObject ComponentScope(Target);
+		FScopeCycleCounterUObject AdditionalScope(Target->AdditionalStatObject());
+		Target->PostPhysicsTick(*this);	
+	}
+}
+
+/** Abstract function to describe this tick. Used to print messages about illegal cycles in the dependency graph **/
+FString FPrimitiveComponentPostPhysicsTickFunction::DiagnosticMessage()
+{
+	return Target->GetFullName() + TEXT("[UPrimitiveComponent::PostPhysicsTick]");
+}
+
+void UPrimitiveComponent::RegisterComponentTickFunctions(bool bRegister)
+{
+	Super::RegisterComponentTickFunctions(bRegister);
+
+	if (bRegister)
+	{
+		if (SetupActorComponentTickFunction(&PostPhysicsComponentTick))
+		{
+			PostPhysicsComponentTick.Target = this;
+
+			// If primary tick is registered, add a prerequisate to it
+			if(PrimaryComponentTick.bCanEverTick)
+			{
+				PostPhysicsComponentTick.AddPrerequisite(this,PrimaryComponentTick); 
+			}
+
+			// Set a prereq for the post physics tick to happen after physics and cloth is finished
+			if(World != NULL)
+			{
+				PostPhysicsComponentTick.AddPrerequisite(World, World->EndClothTickFunction);
+			}
+		}
+	}
+	else
+	{
+		if(PostPhysicsComponentTick.IsTickFunctionRegistered())
+		{
+			PostPhysicsComponentTick.UnRegisterTickFunction();
+		}
+	}
+}
+
+void UPrimitiveComponent::SetPostPhysicsComponentTickEnabled(bool bEnable)
+{
+	// @todo ticking, James, turn this off when not needed
+	if (!IsTemplate() && PostPhysicsComponentTick.bCanEverTick)
+	{
+		PostPhysicsComponentTick.SetTickFunctionEnable(bEnable);
+	}
+}
+
+bool UPrimitiveComponent::IsPostPhysicsComponentTickEnabled() const
+{
+	return PostPhysicsComponentTick.IsTickFunctionEnabled();
+}
+
+
+//////////////////////////////////////////////////////////////////////////
+// Render
+
+void UPrimitiveComponent::CreateRenderState_Concurrent()
+{
+	// Make sure cached cull distance is up-to-date if its zero and we have an LD cull distance
+	if( CachedMaxDrawDistance == 0 && LDMaxDrawDistance > 0 )
+	{
+		CachedMaxDrawDistance = LDMaxDrawDistance;
+	}
+
+	Super::CreateRenderState_Concurrent();
+
+	UpdateBounds();
+
+	// If the primitive isn't hidden and the detail mode setting allows it, add it to the scene.
+	if (ShouldComponentAddToScene())
+	{
+		World->Scene->AddPrimitive(this);
+	}
+}
+
+void UPrimitiveComponent::SendRenderTransform_Concurrent()
+{
+	UpdateBounds();
+
+	// If the primitive isn't hidden update its transform.
+	const bool bDetailModeAllowsRendering	= DetailMode <= GetCachedScalabilityCVars().DetailMode;
+	if( bDetailModeAllowsRendering && (ShouldRender() || bCastHiddenShadow))
+	{
+		// Update the scene info's transform for this primitive.
+		World->Scene->UpdatePrimitiveTransform(this);
+	}
+
+	Super::SendRenderTransform_Concurrent();
+}
+
+void UPrimitiveComponent::OnRegister()
+{
+	Super::OnRegister();
+
+	// Notify the streaming system. Will only update the component data if it's already tracked.
+	IStreamingManager::Get().NotifyPrimitiveUpdated(this);
+
+	if (bCanEverAffectNavigation)
+	{
+		const bool bNavRelevant = bNavigationRelevant = IsNavigationRelevant();
+		if (bNavRelevant)
+		{
+			UNavigationSystem::OnComponentRegistered(this);
+		}
+	}
+	else
+	{
+		bNavigationRelevant = false;
+	}
+}
+
+
+void UPrimitiveComponent::OnUnregister()
+{
+	if (World && World->Scene)
+	{
+		World->Scene->ReleasePrimitive(this);
+	}
+
+	Super::OnUnregister();
+
+	if (bCanEverAffectNavigation)
+	{
+		UNavigationSystem::OnComponentUnregistered(this);
+	}
+}
+
+FPrimitiveComponentInstanceData::FPrimitiveComponentInstanceData(const UPrimitiveComponent* SourceComponent)
+	: FSceneComponentInstanceData(SourceComponent)
+{
+}
+
+void FPrimitiveComponentInstanceData::ApplyToComponent(UActorComponent* Component, const ECacheApplyPhase CacheApplyPhase)
+{
+	FSceneComponentInstanceData::ApplyToComponent(Component, CacheApplyPhase);
+
+	if (ContainsSavedProperties() && Component->IsRegistered())
+	{
+		Component->MarkRenderStateDirty();
+	}
+}
+
+bool FPrimitiveComponentInstanceData::ContainsData() const
+{
+	return (ContainsSavedProperties() || AttachedInstanceComponents.Num() > 0);
+}
+
+FActorComponentInstanceData* UPrimitiveComponent::GetComponentInstanceData() const
+{
+	FPrimitiveComponentInstanceData* InstanceData = new FPrimitiveComponentInstanceData(this);
+
+	if (!InstanceData->ContainsData())
+	{
+		delete InstanceData;
+		InstanceData = nullptr;
+	}
+
+	return InstanceData;
+}
+
+FName UPrimitiveComponent::GetComponentInstanceDataType() const
+{
+	static const FName PrimitiveComponentInstanceDataTypeName(TEXT("PrimitiveComponentInstanceData"));
+	return PrimitiveComponentInstanceDataTypeName;
+}
+
+void UPrimitiveComponent::OnAttachmentChanged()
+{
+	if (World && World->Scene)
+	{
+		World->Scene->UpdatePrimitiveAttachment(this);
+	}
+}
+
+void UPrimitiveComponent::OnChildAttached(USceneComponent* ChildComponent)
+{
+	Super::OnChildAttached(ChildComponent);
+}
+
+
+void UPrimitiveComponent::DestroyRenderState_Concurrent()
+{
+	// Remove the primitive from the scene.
+	if(World && World->Scene)
+	{
+		World->Scene->RemovePrimitive(this);
+	}
+
+	Super::DestroyRenderState_Concurrent();
+}
+
+
+//////////////////////////////////////////////////////////////////////////
+// Physics
+
+void UPrimitiveComponent::CreatePhysicsState()
+{
+	Super::CreatePhysicsState();
+
+	// if we have a scene, we don't want to disable all physics and we have no bodyinstance already
+	if(!BodyInstance.IsValidBodyInstance())
+	{
+		//UE_LOG(LogPrimitiveComponent, Warning, TEXT("Creating Physics State (%s : %s)"), *GetNameSafe(GetOuter()),  *GetName());
+
+		UBodySetup* BodySetup = GetBodySetup();
+		if(BodySetup)
+		{
+			// Create new BodyInstance at given location.
+			if(FMath::Abs(ComponentToWorld.GetDeterminant()) < SMALL_NUMBER)
+			{
+				UE_LOG(LogPrimitiveComponent, Log, TEXT("Zero scaling not supported (%s)"), *GetPathName());
+				return;
+			}
+
+#if UE_WITH_PHYSICS
+			// Create the body.
+			BodyInstance.InitBody(BodySetup, ComponentToWorld, this, World->GetPhysicsScene());
+#endif //UE_WITH_PHYSICS
+
+#if WITH_EDITOR
+			// Make sure we have a valid body instance here. As we do not keep BIs with no collision shapes at all,
+			// we don't want to create cloth collision in these cases
+			if (BodyInstance.IsValidBodyInstance())
+			{
+				const float RealMass = BodyInstance.GetBodyMass();
+				const float CalcedMass = BodySetup->CalculateMass(this);
+				float MassDifference =  RealMass - CalcedMass;
+				if (RealMass > 1.0f && FMath::Abs(MassDifference) > 0.1f)
+				{
+					UE_LOG(LogPhysics, Log, TEXT("Calculated mass differs from real mass for %s:%s. Mass: %f  CalculatedMass: %f"),
+						GetOwner() != NULL ? *GetOwner()->GetName() : TEXT("NoActor"),
+						*GetName(), RealMass, CalcedMass);
+				}
+			}
+#endif // WITH_EDITOR
+		}
+	}
+}	
+
+void UPrimitiveComponent::EnsurePhysicsStateCreated()
+{
+	// if physics is created when it shouldn't OR if physics isn't created when it should
+	// we should fix it up
+	if ( IsPhysicsStateCreated() != ShouldCreatePhysicsState() )
+	{
+		RecreatePhysicsState();
+	}
+}
+
+void UPrimitiveComponent::OnUpdateTransform(bool bSkipPhysicsMove)
+{
+	Super::OnUpdateTransform(bSkipPhysicsMove);
+
+	// Always send new transform to physics
+	if(bPhysicsStateCreated && !bSkipPhysicsMove)
+	{
+		// @todo UE4 rather than always pass false, this function should know if it is being teleported or not!
+		SendPhysicsTransform(false);
+	}
+}
+
+void UPrimitiveComponent::SendPhysicsTransform(bool bTeleport)
+{
+	BodyInstance.SetBodyTransform(ComponentToWorld, bTeleport);
+	BodyInstance.UpdateBodyScale(ComponentToWorld.GetScale3D());
+}
+
+void UPrimitiveComponent::DestroyPhysicsState()
+{
+	// we remove welding related to this component
+	UnWeldFromParent();
+	UnWeldChildren();
+
+	// clean up physics engine representation
+	if(BodyInstance.IsValidBodyInstance())
+	{
+		// We tell the BodyInstance to shut down the physics-engine data.
+		BodyInstance.TermBody();
+	}
+
+	Super::DestroyPhysicsState();
+}
+
+FMatrix UPrimitiveComponent::GetRenderMatrix() const
+{
+	return ComponentToWorld.ToMatrixWithScale();
+}
+
+void UPrimitiveComponent::Serialize(FArchive& Ar)
+{
+	Super::Serialize(Ar);
+
+	if (Ar.UE4Ver() < VER_UE4_COMBINED_LIGHTMAP_TEXTURES)
+	{
+		bHasCachedStaticLighting = false;
+	}
+
+	// as temporary fix for the bug TTP 299926
+	// permanent fix is coming
+	if (IsTemplate())
+	{
+		BodyInstance.FixupData(this);
+	}
+}
+
+#if WITH_EDITOR
+void UPrimitiveComponent::PostEditChangeProperty(FPropertyChangedEvent& PropertyChangedEvent)
+{
+	static const FName NAME_SimulatePhysics = TEXT("bSimulatePhysics");
+	// Keep track of old cached cull distance to see whether we need to re-attach component.
+	const float OldCachedMaxDrawDistance = CachedMaxDrawDistance;
+
+	UProperty* PropertyThatChanged = PropertyChangedEvent.Property;
+	if(PropertyThatChanged)
+	{
+		const FName PropertyName = PropertyThatChanged->GetFName();
+
+		// We disregard cull distance volumes in this case as we have no way of handling cull 
+		// distance changes to without refreshing all cull distance volumes. Saving or updating 
+		// any cull distance volume will set the proper value again.
+		if( PropertyName == TEXT("LDMaxDrawDistance") || PropertyName == TEXT("bAllowCullDistanceVolume") )
+		{
+			CachedMaxDrawDistance = LDMaxDrawDistance;
+		}
+
+		if (PropertyName == NAME_SimulatePhysics)
+		{
+			if (AttachParent != NULL && IsSimulatingPhysics())
+			{
+				UE_LOG(LogPrimitiveComponent, Warning, TEXT("%s is simulating physics but has an attach parent. Detaching."), *GetPathName());
+				DetachFromParent();
+			}
+		}
+
+		// we need to reregister the primitive if the min draw distance changed to propagate the change to the rendering thread
+		if (PropertyThatChanged->GetFName() == GET_MEMBER_NAME_CHECKED(UPrimitiveComponent, MinDrawDistance))
+		{
+			MarkRenderStateDirty();
+		}
+	}
+
+	if (bLightAsIfStatic && GetStaticLightingType() == LMIT_None)
+	{
+		bLightAsIfStatic = false;
+	}
+
+	Super::PostEditChangeProperty(PropertyChangedEvent);
+
+	// Make sure cached cull distance is up-to-date.
+	if( LDMaxDrawDistance > 0 )
+	{
+		CachedMaxDrawDistance = FMath::Min( LDMaxDrawDistance, CachedMaxDrawDistance );
+	}
+	// Directly use LD cull distance if cull distance volumes are disabled.
+	if( !bAllowCullDistanceVolume )
+	{
+		CachedMaxDrawDistance = LDMaxDrawDistance;
+	}
+
+	// Reattach to propagate cull distance change.
+	if( CachedMaxDrawDistance != OldCachedMaxDrawDistance )
+	{
+		MarkRenderStateDirty();
+	}
+
+	// update component, ActorComponent's property update locks navigation system 
+	// so it needs to be called directly here
+	if (PropertyThatChanged && PropertyThatChanged->GetFName() == GET_MEMBER_NAME_CHECKED(UPrimitiveComponent, bCanEverAffectNavigation))
+	{
+		UNavigationSystem::UpdateNavOctree(this);
+	}
+}
+
+bool UPrimitiveComponent::CanEditChange(const UProperty* InProperty) const
+{
+	bool bIsEditable = Super::CanEditChange( InProperty );
+	if (bIsEditable && InProperty)
+	{
+		const FName PropertyName = InProperty->GetFName();
+
+		if (PropertyName == GET_MEMBER_NAME_CHECKED(UPrimitiveComponent, bLightAsIfStatic))
+		{
+			// Disable editing bLightAsIfStatic on static components, since it has no effect
+			return Mobility != EComponentMobility::Static;
+		}
+
+		if (PropertyName == TEXT("LightmassSettings"))
+		{
+			return Mobility != EComponentMobility::Movable || bLightAsIfStatic;
+		}
+
+		if (PropertyName == GET_MEMBER_NAME_CHECKED(UPrimitiveComponent, IndirectLightingCacheQuality)
+			|| PropertyName == GET_MEMBER_NAME_CHECKED(UPrimitiveComponent, bCastCinematicShadow))
+		{
+			return Mobility == EComponentMobility::Movable;
+		}
+
+		if (PropertyName == GET_MEMBER_NAME_CHECKED(UPrimitiveComponent, bCastInsetShadow))
+		{
+			return !bSelfShadowOnly;
+		}
+
+		if (PropertyName == GET_MEMBER_NAME_CHECKED(UPrimitiveComponent, CastShadow))
+		{
+			// Look for any lit materials
+			bool bHasAnyLitMaterials = false;
+			const int32 NumMaterials = GetNumMaterials();
+			for (int32 MaterialIndex = 0; (MaterialIndex < NumMaterials) && !bHasAnyLitMaterials; ++MaterialIndex)
+			{
+				if (UMaterialInterface* Material = GetMaterial(MaterialIndex))
+				{
+					if (Material->GetShadingModel() != MSM_Unlit)
+					{
+						bHasAnyLitMaterials = true;
+					}
+				}
+			}
+
+			// If there's at least one lit section it could cast shadows, so let the property be edited.
+			// The 0 materials catch is in case any components aren't properly implementing the GetMaterial API, they might or might not work with shadows.
+			return (NumMaterials == 0) || bHasAnyLitMaterials;
+		}
+	}
+
+	return bIsEditable;
+}
+
+void UPrimitiveComponent::PostEditChangeChainProperty(FPropertyChangedChainEvent& PropertyChangedEvent)
+{
+	for( FEditPropertyChain::TIterator It(PropertyChangedEvent.PropertyChain.GetHead()); It; ++It )
+	{
+		FName N = *It->GetName();
+		if( FCString::Stricmp( *It->GetName(), TEXT("Scale3D") )		== 0 || 
+			FCString::Stricmp( *It->GetName(), TEXT("Scale") )			== 0 || 
+			FCString::Stricmp( *It->GetName(), TEXT("Translation") )	== 0 || 
+			FCString::Stricmp( *It->GetName(), TEXT("Rotation") )		== 0 )
+		{
+			UpdateComponentToWorld();
+		}
+	}
+
+	Super::PostEditChangeChainProperty( PropertyChangedEvent );
+}
+
+void UPrimitiveComponent::CheckForErrors()
+{
+	AActor* Owner = GetOwner();
+
+	if (CastShadow && bCastDynamicShadow && BoundsScale > 1.0f)
+	{
+		FMessageLog("MapCheck").PerformanceWarning()
+			->AddToken(FUObjectToken::Create(Owner))
+			->AddToken(FTextToken::Create(LOCTEXT( "MapCheck_Message_ShadowCasterUsingBoundsScale", "Actor casts dynamic shadows and has a BoundsScale greater than 1! This will have a large performance hit" )))
+			->AddToken(FMapErrorToken::Create(FMapErrors::ShadowCasterUsingBoundsScale));
+	}
+
+	if (HasStaticLighting() && !HasValidSettingsForStaticLighting() && (!Owner || !Owner->IsA(AWorldSettings::StaticClass())))	// Ignore worldsettings
+	{
+		FMessageLog("MapCheck").Error()
+			->AddToken(FUObjectToken::Create(Owner))
+			->AddToken(FTextToken::Create(LOCTEXT( "MapCheck_Message_InvalidLightmapSettings", "Component is a static type but has invalid lightmap settings!  Indirect lighting will be black.  Common causes are lightmap resolution of 0, LightmapCoordinateIndex out of bounds." )))
+			->AddToken(FMapErrorToken::Create(FMapErrors::StaticComponentHasInvalidLightmapSettings));
+	}
+}
+
+void UPrimitiveComponent::UpdateCollisionProfile()
+{
+	BodyInstance.LoadProfileData(false);
+}
+#endif // WITH_EDITOR
+
+void UPrimitiveComponent::ReceiveComponentDamage(float DamageAmount, FDamageEvent const& DamageEvent, AController* EventInstigator, AActor* DamageCauser)
+{
+	UDamageType const* const DamageTypeCDO = DamageEvent.DamageTypeClass ? DamageEvent.DamageTypeClass->GetDefaultObject<UDamageType>() : GetDefault<UDamageType>();
+	if (DamageEvent.IsOfType(FPointDamageEvent::ClassID))
+	{
+		FPointDamageEvent* const PointDamageEvent = (FPointDamageEvent*) &DamageEvent;
+		if((DamageTypeCDO->DamageImpulse > 0.f) && !PointDamageEvent->ShotDirection.IsNearlyZero())
+		{
+			if (IsSimulatingPhysics(PointDamageEvent->HitInfo.BoneName))
+			{
+				FVector const ImpulseToApply = PointDamageEvent->ShotDirection.GetSafeNormal() * DamageTypeCDO->DamageImpulse;
+				AddImpulseAtLocation(ImpulseToApply, PointDamageEvent->HitInfo.ImpactPoint, PointDamageEvent->HitInfo.BoneName);
+			}
+		}
+	}
+	else if (DamageEvent.IsOfType(FRadialDamageEvent::ClassID))
+	{
+		FRadialDamageEvent* const RadialDamageEvent = (FRadialDamageEvent*) &DamageEvent;
+		if ( (DamageTypeCDO->DamageImpulse > 0.f) )
+		{
+			AddRadialImpulse(RadialDamageEvent->Origin, RadialDamageEvent->Params.OuterRadius, DamageTypeCDO->DamageImpulse, RIF_Linear, DamageTypeCDO->bRadialDamageVelChange);
+		}
+	}
+}
+
+void UPrimitiveComponent::PostLoad()
+{
+	Super::PostLoad();
+	
+	int32 const UE4Version = GetLinkerUE4Version();
+
+	// as temporary fix for the bug TTP 299926
+	// permanent fix is coming
+	if (IsTemplate()==false)
+	{
+		BodyInstance.FixupData(this);
+	}
+
+	if (UE4Version < VER_UE4_RENAME_CANBECHARACTERBASE)
+	{
+		CanCharacterStepUpOn = CanBeCharacterBase_DEPRECATED;
+	}
+
+	// Make sure cached cull distance is up-to-date.
+	if( LDMaxDrawDistance > 0 )
+	{
+		// Directly use LD cull distance if cached one is not set.
+		if( CachedMaxDrawDistance == 0 )
+		{
+			CachedMaxDrawDistance = LDMaxDrawDistance;
+		}
+		// Use min of both if neither is 0. Need to check as 0 has special meaning.
+		else
+		{
+			CachedMaxDrawDistance = FMath::Min( LDMaxDrawDistance, CachedMaxDrawDistance );
+		}
+	} 
+
+	if (bLightAsIfStatic && GetStaticLightingType() == LMIT_None)
+	{
+		bLightAsIfStatic = false;
+	}
+}
+
+void UPrimitiveComponent::PostDuplicate(bool bDuplicateForPIE)
+{
+	if (!bDuplicateForPIE)
+	{
+		VisibilityId = INDEX_NONE;
+	}
+
+	Super::PostDuplicate(bDuplicateForPIE);
+}
+
+#if WITH_EDITOR
+/**
+ * Called after importing property values for this object (paste, duplicate or .t3d import)
+ * Allow the object to perform any cleanup for properties which shouldn't be duplicated or
+ * are unsupported by the script serialization
+ */
+void UPrimitiveComponent::PostEditImport()
+{
+	Super::PostEditImport();
+
+	VisibilityId = INDEX_NONE;
+
+	// as temporary fix for the bug TTP 299926
+	// permanent fix is coming
+	if (IsTemplate()==false)
+	{
+		BodyInstance.FixupData(this);
+	}
+}
+#endif
+
+void UPrimitiveComponent::BeginDestroy()
+{
+	Super::BeginDestroy();
+
+	// Use a fence to keep track of when the rendering thread executes this scene detachment.
+	DetachFence.BeginFence();
+	AActor* Owner = GetOwner();
+	if(Owner)
+	{
+		Owner->DetachFence.BeginFence();
+	}
+}
+
+bool UPrimitiveComponent::IsReadyForFinishDestroy()
+{
+	// Don't allow the primitive component to the purged until its pending scene detachments have completed.
+	return Super::IsReadyForFinishDestroy() && DetachFence.IsFenceComplete();
+}
+
+void UPrimitiveComponent::FinishDestroy()
+{
+	// The detach fence has cleared so we better not be attached to the scene.
+	check(AttachmentCounter.GetValue() == 0);
+	Super::FinishDestroy();
+}
+
+bool UPrimitiveComponent::NeedsLoadForClient() const
+{
+	if(!IsVisible() && !IsCollisionEnabled() && !AlwaysLoadOnClient)
+	{
+		return 0;
+	}
+	else
+	{
+		return Super::NeedsLoadForClient();
+	}
+}
+
+
+bool UPrimitiveComponent::NeedsLoadForServer() const
+{
+	if(!IsCollisionEnabled() && !AlwaysLoadOnServer)
+	{
+		return 0;
+	}
+	else
+	{
+		return Super::NeedsLoadForServer();
+	}
+}
+
+void UPrimitiveComponent::SetOwnerNoSee(bool bNewOwnerNoSee)
+{
+	if(bOwnerNoSee != bNewOwnerNoSee)
+	{
+		bOwnerNoSee = bNewOwnerNoSee;
+		MarkRenderStateDirty();
+	}
+}
+
+void UPrimitiveComponent::SetOnlyOwnerSee(bool bNewOnlyOwnerSee)
+{
+	if(bOnlyOwnerSee != bNewOnlyOwnerSee)
+	{
+		bOnlyOwnerSee = bNewOnlyOwnerSee;
+		MarkRenderStateDirty();
+	}
+}
+
+bool UPrimitiveComponent::ShouldComponentAddToScene() const
+{
+	bool bSceneAdd = USceneComponent::ShouldComponentAddToScene();
+	return bSceneAdd && (ShouldRender() || bCastHiddenShadow);
+}
+
+bool UPrimitiveComponent::ShouldCreatePhysicsState() const
+{
+	bool bShouldCreatePhysicsState = IsRegistered() && (bAlwaysCreatePhysicsState || IsCollisionEnabled());
+
+#if WITH_EDITOR
+	if (BodyInstance.bSimulatePhysics && GetCollisionEnabled() == ECollisionEnabled::NoCollision)
+	{
+		FMessageLog("PIE").Warning(FText::Format(LOCTEXT("InvalidSimulateOptions", "Invalid Simulate Options: Body ({0}) is set to simulate physics but Collision Enabled is set to No Collision"),
+			FText::FromString(GetReadableName())));
+	}
+
+	// if it shouldn't create physics state, but if world wants to enable trace collision for components, allow it
+	if (!bShouldCreatePhysicsState && World && World->bEnableTraceCollision)
+	{
+		bShouldCreatePhysicsState = true;
+	}
+#endif
+	return bShouldCreatePhysicsState;
+}
+
+bool UPrimitiveComponent::HasValidPhysicsState() const
+{
+	return BodyInstance.IsValidBodyInstance();
+}
+
+bool UPrimitiveComponent::IsComponentIndividuallySelected() const
+{
+	bool bIndividuallySelected = false;
+#if WITH_EDITOR
+	if(SelectionOverrideDelegate.IsBound())
+	{
+		bIndividuallySelected = SelectionOverrideDelegate.Execute(this);
+	}
+#endif
+	return bIndividuallySelected;
+}
+
+bool UPrimitiveComponent::ShouldRenderSelected() const
+{
+	const AActor* Owner = GetOwner();
+	return(	bSelectable && 
+			Owner != NULL && 
+#if WITH_EDITOR
+			(Owner->IsSelected() || (Owner->ParentComponentActor != NULL && Owner->ParentComponentActor->IsSelected())) );
+#else
+			Owner->IsSelected() );
+#endif
+}
+
+void UPrimitiveComponent::SetCastShadow(bool NewCastShadow)
+{
+	if(NewCastShadow != CastShadow)
+	{
+		CastShadow = NewCastShadow;
+		MarkRenderStateDirty();
+	}
+}
+
+void UPrimitiveComponent::SetTranslucentSortPriority(int32 NewTranslucentSortPriority)
+{
+	if (NewTranslucentSortPriority != TranslucencySortPriority)
+	{
+		TranslucencySortPriority = NewTranslucentSortPriority;
+		MarkRenderStateDirty();
+	}
+}
+
+void UPrimitiveComponent::PushSelectionToProxy()
+{
+	//although this should only be called for attached components, some billboard components can get in without valid proxies
+	if (SceneProxy)
+	{
+		SceneProxy->SetSelection_GameThread(ShouldRenderSelected(),IsComponentIndividuallySelected());
+	}
+}
+
+
+void UPrimitiveComponent::PushEditorVisibilityToProxy( uint64 InVisibility )
+{
+	//although this should only be called for attached components, some billboard components can get in without valid proxies
+	if (SceneProxy)
+	{
+		SceneProxy->SetHiddenEdViews_GameThread( InVisibility );
+	}
+}
+
+
+void UPrimitiveComponent::PushHoveredToProxy(const bool bInHovered)
+{
+	//although this should only be called for attached components, some billboard components can get in without valid proxies
+	if (SceneProxy)
+	{
+		SceneProxy->SetHovered_GameThread(bInHovered);
+	}
+}
+
+void UPrimitiveComponent::SetCullDistance(float NewCullDistance)
+{
+	if (NewCullDistance > 0)
+	{
+		LDMaxDrawDistance = NewCullDistance;
+	
+		if (LDMaxDrawDistance < CachedMaxDrawDistance)
+		{
+			SetCachedMaxDrawDistance(LDMaxDrawDistance);
+		}
+	}
+}
+
+void UPrimitiveComponent::SetCachedMaxDrawDistance(const float NewCachedMaxDrawDistance)
+{
+	if( !FMath::IsNearlyEqual(CachedMaxDrawDistance, NewCachedMaxDrawDistance) )
+	{
+		CachedMaxDrawDistance = NewCachedMaxDrawDistance;
+		MarkRenderStateDirty();
+	}
+}
+
+
+void UPrimitiveComponent::SetDepthPriorityGroup(ESceneDepthPriorityGroup NewDepthPriorityGroup)
+{
+	if (DepthPriorityGroup != NewDepthPriorityGroup)
+	{
+		DepthPriorityGroup = NewDepthPriorityGroup;
+		MarkRenderStateDirty();
+	}
+}
+
+
+
+void UPrimitiveComponent::SetViewOwnerDepthPriorityGroup(
+	bool bNewUseViewOwnerDepthPriorityGroup,
+	ESceneDepthPriorityGroup NewViewOwnerDepthPriorityGroup
+	)
+{
+	bUseViewOwnerDepthPriorityGroup = bNewUseViewOwnerDepthPriorityGroup;
+	ViewOwnerDepthPriorityGroup = NewViewOwnerDepthPriorityGroup;
+	MarkRenderStateDirty();
+}
+
+bool UPrimitiveComponent::IsWorldGeometry() const
+{
+	// if modify flag doesn't change, and 
+	// it's saying it's movement is static, we considered to be world geom
+	// @Todo collision: we could check if bEnableCollision==true here as well
+	// but then if we disable collision, they just become non world geometry. 
+	// not sure if that would be best way to do this yet
+	return Mobility != EComponentMobility::Movable && GetCollisionObjectType()==ECC_WorldStatic;
+}
+
+ECollisionChannel UPrimitiveComponent::GetCollisionObjectType() const
+{
+	return ECollisionChannel(BodyInstance.GetObjectType());
+}
+
+UMaterialInterface* UPrimitiveComponent::GetMaterial(int32 Index) const
+{
+	return NULL;
+}
+
+void UPrimitiveComponent::SetMaterial(int32 Index, UMaterialInterface* InMaterial)
+{
+}
+
+int32 UPrimitiveComponent::GetNumMaterials() const
+{
+	return 0;
+}
+
+UMaterialInstanceDynamic* UPrimitiveComponent::CreateAndSetMaterialInstanceDynamic(int32 ElementIndex)
+{
+	UMaterialInterface* MaterialInstance = GetMaterial(ElementIndex);
+	UMaterialInstanceDynamic* MID = Cast<UMaterialInstanceDynamic>(MaterialInstance);
+
+	if (MaterialInstance && !MID)
+	{
+		// Create and set the dynamic material instance.
+		MID = UMaterialInstanceDynamic::Create(MaterialInstance, this);
+		SetMaterial(ElementIndex, MID);
+	}
+	else if (!MaterialInstance)
+	{
+		UE_LOG(LogPrimitiveComponent, Warning, TEXT("CreateAndSetMaterialInstanceDynamic on %s: Material index %d is invalid."), *GetPathName(), ElementIndex);
+	}
+
+	return MID;
+}
+
+UMaterialInstanceDynamic* UPrimitiveComponent::CreateAndSetMaterialInstanceDynamicFromMaterial(int32 ElementIndex, class UMaterialInterface* Parent)
+{
+	if (Parent)
+	{
+		SetMaterial(ElementIndex, Parent);
+		return CreateAndSetMaterialInstanceDynamic(ElementIndex);
+	}
+
+	return NULL;
+}
+
+UMaterialInstanceDynamic* UPrimitiveComponent::CreateDynamicMaterialInstance(int32 ElementIndex, class UMaterialInterface* SourceMaterial)
+{
+	if (SourceMaterial)
+	{
+		SetMaterial(ElementIndex, SourceMaterial);
+	}
+
+	UMaterialInterface* MaterialInstance = GetMaterial(ElementIndex);
+	UMaterialInstanceDynamic* MID = Cast<UMaterialInstanceDynamic>(MaterialInstance);
+
+	if (MaterialInstance && !MID)
+	{
+		// Create and set the dynamic material instance.
+		MID = UMaterialInstanceDynamic::Create(MaterialInstance, this);
+		SetMaterial(ElementIndex, MID);
+	}
+	else if (!MaterialInstance)
+	{
+		UE_LOG(LogPrimitiveComponent, Warning, TEXT("CreateDynamicMaterialInstance on %s: Material index %d is invalid."), *GetPathName(), ElementIndex);
+	}
+
+	return MID;
+}
+
+
+//////////////////////////////////////////////////////////////////////////
+// MOVECOMPONENT PROFILING CODE
+
+// LOOKING_FOR_PERF_ISSUES
+#define PERF_MOVECOMPONENT_STATS 0
+
+#if !(UE_BUILD_SHIPPING || UE_BUILD_TEST) && PERF_MOVECOMPONENT_STATS
+extern bool GShouldLogOutAFrameOfMoveComponent;
+
+/** 
+ *	Class to start/stop timer when it goes outside MoveComponent scope.
+ *	We keep all results from different MoveComponent calls until we reach the top level, and then print them all out.
+ *	That way we can show totals before breakdown, and not pollute timings with log time.
+ */
+class FScopedMoveCompTimer
+{
+	/** Struct to contain one temporary MoveActor timing, until we finish entire move and can print it out. */
+	struct FMoveTimer
+	{
+		AActor* Actor;
+		FVector Delta;
+		double Time;
+		int32 Depth;
+		bool bDidLineCheck;
+		bool bDidEncroachCheck;
+	};
+
+	/** Array of all moves within one 'outer' move. */
+	static TArray<FMoveTimer> Moves;
+	/** Current depth in movement hierarchy */
+	static int32 Depth;
+
+	/** Time that this scoped move started. */
+	double	StartTime;
+	/** Index into Moves array to put results of this MoveActor timing. */
+	int32		MoveIndex;
+public:
+
+	/** If we did a line check during this MoveActor. */
+	bool	bDidLineCheck;
+	/** If we did an encroach check during this MoveActor. */
+	bool	bDidEncroachCheck;
+
+	FScopedMoveCompTimer(AActor* Actor, const FVector& Delta)
+		: StartTime(0.0)
+	    , MoveIndex(-1)
+	    , bDidLineCheck(false)
+		, bDidEncroachCheck(false)
+	{
+		if(GShouldLogOutAFrameOfMoveComponent)
+		{
+			// Add new entry to temp results array, and save actor and current stack depth
+			MoveIndex = Moves.AddZeroed(1);
+			Moves(MoveIndex).Actor = Actor;
+			Moves(MoveIndex).Depth = Depth;
+			Moves(MoveIndex).Delta = Delta;
+			Depth++;
+
+			StartTime = FPlatformTime::Seconds(); // Start timer.
+		}
+	}
+
+	~FScopedMoveCompTimer()
+	{
+		if(GShouldLogOutAFrameOfMoveComponent)
+		{
+			// Record total time MoveActor took
+			const double TakeTime = FPlatformTime::Seconds() - StartTime;
+
+			check(Depth > 0);
+			check(MoveIndex < Moves.Num());
+
+			// Update entry with timing results
+			Moves(MoveIndex).Time = TakeTime;
+			Moves(MoveIndex).bDidLineCheck = bDidLineCheck;
+			Moves(MoveIndex).bDidEncroachCheck = bDidEncroachCheck;
+
+			Depth--;
+
+			// Reached the top of the move stack again - output what we accumulated!
+			if(Depth == 0)
+			{
+				for(int32 MoveIdx=0; MoveIdx<Moves.Num(); MoveIdx++)
+				{
+					const FMoveTimer& Move = Moves(MoveIdx);
+
+					// Build indentation
+					FString Indent;
+					for(int32 i=0; i<Move.Depth; i++)
+					{
+						Indent += TEXT("  ");
+					}
+
+					UE_LOG(LogPrimitiveComponent, Log, TEXT("MOVE%s - %s %5.2fms (%f %f %f) %d %d %s"), *Indent, *Move.Actor->GetName(), Move.Time * 1000.f, Move.Delta.X, Move.Delta.Y, Move.Delta.Z, Move.bDidLineCheck, Move.bDidEncroachCheck, *Move.Actor->GetDetailedInfo());
+				}
+
+				// Clear moves array
+				Moves.Reset();
+			}
+		}
+	}
+
+};
+
+// Init statics
+TArray<FScopedMoveCompTimer::FMoveTimer>	FScopedMoveCompTimer::Moves;
+int32											FScopedMoveCompTimer::Depth = 0;
+
+#endif //  !(UE_BUILD_SHIPPING || UE_BUILD_TEST) && PERF_MOVECOMPONENT_STATS
+
+
+static void PullBackHit(FHitResult& Hit, const FVector& Start, const FVector& End, const float Dist)
+{
+	const float DesiredTimeBack = FMath::Clamp(0.1f, 0.1f/Dist, 1.f/Dist) + 0.001f;
+	Hit.Time = FMath::Clamp( Hit.Time - DesiredTimeBack, 0.f, 1.f );
+}
+
+/**
+ * PERF_ISSUE_FINDER
+ *
+ * MoveComponent should not take a long time to execute.  If it is then then there is probably something wrong.
+ *
+ * Turn this on to have the engine log out when a specific actor is taking longer than 
+ * PERF_SHOW_MOVECOMPONENT_TAKING_LONG_TIME_AMOUNT to move.  This is a great way to catch cases where
+ * collision has been enabled but it should not have been.  Or if a specific actor is doing something evil
+ *
+ **/
+//#define PERF_SHOW_MOVECOMPONENT_TAKING_LONG_TIME 1
+const static float PERF_SHOW_MOVECOMPONENT_TAKING_LONG_TIME_AMOUNT = 2.0f; // modify this value to look at larger or smaller sets of "bad" actors
+
+
+static bool ShouldIgnoreHitResult(const UWorld* InWorld, FHitResult const& TestHit, FVector const& MovementDirDenormalized, const AActor* MovingActor, EMoveComponentFlags MoveFlags)
+{
+	if (TestHit.bBlockingHit)
+	{
+		// check "ignore bases" functionality
+		if ( (MoveFlags & MOVECOMP_IgnoreBases) && MovingActor )	//we let overlap components go through because their overlap is still needed and will cause beginOverlap/endOverlap events
+		{
+			// ignore if there's a base relationship between moving actor and hit actor
+			AActor const* const HitActor = TestHit.GetActor();
+			if (HitActor)
+			{
+				if (MovingActor->IsBasedOnActor(HitActor) || HitActor->IsBasedOnActor(MovingActor))
+				{
+					return true;
+				}
+			}
+		}
+	
+		// If we started penetrating, we may want to ignore it if we are moving out of penetration.
+		// This helps prevent getting stuck in walls.
+		if ( TestHit.bStartPenetrating && !(MoveFlags & MOVECOMP_NeverIgnoreBlockingOverlaps) )
+		{
+			const float DotTolerance = CVarInitialOverlapTolerance.GetValueOnGameThread();
+
+			// Dot product of movement direction against 'exit' direction
+			const FVector MovementDir = MovementDirDenormalized.GetSafeNormal();
+			const float MoveDot = (TestHit.ImpactNormal | MovementDir);
+
+			const bool bMovingOut = MoveDot > DotTolerance;
+
+	#if !(UE_BUILD_SHIPPING || UE_BUILD_TEST)
+			{
+				if(CVarShowInitialOverlaps.GetValueOnGameThread() != 0)
+				{
+					UE_LOG(LogTemp, Log, TEXT("Overlapping %s Dir %s Dot %f Normal %s Depth %f"), *GetNameSafe(TestHit.Component.Get()), *MovementDir.ToString(), MoveDot, *TestHit.ImpactNormal.ToString(), TestHit.PenetrationDepth);
+					DrawDebugDirectionalArrow(InWorld, TestHit.TraceStart, TestHit.TraceStart + 30.f * TestHit.ImpactNormal, 5.f, bMovingOut ? FColor(64,128,255) : FColor(255,64,64), true, 4.f);
+					if (TestHit.PenetrationDepth > KINDA_SMALL_NUMBER)
+					{
+						DrawDebugDirectionalArrow(InWorld, TestHit.TraceStart, TestHit.TraceStart + TestHit.PenetrationDepth * TestHit.Normal, 5.f, FColor(64,255,64), true, 4.f);
+					}
+				}
+			}
+	#endif
+
+			// If we are moving out, ignore this result!
+			if (bMovingOut)
+			{
+				return true;
+			}
+		}
+	}
+
+	return false;
+}
+
+static bool ShouldIgnoreOverlapResult(const UWorld* World, const AActor* ThisActor, const UPrimitiveComponent& ThisComponent, const AActor* OtherActor, const UPrimitiveComponent& OtherComponent)
+{
+	// Don't overlap with self
+	if (&ThisComponent == &OtherComponent)
+	{
+		return true;
+	}
+
+	// Both components must set bGenerateOverlapEvents
+	if (!ThisComponent.bGenerateOverlapEvents || !OtherComponent.bGenerateOverlapEvents)
+	{
+		return true;
+	}
+
+	if (!ThisActor || !OtherActor)
+	{
+		return true;
+	}
+
+	if (!World || OtherActor == World->GetWorldSettings() || !OtherActor->IsActorInitialized())
+	{
+		return true;
+	}
+
+	return false;
+}
+
+
+void UPrimitiveComponent::InitSweepCollisionParams(FCollisionQueryParams &OutParams, FCollisionResponseParams& OutResponseParam) const
+{
+	OutResponseParam.CollisionResponse = BodyInstance.GetResponseToChannels();
+	OutParams.AddIgnoredActors(MoveIgnoreActors);
+	OutParams.bTraceAsyncScene = bCheckAsyncSceneOnMove;
+	OutParams.bTraceComplex = bTraceComplexOnMove;
+	OutParams.bReturnPhysicalMaterial = bReturnMaterialOnMove;
+}
+
+
+FCollisionShape UPrimitiveComponent::GetCollisionShape(float Inflation) const
+{
+	// This is intended to be overridden by shape classes, so this is a simple, large bounding shape.
+	FVector Extent = Bounds.BoxExtent + Inflation;
+	if (Inflation < 0.f)
+	{
+		// Don't shrink below zero size.
+		Extent = Extent.ComponentMax(FVector::ZeroVector);
+	}
+
+	return FCollisionShape::MakeBox(Bounds.BoxExtent + Inflation);
+}
+
+bool UPrimitiveComponent::CheckStaticMobilityAndWarn(const FText& ActionText) const
+{
+	// static things can move before they are registered (e.g. immediately after streaming), but not after.
+	if (Mobility == EComponentMobility::Static)
+	{
+		AActor* Actor = GetOwner();
+		if (Actor && Actor->IsActorInitialized())
+		{
+			static const FText WarnText = LOCTEXT("InvalidStaticMove", "Mobility of {0} : {1} has to be 'Movable' if you'd like to {2}. ");
+			FMessageLog("PIE").Warning(FText::Format(WarnText,
+				FText::FromString(GetNameSafe(GetOwner())), FText::FromString(GetName()), ActionText));
+
+			return true;
+		}
+	}
+
+	return false;
+}
+
+bool UPrimitiveComponent::MoveComponent( const FVector& Delta, const FRotator& NewRotation, bool bSweep, FHitResult* OutHit, EMoveComponentFlags MoveFlags)
+{
+	SCOPE_CYCLE_COUNTER(STAT_MoveComponentTime);
+
+#if !(UE_BUILD_SHIPPING || UE_BUILD_TEST) && PERF_MOVECOMPONENT_STATS
+	FScopedMoveCompTimer MoveTimer(this, Delta);
+#endif // !(UE_BUILD_SHIPPING || UE_BUILD_TEST) && PERF_MOVECOMPONENT_STATS
+
+#if defined(PERF_SHOW_MOVECOMPONENT_TAKING_LONG_TIME) || LOOKING_FOR_PERF_ISSUES
+	uint32 MoveCompTakingLongTime=0;
+	CLOCK_CYCLES(MoveCompTakingLongTime);
+#endif
+
+	if ( IsPendingKill() )
+	{
+		//UE_LOG(LogPrimitiveComponent, Log, TEXT("%s deleted move physics %d"),*Actor->GetName(),Actor->Physics);
+		if (OutHit)
+		{
+			*OutHit = FHitResult();
+		}
+		return false;
+	}
+
+	AActor* const Actor = GetOwner();
+
+	// static things can move before they are registered (e.g. immediately after streaming), but not after.
+	// TODO: Static components without an owner can move, should they be able to?
+	static const FText WarnText = LOCTEXT("InvalidMove", "move");
+	if (CheckStaticMobilityAndWarn(WarnText))
+	{
+		if (OutHit)
+		{
+			*OutHit = FHitResult();
+		}
+		return false;
+	}
+
+	if (!bWorldToComponentUpdated)
+	{
+		UpdateComponentToWorld();
+	}
+
+	// Init HitResult
+	FHitResult BlockingHit(1.f);
+	const FVector TraceStart = GetComponentLocation();
+	const FVector TraceEnd = TraceStart + Delta;
+	BlockingHit.TraceStart = TraceStart;
+	BlockingHit.TraceEnd = TraceEnd;
+
+	// Set up.
+	float DeltaSizeSq = Delta.SizeSquared();
+
+	// ComponentSweepMulti does nothing if moving < KINDA_SMALL_NUMBER in distance, so it's important to not try to sweep distances smaller than that. 
+	const float MinMovementDistSq = (bSweep ? FMath::Square(4.f*KINDA_SMALL_NUMBER) : 0.f);
+	if (DeltaSizeSq <= MinMovementDistSq)
+	{
+		// Skip if no vector or rotation.
+		const FQuat NewQuat = NewRotation.Quaternion();
+		if( NewQuat.Equals(ComponentToWorld.GetRotation()) )
+		{
+			// copy to optional output param
+			if (OutHit)
+			{
+				*OutHit = BlockingHit;
+			}
+			return true;
+		}
+		DeltaSizeSq = 0.f;
+	}
+
+	const bool bSkipPhysicsMove = ((MoveFlags & MOVECOMP_SkipPhysicsMove) != MOVECOMP_NoFlags);
+
+	bool bMoved = false;
+	TArray<FOverlapInfo> PendingOverlaps;
+	TArray<FOverlapInfo> OverlapsAtEndLocation;
+	TArray<FOverlapInfo>* OverlapsAtEndLocationPtr = NULL; // When non-null, used as optimization to avoid work in UpdateOverlaps.
+	
+	if ( !bSweep )
+	{
+		// not sweeping, just go directly to the new transform
+		bMoved = InternalSetWorldLocationAndRotation(TraceEnd, NewRotation, bSkipPhysicsMove);
+	}
+	else
+	{
+		TArray<FHitResult> Hits;
+		FVector NewLocation = TraceStart;
+
+		// Perform movement collision checking if needed for this actor.
+		const bool bCollisionEnabled = IsCollisionEnabled();
+		if( bCollisionEnabled && (DeltaSizeSq > 0.f))
+		{
+#if !(UE_BUILD_SHIPPING || UE_BUILD_TEST)
+			if( !IsRegistered() )
+			{
+				if (Actor)
+				{
+					UE_LOG(LogPrimitiveComponent, Fatal,TEXT("%s MovedComponent %s not initialized deleteme %d"),*Actor->GetName(), *GetName(), Actor->IsPendingKill());
+				}
+				else
+				{
+					UE_LOG(LogPrimitiveComponent, Fatal,TEXT("MovedComponent %s not initialized"), *GetFullName());
+				}
+			}
+#endif
+
+#if !(UE_BUILD_SHIPPING || UE_BUILD_TEST) && PERF_MOVECOMPONENT_STATS
+			MoveTimer.bDidLineCheck = true;
+#endif 
+			static const FName Name_MoveComponent(TEXT("MoveComponent"));
+
+			FComponentQueryParams Params(Name_MoveComponent, Actor);
+			FCollisionResponseParams ResponseParam;
+			InitSweepCollisionParams(Params, ResponseParam);
+			bool const bHadBlockingHit = GetWorld()->ComponentSweepMulti(Hits, this, TraceStart, TraceEnd, GetComponentRotation(), Params);
+
+			if (Hits.Num() > 0)
+			{
+				const float DeltaSize = FMath::Sqrt(DeltaSizeSq);
+				for(int32 HitIdx=0; HitIdx<Hits.Num(); HitIdx++)
+				{
+					PullBackHit(Hits[HitIdx], TraceStart, TraceEnd, DeltaSize);
+				}
+			}
+
+			// If we had a valid blocking hit, store it.
+			// If we are looking for overlaps, store those as well.
+			uint32 FirstNonInitialOverlapIdx = INDEX_NONE;
+			if (bHadBlockingHit || bGenerateOverlapEvents)
+			{
+				int32 BlockingHitIndex = INDEX_NONE;
+				float BlockingHitNormalDotDelta = BIG_NUMBER;
+				for( int32 HitIdx = 0; HitIdx < Hits.Num(); HitIdx++ )
+				{
+					const FHitResult& TestHit = Hits[HitIdx];
+
+					if ( !ShouldIgnoreHitResult(GetWorld(), TestHit, Delta, Actor, MoveFlags) )
+					{
+						if (TestHit.bBlockingHit)
+						{
+							if (TestHit.Time == 0.f)
+							{
+								// We may have multiple initial hits, and want to choose the one with the normal most opposed to our movement.
+								const float NormalDotDelta = (TestHit.ImpactNormal | Delta);
+								if (NormalDotDelta < BlockingHitNormalDotDelta)
+								{
+									BlockingHitNormalDotDelta = NormalDotDelta;
+									BlockingHitIndex = HitIdx;
+								}
+							}
+							else if (BlockingHitIndex == INDEX_NONE)
+							{
+								// First non-overlapping blocking hit should be used, if an overlapping hit was not.
+								// This should be the only non-overlapping blocking hit, and last in the results.
+								BlockingHitIndex = HitIdx;
+								break;
+							}							
+						}
+						else if (bGenerateOverlapEvents)
+						{
+							UPrimitiveComponent * OverlapComponent = TestHit.Component.Get();
+							if (OverlapComponent && OverlapComponent->bGenerateOverlapEvents)
+							{
+								if (!ShouldIgnoreOverlapResult(GetWorld(), Actor, *this, TestHit.GetActor(), *OverlapComponent))
+								{
+									// don't process touch events after initial blocking hits
+									if (BlockingHitIndex >= 0 && TestHit.Time > Hits[BlockingHitIndex].Time)
+									{
+										break;
+									}
+
+									if (FirstNonInitialOverlapIdx == INDEX_NONE && TestHit.Time > 0.f)
+									{
+										// We are about to add the first non-initial overlap.
+										FirstNonInitialOverlapIdx = PendingOverlaps.Num();
+									}
+
+									// cache touches
+									PendingOverlaps.AddUnique(FOverlapInfo(TestHit));
+								}
+							}
+						}
+					}
+				}
+
+				// Update blocking hit, if there was a valid one.
+				if (BlockingHitIndex >= 0)
+				{
+					BlockingHit = Hits[BlockingHitIndex];
+				}
+			}
+		
+			// Update NewLocation based on the hit result
+			if (!BlockingHit.bBlockingHit)
+			{
+				NewLocation = TraceEnd;
+			}
+			else
+			{
+				NewLocation = TraceStart + (BlockingHit.Time * (TraceEnd - TraceStart));
+
+				// Sanity check
+				const FVector ToNewLocation = (NewLocation - TraceStart);
+				if (ToNewLocation.SizeSquared() <= MinMovementDistSq)
+				{
+					// We don't want really small movements to put us on or inside a surface.
+					NewLocation = TraceStart;
+					BlockingHit.Time = 0.f;
+
+					// Remove any pending overlaps after this point, we are not going as far as we swept.
+					if (FirstNonInitialOverlapIdx != INDEX_NONE)
+					{
+						PendingOverlaps.SetNum(FirstNonInitialOverlapIdx);
+					}
+				}
+			}
+
+			// We have performed a sweep that tested for all overlaps (not including components on the owning actor).
+			// However any rotation was set at the end and not swept, so we can't assume the overlaps at the end location are correct if rotation changed.
+			if (PendingOverlaps.Num() == 0 && CVarAllowCachedOverlaps->GetInt())
+			{
+				if (Actor && Actor->GetRootComponent() == this && AreSymmetricRotations(NewRotation.Quaternion(), ComponentToWorld.GetRotation(), ComponentToWorld.GetScale3D()))
+				{
+					// We know we are not overlapping any new components at the end location.
+					// Keep known overlapping child components, as long as we know their overlap status could not have changed (ie they are positioned relative to us).
+					if (AreAllCollideableDescendantsRelative())
+					{
+						GetOverlapsWithActor(Actor, OverlapsAtEndLocation);
+						OverlapsAtEndLocationPtr = &OverlapsAtEndLocation;
+					}
+				}
+			}
+
+#if !(UE_BUILD_SHIPPING || UE_BUILD_TEST)
+			if ( (BlockingHit.Time < 1.f) && !IsZeroExtent() )
+			{
+				// this is sole debug purpose to find how capsule trace information was when hit 
+				// to resolve stuck or improve our movement system - To turn this on, use DebugCapsuleSweepPawn
+				APawn const* const ActorPawn = (Actor ? Cast<APawn>(Actor) : NULL);
+				if (ActorPawn && ActorPawn->Controller && ActorPawn->Controller->IsLocalPlayerController())
+				{
+					APlayerController const* const PC = CastChecked<APlayerController>(ActorPawn->Controller);
+					if (PC->CheatManager && PC->CheatManager->bDebugCapsuleSweepPawn)
+					{
+						FVector CylExtent = ActorPawn->GetSimpleCollisionCylinderExtent()*FVector(1.001f,1.001f,1.0f);							
+						FCollisionShape CapsuleShape = FCollisionShape::MakeCapsule(CylExtent);
+						PC->CheatManager->AddCapsuleSweepDebugInfo(TraceStart, TraceEnd, BlockingHit.ImpactPoint, BlockingHit.Normal, BlockingHit.ImpactNormal, BlockingHit.Location, CapsuleShape.GetCapsuleHalfHeight(), CapsuleShape.GetCapsuleRadius(), true, (BlockingHit.bStartPenetrating && BlockingHit.bBlockingHit) ? true: false);
+					}
+				}
+			}
+#endif
+		}
+		else if (DeltaSizeSq > 0.f)
+		{
+			// apply move delta even if components has collisions disabled
+			NewLocation += Delta;
+		}
+		else if (DeltaSizeSq == 0.f && bCollisionEnabled)
+		{
+			// We didn't move, and any rotation that doesn't change our overlap bounds means we already know what we are overlapping at this point.
+			// Can only do this if current known overlaps are valid (not deferring updates)
+			if (CVarAllowCachedOverlaps->GetInt() && Actor && Actor->GetRootComponent() == this && !IsDeferringMovementUpdates())
+			{
+				// Only if we know that we won't change our overlap status with children by moving.
+				if (AreSymmetricRotations(NewRotation.Quaternion(), ComponentToWorld.GetRotation(), ComponentToWorld.GetScale3D()) &&
+					AreAllCollideableDescendantsRelative())
+				{
+					OverlapsAtEndLocation = OverlappingComponents;
+					OverlapsAtEndLocationPtr = &OverlapsAtEndLocation;
+				}
+			}
+		}
+
+		// Update the location.  This will teleport any child components as well (not sweep).
+		bMoved = InternalSetWorldLocationAndRotation(NewLocation, NewRotation, bSkipPhysicsMove);
+	}
+
+	// Handle overlap notifications.
+	if (bMoved)
+	{
+		// Check if we are deferring the movement updates.
+		if (IsDeferringMovementUpdates())
+		{
+			// Defer UpdateOverlaps until the scoped move ends.
+			FScopedMovementUpdate* ScopedUpdate = GetCurrentScopedMovement();
+			checkSlow(ScopedUpdate != NULL);
+			ScopedUpdate->AppendOverlaps(PendingOverlaps, OverlapsAtEndLocationPtr);
+		}
+		else
+		{
+			// still need to do this even if bGenerateOverlapEvents is false for this component, since we could have child components where it is true
+			UpdateOverlaps(&PendingOverlaps, true, OverlapsAtEndLocationPtr);
+		}
+	}
+
+	// Handle blocking hit notifications.
+	if (BlockingHit.bBlockingHit)
+	{
+		if (IsDeferringMovementUpdates())
+		{
+			FScopedMovementUpdate* ScopedUpdate = GetCurrentScopedMovement();
+			checkSlow(ScopedUpdate != NULL);
+			ScopedUpdate->AppendBlockingHit(BlockingHit);
+		}
+		else
+		{
+			DispatchBlockingHit(*Actor, BlockingHit);
+		}
+	}
+
+#if defined(PERF_SHOW_MOVECOMPONENT_TAKING_LONG_TIME) || LOOKING_FOR_PERF_ISSUES
+	UNCLOCK_CYCLES(MoveCompTakingLongTime);
+	const float MSec = FPlatformTime::ToMilliseconds(MoveCompTakingLongTime);
+	if( MSec > PERF_SHOW_MOVECOMPONENT_TAKING_LONG_TIME_AMOUNT )
+	{
+		if (Actor)
+		{
+			UE_LOG(LogPrimitiveComponent, Log, TEXT("%10f executing MoveComponent for %s owned by %s"), MSec, *GetName(), *Actor->GetFullName() );
+		}
+		else
+		{
+			UE_LOG(LogPrimitiveComponent, Log, TEXT("%10f executing MoveComponent for %s"), MSec, *GetFullName() );
+		}
+	}
+#endif
+
+	// copy to optional output param
+	if (OutHit)
+	{
+		*OutHit = BlockingHit;
+	}
+
+	// Return whether we moved at all.
+	return bMoved;
+}
+
+
+void UPrimitiveComponent::DispatchBlockingHit(AActor& Owner, FHitResult const& BlockingHit)
+{
+	Owner.DispatchBlockingHit(this, BlockingHit.Component.Get(), true, BlockingHit);
+
+	// BlockingHit.GetActor() could be marked for deletion in DispatchBlockingHit(), which would make the weak pointer return NULL.
+	if ( BlockingHit.GetActor() != nullptr && BlockingHit.Component.Get() != nullptr)
+	{
+		BlockingHit.GetActor()->DispatchBlockingHit(BlockingHit.Component.Get(), this, false, BlockingHit);
+	}
+}
+
+
+bool UPrimitiveComponent::IsNavigationRelevant() const 
+{ 
+	if (!CanEverAffectNavigation())
+	{
+		return false;
+	}
+
+	if (HasCustomNavigableGeometry() >= EHasCustomNavigableGeometry::EvenIfNotCollidable)
+	{
+		return true;
+	}
+
+	const FCollisionResponseContainer& ResponseToChannels = GetCollisionResponseToChannels();
+	return IsCollisionEnabled() &&
+		(ResponseToChannels.GetResponse(ECC_Pawn) == ECR_Block || ResponseToChannels.GetResponse(ECC_Vehicle) == ECR_Block);
+}
+
+FBox UPrimitiveComponent::GetNavigationBounds() const
+{
+	return Bounds.GetBox();
+}
+
+void UPrimitiveComponent::SetCanEverAffectNavigation(bool bRelevant)
+{
+	if (bCanEverAffectNavigation != bRelevant)
+	{
+		bCanEverAffectNavigation = bRelevant;
+
+		// update octree if already registered
+		if (bRegistered)
+		{
+			if (bRelevant)
+			{
+				bNavigationRelevant = IsNavigationRelevant();
+				UNavigationSystem::OnComponentRegistered(this);
+			}
+			else
+			{
+				UNavigationSystem::OnComponentUnregistered(this);
+			}
+		}
+	}
+}
+
+//////////////////////////////////////////////////////////////////////////
+// COLLISION
+
+extern float DebugLineLifetime;
+
+
+bool UPrimitiveComponent::LineTraceComponent(struct FHitResult& OutHit, const FVector Start, const FVector End, const struct FCollisionQueryParams& Params)
+{
+	bool bHaveHit = BodyInstance.LineTrace(OutHit, Start, End, Params.bTraceComplex, Params.bReturnPhysicalMaterial); 
+
+#if !(UE_BUILD_SHIPPING || UE_BUILD_TEST)
+	if((GetWorld()->DebugDrawTraceTag != NAME_None) && (GetWorld()->DebugDrawTraceTag == Params.TraceTag))
+	{
+		TArray<FHitResult> Hits;
+		if (bHaveHit)
+		{
+			Hits.Add(OutHit);
+		}
+
+		DrawLineTraces(GetWorld(), Start, End, Hits, DebugLineLifetime);
+	}
+#endif //!(UE_BUILD_SHIPPING || UE_BUILD_TEST)
+	return bHaveHit;
+}
+
+// @Todo change this to shape with sweep
+bool UPrimitiveComponent::SweepComponent(struct FHitResult& OutHit, const FVector Start, const FVector End, const FCollisionShape &CollisionShape, bool bTraceComplex)
+{
+	return BodyInstance.Sweep(OutHit, Start, End, CollisionShape, bTraceComplex);
+}
+
+bool UPrimitiveComponent::ComponentOverlapComponent(class UPrimitiveComponent* PrimComp, const FVector Pos, const FRotator Rot, const struct FCollisionQueryParams& Params)
+{
+	//@TODO: BOX2D: Implement UPrimitiveComponent::ComponentOverlapComponent
+
+	// if target is skeletalmeshcomponent and do not support singlebody physics
+	USkeletalMeshComponent * OtherComp = Cast<USkeletalMeshComponent>(PrimComp);
+	if (OtherComp)
+	{
+		UE_LOG(LogCollision, Log, TEXT("ComponentOverlapMulti : (%s) Does not support skeletalmesh with Physics Asset"), *PrimComp->GetPathName());
+		return false;
+	}
+#if WITH_PHYSX
+	// will have to do per component - default single body or physicsinstance 
+	const PxRigidActor* TargetRigidBody = (PrimComp)? PrimComp->BodyInstance.GetPxRigidActor():NULL;
+	if (TargetRigidBody==NULL || TargetRigidBody->getNbShapes()==0)
+	{
+		return false;
+	}
+
+	// calculate the test global pose of the actor
+	PxTransform PTestGlobalPose = U2PTransform(FTransform(Rot, Pos));
+	// Get all the shapes from the actor
+	TArray<PxShape*> PTargetShapes;
+	PTargetShapes.AddZeroed(TargetRigidBody->getNbShapes());
+	int32 NumTargetShapes = TargetRigidBody->getShapes(PTargetShapes.GetData(), PTargetShapes.Num());
+
+	bool bHaveOverlap = false;
+
+	for (int32 TargetShapeIdx=0; TargetShapeIdx<PTargetShapes.Num(); ++TargetShapeIdx)
+	{
+		const PxShape * PTargetShape = PTargetShapes[TargetShapeIdx];
+		check (PTargetShape);
+
+		// Calc shape global pose
+		PxTransform PShapeGlobalPose = PTestGlobalPose.transform(PTargetShape->getLocalPose());
+
+		GET_GEOMETRY_FROM_SHAPE(PGeom, PTargetShape);
+
+		if(PGeom != NULL)
+		{
+			bHaveOverlap = BodyInstance.OverlapPhysX(*PGeom, PShapeGlobalPose);
+
+#if !(UE_BUILD_SHIPPING || UE_BUILD_TEST)
+			if((GetWorld()->DebugDrawTraceTag != NAME_None) && (GetWorld()->DebugDrawTraceTag == Params.TraceTag))
+			{
+				TArray<FOverlapResult> Overlaps;
+				if (bHaveOverlap)
+				{
+					FOverlapResult Result;
+					Result.Actor = PrimComp->GetOwner();
+					Result.Component = PrimComp;
+					Result.bBlockingHit = true;
+					Overlaps.Add(Result);
+				}
+
+				DrawGeomOverlaps(GetWorld(), *PGeom, PShapeGlobalPose, Overlaps, DebugLineLifetime);
+			}
+#endif //!(UE_BUILD_SHIPPING || UE_BUILD_TEST)
+
+			if (bHaveOverlap)
+			{
+				break;
+			}
+		}
+	}
+
+	return bHaveOverlap;
+#endif //WITH_PHYSX
+	return false;
+}
+
+
+bool UPrimitiveComponent::OverlapComponent(const FVector& Pos, const FQuat& Rot, const struct FCollisionShape& CollisionShape)
+{
+	return BodyInstance.OverlapTest(Pos, Rot, CollisionShape);
+}
+
+bool UPrimitiveComponent::ComputePenetration(FMTDResult& OutMTD, const FCollisionShape & CollisionShape, const FVector& Pos, const FQuat& Rot)
+{
+	if(FBodyInstance* BodyInstance = GetBodyInstance())
+	{
+		return BodyInstance->OverlapTest(Pos, Rot, CollisionShape, &OutMTD);
+	}
+
+	return false;
+}
+
+bool UPrimitiveComponent::IsOverlappingComponent(UPrimitiveComponent const* OtherComp) const
+{
+	for (int32 i=0; i < OverlappingComponents.Num(); ++i)
+	{
+		if (OverlappingComponents[i].OverlapInfo.Component.Get() == OtherComp)
+		{
+			return true;
+		}
+	}
+	return false;
+}
+
+bool UPrimitiveComponent::IsOverlappingComponent( const FOverlapInfo& Overlap ) const
+{
+	return OverlappingComponents.Find(Overlap) != INDEX_NONE;
+}
+
+bool UPrimitiveComponent::IsOverlappingActor(const AActor* Other) const
+{
+	if (Other)
+	{
+		for (int32 OverlapIdx=0; OverlapIdx<OverlappingComponents.Num(); ++OverlapIdx)
+		{
+			UPrimitiveComponent const* const PrimComp = OverlappingComponents[OverlapIdx].OverlapInfo.Component.Get();
+			if ( PrimComp && (PrimComp->GetOwner() == Other) )
+			{
+				return true;
+			}
+		}
+	}
+
+	return false;
+}
+
+bool UPrimitiveComponent::GetOverlapsWithActor(const AActor* Actor, TArray<FOverlapInfo>& OutOverlaps) const
+{
+	const int32 InitialCount = OutOverlaps.Num();
+	if (Actor)
+	{
+		for (int32 OverlapIdx=0; OverlapIdx<OverlappingComponents.Num(); ++OverlapIdx)
+		{
+			UPrimitiveComponent const* const PrimComp = OverlappingComponents[OverlapIdx].OverlapInfo.Component.Get();
+			if ( PrimComp && (PrimComp->GetOwner() == Actor) )
+			{
+				OutOverlaps.Add(OverlappingComponents[OverlapIdx]);
+			}
+		}
+	}
+
+	return InitialCount != OutOverlaps.Num();
+}
+
+#if WITH_EDITOR
+bool UPrimitiveComponent::ComponentIsTouchingSelectionBox(const FBox& InSelBBox, const FEngineShowFlags& ShowFlags, const bool bConsiderOnlyBSP, const bool bMustEncompassEntireComponent) const
+{
+	if (!bConsiderOnlyBSP)
+	{
+		const FBox& ComponentBounds = Bounds.GetBox();
+
+		// Check the component bounds versus the selection box
+		// If the selection box must encompass the entire component, then both the min and max vector of the bounds must be inside in the selection
+		// box to be valid. If the selection box only has to touch the component, then it is sufficient to check if it intersects with the bounds.
+		if ((!bMustEncompassEntireComponent && InSelBBox.Intersect(ComponentBounds))
+			|| (bMustEncompassEntireComponent && InSelBBox.IsInside(ComponentBounds)))
+		{
+			return true;
+		}
+	}
+
+	return false;
+}
+
+bool UPrimitiveComponent::ComponentIsTouchingSelectionFrustum(const FConvexVolume& InFrustum, const FEngineShowFlags& ShowFlags, const bool bConsiderOnlyBSP, const bool bMustEncompassEntireComponent) const
+{
+	if (!bConsiderOnlyBSP)
+	{
+		bool bIsFullyContained;
+		if (InFrustum.IntersectBox(Bounds.Origin, Bounds.BoxExtent, bIsFullyContained))
+		{
+			return !bMustEncompassEntireComponent || bIsFullyContained;
+		}
+	}
+
+	return false;
+}
+#endif
+
+
+
+/** Used to determine if it is ok to call a notification on this object */
+extern bool IsActorValidToNotify(AActor* Actor);
+
+// @fixme, duplicated, make an inline member?
+bool IsPrimCompValidAndAlive(UPrimitiveComponent* PrimComp)
+{
+	return (PrimComp != NULL) && !PrimComp->IsPendingKill();
+}
+
+// @todo, don't need to pass in Other actor?
+void UPrimitiveComponent::BeginComponentOverlap(const FOverlapInfo& OtherOverlap, bool bDoNotifies)
+{
+	// If pending kill, we should not generate any new overlaps
+	if (IsPendingKill())
+	{
+		return;
+	}
+
+	//	UE_LOG(LogActor, Log, TEXT("BEGIN OVERLAP! Self=%s SelfComp=%s, Other=%s, OtherComp=%s"), *GetNameSafe(this), *GetNameSafe(MyComp), *GetNameSafe(OtherComp->GetOwner()), *GetNameSafe(OtherComp));
+
+	UPrimitiveComponent* OtherComp = OtherOverlap.OverlapInfo.Component.Get();
+
+	if (OtherComp)
+	{
+		bool const bComponentsAlreadyTouching = IsOverlappingComponent(OtherOverlap);
+		if (!bComponentsAlreadyTouching)
+		{
+			AActor* const OtherActor = OtherComp->GetOwner();
+			AActor* const MyActor = GetOwner();
+
+			const bool bNotifyActorTouch = bDoNotifies && !MyActor->IsOverlappingActor(OtherActor);
+
+			// Perform reflexive touch.
+			OverlappingComponents.Add(OtherOverlap);										// already verified uniqueness above
+			OtherComp->OverlappingComponents.AddUnique(FOverlapInfo(this, INDEX_NONE));		// uniqueness unverified, so addunique
+			
+			if (bDoNotifies)
+			{
+				// first execute component delegates
+				if (!IsPendingKill())
+				{
+					OnComponentBeginOverlap.Broadcast(OtherActor, OtherComp, OtherOverlap.GetBodyIndex(), OtherOverlap.bFromSweep, OtherOverlap.OverlapInfo);
+				}
+
+				if (!OtherComp->IsPendingKill())
+				{
+					OtherComp->OnComponentBeginOverlap.Broadcast(MyActor, this, INDEX_NONE, OtherOverlap.bFromSweep, OtherOverlap.OverlapInfo);
+				}
+
+				// then execute actor notification if this is a new actor touch
+				if (bNotifyActorTouch)
+				{
+					// First actor virtuals
+					if (IsActorValidToNotify(MyActor))
+					{
+						MyActor->ReceiveActorBeginOverlap(OtherActor);
+					}
+
+					if (IsActorValidToNotify(OtherActor))
+					{
+						OtherActor->ReceiveActorBeginOverlap(MyActor);
+					}
+
+					// Then level-script delegates
+					if (IsActorValidToNotify(MyActor))
+					{
+						MyActor->OnActorBeginOverlap.Broadcast(OtherActor);
+					}
+
+					if (IsActorValidToNotify(OtherActor))
+					{
+						OtherActor->OnActorBeginOverlap.Broadcast(MyActor);
+					}
+				}
+			}
+		}
+	}
+}
+
+
+void UPrimitiveComponent::EndComponentOverlap(const FOverlapInfo& OtherOverlap, bool bDoNotifies, bool bNoNotifySelf)
+{
+	UPrimitiveComponent* OtherComp = OtherOverlap.OverlapInfo.Component.Get();
+
+	AActor* const OtherActor = OtherComp ? OtherComp->GetOwner() : NULL;
+	AActor* const MyActor = GetOwner();
+
+	//	UE_LOG(LogActor, Log, TEXT("END OVERLAP! Self=%s SelfComp=%s, Other=%s, OtherComp=%s"), *GetNameSafe(this), *GetNameSafe(MyComp), *GetNameSafe(OtherActor), *GetNameSafe(OtherComp));
+
+	const int32 OtherOverlapIdx = OtherComp ? OtherComp->OverlappingComponents.Find(FOverlapInfo(this, INDEX_NONE)) : INDEX_NONE;
+	if (OtherOverlapIdx != INDEX_NONE)
+	{
+		OtherComp->OverlappingComponents.RemoveAtSwap(OtherOverlapIdx, 1, false);
+	}
+
+	const int32 OverlapIdx = OverlappingComponents.Find(OtherOverlap);
+	if (OverlapIdx != INDEX_NONE)
+	{
+		OverlappingComponents.RemoveAtSwap(OverlapIdx, 1, false);
+	
+		if ((OtherActor != NULL) && bDoNotifies)
+		{
+			if (!bNoNotifySelf && IsPrimCompValidAndAlive(this))
+			{
+				OnComponentEndOverlap.Broadcast(OtherActor, OtherComp, OtherOverlap.GetBodyIndex());
+			}
+
+			if (IsPrimCompValidAndAlive(OtherComp))
+			{
+				OtherComp->OnComponentEndOverlap.Broadcast(MyActor, this, INDEX_NONE);
+			}
+		}
+	}
+
+	// if this was the last touch on the other actor, notify that we've untouched the actor as well
+	if (bDoNotifies && MyActor && OtherActor && !MyActor->IsOverlappingActor(OtherActor) )
+	{			
+		if (IsActorValidToNotify(MyActor))
+		{
+			MyActor->ReceiveActorEndOverlap(OtherActor);
+			MyActor->OnActorEndOverlap.Broadcast(OtherActor);
+		}
+
+		if (IsActorValidToNotify(OtherActor))
+		{
+			OtherActor->ReceiveActorEndOverlap(MyActor);
+			OtherActor->OnActorEndOverlap.Broadcast(MyActor);
+		}
+	}
+}
+
+
+
+void UPrimitiveComponent::GetOverlappingActors(TArray<AActor*>& OutOverlappingActors, UClass* ClassFilter) const
+{
+	OutOverlappingActors.Reset();
+
+	for (auto CompIt = OverlappingComponents.CreateConstIterator(); CompIt; ++CompIt)
+	{
+		const FOverlapInfo& OtherOvelap = *CompIt;
+		
+		if (OtherOvelap.OverlapInfo.Component.IsValid())
+		{
+			AActor* OtherActor = OtherOvelap.OverlapInfo.Component->GetOwner();
+			if (OtherActor)
+			{
+				if ( !ClassFilter || OtherActor->IsA(ClassFilter) )
+				{
+					OutOverlappingActors.AddUnique(OtherActor);
+				}
+			}
+		}
+	}
+}
+
+
+void UPrimitiveComponent::GetOverlappingComponents(TArray<UPrimitiveComponent*>& OutOverlappingComponents) const
+{
+	OutOverlappingComponents.Reset( OverlappingComponents.Num() );
+
+	for (auto CompIt = OverlappingComponents.CreateConstIterator(); CompIt; ++CompIt)
+	{
+		UPrimitiveComponent* const OtherComp = (*CompIt).OverlapInfo.Component.Get();
+		if (OtherComp)
+		{
+			OutOverlappingComponents.Add(OtherComp);
+		}
+	}
+}
+
+
+void UPrimitiveComponent::UpdateNavigationData()
+{
+	if (bNavigationRelevant)
+	{
+		Super::UpdateNavigationData();
+	}
+}
+
+
+bool UPrimitiveComponent::AreAllCollideableDescendantsRelative(bool bAllowCachedValue) const
+{
+	UPrimitiveComponent* MutableThis = const_cast<UPrimitiveComponent*>(this);
+	if (AttachChildren.Num() > 0)
+	{
+		UWorld* MyWorld = GetWorld();
+		check(MyWorld);
+
+		// Throttle this test when it has been false in the past, since it rarely changes afterwards.
+		if (bAllowCachedValue && !bCachedAllCollideableDescendantsRelative && MyWorld->TimeSince(LastCheckedAllCollideableDescendantsTime) < 1.f)
+		{
+			return false;
+		}
+
+		// Check all descendant PrimitiveComponents
+		TInlineComponentArray<USceneComponent*> ComponentStack;
+		ComponentStack.Reserve(FMath::Max<uint32>(NumInlinedActorComponents, AttachChildren.Num()));
+
+		ComponentStack.Append(AttachChildren);
+		while (ComponentStack.Num() > 0)
+		{
+			USceneComponent* const CurrentComp = ComponentStack.Pop(false);
+			if (CurrentComp)
+			{
+				// Is the component not using relative position?
+				if (CurrentComp->bAbsoluteLocation || CurrentComp->bAbsoluteRotation)
+				{
+					// Can we possibly collide with the component?
+					if (CurrentComp->IsCollisionEnabled() && CurrentComp->GetCollisionResponseToChannel(GetCollisionObjectType()) != ECR_Ignore)
+					{
+						MutableThis->bCachedAllCollideableDescendantsRelative = false;
+						MutableThis->LastCheckedAllCollideableDescendantsTime = MyWorld->GetTimeSeconds();
+						return false;
+					}
+				}
+
+				ComponentStack.Append(CurrentComp->AttachChildren);
+			}
+		}
+	}
+
+	MutableThis->bCachedAllCollideableDescendantsRelative = true;
+	return true;
+}
+
+const TArray<FOverlapInfo>& UPrimitiveComponent::GetOverlapInfos() const
+{
+	return OverlappingComponents;
+}
+
+void UPrimitiveComponent::IgnoreActorWhenMoving(AActor* Actor, bool bShouldIgnore)
+{
+	// Clean up stale references
+	MoveIgnoreActors.RemoveSwap(NULL);
+
+	// Add/Remove the actor from the list
+	if (Actor)
+	{
+		if (bShouldIgnore)
+		{
+			MoveIgnoreActors.AddUnique(Actor);
+		}
+		else
+		{
+			MoveIgnoreActors.RemoveSingleSwap(Actor);
+		}
+	}
+}
+
+TArray<TWeakObjectPtr<AActor> > & UPrimitiveComponent::GetMoveIgnoreActors()
+{
+	// Clean up stale references
+	MoveIgnoreActors.RemoveSwap(NULL);
+	return MoveIgnoreActors;
+}
+
+void UPrimitiveComponent::ClearMoveIgnoreActors()
+{
+	MoveIgnoreActors.Empty();
+}
+
+
+void UPrimitiveComponent::UpdateOverlaps(TArray<FOverlapInfo> const* PendingOverlaps, bool bDoNotifies, const TArray<FOverlapInfo>* OverlapsAtEndLocation)
+{
+	SCOPE_CYCLE_COUNTER(STAT_UpdateOverlaps); 
+
+	if (IsDeferringMovementUpdates())
+	{
+		return;
+	}
+
+	// first, dispatch any pending overlaps
+	if (bGenerateOverlapEvents && IsCollisionEnabled())
+	{
+		// if we haven't begun play, we're still setting things up (e.g. we might be inside one of the construction scripts)
+		// so we don't want to generate overlaps yet.
+		AActor* const MyActor = GetOwner();
+		if ( MyActor && MyActor->IsActorInitialized() )
+		{
+			if (PendingOverlaps)
+			{
+				for (int32 Idx=0; Idx<PendingOverlaps->Num(); ++Idx)
+				{
+					BeginComponentOverlap( (*PendingOverlaps)[Idx], bDoNotifies );
+				}
+			}
+
+			// now generate full list of new touches, so we can compare to existing list and
+			// determine what changed
+			typedef TArray<FOverlapInfo, TInlineAllocator<4>> TInlineOverlapInfoArray;
+			TInlineOverlapInfoArray NewOverlappingComponents;
+
+			// If pending kill, we should not generate any new overlaps
+			if (!IsPendingKill())
+			{
+				// Might be able to avoid testing for new overlaps at the end location.
+				if (OverlapsAtEndLocation != NULL && CVarAllowCachedOverlaps->GetInt())
+				{
+					UE_LOG(LogPrimitiveComponent, VeryVerbose, TEXT("%s Skipping overlap test!"), *GetName());
+					NewOverlappingComponents = *OverlapsAtEndLocation;
+				}
+				else
+				{
+					UE_LOG(LogPrimitiveComponent, VeryVerbose, TEXT("%s Performing overlaps!"), *GetName());
+					UWorld* const MyWorld = MyActor->GetWorld();
+					TArray<FOverlapResult> Overlaps;
+					static FName NAME_UpdateOverlaps = FName(TEXT("UpdateOverlaps"));
+					// note this will include overlaps with components in the same actor.  
+					FComponentQueryParams Params(NAME_UpdateOverlaps);
+					Params.bTraceAsyncScene = bCheckAsyncSceneOnMove;
+					Params.AddIgnoredActors(MoveIgnoreActors);
+					MyWorld->ComponentOverlapMulti(Overlaps, this, GetComponentLocation(), GetComponentRotation(), Params);
+
+					for( int32 ResultIdx=0; ResultIdx<Overlaps.Num(); ResultIdx++ )
+					{
+						const FOverlapResult& Result = Overlaps[ResultIdx];
+
+						UPrimitiveComponent* const HitComp = Result.Component.Get();
+						if (!Result.bBlockingHit && HitComp && (HitComp != this) && HitComp->bGenerateOverlapEvents)
+						{
+							if (!ShouldIgnoreOverlapResult(MyWorld, MyActor, *this, Result.GetActor(), *HitComp))
+							{
+								NewOverlappingComponents.Add(FOverlapInfo(HitComp, Result.ItemIndex));		// don't need to add unique unless the overlap check can return dupes
+							}
+						}
+					}
+				}
+			}
+
+			if (OverlappingComponents.Num() > 0)
+			{
+				// make a copy of the old that we can manipulate to avoid n^2 searching later
+				TInlineOverlapInfoArray OldOverlappingComponents(OverlappingComponents);
+
+				// Now we want to compare the old and new overlap lists to determine 
+				// what overlaps are in old and not in new (need end overlap notifies), and 
+				// what overlaps are in new and not in old (need begin overlap notifies).
+				// We do this by removing common entries from both lists, since overlapping status has not changed for them.
+				// What is left over will be what has changed.
+				for (int32 CompIdx=0; CompIdx < OldOverlappingComponents.Num() && NewOverlappingComponents.Num() > 0; ++CompIdx)
+				{
+					// RemoveSingleSwap is ok, since it is not necessary to maintain order
+					const bool bAllowShrinking = false;
+					if (NewOverlappingComponents.RemoveSingleSwap(OldOverlappingComponents[CompIdx], bAllowShrinking) > 0)
+					{
+						OldOverlappingComponents.RemoveAtSwap(CompIdx, 1, bAllowShrinking);
+						--CompIdx;
+					}
+				}
+
+				// OldOverlappingComponents now contains only previous overlaps that are confirmed to no longer be valid.
+				for (auto CompIt = OldOverlappingComponents.CreateIterator(); CompIt; ++CompIt)
+				{
+					const FOverlapInfo& OtherOverlap = *CompIt;
+					if (OtherOverlap.OverlapInfo.Component.IsValid())
+					{
+						EndComponentOverlap(OtherOverlap, bDoNotifies, false);
+					}
+				}
+			}
+
+			// NewOverlappingComponents now contains only new overlaps that didn't exist previously.
+			for (auto CompIt = NewOverlappingComponents.CreateIterator(); CompIt; ++CompIt)
+			{
+				const FOverlapInfo& OtherOverlap = *CompIt;
+				if (OtherOverlap.OverlapInfo.Component.IsValid())
+				{
+					BeginComponentOverlap(OtherOverlap, bDoNotifies);
+				}
+			}
+		}
+	}
+	else
+	{
+		// bGenerateOverlapEvents is false or collision is disabled
+
+		// End all overlaps that exist, in case bGenerateOverlapEvents was true last tick (i.e. was just turned off)
+		// Iterate backwards since EndComponentOverlap will remove items from OverlappingComponents.
+		for (int32 OverlapIdx = OverlappingComponents.Num()-1; OverlapIdx >= 0; --OverlapIdx)
+		{
+			const FOverlapInfo& OtherOverlap = OverlappingComponents[OverlapIdx];
+			if (OtherOverlap.OverlapInfo.Component.IsValid())
+			{
+				EndComponentOverlap(OtherOverlap, bDoNotifies, false);
+			}
+		}
+	}
+
+	// Update physics volume using most current overlaps
+	if (bShouldUpdatePhysicsVolume)
+	{
+		UpdatePhysicsVolume(bDoNotifies);
+	}
+
+	// now update any children down the chain.
+	for (int32 ChildIdx=0; ChildIdx<AttachChildren.Num(); ++ChildIdx)
+	{
+		USceneComponent* const ChildComp = AttachChildren[ChildIdx];
+		if (ChildComp)
+		{
+			// Do not pass on OverlapsAtEndLocation, it only applied to this component.
+			ChildComp->UpdateOverlaps(NULL, bDoNotifies, NULL);
+		}
+	}
+}
+
+bool UPrimitiveComponent::ComponentOverlapMulti(TArray<struct FOverlapResult>& OutOverlaps, const UWorld* World, const FVector& Pos, const FRotator& Rot, ECollisionChannel TestChannel, const struct FComponentQueryParams& Params, const struct FCollisionObjectQueryParams& ObjectQueryParams) const
+{
+	FComponentQueryParams ParamsWithSelf = Params;
+	ParamsWithSelf.AddIgnoredComponent(this);
+	OutOverlaps.Reset();
+	return BodyInstance.OverlapMulti(OutOverlaps, World, /*pWorldToComponent=*/ nullptr, Pos, Rot, TestChannel, ParamsWithSelf, FCollisionResponseParams(GetCollisionResponseToChannels()), ObjectQueryParams);
+}
+
+void UPrimitiveComponent::UpdatePhysicsVolume( bool bTriggerNotifiers )
+{
+	if (bShouldUpdatePhysicsVolume && !IsPendingKill() && GetWorld())
+	{
+		SCOPE_CYCLE_COUNTER(STAT_UpdatePhysicsVolume);
+
+		if (bGenerateOverlapEvents && IsCollisionEnabled())
+		{
+			APhysicsVolume* BestVolume = GetWorld()->GetDefaultPhysicsVolume();
+			int32 BestPriority = BestVolume->Priority;
+
+			for (auto CompIt = OverlappingComponents.CreateIterator(); CompIt; ++CompIt)
+			{
+				const FOverlapInfo& Overlap = *CompIt;
+				UPrimitiveComponent* OtherComponent = Overlap.OverlapInfo.Component.Get();
+				if (OtherComponent)
+				{
+					APhysicsVolume* V = Cast<APhysicsVolume>(OtherComponent->GetOwner());
+					if (V && V->Priority > BestPriority)
+					{
+						if (V->IsOverlapInVolume(*this))
+						{
+							BestPriority = V->Priority;
+							BestVolume = V;
+						}
+					}
+				}
+			}
+
+			if (BestVolume != PhysicsVolume)
+			{
+				SetPhysicsVolume(BestVolume, bTriggerNotifiers);
+			}
+		}
+		else
+		{
+			Super::UpdatePhysicsVolume(bTriggerNotifiers);
+		}
+	}
+}
+
+
+void UPrimitiveComponent::DispatchMouseOverEvents(UPrimitiveComponent* CurrentComponent, UPrimitiveComponent* NewComponent)
+{
+	if (NewComponent)
+	{
+		AActor* NewOwner = NewComponent->GetOwner();
+
+		bool bBroadcastComponentBegin = true;
+		bool bBroadcastActorBegin = true;
+		if (CurrentComponent)
+		{
+			AActor* CurrentOwner = CurrentComponent->GetOwner();
+
+			if (NewComponent == CurrentComponent)
+			{
+				bBroadcastComponentBegin = false;
+			}
+			else
+			{
+				bBroadcastActorBegin = (NewOwner != CurrentOwner);
+
+				if (!CurrentComponent->IsPendingKill())
+				{
+					CurrentComponent->OnEndCursorOver.Broadcast(CurrentComponent);
+				}
+				if (bBroadcastActorBegin && IsActorValidToNotify(CurrentOwner))
+				{
+					CurrentOwner->ReceiveActorEndCursorOver();
+					if (IsActorValidToNotify(CurrentOwner))
+					{
+						CurrentOwner->OnEndCursorOver.Broadcast();
+					}
+				}
+			}
+		}
+
+		if (bBroadcastComponentBegin)
+		{
+			if (bBroadcastActorBegin && IsActorValidToNotify(NewOwner))
+			{
+				NewOwner->ReceiveActorBeginCursorOver();
+				if (IsActorValidToNotify(NewOwner))
+				{
+					NewOwner->OnBeginCursorOver.Broadcast();
+				}
+			}
+			if (!NewComponent->IsPendingKill())
+			{
+				NewComponent->OnBeginCursorOver.Broadcast(NewComponent);
+			}
+		}
+	}
+	else if (CurrentComponent)
+	{
+		AActor* CurrentOwner = CurrentComponent->GetOwner();
+
+		if (!CurrentComponent->IsPendingKill())
+		{
+			CurrentComponent->OnEndCursorOver.Broadcast(CurrentComponent);
+		}
+
+		if (IsActorValidToNotify(CurrentOwner))
+		{
+			CurrentOwner->ReceiveActorEndCursorOver();
+			if (IsActorValidToNotify(CurrentOwner))
+			{
+				CurrentOwner->OnEndCursorOver.Broadcast();
+			}
+		}
+	}
+}
+
+void UPrimitiveComponent::DispatchTouchOverEvents(ETouchIndex::Type FingerIndex, UPrimitiveComponent* CurrentComponent, UPrimitiveComponent* NewComponent)
+{
+	if (NewComponent)
+	{
+		AActor* NewOwner = NewComponent->GetOwner();
+
+		bool bBroadcastComponentBegin = true;
+		bool bBroadcastActorBegin = true;
+		if (CurrentComponent)
+		{
+			AActor* CurrentOwner = CurrentComponent->GetOwner();
+
+			if (NewComponent == CurrentComponent)
+			{
+				bBroadcastComponentBegin = false;
+			}
+			else
+			{
+				bBroadcastActorBegin = (NewOwner != CurrentOwner);
+
+				if (!CurrentComponent->IsPendingKill())
+				{
+					CurrentComponent->OnInputTouchLeave.Broadcast(FingerIndex, CurrentComponent);
+				}
+				if (bBroadcastActorBegin && IsActorValidToNotify(CurrentOwner))
+				{
+					CurrentOwner->ReceiveActorOnInputTouchLeave(FingerIndex);
+					if (IsActorValidToNotify(CurrentOwner))
+					{
+						CurrentOwner->OnInputTouchLeave.Broadcast(FingerIndex);
+					}
+				}
+			}
+		}
+
+		if (bBroadcastComponentBegin)
+		{
+			if (bBroadcastActorBegin && IsActorValidToNotify(NewOwner))
+			{
+				NewOwner->ReceiveActorOnInputTouchEnter(FingerIndex);
+				if (IsActorValidToNotify(NewOwner))
+				{
+					NewOwner->OnInputTouchEnter.Broadcast(FingerIndex);
+				}
+			}
+			if (!NewComponent->IsPendingKill())
+			{
+				NewComponent->OnInputTouchEnter.Broadcast(FingerIndex, NewComponent);
+			}
+		}
+	}
+	else if (CurrentComponent)
+	{
+		AActor* CurrentOwner = CurrentComponent->GetOwner();
+
+		if (!CurrentComponent->IsPendingKill())
+		{
+			CurrentComponent->OnInputTouchLeave.Broadcast(FingerIndex, CurrentComponent);
+		}
+
+		if (IsActorValidToNotify(CurrentOwner))
+		{
+			CurrentOwner->ReceiveActorOnInputTouchLeave(FingerIndex);
+			if (IsActorValidToNotify(CurrentOwner))
+			{
+				CurrentOwner->OnInputTouchLeave.Broadcast(FingerIndex);
+			}
+		}
+	}
+}
+
+void UPrimitiveComponent::DispatchOnClicked()
+{
+	if (IsActorValidToNotify(GetOwner()))
+	{
+		GetOwner()->ReceiveActorOnClicked();
+		if (IsActorValidToNotify(GetOwner()))
+		{
+			GetOwner()->OnClicked.Broadcast();
+		}
+	}
+
+	if (!IsPendingKill())
+	{
+		OnClicked.Broadcast(this);
+	}
+}
+
+void UPrimitiveComponent::DispatchOnReleased()
+{
+	if (IsActorValidToNotify(GetOwner()))
+	{
+		GetOwner()->ReceiveActorOnReleased();
+		if (IsActorValidToNotify(GetOwner()))
+		{
+			GetOwner()->OnReleased.Broadcast();
+		}
+	}
+
+	if (!IsPendingKill())
+	{
+		OnReleased.Broadcast(this);
+	}
+}
+
+void UPrimitiveComponent::DispatchOnInputTouchBegin(const ETouchIndex::Type FingerIndex)
+{
+	if (IsActorValidToNotify(GetOwner()))
+	{
+		GetOwner()->ReceiveActorOnInputTouchBegin(FingerIndex);
+		if (IsActorValidToNotify(GetOwner()))
+		{
+			GetOwner()->OnInputTouchBegin.Broadcast(FingerIndex);
+		}
+	}
+
+	if (!IsPendingKill())
+	{
+		OnInputTouchBegin.Broadcast(FingerIndex, this);
+	}
+}
+
+void UPrimitiveComponent::DispatchOnInputTouchEnd(const ETouchIndex::Type FingerIndex)
+{
+	if (IsActorValidToNotify(GetOwner()))
+	{
+		GetOwner()->ReceiveActorOnInputTouchEnd(FingerIndex);
+		if (IsActorValidToNotify(GetOwner()))
+		{
+			GetOwner()->OnInputTouchEnd.Broadcast(FingerIndex);
+		}
+	}
+
+	if (!IsPendingKill())
+	{
+		OnInputTouchEnd.Broadcast(FingerIndex, this);
+	}
+}
+
+SIZE_T UPrimitiveComponent::GetResourceSize( EResourceSizeMode::Type Mode )
+{
+	SIZE_T ResSize = 0;
+
+	if (BodyInstance.IsValidBodyInstance())
+	{
+		ResSize = BodyInstance.GetBodyInstanceResourceSize(Mode);
+	}
+
+	return ResSize;
+}
+
+void UPrimitiveComponent::SetRenderCustomDepth(bool bValue)
+{
+	if( bRenderCustomDepth != bValue )
+	{
+		bRenderCustomDepth = bValue;
+		MarkRenderStateDirty();
+	}
+}
+
+void UPrimitiveComponent::SetRenderInMainPass(bool bValue)
+{
+	if (bRenderInMainPass != bValue)
+	{
+		bRenderInMainPass = bValue;
+		MarkRenderStateDirty();
+	}
+}
+
+void UPrimitiveComponent::SetLODParentPrimitive(UPrimitiveComponent * InLODParentPrimitive)
+{
+	// @todo, what do we do with old parent. We can't just reset undo parent because the parent migh be used by other primitive
+	LODParentPrimitive = InLODParentPrimitive;
+	MarkRenderStateDirty();
+}
+#if WITH_EDITOR
+const int32 UPrimitiveComponent::GetNumUncachedStaticLightingInteractions() const
+{
+	int32 NumUncachedStaticLighting = 0;
+
+	NumUncachedStaticLighting += Super::GetNumUncachedStaticLightingInteractions();
+	if (SceneProxy)
+	{
+		NumUncachedStaticLighting += SceneProxy->GetNumUncachedStaticLightingInteractions();
+	}
+	return NumUncachedStaticLighting;
+}
+#endif
+
+
+bool UPrimitiveComponent::CanCharacterStepUp(APawn* Pawn) const
+{
+	if ( CanCharacterStepUpOn != ECB_Owner )
+	{
+		return CanCharacterStepUpOn == ECB_Yes;
+	}
+	else
+	{	
+		const AActor* Owner = GetOwner();
+		return Owner && Owner->CanBeBaseForCharacter(Pawn);
+	}
+}
+
+bool UPrimitiveComponent::CanEditSimulatePhysics()
+{
+	//Even if there's no collision but there is a body setup, we still let them simulate physics.
+	// The object falls through the world - this behavior is debatable but what we decided on for now
+	return GetBodySetup() != nullptr;
+}
+
+void UPrimitiveComponent::SetCustomNavigableGeometry(const EHasCustomNavigableGeometry::Type InType)
+{
+	bHasCustomNavigableGeometry = InType;
+}
+
+#undef LOCTEXT_NAMESPACE
+
