Engine: Moved actor components into Components folder

diff --git a/Engine/Source/Runtime/Engine/Private/Components/ReflectionCaptureComponent.cpp b/Engine/Source/Runtime/Engine/Private/Components/ReflectionCaptureComponent.cpp
new file mode 100644
--- /dev/null
+++ b/Engine/Source/Runtime/Engine/Private/Components/ReflectionCaptureComponent.cpp
@@ -0,0 +1,1574 @@
+// Copyright 1998-2015 Epic Games, Inc. All Rights Reserved.
+
+/*=============================================================================
+	
+=============================================================================*/
+
+#include "EnginePrivate.h"
+#include "Engine/ReflectionCapture.h"
+#include "DerivedDataCacheInterface.h"
+#include "TargetPlatform.h"
+#include "EngineModule.h"
+#include "RendererInterface.h"
+#include "ShaderCompiler.h"
+#include "Engine/SphereReflectionCapture.h"
+#include "Components/SphereReflectionCaptureComponent.h"
+#include "Components/DrawSphereComponent.h"
+#include "Components/BoxReflectionCaptureComponent.h"
+#include "Engine/PlaneReflectionCapture.h"
+#include "Engine/BoxReflectionCapture.h"
+#include "Components/PlaneReflectionCaptureComponent.h"
+#include "Components/ReflectionCaptureComponent.h"
+#include "Components/BoxComponent.h"
+#include "Components/SkyLightComponent.h"
+
+/** 
+ * Size of all reflection captures.
+ * Reflection capture derived data versions must be changed if modifying this
+ * Note: update HardcodedNumCaptureArrayMips if changing this
+ */
+ENGINE_API int32 GReflectionCaptureSize = 128;
+
+void UWorld::UpdateAllReflectionCaptures()
+{
+	if ( FeatureLevel < ERHIFeatureLevel::SM4)
+	{
+		UE_LOG(LogMaterial, Warning, TEXT("Update reflection captures only works with an active feature level of SM4 or greater."));
+		return;
+	}
+	
+	for (TObjectIterator<UReflectionCaptureComponent> It; It; ++It)
+	{
+		UReflectionCaptureComponent* CaptureComponent = *It;
+
+		if (ContainsActor(CaptureComponent->GetOwner()) && !CaptureComponent->IsPendingKill())
+		{
+			// Purge cached derived data and force an update
+			CaptureComponent->SetCaptureIsDirty();			
+		}
+	}
+	UReflectionCaptureComponent::UpdateReflectionCaptureContents(this);
+	
+	for (TObjectIterator<USkyLightComponent> It; It; ++It)
+	{
+		USkyLightComponent* SkylightComponent = *It;
+		if (ContainsActor(SkylightComponent->GetOwner()) && !SkylightComponent->IsPendingKill())
+		{			
+			SkylightComponent->SetCaptureIsDirty();			
+		}
+	}
+	USkyLightComponent::UpdateSkyCaptureContents(this);
+}
+
+AReflectionCapture::AReflectionCapture(const FObjectInitializer& ObjectInitializer)
+	: Super(ObjectInitializer)
+{
+	CaptureComponent = CreateAbstractDefaultSubobject<UReflectionCaptureComponent>(TEXT("NewReflectionComponent"));
+
+#if WITH_EDITORONLY_DATA
+	SpriteComponent = CreateEditorOnlyDefaultSubobject<UBillboardComponent>(TEXT("Sprite"));
+	if (!IsRunningCommandlet() && (SpriteComponent != nullptr))
+	{
+		// Structure to hold one-time initialization
+		struct FConstructorStatics
+		{
+			FName NAME_ReflectionCapture;
+			ConstructorHelpers::FObjectFinderOptional<UTexture2D> DecalTexture;
+			FConstructorStatics()
+				: NAME_ReflectionCapture(TEXT("ReflectionCapture"))
+				, DecalTexture(TEXT("/Engine/EditorResources/S_ReflActorIcon"))
+			{
+			}
+		};
+		static FConstructorStatics ConstructorStatics;
+
+		SpriteComponent->Sprite = ConstructorStatics.DecalTexture.Get();
+		SpriteComponent->RelativeScale3D = FVector(0.5f, 0.5f, 0.5f);
+		SpriteComponent->bHiddenInGame = true;
+		SpriteComponent->bAbsoluteScale = true;
+		SpriteComponent->SetCollisionProfileName(UCollisionProfile::NoCollision_ProfileName);
+		SpriteComponent->bIsScreenSizeScaled = true;
+	}
+#endif // WITH_EDITORONLY_DATA
+}
+
+#if WITH_EDITOR
+void AReflectionCapture::PostEditMove(bool bFinished)
+{
+	Super::PostEditMove(bFinished);
+
+	CaptureComponent->SetCaptureIsDirty();
+}
+#endif // WITH_EDITOR
+
+ASphereReflectionCapture::ASphereReflectionCapture(const FObjectInitializer& ObjectInitializer)
+	: Super(ObjectInitializer.SetDefaultSubobjectClass<USphereReflectionCaptureComponent>(TEXT("NewReflectionComponent")))
+{
+	USphereReflectionCaptureComponent* SphereComponent = CastChecked<USphereReflectionCaptureComponent>(GetCaptureComponent());
+	RootComponent = SphereComponent;
+#if WITH_EDITORONLY_DATA
+	if (GetSpriteComponent())
+	{
+		GetSpriteComponent()->AttachParent = SphereComponent;
+	}
+#endif	//WITH_EDITORONLY_DATA
+
+	UDrawSphereComponent* DrawInfluenceRadius = CreateDefaultSubobject<UDrawSphereComponent>(TEXT("DrawRadius0"));
+	DrawInfluenceRadius->AttachParent = GetCaptureComponent();
+	DrawInfluenceRadius->bDrawOnlyIfSelected = true;
+	DrawInfluenceRadius->bUseEditorCompositing = true;
+	DrawInfluenceRadius->SetCollisionProfileName(UCollisionProfile::NoCollision_ProfileName);
+	SphereComponent->PreviewInfluenceRadius = DrawInfluenceRadius;
+
+	DrawCaptureRadius = CreateDefaultSubobject<UDrawSphereComponent>(TEXT("DrawRadius1"));
+	DrawCaptureRadius->AttachParent = GetCaptureComponent();
+	DrawCaptureRadius->bDrawOnlyIfSelected = true;
+	DrawCaptureRadius->bUseEditorCompositing = true;
+	DrawCaptureRadius->SetCollisionProfileName(UCollisionProfile::NoCollision_ProfileName);
+	DrawCaptureRadius->ShapeColor = FColor(100, 90, 40);
+}
+
+#if WITH_EDITOR
+void ASphereReflectionCapture::EditorApplyScale(const FVector& DeltaScale, const FVector* PivotLocation, bool bAltDown, bool bShiftDown, bool bCtrlDown)
+{
+	USphereReflectionCaptureComponent* SphereComponent = Cast<USphereReflectionCaptureComponent>(GetCaptureComponent());
+	check(SphereComponent);
+	const FVector ModifiedScale = DeltaScale * ( AActor::bUsePercentageBasedScaling ? 5000.0f : 50.0f );
+	FMath::ApplyScaleToFloat(SphereComponent->InfluenceRadius, ModifiedScale);
+	GetCaptureComponent()->SetCaptureIsDirty();
+	PostEditChange();
+}
+
+void APlaneReflectionCapture::EditorApplyScale(const FVector& DeltaScale, const FVector* PivotLocation, bool bAltDown, bool bShiftDown, bool bCtrlDown)
+{
+	UPlaneReflectionCaptureComponent* PlaneComponent = Cast<UPlaneReflectionCaptureComponent>(GetCaptureComponent());
+	check(PlaneComponent);
+	const FVector ModifiedScale = DeltaScale * ( AActor::bUsePercentageBasedScaling ? 5000.0f : 50.0f );
+	FMath::ApplyScaleToFloat(PlaneComponent->InfluenceRadiusScale, ModifiedScale);
+	GetCaptureComponent()->SetCaptureIsDirty();
+	PostEditChange();
+}
+#endif
+
+ABoxReflectionCapture::ABoxReflectionCapture(const FObjectInitializer& ObjectInitializer)
+	: Super(ObjectInitializer.SetDefaultSubobjectClass<UBoxReflectionCaptureComponent>(TEXT("NewReflectionComponent")))
+{
+	UBoxReflectionCaptureComponent* BoxComponent = CastChecked<UBoxReflectionCaptureComponent>(GetCaptureComponent());
+	BoxComponent->RelativeScale3D = FVector(1000, 1000, 400);
+	RootComponent = BoxComponent;
+#if WITH_EDITORONLY_DATA
+	if (GetSpriteComponent())
+	{
+		GetSpriteComponent()->AttachParent = BoxComponent;
+	}
+#endif	//WITH_EDITORONLY_DATA
+	UBoxComponent* DrawInfluenceBox = CreateDefaultSubobject<UBoxComponent>(TEXT("DrawBox0"));
+	DrawInfluenceBox->AttachParent = GetCaptureComponent();
+	DrawInfluenceBox->bDrawOnlyIfSelected = true;
+	DrawInfluenceBox->bUseEditorCompositing = true;
+	DrawInfluenceBox->SetCollisionProfileName(UCollisionProfile::NoCollision_ProfileName);
+	DrawInfluenceBox->InitBoxExtent(FVector(1, 1, 1));
+	BoxComponent->PreviewInfluenceBox = DrawInfluenceBox;
+
+	UBoxComponent* DrawCaptureBox = CreateDefaultSubobject<UBoxComponent>(TEXT("DrawBox1"));
+	DrawCaptureBox->AttachParent = GetCaptureComponent();
+	DrawCaptureBox->bDrawOnlyIfSelected = true;
+	DrawCaptureBox->bUseEditorCompositing = true;
+	DrawCaptureBox->SetCollisionProfileName(UCollisionProfile::NoCollision_ProfileName);
+	DrawCaptureBox->ShapeColor = FColor(100, 90, 40);
+	DrawCaptureBox->InitBoxExtent(FVector(1, 1, 1));
+	BoxComponent->PreviewCaptureBox = DrawCaptureBox;
+}
+
+APlaneReflectionCapture::APlaneReflectionCapture(const FObjectInitializer& ObjectInitializer)
+	: Super(ObjectInitializer.SetDefaultSubobjectClass<UPlaneReflectionCaptureComponent>(TEXT("NewReflectionComponent")))
+{
+	UPlaneReflectionCaptureComponent* PlaneComponent = CastChecked<UPlaneReflectionCaptureComponent>(GetCaptureComponent());
+	PlaneComponent->RelativeScale3D = FVector(1, 1000, 1000);
+	RootComponent = PlaneComponent;
+#if WITH_EDITORONLY_DATA
+	if (GetSpriteComponent())
+	{
+		GetSpriteComponent()->AttachParent = PlaneComponent;
+	}
+#endif	//#if WITH_EDITORONLY_DATA
+	UDrawSphereComponent* DrawInfluenceRadius = CreateDefaultSubobject<UDrawSphereComponent>(TEXT("DrawRadius0"));
+	DrawInfluenceRadius->AttachParent = GetCaptureComponent();
+	DrawInfluenceRadius->bDrawOnlyIfSelected = true;
+	DrawInfluenceRadius->bAbsoluteScale = true;
+	DrawInfluenceRadius->bUseEditorCompositing = true;
+	DrawInfluenceRadius->SetCollisionProfileName(UCollisionProfile::NoCollision_ProfileName);
+	PlaneComponent->PreviewInfluenceRadius = DrawInfluenceRadius;
+
+	UBoxComponent* DrawCaptureBox = CreateDefaultSubobject<UBoxComponent>(TEXT("DrawBox1"));
+	DrawCaptureBox->AttachParent = GetCaptureComponent();
+	DrawCaptureBox->bDrawOnlyIfSelected = true;
+	DrawCaptureBox->bUseEditorCompositing = true;
+	DrawCaptureBox->SetCollisionProfileName(UCollisionProfile::NoCollision_ProfileName);
+	DrawCaptureBox->ShapeColor = FColor(100, 90, 40);
+	DrawCaptureBox->InitBoxExtent(FVector(1, 1, 1));
+	PlaneComponent->PreviewCaptureBox = DrawCaptureBox;
+}
+
+// Generate a new guid to force a recache of all reflection derived data
+// Note: changing this will cause saved capture data in maps to be discarded
+// A resave of those maps will be required to guarantee valid reflections when cooking for ES2
+FGuid ReflectionCaptureDDCVer(0x0c669396, 0x9cb849ae, 0x9f4120ff, 0x5812f4d2);
+
+FString FReflectionCaptureFullHDRDerivedData::GetDDCKeyString(const FGuid& StateId)
+{
+	return FDerivedDataCacheInterface::BuildCacheKey(TEXT("REFL_FULL"), *ReflectionCaptureDDCVer.ToString(), *StateId.ToString());
+}
+
+// For temporary backwards compatibility
+#define LEGACY_REFLECTIONCAPTURE_FULL_DERIVEDDATA_VER	TEXT("82f35deee55a5ec9956d4897f3d46e5a")
+
+FString FReflectionCaptureFullHDRDerivedData::GetLegacyDDCKeyString(const FGuid& StateId)
+{
+	return FDerivedDataCacheInterface::BuildCacheKey(TEXT("REFL_FULL"), LEGACY_REFLECTIONCAPTURE_FULL_DERIVEDDATA_VER, *StateId.ToString());
+}
+
+FReflectionCaptureFullHDRDerivedData::~FReflectionCaptureFullHDRDerivedData()
+{
+	DEC_MEMORY_STAT_BY(STAT_ReflectionCaptureMemory, CompressedCapturedData.GetAllocatedSize());
+}
+
+void FReflectionCaptureFullHDRDerivedData::InitializeFromUncompressedData(const TArray<uint8>& UncompressedData)
+{
+	DEC_MEMORY_STAT_BY(STAT_ReflectionCaptureMemory, CompressedCapturedData.GetAllocatedSize());
+
+	int32 UncompressedSize = UncompressedData.Num() * UncompressedData.GetTypeSize();
+
+	TArray<uint8> TempCompressedMemory;
+	// Compressed can be slightly larger than uncompressed
+	TempCompressedMemory.Empty(UncompressedSize * 4 / 3);
+	TempCompressedMemory.AddUninitialized(UncompressedSize * 4 / 3);
+	int32 CompressedSize = TempCompressedMemory.Num() * TempCompressedMemory.GetTypeSize();
+
+	verify(FCompression::CompressMemory(
+		(ECompressionFlags)(COMPRESS_ZLIB | COMPRESS_BiasMemory), 
+		TempCompressedMemory.GetData(), 
+		CompressedSize, 
+		UncompressedData.GetData(), 
+		UncompressedSize));
+
+	// Note: change REFLECTIONCAPTURE_FULL_DERIVEDDATA_VER when modifying the serialization layout
+	FMemoryWriter FinalArchive(CompressedCapturedData, true);
+	FinalArchive << UncompressedSize;
+	FinalArchive << CompressedSize;
+	FinalArchive.Serialize(TempCompressedMemory.GetData(), CompressedSize);
+
+	INC_MEMORY_STAT_BY(STAT_ReflectionCaptureMemory, CompressedCapturedData.GetAllocatedSize());
+}
+
+void FReflectionCaptureFullHDRDerivedData::GetUncompressedData(TArray<uint8>& UncompressedData) const
+{
+	FMemoryReader Ar(CompressedCapturedData);
+
+	// Note: change REFLECTIONCAPTURE_FULL_DERIVEDDATA_VER when modifying the serialization layout
+	int32 UncompressedSize;
+	Ar << UncompressedSize;
+
+	int32 CompressedSize;
+	Ar << CompressedSize;
+
+	TArray<uint8> CompressedData;
+	CompressedData.Empty(CompressedSize);
+	CompressedData.AddUninitialized(CompressedSize);
+	Ar.Serialize(CompressedData.GetData(), CompressedSize);
+
+	UncompressedData.Empty(UncompressedSize);
+	UncompressedData.AddUninitialized(UncompressedSize);
+
+	verify(FCompression::UncompressMemory((ECompressionFlags)COMPRESS_ZLIB, UncompressedData.GetData(), UncompressedSize, CompressedData.GetData(), CompressedSize));
+}
+
+FColor RGBMEncode( FLinearColor Color )
+{
+	FColor Encoded;
+
+	// Convert to gamma space
+	Color.R = FMath::Sqrt( Color.R );
+	Color.G = FMath::Sqrt( Color.G );
+	Color.B = FMath::Sqrt( Color.B );
+
+	// Range
+	Color /= 16.0f;
+	
+	float MaxValue = FMath::Max( FMath::Max(Color.R, Color.G), FMath::Max(Color.B, DELTA) );
+	
+	if( MaxValue > 0.75f )
+	{
+		// Fit to valid range by leveling off intensity
+		float Tonemapped = ( MaxValue - 0.75 * 0.75 ) / ( MaxValue - 0.5 );
+		Color *= Tonemapped / MaxValue;
+		MaxValue = Tonemapped;
+	}
+
+	Encoded.A = FMath::Min( FMath::CeilToInt( MaxValue * 255.0f ), 255 );
+	Encoded.R = FMath::RoundToInt( ( Color.R * 255.0f / Encoded.A ) * 255.0f );
+	Encoded.G = FMath::RoundToInt( ( Color.G * 255.0f / Encoded.A ) * 255.0f );
+	Encoded.B = FMath::RoundToInt( ( Color.B * 255.0f / Encoded.A ) * 255.0f );
+
+	return Encoded;
+}
+
+// Based off of CubemapGen
+// https://code.google.com/p/cubemapgen/
+
+#define FACE_X_POS 0
+#define FACE_X_NEG 1
+#define FACE_Y_POS 2
+#define FACE_Y_NEG 3
+#define FACE_Z_POS 4
+#define FACE_Z_NEG 5
+
+#define EDGE_LEFT   0	 // u = 0
+#define EDGE_RIGHT  1	 // u = 1
+#define EDGE_TOP    2	 // v = 0
+#define EDGE_BOTTOM 3	 // v = 1
+
+#define CORNER_NNN  0
+#define CORNER_NNP  1
+#define CORNER_NPN  2
+#define CORNER_NPP  3
+#define CORNER_PNN  4
+#define CORNER_PNP  5
+#define CORNER_PPN  6
+#define CORNER_PPP  7
+
+// D3D cube map face specification
+//   mapping from 3D x,y,z cube map lookup coordinates 
+//   to 2D within face u,v coordinates
+//
+//   --------------------> U direction 
+//   |                   (within-face texture space)
+//   |         _____
+//   |        |     |
+//   |        | +Y  |
+//   |   _____|_____|_____ _____
+//   |  |     |     |     |     |
+//   |  | -X  | +Z  | +X  | -Z  |
+//   |  |_____|_____|_____|_____|
+//   |        |     |
+//   |        | -Y  |
+//   |        |_____|
+//   |
+//   v   V direction
+//      (within-face texture space)
+
+// Index by [Edge][FaceOrEdge]
+static const int32 CubeEdgeListA[12][2] =
+{
+	{FACE_X_POS, EDGE_LEFT},
+	{FACE_X_POS, EDGE_RIGHT},
+	{FACE_X_POS, EDGE_TOP},
+	{FACE_X_POS, EDGE_BOTTOM},
+
+	{FACE_X_NEG, EDGE_LEFT},
+	{FACE_X_NEG, EDGE_RIGHT},
+	{FACE_X_NEG, EDGE_TOP},
+	{FACE_X_NEG, EDGE_BOTTOM},
+
+	{FACE_Z_POS, EDGE_TOP},
+	{FACE_Z_POS, EDGE_BOTTOM},
+	{FACE_Z_NEG, EDGE_TOP},
+	{FACE_Z_NEG, EDGE_BOTTOM}
+};
+
+static const int32 CubeEdgeListB[12][2] =
+{
+	{FACE_Z_POS, EDGE_RIGHT },
+	{FACE_Z_NEG, EDGE_LEFT  },
+	{FACE_Y_POS, EDGE_RIGHT },
+	{FACE_Y_NEG, EDGE_RIGHT },
+
+	{FACE_Z_NEG, EDGE_RIGHT },
+	{FACE_Z_POS, EDGE_LEFT  },
+	{FACE_Y_POS, EDGE_LEFT  },
+	{FACE_Y_NEG, EDGE_LEFT  },
+
+	{FACE_Y_POS, EDGE_BOTTOM },
+	{FACE_Y_NEG, EDGE_TOP    },
+	{FACE_Y_POS, EDGE_TOP    },
+	{FACE_Y_NEG, EDGE_BOTTOM },
+};
+
+// Index by [Face][Corner]
+static const int32 CubeCornerList[6][4] =
+{
+	{ CORNER_PPP, CORNER_PPN, CORNER_PNP, CORNER_PNN },
+	{ CORNER_NPN, CORNER_NPP, CORNER_NNN, CORNER_NNP },
+	{ CORNER_NPN, CORNER_PPN, CORNER_NPP, CORNER_PPP },
+	{ CORNER_NNP, CORNER_PNP, CORNER_NNN, CORNER_PNN },
+	{ CORNER_NPP, CORNER_PPP, CORNER_NNP, CORNER_PNP },
+	{ CORNER_PPN, CORNER_NPN, CORNER_PNN, CORNER_NNN }
+};
+
+static void EdgeWalkSetup( bool ReverseDirection, int32 Edge, int32 MipSize, int32& EdgeStart, int32& EdgeStep )
+{
+	if( ReverseDirection )
+	{
+		switch( Edge )
+		{
+		case EDGE_LEFT:		//start at lower left and walk up
+			EdgeStart = MipSize * (MipSize - 1);
+			EdgeStep = -MipSize;
+			break;
+		case EDGE_RIGHT:	//start at lower right and walk up
+			EdgeStart = MipSize * (MipSize - 1) + (MipSize - 1);
+			EdgeStep = -MipSize;
+			break;
+		case EDGE_TOP:		//start at upper right and walk left
+			EdgeStart = (MipSize - 1);
+			EdgeStep = -1;
+			break;
+		case EDGE_BOTTOM:	//start at lower right and walk left
+			EdgeStart = MipSize * (MipSize - 1) + (MipSize - 1);
+			EdgeStep = -1;
+			break;
+		}            
+	}
+	else
+	{
+		switch( Edge )
+		{
+		case EDGE_LEFT:		//start at upper left and walk down
+			EdgeStart = 0;
+			EdgeStep = MipSize;
+			break;
+		case EDGE_RIGHT:	//start at upper right and walk down
+			EdgeStart = (MipSize - 1);
+			EdgeStep = MipSize;
+			break;
+		case EDGE_TOP:		//start at upper left and walk left
+			EdgeStart = 0;
+			EdgeStep = 1;
+			break;
+		case EDGE_BOTTOM:	//start at lower left and walk left
+			EdgeStart = MipSize * (MipSize - 1);
+			EdgeStep = 1;
+			break;
+		}
+	}
+}
+
+void FReflectionCaptureEncodedHDRDerivedData::GenerateFromDerivedDataSource(const FReflectionCaptureFullHDRDerivedData& FullHDRDerivedData, float Brightness)
+{
+	const int32 EffectiveTopMipSize = GReflectionCaptureSize;
+	const int32 NumMips = FMath::CeilLogTwo(EffectiveTopMipSize) + 1;
+
+	TArray<uint8> CubemapData;
+	FullHDRDerivedData.GetUncompressedData(CubemapData);
+
+	int32 SourceMipBaseIndex = 0;
+	int32 DestMipBaseIndex = 0;
+
+	CapturedData.Empty(CubemapData.Num() * sizeof(FColor) / sizeof(FFloat16Color));
+	CapturedData.AddUninitialized(CubemapData.Num() * sizeof(FColor) / sizeof(FFloat16Color));
+
+	// Note: change REFLECTIONCAPTURE_ENCODED_DERIVEDDATA_VER when modifying the encoded data layout or contents
+
+	for (int32 MipIndex = 0; MipIndex < NumMips; MipIndex++)
+	{
+		const int32 MipSize = 1 << (NumMips - MipIndex - 1);
+		const int32 SourceCubeFaceBytes = MipSize * MipSize * sizeof(FFloat16Color);
+		const int32 DestCubeFaceBytes = MipSize * MipSize * sizeof(FColor);
+
+		const FFloat16Color*	MipSrcData = (const FFloat16Color*)&CubemapData[ SourceMipBaseIndex ];
+		FColor*					MipDstData = (FColor*)&CapturedData[ DestMipBaseIndex ];
+
+		// Fix cubemap seams by averaging colors across edges
+
+		int32 CornerTable[4] =
+		{
+			0,
+			MipSize - 1,
+			MipSize * (MipSize - 1),
+			MipSize * (MipSize - 1) + MipSize - 1,
+		};
+
+		// Average corner colors
+		FLinearColor AvgCornerColors[8];
+		memset( AvgCornerColors, 0, sizeof( AvgCornerColors ) );
+		for( int32 Face = 0; Face < CubeFace_MAX; Face++ )
+		{
+			const FFloat16Color* FaceSrcData = MipSrcData + Face * MipSize * MipSize;
+
+			for( int32 Corner = 0; Corner < 4; Corner++ )
+			{
+				AvgCornerColors[ CubeCornerList[Face][Corner] ] += FLinearColor( FaceSrcData[ CornerTable[Corner] ] );
+			}
+		}
+
+		// Encode corners
+		for( int32 Face = 0; Face < CubeFace_MAX; Face++ )
+		{
+			FColor* FaceDstData = MipDstData + Face * MipSize * MipSize;
+
+			for( int32 Corner = 0; Corner < 4; Corner++ )
+			{
+				const FLinearColor LinearColor = AvgCornerColors[ CubeCornerList[Face][Corner] ] / 3.0f;
+				FaceDstData[ CornerTable[Corner] ] = RGBMEncode( LinearColor * Brightness );
+			}
+		}
+
+		// Average edge colors
+		for( int32 EdgeIndex = 0; EdgeIndex < 12; EdgeIndex++ )
+		{
+			int32 FaceA = CubeEdgeListA[ EdgeIndex ][0];
+			int32 EdgeA = CubeEdgeListA[ EdgeIndex ][1];
+
+			int32 FaceB = CubeEdgeListB[ EdgeIndex ][0];
+			int32 EdgeB = CubeEdgeListB[ EdgeIndex ][1];
+
+			const FFloat16Color*	FaceSrcDataA = MipSrcData + FaceA * MipSize * MipSize;
+			FColor*					FaceDstDataA = MipDstData + FaceA * MipSize * MipSize;
+
+			const FFloat16Color*	FaceSrcDataB = MipSrcData + FaceB * MipSize * MipSize;
+			FColor*					FaceDstDataB = MipDstData + FaceB * MipSize * MipSize;
+
+			int32 EdgeStartA = 0;
+			int32 EdgeStepA = 0;
+			int32 EdgeStartB = 0;
+			int32 EdgeStepB = 0;
+
+			EdgeWalkSetup( false, EdgeA, MipSize, EdgeStartA, EdgeStepA );
+			EdgeWalkSetup( EdgeA == EdgeB || EdgeA + EdgeB == 3, EdgeB, MipSize, EdgeStartB, EdgeStepB );
+
+			// Walk edge
+			// Skip corners
+			for( int32 Texel = 1; Texel < MipSize - 1; Texel++ )       
+			{
+				int32 EdgeTexelA = EdgeStartA + EdgeStepA * Texel;
+				int32 EdgeTexelB = EdgeStartB + EdgeStepB * Texel;
+
+				check( 0 <= EdgeTexelA && EdgeTexelA < MipSize * MipSize );
+				check( 0 <= EdgeTexelB && EdgeTexelB < MipSize * MipSize );
+
+				const FLinearColor EdgeColorA = FLinearColor( FaceSrcDataA[ EdgeTexelA ] );
+				const FLinearColor EdgeColorB = FLinearColor( FaceSrcDataB[ EdgeTexelB ] );
+				const FLinearColor AvgColor = 0.5f * ( EdgeColorA + EdgeColorB );
+				
+				FaceDstDataA[ EdgeTexelA ] = FaceDstDataB[ EdgeTexelB ] = RGBMEncode( AvgColor * Brightness );
+			}
+		}
+		
+		// Encode rest of texels
+		for (int32 CubeFace = 0; CubeFace < CubeFace_MAX; CubeFace++)
+		{
+			const int32 FaceSourceIndex = SourceMipBaseIndex + CubeFace * SourceCubeFaceBytes;
+			const int32 FaceDestIndex = DestMipBaseIndex + CubeFace * DestCubeFaceBytes;
+			const FFloat16Color* FaceSourceData = (const FFloat16Color*)&CubemapData[FaceSourceIndex];
+			FColor* FaceDestData = (FColor*)&CapturedData[FaceDestIndex];
+
+			// Convert each texel from linear space FP16 to RGBM FColor
+			// Note: Brightness on the capture is baked into the encoded HDR data
+			// Skip edges
+			for( int32 y = 1; y < MipSize - 1; y++ )
+			{
+				for( int32 x = 1; x < MipSize - 1; x++ )
+				{
+					int32 TexelIndex = x + y * MipSize;
+					const FLinearColor LinearColor = FLinearColor( FaceSourceData[ TexelIndex ]) * Brightness;
+					FaceDestData[ TexelIndex ] = RGBMEncode( LinearColor );
+				}
+			}
+		}
+
+		SourceMipBaseIndex += SourceCubeFaceBytes * CubeFace_MAX;
+		DestMipBaseIndex += DestCubeFaceBytes * CubeFace_MAX;
+	}
+}
+
+// Generate a new guid to force a recache of all encoded HDR derived data
+#define REFLECTIONCAPTURE_ENCODED_DERIVEDDATA_VER TEXT("96DFC022836B48889143E9DF484C3296")
+
+FString FReflectionCaptureEncodedHDRDerivedData::GetDDCKeyString(const FGuid& StateId)
+{
+	return FDerivedDataCacheInterface::BuildCacheKey(TEXT("REFL_ENC"), REFLECTIONCAPTURE_ENCODED_DERIVEDDATA_VER, *StateId.ToString());
+}
+
+FReflectionCaptureEncodedHDRDerivedData::~FReflectionCaptureEncodedHDRDerivedData()
+{
+	DEC_MEMORY_STAT_BY(STAT_ReflectionCaptureMemory,CapturedData.GetAllocatedSize());
+}
+
+TRefCountPtr<FReflectionCaptureEncodedHDRDerivedData> FReflectionCaptureEncodedHDRDerivedData::GenerateEncodedHDRData(const FReflectionCaptureFullHDRDerivedData& FullHDRData, const FGuid& StateId, float Brightness)
+{
+	TRefCountPtr<FReflectionCaptureEncodedHDRDerivedData> EncodedHDRData = new FReflectionCaptureEncodedHDRDerivedData();
+	const FString KeyString = GetDDCKeyString(StateId);
+
+	if (!GetDerivedDataCacheRef().GetSynchronous(*KeyString, EncodedHDRData->CapturedData))
+	{
+		EncodedHDRData->GenerateFromDerivedDataSource(FullHDRData, Brightness);
+
+		if (EncodedHDRData->CapturedData.Num() > 0)
+		{
+			GetDerivedDataCacheRef().Put(*KeyString, EncodedHDRData->CapturedData);
+		}
+	}
+
+	check(EncodedHDRData->CapturedData.Num() > 0);
+	INC_MEMORY_STAT_BY(STAT_ReflectionCaptureMemory,EncodedHDRData->CapturedData.GetAllocatedSize());
+	return EncodedHDRData;
+}
+
+/** 
+ * A cubemap texture resource that knows how to upload the packed capture data from a reflection capture. 
+ * @todo - support texture streaming and compression
+ */
+class FReflectionTextureCubeResource : public FTexture
+{
+public:
+
+	FReflectionTextureCubeResource() :
+		Size(0),
+		NumMips(0),
+		Format(PF_Unknown),
+		SourceData(NULL)
+	{}
+
+	void SetupParameters(int32 InSize, int32 InNumMips, EPixelFormat InFormat, TArray<uint8>* InSourceData)
+	{
+		Size = InSize;
+		NumMips = InNumMips;
+		Format = InFormat;
+		SourceData = InSourceData;
+	}
+
+	virtual void InitRHI() override
+	{
+		FRHIResourceCreateInfo CreateInfo;
+		TextureCubeRHI = RHICreateTextureCube(Size, Format, NumMips, 0, CreateInfo);
+		TextureRHI = TextureCubeRHI;
+
+		if (SourceData)
+		{
+			check(SourceData->Num() > 0);
+
+			const int32 BlockBytes = GPixelFormats[Format].BlockBytes;
+			int32 MipBaseIndex = 0;
+
+			for (int32 MipIndex = 0; MipIndex < NumMips; MipIndex++)
+			{
+				const int32 MipSize = 1 << (NumMips - MipIndex - 1);
+				const int32 CubeFaceBytes = MipSize * MipSize * BlockBytes;
+
+				for (int32 CubeFace = 0; CubeFace < CubeFace_MAX; CubeFace++)
+				{
+					uint32 DestStride = 0;
+					uint8* DestBuffer = (uint8*)RHILockTextureCubeFace(TextureCubeRHI, CubeFace, 0, MipIndex, RLM_WriteOnly, DestStride, false);
+
+					// Handle DestStride by copying each row
+					for (int32 Y = 0; Y < MipSize; Y++)
+					{
+						uint8* DestPtr = ((uint8*)DestBuffer + Y * DestStride);
+						const int32 SourceIndex = MipBaseIndex + CubeFace * CubeFaceBytes + Y * MipSize * BlockBytes;
+						const uint8* SourcePtr = &(*SourceData)[SourceIndex];
+						FMemory::Memcpy(DestPtr, SourcePtr, MipSize * BlockBytes);
+					}
+
+					RHIUnlockTextureCubeFace(TextureCubeRHI, CubeFace, 0, MipIndex, false);
+				}
+
+				MipBaseIndex += CubeFaceBytes * CubeFace_MAX;
+			}
+
+			if (!GIsEditor)
+			{
+				// Toss the source data now that we've created the cubemap
+				// Note: can't do this if we ever use this texture resource in the editor and want to save the data later
+				DEC_MEMORY_STAT_BY(STAT_ReflectionCaptureMemory,SourceData->GetAllocatedSize());
+				SourceData->Empty();
+			}
+		}
+
+		// Create the sampler state RHI resource.
+		FSamplerStateInitializerRHI SamplerStateInitializer
+		(
+			SF_Trilinear,
+			AM_Clamp,
+			AM_Clamp,
+			AM_Clamp
+		);
+		SamplerStateRHI = RHICreateSamplerState(SamplerStateInitializer);
+
+		INC_MEMORY_STAT_BY(STAT_ReflectionCaptureTextureMemory,CalcTextureSize(Size,Size,Format,NumMips) * 6);
+	}
+
+	virtual void ReleaseRHI()
+	{
+		DEC_MEMORY_STAT_BY(STAT_ReflectionCaptureTextureMemory,CalcTextureSize(Size,Size,Format,NumMips) * 6);
+		TextureCubeRHI.SafeRelease();
+		FTexture::ReleaseRHI();
+	}
+
+	virtual uint32 GetSizeX() const
+	{
+		return Size;
+	}
+
+	virtual uint32 GetSizeY() const
+	{
+		return Size;
+	}
+
+	FTextureRHIParamRef GetTextureRHI() 
+	{
+		return TextureCubeRHI;
+	}
+
+private:
+
+	int32 Size;
+	int32 NumMips;
+	EPixelFormat Format;
+	FTextureCubeRHIRef TextureCubeRHI;
+	TArray<uint8>* SourceData;
+};
+
+
+TArray<UReflectionCaptureComponent*> UReflectionCaptureComponent::ReflectionCapturesToUpdate;
+TArray<UReflectionCaptureComponent*> UReflectionCaptureComponent::ReflectionCapturesToUpdateForLoad;
+
+UReflectionCaptureComponent::UReflectionCaptureComponent(const FObjectInitializer& ObjectInitializer)
+	: Super(ObjectInitializer)
+{
+	Brightness = 1;
+	// Shouldn't be able to change reflection captures at runtime
+	Mobility = EComponentMobility::Static;
+
+	bCaptureDirty = false;
+	bDerivedDataDirty = false;
+}
+
+void UReflectionCaptureComponent::CreateRenderState_Concurrent()
+{
+	Super::CreateRenderState_Concurrent();
+
+	UpdatePreviewShape();
+
+	if (ShouldRender())
+	{
+		World->Scene->AddReflectionCapture(this);
+	}
+}
+
+void UReflectionCaptureComponent::SendRenderTransform_Concurrent()
+{	
+	// Don't update the transform of a component that needs to be recaptured,
+	// Otherwise the RT will get the new transform one frame before the capture
+	if (!bCaptureDirty)
+	{
+		UpdatePreviewShape();
+
+		if (ShouldRender())
+		{
+			World->Scene->UpdateReflectionCaptureTransform(this);
+		}
+	}
+
+	Super::SendRenderTransform_Concurrent();
+}
+
+void UReflectionCaptureComponent::DestroyRenderState_Concurrent()
+{
+	Super::DestroyRenderState_Concurrent();
+	World->Scene->RemoveReflectionCapture(this);
+}
+
+void UReflectionCaptureComponent::PostInitProperties()
+{
+	Super::PostInitProperties();
+
+	// Create a new guid in case this is a newly created component
+	// If not, this guid will be overwritten when serialized
+	FPlatformMisc::CreateGuid(StateId);
+
+	if (!HasAnyFlags(RF_ClassDefaultObject))
+	{
+		ReflectionCapturesToUpdateForLoad.AddUnique(this);
+		bCaptureDirty = true; 
+	}
+}
+
+void UReflectionCaptureComponent::SerializeSourceData(FArchive& Ar)
+{
+	if (Ar.UE4Ver() >= VER_UE4_REFLECTION_DATA_IN_PACKAGES)
+	{
+		if (Ar.IsSaving())
+		{
+			Ar << ReflectionCaptureDDCVer;
+
+			int32 StartOffset = Ar.Tell();
+			Ar << StartOffset;
+
+			bool bValid = FullHDRDerivedData != NULL;
+
+			Ar << bValid;
+
+			if (bValid)
+			{
+				Ar << FullHDRDerivedData->CompressedCapturedData;
+			}
+
+			int32 EndOffset = Ar.Tell();
+			Ar.Seek(StartOffset);
+			Ar << EndOffset;
+			Ar.Seek(EndOffset);
+		}
+		else if (Ar.IsLoading())
+		{
+			FGuid SavedVersion;
+			Ar << SavedVersion;
+
+			int32 EndOffset = 0;
+			Ar << EndOffset;
+
+			if (SavedVersion != ReflectionCaptureDDCVer)
+			{
+				// Guid version of saved source data doesn't match latest, skip the data
+				// The skipping is done so we don't have to maintain legacy serialization code paths when changing the format
+				Ar.Seek(EndOffset);
+			}
+			else
+			{
+				bool bValid = false;
+				Ar << bValid;
+
+				if (bValid)
+				{
+					FullHDRDerivedData = new FReflectionCaptureFullHDRDerivedData();
+					Ar << FullHDRDerivedData->CompressedCapturedData;
+				}
+			}
+		}
+	}
+}
+
+void UReflectionCaptureComponent::Serialize(FArchive& Ar)
+{
+	DECLARE_SCOPE_CYCLE_COUNTER(TEXT("UReflectionCaptureComponent::Serialize"), STAT_ReflectionCaptureComponent_Serialize, STATGROUP_LoadTime);
+
+	Super::Serialize(Ar);
+
+	bool bCooked = false;
+
+	if (Ar.UE4Ver() >= VER_UE4_REFLECTION_CAPTURE_COOKING)
+	{
+		bCooked = Ar.IsCooking();
+		// Save a bool indicating whether this is cooked data
+		// This is needed when loading cooked data, to know to serialize differently
+		Ar << bCooked;
+	}
+
+	if (FPlatformProperties::RequiresCookedData() && !bCooked && Ar.IsLoading())
+	{
+		UE_LOG(LogMaterial, Fatal, TEXT("This platform requires cooked packages, and this reflection capture does not contain cooked data %s."), *GetName());
+	}
+
+	if (bCooked)
+	{
+		static FName FullHDR(TEXT("FullHDR"));
+		static FName EncodedHDR(TEXT("EncodedHDR"));
+
+		// Saving for cooking path
+		if (Ar.IsCooking())
+		{
+			// Get all the reflection capture formats that the target platform wants
+			TArray<FName> Formats;
+			Ar.CookingTarget()->GetReflectionCaptureFormats(Formats);
+
+			int32 NumFormats = Formats.Num();
+			Ar << NumFormats;
+
+			for (int32 FormatIndex = 0; FormatIndex < NumFormats; FormatIndex++)
+			{
+				FName CurrentFormat = Formats[FormatIndex];
+
+				Ar << CurrentFormat;
+
+				if (CurrentFormat == FullHDR)
+				{
+					// FullHDRDerivedData would have been set in PostLoad during cooking if it exists in the DDC
+					// Can't generate it if missing, since that requires rendering the scene
+					bool bValid = FullHDRDerivedData != NULL;
+
+					Ar << bValid;
+
+					if (bValid)
+					{
+						Ar << FullHDRDerivedData->CompressedCapturedData;
+					}
+				}
+				else
+				{
+					check(CurrentFormat == EncodedHDR);
+
+					TRefCountPtr<FReflectionCaptureEncodedHDRDerivedData> EncodedHDRData;
+					
+					// FullHDRDerivedData would have been set in PostLoad during cooking if it exists in the DDC
+					// Generate temporary encoded HDR data for saving
+					if (FullHDRDerivedData != NULL)
+					{
+						EncodedHDRData = FReflectionCaptureEncodedHDRDerivedData::GenerateEncodedHDRData(*FullHDRDerivedData, StateId, Brightness);
+					}
+					
+					bool bValid = EncodedHDRData != NULL;
+
+					Ar << bValid;
+
+					if (bValid)
+					{
+						Ar << EncodedHDRData->CapturedData;
+					}
+					else if (!IsTemplate())
+					{
+						// Temporary warning until the cooker can do scene captures itself in the case of missing DDC
+						UE_LOG(LogMaterial, Warning, TEXT("Reflection capture requires encoded HDR data but none was found in the DDC!  This reflection will be black.  Fix by loading the map in the editor once.  %s."), *GetFullName());
+					}
+				}
+			}
+		}
+		else
+		{
+			// Loading cooked data path
+
+			int32 NumFormats = 0;
+			Ar << NumFormats;
+
+			for (int32 FormatIndex = 0; FormatIndex < NumFormats; FormatIndex++)
+			{
+				FName CurrentFormat;
+				Ar << CurrentFormat;
+
+				bool bValid = false;
+				Ar << bValid;
+
+				if (bValid)
+				{
+					if (CurrentFormat == FullHDR)
+					{
+						FullHDRDerivedData = new FReflectionCaptureFullHDRDerivedData();
+						Ar << FullHDRDerivedData->CompressedCapturedData;
+					}
+					else 
+					{
+						check(CurrentFormat == EncodedHDR);
+						EncodedHDRDerivedData = new FReflectionCaptureEncodedHDRDerivedData();
+						Ar << EncodedHDRDerivedData->CapturedData;
+					}
+				}
+				else if (CurrentFormat == EncodedHDR)
+				{
+					// Temporary warning until the cooker can do scene captures itself in the case of missing DDC
+					UE_LOG(LogMaterial, Error, TEXT("Reflection capture was loaded without any valid capture data and will be black.  This can happen if the DDC was not up to date during cooking.  Load the map in the editor once before cooking to fix.  %s."), *GetFullName());
+				}
+			}
+		}
+	}
+	else
+	{
+		SerializeSourceData(Ar);
+	}
+}
+
+void UReflectionCaptureComponent::PostLoad()
+{
+	Super::PostLoad();
+
+	bool bRetainAllFeatureLevelData = GIsEditor && GMaxRHIFeatureLevel >= ERHIFeatureLevel::SM4;
+	bool bEncodedDataRequired = bRetainAllFeatureLevelData || (GMaxRHIFeatureLevel == ERHIFeatureLevel::ES2 || GMaxRHIFeatureLevel == ERHIFeatureLevel::ES3_1);
+	bool bFullDataRequired = GMaxRHIFeatureLevel >= ERHIFeatureLevel::SM4;
+
+	// If we're loading on a platform that doesn't require cooked data, attempt to load missing data from the DDC
+	if (!FPlatformProperties::RequiresCookedData())
+	{
+		// Only check the DDC if we don't already have it loaded
+		// If we are loading cooked then FullHDRDerivedData might be setup already (FullHDRDerivedData is set in Serialize)
+		if (!FullHDRDerivedData)
+		{
+			FullHDRDerivedData = new FReflectionCaptureFullHDRDerivedData();
+
+			// Attempt to load the full HDR data from the DDC
+			if (!GetDerivedDataCacheRef().GetSynchronous(*FReflectionCaptureFullHDRDerivedData::GetDDCKeyString(StateId), FullHDRDerivedData->CompressedCapturedData))
+			{
+				// Try to load the previous version
+				// This is a temporary measure for backwards compatibility of DDC
+				if (!GetDerivedDataCacheRef().GetSynchronous(*FReflectionCaptureFullHDRDerivedData::GetLegacyDDCKeyString(StateId), FullHDRDerivedData->CompressedCapturedData))
+				{
+					bDerivedDataDirty = true;
+					delete FullHDRDerivedData;
+					FullHDRDerivedData = NULL;
+
+					if (!FApp::CanEverRender())
+					{
+						// Warn, especially when running the DDC commandlet to build a DDC for the binary version of UE4.
+						UE_LOG(LogMaterial, Warning, TEXT("Reflection capture was loaded without any valid capture data and will be black.  This can happen if the DDC was not up to date during cooking.  Load the map in the editor once before cooking to fix.  %s."), *GetFullName());
+					}
+				}
+			}
+		}
+
+		// If we have full HDR data but not encoded HDR data, generate the encoded data now
+		if (FullHDRDerivedData 
+			&& !EncodedHDRDerivedData 
+			&& bEncodedDataRequired)
+		{
+			EncodedHDRDerivedData = FReflectionCaptureEncodedHDRDerivedData::GenerateEncodedHDRData(*FullHDRDerivedData, StateId, Brightness);
+		}
+	}
+	
+	// Initialize rendering resources for the current feature level, and toss data only needed by other feature levels (unless in editor mode, in which all feature level data should be resident.)
+	if (FullHDRDerivedData && bFullDataRequired)
+	{
+		// Don't need encoded HDR data for rendering on this feature level
+		INC_MEMORY_STAT_BY(STAT_ReflectionCaptureMemory, FullHDRDerivedData->CompressedCapturedData.GetAllocatedSize());
+
+		if (GMaxRHIFeatureLevel == ERHIFeatureLevel::SM4)
+		{
+			SM4FullHDRCubemapTexture = new FReflectionTextureCubeResource();
+			SM4FullHDRCubemapTexture->SetupParameters(GReflectionCaptureSize, FMath::CeilLogTwo(GReflectionCaptureSize) + 1, PF_FloatRGBA, &FullHDRDerivedData->GetCapturedDataForSM4Load());
+			BeginInitResource(SM4FullHDRCubemapTexture);
+		}
+
+		if (!bEncodedDataRequired)
+		{
+			EncodedHDRDerivedData = NULL;
+		}
+	}
+
+	if (EncodedHDRDerivedData && bEncodedDataRequired)
+	{
+		if (FPlatformProperties::RequiresCookedData())
+		{
+			INC_MEMORY_STAT_BY(STAT_ReflectionCaptureMemory, EncodedHDRDerivedData->CapturedData.GetAllocatedSize());
+		}
+
+		// Create a cubemap texture out of the encoded HDR data
+		EncodedHDRCubemapTexture = new FReflectionTextureCubeResource();
+		EncodedHDRCubemapTexture->SetupParameters(GReflectionCaptureSize, FMath::CeilLogTwo(GReflectionCaptureSize) + 1, PF_B8G8R8A8, &EncodedHDRDerivedData->CapturedData);
+		BeginInitResource(EncodedHDRCubemapTexture);
+
+		// free up the full hdr data if we no longer need it.
+		if (FullHDRDerivedData && !bFullDataRequired)
+		{
+			delete FullHDRDerivedData;
+			FullHDRDerivedData = NULL;
+		}
+	}
+}
+
+void UReflectionCaptureComponent::PreSave() 
+{
+	Super::PreSave();
+
+	// This is done on save of the package, because this capture data can only be generated by the renderer
+	// So we must make efforts to ensure that it is generated in the editor, because it can't be generated during cooking when missing
+	// Note: This will only work when registered
+	if (World)
+	{
+		ReadbackFromGPUAndSaveDerivedData(World);
+	}
+}
+
+void UReflectionCaptureComponent::PostDuplicate(bool bDuplicateForPIE)
+{
+	if (!bDuplicateForPIE)
+	{
+		// Reset the StateId on duplication since it needs to be unique for each capture.
+		// PostDuplicate covers direct calls to StaticDuplicateObject, but not actor duplication (see PostEditImport)
+		FPlatformMisc::CreateGuid(StateId);
+	}
+}
+
+void UReflectionCaptureComponent::UpdateDerivedData(FReflectionCaptureFullHDRDerivedData* NewDerivedData)
+{
+	if (FullHDRDerivedData)
+	{
+		// Delete the derived data on the rendering thread, since the rendering thread may be reading from its contents in FScene::UpdateReflectionCaptureContents
+		ENQUEUE_UNIQUE_RENDER_COMMAND_ONEPARAMETER( 
+			DeleteCaptureDataCommand,
+			FReflectionCaptureFullHDRDerivedData*, FullHDRDerivedData, FullHDRDerivedData,
+		{
+			delete FullHDRDerivedData;
+		});
+	}
+
+	FullHDRDerivedData = NewDerivedData;
+}
+
+FReflectionCaptureProxy* UReflectionCaptureComponent::CreateSceneProxy()
+{
+	return new FReflectionCaptureProxy(this);
+}
+
+#if WITH_EDITOR
+void UReflectionCaptureComponent::PostEditChangeProperty(FPropertyChangedEvent& PropertyChangedEvent)
+{
+	// AActor::PostEditChange will ForceUpdateComponents()
+	Super::PostEditChangeProperty(PropertyChangedEvent);
+
+	if (PropertyChangedEvent.Property && PropertyChangedEvent.Property->GetFName() == GET_MEMBER_NAME_CHECKED(UReflectionCaptureComponent, bVisible))
+	{
+		SetCaptureIsDirty();
+	}
+}
+
+void UReflectionCaptureComponent::PostEditImport()
+{
+	Super::PostEditImport(); 
+
+	// Generate a new StateId.  This is needed to cover actor duplication through alt-drag or copy-paste.
+	SetCaptureIsDirty();
+}
+#endif // WITH_EDITOR
+
+void UReflectionCaptureComponent::BeginDestroy()
+{
+	// Deregister the component from the update queue
+	if (bCaptureDirty)
+	{
+		ReflectionCapturesToUpdate.Remove(this);
+		ReflectionCapturesToUpdateForLoad.Remove(this);
+	}
+
+	// Have to do this because we can't use GetWorld in BeginDestroy
+	for (TSet<FSceneInterface*>::TConstIterator SceneIt(GetRendererModule().GetAllocatedScenes()); SceneIt; ++SceneIt)
+	{
+		FSceneInterface* Scene = *SceneIt;
+		Scene->ReleaseReflectionCubemap(this);
+	}
+
+	if (SM4FullHDRCubemapTexture)
+	{
+		BeginReleaseResource(SM4FullHDRCubemapTexture);
+	}
+
+	if (EncodedHDRCubemapTexture)
+	{
+		BeginReleaseResource(EncodedHDRCubemapTexture);
+	}
+
+	// Begin a fence to track the progress of the above BeginReleaseResource being completed on the RT
+	ReleaseResourcesFence.BeginFence();
+
+	Super::BeginDestroy();
+}
+
+bool UReflectionCaptureComponent::IsReadyForFinishDestroy()
+{
+	// Wait until the fence is complete before allowing destruction
+	return Super::IsReadyForFinishDestroy() && ReleaseResourcesFence.IsFenceComplete();
+}
+
+void UReflectionCaptureComponent::FinishDestroy()
+{
+	UpdateDerivedData(NULL);
+
+	if (SM4FullHDRCubemapTexture)
+	{
+		delete SM4FullHDRCubemapTexture;
+		SM4FullHDRCubemapTexture = NULL;
+	}
+
+	if (EncodedHDRCubemapTexture)
+	{
+		delete EncodedHDRCubemapTexture;
+		EncodedHDRCubemapTexture = NULL;
+	}
+
+	Super::FinishDestroy();
+}
+
+void UReflectionCaptureComponent::SetCaptureIsDirty() 
+{ 
+	if (bVisible)
+	{
+		UpdateDerivedData(NULL);
+		FPlatformMisc::CreateGuid(StateId);
+		bDerivedDataDirty = true;
+		ReflectionCapturesToUpdate.AddUnique(this);
+		bCaptureDirty = true; 
+	}
+}
+
+void ReadbackFromSM4Cubemap_RenderingThread(FRHICommandListImmediate& RHICmdList, FReflectionTextureCubeResource* SM4FullHDRCubemapTexture, FReflectionCaptureFullHDRDerivedData* OutDerivedData)
+{
+	const int32 EffectiveTopMipSize = GReflectionCaptureSize;
+	const int32 NumMips = FMath::CeilLogTwo(EffectiveTopMipSize) + 1;
+
+	TArray<uint8> CaptureData;
+	int32 CaptureDataSize = 0;
+
+	for (int32 MipIndex = 0; MipIndex < NumMips; MipIndex++)
+	{
+		const int32 MipSize = 1 << (NumMips - MipIndex - 1);
+
+		for (int32 CubeFace = 0; CubeFace < CubeFace_MAX; CubeFace++)
+		{
+			CaptureDataSize += MipSize * MipSize * sizeof(FFloat16Color);
+		}
+	}
+
+	CaptureData.Empty(CaptureDataSize);
+	CaptureData.AddZeroed(CaptureDataSize);
+	int32 MipBaseIndex = 0;
+
+
+
+	for (int32 MipIndex = 0; MipIndex < NumMips; MipIndex++)
+	{
+		check(SM4FullHDRCubemapTexture->GetTextureRHI()->GetFormat() == PF_FloatRGBA);
+		const int32 MipSize = 1 << (NumMips - MipIndex - 1);
+		const int32 CubeFaceBytes = MipSize * MipSize * sizeof(FFloat16Color);
+
+		for (int32 CubeFace = 0; CubeFace < CubeFace_MAX; CubeFace++)
+		{
+			TArray<FFloat16Color> SurfaceData;
+			// Read each mip face
+			//@todo - do this without blocking the GPU so many times
+			//@todo - pool the temporary textures in RHIReadSurfaceFloatData instead of always creating new ones
+			RHICmdList.ReadSurfaceFloatData(SM4FullHDRCubemapTexture->GetTextureRHI(), FIntRect(0, 0, MipSize, MipSize), SurfaceData, (ECubeFace)CubeFace, 0, MipIndex);
+			const int32 DestIndex = MipBaseIndex + CubeFace * CubeFaceBytes;
+			uint8* FaceData = &CaptureData[DestIndex];
+			check(SurfaceData.Num() * SurfaceData.GetTypeSize() == CubeFaceBytes);
+			FMemory::Memcpy(FaceData, SurfaceData.GetData(), CubeFaceBytes);
+		}
+
+		MipBaseIndex += CubeFaceBytes * CubeFace_MAX;
+	}
+
+	OutDerivedData->InitializeFromUncompressedData(CaptureData);
+}
+
+void ReadbackFromSM4Cubemap(FReflectionTextureCubeResource* SM4FullHDRCubemapTexture, FReflectionCaptureFullHDRDerivedData* OutDerivedData)
+{
+	ENQUEUE_UNIQUE_RENDER_COMMAND_TWOPARAMETER( 
+		ReadbackReflectionData,
+		FReflectionTextureCubeResource*, SM4FullHDRCubemapTexture, SM4FullHDRCubemapTexture,
+		FReflectionCaptureFullHDRDerivedData*, OutDerivedData, OutDerivedData,
+	{
+		ReadbackFromSM4Cubemap_RenderingThread(RHICmdList, SM4FullHDRCubemapTexture, OutDerivedData);
+	});
+
+	FlushRenderingCommands();
+}
+
+void UReflectionCaptureComponent::ReadbackFromGPUAndSaveDerivedData(UWorld* WorldToUpdate)
+{
+	if (bDerivedDataDirty && !IsRunningCommandlet() && WorldToUpdate->FeatureLevel >= ERHIFeatureLevel::SM4)
+	{
+		FReflectionCaptureFullHDRDerivedData* NewDerivedData = new FReflectionCaptureFullHDRDerivedData();
+
+		if (WorldToUpdate->FeatureLevel == ERHIFeatureLevel::SM4)
+		{
+			if (SM4FullHDRCubemapTexture)
+			{
+				ReadbackFromSM4Cubemap(SM4FullHDRCubemapTexture, NewDerivedData);
+			}
+		}
+		else
+		{
+			WorldToUpdate->Scene->GetReflectionCaptureData(this, *NewDerivedData);
+		}
+
+		if (NewDerivedData->CompressedCapturedData.Num() > 0)
+		{
+			GetDerivedDataCacheRef().Put(*FReflectionCaptureFullHDRDerivedData::GetDDCKeyString(StateId), NewDerivedData->CompressedCapturedData);
+
+			// Update our copy in memory
+			UpdateDerivedData(NewDerivedData);
+		}
+		else
+		{
+			delete NewDerivedData;
+		}
+	}
+}
+
+void UReflectionCaptureComponent::UpdateReflectionCaptureContents(UWorld* WorldToUpdate)
+{
+	if (WorldToUpdate->Scene 
+		// Don't capture and read back capture contents if we are currently doing async shader compiling
+		// This will keep the update requests in the queue until compiling finishes
+		// Note: this will also prevent uploads of cubemaps from DDC, which is unintentional
+		&& (GShaderCompilingManager == NULL || !GShaderCompilingManager->IsCompiling()))
+	{
+		TArray<UReflectionCaptureComponent*> WorldCombinedCaptures;
+
+		for (int32 CaptureIndex = ReflectionCapturesToUpdate.Num() - 1; CaptureIndex >= 0; CaptureIndex--)
+		{
+			UReflectionCaptureComponent* CaptureComponent = ReflectionCapturesToUpdate[CaptureIndex];
+
+			if (!CaptureComponent->GetOwner() || WorldToUpdate->ContainsActor(CaptureComponent->GetOwner()))
+			{
+				WorldCombinedCaptures.Add(CaptureComponent);
+				ReflectionCapturesToUpdate.RemoveAt(CaptureIndex);
+			}
+		}
+
+		TArray<UReflectionCaptureComponent*> WorldCapturesToUpdateForLoad;
+
+		for (int32 CaptureIndex = ReflectionCapturesToUpdateForLoad.Num() - 1; CaptureIndex >= 0; CaptureIndex--)
+		{
+			UReflectionCaptureComponent* CaptureComponent = ReflectionCapturesToUpdateForLoad[CaptureIndex];
+
+			if (!CaptureComponent->GetOwner() || WorldToUpdate->ContainsActor(CaptureComponent->GetOwner()))
+			{
+				WorldCombinedCaptures.Add(CaptureComponent);
+				WorldCapturesToUpdateForLoad.Add(CaptureComponent);
+				ReflectionCapturesToUpdateForLoad.RemoveAt(CaptureIndex);
+			}
+		}
+
+		const auto FeatureLevel = WorldToUpdate->Scene->GetFeatureLevel();
+
+		if (FeatureLevel == ERHIFeatureLevel::SM4)
+		{
+			for (int32 CaptureIndex = 0; CaptureIndex < WorldCombinedCaptures.Num(); CaptureIndex++)
+			{
+				UReflectionCaptureComponent* ReflectionComponent = WorldCombinedCaptures[CaptureIndex];
+
+				if (!ReflectionComponent->SM4FullHDRCubemapTexture)
+				{
+					// Create the cubemap if it wasn't already - this happens when updating a reflection capture that doesn't have valid DDC
+					ReflectionComponent->SM4FullHDRCubemapTexture = new FReflectionTextureCubeResource();
+					ReflectionComponent->SM4FullHDRCubemapTexture->SetupParameters(GReflectionCaptureSize, FMath::CeilLogTwo(GReflectionCaptureSize) + 1, PF_FloatRGBA, NULL);
+					BeginInitResource(ReflectionComponent->SM4FullHDRCubemapTexture);
+					ReflectionComponent->MarkRenderStateDirty();
+				}
+			}
+		}
+
+		WorldToUpdate->Scene->AllocateReflectionCaptures(WorldCombinedCaptures);
+
+		if (FeatureLevel >= ERHIFeatureLevel::SM4 && !FPlatformProperties::RequiresCookedData())
+		{
+			for (int32 CaptureIndex = 0; CaptureIndex < WorldCapturesToUpdateForLoad.Num(); CaptureIndex++)
+			{
+				// Save the derived data for any captures that were dirty on load
+				// This allows the derived data to get cached without having to resave a map
+				WorldCapturesToUpdateForLoad[CaptureIndex]->ReadbackFromGPUAndSaveDerivedData(WorldToUpdate);
+			}
+		}
+	}
+}
+
+#if WITH_EDITOR
+void UReflectionCaptureComponent::PreFeatureLevelChange(ERHIFeatureLevel::Type PendingFeatureLevel)
+{
+	if (PendingFeatureLevel == ERHIFeatureLevel::ES2 || PendingFeatureLevel == ERHIFeatureLevel::ES3_1)
+	{
+		// generate encoded hdr data for ES2 preview mode.
+		if (World != nullptr)
+		{
+			// Capture full hdr derived data first.
+			ReadbackFromGPUAndSaveDerivedData(World);
+		}
+		if (FullHDRDerivedData)
+		{
+			EncodedHDRDerivedData = FReflectionCaptureEncodedHDRDerivedData::GenerateEncodedHDRData(*FullHDRDerivedData, StateId, Brightness);
+			if (FPlatformProperties::RequiresCookedData() && EncodedHDRDerivedData)
+			{
+				INC_MEMORY_STAT_BY(STAT_ReflectionCaptureMemory, EncodedHDRDerivedData->CapturedData.GetAllocatedSize());
+			}
+
+			if (EncodedHDRCubemapTexture == nullptr)
+			{
+				EncodedHDRCubemapTexture = new FReflectionTextureCubeResource();
+			}
+			EncodedHDRCubemapTexture->SetupParameters(GReflectionCaptureSize, FMath::CeilLogTwo(GReflectionCaptureSize) + 1, PF_B8G8R8A8, &EncodedHDRDerivedData->CapturedData);
+			BeginInitResource(EncodedHDRCubemapTexture);
+		}
+	}
+	else
+	{
+		// Release ES2's unused textures & data.
+		if (FPlatformProperties::RequiresCookedData() && EncodedHDRDerivedData)
+		{
+			DEC_MEMORY_STAT_BY(STAT_ReflectionCaptureMemory, EncodedHDRDerivedData->CapturedData.GetAllocatedSize());
+		}
+		EncodedHDRDerivedData = nullptr;
+		if (EncodedHDRCubemapTexture)
+		{
+			BeginReleaseResource(EncodedHDRCubemapTexture);
+			FlushRenderingCommands();
+			delete EncodedHDRCubemapTexture;
+			EncodedHDRCubemapTexture = nullptr;
+		}
+
+		// for >= SM4 capture should be updated.
+		SetCaptureIsDirty();
+	}
+
+	if (PendingFeatureLevel == ERHIFeatureLevel::SM4)
+	{
+		if (SM4FullHDRCubemapTexture == nullptr)
+		{
+			SM4FullHDRCubemapTexture = new FReflectionTextureCubeResource();
+			SM4FullHDRCubemapTexture->SetupParameters(GReflectionCaptureSize, FMath::CeilLogTwo(GReflectionCaptureSize) + 1, PF_FloatRGBA, NULL);
+			BeginInitResource(SM4FullHDRCubemapTexture);
+		}
+	}
+	else
+	{
+		// release SM4 texture
+		if (SM4FullHDRCubemapTexture != nullptr)
+		{
+			BeginReleaseResource(SM4FullHDRCubemapTexture);
+			FlushRenderingCommands();
+			delete SM4FullHDRCubemapTexture;
+			SM4FullHDRCubemapTexture = nullptr;
+		}
+	}
+}
+#endif // WITH_EDITOR
+
+USphereReflectionCaptureComponent::USphereReflectionCaptureComponent(const FObjectInitializer& ObjectInitializer)
+	: Super(ObjectInitializer)
+{
+	InfluenceRadius = 3000;
+}
+
+void USphereReflectionCaptureComponent::UpdatePreviewShape()
+{
+	if (PreviewInfluenceRadius)
+	{
+		PreviewInfluenceRadius->InitSphereRadius(InfluenceRadius);
+	}
+}
+
+#if WITH_EDITOR
+void USphereReflectionCaptureComponent::PostEditChangeProperty(FPropertyChangedEvent& PropertyChangedEvent)
+{
+	// AActor::PostEditChange will ForceUpdateComponents()
+	Super::PostEditChangeProperty(PropertyChangedEvent);
+
+	if (PropertyChangedEvent.Property && 
+		(PropertyChangedEvent.Property->GetFName() == GET_MEMBER_NAME_CHECKED(USphereReflectionCaptureComponent, InfluenceRadius)))
+	{
+		SetCaptureIsDirty();
+	}
+}
+#endif // WITH_EDITOR
+
+UBoxReflectionCaptureComponent::UBoxReflectionCaptureComponent(const FObjectInitializer& ObjectInitializer)
+	: Super(ObjectInitializer)
+{
+	BoxTransitionDistance = 100;
+}
+
+void UBoxReflectionCaptureComponent::UpdatePreviewShape()
+{
+	if (PreviewCaptureBox)
+	{
+		PreviewCaptureBox->InitBoxExtent((ComponentToWorld.GetScale3D() - FVector(BoxTransitionDistance)) / ComponentToWorld.GetScale3D());
+	}
+}
+
+float UBoxReflectionCaptureComponent::GetInfluenceBoundingRadius() const
+{
+	return (ComponentToWorld.GetScale3D() + FVector(BoxTransitionDistance)).Size();
+}
+
+#if WITH_EDITOR
+void UBoxReflectionCaptureComponent::PostEditChangeProperty(FPropertyChangedEvent& PropertyChangedEvent)
+{
+	// AActor::PostEditChange will ForceUpdateComponents()
+	Super::PostEditChangeProperty(PropertyChangedEvent);
+
+	if (PropertyChangedEvent.Property && PropertyChangedEvent.Property->GetFName() == GET_MEMBER_NAME_CHECKED(UBoxReflectionCaptureComponent, BoxTransitionDistance))
+	{
+		SetCaptureIsDirty();
+	}
+}
+#endif // WITH_EDITOR
+
+UPlaneReflectionCaptureComponent::UPlaneReflectionCaptureComponent(const FObjectInitializer& ObjectInitializer)
+	: Super(ObjectInitializer)
+{
+	InfluenceRadiusScale = 2;
+}
+
+void UPlaneReflectionCaptureComponent::UpdatePreviewShape()
+{
+	if (PreviewInfluenceRadius)
+	{
+		PreviewInfluenceRadius->InitSphereRadius(GetInfluenceBoundingRadius());
+	}
+}
+
+float UPlaneReflectionCaptureComponent::GetInfluenceBoundingRadius() const
+{
+	return FVector2D(ComponentToWorld.GetScale3D().Y, ComponentToWorld.GetScale3D().Z).Size() * InfluenceRadiusScale;
+}
+
+FReflectionCaptureProxy::FReflectionCaptureProxy(const UReflectionCaptureComponent* InComponent)
+{
+	PackedIndex = INDEX_NONE;
+
+	const USphereReflectionCaptureComponent* SphereComponent = Cast<const USphereReflectionCaptureComponent>(InComponent);
+	const UBoxReflectionCaptureComponent* BoxComponent = Cast<const UBoxReflectionCaptureComponent>(InComponent);
+	const UPlaneReflectionCaptureComponent* PlaneComponent = Cast<const UPlaneReflectionCaptureComponent>(InComponent);
+
+	// Initialize shape specific settings
+	Shape = EReflectionCaptureShape::Num;
+	BoxTransitionDistance = 0;
+
+	if (SphereComponent)
+	{
+		Shape = EReflectionCaptureShape::Sphere;
+	}
+	else if (BoxComponent)
+	{
+		Shape = EReflectionCaptureShape::Box;
+		BoxTransitionDistance = BoxComponent->BoxTransitionDistance;
+	}
+	else if (PlaneComponent)
+	{
+		Shape = EReflectionCaptureShape::Plane;
+	}
+	else
+	{
+		check(0);
+	}
+	
+	// Initialize common settings
+	Component = InComponent;
+	SM4FullHDRCubemap = InComponent->SM4FullHDRCubemapTexture;
+	EncodedHDRCubemap = InComponent->EncodedHDRCubemapTexture;
+	SetTransform(InComponent->ComponentToWorld.ToMatrixWithScale());
+	InfluenceRadius = InComponent->GetInfluenceBoundingRadius();
+	Brightness = InComponent->Brightness;
+	Guid = GetTypeHash( Component->GetPathName() );
+}
+
+void FReflectionCaptureProxy::SetTransform(const FMatrix& InTransform)
+{
+	Position = InTransform.GetOrigin();
+	BoxTransform = InTransform.Inverse();
+
+	FVector ForwardVector(1.0f,0.0f,0.0f);
+	FVector RightVector(0.0f,-1.0f,0.0f);
+	const FVector4 PlaneNormal = InTransform.TransformVector(ForwardVector);
+
+	// Normalize the plane
+	ReflectionPlane = FPlane(Position, FVector(PlaneNormal).GetSafeNormal());
+	const FVector ReflectionXAxis = InTransform.TransformVector(RightVector);
+	const FVector ScaleVector = InTransform.GetScaleVector();
+	BoxScales = ScaleVector;
+	// Include the owner's draw scale in the axes
+	ReflectionXAxisAndYScale = ReflectionXAxis.GetSafeNormal() * ScaleVector.Y;
+	ReflectionXAxisAndYScale.W = ScaleVector.Y / ScaleVector.Z;
+}
+
+/** Returns CaptureComponent subobject **/
+UReflectionCaptureComponent* AReflectionCapture::GetCaptureComponent() const { return CaptureComponent; }
+#if WITH_EDITORONLY_DATA
+/** Returns SpriteComponent subobject **/
+UBillboardComponent* AReflectionCapture::GetSpriteComponent() const { return SpriteComponent; }
+#endif
+/** Returns DrawCaptureRadius subobject **/
+UDrawSphereComponent* ASphereReflectionCapture::GetDrawCaptureRadius() const { return DrawCaptureRadius; }
diff --git a/Engine/Source/Runtime/Engine/Private/Components/SceneCaptureComponent.cpp b/Engine/Source/Runtime/Engine/Private/Components/SceneCaptureComponent.cpp
new file mode 100644
--- /dev/null
+++ b/Engine/Source/Runtime/Engine/Private/Components/SceneCaptureComponent.cpp
@@ -0,0 +1,420 @@
+// Copyright 1998-2015 Epic Games, Inc. All Rights Reserved.
+
+/*=============================================================================
+	
+=============================================================================*/
+
+#include "EnginePrivate.h"
+#include "../../Renderer/Private/ScenePrivate.h"
+#include "Engine/SceneCapture.h"
+#include "Engine/SceneCapture2D.h"
+#include "Components/SceneCaptureComponent2D.h"
+#include "Engine/SceneCaptureCube.h"
+#include "Components/SceneCaptureComponentCube.h"
+#include "Components/DrawFrustumComponent.h"
+
+ASceneCapture::ASceneCapture(const FObjectInitializer& ObjectInitializer)
+	: Super(ObjectInitializer)
+{
+	MeshComp = CreateDefaultSubobject<UStaticMeshComponent>(TEXT("CamMesh0"));
+
+	MeshComp->SetCollisionProfileName(UCollisionProfile::NoCollision_ProfileName);
+
+	MeshComp->bHiddenInGame = true;
+	MeshComp->CastShadow = false;
+	MeshComp->PostPhysicsComponentTick.bCanEverTick = false;
+	RootComponent = MeshComp;
+}
+// -----------------------------------------------
+
+ASceneCapture2D::ASceneCapture2D(const FObjectInitializer& ObjectInitializer)
+	: Super(ObjectInitializer)
+{
+	DrawFrustum = CreateDefaultSubobject<UDrawFrustumComponent>(TEXT("DrawFrust0"));
+	DrawFrustum->AlwaysLoadOnClient = false;
+	DrawFrustum->AlwaysLoadOnServer = false;
+	DrawFrustum->AttachParent = GetMeshComp();
+
+	CaptureComponent2D = CreateDefaultSubobject<USceneCaptureComponent2D>(TEXT("NewSceneCaptureComponent2D"));
+	CaptureComponent2D->AttachParent = GetMeshComp();
+}
+
+void ASceneCapture2D::OnInterpToggle(bool bEnable)
+{
+	CaptureComponent2D->SetVisibility(bEnable);
+}
+
+void ASceneCapture2D::UpdateDrawFrustum()
+{
+	if(DrawFrustum && CaptureComponent2D)
+	{
+		DrawFrustum->FrustumStartDist = GNearClippingPlane;
+		
+		// 1000 is the default frustum distance, ideally this would be infinite but that might cause rendering issues
+		DrawFrustum->FrustumEndDist = (CaptureComponent2D->MaxViewDistanceOverride > DrawFrustum->FrustumStartDist)
+			? CaptureComponent2D->MaxViewDistanceOverride : 1000.0f;
+
+		DrawFrustum->FrustumAngle = CaptureComponent2D->FOVAngle;
+		//DrawFrustum->FrustumAspectRatio = CaptureComponent2D->AspectRatio;
+	}
+}
+
+void ASceneCapture2D::PostActorCreated()
+{
+	Super::PostActorCreated();
+
+	// no need load the editor mesh when there is no editor
+#if WITH_EDITOR
+	if(GetMeshComp())
+	{
+		if (!IsRunningCommandlet())
+		{
+			if( !GetMeshComp()->StaticMesh)
+			{
+				UStaticMesh* CamMesh = LoadObject<UStaticMesh>(NULL, TEXT("/Engine/EditorMeshes/MatineeCam_SM.MatineeCam_SM"), NULL, LOAD_None, NULL);
+				GetMeshComp()->SetStaticMesh(CamMesh);
+			}
+		}
+	}
+#endif
+
+	// Sync component with CameraActor frustum settings.
+	UpdateDrawFrustum();
+}
+// -----------------------------------------------
+
+ASceneCaptureCube::ASceneCaptureCube(const FObjectInitializer& ObjectInitializer)
+	: Super(ObjectInitializer)
+{
+	DrawFrustum = CreateDefaultSubobject<UDrawFrustumComponent>(TEXT("DrawFrust0"));
+	DrawFrustum->AlwaysLoadOnClient = false;
+	DrawFrustum->AlwaysLoadOnServer = false;
+	DrawFrustum->AttachParent = GetMeshComp();
+
+	CaptureComponentCube = CreateDefaultSubobject<USceneCaptureComponentCube>(TEXT("NewSceneCaptureComponentCube"));
+	CaptureComponentCube->AttachParent = GetMeshComp();
+}
+
+void ASceneCaptureCube::OnInterpToggle(bool bEnable)
+{
+	CaptureComponentCube->SetVisibility(bEnable);
+}
+
+void ASceneCaptureCube::UpdateDrawFrustum()
+{
+	if(DrawFrustum && CaptureComponentCube)
+	{
+		DrawFrustum->FrustumStartDist = GNearClippingPlane;
+
+		// 1000 is the default frustum distance, ideally this would be infinite but that might cause rendering issues
+		DrawFrustum->FrustumEndDist = (CaptureComponentCube->MaxViewDistanceOverride > DrawFrustum->FrustumStartDist)
+			? CaptureComponentCube->MaxViewDistanceOverride : 1000.0f;
+
+		DrawFrustum->FrustumAngle = 90;
+	}
+}
+
+void ASceneCaptureCube::PostActorCreated()
+{
+	Super::PostActorCreated();
+
+	// no need load the editor mesh when there is no editor
+#if WITH_EDITOR
+	if(GetMeshComp())
+	{
+		if (!IsRunningCommandlet())
+		{
+			if( !GetMeshComp()->StaticMesh)
+			{
+				UStaticMesh* CamMesh = LoadObject<UStaticMesh>(NULL, TEXT("/Engine/EditorMeshes/MatineeCam_SM.MatineeCam_SM"), NULL, LOAD_None, NULL);
+				GetMeshComp()->SetStaticMesh(CamMesh);
+			}
+		}
+	}
+#endif
+
+	// Sync component with CameraActor frustum settings.
+	UpdateDrawFrustum();
+}
+#if WITH_EDITOR
+
+void ASceneCaptureCube::PostEditMove(bool bFinished)
+{
+	Super::PostEditMove(bFinished);
+	if(bFinished)
+	{
+		CaptureComponentCube->UpdateContent();
+	}
+}
+#endif
+// -----------------------------------------------
+USceneCaptureComponent::USceneCaptureComponent(const FObjectInitializer& ObjectInitializer)
+	: Super(ObjectInitializer), ShowFlags(FEngineShowFlags(ESFIM_Game))
+{
+	bCaptureEveryFrame = true;
+	MaxViewDistanceOverride = -1;
+
+	// Disable features that are not desired when capturing the scene
+	ShowFlags.MotionBlur = 0; // motion blur doesn't work correctly with scene captures.
+	ShowFlags.SeparateTranslucency = 0;
+	ShowFlags.HMDDistortion = 0;
+}
+
+void USceneCaptureComponent::PostLoad()
+{
+	Super::PostLoad();
+
+	// Make sure any loaded saved flag settings are reflected in our FEngineShowFlags
+	UpdateShowFlags();
+}
+
+void USceneCaptureComponent::AddReferencedObjects(UObject* InThis, FReferenceCollector& Collector)
+{
+	USceneCaptureComponent* This = CastChecked<USceneCaptureComponent>(InThis);
+
+	FSceneViewStateInterface* Ref = This->ViewState.GetReference();
+
+	if (Ref)
+	{
+		Ref->AddReferencedObjects(Collector);
+	}
+
+	Super::AddReferencedObjects(This, Collector);
+}
+
+void USceneCaptureComponent::HideComponent(UPrimitiveComponent* InComponent)
+{
+	if (InComponent)
+	{
+		HiddenComponents.AddUnique(InComponent);
+	}
+}
+
+void USceneCaptureComponent::HideActorComponents(AActor* InActor)
+{
+	if (InActor)
+	{
+		TInlineComponentArray<UPrimitiveComponent*> PrimitiveComponents;
+		InActor->GetComponents(PrimitiveComponents);
+		for (int32 ComponentIndex = 0, NumComponents = PrimitiveComponents.Num(); ComponentIndex < NumComponents; ++ComponentIndex)
+		{
+			HiddenComponents.AddUnique(PrimitiveComponents[ComponentIndex]);
+		}
+	}
+}
+
+FSceneViewStateInterface* USceneCaptureComponent::GetViewState()
+{
+	FSceneViewStateInterface* ViewStateInterface = ViewState.GetReference();
+	if (bCaptureEveryFrame && ViewStateInterface == NULL)
+	{
+		ViewState.Allocate();
+		ViewStateInterface = ViewState.GetReference();
+	}
+	else if (!bCaptureEveryFrame && ViewStateInterface)
+	{
+		ViewState.Destroy();
+		ViewStateInterface = NULL;
+	}
+	return ViewStateInterface;
+}
+
+void USceneCaptureComponent::UpdateShowFlags()
+{
+	for (FEngineShowFlagsSetting ShowFlagSetting : ShowFlagSettings)
+	{
+		int32 SettingIndex = ShowFlags.FindIndexByName(*(ShowFlagSetting.ShowFlagName));
+		if (SettingIndex != INDEX_NONE)
+		{ 
+			ShowFlags.SetSingleFlag(SettingIndex, ShowFlagSetting.Enabled);
+		}
+	}
+}
+
+#if WITH_EDITOR
+void USceneCaptureComponent::PostEditChangeProperty(struct FPropertyChangedEvent& PropertyChangedEvent)
+{
+	Super::PostEditChangeProperty(PropertyChangedEvent);
+
+	const FName Name = (PropertyChangedEvent.Property != nullptr) ? PropertyChangedEvent.Property->GetFName() : NAME_None;
+	const FName MemberPropertyName = (PropertyChangedEvent.MemberProperty != NULL) ? PropertyChangedEvent.MemberProperty->GetFName() : NAME_None;
+
+	// If our ShowFlagSetting UStruct changed, (or if PostEditChange was called without specifying a property) update the actual show flags
+	if (MemberPropertyName.IsEqual("ShowFlagSettings") || MemberPropertyName.IsNone())
+	{
+		UpdateShowFlags();
+	}
+}
+#endif
+
+bool USceneCaptureComponent::GetSettingForShowFlag(FString FlagName, FEngineShowFlagsSetting** ShowFlagSettingOut)
+{
+	bool HasSetting = false;
+	for (int32 ShowFlagSettingsIndex = 0; ShowFlagSettingsIndex < ShowFlagSettings.Num(); ++ShowFlagSettingsIndex)
+	{
+		if (ShowFlagSettings[ShowFlagSettingsIndex].ShowFlagName.Equals(FlagName))
+		{
+			HasSetting = true;
+			*ShowFlagSettingOut = &(ShowFlagSettings[ShowFlagSettingsIndex]);
+			break;
+		}
+	}	
+	return HasSetting;
+}
+
+// -----------------------------------------------
+
+
+USceneCaptureComponent2D::USceneCaptureComponent2D(const FObjectInitializer& ObjectInitializer)
+	: Super(ObjectInitializer)
+{
+	FOVAngle = 90.0f;
+	bAutoActivate = true;
+	PrimaryComponentTick.bCanEverTick = true;
+	PrimaryComponentTick.TickGroup = TG_DuringPhysics;
+	// previous behavior was to capture from raw scene color 
+	CaptureSource = SCS_SceneColorHDR;
+	// default to full blend weight..
+	PostProcessBlendWeight = 1.0f;
+}
+
+void USceneCaptureComponent2D::SendRenderTransform_Concurrent()
+{	
+	UpdateContent();
+
+	Super::SendRenderTransform_Concurrent();
+}
+
+void USceneCaptureComponent2D::TickComponent(float DeltaTime, enum ELevelTick TickType, FActorComponentTickFunction *ThisTickFunction)
+{
+	Super::TickComponent(DeltaTime, TickType, ThisTickFunction);
+
+	if (bCaptureEveryFrame)
+	{
+		UpdateComponentToWorld(false);
+	}
+}
+
+static TMultiMap<TWeakObjectPtr<UWorld>, USceneCaptureComponent2D*> SceneCapturesToUpdateMap;
+
+void USceneCaptureComponent2D::UpdateContent()
+{
+	if (World && World->Scene && IsVisible())
+	{
+		// Defer until after updates finish
+		SceneCapturesToUpdateMap.AddUnique(World, this);
+	}	
+}
+
+void USceneCaptureComponent2D::UpdateDeferredCaptures( FSceneInterface* Scene )
+{
+	UWorld* World = Scene->GetWorld();
+	if( World && SceneCapturesToUpdateMap.Num() > 0 )
+	{
+		// Only update the scene captures assoicated with the current scene.
+		// Updating others not associated with the scene would cause invalid data to be rendered into the target
+		TArray<USceneCaptureComponent2D*> SceneCapturesToUpdate;
+		SceneCapturesToUpdateMap.MultiFind( World, SceneCapturesToUpdate );
+		
+		for( USceneCaptureComponent2D* Component : SceneCapturesToUpdate )
+		{
+			Scene->UpdateSceneCaptureContents( Component );
+		}
+		
+		// All scene captures for this world have been updated
+		SceneCapturesToUpdateMap.Remove( World) ;
+	}
+}
+
+#if WITH_EDITOR
+void USceneCaptureComponent2D::PostEditChangeProperty(FPropertyChangedEvent& PropertyChangedEvent)
+{
+	// AActor::PostEditChange will ForceUpdateComponents()
+	Super::PostEditChangeProperty(PropertyChangedEvent);
+
+	UpdateContent();
+}
+#endif // WITH_EDITOR
+
+
+// -----------------------------------------------
+
+
+USceneCaptureComponentCube::USceneCaptureComponentCube(const FObjectInitializer& ObjectInitializer)
+	: Super(ObjectInitializer)
+{
+	bAutoActivate = true;
+	PrimaryComponentTick.bCanEverTick = true;
+	PrimaryComponentTick.TickGroup = TG_DuringPhysics;
+}
+
+void USceneCaptureComponentCube::SendRenderTransform_Concurrent()
+{	
+	UpdateContent();
+
+	Super::SendRenderTransform_Concurrent();
+}
+
+void USceneCaptureComponentCube::TickComponent(float DeltaTime, enum ELevelTick TickType, FActorComponentTickFunction *ThisTickFunction)
+{
+	Super::TickComponent(DeltaTime, TickType, ThisTickFunction);
+
+	if (bCaptureEveryFrame)
+	{
+		UpdateComponentToWorld(false);
+	}
+}
+
+static TMultiMap<TWeakObjectPtr<UWorld>, USceneCaptureComponentCube*> CubedSceneCapturesToUpdateMap;
+
+void USceneCaptureComponentCube::UpdateContent()
+{
+	if (World && World->Scene && IsVisible())
+	{
+		// Defer until after updates finish
+		CubedSceneCapturesToUpdateMap.AddUnique( World, this );
+	}	
+}
+
+void USceneCaptureComponentCube::UpdateDeferredCaptures( FSceneInterface* Scene )
+{
+	UWorld* World = Scene->GetWorld();
+	
+	if( World && CubedSceneCapturesToUpdateMap.Num() > 0 )
+	{
+		// Only update the scene captures assoicated with the current scene.
+		// Updating others not associated with the scene would cause invalid data to be rendered into the target
+		TArray<USceneCaptureComponentCube*> SceneCapturesToUpdate;
+		CubedSceneCapturesToUpdateMap.MultiFind( World, SceneCapturesToUpdate );
+		
+		for( USceneCaptureComponentCube* Component : SceneCapturesToUpdate )
+		{
+			Scene->UpdateSceneCaptureContents( Component );
+		}
+		
+		// All scene captures for this world have been updated
+		CubedSceneCapturesToUpdateMap.Remove( World );
+	}
+
+}
+
+#if WITH_EDITOR
+void USceneCaptureComponentCube::PostEditChangeProperty(FPropertyChangedEvent& PropertyChangedEvent)
+{
+	// AActor::PostEditChange will ForceUpdateComponents()
+	Super::PostEditChangeProperty(PropertyChangedEvent);
+
+	UpdateContent();
+}
+#endif // WITH_EDITOR
+
+/** Returns MeshComp subobject **/
+UStaticMeshComponent* ASceneCapture::GetMeshComp() const { return MeshComp; }
+/** Returns CaptureComponent2D subobject **/
+USceneCaptureComponent2D* ASceneCapture2D::GetCaptureComponent2D() const { return CaptureComponent2D; }
+/** Returns DrawFrustum subobject **/
+UDrawFrustumComponent* ASceneCapture2D::GetDrawFrustum() const { return DrawFrustum; }
+/** Returns CaptureComponentCube subobject **/
+USceneCaptureComponentCube* ASceneCaptureCube::GetCaptureComponentCube() const { return CaptureComponentCube; }
+/** Returns DrawFrustum subobject **/
+UDrawFrustumComponent* ASceneCaptureCube::GetDrawFrustum() const { return DrawFrustum; }
