Engine: Moved actor components into Components folder

diff --git a/Engine/Source/Runtime/Engine/Private/Components/SceneComponent.cpp b/Engine/Source/Runtime/Engine/Private/Components/SceneComponent.cpp
new file mode 100644
--- /dev/null
+++ b/Engine/Source/Runtime/Engine/Private/Components/SceneComponent.cpp
@@ -0,0 +1,2216 @@
+// Copyright 1998-2015 Epic Games, Inc. All Rights Reserved.
+
+/*=============================================================================
+	SceneComponent.cpp
+=============================================================================*/
+
+#include "EnginePrivate.h"
+#include "PhysicsPublic.h"
+#include "MessageLog.h"
+#include "Net/UnrealNetwork.h"
+#include "GameFramework/PhysicsVolume.h"
+#include "ComponentReregisterContext.h"
+
+#define LOCTEXT_NAMESPACE "SceneComponent"
+
+
+DEFINE_LOG_CATEGORY_STATIC(LogSceneComponent, Log, All);
+
+FOverlapInfo::FOverlapInfo(UPrimitiveComponent* InComponent, int32 InBodyIndex)
+	: bFromSweep(false)
+{
+	OverlapInfo.Component = InComponent;
+	OverlapInfo.Item = InBodyIndex;
+}
+
+const FName& USceneComponent::GetDefaultSceneRootVariableName()
+{
+	static FName DefaultSceneRootVariableName = FName(TEXT("DefaultSceneRoot"));
+
+	return DefaultSceneRootVariableName;
+}
+
+USceneComponent::USceneComponent(const FObjectInitializer& ObjectInitializer /*= FObjectInitializer::Get()*/)
+	: Super(ObjectInitializer)
+{
+	Mobility = EComponentMobility::Movable;
+	RelativeScale3D = FVector(1.0f, 1.0f, 1.0f);
+	// default behavior is visible
+	bVisible = true;
+	bAutoActivate = false;
+
+	NetUpdateTransform = false;
+}
+
+void USceneComponent::AddReferencedObjects(UObject* InThis, FReferenceCollector& Collector)
+{
+#if WITH_EDITORONLY_DATA
+	USceneComponent* This = CastChecked<USceneComponent>(InThis);
+	Collector.AddReferencedObject(This->SpriteComponent);
+#endif
+	Super::AddReferencedObjects(InThis, Collector);
+}
+
+FTransform USceneComponent::CalcNewComponentToWorld(const FTransform& NewRelativeTransform, const USceneComponent* Parent) const
+{
+	Parent = Parent ? Parent : AttachParent;
+	if (Parent != NULL)
+	{
+		const FTransform ParentToWorld = Parent->GetSocketTransform(AttachSocketName);
+		FTransform NewCompToWorld = NewRelativeTransform * ParentToWorld;
+
+		if(bAbsoluteLocation)
+		{
+			NewCompToWorld.SetTranslation(NewRelativeTransform.GetTranslation());
+		}
+
+		if(bAbsoluteRotation)
+		{
+			NewCompToWorld.SetRotation(NewRelativeTransform.GetRotation());
+		}
+
+		if(bAbsoluteScale)
+		{
+			NewCompToWorld.SetScale3D(NewRelativeTransform.GetScale3D());
+		}
+
+		return NewCompToWorld;
+	}
+	else
+	{
+		return NewRelativeTransform;
+	}
+}
+
+void USceneComponent::OnUpdateTransform(bool bSkipPhysicsMove)
+{
+}
+
+void USceneComponent::UpdateComponentToWorldWithParent(USceneComponent * Parent, bool bSkipPhysicsMove)
+{
+	// If our parent hasn't been updated before, we'll need walk up our parent attach hierarchy
+	if (Parent && !Parent->bWorldToComponentUpdated)
+	{
+		Parent->UpdateComponentToWorld();
+
+		// Updating the parent may (depending on if we were already attached to parent) result in our being updated, so just return
+		if (bWorldToComponentUpdated)
+		{
+			return;
+		}
+	}
+
+	bWorldToComponentUpdated = true;
+
+	// Calculate the new ComponentToWorld transform
+	const FTransform RelativeTransform(RelativeRotation, RelativeLocation, RelativeScale3D);
+	const FTransform NewTransform = CalcNewComponentToWorld(RelativeTransform, Parent);
+
+	// If transform has changed..
+	if (!ComponentToWorld.Equals(NewTransform, SMALL_NUMBER))
+	{
+		// Update transform
+		ComponentToWorld = NewTransform;
+
+		PropagateTransformUpdate(true, bSkipPhysicsMove);
+	}
+	else
+	{
+		PropagateTransformUpdate(false);
+	}
+}
+
+void USceneComponent::OnRegister()
+{
+	// If we need to perform a call to AttachTo, do that now
+	// At this point scene component still has no any state (rendering, physics),
+	// so this call will just add this component to an AttachChildren array of a the Parent component
+	AttachTo(AttachParent, AttachSocketName);
+	
+	Super::OnRegister();
+
+#if WITH_EDITORONLY_DATA
+	if (bVisualizeComponent && SpriteComponent == nullptr && GetOwner() && !GetWorld()->IsGameWorld() )
+	{
+		// Create a new billboard component to serve as a visualization of the actor until there is another primitive component
+		SpriteComponent = NewObject<UBillboardComponent>(GetOwner(), NAME_None, RF_Transactional | RF_TextExportTransient);
+
+		SpriteComponent->Sprite = LoadObject<UTexture2D>(nullptr, TEXT("/Engine/EditorResources/EmptyActor.EmptyActor"));
+		SpriteComponent->RelativeScale3D = FVector(0.5f, 0.5f, 0.5f);
+		SpriteComponent->Mobility = EComponentMobility::Movable;
+		SpriteComponent->AlwaysLoadOnClient = false;
+		SpriteComponent->AlwaysLoadOnServer = false;
+		SpriteComponent->SpriteInfo.Category = TEXT("Misc");
+		SpriteComponent->SpriteInfo.DisplayName = NSLOCTEXT( "SpriteCategory", "Misc", "Misc" );
+		SpriteComponent->CreationMethod = CreationMethod;
+		SpriteComponent->bIsScreenSizeScaled = true;
+		SpriteComponent->bUseInEditorScaling = true;
+
+		SpriteComponent->AttachTo(this);
+		SpriteComponent->RegisterComponent();
+	}
+#endif
+}
+
+void USceneComponent::UpdateComponentToWorld(bool bSkipPhysicsMove)
+{
+	UpdateComponentToWorldWithParent(AttachParent, bSkipPhysicsMove);
+}
+
+
+void USceneComponent::PropagateTransformUpdate(bool bTransformChanged, bool bSkipPhysicsMove)
+{
+	if (IsDeferringMovementUpdates())
+	{
+		// We are deferring these updates until later.
+		return;
+	}
+
+	if (bTransformChanged)
+	{
+		// Then update bounds
+		UpdateBounds();
+
+		// Always send new transform to physics
+		OnUpdateTransform(bSkipPhysicsMove);
+
+		// Flag render transform as dirty
+		MarkRenderTransformDirty();
+		
+		// Now go and update children
+		UpdateChildTransforms();
+
+		// Refresh navigation
+		UpdateNavigationData();
+	}
+	else
+	{
+		// We update bounds even if transform doesn't change, as shape/mesh etc might have done
+		UpdateBounds();
+
+		// Now go and update children
+		UpdateChildTransforms();
+
+		// Need to flag as dirty so new bounds are sent to render thread
+		MarkRenderTransformDirty();
+	}
+}
+
+
+class FScopedMovementUpdate* USceneComponent::GetCurrentScopedMovement() const
+{
+	if (ScopedMovementStack.Num() > 0)
+	{
+		return ScopedMovementStack.Last();
+	}
+
+	return NULL;
+}
+
+
+bool USceneComponent::IsDeferringMovementUpdates() const
+{
+	if (ScopedMovementStack.Num() > 0)
+	{
+		checkSlow(ScopedMovementStack.Last()->IsDeferringUpdates());
+		return true;
+	}
+	
+	return false;
+}
+
+
+void USceneComponent::BeginScopedMovementUpdate(class FScopedMovementUpdate& ScopedUpdate)
+{
+	checkSlow(IsInGameThread());
+	checkSlow(ScopedUpdate.IsDeferringUpdates());
+	
+	ScopedMovementStack.Push(&ScopedUpdate);
+}
+
+
+void USceneComponent::EndScopedMovementUpdate(class FScopedMovementUpdate& CompletedScope)
+{
+	SCOPE_CYCLE_COUNTER(STAT_EndScopedMovementUpdate);
+	checkSlow(IsInGameThread());
+
+	// Special case when shutting down
+	if (ScopedMovementStack.Num() == 0)
+	{
+		return;
+	}
+
+	// Process top of the stack
+	FScopedMovementUpdate* CurrentScopedUpdate = ScopedMovementStack.Pop(false);
+	check(CurrentScopedUpdate == &CompletedScope);
+
+	if (CurrentScopedUpdate)
+	{
+		checkSlow(CurrentScopedUpdate->IsDeferringUpdates());
+		if (ScopedMovementStack.Num() == 0)
+		{
+			// This was the last item on the stack, time to apply the updates if necessary
+			const bool bMoved = CurrentScopedUpdate->IsTransformDirty();
+			if (bMoved)
+			{
+				PropagateTransformUpdate(true);
+			}
+
+			// We may have moved somewhere and then moved back to the start, we still need to update overlaps if we touched things along the way.
+			if (bMoved || CurrentScopedUpdate->HasPendingOverlaps() || CurrentScopedUpdate->GetOverlapsAtEnd())
+			{
+				UpdateOverlaps(&CurrentScopedUpdate->GetPendingOverlaps(), true, CurrentScopedUpdate->GetOverlapsAtEnd());
+			}
+
+			// Dispatch all deferred blocking hits
+			if (CompletedScope.BlockingHits.Num() > 0)
+			{
+				AActor* const Owner = GetOwner();
+				if (Owner)
+				{
+					// If we have blocking hits, we must be a primitive component.
+					UPrimitiveComponent* PrimitiveThis = CastChecked<UPrimitiveComponent>(this);
+					for (const FHitResult& Hit : CompletedScope.BlockingHits)
+					{
+						PrimitiveThis->DispatchBlockingHit(*Owner, Hit);
+					}
+				}
+			}
+		}
+		else
+		{
+			// Combine with next item on the stack
+			FScopedMovementUpdate* OuterScopedUpdate = ScopedMovementStack.Last();
+			OuterScopedUpdate->OnInnerScopeComplete(*CurrentScopedUpdate);
+		}
+	}
+}
+
+
+void USceneComponent::DestroyComponent(bool bPromoteChildren/*= false*/)
+{
+	if (bPromoteChildren)
+	{
+		AActor* Owner = GetOwner();
+		if (Owner != NULL)
+		{
+			Owner->Modify();
+			USceneComponent* ChildToPromote = nullptr;
+
+			// Handle removal of the root node
+			if (this == Owner->GetRootComponent())
+			{
+				// Always choose non editor-only child nodes over editor-only child nodes (since we don't want editor-only nodes to end up with non editor-only child nodes)
+				// Exclude scene components owned by attached child actors
+				USceneComponent** FindResult =
+					AttachChildren.FindByPredicate([Owner](USceneComponent* Child){ return Child != nullptr && !Child->IsEditorOnly() && Child->GetOwner() == Owner; });
+
+				if (FindResult != nullptr)
+				{
+					ChildToPromote = *FindResult;
+				}
+				else
+				{
+					// Didn't find a suitable component to promote so create a new default component
+
+					Rename(NULL, GetOuter(), REN_DoNotDirty | REN_DontCreateRedirectors);
+
+					// Construct a new default root component
+					auto NewRootComponent = NewObject<USceneComponent>(Owner, USceneComponent::GetDefaultSceneRootVariableName(), RF_Transactional);
+					NewRootComponent->Mobility = Mobility;
+					NewRootComponent->SetWorldLocationAndRotation(GetComponentLocation(), GetComponentRotation());
+#if WITH_EDITORONLY_DATA
+					NewRootComponent->bVisualizeComponent = true;
+#endif
+					Owner->AddInstanceComponent(NewRootComponent);
+					NewRootComponent->RegisterComponent();
+
+					// Designate the new default root as the child we're promoting
+					ChildToPromote = NewRootComponent;
+				}
+
+				Owner->Modify();
+
+				// Set the selected child node as the new root
+				check(ChildToPromote != nullptr);
+				Owner->SetRootComponent(ChildToPromote);
+			}
+			else    // ...not the root node, so we'll promote the selected child node to this position in its AttachParent's child array.
+			{
+				// Cache our AttachParent
+				USceneComponent* CachedAttachParent = AttachParent;
+				check(CachedAttachParent != nullptr);
+
+				// Find the our position in its AttachParent's child array
+				int32 Index = CachedAttachParent->AttachChildren.Find(this);
+				check(Index != INDEX_NONE);
+
+				// Detach from parent
+				DetachFromParent(true);
+
+				// Find an appropriate child node to promote to this node's position in the hierarchy
+				if (AttachChildren.Num() > 0)
+				{
+					// Always choose non editor-only child nodes over editor-only child nodes (since we don't want editor-only nodes to end up with non editor-only child nodes)
+					USceneComponent** FindResult =
+						AttachChildren.FindByPredicate([Owner](USceneComponent* Child){ return Child != nullptr && !Child->IsEditorOnly(); });
+
+					if (FindResult != nullptr)
+					{
+						ChildToPromote = *FindResult;
+					}
+					else
+					{
+						// Default to first child node
+						check(AttachChildren[0] != nullptr);
+						ChildToPromote = AttachChildren[0];
+					}
+				}
+
+				if (ChildToPromote != nullptr)
+				{
+					// Attach the child node that we're promoting to the parent and move it to the same position as the old node was in the array
+					ChildToPromote->AttachTo(CachedAttachParent, NAME_None, EAttachLocation::KeepWorldPosition);
+					CachedAttachParent->AttachChildren.Remove(ChildToPromote);
+
+					Index = FMath::Clamp<int32>(Index, 0, CachedAttachParent->AttachChildren.Num());
+					CachedAttachParent->AttachChildren.Insert(ChildToPromote, Index);
+				}
+			}
+
+			// Detach child nodes from the node that's being removed and re-attach them to the child that's being promoted
+			TArray<USceneComponent*> AttachChildrenLocalCopy(AttachChildren);
+			for (auto ChildCompIt = AttachChildrenLocalCopy.CreateIterator(); ChildCompIt; ++ChildCompIt)
+			{
+				USceneComponent* Child = *ChildCompIt;
+				check(Child != nullptr);
+
+				// Note: This will internally call Modify(), so we don't need to call it here
+				Child->DetachFromParent(true);
+				if (Child != ChildToPromote)
+				{
+					Child->AttachTo(ChildToPromote, NAME_None, EAttachLocation::KeepWorldPosition);
+				}
+			}
+		}
+	}
+	Super::DestroyComponent(bPromoteChildren);
+}
+
+void USceneComponent::OnComponentDestroyed()
+{
+	Super::OnComponentDestroyed();
+
+#if WITH_EDITORONLY_DATA
+	if (SpriteComponent)
+	{
+		SpriteComponent->DestroyComponent();
+	}
+#endif
+
+	ScopedMovementStack.Reset();
+
+	// Ensure we are detached before destroying
+	DetachFromParent();
+}
+
+FBoxSphereBounds USceneComponent::CalcBounds(const FTransform& LocalToWorld) const
+{
+	FBoxSphereBounds NewBounds;
+	NewBounds.Origin = LocalToWorld.GetLocation();
+	NewBounds.BoxExtent = FVector::ZeroVector;
+	NewBounds.SphereRadius = 0.f;
+	return NewBounds;
+}
+
+void USceneComponent::CalcBoundingCylinder(float& CylinderRadius, float& CylinderHalfHeight) const
+{
+	CylinderRadius = FMath::Sqrt( FMath::Square(Bounds.BoxExtent.X) + FMath::Square(Bounds.BoxExtent.Y) );
+	CylinderHalfHeight = Bounds.BoxExtent.Z;
+}
+
+void USceneComponent::UpdateBounds()
+{
+#if WITH_EDITOR
+	FBoxSphereBounds OriginalBounds = Bounds; // Save old bounds
+#endif
+
+	// if use parent bound if attach parent exists, and the flag is set
+	// since parents tick first before child, this should work correctly
+	if ( bUseAttachParentBound && AttachParent != NULL )
+	{
+		Bounds = AttachParent->Bounds;
+	}
+	else
+	{
+		// Calculate new bounds
+		Bounds = CalcBounds(ComponentToWorld);
+	}
+
+#if WITH_EDITOR
+	// If bounds have changed (in editor), trigger data rebuild
+	if ( IsRegistered() && (World != NULL) && !World->IsGameWorld() &&
+		(OriginalBounds.Origin.Equals(Bounds.Origin) == false || OriginalBounds.BoxExtent.Equals(Bounds.BoxExtent) == false) )
+	{
+		GEngine->TriggerStreamingDataRebuild();
+	}
+#endif
+}
+
+void USceneComponent::SetRelativeLocation(FVector NewLocation, bool bSweep, FHitResult* OutSweepHitResult)
+{
+	SetRelativeLocationAndRotation(NewLocation, RelativeRotation, bSweep, OutSweepHitResult);
+}
+
+void USceneComponent::SetRelativeRotation(FRotator NewRotation, bool bSweep, FHitResult* OutSweepHitResult)
+{
+	SetRelativeLocationAndRotation(RelativeLocation, NewRotation, bSweep, OutSweepHitResult);
+}
+
+void USceneComponent::SetRelativeLocationAndRotation(FVector NewLocation, FRotator NewRotation, bool bSweep, FHitResult* OutSweepHitResult)
+{
+	if(!NewLocation.Equals(RelativeLocation) || !NewRotation.Equals(RelativeRotation))
+	{
+		if (!bWorldToComponentUpdated)
+		{
+			UpdateComponentToWorld();
+		}
+
+		const FTransform DesiredRelTransform(NewRotation, NewLocation);
+		const FTransform DesiredWorldTransform = CalcNewComponentToWorld(DesiredRelTransform);
+		const FVector DesiredDelta = DesiredWorldTransform.GetTranslation() - ComponentToWorld.GetTranslation();
+		
+		MoveComponent(DesiredDelta, DesiredWorldTransform.Rotator(), bSweep, OutSweepHitResult);
+	}
+	else if (OutSweepHitResult)
+	{
+		*OutSweepHitResult = FHitResult();
+	}
+}
+
+void USceneComponent::AddRelativeLocation(FVector DeltaLocation, bool bSweep, FHitResult* OutSweepHitResult)
+{
+	SetRelativeLocationAndRotation(RelativeLocation + DeltaLocation, RelativeRotation, bSweep, OutSweepHitResult);
+}
+
+void USceneComponent::AddRelativeRotation(FRotator DeltaRotation, bool bSweep, FHitResult* OutSweepHitResult)
+{
+	SetRelativeLocationAndRotation(RelativeLocation, RelativeRotation + DeltaRotation, bSweep, OutSweepHitResult);
+}
+
+void USceneComponent::AddLocalOffset(FVector DeltaLocation, bool bSweep, FHitResult* OutSweepHitResult)
+{
+	if (!DeltaLocation.IsNearlyZero())
+	{
+		const FQuat RelativeRotQuat = RelativeRotation.Quaternion();
+		const FVector LocalOffset = RelativeRotQuat.RotateVector(DeltaLocation);
+
+		SetRelativeLocationAndRotation(RelativeLocation + LocalOffset, RelativeRotation, bSweep, OutSweepHitResult);
+	}
+	else if (OutSweepHitResult)
+	{
+		*OutSweepHitResult = FHitResult();
+	}
+}
+
+void USceneComponent::AddLocalRotation(FRotator DeltaRotation, bool bSweep, FHitResult* OutSweepHitResult)
+{
+	if (!DeltaRotation.IsZero())
+	{
+		const FQuat RelativeRotQuat = RelativeRotation.Quaternion();
+		const FQuat DeltaRotQuat = DeltaRotation.Quaternion();
+		const FQuat NewRelRotQuat = RelativeRotQuat * DeltaRotQuat;
+
+		SetRelativeLocationAndRotation(RelativeLocation, NewRelRotQuat.Rotator(), bSweep, OutSweepHitResult);
+	}
+	else if (OutSweepHitResult)
+	{
+		*OutSweepHitResult = FHitResult();
+	}
+}
+
+void USceneComponent::AddLocalTransform(const FTransform& DeltaTransform, bool bSweep, FHitResult* OutSweepHitResult)
+{
+	const FTransform RelativeTransform( RelativeRotation, RelativeLocation, FVector(1,1,1) ); // don't use scaling, so it matches how AddLocalRotation/Offset work
+	const FTransform NewRelTransform = DeltaTransform * RelativeTransform;
+
+	SetRelativeTransform(NewRelTransform, bSweep, OutSweepHitResult);
+}
+
+void USceneComponent::AddWorldOffset(FVector DeltaLocation, bool bSweep, FHitResult* OutSweepHitResult)
+{
+	if (!DeltaLocation.IsZero())
+	{
+		const FVector NewWorldLocation = DeltaLocation + ComponentToWorld.GetTranslation();
+		SetWorldLocation(NewWorldLocation, bSweep, OutSweepHitResult);
+	}
+	else if (OutSweepHitResult)
+	{
+		*OutSweepHitResult = FHitResult();
+	}
+}
+
+void USceneComponent::AddWorldRotation(FRotator DeltaRotation, bool bSweep, FHitResult* OutSweepHitResult)
+{
+	if (!DeltaRotation.IsZero())
+	{
+		const FQuat NewWorldRotation = DeltaRotation.Quaternion() * ComponentToWorld.GetRotation();
+		SetWorldRotation(NewWorldRotation.Rotator(), bSweep, OutSweepHitResult);
+	}
+	else if (OutSweepHitResult)
+	{
+		*OutSweepHitResult = FHitResult();
+	}
+}
+
+void USceneComponent::AddWorldTransform(const FTransform& DeltaTransform, bool bSweep, FHitResult* OutSweepHitResult)
+{
+	const FQuat NewWorldRotation = DeltaTransform.GetRotation() * ComponentToWorld.GetRotation();
+	const FVector NewWorldLocation = DeltaTransform.GetTranslation() + ComponentToWorld.GetTranslation();
+
+	SetWorldTransform(FTransform(NewWorldRotation, NewWorldLocation, FVector(1,1,1)));
+}
+
+void USceneComponent::SetRelativeScale3D(FVector NewScale3D)
+{
+	if( NewScale3D != RelativeScale3D )
+	{
+		if (NewScale3D.ContainsNaN())
+		{
+			UE_LOG(LogBlueprint, Warning, TEXT("SetRelativeScale3D : Invalid Scale entered (%s). Resetting to 1.f."), *NewScale3D.ToString());
+			NewScale3D = FVector(1.f);
+		}
+		RelativeScale3D = NewScale3D;		
+		UpdateComponentToWorld();
+		
+		if (IsRegistered())
+		{
+			UpdateOverlaps();
+		}
+	}
+}
+
+void USceneComponent::ResetRelativeTransform()
+{
+	SetRelativeLocationAndRotation(FVector::ZeroVector, FRotator::ZeroRotator);
+	SetRelativeScale3D(FVector(1.f));
+}
+
+void USceneComponent::SetRelativeTransform(const FTransform& NewTransform, bool bSweep, FHitResult* OutSweepHitResult)
+{
+	SetRelativeLocationAndRotation(NewTransform.GetTranslation(), NewTransform.Rotator(), bSweep, OutSweepHitResult);
+	SetRelativeScale3D(NewTransform.GetScale3D());
+}
+
+FTransform USceneComponent::GetRelativeTransform() const
+{
+	const FTransform RelativeTransform( RelativeRotation, RelativeLocation, RelativeScale3D );
+	return RelativeTransform;
+}
+
+void USceneComponent::SetWorldLocation(FVector NewLocation, bool bSweep, FHitResult* OutSweepHitResult)
+{
+	FVector NewRelLocation = NewLocation;
+
+	// If attached to something, transform into local space
+	if(AttachParent != NULL && !bAbsoluteLocation)
+	{
+		FTransform ParentToWorld = AttachParent->GetSocketTransform(AttachSocketName);
+		NewRelLocation = ParentToWorld.InverseTransformPosition(NewLocation);
+	}
+
+	SetRelativeLocation(NewRelLocation, bSweep, OutSweepHitResult);
+}
+
+void USceneComponent::SetWorldRotation(FRotator NewRotation, bool bSweep, FHitResult* OutSweepHitResult)
+{
+	FRotator NewRelRotation = NewRotation;
+
+	// If already attached to something, transform into local space
+	if(AttachParent != NULL && !bAbsoluteRotation)
+	{
+		FTransform ParentToWorld = AttachParent->GetSocketTransform(AttachSocketName);
+		// Quat multiplication works reverse way, make sure you do Parent(-1) * World = Local, not World*Parent(-) = Local (the way matrix does)
+		FQuat NewRelQuat = ParentToWorld.GetRotation().Inverse() * NewRotation.Quaternion();
+		NewRelRotation = NewRelQuat.Rotator();
+	}
+
+	SetRelativeRotation(NewRelRotation, bSweep, OutSweepHitResult);
+}
+
+void USceneComponent::SetWorldScale3D(FVector NewScale)
+{
+	FVector NewRelScale = NewScale;
+
+	// If attached to something, transform into local space
+	if(AttachParent != NULL && !bAbsoluteScale)
+	{
+		FTransform ParentToWorld = AttachParent->GetSocketTransform(AttachSocketName);
+		NewRelScale = NewScale * ParentToWorld.GetSafeScaleReciprocal(ParentToWorld.GetScale3D());
+	}
+
+	SetRelativeScale3D(NewRelScale);
+}
+
+void USceneComponent::SetWorldTransform(const FTransform& NewTransform, bool bSweep, FHitResult* OutSweepHitResult)
+{
+	// If attached to something, transform into local space
+	FQuat NewRotation = NewTransform.GetRotation();
+	FVector NewLocation = NewTransform.GetTranslation();
+	FVector NewScale = NewTransform.GetScale3D();
+
+	if (AttachParent != NULL)
+	{
+		FTransform ParentToWorld = AttachParent->GetSocketTransform(AttachSocketName);
+		FTransform RelativeTM = NewTransform.GetRelativeTransform(ParentToWorld);
+		if (!bAbsoluteLocation)
+		{
+			NewLocation = RelativeTM.GetTranslation();
+		}
+
+		if (!bAbsoluteRotation)
+		{
+			// Quat multiplication works reverse way, make sure you do Parent(-1) * World = Local, not World*Parent(-) = Local (the way matrix does)
+			NewRotation = RelativeTM.GetRotation();
+		}
+
+		if (!bAbsoluteScale)
+		{
+			NewScale = RelativeTM.GetScale3D();
+		}
+	}
+
+	SetRelativeTransform(FTransform(NewRotation, NewLocation, NewScale), bSweep, OutSweepHitResult);
+}
+
+void USceneComponent::SetWorldLocationAndRotation(FVector NewLocation, FRotator NewRotation, bool bSweep, FHitResult* OutSweepHitResult)
+{
+	SetWorldLocationAndRotation(NewLocation, NewRotation.Quaternion(), bSweep, OutSweepHitResult);
+}
+
+void USceneComponent::SetWorldLocationAndRotation(FVector NewLocation, const FQuat& NewRotation, bool bSweep, FHitResult* OutSweepHitResult)
+{
+	// If attached to something, transform into local space
+	FQuat NewFinalRotation = NewRotation;
+	if (AttachParent != NULL)
+	{
+		FTransform ParentToWorld = AttachParent->GetSocketTransform(AttachSocketName);
+
+		if (!bAbsoluteLocation)
+		{
+			NewLocation = ParentToWorld.InverseTransformPosition(NewLocation);
+		}
+
+		if (!bAbsoluteRotation)
+		{
+			// Quat multiplication works reverse way, make sure you do Parent(-1) * World = Local, not World*Parent(-) = Local (the way matrix does)
+			FQuat NewRelQuat = ParentToWorld.GetRotation().Inverse() * NewRotation;
+			NewFinalRotation = NewRelQuat;
+		}
+	}
+
+	SetRelativeLocationAndRotation(NewLocation, NewFinalRotation.Rotator(), bSweep, OutSweepHitResult);
+}
+
+void USceneComponent::SetWorldLocationAndRotationNoPhysics(const FVector& NewLocation, const FRotator& NewRotation)
+{
+	// If attached to something, transform into local space
+	FQuat NewFinalRotation = NewRotation.Quaternion();
+	if (AttachParent != NULL)
+	{
+		const FTransform ParentToWorld = AttachParent->GetSocketTransform(AttachSocketName);
+
+		if(bAbsoluteLocation)
+		{
+			RelativeLocation = NewLocation;
+		}
+		else
+		{
+			RelativeLocation = ParentToWorld.InverseTransformPosition(NewLocation);
+		}
+
+		if(bAbsoluteRotation)
+		{
+			RelativeRotation = NewRotation;
+		}
+		else
+		{
+			// Quat multiplication works reverse way, make sure you do Parent(-1) * World = Local, not World*Parent(-) = Local (the way matrix does)
+			const FQuat NewRelQuat = ParentToWorld.GetRotation().Inverse() * NewRotation.Quaternion();
+			RelativeRotation = NewRelQuat.Rotator();
+		}
+	}
+	else
+	{
+		RelativeLocation = NewLocation;
+		RelativeRotation = NewRotation;
+	}
+
+	UpdateComponentToWorld(true);
+}
+
+void USceneComponent::SetAbsolute(bool bNewAbsoluteLocation, bool bNewAbsoluteRotation, bool bNewAbsoluteScale)
+{
+	bAbsoluteLocation = bNewAbsoluteLocation;
+	bAbsoluteRotation = bNewAbsoluteRotation;
+	bAbsoluteScale = bNewAbsoluteScale;
+
+	UpdateComponentToWorld();
+}
+
+FTransform USceneComponent::K2_GetComponentToWorld() const
+{
+	return GetComponentToWorld();
+}
+
+FVector USceneComponent::GetForwardVector() const
+{
+	return ComponentToWorld.GetUnitAxis( EAxis::X );
+}
+
+FVector USceneComponent::GetRightVector() const
+{
+	return ComponentToWorld.GetUnitAxis( EAxis::Y );
+}
+
+FVector USceneComponent::GetUpVector() const
+{
+	return ComponentToWorld.GetUnitAxis( EAxis::Z );
+}
+
+FVector USceneComponent::K2_GetComponentLocation() const
+{
+	return GetComponentLocation();
+}
+
+FRotator USceneComponent::K2_GetComponentRotation() const
+{
+	return GetComponentRotation();
+}
+
+FVector USceneComponent::K2_GetComponentScale() const
+{
+	return GetComponentScale();
+}
+
+
+USceneComponent* USceneComponent::GetAttachParent() const
+{
+	return AttachParent;
+}
+
+void USceneComponent::GetParentComponents(TArray<class USceneComponent*>& Parents) const
+{
+	Parents.Empty();
+
+	USceneComponent* ParentIterator = AttachParent;
+	while (ParentIterator != NULL)
+	{
+		Parents.Add(ParentIterator);
+		ParentIterator = ParentIterator->AttachParent;
+	}
+}
+
+int32 USceneComponent::GetNumChildrenComponents() const
+{
+	return AttachChildren.Num();
+}
+
+USceneComponent* USceneComponent::GetChildComponent(int32 ChildIndex) const
+{
+	if (ChildIndex < 0)
+	{
+		UE_LOG(LogBlueprint, Log, TEXT("SceneComponent::GetChild called with a negative ChildIndex: %d"), ChildIndex);
+		return NULL;
+	}
+	else if (ChildIndex >= AttachChildren.Num())
+	{
+		UE_LOG(LogBlueprint, Log, TEXT("SceneComponent::GetChild called with an out of range ChildIndex: %d; Number of children is %d."), ChildIndex, AttachChildren.Num());
+		return NULL;
+	}
+
+	return AttachChildren[ChildIndex];
+}
+
+void USceneComponent::GetChildrenComponents(bool bIncludeAllDescendants, TArray<USceneComponent*>& Children) const
+{
+	Children.Empty();
+
+	if (bIncludeAllDescendants)
+	{
+		AppendDescendants(Children);
+	}
+	else
+	{
+		Children.Append(AttachChildren);
+	}
+}
+
+void USceneComponent::AppendDescendants(TArray<USceneComponent*>& Children) const
+{
+	Children.Append(AttachChildren);
+
+	for (USceneComponent* Child : AttachChildren)
+	{
+		if (Child)
+		{
+			Child->AppendDescendants(Children);
+		}
+	}
+}
+
+//This function is used for giving AttachTo different bWeldSimulatedBodies default, but only when called from BP
+void USceneComponent::K2_AttachTo(class USceneComponent* InParent, FName InSocketName, EAttachLocation::Type AttachLocationType, bool bWeldSimulatedBodies /*= true*/)
+{
+	AttachTo(InParent, InSocketName, AttachLocationType, bWeldSimulatedBodies);
+}
+
+void USceneComponent::AttachTo(class USceneComponent* Parent, FName InSocketName, EAttachLocation::Type AttachType /*= EAttachLocation::KeepRelativeOffset */, bool bWeldSimulatedBodies /*= false*/)
+{
+	if(Parent != NULL)
+	{
+		if (Parent == AttachParent && InSocketName == AttachSocketName && Parent->AttachChildren.Contains(this))
+		{
+			// already attached!
+			return;
+		}
+
+		if(Parent == this)
+		{
+			FMessageLog("PIE").Warning(FText::Format(LOCTEXT("AttachToSelfWarning", "AttachTo: '{0}' cannot be attached to itself. Aborting."), 
+				FText::FromString(GetPathName())));
+			return;
+		}
+
+		if(Parent->IsAttachedTo(this))
+		{
+			FMessageLog("PIE").Warning(FText::Format(LOCTEXT("AttachCycleWarning", "AttachTo: '{0}' already attached to '{1}', would form cycle. Aborting."), 
+				FText::FromString(Parent->GetPathName()), 
+				FText::FromString(GetPathName())));
+			return;
+		}
+
+		if(!Parent->CanAttachAsChild(this, InSocketName))
+		{
+			UE_LOG(LogSceneComponent, Warning, TEXT("AttachTo: '%s' will not allow '%s' to be attached as a child."), *Parent->GetPathName(), *GetPathName());
+			return;
+		}
+
+		// Don't allow components with static mobility to be attached to non-static parents (except during UCS)
+		if(!IsOwnerRunningUserConstructionScript() && Mobility == EComponentMobility::Static && Parent->Mobility != EComponentMobility::Static)
+		{
+			FString ExtraBlueprintInfo;
+#if WITH_EDITORONLY_DATA
+			UClass* ParentClass = Parent->GetOuter()->GetClass();
+			if (ParentClass->ClassGeneratedBy && ParentClass->ClassGeneratedBy->IsA(UBlueprint::StaticClass())) 
+			{ 
+				ExtraBlueprintInfo = FString::Printf(TEXT(" (in blueprint \"%s\")"), *ParentClass->ClassGeneratedBy->GetName());  
+			}
+#endif //WITH_EDITORONLY_DATA
+			FMessageLog("PIE").Warning(FText::Format(LOCTEXT("NoStaticToDynamicWarning", "AttachTo: '{0}' is not static {1}, cannot attach '{2}' which is static to it. Aborting."), 
+				FText::FromString(Parent->GetPathName()), 
+				FText::FromString(ExtraBlueprintInfo), 
+				FText::FromString(GetPathName())));
+			return;
+		}
+
+		// if our template type doesn't match
+		if (Parent->IsTemplate() != IsTemplate())
+		{
+			if (Parent->IsTemplate())
+			{
+				ensureMsgf(false, TEXT("Template Mismatch during attachment. Attaching instanced component to template component. Parent '%s' Self '%s'"), *Parent->GetName(), *GetName());
+			}
+			else
+			{
+				ensureMsgf(false, TEXT("Template Mismatch during attachment. Attaching template component to instanced component. Parent '%s' Self '%s'"), *Parent->GetName(), *GetName());
+			}
+			return;
+		}
+
+		// Don't call UpdateOverlaps() when detaching, since we are going to do it anyway after we reattach below.
+		// Aside from a perf benefit this also maintains correct behavior when we don't have KeepWorldPosition set.
+		const bool bSavedDisableDetachmentUpdateOverlaps = bDisableDetachmentUpdateOverlaps;
+		bDisableDetachmentUpdateOverlaps = true;
+
+		// Find out if we're already attached, and save off our position in the array if we are
+		int32 LastAttachIndex = INDEX_NONE;
+		Parent->AttachChildren.Find(this, LastAttachIndex);
+
+		// Make sure we are detached
+		const bool bMaintainWorldPosition = (AttachType == EAttachLocation::KeepWorldPosition);
+		DetachFromParent(bMaintainWorldPosition);
+		
+		// Restore detachment update overlaps flag.
+		bDisableDetachmentUpdateOverlaps = bSavedDisableDetachmentUpdateOverlaps;
+
+		{
+			//This code requires some explaining. Inside the editor we allow user to attach physically simulated objects to other objects. This is done for convenience so that users can group things together in hierarchy.
+			//At runtime we must not attach physically simulated objects as it will cause double transform updates, and you should just use a physical constraint if attachment is the desired behavior.
+			//Note if bWeldSimulatedBodies = true then they actually want to keep these objects simulating together
+			//We must fixup the relative location,rotation,scale as the attachment is no longer valid. Blueprint uses simple construction to try and attach before ComponentToWorld has ever been updated, so we cannot rely on it.
+			//As such we must calculate the proper Relative information
+			//Also physics state may not be created yet so we use bSimulatePhysics to determine if the object has any intention of being physically simulated
+			UPrimitiveComponent * PrimitiveComponent = Cast<UPrimitiveComponent>(this);
+
+			if (PrimitiveComponent && PrimitiveComponent->BodyInstance.bSimulatePhysics && !bWeldSimulatedBodies && GetWorld() && GetWorld()->IsGameWorld() && !GetWorld()->bIsRunningConstructionScript)
+			{
+				//Since the object is physically simulated it can't be the case that it's a child of object A and being attached to object B (at runtime)
+				if (bMaintainWorldPosition == false)	//User tried to attach but physically based so detach. However, if they provided relative coordinates we should still get the correct position
+				{
+					UpdateComponentToWorldWithParent(Parent, false);
+					RelativeLocation = ComponentToWorld.GetLocation();
+					RelativeRotation = ComponentToWorld.GetRotation().Rotator();
+					RelativeScale3D = ComponentToWorld.GetScale3D();
+				}
+
+				return;
+			}
+		}
+
+		// Detach removes all Prerequisite, so will need to add after Detach happens
+		PrimaryComponentTick.AddPrerequisite(Parent, Parent->PrimaryComponentTick); // force us to tick after the parent does
+
+		// Save pointer from child to parent
+		AttachParent = Parent;
+		AttachSocketName = InSocketName;
+
+		OnAttachmentChanged();
+
+		// Preserve order of previous attachment if valid (in case we're doing a reattach operation inside a loop that might assume the AttachChildren order won't change)
+		if(LastAttachIndex != INDEX_NONE)
+		{
+			Parent->AttachChildren.Insert(this, LastAttachIndex);
+		}
+		else
+		{
+			Parent->AttachChildren.Add(this);
+		}
+
+		switch ( AttachType )
+		{
+		case EAttachLocation::KeepWorldPosition:
+			{
+				// Update RelativeLocation and RelativeRotation to maintain current world position after attachment
+				FTransform ParentToWorld = AttachParent->GetSocketTransform(AttachSocketName);
+				FTransform RelativeTM = ComponentToWorld.GetRelativeTransform(ParentToWorld);
+
+				if(bAbsoluteLocation)
+				{
+					RelativeLocation = ComponentToWorld.GetTranslation();
+				}
+				else
+				{
+					RelativeLocation = RelativeTM.GetTranslation();
+				}
+
+				if(bAbsoluteRotation)
+				{
+					RelativeRotation = ComponentToWorld.Rotator();
+				}
+				else
+				{
+					RelativeRotation = RelativeTM.Rotator();
+				}
+				if(bAbsoluteScale)
+				{
+					RelativeScale3D = ComponentToWorld.GetScale3D();
+				}
+				else
+				{
+					RelativeScale3D = RelativeTM.GetScale3D();
+				}
+			}
+			break;
+
+		case EAttachLocation::SnapToTarget:
+			// Zero out relative location and rotation so we snap to the parent's transform
+			{
+				RelativeLocation = FVector::ZeroVector;
+				RelativeRotation = FRotator::ZeroRotator;
+				if(bAbsoluteScale)
+				{
+					RelativeScale3D = ComponentToWorld.GetScale3D();
+				}
+				else
+				{
+					// when snap, we'd like to give socket or bone scale only
+					// to do so, get parent and get socket relative to parent
+					FTransform ParentToWorld = AttachParent->GetComponentToWorld();
+					FTransform SocketTransform = AttachParent->GetSocketTransform(AttachSocketName);
+					FTransform RelativeTM = SocketTransform.GetRelativeTransform(ParentToWorld);
+					RelativeScale3D = RelativeTM.GetScale3D();
+				}
+			}
+			break;
+
+		default:
+			// Leave the transform alone (relative offset to parent stays the same)
+			break;
+		}
+
+#if WITH_EDITOR
+		if(GEngine)
+		{
+			if(GetOwner() && this == GetOwner()->GetRootComponent())
+			{
+				GEngine->BroadcastLevelActorAttached(GetOwner(), AttachParent->GetOwner());
+			}
+		}
+#endif
+
+		AttachParent->OnChildAttached(this);
+
+		// calculate transform with new attachment condition
+		UpdateComponentToWorld();
+
+		if (UPrimitiveComponent * PrimitiveComponent = Cast<UPrimitiveComponent>(this))
+		{
+			if (FBodyInstance* BI = PrimitiveComponent->GetBodyInstance())
+			{
+				if (bWeldSimulatedBodies)
+				{
+					PrimitiveComponent->WeldToImplementation(AttachParent, AttachSocketName, bWeldSimulatedBodies);
+				}
+			}
+		}
+
+		// Update overlaps, in case location changed or overlap state depends on attachment.
+		if (IsRegistered())
+		{
+			UpdateOverlaps();
+		}
+	}
+}
+
+void USceneComponent::SnapTo(class USceneComponent* Parent, FName InSocketName)
+{
+	AttachTo(Parent, InSocketName, EAttachLocation::SnapToTarget);
+}
+
+void USceneComponent::DetachFromParent(bool bMaintainWorldPosition, bool bCallModify)
+{
+	if(AttachParent != NULL)
+	{
+		if (UPrimitiveComponent * PrimComp = Cast<UPrimitiveComponent>(this))
+		{
+			PrimComp->UnWeldFromParent();
+		}
+
+		// Make sure parent points to us if we're registered
+		checkf(!bRegistered || AttachParent->AttachChildren.Contains(this), TEXT("Attempt to detach SceneComponent '%s' owned by '%s' from AttachParent '%s' while not attached."), *GetName(), (GetOwner() ? *GetOwner()->GetName() : TEXT("Unowned")), *AttachParent->GetName());
+
+		if (bCallModify)
+		{
+			Modify();
+			AttachParent->Modify();
+		}
+
+		PrimaryComponentTick.RemovePrerequisite(AttachParent, AttachParent->PrimaryComponentTick); // no longer required to tick after the attachment
+
+		AttachParent->AttachChildren.Remove(this);
+		AttachParent->OnChildDetached(this);
+
+#if WITH_EDITOR
+		if(GEngine)
+		{
+			if(GetOwner() && this == GetOwner()->GetRootComponent())
+			{
+				GEngine->BroadcastLevelActorDetached(GetOwner(), AttachParent->GetOwner());
+			}
+		}
+#endif
+		AttachParent = NULL;
+		AttachSocketName = NAME_None;
+
+		OnAttachmentChanged();
+
+		// If desired, update RelativeLocation and RelativeRotation to maintain current world position after detachment
+		if(bMaintainWorldPosition)
+		{
+			RelativeLocation = ComponentToWorld.GetTranslation();
+			RelativeRotation = ComponentToWorld.Rotator();
+			RelativeScale3D = ComponentToWorld.GetScale3D();
+		}
+
+		// calculate transform with new attachment condition
+		UpdateComponentToWorld();
+
+		// Update overlaps, in case location changed or overlap state depends on attachment.
+		if (IsRegistered() && !bDisableDetachmentUpdateOverlaps)
+		{
+			UpdateOverlaps();
+		}
+	}
+}
+
+USceneComponent* USceneComponent::GetAttachmentRoot() const
+{
+	const USceneComponent* Top;
+	for( Top=this; Top && Top->AttachParent; Top=Top->AttachParent );
+	return const_cast<USceneComponent*>(Top);
+}
+
+AActor* USceneComponent::GetAttachmentRootActor() const
+{
+	const USceneComponent* const AttachmentRootComponent = GetAttachmentRoot();
+	return AttachmentRootComponent ? AttachmentRootComponent->GetOwner() : NULL;
+}
+
+bool USceneComponent::IsAttachedTo(class USceneComponent* TestComp) const
+{
+	if(TestComp != NULL)
+	{
+		for( const USceneComponent* Comp=this->AttachParent; Comp!=NULL; Comp=Comp->AttachParent )
+		{
+			if( TestComp == Comp )
+			{
+				return true;
+			}
+		}
+	}
+	return false;
+}
+
+FSceneComponentInstanceData::FSceneComponentInstanceData(const USceneComponent* SourceComponent)
+	: FActorComponentInstanceData(SourceComponent)
+{
+	AActor* SourceOwner = SourceComponent->GetOwner();
+	for (int32 i = SourceComponent->AttachChildren.Num()-1; i >= 0; --i)
+	{
+		USceneComponent* SceneComponent = SourceComponent->AttachChildren[i];
+		if (SceneComponent && SceneComponent->GetOwner() == SourceOwner && !SceneComponent->IsCreatedByConstructionScript())
+		{
+			AttachedInstanceComponents.Add(SceneComponent);
+		}
+	}
+}
+
+void FSceneComponentInstanceData::ApplyToComponent(UActorComponent* Component, const ECacheApplyPhase CacheApplyPhase)
+{
+	FActorComponentInstanceData::ApplyToComponent(Component, CacheApplyPhase);
+
+	USceneComponent* SceneComponent = CastChecked<USceneComponent>(Component);
+
+	if (ContainsSavedProperties())
+	{
+		SceneComponent->UpdateComponentToWorld();
+	}
+
+	for (USceneComponent* ChildComponent : AttachedInstanceComponents)
+	{
+		if (ChildComponent)
+		{
+			ChildComponent->AttachTo(SceneComponent);
+		}
+	}
+}
+
+void FSceneComponentInstanceData::FindAndReplaceInstances(const TMap<UObject*, UObject*>& OldToNewInstanceMap)
+{
+	for (USceneComponent*& ChildComponent : AttachedInstanceComponents)
+	{
+		if (UObject* const* NewChildComponent = OldToNewInstanceMap.Find(ChildComponent))
+		{
+			ChildComponent = CastChecked<USceneComponent>(*NewChildComponent, ECastCheckedType::NullAllowed);
+		}
+	}
+}
+
+FActorComponentInstanceData* USceneComponent::GetComponentInstanceData() const
+{
+	FActorComponentInstanceData* InstanceData = nullptr;
+
+	for (USceneComponent* Child : AttachChildren)
+	{
+		if (Child && !Child->IsCreatedByConstructionScript())
+		{
+			InstanceData = new FSceneComponentInstanceData(this);
+			break;
+		}
+	}
+
+	if (InstanceData == nullptr)
+	{
+		InstanceData = Super::GetComponentInstanceData();
+	}
+
+	return InstanceData;
+}
+
+FName USceneComponent::GetComponentInstanceDataType() const
+{
+	static const FName SceneComponentInstanceDataTypeName(TEXT("SceneComponentInstanceData"));
+	return SceneComponentInstanceDataTypeName;
+}
+
+void USceneComponent::UpdateChildTransforms()
+{
+	for(int32 i=0; i<AttachChildren.Num(); i++)
+	{
+		USceneComponent* ChildComp = AttachChildren[i];
+		if(ChildComp != NULL)
+		{
+			ChildComp->UpdateComponentToWorld();
+		}
+	}
+}
+
+void USceneComponent::Serialize(FArchive& Ar)
+{
+	Super::Serialize(Ar);
+
+	// Copy from deprecated properties
+	if(Ar.UE4Ver() < VER_UE4_SCENECOMP_TRANSLATION_TO_LOCATION)
+	{
+		RelativeLocation = RelativeTranslation_DEPRECATED;
+		bAbsoluteLocation = bAbsoluteTranslation_DEPRECATED;
+	}
+}
+
+
+void USceneComponent::PostInterpChange(UProperty* PropertyThatChanged)
+{
+	Super::PostInterpChange(PropertyThatChanged);
+
+	static FName RelativeScale3D(GET_MEMBER_NAME_CHECKED(USceneComponent, RelativeScale3D));
+
+	if (PropertyThatChanged->GetFName() == RelativeScale3D)
+	{
+		UpdateComponentToWorld();
+	}
+}
+
+TArray<FName> USceneComponent::GetAllSocketNames() const
+{
+	TArray<FComponentSocketDescription> SocketList;
+	QuerySupportedSockets(/*out*/ SocketList);
+
+	TArray<FName> ResultList;
+	ResultList.Reserve(SocketList.Num());
+
+	for (const FComponentSocketDescription& SocketDesc : SocketList)
+	{
+		ResultList.Add(SocketDesc.Name);
+	}
+
+	return ResultList;
+}
+
+FTransform USceneComponent::GetSocketTransform(FName SocketName, ERelativeTransformSpace TransformSpace) const
+{
+	switch(TransformSpace)
+	{
+		case RTS_Actor:
+		{
+			return ComponentToWorld.GetRelativeTransform( GetOwner()->GetTransform() );
+			break;
+		}
+		case RTS_Component:
+		{
+			return FTransform::Identity;
+		}
+		default:
+		{
+			return ComponentToWorld;
+		}
+	}
+}
+
+FVector USceneComponent::GetSocketLocation(FName SocketName) const
+{
+	return GetSocketTransform(SocketName, RTS_World).GetTranslation();
+}
+
+FRotator USceneComponent::GetSocketRotation(FName SocketName) const
+{
+	return GetSocketTransform(SocketName, RTS_World).GetRotation().Rotator();
+}
+
+FQuat USceneComponent::GetSocketQuaternion(FName SocketName) const
+{
+	return GetSocketTransform(SocketName, RTS_World).GetRotation();
+}
+
+bool USceneComponent::DoesSocketExist(FName InSocketName) const
+{
+	return false;
+}
+
+bool USceneComponent::HasAnySockets() const
+{
+	return false;
+}
+
+void USceneComponent::QuerySupportedSockets(TArray<FComponentSocketDescription>& OutSockets) const
+{
+}
+
+FVector USceneComponent::GetComponentVelocity() const
+{
+	return ComponentVelocity;
+}
+
+void USceneComponent::GetSocketWorldLocationAndRotation(FName InSocketName, FVector& OutLocation, FRotator& OutRotation) const
+{
+	FTransform const SocketWorldTransform = GetSocketTransform(InSocketName);
+
+	// assemble output
+	OutLocation = SocketWorldTransform.GetLocation();
+	OutRotation = SocketWorldTransform.Rotator();
+}
+
+bool USceneComponent::IsWorldGeometry() const
+{
+	return false;
+}
+
+ECollisionEnabled::Type USceneComponent::GetCollisionEnabled() const
+{
+	return ECollisionEnabled::NoCollision;
+}
+
+const FCollisionResponseContainer& USceneComponent::GetCollisionResponseToChannels() const
+{
+	return FCollisionResponseContainer::GetDefaultResponseContainer();
+}
+
+ECollisionResponse USceneComponent::GetCollisionResponseToChannel(ECollisionChannel Channel) const
+{
+	return ECR_Ignore;
+}
+
+ECollisionChannel USceneComponent::GetCollisionObjectType() const
+{
+	return ECC_WorldDynamic;
+}
+
+ECollisionResponse USceneComponent::GetCollisionResponseToComponent(class USceneComponent* OtherComponent) const
+{
+	// Ignore if no component, or either component has no collision
+	if(OtherComponent == NULL || GetCollisionEnabled() == ECollisionEnabled::NoCollision || OtherComponent->GetCollisionEnabled() == ECollisionEnabled::NoCollision)
+	{
+		return ECR_Ignore;
+	}
+
+	ECollisionResponse OutResponse;
+	ECollisionChannel MyCollisionObjectType = GetCollisionObjectType();
+	ECollisionChannel OtherCollisionObjectType = OtherComponent->GetCollisionObjectType();
+
+	/**
+	 * We decide minimum of behavior of both will decide the resulting response
+	 * If A wants to block B, but B wants to touch A, touch will be the result of this collision
+	 * However if A is static, then we don't care about B's response to A (static) but A's response to B overrides the verdict
+	 * Vice versa, if B is static, we don't care about A's response to static but B's response to A will override the verdict
+	 * To make this work, if MyCollisionObjectType is static, set OtherResponse to be ECR_Block, so to be ignored at the end
+	 **/
+	ECollisionResponse MyResponse = GetCollisionResponseToChannel(OtherCollisionObjectType);
+	ECollisionResponse OtherResponse = OtherComponent->GetCollisionResponseToChannel(MyCollisionObjectType);
+
+	OutResponse = FMath::Min<ECollisionResponse>(MyResponse, OtherResponse);
+
+
+	return OutResponse;
+}
+
+void USceneComponent::SetMobility(EComponentMobility::Type NewMobility)
+{
+	if (NewMobility != Mobility)
+	{
+		FComponentReregisterContext ReregisterContext(this);
+		Mobility = NewMobility;
+
+		if (Mobility == EComponentMobility::Movable)	//if we're now movable all children should be updated as having static children is invalid
+		{
+			for (USceneComponent* ChildComponent : AttachChildren)
+			{
+				if (ChildComponent)
+				{
+					ChildComponent->SetMobility(NewMobility);
+				}
+			}
+		}
+	}
+}
+
+bool USceneComponent::IsSimulatingPhysics(FName BoneName) const
+{
+	return false;
+}
+
+bool USceneComponent::IsAnySimulatingPhysics() const
+{
+	return IsSimulatingPhysics();
+}
+
+APhysicsVolume* USceneComponent::GetPhysicsVolume() const
+{
+	if (PhysicsVolume)
+	{
+		return PhysicsVolume;
+	}
+	else if (const UWorld* MyWorld = GetWorld())
+	{
+		return MyWorld->GetDefaultPhysicsVolume();
+	}
+
+	return NULL;
+}
+
+void USceneComponent::UpdatePhysicsVolume( bool bTriggerNotifiers )
+{
+	if ( bShouldUpdatePhysicsVolume && !IsPendingKill() && GetWorld() )
+	{
+		SCOPE_CYCLE_COUNTER(STAT_UpdatePhysicsVolume);
+
+		UWorld* const MyWorld = GetWorld();
+		APhysicsVolume* NewVolume = MyWorld->GetDefaultPhysicsVolume();
+		// Avoid doing anything if there are no other physics volumes in the world.
+		if (MyWorld->GetNonDefaultPhysicsVolumeCount() > 0)
+		{
+			// Avoid a full overlap query if we can do some quick bounds tests against the volumes.
+			static uint32 MaxVolumesToCheck = 100;
+			uint32 VolumeIndex = 0;
+			bool bAnyPotentialOverlap = false;
+			for (auto VolumeIter = MyWorld->GetNonDefaultPhysicsVolumeIterator(); VolumeIter && !bAnyPotentialOverlap; ++VolumeIter, ++VolumeIndex)
+			{
+				const APhysicsVolume* Volume = *VolumeIter;
+				if( Volume != nullptr )
+				{
+					const USceneComponent* VolumeRoot = Volume->GetRootComponent();
+					if (VolumeRoot)
+					{
+						if (FBoxSphereBounds::SpheresIntersect(VolumeRoot->Bounds, Bounds))
+						{
+							if (FBoxSphereBounds::BoxesIntersect(VolumeRoot->Bounds, Bounds))
+							{
+								bAnyPotentialOverlap = true;
+							}
+						}
+					}
+				}
+
+				// Bail if too many volumes. Later we'll probably convert to using an octree so this wouldn't be a concern.
+				if (VolumeIndex >= MaxVolumesToCheck)
+				{
+					bAnyPotentialOverlap = true;
+					break;
+				}
+			}
+			
+			if (bAnyPotentialOverlap)
+			{
+				// check for all volumes that overlap the component
+				TArray<FOverlapResult> Hits;
+				static FName NAME_PhysicsVolumeTrace = FName(TEXT("PhysicsVolumeTrace"));
+				FComponentQueryParams Params(NAME_PhysicsVolumeTrace, GetOwner());
+
+				bool bOverlappedOrigin = false;
+				const UPrimitiveComponent* SelfAsPrimitive = Cast<UPrimitiveComponent>(this);
+				if (SelfAsPrimitive)
+				{
+					MyWorld->ComponentOverlapMultiByChannel(Hits, SelfAsPrimitive, GetComponentLocation(), GetComponentRotation(), GetCollisionObjectType(), Params);
+				}
+				else
+				{
+					bOverlappedOrigin = true;
+					MyWorld->OverlapMultiByChannel(Hits, GetComponentLocation(), FQuat::Identity, GetCollisionObjectType(), FCollisionShape::MakeSphere(0.f), Params);
+				}				
+				
+				for (int32 HitIdx = 0; HitIdx < Hits.Num(); HitIdx++)
+				{
+					const FOverlapResult& Link = Hits[HitIdx];
+					APhysicsVolume* const V = Cast<APhysicsVolume>(Link.GetActor());
+					if (V && (V->Priority > NewVolume->Priority))
+					{
+						if (bOverlappedOrigin || V->IsOverlapInVolume(*this))
+						{
+							NewVolume = V;
+						}
+					}
+				}
+			}
+		}
+
+		if (PhysicsVolume != NewVolume)
+		{
+			SetPhysicsVolume( NewVolume, bTriggerNotifiers );
+		}
+	}
+}
+
+void USceneComponent::SetPhysicsVolume( APhysicsVolume * NewVolume,  bool bTriggerNotifiers )
+{
+	// Owner can be NULL
+	// The Notifier can be triggered with NULL Actor
+	// Still the delegate should be still called
+	if( bTriggerNotifiers )
+	{
+		if( NewVolume != PhysicsVolume )
+		{
+			AActor *A = GetOwner();
+			if( PhysicsVolume )
+			{
+				PhysicsVolume->ActorLeavingVolume(A);
+				PhysicsVolumeChangedDelegate.Broadcast(NewVolume);
+			}
+			PhysicsVolume = NewVolume;
+			if( PhysicsVolume )
+			{
+				PhysicsVolume->ActorEnteredVolume(A);
+			}
+		}
+	}
+	else
+	{
+		PhysicsVolume = NewVolume;
+	}
+}
+
+void USceneComponent::BeginDestroy()
+{
+	PhysicsVolumeChangedDelegate.Clear();
+
+	Super::BeginDestroy();
+}
+
+bool USceneComponent::InternalSetWorldLocationAndRotation(FVector NewLocation, FRotator NewRotation, bool bNoPhysics)
+{
+	// If attached to something, transform into local space
+	if (AttachParent != NULL)
+	{
+		FTransform const ParentToWorld = AttachParent->GetSocketTransform(AttachSocketName);
+
+		if (!bAbsoluteLocation)
+		{
+			NewLocation = ParentToWorld.InverseTransformPosition(NewLocation);
+		}
+
+		if (!bAbsoluteRotation)
+		{
+			FQuat const NewQuat = NewRotation.Quaternion();
+
+			// Quat multiplication works reverse way, make sure you do Parent(-1) * World = Local, not World*Parent(-) = Local (the way matrix does)
+			FQuat NewRelQuat = ParentToWorld.GetRotation().Inverse() * NewQuat;
+			NewRotation = NewRelQuat.Rotator();
+		}
+	}
+
+	bool bUpdated = false;
+	if (!NewLocation.Equals(RelativeLocation) || !NewRotation.Equals(RelativeRotation))
+	{
+		RelativeLocation = NewLocation;
+		RelativeRotation = NewRotation;
+		UpdateComponentToWorld(bNoPhysics);
+		bUpdated = true;
+	}
+
+	return bUpdated;
+}
+
+void USceneComponent::UpdateOverlaps(TArray<FOverlapInfo> const* PendingOverlaps, bool bDoNotifies, const TArray<FOverlapInfo>* OverlapsAtEndLocation)
+{
+	SCOPE_CYCLE_COUNTER(STAT_UpdateOverlaps); 
+
+	if (IsDeferringMovementUpdates())
+	{
+		return;
+	}
+	
+	if (bShouldUpdatePhysicsVolume)
+	{
+		UpdatePhysicsVolume(bDoNotifies);
+	}
+
+	// SceneComponent has no physical representation, so no overlaps to test for/
+	// But, we need to test down the attachment chain since there might be PrimitiveComponents below.
+	for (int32 ChildIdx=0; ChildIdx<AttachChildren.Num(); ++ChildIdx)
+	{
+		if (AttachChildren[ChildIdx])
+		{
+			// Do not pass on OverlapsAtEndLocation, it only applied to this component.
+			AttachChildren[ChildIdx]->UpdateOverlaps(NULL, bDoNotifies);
+		}
+	}
+}
+
+
+bool USceneComponent::MoveComponent( const FVector& Delta, const FRotator& NewRotation, bool bSweep, FHitResult* OutHit, EMoveComponentFlags MoveFlags )
+{
+#if !(UE_BUILD_SHIPPING || UE_BUILD_TEST)
+	// make sure mobility is movable, otherwise you shouldn't try to move
+	if (Mobility != EComponentMobility::Movable)
+	{
+		if (UWorld * World = GetWorld())
+		{
+			ULevel* Level = GetComponentLevel();
+
+			// It's only a problem if we're in gameplay, and the owning level is visible
+			if (World->HasBegunPlay() && IsRegistered() && Level && Level->bIsVisible)
+			{
+				FMessageLog("PIE").Warning(FText::Format(LOCTEXT("InvalidMove", "Mobility of {0} : {1} has to be 'Movable' if you'd like to move. "),
+					FText::FromString(GetNameSafe(GetOwner())), FText::FromString(GetName())));
+				return false;
+			}
+		}
+	}
+#endif
+
+	// Fill in optional output param. SceneComponent doesn't sweep, so this is just an empty result.
+	if (OutHit)
+	{
+		*OutHit = FHitResult(1.f);
+	}
+
+	if (!bWorldToComponentUpdated)
+	{
+		UpdateComponentToWorld();
+	}
+
+	// early out for zero case
+	if( Delta.IsZero() )
+	{
+		// Skip if no vector or rotation.
+		const FQuat NewQuat = NewRotation.Quaternion();
+		if( NewQuat.Equals(ComponentToWorld.GetRotation(), SMALL_NUMBER) )
+		{
+			return true;
+		}
+	}
+
+	// just teleport, sweep is supported for PrimitiveComponents.  this will update child components as well.
+	InternalSetWorldLocationAndRotation(GetComponentLocation() + Delta, NewRotation);
+
+	// Only update overlaps if not deferring updates within a scope
+	if (!IsDeferringMovementUpdates())
+	{
+		// need to update overlap detection in case PrimitiveComponents are attached.
+		UpdateOverlaps();
+	}
+
+	return true;
+}
+
+bool USceneComponent::IsVisibleInEditor() const
+{
+	// in editor, we only check bVisible
+	return bVisible;
+}
+
+bool USceneComponent::ShouldRender() const
+{
+	AActor* Owner = GetOwner();
+	const bool bShowInEditor = 
+#if WITH_EDITOR
+		GIsEditor ? (!Owner || !Owner->IsHiddenEd()) : false;
+#else
+		false;
+#endif
+	UWorld *World = GetWorld();
+	const bool bInGameWorld = GetWorld() && GetWorld()->UsesGameHiddenFlags();
+
+	const bool bShowInGame = IsVisible() && (!Owner || !Owner->bHidden);
+	return ((bInGameWorld && bShowInGame) || (!bInGameWorld && bShowInEditor)) && bVisible == true;
+}
+
+bool USceneComponent::CanEverRender() const
+{
+	AActor* Owner = GetOwner();
+	const bool bShowInEditor =
+#if WITH_EDITOR
+		GIsEditor ? (!Owner || !Owner->IsHiddenEd()) : false;
+#else
+		false;
+#endif
+	UWorld *World = GetWorld();
+	const bool bInGameWorld = World && World->UsesGameHiddenFlags();
+
+	const bool bShowInGame = (!Owner || !Owner->bHidden);
+	return ((bInGameWorld && bShowInGame) || (!bInGameWorld && bShowInEditor));
+}
+
+bool USceneComponent::ShouldComponentAddToScene() const
+{
+	// If the detail mode setting allows it, add it to the scene.
+	return DetailMode <= GetCachedScalabilityCVars().DetailMode;
+}
+
+bool USceneComponent::IsVisible() const
+{
+	// if hidden in game, nothing to do
+	if (bHiddenInGame)
+	{
+		return false;
+	}
+
+	return ( bVisible ); 
+}
+
+void USceneComponent::ToggleVisibility(bool bPropagateToChildren)
+{
+	SetVisibility(!bVisible,bPropagateToChildren);
+}
+
+void USceneComponent::SetVisibility(bool bNewVisibility, bool bPropagateToChildren)
+{
+	if ( bNewVisibility != bVisible )
+	{
+		bVisible = bNewVisibility;
+		MarkRenderStateDirty();
+	}
+
+	if (bPropagateToChildren && AttachChildren.Num() > 0)
+	{
+		// fully traverse down the attachment tree
+		// we do it entirely inline here instead of recursing in case a primitivecomponent is a child of a non-primitivecomponent
+		TArray<USceneComponent*> ComponentStack;
+
+		// presize to minimize reallocs
+		ComponentStack.Reserve(FMath::Max(32, AttachChildren.Num()));
+
+		// prime the pump
+		ComponentStack.Append(AttachChildren);
+
+		while (ComponentStack.Num() > 0)
+		{
+			USceneComponent* const CurrentComp = ComponentStack.Pop();
+			if (CurrentComp)
+			{
+				ComponentStack.Append(CurrentComp->AttachChildren);
+				// don't propagate, we are handling it already
+				CurrentComp->SetVisibility(bNewVisibility, false);
+			}
+		}
+	}
+}
+
+void USceneComponent::SetHiddenInGame(bool NewHiddenGame, bool bPropagateToChildren)
+{
+	if( NewHiddenGame != bHiddenInGame )
+	{
+		bHiddenInGame = NewHiddenGame;
+		MarkRenderStateDirty();
+	}
+
+	if (bPropagateToChildren && AttachChildren.Num() > 0)
+	{
+		// fully traverse down the attachment tree
+		// we do it entirely inline here instead of recursing in case a primitivecomponent is a child of a non-primitivecomponent
+		TArray<USceneComponent*> ComponentStack;
+
+		// presize to minimize reallocs
+		ComponentStack.Reserve(FMath::Max(32, AttachChildren.Num()));
+
+		// prime the pump
+		ComponentStack.Append(AttachChildren);
+
+		while (ComponentStack.Num() > 0)
+		{
+			USceneComponent* const CurrentComp = ComponentStack.Pop();
+			if (CurrentComp)
+			{
+				ComponentStack.Append(CurrentComp->AttachChildren);
+
+				UPrimitiveComponent* const PrimComp = Cast<UPrimitiveComponent>(CurrentComp);
+				if (PrimComp)
+				{
+					// don't propagate, we are handling it already
+					PrimComp->SetHiddenInGame(NewHiddenGame, false);
+				}
+			}
+		}
+	}
+}
+
+void USceneComponent::ApplyWorldOffset(const FVector& InOffset, bool bWorldShift)
+{
+	Super::ApplyWorldOffset(InOffset, bWorldShift);
+	
+	// Calculate current ComponentToWorld transform
+	// We do this because at level load/duplication ComponentToWorld is uninitialized
+	{
+		const FTransform RelativeTransform(RelativeRotation, RelativeLocation, RelativeScale3D);
+		ComponentToWorld = CalcNewComponentToWorld(RelativeTransform);
+	}
+
+	// Update bounds
+	Bounds.Origin+= InOffset;
+
+	// Update component location
+	FVector NewLocation = GetComponentLocation() + InOffset;
+	// If attached to something, transform into local space
+	if (AttachParent != NULL)
+	{
+		FTransform const ParentToWorld = AttachParent->GetSocketTransform(AttachSocketName);
+		if (!bAbsoluteLocation)
+		{
+			NewLocation = ParentToWorld.InverseTransformPosition(GetComponentLocation());
+		}
+	}
+	
+	RelativeLocation = NewLocation;
+		
+	// Calculate the new ComponentToWorld transform
+	const FTransform RelativeTransform(RelativeRotation, RelativeLocation, RelativeScale3D);
+	ComponentToWorld = CalcNewComponentToWorld(RelativeTransform);
+
+	// Physics move is skipped if physics state is not created or physics scene supports origin shifting
+	// We still need to send transform to physics scene to "transform back" actors which should ignore origin shifting
+	// (such actors receive Zero offset)
+	const bool bSkipPhysicsTransform = (!bPhysicsStateCreated || (bWorldShift && FPhysScene::SupportsOriginShifting() && !InOffset.IsZero()));
+	OnUpdateTransform(bSkipPhysicsTransform);
+	
+	// We still need to send transform to RT to "transform back" primitives which should ignore origin shifting
+	// (such primitives receive Zero offset)
+	if (!bWorldShift || InOffset.IsZero())
+	{
+		MarkRenderTransformDirty();
+	}
+
+	// Update physics volume if desired	
+	if (bShouldUpdatePhysicsVolume && !bWorldShift)
+	{
+		UpdatePhysicsVolume(true);
+	}
+
+	// Update children
+	for(int32 i=0; i<AttachChildren.Num(); i++)
+	{
+		USceneComponent* ChildComp = AttachChildren[i];
+		if(ChildComp != NULL)
+		{
+			ChildComp->ApplyWorldOffset(InOffset, bWorldShift);
+		}
+	}
+}
+
+FBoxSphereBounds USceneComponent::GetPlacementExtent() const
+{
+	return CalcBounds( FTransform::Identity );
+}
+
+void USceneComponent::OnRep_Transform()
+{
+	NetUpdateTransform = true;
+}
+
+void USceneComponent::OnRep_Visibility(bool OldValue)
+{
+	bool ReppedValue = bVisible;
+	bVisible = OldValue;
+	SetVisibility(ReppedValue);
+}
+
+void USceneComponent::PreNetReceive()
+{
+	Super::PreNetReceive();
+
+	NetOldAttachSocketName = AttachSocketName;
+	NetOldAttachParent = AttachParent;
+}
+
+void USceneComponent::PostNetReceive()
+{
+	Super::PostNetReceive();
+
+	// If we have no attach parent, attach to parent's root component.
+	bool UpdateAttach = false;
+
+	UpdateAttach |= (NetOldAttachParent != AttachParent);
+	UpdateAttach |= (NetOldAttachSocketName != AttachSocketName);
+
+	if (AttachParent == NULL)
+	{
+		USceneComponent * ParentRoot = GetOwner()->GetRootComponent();
+		if (ParentRoot != this)
+		{
+			UpdateAttach = true;
+			AttachParent = ParentRoot;
+		}
+	}
+	
+	if (UpdateAttach)
+	{
+		Exchange(NetOldAttachParent, AttachParent);
+		Exchange(NetOldAttachSocketName, AttachSocketName);
+		
+		AttachTo(NetOldAttachParent, NetOldAttachSocketName);
+	}
+
+	if (NetUpdateTransform)
+	{
+		UpdateComponentToWorld(true);
+		NetUpdateTransform = false;
+	}
+}
+
+void USceneComponent::GetLifetimeReplicatedProps(TArray< FLifetimeProperty > & OutLifetimeProps) const
+{
+	Super::GetLifetimeReplicatedProps(OutLifetimeProps);
+	
+	DOREPLIFETIME(USceneComponent, bAbsoluteLocation);
+	DOREPLIFETIME(USceneComponent, bAbsoluteRotation);
+	DOREPLIFETIME(USceneComponent, bAbsoluteScale);
+	DOREPLIFETIME(USceneComponent, bVisible);
+	DOREPLIFETIME(USceneComponent, AttachParent);
+	DOREPLIFETIME(USceneComponent, AttachSocketName);
+	DOREPLIFETIME(USceneComponent, RelativeLocation);
+	DOREPLIFETIME(USceneComponent, RelativeRotation);
+	DOREPLIFETIME(USceneComponent, RelativeScale3D);
+}
+
+#if WITH_EDITOR
+bool USceneComponent::CanEditChange( const UProperty* Property ) const
+{
+	bool bIsEditable = Super::CanEditChange( Property );
+	if( bIsEditable && Property != NULL )
+	{
+		AActor* Owner = GetOwner();
+		if(Owner != NULL)
+		{
+			if(Property->GetFName() == TEXT( "RelativeLocation" ) ||
+			   Property->GetFName() == TEXT( "RelativeRotation" ) ||
+			   Property->GetFName() == TEXT( "RelativeScale3D" ))
+			{
+				bIsEditable = !Owner->bLockLocation;
+			}
+		}
+	}
+
+	return bIsEditable;
+}
+#endif
+
+/**
+ * FScopedMovementUpdate implementation
+ */
+
+static uint32 s_ScopedWarningCount = 0;
+
+FScopedMovementUpdate::FScopedMovementUpdate( class USceneComponent* Component, EScopedUpdate::Type ScopeBehavior )
+: Owner(Component)
+, bDeferUpdates(ScopeBehavior == EScopedUpdate::DeferredUpdates)
+, bHasValidOverlapsAtEnd(false)
+{
+	if (IsValid(Component))
+	{
+		InitialTransform = Component->GetComponentToWorld();
+		InitialRelativeLocation = Component->RelativeLocation;
+		InitialRelativeRotation = Component->RelativeRotation;
+		InitialRelativeScale = Component->RelativeScale3D;
+
+		if (ScopeBehavior == EScopedUpdate::ImmediateUpdates)
+		{
+			// We only allow ScopeUpdateImmediately if there is no outer scope, or if the outer scope is also ScopeUpdateImmediately.
+			FScopedMovementUpdate* OuterScope = Component->GetCurrentScopedMovement();
+			if (OuterScope && OuterScope->bDeferUpdates)
+			{
+				if (s_ScopedWarningCount < 100 || (GFrameCounter & 31) == 0)
+				{
+					s_ScopedWarningCount++;
+					UE_LOG(LogSceneComponent, Error, TEXT("FScopedMovementUpdate attempting to use immediate updates within deferred scope, will use deferred updates instead."));
+				}
+
+				bDeferUpdates = true;
+			}
+		}			
+
+		if (bDeferUpdates)
+		{
+			Component->BeginScopedMovementUpdate(*this);
+		}
+	}
+	else
+	{
+		Owner = NULL;
+	}
+}
+
+
+FScopedMovementUpdate::~FScopedMovementUpdate()
+{
+	if (bDeferUpdates && IsValid(Owner))
+	{
+		Owner->EndScopedMovementUpdate(*this);
+	}
+	Owner = NULL;
+}
+
+
+bool FScopedMovementUpdate::IsTransformDirty() const
+{
+	if (IsValid(Owner))
+	{
+		return !InitialTransform.Equals(Owner->GetComponentToWorld());
+	}
+
+	return false;
+}
+
+
+void FScopedMovementUpdate::RevertMove()
+{
+	USceneComponent* Component = Owner;
+	if (IsValid(Component))
+	{
+		bHasValidOverlapsAtEnd = false;
+		PendingOverlaps.Reset();
+		BlockingHits.Reset();
+		
+		if (IsTransformDirty())
+		{
+			// Teleport to start
+			Component->ComponentToWorld = InitialTransform;
+			Component->RelativeLocation = InitialRelativeLocation;
+			Component->RelativeRotation = InitialRelativeRotation;
+			Component->RelativeScale3D = InitialRelativeScale;
+
+			if (!IsDeferringUpdates())
+			{
+				Component->PropagateTransformUpdate(true);
+				Component->UpdateOverlaps();
+			}
+		}
+	}
+}
+
+
+void FScopedMovementUpdate::AppendOverlaps(const TArray<struct FOverlapInfo>& OtherOverlaps, const TArray<FOverlapInfo>* OverlapsAtEndLocation)
+{
+	PendingOverlaps.Append(OtherOverlaps);
+	if (OverlapsAtEndLocation)
+	{
+		bHasValidOverlapsAtEnd = true;
+		OverlapsAtEnd = *OverlapsAtEndLocation;		
+	}
+	else
+	{
+		bHasValidOverlapsAtEnd = false;
+	}
+}
+
+
+void FScopedMovementUpdate::OnInnerScopeComplete(const FScopedMovementUpdate& InnerScope)
+{
+	if (IsValid(Owner))
+	{
+		checkSlow(IsDeferringUpdates());
+		checkSlow(InnerScope.IsDeferringUpdates());
+
+		// Combine with the next item on the stack
+		AppendOverlaps(InnerScope.GetPendingOverlaps(), InnerScope.GetOverlapsAtEnd());
+		BlockingHits.Append(InnerScope.GetPendingBlockingHits());
+	}	
+}
+
+#if WITH_EDITOR
+const int32 USceneComponent::GetNumUncachedStaticLightingInteractions() const
+{
+	int32 NumUncachedStaticLighting = 0;
+	for (int32 i = 0; i < AttachChildren.Num(); i++)
+	{
+		if (AttachChildren[i])
+		{
+			NumUncachedStaticLighting += AttachChildren[i]->GetNumUncachedStaticLightingInteractions();
+		}
+	}
+	return NumUncachedStaticLighting;
+}
+#endif
+
+void USceneComponent::UpdateNavigationData()
+{
+	if (UNavigationSystem::ShouldUpdateNavOctreeOnComponentChange() &&
+		IsRegistered() && World && World->IsGameWorld() &&
+		World->GetNetMode() < ENetMode::NM_Client)
+	{
+		UNavigationSystem::UpdateNavOctree(this);
+	}
+}
+
+
+// K2 versions of various transform changing operations.
+// Note: we pass null for the hit result if not sweeping, for better perf.
+// This assumes this K2 function is only used by blueprints, which initializes the param for each function call.
+
+void USceneComponent::K2_SetRelativeLocationAndRotation(FVector NewLocation, FRotator NewRotation, bool bSweep, FHitResult& SweepHitResult)
+{
+	SetRelativeLocationAndRotation(NewLocation, NewRotation, bSweep, (bSweep ? &SweepHitResult : nullptr));
+}
+
+void USceneComponent::K2_SetWorldLocationAndRotation(FVector NewLocation, FRotator NewRotation, bool bSweep, FHitResult& SweepHitResult)
+{
+	SetWorldLocationAndRotation(NewLocation, NewRotation, bSweep, (bSweep ? &SweepHitResult : nullptr));
+}
+
+void USceneComponent::K2_SetRelativeLocation(FVector NewLocation, bool bSweep, FHitResult& SweepHitResult)
+{
+	SetRelativeLocation(NewLocation, bSweep, (bSweep ? &SweepHitResult : nullptr));
+}
+
+void USceneComponent::K2_SetRelativeRotation(FRotator NewRotation, bool bSweep, FHitResult& SweepHitResult)
+{
+	SetRelativeRotation(NewRotation, bSweep, (bSweep ? &SweepHitResult : nullptr));
+}
+
+void USceneComponent::K2_SetRelativeTransform(const FTransform& NewTransform, bool bSweep, FHitResult& SweepHitResult)
+{
+	SetRelativeTransform(NewTransform, bSweep, (bSweep ? &SweepHitResult : nullptr));
+}
+
+void USceneComponent::K2_AddRelativeLocation(FVector DeltaLocation, bool bSweep, FHitResult& SweepHitResult)
+{
+	AddRelativeLocation(DeltaLocation, bSweep, (bSweep ? &SweepHitResult : nullptr));
+}
+
+void USceneComponent::K2_AddRelativeRotation(FRotator DeltaRotation, bool bSweep, FHitResult& SweepHitResult)
+{
+	AddRelativeRotation(DeltaRotation, bSweep, (bSweep ? &SweepHitResult : nullptr));
+}
+
+void USceneComponent::K2_AddLocalOffset(FVector DeltaLocation, bool bSweep, FHitResult& SweepHitResult)
+{
+	AddLocalOffset(DeltaLocation, bSweep, (bSweep ? &SweepHitResult : nullptr));
+}
+
+void USceneComponent::K2_AddLocalRotation(FRotator DeltaRotation, bool bSweep, FHitResult& SweepHitResult)
+{
+	AddLocalRotation(DeltaRotation, bSweep, (bSweep ? &SweepHitResult : nullptr));
+}
+
+void USceneComponent::K2_AddLocalTransform(const FTransform& DeltaTransform, bool bSweep, FHitResult& SweepHitResult)
+{
+	AddLocalTransform(DeltaTransform, bSweep, (bSweep ? &SweepHitResult : nullptr));
+}
+
+void USceneComponent::K2_SetWorldLocation(FVector NewLocation, bool bSweep, FHitResult& SweepHitResult)
+{
+	SetWorldLocation(NewLocation, bSweep, (bSweep ? &SweepHitResult : nullptr));
+}
+
+void USceneComponent::K2_SetWorldRotation(FRotator NewRotation, bool bSweep, FHitResult& SweepHitResult)
+{
+	SetWorldRotation(NewRotation, bSweep, (bSweep ? &SweepHitResult : nullptr));
+}
+
+void USceneComponent::K2_SetWorldTransform(const FTransform& NewTransform, bool bSweep, FHitResult& SweepHitResult)
+{
+	SetWorldTransform(NewTransform, bSweep, (bSweep ? &SweepHitResult : nullptr));
+}
+
+void USceneComponent::K2_AddWorldOffset(FVector DeltaLocation, bool bSweep, FHitResult& SweepHitResult)
+{
+	AddWorldOffset(DeltaLocation, bSweep, (bSweep ? &SweepHitResult : nullptr));
+}
+
+void USceneComponent::K2_AddWorldRotation(FRotator DeltaRotation, bool bSweep, FHitResult& SweepHitResult)
+{
+	AddWorldRotation(DeltaRotation, bSweep, (bSweep ? &SweepHitResult : nullptr));
+}
+
+void USceneComponent::K2_AddWorldTransform(const FTransform& DeltaTransform, bool bSweep, FHitResult& SweepHitResult)
+{
+	AddWorldTransform(DeltaTransform, bSweep, (bSweep ? &SweepHitResult : nullptr));
+}
+
+#undef LOCTEXT_NAMESPACE
diff --git a/Engine/Source/Runtime/Engine/Private/Components/ShapeComponent.cpp b/Engine/Source/Runtime/Engine/Private/Components/ShapeComponent.cpp
new file mode 100644
--- /dev/null
+++ b/Engine/Source/Runtime/Engine/Private/Components/ShapeComponent.cpp
@@ -0,0 +1,81 @@
+// Copyright 1998-2015 Epic Games, Inc. All Rights Reserved.
+
+
+#include "EnginePrivate.h"
+#include "AI/Navigation/NavAreas/NavArea_Obstacle.h"
+
+UShapeComponent::UShapeComponent(const FObjectInitializer& ObjectInitializer)
+	: Super(ObjectInitializer)
+{
+	static const FName CollisionProfileName(TEXT("OverlapAllDynamic"));
+	BodyInstance.SetCollisionProfileName(CollisionProfileName);
+	// when we deprecated this variable, we switched on for the shapecomponent collision profile
+	// the problem with adding collision profile later, the instanced data will be wiped
+	// since shape component is so popular for BP and so on, I'm adding manual support for compatibility
+	// this only works since this variable is getting deprecated. 
+	BodyInstance.ResponseToChannels_DEPRECATED.SetAllChannels(ECR_Block);
+	BodyInstance.bAutoWeld = true;	//UShapeComponent by default has auto welding
+
+	bHiddenInGame = true;
+	bCastDynamicShadow = false;
+	ShapeColor = FColor(223, 149, 157, 255);
+	bShouldCollideWhenPlacing = false;
+	
+	bHasCustomNavigableGeometry = EHasCustomNavigableGeometry::Yes;
+	bCanEverAffectNavigation = true;
+	bDynamicObstacle = false;
+	AreaClass = UNavArea_Obstacle::StaticClass();
+}
+
+FPrimitiveSceneProxy* UShapeComponent::CreateSceneProxy()
+{
+	check( false && "Subclass needs to Implement this" );
+	return NULL;
+}
+
+FBoxSphereBounds UShapeComponent::CalcBounds(const FTransform& LocalToWorld) const
+{
+	check( false && "Subclass needs to Implement this" );
+	return FBoxSphereBounds();
+}
+
+void UShapeComponent::UpdateBodySetup()
+{
+	check( false && "Subclass needs to Implement this" );
+}
+
+UBodySetup* UShapeComponent::GetBodySetup()
+{
+	UpdateBodySetup();
+	return ShapeBodySetup;
+}
+
+void UShapeComponent::GetUsedMaterials( TArray<UMaterialInterface*>& OutMaterials ) const
+{
+	OutMaterials.Add( ShapeMaterial );
+}
+
+#if WITH_EDITOR
+void UShapeComponent::PostEditChangeProperty(FPropertyChangedEvent& PropertyChangedEvent)
+{
+	if (!IsTemplate())
+	{
+		UpdateBodySetup(); // do this before reregistering components so that new values are used for collision
+	}
+
+	Super::PostEditChangeProperty(PropertyChangedEvent);
+}
+#endif // WITH_EDITOR
+
+bool UShapeComponent::DoCustomNavigableGeometryExport(FNavigableGeometryExport& GeomExport) const
+{
+	return !bDynamicObstacle;
+}
+
+void UShapeComponent::GetNavigationData(FNavigationRelevantData& Data) const
+{
+	if (bDynamicObstacle)
+	{
+		Data.Modifiers.CreateAreaModifiers(this, AreaClass);
+	}
+}
