Engine: Moved actor components into Components folder

diff --git a/Engine/Source/Runtime/Engine/Private/Components/SkeletalMeshComponent.cpp b/Engine/Source/Runtime/Engine/Private/Components/SkeletalMeshComponent.cpp
new file mode 100644
--- /dev/null
+++ b/Engine/Source/Runtime/Engine/Private/Components/SkeletalMeshComponent.cpp
@@ -0,0 +1,2094 @@
+// Copyright 1998-2015 Epic Games, Inc. All Rights Reserved.
+
+/*=============================================================================
+	UnSkeletalComponent.cpp: Actor component implementation.
+=============================================================================*/
+
+#include "EnginePrivate.h"
+#include "ParticleDefinitions.h"
+#include "BlueprintUtilities.h"
+#include "SkeletalRenderCPUSkin.h"
+#include "SkeletalRenderGPUSkin.h"
+#include "AnimEncoding.h"
+#include "AnimationUtils.h"
+#include "AnimationRuntime.h"
+#include "PhysXASync.h"
+#include "AnimTree.h"
+#include "Animation/AnimNodeBase.h"
+#include "Animation/VertexAnim/VertexAnimation.h"
+#include "Particles/ParticleSystemComponent.h"
+#include "Animation/AnimSingleNodeInstance.h"
+#include "Engine/SkeletalMeshSocket.h"
+#if WITH_EDITOR
+#include "ShowFlags.h"
+#include "Collision.h"
+#include "ConvexVolume.h"
+#endif
+#include "PhysicsEngine/PhysicsSettings.h"
+#include "PhysicsEngine/PhysicsAsset.h"
+
+TAutoConsoleVariable<int32> CVarUseParallelAnimationEvaluation(TEXT("a.ParallelAnimEvaluation"), 1, TEXT("If 1, animation evaluation will be run across the task graph system. If 0, evaluation will run purely on the game thread"));
+
+DECLARE_CYCLE_STAT_EXTERN(TEXT("Anim Instance Spawn Time"), STAT_AnimSpawnTime, STATGROUP_Anim, );
+DEFINE_STAT(STAT_AnimSpawnTime);
+
+class FParallelAnimationEvaluationTask
+{
+	TWeakObjectPtr<USkeletalMeshComponent> SkeletalMeshComponent;
+
+public:
+	FParallelAnimationEvaluationTask(TWeakObjectPtr<USkeletalMeshComponent> InSkeletalMeshComponent)
+		: SkeletalMeshComponent(InSkeletalMeshComponent)
+	{
+	}
+
+	FORCEINLINE TStatId GetStatId() const
+	{
+		RETURN_QUICK_DECLARE_CYCLE_STAT(FParallelAnimationEvaluationTask, STATGROUP_TaskGraphTasks);
+	}
+	static ENamedThreads::Type GetDesiredThread()
+	{
+		return ENamedThreads::AnyThread;
+	}
+	static ESubsequentsMode::Type GetSubsequentsMode()
+	{
+		return ESubsequentsMode::TrackSubsequents;
+	}
+
+	void DoTask(ENamedThreads::Type CurrentThread, const FGraphEventRef& MyCompletionGraphEvent)
+	{
+		if (USkeletalMeshComponent* Comp = SkeletalMeshComponent.Get())
+		{
+			Comp->ParallelAnimationEvaluation();
+		}
+	}
+};
+
+class FParallelAnimationCompletionTask
+{
+	TWeakObjectPtr<USkeletalMeshComponent> SkeletalMeshComponent;
+
+public:
+	FParallelAnimationCompletionTask(TWeakObjectPtr<USkeletalMeshComponent> InSkeletalMeshComponent)
+		: SkeletalMeshComponent(InSkeletalMeshComponent)
+	{
+	}
+
+	FORCEINLINE TStatId GetStatId() const
+	{
+		RETURN_QUICK_DECLARE_CYCLE_STAT(FParallelAnimationCompletionTask, STATGROUP_TaskGraphTasks);
+	}
+	static ENamedThreads::Type GetDesiredThread()
+	{
+		return ENamedThreads::GameThread;
+	}
+	static ESubsequentsMode::Type GetSubsequentsMode()
+	{
+		return ESubsequentsMode::TrackSubsequents;
+	}
+
+	void DoTask(ENamedThreads::Type CurrentThread, const FGraphEventRef& MyCompletionGraphEvent)
+	{
+		SCOPE_CYCLE_COUNTER(STAT_AnimGameThreadTime);
+		if (USkeletalMeshComponent* Comp = SkeletalMeshComponent.Get())
+		{
+			Comp->CompleteParallelAnimationEvaluation();
+		}
+	}
+};
+
+USkeletalMeshComponent::USkeletalMeshComponent(const FObjectInitializer& ObjectInitializer)
+	: Super(ObjectInitializer)
+{
+	bAutoActivate = true;
+	PrimaryComponentTick.bCanEverTick = true;
+	PrimaryComponentTick.TickGroup = TG_PrePhysics;
+	bWantsInitializeComponent = true;
+	GlobalAnimRateScale = 1.0f;
+	bNoSkeletonUpdate = false;
+	MeshComponentUpdateFlag = EMeshComponentUpdateFlag::AlwaysTickPoseAndRefreshBones;
+	KinematicBonesUpdateType = EKinematicBonesUpdateToPhysics::SkipSimulatingBones;
+	bGenerateOverlapEvents = false;
+	LineCheckBoundsScale = FVector(1.0f, 1.0f, 1.0f);
+
+	PreClothTickFunction.TickGroup = TG_PreCloth;
+	PreClothTickFunction.bCanEverTick = true;
+	PreClothTickFunction.bStartWithTickEnabled = true;
+
+
+#if WITH_APEX_CLOTHING
+	ClothMaxDistanceScale = 1.0f;
+	ClothTeleportMode = FClothingActor::Continuous;
+	PrevRootBoneMatrix = GetBoneMatrix(0); // save the root bone transform
+	bResetAfterTeleport = true;
+	TeleportDistanceThreshold = 300.0f;
+	TeleportRotationThreshold = 0.0f;// angles in degree, disabled by default
+	// pre-compute cloth teleport thresholds for performance
+	ClothTeleportCosineThresholdInRad = FMath::Cos(FMath::DegreesToRadians(TeleportRotationThreshold));
+	ClothTeleportDistThresholdSquared = TeleportDistanceThreshold * TeleportDistanceThreshold;
+	bNeedTeleportAndResetOnceMore = false;
+	ClothBlendWeight = 1.0f;
+	bPreparedClothMorphTargets = false;
+#if WITH_CLOTH_COLLISION_DETECTION
+	ClothingCollisionRevision = 0;
+#endif// #if WITH_CLOTH_COLLISION_DETECTION
+
+#endif//#if WITH_APEX_CLOTHING
+
+	DefaultPlayRate_DEPRECATED = 1.0f;
+	bDefaultPlaying_DEPRECATED = true;
+	bEnablePhysicsOnDedicatedServer = UPhysicsSettings::Get()->bSimulateSkeletalMeshOnDedicatedServer;
+	bEnableUpdateRateOptimizations = false;
+	RagdollAggregateThreshold = UPhysicsSettings::Get()->RagdollAggregateThreshold;
+
+	bTickInEditor = true;
+}
+
+
+void USkeletalMeshComponent::RegisterComponentTickFunctions(bool bRegister)
+{
+	Super::RegisterComponentTickFunctions(bRegister);
+
+	if (bRegister)
+	{
+		if (SetupActorComponentTickFunction(&PreClothTickFunction))
+		{
+			PreClothTickFunction.Target = this;
+			// Set a prereq for the pre cloth tick to happen after physics is finished
+			if (World != NULL)
+			{
+				PreClothTickFunction.AddPrerequisite(World, World->EndPhysicsTickFunction);
+			}
+		}
+	}
+	else
+	{
+		if (PreClothTickFunction.IsTickFunctionRegistered())
+		{
+			PreClothTickFunction.UnRegisterTickFunction();
+		}
+	}
+
+}
+
+bool USkeletalMeshComponent::NeedToSpawnAnimScriptInstance(bool bForceInit) const
+{
+	if (AnimationMode == EAnimationMode::AnimationBlueprint && (AnimBlueprintGeneratedClass != NULL) && 
+		(SkeletalMesh != NULL) && (SkeletalMesh->Skeleton->IsCompatible(AnimBlueprintGeneratedClass->TargetSkeleton)))
+	{
+		if (bForceInit || (AnimScriptInstance == NULL) || (AnimScriptInstance->GetClass() != AnimBlueprintGeneratedClass) )
+		{
+			return true;
+		}
+	}
+
+	return false;
+}
+
+bool USkeletalMeshComponent::IsAnimBlueprintInstanced() const
+{
+	return (AnimScriptInstance && AnimScriptInstance->GetClass() == AnimBlueprintGeneratedClass);
+}
+
+void USkeletalMeshComponent::OnRegister()
+{
+	Super::OnRegister();
+
+	InitAnim(false);
+
+	if (MeshComponentUpdateFlag == EMeshComponentUpdateFlag::OnlyTickPoseWhenRendered && !FApp::CanEverRender())
+	{
+		SetComponentTickEnabled(false);
+	}
+}
+
+void USkeletalMeshComponent::OnUnregister()
+{
+#if WITH_APEX_CLOTHING
+	//clothing actors will be re-created in TickClothing
+	ReleaseAllClothingResources();
+#endif// #if WITH_APEX_CLOTHING
+
+	Super::OnUnregister();
+}
+
+void USkeletalMeshComponent::InitAnim(bool bForceReinit)
+{
+	// a lot of places just call InitAnim without checking Mesh, so 
+	// I'm moving the check here
+	if ( SkeletalMesh != NULL && IsRegistered() )
+	{
+		bool bBlueprintMismatch = (AnimBlueprintGeneratedClass != NULL) && 
+			(AnimScriptInstance != NULL) && (AnimScriptInstance->GetClass() != AnimBlueprintGeneratedClass);
+
+		bool bSkeletonMismatch = AnimScriptInstance && AnimScriptInstance->CurrentSkeleton && (AnimScriptInstance->CurrentSkeleton!=SkeletalMesh->Skeleton);
+
+		if (bBlueprintMismatch || bSkeletonMismatch )
+		{
+			ClearAnimScriptInstance();
+		}
+
+		// this has to be called before Initialize Animation because it will required RequiredBones list when InitializeAnimScript
+		RecalcRequiredBones(0);
+
+		InitializeAnimScriptInstance(bForceReinit);
+
+        //Make sure we have a valid pose		
+        TickAnimation(0.f); 
+
+		RefreshBoneTransforms();
+		FlipEditableSpaceBases();
+		UpdateComponentToWorld();
+	}
+}
+
+void USkeletalMeshComponent::InitializeAnimScriptInstance(bool bForceReinit)
+{
+	if (IsRegistered())
+	{
+		if (NeedToSpawnAnimScriptInstance(bForceReinit))
+		{
+			SCOPE_CYCLE_COUNTER(STAT_AnimSpawnTime);
+			AnimScriptInstance = NewObject<UAnimInstance>(this, AnimBlueprintGeneratedClass);
+
+			if (AnimScriptInstance)
+			{
+				AnimScriptInstance->InitializeAnimation();
+			}
+		}
+		else if (AnimationMode == EAnimationMode::AnimationSingleNode)
+		{
+			SCOPE_CYCLE_COUNTER(STAT_AnimSpawnTime);
+
+			UAnimSingleNodeInstance* OldInstance = NULL;
+			if (!bForceReinit)
+			{
+				OldInstance = Cast<UAnimSingleNodeInstance>(AnimScriptInstance);
+			}
+
+			AnimScriptInstance = NewObject<UAnimSingleNodeInstance>(this);
+
+			if (AnimScriptInstance)
+			{
+				AnimScriptInstance->InitializeAnimation();
+			}
+
+			if (OldInstance && AnimScriptInstance)
+			{
+				// Copy data from old instance unless we force reinitialized
+				FSingleAnimationPlayData CachedData;
+				CachedData.PopulateFrom(OldInstance);
+				CachedData.Initialize(Cast<UAnimSingleNodeInstance>(AnimScriptInstance));
+			}
+		}
+		else if (AnimScriptInstance)
+		{
+			AnimScriptInstance->InitializeAnimation();
+		}		
+	}
+}
+
+bool USkeletalMeshComponent::IsWindEnabled() const
+{
+#if WITH_APEX_CLOTHING
+	// Wind is enabled in game worlds
+	return GetWorld() && GetWorld()->IsGameWorld();
+#else
+	return false;
+#endif
+}
+
+void USkeletalMeshComponent::ClearAnimScriptInstance()
+{
+	AnimScriptInstance = NULL;
+}
+
+void USkeletalMeshComponent::CreateRenderState_Concurrent()
+{
+	// Update bHasValidBodies flag
+	UpdateHasValidBodies();
+
+	Super::CreateRenderState_Concurrent();
+}
+
+void USkeletalMeshComponent::InitializeComponent()
+{
+	Super::InitializeComponent();
+
+	InitAnim(false);
+}
+
+#if CHART_DISTANCE_FACTORS
+static void AddDistanceFactorToChart(float DistanceFactor)
+{
+	if(DistanceFactor < SMALL_NUMBER)
+	{
+		return;
+	}
+
+	if(DistanceFactor >= GDistanceFactorDivision[NUM_DISTANCEFACTOR_BUCKETS-2])
+	{
+		GDistanceFactorChart[NUM_DISTANCEFACTOR_BUCKETS-1]++;
+	}
+	else if(DistanceFactor < GDistanceFactorDivision[0])
+	{
+		GDistanceFactorChart[0]++;
+	}
+	else
+	{
+		for(int32 i=1; i<NUM_DISTANCEFACTOR_BUCKETS-2; i++)
+		{
+			if(DistanceFactor < GDistanceFactorDivision[i])
+			{
+				GDistanceFactorChart[i]++;
+				break;
+			}
+		}
+	}
+}
+#endif // CHART_DISTANCE_FACTORS
+
+#if WITH_EDITOR
+void USkeletalMeshComponent::PostEditChangeProperty(FPropertyChangedEvent& PropertyChangedEvent)
+{
+	Super::PostEditChangeProperty(PropertyChangedEvent);
+
+	UProperty* PropertyThatChanged = PropertyChangedEvent.Property;
+
+	if ( PropertyThatChanged != NULL )
+	{
+		// if the blueprint has changed, recreate the AnimInstance
+		if ( PropertyThatChanged->GetFName() == GET_MEMBER_NAME_CHECKED( USkeletalMeshComponent, AnimationMode ) )
+		{
+			if (AnimationMode == EAnimationMode::AnimationBlueprint)
+			{
+				if (AnimBlueprintGeneratedClass == NULL)
+				{
+					ClearAnimScriptInstance();
+				}
+				else
+				{
+					if (NeedToSpawnAnimScriptInstance(false))
+					{
+						SCOPE_CYCLE_COUNTER(STAT_AnimSpawnTime);
+						AnimScriptInstance = NewObject<UAnimInstance>(this, AnimBlueprintGeneratedClass);
+						AnimScriptInstance->InitializeAnimation();
+					}
+				}
+			}
+		}
+
+		if ( PropertyThatChanged->GetFName() == GET_MEMBER_NAME_CHECKED( USkeletalMeshComponent, AnimBlueprintGeneratedClass ) )
+		{
+			InitAnim(false);
+		}
+
+		if(PropertyThatChanged->GetFName() == GET_MEMBER_NAME_CHECKED( USkeletalMeshComponent, SkeletalMesh))
+		{
+			ValidateAnimation();
+
+			if(OnSkeletalMeshPropertyChanged.IsBound())
+			{
+				OnSkeletalMeshPropertyChanged.Broadcast();
+			}
+		}
+
+		// when user changes simulate physics, just make sure to update blendphysics together
+		// bBlendPhysics isn't the editor exposed property, it should work with simulate physics
+		if ( PropertyThatChanged->GetFName() == GET_MEMBER_NAME_CHECKED( FBodyInstance, bSimulatePhysics ))
+		{
+			bBlendPhysics = BodyInstance.bSimulatePhysics;
+		}
+
+		if ( PropertyThatChanged->GetFName() == GET_MEMBER_NAME_CHECKED( FSingleAnimationPlayData, AnimToPlay ))
+		{
+			// make sure the animation skeleton matches the current skeletalmesh
+			if (AnimationData.AnimToPlay != NULL && SkeletalMesh && AnimationData.AnimToPlay->GetSkeleton() != SkeletalMesh->Skeleton)
+			{
+				UE_LOG(LogAnimation, Warning, TEXT("Invalid animation"));
+				AnimationData.AnimToPlay = NULL;
+			}
+			else
+			{
+				PlayAnimation(AnimationData.AnimToPlay, false);
+			}
+		}
+
+		if ( PropertyThatChanged->GetFName() == GET_MEMBER_NAME_CHECKED( FSingleAnimationPlayData, SavedPosition ))
+		{
+			AnimationData.ValidatePosition();
+			SetPosition(AnimationData.SavedPosition, false);
+		}
+	}
+}
+
+void USkeletalMeshComponent::LoadedFromAnotherClass(const FName& OldClassName)
+{
+	Super::LoadedFromAnotherClass(OldClassName);
+
+	if(GetLinkerUE4Version() < VER_UE4_REMOVE_SINGLENODEINSTANCE)
+	{
+		static FName SingleAnimSkeletalComponent_NAME(TEXT("SingleAnimSkeletalComponent"));
+
+		if(OldClassName == SingleAnimSkeletalComponent_NAME)
+		{
+			SetAnimationMode(EAnimationMode::Type::AnimationSingleNode);
+
+			// support old compatibility code that changed variable name
+			if (SequenceToPlay_DEPRECATED!=NULL && AnimToPlay_DEPRECATED== NULL)
+			{
+				AnimToPlay_DEPRECATED = SequenceToPlay_DEPRECATED;
+				SequenceToPlay_DEPRECATED = NULL;
+			}
+
+			AnimationData.AnimToPlay = AnimToPlay_DEPRECATED;
+			AnimationData.bSavedLooping = bDefaultLooping_DEPRECATED;
+			AnimationData.bSavedPlaying = bDefaultPlaying_DEPRECATED;
+			AnimationData.SavedPosition = DefaultPosition_DEPRECATED;
+			AnimationData.SavedPlayRate = DefaultPlayRate_DEPRECATED;
+
+			MarkPackageDirty();
+		}
+	}
+}
+#endif // WITH_EDITOR
+
+void USkeletalMeshComponent::TickAnimation(float DeltaTime)
+{
+	SCOPE_CYCLE_COUNTER(STAT_AnimTickTime);
+	SCOPE_CYCLE_COUNTER(STAT_AnimGameThreadTime);
+	if (SkeletalMesh != NULL)
+	{
+		if (AnimScriptInstance != NULL)
+		{
+			// Tick the animation
+			AnimScriptInstance->UpdateAnimation(DeltaTime * GlobalAnimRateScale);
+
+			// TODO @LinaH - I've hit access violations due to AnimScriptInstance being NULL after this, probably due to
+			// AnimNotifies?  Please take a look and fix as we discussed.  Temporary fix:
+			if (AnimScriptInstance != NULL)
+			{
+				// now all tick/trigger/kismet is done
+				// add MorphTarget Curves from Kismet driven or any other source
+				// and overwrite if it exists
+				// Tick always should maintain this list, not Evaluate
+				for( auto Iter = MorphTargetCurves.CreateConstIterator(); Iter; ++Iter )
+				{
+					float *CurveValPtr = AnimScriptInstance->MorphTargetCurves.Find(Iter.Key());
+					if ( CurveValPtr )
+					{
+						// override the value if Kismet request was made
+						*CurveValPtr = Iter.Value();
+					}
+					else
+					{
+						AnimScriptInstance->MorphTargetCurves.Add(Iter.Key(), Iter.Value());
+					}				
+				}
+
+				//Update material parameters
+				UpdateMaterialParameters();
+			}
+		}
+	}
+}
+
+void USkeletalMeshComponent::UpdateMaterialParameters()
+{
+	if(AnimScriptInstance->MaterialParameterCurves.Num() > 0)
+	{
+		for(auto Iter = AnimScriptInstance->MaterialParameterCurves.CreateConstIterator(); Iter; ++Iter)
+		{
+			FName ParameterName = Iter.Key();
+			float ParameterValue = Iter.Value();
+
+			UE_LOG(LogAnimation, Verbose, TEXT("Material Parameter change by Animation (%s : %0.2f)"), *ParameterName.ToString(), ParameterValue);
+			for(int32 MaterialIndex = 0; MaterialIndex < GetNumMaterials(); ++MaterialIndex)
+			{
+				UMaterialInterface* MaterialInterface = GetMaterial(MaterialIndex);
+				if(MaterialInterface)
+				{
+					float TestValue; //not used but needed for GetScalarParameterValue call
+					if(MaterialInterface->GetScalarParameterValue(ParameterName, TestValue))
+					{
+						UMaterialInstanceDynamic* DynamicMaterial = Cast<UMaterialInstanceDynamic>(MaterialInterface);
+						if(!DynamicMaterial) //Is it already a UMaterialInstanceDynamic (ie we used it last tick)
+						{
+							DynamicMaterial = CreateAndSetMaterialInstanceDynamic(MaterialIndex);
+						}
+						DynamicMaterial->SetScalarParameterValue(ParameterName, ParameterValue);
+
+						// we don't break here because we can have multiple materials wanted to be driven by same parameter
+					}
+				}
+			}
+		}
+	}
+}
+
+bool USkeletalMeshComponent::UpdateLODStatus()
+{
+	if ( Super::UpdateLODStatus() )
+	{
+		bRequiredBonesUpToDate = false;
+
+#if WITH_APEX_CLOTHING
+		SetClothingLOD(PredictedLODLevel);
+#endif // #if WITH_APEX_CLOTHING
+		return true;
+	}
+
+	return false;
+}
+
+bool USkeletalMeshComponent::ShouldUpdateTransform(bool bLODHasChanged) const
+{
+#if WITH_EDITOR
+	// If we're in an editor world (Non running, WorldType will be PIE when simulating or in PIE) then we only want transform updates on LOD changes as the
+	// animation isn't running so it would just waste CPU time
+	if(GetWorld()->WorldType == EWorldType::Editor && !bLODHasChanged)
+	{
+		return false;
+	}
+#endif
+
+	// If forcing RefPose we can skip updating the skeleton for perf, except if it's using MorphTargets.
+	const bool bSkipBecauseOfRefPose = bForceRefpose && bOldForceRefPose && (MorphTargetCurves.Num() == 0) && ((AnimScriptInstance)? AnimScriptInstance->MorphTargetCurves.Num() == 0 : true);
+
+	// LOD changing should always trigger an update.
+	return (bLODHasChanged || (!bNoSkeletonUpdate && !bSkipBecauseOfRefPose && Super::ShouldUpdateTransform(bLODHasChanged)));
+}
+
+bool USkeletalMeshComponent::ShouldTickPose() const
+{
+	// When we stop root motion we go back to ticking after CharacterMovement. Unfortunately that means that we could tick twice that frame.
+	// So only enforce a single tick per frame.
+	const bool bAlreadyTickedThisFrame = PoseTickedThisFrame();
+	return (Super::ShouldTickPose() && IsRegistered() && AnimScriptInstance && !bAutonomousTickPose && !bPauseAnims && GetWorld()->AreActorsInitialized() && !bNoSkeletonUpdate && !bAlreadyTickedThisFrame);
+}
+
+void USkeletalMeshComponent::TickPose(float DeltaTime, bool bNeedsValidRootMotion)
+{
+	Super::TickPose(DeltaTime, bNeedsValidRootMotion);
+
+	if (!bEnableUpdateRateOptimizations || !AnimUpdateRateParams->ShouldSkipUpdate())
+	{
+		float TimeAdjustment = AnimUpdateRateParams->GetTimeAdjustment();
+		TickAnimation(DeltaTime + TimeAdjustment);
+		LastPoseTickTime = GetWorld()->TimeSeconds;
+	}
+	//return bPoseTicked;
+}
+
+void USkeletalMeshComponent::TickComponent(float DeltaTime, enum ELevelTick TickType, FActorComponentTickFunction *ThisTickFunction)
+{
+	Super::TickComponent(DeltaTime, TickType, ThisTickFunction);
+
+	// Update bOldForceRefPose
+	bOldForceRefPose = bForceRefpose;
+}
+
+
+/** 
+ *	Utility for taking two arrays of bone indices, which must be strictly increasing, and finding the intersection between them.
+ *	That is - any item in the output should be present in both A and B. Output is strictly increasing as well.
+ */
+static void IntersectBoneIndexArrays(TArray<FBoneIndexType>& Output, const TArray<FBoneIndexType>& A, const TArray<FBoneIndexType>& B)
+{
+	int32 APos = 0;
+	int32 BPos = 0;
+	while(	APos < A.Num() && BPos < B.Num() )
+	{
+		// If value at APos is lower, increment APos.
+		if( A[APos] < B[BPos] )
+		{
+			APos++;
+		}
+		// If value at BPos is lower, increment APos.
+		else if( B[BPos] < A[APos] )
+		{
+			BPos++;
+		}
+		// If they are the same, put value into output, and increment both.
+		else
+		{
+			Output.Add( A[APos] );
+			APos++;
+			BPos++;
+		}
+	}
+}
+
+
+void USkeletalMeshComponent::FillSpaceBases(const USkeletalMesh* InSkeletalMesh, const TArray<FTransform>& SourceAtoms, TArray<FTransform>& DestSpaceBases) const
+{
+	SCOPE_CYCLE_COUNTER(STAT_SkelComposeTime);
+
+	if( !InSkeletalMesh )
+	{
+		return;
+	}
+
+	// right now all this does is populate DestSpaceBases
+	check( InSkeletalMesh->RefSkeleton.GetNum() == SourceAtoms.Num());
+	check( InSkeletalMesh->RefSkeleton.GetNum() == DestSpaceBases.Num());
+
+	const int32 NumBones = SourceAtoms.Num();
+
+#if (UE_BUILD_DEBUG || UE_BUILD_DEVELOPMENT)
+	/** Keep track of which bones have been processed for fast look up */
+	TArray<uint8> BoneProcessed;
+	BoneProcessed.AddZeroed(NumBones);
+#endif
+
+	const FTransform* LocalTransformsData = SourceAtoms.GetData();
+	FTransform* SpaceBasesData = DestSpaceBases.GetData();
+
+	// First bone is always root bone, and it doesn't have a parent.
+	{
+		check(FillSpaceBasesRequiredBones[0] == 0);
+		DestSpaceBases[0] = SourceAtoms[0];
+
+#if (UE_BUILD_DEBUG || UE_BUILD_DEVELOPMENT)
+		// Mark bone as processed
+		BoneProcessed[0] = 1;
+#endif
+	}
+
+	for (int32 i = 1; i<FillSpaceBasesRequiredBones.Num(); i++)
+	{
+		const int32 BoneIndex = FillSpaceBasesRequiredBones[i];
+		FTransform* SpaceBase = SpaceBasesData + BoneIndex;
+
+		FPlatformMisc::Prefetch(SpaceBase);
+
+#if (UE_BUILD_DEBUG || UE_BUILD_DEVELOPMENT)
+		// Mark bone as processed
+		BoneProcessed[BoneIndex] = 1;
+#endif
+		// For all bones below the root, final component-space transform is relative transform * component-space transform of parent.
+		const int32 ParentIndex = InSkeletalMesh->RefSkeleton.GetParentIndex(BoneIndex);
+		FTransform* ParentSpaceBase = SpaceBasesData + ParentIndex;
+		FPlatformMisc::Prefetch(ParentSpaceBase);
+
+#if (UE_BUILD_DEBUG || UE_BUILD_DEVELOPMENT)
+		// Check the precondition that Parents occur before Children in the RequiredBones array.
+		checkSlow(BoneProcessed[ParentIndex] == 1);
+#endif
+		FTransform::Multiply(SpaceBase, LocalTransformsData + BoneIndex, ParentSpaceBase);
+
+		checkSlow(SpaceBase->IsRotationNormalized());
+		checkSlow(!SpaceBase->ContainsNaN());
+	}
+
+	/**
+	 * Normalize rotations.
+	 * We want to remove any loss of precision due to accumulation of error.
+	 * i.e. A componentSpace transform is the accumulation of all of its local space parents. The further down the chain, the greater the error.
+	 * SpaceBases are used by external systems, we feed this to PhysX, send this to gameplay through bone and socket queries, etc.
+	 * So this is a good place to make sure all transforms are normalized.
+	 */
+	FAnimationRuntime::NormalizeRotations(DestSpaceBases);
+}
+
+/** Takes sorted array Base and then adds any elements from sorted array Insert which is missing from it, preserving order.
+ * this assumes both arrays are sorted and contain unique bone indices. */
+static void MergeInBoneIndexArrays(TArray<FBoneIndexType>& BaseArray, TArray<FBoneIndexType>& InsertArray)
+{
+	// Then we merge them into the array of required bones.
+	int32 BaseBonePos = 0;
+	int32 InsertBonePos = 0;
+
+	// Iterate over each of the bones we need.
+	while( InsertBonePos < InsertArray.Num() )
+	{
+		// Find index of physics bone
+		FBoneIndexType InsertBoneIndex = InsertArray[InsertBonePos];
+
+		// If at end of BaseArray array - just append.
+		if( BaseBonePos == BaseArray.Num() )
+		{
+			BaseArray.Add(InsertBoneIndex);
+			BaseBonePos++;
+			InsertBonePos++;
+		}
+		// If in the middle of BaseArray, merge together.
+		else
+		{
+			// Check that the BaseArray array is strictly increasing, otherwise merge code does not work.
+			check( BaseBonePos == 0 || BaseArray[BaseBonePos-1] < BaseArray[BaseBonePos] );
+
+			// Get next required bone index.
+			FBoneIndexType BaseBoneIndex = BaseArray[BaseBonePos];
+
+			// We have a bone in BaseArray not required by Insert. Thats ok - skip.
+			if( BaseBoneIndex < InsertBoneIndex )
+			{
+				BaseBonePos++;
+			}
+			// Bone required by Insert is in 
+			else if( BaseBoneIndex == InsertBoneIndex )
+			{
+				BaseBonePos++;
+				InsertBonePos++;
+			}
+			// Bone required by Insert is missing - insert it now.
+			else // BaseBoneIndex > InsertBoneIndex
+			{
+				BaseArray.InsertUninitialized(BaseBonePos);
+				BaseArray[BaseBonePos] = InsertBoneIndex;
+
+				BaseBonePos++;
+				InsertBonePos++;
+			}
+		}
+	}
+}
+
+
+void USkeletalMeshComponent::RecalcRequiredBones(int32 LODIndex)
+{
+	if (!SkeletalMesh)
+	{
+		return;
+	}
+
+	FSkeletalMeshResource* SkelMeshResource = GetSkeletalMeshResource();
+	check(SkelMeshResource);
+
+	// The list of bones we want is taken from the predicted LOD level.
+	FStaticLODModel& LODModel = SkelMeshResource->LODModels[LODIndex];
+	RequiredBones = LODModel.RequiredBones;
+	
+	const UPhysicsAsset* const PhysicsAsset = GetPhysicsAsset();
+	// If we have a PhysicsAsset, we also need to make sure that all the bones used by it are always updated, as its used
+	// by line checks etc. We might also want to kick in the physics, which means having valid bone transforms.
+	if(PhysicsAsset)
+	{
+		TArray<FBoneIndexType> PhysAssetBones;
+		for(int32 i=0; i<PhysicsAsset->BodySetup.Num(); i++ )
+		{
+			int32 PhysBoneIndex = SkeletalMesh->RefSkeleton.FindBoneIndex( PhysicsAsset->BodySetup[i]->BoneName );
+			if(PhysBoneIndex != INDEX_NONE)
+			{
+				PhysAssetBones.Add(PhysBoneIndex);
+			}	
+		}
+
+		// Then sort array of required bones in hierarchy order
+		PhysAssetBones.Sort();
+
+		// Make sure all of these are in RequiredBones.
+		MergeInBoneIndexArrays(RequiredBones, PhysAssetBones);
+	}
+
+	// Make sure that bones with per-poly collision are also always updated.
+	// TODO UE4
+
+	// Purge invisible bones and their children
+	// this has to be done before mirror table check/physics body checks
+	// mirror table/phys body ones has to be calculated
+	if (ShouldUpdateBoneVisibility())
+	{
+		check(BoneVisibilityStates.Num() == GetNumSpaceBases());
+
+		int32 VisibleBoneWriteIndex = 0;
+		for (int32 i = 0; i < RequiredBones.Num(); ++i)
+		{
+			FBoneIndexType CurBoneIndex = RequiredBones[i];
+
+			// Current bone visible?
+			if (BoneVisibilityStates[CurBoneIndex] == BVS_Visible)
+			{
+				RequiredBones[VisibleBoneWriteIndex++] = CurBoneIndex;
+			}
+		}
+
+		// Remove any trailing junk in the RequiredBones array
+		const int32 NumBonesHidden = RequiredBones.Num() - VisibleBoneWriteIndex;
+		if (NumBonesHidden > 0)
+		{
+			RequiredBones.RemoveAt(VisibleBoneWriteIndex, NumBonesHidden);
+		}
+	}
+
+	// Add in any bones that may be required when mirroring.
+	// JTODO: This is only required if there are mirroring nodes in the tree, but hard to know...
+	if(SkeletalMesh->SkelMirrorTable.Num() > 0 && 
+		SkeletalMesh->SkelMirrorTable.Num() == LocalAtoms.Num())
+	{
+		TArray<FBoneIndexType> MirroredDesiredBones;
+		MirroredDesiredBones.AddUninitialized(RequiredBones.Num());
+
+		// Look up each bone in the mirroring table.
+		for(int32 i=0; i<RequiredBones.Num(); i++)
+		{
+			MirroredDesiredBones[i] = SkeletalMesh->SkelMirrorTable[RequiredBones[i]].SourceIndex;
+		}
+
+		// Sort to ensure strictly increasing order.
+		MirroredDesiredBones.Sort();
+
+		// Make sure all of these are in RequiredBones, and 
+		MergeInBoneIndexArrays(RequiredBones, MirroredDesiredBones);
+	}
+
+	TArray<FBoneIndexType> NeededBonesForFillSpaceBases;
+	{
+		TArray<FBoneIndexType> ForceAnimatedSocketBones;
+
+		for (const USkeletalMeshSocket* Socket : SkeletalMesh->GetActiveSocketList())
+		{
+			int32 BoneIndex = SkeletalMesh->RefSkeleton.FindBoneIndex(Socket->BoneName);
+			if (BoneIndex != INDEX_NONE)
+			{
+				if (Socket->bForceAlwaysAnimated)
+				{
+					ForceAnimatedSocketBones.AddUnique(BoneIndex);
+				}
+				else
+				{
+					NeededBonesForFillSpaceBases.AddUnique(BoneIndex);
+				}
+			}
+		}
+
+		// Then sort array of required bones in hierarchy order
+		ForceAnimatedSocketBones.Sort();
+
+		// Make sure all of these are in RequiredBones.
+		MergeInBoneIndexArrays(RequiredBones, ForceAnimatedSocketBones);
+	}
+
+
+	// Ensure that we have a complete hierarchy down to those bones.
+	FAnimationRuntime::EnsureParentsPresent(RequiredBones, SkeletalMesh);
+
+	FillSpaceBasesRequiredBones.Empty(RequiredBones.Num() + NeededBonesForFillSpaceBases.Num());
+	FillSpaceBasesRequiredBones = RequiredBones;
+	
+	NeededBonesForFillSpaceBases.Sort();
+	MergeInBoneIndexArrays(FillSpaceBasesRequiredBones, NeededBonesForFillSpaceBases);
+	FAnimationRuntime::EnsureParentsPresent(FillSpaceBasesRequiredBones, SkeletalMesh);
+
+	// Sanitise bones that we aren't going to be updating
+	for (int32 BoneIndex = 0; BoneIndex < LocalAtoms.Num(); ++BoneIndex)
+	{
+		if (!RequiredBones.Contains(BoneIndex))
+		{
+			LocalAtoms[BoneIndex] = SkeletalMesh->RefSkeleton.GetRefBonePose()[BoneIndex];
+		}
+	}
+
+	// make sure animation requiredBone to mark as dirty
+	if (AnimScriptInstance)
+	{
+		AnimScriptInstance->RecalcRequiredBones();
+	}
+
+	bRequiredBonesUpToDate = true;
+
+	// Invalidate cached bones.
+	CachedLocalAtoms.Empty();
+	CachedSpaceBases.Empty();
+}
+
+void USkeletalMeshComponent::EvaluateAnimation(const USkeletalMesh* InSkeletalMesh, UAnimInstance* InAnimInstance, TArray<FTransform>& OutLocalAtoms, TArray<FActiveVertexAnim>& OutVertexAnims, FVector& OutRootBoneTranslation) const
+{
+	SCOPE_CYCLE_COUNTER(STAT_AnimBlendTime);
+
+	if( !InSkeletalMesh )
+	{
+		return;
+	}
+
+	// We can only evaluate animation if RequiredBones is properly setup for the right mesh!
+	if( InSkeletalMesh->Skeleton && InAnimInstance
+		&& ensure(bRequiredBonesUpToDate)
+		&& InAnimInstance->RequiredBones.IsValid()
+		&& (InAnimInstance->RequiredBones.GetAsset() == InSkeletalMesh) )
+	{
+		if( !bForceRefpose )
+		{
+			// Create an evaluation context
+			FPoseContext EvaluationContext(InAnimInstance);
+			EvaluationContext.ResetToRefPose();
+			
+			// Run the anim blueprint
+			InAnimInstance->EvaluateAnimation(EvaluationContext);
+
+			// can we avoid that copy?
+			if( EvaluationContext.Pose.Bones.Num() > 0 )
+			{
+				OutLocalAtoms = EvaluationContext.Pose.Bones;
+
+				// Make sure rotations are normalized to account for accumulation of errors.
+				FAnimationRuntime::NormalizeRotations(OutLocalAtoms);
+			}
+			else
+			{
+				FAnimationRuntime::FillWithRefPose(OutLocalAtoms, InAnimInstance->RequiredBones);
+			}
+		}
+		else
+		{
+			FAnimationRuntime::FillWithRefPose(OutLocalAtoms, InAnimInstance->RequiredBones);
+		}
+
+		OutVertexAnims = UpdateActiveVertexAnims(InSkeletalMesh, InAnimInstance->MorphTargetCurves, InAnimInstance->VertexAnims);
+	}
+	else
+	{
+		OutLocalAtoms = InSkeletalMesh->RefSkeleton.GetRefBonePose();
+
+		// if it's only morph, there is no reason to blend
+		if ( MorphTargetCurves.Num() > 0 )
+		{
+			TArray<struct FActiveVertexAnim> EmptyVertexAnims;
+			OutVertexAnims = UpdateActiveVertexAnims(InSkeletalMesh, MorphTargetCurves, EmptyVertexAnims);
+		}
+	}
+
+	// Remember the root bone's translation so we can move the bounds.
+	OutRootBoneTranslation = OutLocalAtoms[0].GetTranslation() - InSkeletalMesh->RefSkeleton.GetRefBonePose()[0].GetTranslation();
+}
+
+void USkeletalMeshComponent::UpdateSlaveComponent()
+{
+	check (MasterPoseComponent.IsValid());
+
+	if(MasterPoseComponent->IsA(USkeletalMeshComponent::StaticClass()))
+	{
+		USkeletalMeshComponent* MasterSMC= CastChecked<USkeletalMeshComponent>(MasterPoseComponent.Get());
+
+		if ( MasterSMC->AnimScriptInstance )
+		{
+			ActiveVertexAnims = UpdateActiveVertexAnims(SkeletalMesh, MasterSMC->AnimScriptInstance->MorphTargetCurves, MasterSMC->AnimScriptInstance->VertexAnims);
+		}
+	}
+
+	Super::UpdateSlaveComponent();
+}
+
+void USkeletalMeshComponent::PerformAnimationEvaluation(const USkeletalMesh* InSkeletalMesh, UAnimInstance* InAnimInstance, TArray<FTransform>& OutSpaceBases, TArray<FTransform>& OutLocalAtoms, TArray<FActiveVertexAnim>& OutVertexAnims, FVector& OutRootBoneTranslation) const
+{
+	SCOPE_CYCLE_COUNTER(STAT_PerformAnimEvaluation);
+	// Can't do anything without a SkeletalMesh
+	// Do nothing more if no bones in skeleton.
+	if (!InSkeletalMesh || OutSpaceBases.Num() == 0)
+	{
+		return;
+	}
+
+	// evaluate pure animations, and fill up LocalAtoms
+	EvaluateAnimation(InSkeletalMesh, InAnimInstance, OutLocalAtoms, OutVertexAnims, OutRootBoneTranslation);
+	// Fill SpaceBases from LocalAtoms
+	FillSpaceBases(InSkeletalMesh, OutLocalAtoms, OutSpaceBases);
+}
+
+const TCHAR* B(bool b)
+{
+	return b ? TEXT("true") : TEXT("false");
+}
+
+void USkeletalMeshComponent::RefreshBoneTransforms(FActorComponentTickFunction* TickFunction)
+{
+	SCOPE_CYCLE_COUNTER(STAT_RefreshBoneTransforms);
+	SCOPE_CYCLE_COUNTER(STAT_AnimGameThreadTime);
+
+	check(IsInGameThread()); //Only want to call this from the game thread as we set up tasks etc
+
+	if (!SkeletalMesh || GetNumSpaceBases() == 0)
+	{
+		return;
+	}
+
+	const bool bDoEvaluationRateOptimization = bEnableUpdateRateOptimizations && AnimUpdateRateParams->DoEvaluationRateOptimizations();
+
+	//Handle update rate optimization setup
+	//Dont mark cache as invalid if we aren't performing optimization anyway
+	const bool bInvalidCachedBones = bDoEvaluationRateOptimization &&
+									 ((LocalAtoms.Num() != SkeletalMesh->RefSkeleton.GetNum())
+									 || (LocalAtoms.Num() != CachedLocalAtoms.Num())
+									 || (GetNumSpaceBases() != CachedSpaceBases.Num()));
+
+	const bool bShouldDoEvaluation = !bDoEvaluationRateOptimization || bInvalidCachedBones || !AnimUpdateRateParams->ShouldSkipEvaluation();
+
+	const bool bDoPAE = !!CVarUseParallelAnimationEvaluation.GetValueOnGameThread() && FApp::ShouldUseThreadingForPerformance();
+
+	const bool bDoParallelEvaluation = AnimEvaluationContext.bDoEvaluation && TickFunction && bDoPAE;
+
+	if (IsValidRef(ParallelAnimationEvaluationTask))
+	{
+		//Are already processing eval on another thread, we wait for eval thread to finish
+		if (!bDoParallelEvaluation) //we are already running parallel evaluation, so if we are going to try it again just return
+		{
+			//If we are not going to attempt parallel evaluation then wait here for the existing task to finish and then
+			//perform complete directly so that when we return all the calculations are complete
+			FTaskGraphInterface::Get().WaitUntilTaskCompletes(ParallelAnimationEvaluationTask, ENamedThreads::GameThread);
+			CompleteParallelAnimationEvaluation(); //Perform completion now
+		}
+		return;
+	}
+
+	AActor* Owner = GetOwner();
+	UE_LOG(LogAnimation, Verbose, TEXT("RefreshBoneTransforms(%s)"), *GetNameSafe(Owner));
+
+	// Recalculate the RequiredBones array, if necessary
+	if (!bRequiredBonesUpToDate)
+	{
+		RecalcRequiredBones(PredictedLODLevel);
+	}
+
+	AnimEvaluationContext.SkeletalMesh = SkeletalMesh;
+	AnimEvaluationContext.AnimInstance = AnimScriptInstance;
+
+	AnimEvaluationContext.bDoEvaluation = bShouldDoEvaluation;
+	
+	AnimEvaluationContext.bDoInterpolation = bDoEvaluationRateOptimization && !bInvalidCachedBones && AnimUpdateRateParams->ShouldInterpolateSkippedFrames();
+	AnimEvaluationContext.bDuplicateToCacheBones = bInvalidCachedBones || (bDoEvaluationRateOptimization && AnimEvaluationContext.bDoEvaluation && !AnimEvaluationContext.bDoInterpolation);
+
+	if (!bDoEvaluationRateOptimization)
+	{
+		//If we aren't optimizing clear the cached local atoms
+		CachedLocalAtoms.Empty();
+		CachedSpaceBases.Empty();
+	}
+
+	if (bDoParallelEvaluation)
+	{
+		if (SkeletalMesh->RefSkeleton.GetNum() != AnimEvaluationContext.LocalAtoms.Num())
+		{
+			// Initialize Parallel Task arrays
+			AnimEvaluationContext.LocalAtoms = LocalAtoms;
+			AnimEvaluationContext.SpaceBases = GetSpaceBases();
+			AnimEvaluationContext.VertexAnims = ActiveVertexAnims;
+		}
+
+		// start parallel work
+		check(!IsValidRef(ParallelAnimationEvaluationTask));
+		ParallelAnimationEvaluationTask = TGraphTask<FParallelAnimationEvaluationTask>::CreateTask().ConstructAndDispatchWhenReady(this);
+
+		// set up a task to run on the game thread to accept the results
+		FGraphEventArray Prerequistes;
+		Prerequistes.Add(ParallelAnimationEvaluationTask);
+		FGraphEventRef TickCompletionEvent = TGraphTask<FParallelAnimationCompletionTask>::CreateTask(&Prerequistes).ConstructAndDispatchWhenReady(this);
+
+		TickFunction->GetCompletionHandle()->DontCompleteUntil(TickCompletionEvent);
+	}
+	else
+	{
+		if (AnimEvaluationContext.bDoEvaluation)
+		{
+			if (AnimEvaluationContext.bDoInterpolation)
+			{
+				PerformAnimationEvaluation(SkeletalMesh, AnimScriptInstance, CachedSpaceBases, CachedLocalAtoms, ActiveVertexAnims, RootBoneTranslation);
+			}
+			else
+			{
+				PerformAnimationEvaluation(SkeletalMesh, AnimScriptInstance, GetEditableSpaceBases(), LocalAtoms, ActiveVertexAnims, RootBoneTranslation);
+			}
+		}
+		else if (!AnimEvaluationContext.bDoInterpolation)
+		{
+			LocalAtoms = CachedLocalAtoms;
+			GetEditableSpaceBases() = CachedSpaceBases;
+		}
+
+		PostAnimEvaluation(AnimEvaluationContext);
+	}
+}
+
+void USkeletalMeshComponent::PostAnimEvaluation(FAnimationEvaluationContext& EvaluationContext)
+{
+	if (AnimScriptInstance)
+	{
+		AnimScriptInstance->PostAnimEvaluation();
+	}
+
+	AnimEvaluationContext.Clear();
+
+	SCOPE_CYCLE_COUNTER(STAT_PostAnimEvaluation);
+	if (EvaluationContext.bDuplicateToCacheBones)
+	{
+		CachedSpaceBases = GetEditableSpaceBases();
+		CachedLocalAtoms = LocalAtoms;
+	}
+
+	if (EvaluationContext.bDoInterpolation)
+	{
+		SCOPE_CYCLE_COUNTER(STAT_InterpolateSkippedFrames);
+
+		const float Alpha = AnimUpdateRateParams->GetInterpolationAlpha();
+		FAnimationRuntime::LerpBoneTransforms(LocalAtoms, CachedLocalAtoms, Alpha, RequiredBones);
+		if (bDoubleBufferedBlendSpaces)
+		{
+			// We need to prep our space bases for interp (TODO: Would a new Lerp function that took
+			// separate input and output be quicker than current copy + lerp in place?)
+			GetEditableSpaceBases() = GetSpaceBases();
+		}
+		FAnimationRuntime::LerpBoneTransforms(GetEditableSpaceBases(), CachedSpaceBases, Alpha, RequiredBones);
+	}
+
+	bNeedToFlipSpaceBaseBuffers = true;
+
+	// Transforms updated, cached local bounds are now out of date.
+	InvalidateCachedBounds();
+
+	// update physics data from animated data
+	UpdateKinematicBonesToAnim(GetEditableSpaceBases(), false, true);
+	UpdateRBJointMotors();
+
+	// @todo anim : hack TTP 224385	ANIM: Skeletalmesh double buffer
+	// this is problem because intermediate buffer changes physics position as well
+	// this causes issue where a half of frame, physics position is fixed with anim pose, and the other half is real simulated position
+	// if you enable physics in tick, since that's before physics update, you'll get animation pose dominating physics pose, which isn't what you want. (Or what you'll see)
+	// so do not update transform if physics is on. This problem will be solved by double buffer, when we keep one buffer for intermediate, and the other buffer for result query
+	if (!bBlendPhysics && !ShouldBlendPhysicsBones())
+	{
+		SCOPE_CYCLE_COUNTER(STAT_UpdateLocalToWorldAndOverlaps);
+
+		// Updated last good bone positions
+		FlipEditableSpaceBases();
+
+		// New bone positions need to be sent to render thread
+		UpdateComponentToWorld();
+
+		// animation often change overlap. 
+		UpdateOverlaps();
+	}
+
+	MarkRenderDynamicDataDirty();
+}
+
+//
+//	USkeletalMeshComponent::UpdateBounds
+//
+
+void USkeletalMeshComponent::UpdateBounds()
+{
+#if WITH_EDITOR
+	FBoxSphereBounds OriginalBounds = Bounds; // Save old bounds
+#endif
+
+	// if use parent bound if attach parent exists, and the flag is set
+	// since parents tick first before child, this should work correctly
+	if ( bUseAttachParentBound && AttachParent != NULL )
+	{
+		Bounds = AttachParent->Bounds;
+	}
+	else
+	{
+		// fixme laurent - extend concept of LocalBounds to all SceneComponent
+		// as rendered calls CalcBounds*() directly in FScene::UpdatePrimitiveTransform, which is pretty expensive for SkelMeshes.
+		// No need to calculated that again, just use cached local bounds.
+		if( bCachedLocalBoundsUpToDate )
+		{
+			Bounds = CachedLocalBounds.TransformBy(ComponentToWorld);
+		}
+		else
+		{
+			// Calculate new bounds
+			Bounds = CalcBounds(ComponentToWorld);
+
+			bCachedLocalBoundsUpToDate = true;
+			CachedLocalBounds = Bounds.TransformBy(ComponentToWorld.Inverse());
+		}
+	}
+
+#if WITH_EDITOR
+	// If bounds have changed (in editor), trigger data rebuild
+	if ( IsRegistered() && (World != NULL) && !GetWorld()->IsGameWorld() && 
+		(OriginalBounds.Origin.Equals(Bounds.Origin) == false || OriginalBounds.BoxExtent.Equals(Bounds.BoxExtent) == false) )
+	{
+		GEngine->TriggerStreamingDataRebuild();
+	}
+#endif
+}
+
+FBoxSphereBounds USkeletalMeshComponent::CalcBounds(const FTransform& LocalToWorld) const
+{
+	FVector RootBoneOffset;
+
+	// if to use MasterPoseComponent's fixed skel bounds, 
+	// send MasterPoseComponent's Root Bone Translation
+	if(MasterPoseComponent.IsValid() && MasterPoseComponent->SkeletalMesh &&
+		MasterPoseComponent->bComponentUseFixedSkelBounds &&
+		MasterPoseComponent->IsA((USkeletalMeshComponent::StaticClass())))
+	{
+		USkeletalMeshComponent* BaseComponent = CastChecked<USkeletalMeshComponent>(MasterPoseComponent.Get());
+		RootBoneOffset = BaseComponent->RootBoneTranslation; // Adjust bounds by root bone translation
+	}
+	else
+	{
+		RootBoneOffset = RootBoneTranslation;
+	}
+
+	FBoxSphereBounds NewBounds = CalcMeshBound( RootBoneOffset, bHasValidBodies, LocalToWorld );
+#if WITH_APEX_CLOTHING
+	AddClothingBounds(NewBounds);
+#endif// #if WITH_APEX_CLOTHING
+	return NewBounds;
+}
+
+void USkeletalMeshComponent::SetSkeletalMesh(USkeletalMesh* InSkelMesh)
+{
+	if (InSkelMesh == SkeletalMesh)
+	{
+		// do nothing if the input mesh is the same mesh we're already using.
+		return;
+	}
+
+	UPhysicsAsset* OldPhysAsset = GetPhysicsAsset();
+
+	Super::SetSkeletalMesh(InSkelMesh);
+
+#if WITH_EDITOR
+	ValidateAnimation();
+#endif
+
+	if (GetPhysicsAsset() != OldPhysAsset && IsPhysicsStateCreated())
+	{
+		RecreatePhysicsState();
+	}
+
+	UpdateHasValidBodies();
+
+	InitAnim(false);
+
+#if WITH_APEX_CLOTHING
+	ValidateClothingActors();
+#endif
+}
+
+bool USkeletalMeshComponent::AllocateTransformData()
+{
+	// Allocate transforms if not present.
+	if ( Super::AllocateTransformData() )
+	{
+		if( LocalAtoms.Num() != SkeletalMesh->RefSkeleton.GetNum() )
+		{
+			LocalAtoms.Empty( SkeletalMesh->RefSkeleton.GetNum() );
+			LocalAtoms.AddUninitialized( SkeletalMesh->RefSkeleton.GetNum() );
+		}
+
+		return true;
+	}
+
+	LocalAtoms.Empty();
+	
+	return false;
+}
+
+void USkeletalMeshComponent::DeallocateTransformData()
+{
+	Super::DeallocateTransformData();
+	LocalAtoms.Empty();
+}
+
+void USkeletalMeshComponent::SetForceRefPose(bool bNewForceRefPose)
+{
+	bForceRefpose = bNewForceRefPose;
+	MarkRenderStateDirty();
+}
+
+void USkeletalMeshComponent::SetAnimInstanceClass(class UClass* NewClass)
+{
+	if (NewClass != NULL)
+	{
+		UAnimBlueprintGeneratedClass* NewGeneratedClass = Cast<UAnimBlueprintGeneratedClass>(NewClass);
+		ensure(NULL != NewGeneratedClass);
+		// set the animation mode
+		AnimationMode = EAnimationMode::Type::AnimationBlueprint;
+
+		if (NewGeneratedClass != AnimBlueprintGeneratedClass)
+		{
+			// Only need to initialize if it hasn't already been set.
+			AnimBlueprintGeneratedClass = NewGeneratedClass;
+			ClearAnimScriptInstance();
+			InitAnim(true);
+		}
+	}
+	else
+	{
+		// Need to clear the instance as well as the blueprint.
+		// @todo is this it?
+		AnimBlueprintGeneratedClass = NULL;
+		ClearAnimScriptInstance();
+	}
+}
+
+UAnimInstance* USkeletalMeshComponent::GetAnimInstance() const
+{
+	return AnimScriptInstance;
+}
+
+FMatrix USkeletalMeshComponent::GetTransformMatrix() const
+{
+	FTransform RootTransform = GetBoneTransform(0);
+	FVector Translation;
+	FQuat Rotation;
+	
+	// if in editor, it should always use localToWorld
+	// if root motion is ignored, use root transform 
+	if( GetWorld()->IsGameWorld() || !SkeletalMesh )
+	{
+		// add root translation info
+		Translation = RootTransform.GetLocation();
+	}
+	else
+	{
+		Translation = ComponentToWorld.TransformPosition(SkeletalMesh->RefSkeleton.GetRefBonePose()[0].GetTranslation());
+	}
+
+	// if root rotation is ignored, use root transform rotation
+	Rotation = RootTransform.GetRotation();
+
+	// now I need to get scale
+	// only LocalToWorld will have scale
+	FVector ScaleVector = ComponentToWorld.GetScale3D();
+
+	Rotation.Normalize();
+	return FScaleMatrix(ScaleVector)*FQuatRotationTranslationMatrix(Rotation, Translation);
+}
+
+void USkeletalMeshComponent::NotifySkelControlBeyondLimit( USkelControlLookAt* LookAt ) {}
+
+
+
+void USkeletalMeshComponent::SkelMeshCompOnParticleSystemFinished( UParticleSystemComponent* PSC )
+{
+	PSC->DetachFromParent();
+	PSC->UnregisterComponent();
+}
+
+
+void USkeletalMeshComponent::HideBone( int32 BoneIndex, EPhysBodyOp PhysBodyOption)
+{
+	Super::HideBone(BoneIndex, PhysBodyOption);
+
+	if (!SkeletalMesh)
+	{
+		return;
+	}
+
+	LocalAtoms[ BoneIndex ].SetScale3D(FVector::ZeroVector);
+	bRequiredBonesUpToDate = false;
+
+	if( PhysBodyOption!=PBO_None )
+	{
+		FName HideBoneName = SkeletalMesh->RefSkeleton.GetBoneName(BoneIndex);
+		if ( PhysBodyOption == PBO_Term )
+		{
+			TermBodiesBelow(HideBoneName);
+		}
+		else if ( PhysBodyOption == PBO_Disable )
+		{
+			// Disable collision
+			// @JTODO
+			//SetCollisionBelow(false, HideBoneName);
+		}
+	}
+}
+
+void USkeletalMeshComponent::UnHideBone( int32 BoneIndex )
+{
+	Super::UnHideBone(BoneIndex);
+
+	if (!SkeletalMesh)
+	{
+		return;
+	}
+
+	LocalAtoms[ BoneIndex ].SetScale3D(FVector(1.0f));
+	bRequiredBonesUpToDate = false;
+
+	FName HideBoneName = SkeletalMesh->RefSkeleton.GetBoneName(BoneIndex);
+	// It's okay to turn this on for terminated bodies
+	// It won't do any if BodyData isn't found
+	// @JTODO
+	//SetCollisionBelow(true, HideBoneName);
+
+}
+
+
+bool USkeletalMeshComponent::IsAnySimulatingPhysics() const
+{
+	for ( int32 BodyIndex=0; BodyIndex<Bodies.Num(); ++BodyIndex )
+	{
+		if (Bodies[BodyIndex]->IsInstanceSimulatingPhysics())
+		{
+			return true;
+		}
+	}
+
+	return false;
+}
+
+/** 
+ * Render bones for debug display
+ */
+void USkeletalMeshComponent::DebugDrawBones(UCanvas* Canvas, bool bSimpleBones) const
+{
+	if (GetWorld()->IsGameWorld() && SkeletalMesh && Canvas && MasterPoseComponent == NULL)
+	{
+		// draw spacebases, we could cache parent bones, but this is mostly debug feature, I'm not caching it right now
+		for ( int32 Index=0; Index<RequiredBones.Num(); ++Index )
+		{
+			int32 BoneIndex = RequiredBones[Index];
+			int32 ParentIndex = SkeletalMesh->RefSkeleton.GetParentIndex(BoneIndex);
+			FTransform BoneTM = (GetSpaceBases()[BoneIndex] * ComponentToWorld);
+			FVector Start, End;
+			FLinearColor LineColor;
+
+			End = BoneTM.GetLocation();
+
+			if (ParentIndex >=0)
+			{
+				Start = (GetSpaceBases()[ParentIndex] * ComponentToWorld).GetLocation();
+				LineColor = FLinearColor::White;
+			}
+			else
+			{
+				Start = ComponentToWorld.GetLocation();
+				LineColor = FLinearColor::Red;
+			}
+
+			if(bSimpleBones)
+			{
+				DrawDebugCanvasLine(Canvas, Start, End, LineColor);
+			}
+			else
+			{
+				static const float SphereRadius = 1.0f;
+
+				//Calc cone size 
+				FVector EndToStart = (Start-End);
+				float ConeLength = EndToStart.Size();
+				float Angle = FMath::RadiansToDegrees(FMath::Atan(SphereRadius / ConeLength));
+
+				DrawDebugCanvasWireSphere(Canvas, End, LineColor, SphereRadius, 10);
+				DrawDebugCanvasWireCone(Canvas, FTransform(FRotationMatrix::MakeFromX(EndToStart)*FTranslationMatrix(End)), ConeLength, Angle, 4, LineColor);
+			}
+
+			RenderAxisGizmo(BoneTM, Canvas);
+		}
+	}
+}
+
+// Render a coordinate system indicator
+void USkeletalMeshComponent::RenderAxisGizmo( const FTransform& Transform, UCanvas* Canvas ) const
+{
+	// Display colored coordinate system axes for this joint.
+	const float AxisLength = 3.75f;
+	const FVector Origin = Transform.GetLocation();
+
+	// Red = X
+	FVector XAxis = Transform.TransformVector( FVector(1.0f,0.0f,0.0f) );
+	XAxis.Normalize();
+	DrawDebugCanvasLine(Canvas, Origin, Origin + XAxis * AxisLength, FLinearColor( 1.f, 0.3f, 0.3f));		
+
+	// Green = Y
+	FVector YAxis = Transform.TransformVector( FVector(0.0f,1.0f,0.0f) );
+	YAxis.Normalize();
+	DrawDebugCanvasLine(Canvas, Origin, Origin + YAxis * AxisLength, FLinearColor( 0.3f, 1.f, 0.3f));	
+
+	// Blue = Z
+	FVector ZAxis = Transform.TransformVector( FVector(0.0f,0.0f,1.0f) );
+	ZAxis.Normalize();
+	DrawDebugCanvasLine(Canvas, Origin, Origin + ZAxis * AxisLength, FLinearColor( 0.3f, 0.3f, 1.f));	
+}
+
+void USkeletalMeshComponent::SetMorphTarget(FName MorphTargetName, float Value, bool bRemoveZeroWeight)
+{
+	float *CurveValPtr = MorphTargetCurves.Find(MorphTargetName);
+	bool bShouldAddToList = !bRemoveZeroWeight || FPlatformMath::Abs(Value) > ZERO_ANIMWEIGHT_THRESH;
+	if ( bShouldAddToList )
+	{
+		if ( CurveValPtr )
+		{
+			// sum up, in the future we might normalize, but for now this just sums up
+			// this won't work well if all of them have full weight - i.e. additive 
+			*CurveValPtr = Value;
+		}
+		else
+		{
+			MorphTargetCurves.Add(MorphTargetName, Value);
+		}
+	}
+	// if less than ZERO_ANIMWEIGHT_THRESH
+	// no reason to keep them on the list
+	else 
+	{
+		// remove if found
+		MorphTargetCurves.Remove(MorphTargetName);
+	}
+}
+
+void USkeletalMeshComponent::ClearMorphTargets()
+{
+	MorphTargetCurves.Empty();
+}
+
+float USkeletalMeshComponent::GetMorphTarget( FName MorphTargetName ) const
+{
+	const float *CurveValPtr = MorphTargetCurves.Find(MorphTargetName);
+	
+	if(CurveValPtr)
+	{
+		return *CurveValPtr;
+	}
+	else
+	{
+		return 0.0f;
+	}
+}
+
+FVector USkeletalMeshComponent::GetClosestCollidingRigidBodyLocation(const FVector& TestLocation) const
+{
+	float BestDistSq = BIG_NUMBER;
+	FVector Best = TestLocation;
+
+	UPhysicsAsset* PhysicsAsset = GetPhysicsAsset();
+	if( PhysicsAsset )
+	{
+		for (int32 i=0; i<Bodies.Num(); i++)
+		{
+			FBodyInstance* BodyInstance = Bodies[i];
+			if( BodyInstance && BodyInstance->IsValidBodyInstance() && (BodyInstance->GetCollisionEnabled() != ECollisionEnabled::NoCollision) )
+			{
+				const FVector BodyLocation = BodyInstance->GetUnrealWorldTransform().GetTranslation();
+				const float DistSq = (BodyLocation - TestLocation).SizeSquared();
+				if( DistSq < BestDistSq )
+				{
+					Best = BodyLocation;
+					BestDistSq = DistSq;
+				}
+			}
+		}
+	}
+
+	return Best;
+}
+
+SIZE_T USkeletalMeshComponent::GetResourceSize( EResourceSizeMode::Type Mode )
+{
+	SIZE_T ResSize = 0;
+
+	for (int32 i=0; i < Bodies.Num(); ++i)
+	{
+		if (Bodies[i] != NULL && Bodies[i]->IsValidBodyInstance())
+		{
+			ResSize += Bodies[i]->GetBodyInstanceResourceSize(Mode);
+		}
+	}
+
+	return ResSize;
+}
+
+void USkeletalMeshComponent::SetAnimationMode(EAnimationMode::Type InAnimationMode)
+{
+	if (AnimationMode != InAnimationMode)
+	{
+		AnimationMode = InAnimationMode;
+		ClearAnimScriptInstance();
+		InitializeAnimScriptInstance();
+	}
+}
+
+EAnimationMode::Type USkeletalMeshComponent::GetAnimationMode() const
+{
+	return AnimationMode;
+}
+
+void USkeletalMeshComponent::PlayAnimation(class UAnimationAsset* NewAnimToPlay, bool bLooping)
+{
+	SetAnimationMode(EAnimationMode::AnimationSingleNode);
+	SetAnimation(NewAnimToPlay);
+	Play(bLooping);
+}
+
+void USkeletalMeshComponent::SetAnimation(UAnimationAsset* NewAnimToPlay)
+{
+	UAnimSingleNodeInstance* SingleNodeInstance = GetSingleNodeInstance();
+	if (SingleNodeInstance)
+	{
+		SingleNodeInstance->SetAnimationAsset(NewAnimToPlay, false);
+		SingleNodeInstance->bPlaying = false;
+	}
+	else if( AnimScriptInstance != NULL )
+	{
+		UE_LOG(LogAnimation, Warning, TEXT("Currently in Animation Blueprint mode. Please change AnimationMode to Use Animation Asset"));
+	}
+}
+
+void USkeletalMeshComponent::SetVertexAnimation(UVertexAnimation* NewVertexAnimation)
+{
+	UAnimSingleNodeInstance* SingleNodeInstance = GetSingleNodeInstance();
+	if (SingleNodeInstance)
+	{
+		SingleNodeInstance->SetVertexAnimation(NewVertexAnimation, false);
+		// when set the asset, we shouldn't automatically play. 
+		SingleNodeInstance->bPlaying = false;
+	}
+	else if( AnimScriptInstance != NULL )
+	{
+		UE_LOG(LogAnimation, Warning, TEXT("Currently in Animation Blueprint mode. Please change AnimationMode to Use Animation Asset"));
+	}
+}
+
+void USkeletalMeshComponent::Play(bool bLooping)
+{
+	UAnimSingleNodeInstance* SingleNodeInstance = GetSingleNodeInstance();
+	if (SingleNodeInstance)
+	{
+		SingleNodeInstance->bPlaying = true;
+		SingleNodeInstance->bLooping = bLooping;
+	}
+	else if( AnimScriptInstance != NULL )
+	{
+		UE_LOG(LogAnimation, Warning, TEXT("Currently in Animation Blueprint mode. Please change AnimationMode to Use Animation Asset"));
+	}
+}
+
+void USkeletalMeshComponent::Stop()
+{
+	UAnimSingleNodeInstance* SingleNodeInstance = GetSingleNodeInstance();
+	if (SingleNodeInstance)
+	{
+		SingleNodeInstance->bPlaying = false;
+	}
+	else if( AnimScriptInstance != NULL )
+	{
+		UE_LOG(LogAnimation, Warning, TEXT("Currently in Animation Blueprint mode. Please change AnimationMode to Use Animation Asset"));
+	}
+}
+
+bool USkeletalMeshComponent::IsPlaying() const
+{
+	UAnimSingleNodeInstance* SingleNodeInstance = GetSingleNodeInstance();
+	if (SingleNodeInstance)
+	{
+		return SingleNodeInstance->bPlaying;
+	}
+	else if( AnimScriptInstance != NULL )
+	{
+		UE_LOG(LogAnimation, Warning, TEXT("Currently in Animation Blueprint mode. Please change AnimationMode to Use Animation Asset"));
+	}
+
+	return false;
+}
+
+void USkeletalMeshComponent::SetPosition(float InPos, bool bFireNotifies)
+{
+	UAnimSingleNodeInstance* SingleNodeInstance = GetSingleNodeInstance();
+	if (SingleNodeInstance)
+	{
+		SingleNodeInstance->SetPosition(InPos, bFireNotifies);
+	}
+	else if( AnimScriptInstance != NULL )
+	{
+		UE_LOG(LogAnimation, Warning, TEXT("Currently in Animation Blueprint mode. Please change AnimationMode to Use Animation Asset"));
+	}
+}
+
+float USkeletalMeshComponent::GetPosition() const
+{
+	UAnimSingleNodeInstance* SingleNodeInstance = GetSingleNodeInstance();
+	if (SingleNodeInstance)
+	{
+		return SingleNodeInstance->CurrentTime;
+	}
+	else if( AnimScriptInstance != NULL )
+	{
+		UE_LOG(LogAnimation, Warning, TEXT("Currently in Animation Blueprint mode. Please change AnimationMode to Use Animation Asset"));
+	}
+
+	return 0.f;
+}
+
+void USkeletalMeshComponent::SetPlayRate(float Rate)
+{
+	UAnimSingleNodeInstance* SingleNodeInstance = GetSingleNodeInstance();
+	if (SingleNodeInstance)
+	{
+		SingleNodeInstance->SetPlayRate(Rate);
+	}
+	else if( AnimScriptInstance != NULL )
+	{
+		UE_LOG(LogAnimation, Warning, TEXT("Currently in Animation Blueprint mode. Please change AnimationMode to Use Animation Asset"));
+	}
+}
+
+float USkeletalMeshComponent::GetPlayRate() const
+{
+	UAnimSingleNodeInstance* SingleNodeInstance = GetSingleNodeInstance();
+	if (SingleNodeInstance)
+	{
+		return SingleNodeInstance->PlayRate;
+	}
+	else if( AnimScriptInstance != NULL )
+	{
+		UE_LOG(LogAnimation, Warning, TEXT("Currently in Animation Blueprint mode. Please change AnimationMode to Use Animation Asset"));
+	}
+
+	return 0.f;
+}
+
+class UAnimSingleNodeInstance* USkeletalMeshComponent::GetSingleNodeInstance() const
+{
+	return Cast<class UAnimSingleNodeInstance>(AnimScriptInstance);
+}
+
+void FSingleAnimationPlayData::Initialize(UAnimSingleNodeInstance* Instance)
+{
+	Instance->SetAnimationAsset(AnimToPlay);
+	Instance->SetPosition(SavedPosition, false);
+	Instance->SetPlayRate(SavedPlayRate);
+	Instance->SetPlaying(bSavedPlaying);
+	Instance->SetLooping(bSavedLooping);
+}
+
+void FSingleAnimationPlayData::PopulateFrom(UAnimSingleNodeInstance* Instance)
+{
+	AnimToPlay = Instance->CurrentAsset;
+	SavedPosition = Instance->CurrentTime;
+	SavedPlayRate = Instance->PlayRate;
+	bSavedPlaying = Instance->bPlaying;
+	bSavedLooping = Instance->bLooping;
+}
+
+void FSingleAnimationPlayData::ValidatePosition()
+{
+	float Min=0,Max=0;
+
+	if (AnimToPlay)
+	{
+		UAnimSequenceBase* SequenceBase = Cast<UAnimSequenceBase>(AnimToPlay);
+		if (SequenceBase)
+		{
+			Max = SequenceBase->SequenceLength;
+		}
+	}
+	else if (VertexAnimToPlay)
+	{
+		Max = VertexAnimToPlay->GetAnimLength();
+	}
+
+	SavedPosition = FMath::Clamp<float>(SavedPosition, Min, Max);
+}
+
+FTransform USkeletalMeshComponent::ConvertLocalRootMotionToWorld(const FTransform& InTransform)
+{
+	// Make sure component to world is up to date
+	if (!bWorldToComponentUpdated)
+	{
+		UpdateComponentToWorld();
+	}
+
+	const FTransform NewWorldTransform = InTransform * ComponentToWorld;
+	const FVector DeltaWorldTranslation = NewWorldTransform.GetTranslation() - ComponentToWorld.GetTranslation();
+	const FQuat DeltaWorldRotation = ComponentToWorld.GetRotation().Inverse() * NewWorldTransform.GetRotation();
+
+	const FTransform DeltaWorldTransform(DeltaWorldRotation, DeltaWorldTranslation);
+
+	UE_LOG(LogRootMotion, Log,  TEXT("ConvertLocalRootMotionToWorld LocalT: %s, LocalR: %s, WorldT: %s, WorldR: %s."),
+		*InTransform.GetTranslation().ToCompactString(), *InTransform.GetRotation().Rotator().ToCompactString(), 
+		*DeltaWorldTransform.GetTranslation().ToCompactString(), *DeltaWorldTransform.GetRotation().Rotator().ToCompactString() );
+
+	return DeltaWorldTransform;
+}
+
+FRootMotionMovementParams USkeletalMeshComponent::ConsumeRootMotion()
+{
+	if (AnimScriptInstance)
+	{
+		float InterpAlpha = bEnableUpdateRateOptimizations ? AnimUpdateRateParams->GetRootMotionInterp() : 1.f;
+		return AnimScriptInstance->ConsumeExtractedRootMotion(InterpAlpha);
+	}
+	return FRootMotionMovementParams();
+}
+
+
+float USkeletalMeshComponent::CalculateMass(FName BoneName)
+{
+	float Mass = 0.0f;
+
+	if (Bodies.Num())
+	{
+		for (int32 i = 0; i < Bodies.Num(); ++i)
+		{
+			//if bone name is not provided calculate entire mass - otherwise get mass for just the bone
+			if (Bodies[i]->BodySetup.IsValid() && (BoneName == NAME_None || BoneName == Bodies[i]->BodySetup->BoneName))
+			{
+				Mass += Bodies[i]->BodySetup->CalculateMass(this);
+			}
+		}
+	}
+	else	//We want to calculate mass before we've initialized body instances - in this case use physics asset setup
+	{
+		TArray<class UBodySetup*> * BodySetups = NULL;
+		if (UPhysicsAsset * PhysicsAsset = GetPhysicsAsset())
+		{
+			BodySetups = &PhysicsAsset->BodySetup;
+		}
+
+		if (BodySetups)
+		{
+			for (int32 i = 0; i < BodySetups->Num(); ++i)
+			{
+				if ((*BodySetups)[i] && (BoneName == NAME_None || BoneName == (*BodySetups)[i]->BoneName))
+				{
+					Mass += (*BodySetups)[i]->CalculateMass(this);
+				}
+			}
+		}
+	}
+
+	return Mass;
+}
+
+#if WITH_EDITOR
+
+bool USkeletalMeshComponent::ComponentIsTouchingSelectionBox(const FBox& InSelBBox, const FEngineShowFlags& ShowFlags, const bool bConsiderOnlyBSP, const bool bMustEncompassEntireComponent) const
+{
+	if (!bConsiderOnlyBSP && ShowFlags.SkeletalMeshes && MeshObject != nullptr)
+	{
+		FSkeletalMeshResource* SkelMeshResource = GetSkeletalMeshResource();
+		check(SkelMeshResource);
+		check(SkelMeshResource->LODModels.Num() > 0);
+
+		// Transform hard and soft verts into world space. Note that this assumes skeletal mesh is in reference pose...
+		const FStaticLODModel& LODModel = SkelMeshResource->LODModels[0];
+		for (const auto& Chunk : LODModel.Chunks)
+		{
+			for (const auto& Vertex : Chunk.RigidVertices)
+			{
+				const FVector Location = ComponentToWorld.TransformPosition(Vertex.Position);
+				const bool bLocationIntersected = FMath::PointBoxIntersection(Location, InSelBBox);
+
+				// If the selection box doesn't have to encompass the entire component and a skeletal mesh vertex has intersected with
+				// the selection box, this component is being touched by the selection box
+				if (!bMustEncompassEntireComponent && bLocationIntersected)
+				{
+					return true;
+				}
+
+				// If the selection box has to encompass the entire component and a skeletal mesh vertex didn't intersect with the selection
+				// box, this component does not qualify
+				else if (bMustEncompassEntireComponent && !bLocationIntersected)
+				{
+					return false;
+				}
+			}
+
+			for (const auto& Vertex : Chunk.SoftVertices)
+			{
+				const FVector Location = ComponentToWorld.TransformPosition(Vertex.Position);
+				const bool bLocationIntersected = FMath::PointBoxIntersection(Location, InSelBBox);
+
+				// If the selection box doesn't have to encompass the entire component and a skeletal mesh vertex has intersected with
+				// the selection box, this component is being touched by the selection box
+				if (!bMustEncompassEntireComponent && bLocationIntersected)
+				{
+					return true;
+				}
+
+				// If the selection box has to encompass the entire component and a skeletal mesh vertex didn't intersect with the selection
+				// box, this component does not qualify
+				else if (bMustEncompassEntireComponent && !bLocationIntersected)
+				{
+					return false;
+				}
+			}
+		}
+
+		// If the selection box has to encompass all of the component and none of the component's verts failed the intersection test, this component
+		// is consider touching
+		return true;
+	}
+
+	return false;
+}
+
+bool USkeletalMeshComponent::ComponentIsTouchingSelectionFrustum(const FConvexVolume& InFrustum, const FEngineShowFlags& ShowFlags, const bool bConsiderOnlyBSP, const bool bMustEncompassEntireComponent) const
+{
+	if (!bConsiderOnlyBSP && ShowFlags.SkeletalMeshes && MeshObject != nullptr)
+	{
+		FSkeletalMeshResource* SkelMeshResource = GetSkeletalMeshResource();
+		check(SkelMeshResource);
+		check(SkelMeshResource->LODModels.Num() > 0);
+
+		// Transform hard and soft verts into world space. Note that this assumes skeletal mesh is in reference pose...
+		const FStaticLODModel& LODModel = SkelMeshResource->LODModels[0];
+		for (const auto& Chunk : LODModel.Chunks)
+		{
+			for (const auto& Vertex : Chunk.RigidVertices)
+			{
+				const FVector Location = ComponentToWorld.TransformPosition(Vertex.Position);
+				const bool bLocationIntersected = InFrustum.IntersectSphere(Location, 0.0f);
+
+				// If the selection box doesn't have to encompass the entire component and a skeletal mesh vertex has intersected with
+				// the selection box, this component is being touched by the selection box
+				if (!bMustEncompassEntireComponent && bLocationIntersected)
+				{
+					return true;
+				}
+
+				// If the selection box has to encompass the entire component and a skeletal mesh vertex didn't intersect with the selection
+				// box, this component does not qualify
+				else if (bMustEncompassEntireComponent && !bLocationIntersected)
+				{
+					return false;
+				}
+			}
+
+			for (const auto& Vertex : Chunk.SoftVertices)
+			{
+				const FVector Location = ComponentToWorld.TransformPosition(Vertex.Position);
+				const bool bLocationIntersected = InFrustum.IntersectSphere(Location, 0.0f);
+
+				// If the selection box doesn't have to encompass the entire component and a skeletal mesh vertex has intersected with
+				// the selection box, this component is being touched by the selection box
+				if (!bMustEncompassEntireComponent && bLocationIntersected)
+				{
+					return true;
+				}
+
+				// If the selection box has to encompass the entire component and a skeletal mesh vertex didn't intersect with the selection
+				// box, this component does not qualify
+				else if (bMustEncompassEntireComponent && !bLocationIntersected)
+				{
+					return false;
+				}
+			}
+		}
+
+		// If the selection box has to encompass all of the component and none of the component's verts failed the intersection test, this component
+		// is consider touching
+		return true;
+	}
+
+	return false;
+}
+
+
+void USkeletalMeshComponent::UpdateCollisionProfile()
+{
+	Super::UpdateCollisionProfile();
+
+	for(int32 i=0; i < Bodies.Num(); ++i)
+	{
+		if(Bodies[i]->BodySetup.IsValid())
+		{
+			Bodies[i]->LoadProfileData(false);
+		}
+	}
+}
+
+FDelegateHandle USkeletalMeshComponent::RegisterOnSkeletalMeshPropertyChanged( const FOnSkeletalMeshPropertyChanged& Delegate )
+{
+	return OnSkeletalMeshPropertyChanged.Add(Delegate);
+}
+
+void USkeletalMeshComponent::UnregisterOnSkeletalMeshPropertyChanged( const FOnSkeletalMeshPropertyChanged& Delegate )
+{
+	OnSkeletalMeshPropertyChanged.DEPRECATED_Remove(Delegate);
+}
+
+void USkeletalMeshComponent::UnregisterOnSkeletalMeshPropertyChanged( FDelegateHandle Handle )
+{
+	OnSkeletalMeshPropertyChanged.Remove(Handle);
+}
+
+void USkeletalMeshComponent::ValidateAnimation()
+{
+	if (SkeletalMesh && SkeletalMesh->Skeleton == NULL)
+	{
+		UE_LOG(LogAnimation, Warning, TEXT("SkeletalMesh %s has no skeleton. This needs to fixed before an animation can be set"), *SkeletalMesh->GetName());
+		if (AnimationMode == EAnimationMode::AnimationSingleNode)
+		{
+			AnimationData.AnimToPlay = NULL;
+		}
+		else
+		{
+			AnimBlueprintGeneratedClass = NULL;
+		}
+		return;
+	}
+
+	if(AnimationMode == EAnimationMode::AnimationSingleNode)
+	{
+		if(AnimationData.AnimToPlay && SkeletalMesh && AnimationData.AnimToPlay->GetSkeleton() != SkeletalMesh->Skeleton)
+		{
+			UE_LOG(LogAnimation, Warning, TEXT("Animation %s is incompatible with skeleton %s, removing animation from actor."), *AnimationData.AnimToPlay->GetName(), *SkeletalMesh->Skeleton->GetName());
+			AnimationData.AnimToPlay = NULL;
+		}
+	}
+	else
+	{
+		if(AnimBlueprintGeneratedClass && SkeletalMesh && AnimBlueprintGeneratedClass->TargetSkeleton != SkeletalMesh->Skeleton)
+		{
+			if(SkeletalMesh->Skeleton == nullptr)
+			{
+				UE_LOG(LogAnimation, Warning, TEXT("AnimBP %s is incompatible because mesh %s has no skeleton, removing AnimBP from actor."), *AnimBlueprintGeneratedClass->GetName(), *SkeletalMesh->GetName());
+			}
+			else
+			{
+				UE_LOG(LogAnimation, Warning, TEXT("AnimBP %s is incompatible with skeleton %s, removing AnimBP from actor."), *AnimBlueprintGeneratedClass->GetName(), *SkeletalMesh->Skeleton->GetName());
+			}
+
+			AnimBlueprintGeneratedClass = NULL;
+		}
+	}
+}
+
+#endif
+
+bool USkeletalMeshComponent::IsPlayingRootMotion()
+{
+	return (AnimScriptInstance ? (AnimScriptInstance->GetRootMotionMontageInstance() != NULL) : false);
+}
+
+bool USkeletalMeshComponent::IsPlayingRootMotionFromEverything()
+{
+	return AnimScriptInstance ? (AnimScriptInstance->RootMotionMode == ERootMotionMode::RootMotionFromEverything) : false;
+}
+
+void USkeletalMeshComponent::SetRootBodyIndex(int32 InBodyIndex)
+{
+	RootBodyData.BodyIndex = InBodyIndex;
+
+	if(Bodies.IsValidIndex(RootBodyData.BodyIndex) && SkeletalMesh && 
+		Bodies[RootBodyData.BodyIndex]->BodySetup.IsValid() && Bodies[RootBodyData.BodyIndex]->BodySetup.Get()->BoneName != NAME_None)
+	{
+		int32 BoneIndex = GetBoneIndex(Bodies[RootBodyData.BodyIndex]->BodySetup->BoneName);
+		// if bone index is valid and not 0, it SHOULD have parnet index
+		if (ensure (BoneIndex != INDEX_NONE))
+		{
+			int32 ParentIndex = SkeletalMesh->RefSkeleton.GetParentIndex(BoneIndex);
+			if (BoneIndex != 0 && ensure (ParentIndex != INDEX_NONE))
+			{
+				const TArray<FTransform>& RefPose = SkeletalMesh->RefSkeleton.GetRefBonePose();
+
+				FTransform RelativeTransform = FTransform(SkeletalMesh->RefBasesInvMatrix[BoneIndex]) * RefPose[ParentIndex];
+				// now get offset 
+				RootBodyData.TransformToRoot = RelativeTransform;
+			}
+			else
+			{
+				RootBodyData.TransformToRoot = FTransform::Identity;
+			}
+		}
+		else
+		{
+			RootBodyData.TransformToRoot = FTransform::Identity;
+		}
+	}
+	else
+	{
+		// error - this should not happen
+		ensure(false);
+		RootBodyData.TransformToRoot = FTransform::Identity;
+	}
+}
\ No newline at end of file
