Engine: Moved actor components into Components folder

diff --git a/Engine/Source/Runtime/Engine/Private/Components/SkinnedMeshComponent.cpp b/Engine/Source/Runtime/Engine/Private/Components/SkinnedMeshComponent.cpp
new file mode 100644
--- /dev/null
+++ b/Engine/Source/Runtime/Engine/Private/Components/SkinnedMeshComponent.cpp
@@ -0,0 +1,2011 @@
+// Copyright 1998-2015 Epic Games, Inc. All Rights Reserved.
+
+/*=============================================================================
+	UnSkeletalComponent.cpp: Actor component implementation.
+=============================================================================*/
+
+#include "EnginePrivate.h"
+#include "ParticleDefinitions.h"
+#include "BlueprintUtilities.h"
+#include "SkeletalRenderCPUSkin.h"
+#include "SkeletalRenderGPUSkin.h"
+#include "AnimationUtils.h"
+#include "AnimTree.h"
+#include "Animation/VertexAnim/MorphTarget.h"
+#include "ComponentReregisterContext.h"
+#include "Engine/SkeletalMeshSocket.h"
+#include "PhysicsEngine/PhysicsAsset.h"
+
+DEFINE_LOG_CATEGORY_STATIC(LogSkinnedMeshComp, Log, All);
+
+int32 GSkeletalMeshLODBias = 0;
+FAutoConsoleVariableRef CVarSkeletalMeshLODBias(
+	TEXT("r.SkeletalMeshLODBias"),
+	GSkeletalMeshLODBias,
+	TEXT("LOD bias for skeletal meshes."),
+	ECVF_Scalability
+	);
+
+
+namespace FAnimUpdateRateManager
+{
+	// Global counter to spread SkinnedMeshComponent tick updates.
+	static uint8 UpdateRateGroupCount = 0;
+
+	static float TargetFrameTimeForUpdateRate = 1.f / 30.f; //Target frame rate for lookahead URO
+
+	struct FAnimUpdateRateParametersTracker
+	{
+		FAnimUpdateRateParameters UpdateRateParameters;
+
+		/** Frame counter to call AnimUpdateRateTick() just once per frame. */
+		uint32 AnimUpdateRateFrameCount;
+
+		/** Counter to stagger update and evaluation across skinned mesh components */
+		uint8 AnimUpdateRateShiftTag;
+
+		/** List of all USkinnedMeshComponents that use this set of parameters */
+		TArray<USkinnedMeshComponent*> RegisteredComponents;
+
+		FAnimUpdateRateParametersTracker() : AnimUpdateRateFrameCount(0), AnimUpdateRateShiftTag(0) {}
+
+		uint8 GetAnimUpdateRateShiftTag()
+		{
+			// If hasn't been initialized yet, pick a unique ID, to spread population over frames.
+			if (AnimUpdateRateShiftTag == 0)
+			{
+				AnimUpdateRateShiftTag = ++UpdateRateGroupCount;
+			}
+
+			return AnimUpdateRateShiftTag;
+		}
+
+		bool IsHumanControlled() const
+		{
+			AActor* Owner = RegisteredComponents[0]->GetOwner();
+			AController * Controller = Owner ? Owner->GetInstigatorController() : NULL;
+			return Cast<APlayerController>(Controller) != NULL;
+		}
+	};
+
+
+	TMap<UObject*, FAnimUpdateRateParametersTracker*> ActorToUpdateRateParams;
+
+	UObject* GetMapIndexForComponent(USkinnedMeshComponent* SkinnedComponent)
+	{
+		UObject* TrackerIndex = SkinnedComponent->GetOwner();
+		if (TrackerIndex == nullptr)
+		{
+			TrackerIndex = SkinnedComponent;
+		}
+		return TrackerIndex;
+	}
+
+	FAnimUpdateRateParameters* GetUpdateRateParameters(USkinnedMeshComponent* SkinnedComponent)
+	{
+		UObject* TrackerIndex = GetMapIndexForComponent(SkinnedComponent);
+
+		FAnimUpdateRateParametersTracker** ExistingTracker = ActorToUpdateRateParams.Find(TrackerIndex);
+		if (!ExistingTracker)
+		{
+			ExistingTracker = &ActorToUpdateRateParams.Add(TrackerIndex);
+			(*ExistingTracker) = new FAnimUpdateRateParametersTracker();
+		}
+		
+		check(ExistingTracker);
+		check(*ExistingTracker);
+
+		checkSlow(!(*ExistingTracker)->RegisteredComponents.Contains(SkinnedComponent)); // We have already been registered? Something has gone very wrong!
+
+		(*ExistingTracker)->RegisteredComponents.Add(SkinnedComponent);
+		return &(*ExistingTracker)->UpdateRateParameters;
+	}
+
+	void CleanupUpdateRateParametersRef(USkinnedMeshComponent* SkinnedComponent)
+	{
+		UObject* TrackerIndex = GetMapIndexForComponent(SkinnedComponent);
+
+		FAnimUpdateRateParametersTracker* Tracker = ActorToUpdateRateParams.FindChecked(TrackerIndex);
+		Tracker->RegisteredComponents.Remove(SkinnedComponent);
+		if (Tracker->RegisteredComponents.Num() == 0)
+		{
+			ActorToUpdateRateParams.Remove(TrackerIndex);
+			delete Tracker;
+		}
+	}
+
+	void AnimUpdateRateSetParams(FAnimUpdateRateParametersTracker* Tracker, float DeltaTime, bool bRecentlyRendered, float MaxDistanceFactor, bool bNeedsValidRootMotion, bool bUsingRootMotionFromEverything)
+	{
+		// default rules for setting update rates
+
+		// Human controlled characters should be ticked always fully to minimize latency w/ game play events triggered by animation.
+		const bool bHumanControlled = Tracker->IsHumanControlled();
+
+		bool bNeedsEveryFrame = bNeedsValidRootMotion && !bUsingRootMotionFromEverything;
+
+		// Not rendered, including dedicated servers. we can skip the Evaluation part.
+		if (!bRecentlyRendered)
+		{
+			Tracker->UpdateRateParameters.SetTrailMode(DeltaTime, Tracker->GetAnimUpdateRateShiftTag(), (bHumanControlled ? 1 : 4), 4, false);
+		}
+		// Visible controlled characters or playing root motion. Need evaluation and ticking done every frame.
+		else  if (bHumanControlled || bNeedsEveryFrame)
+		{
+			Tracker->UpdateRateParameters.SetTrailMode(DeltaTime, Tracker->GetAnimUpdateRateShiftTag(), 1, 1, false);
+		}
+		else
+		{
+			// For visible meshes, figure out how often bones should be evaluated VS interpolated to previous evaluation.
+			// This is based on screen size and makes sense for a 3D FPS game, different games or Actors can implement different rules.
+			int32 DesiredEvaluationRate;
+			if (MaxDistanceFactor > 0.4f)
+			{
+				DesiredEvaluationRate = 1;
+			}
+			else if (MaxDistanceFactor > .2f)
+			{
+				DesiredEvaluationRate = 2;
+			}
+			else
+			{
+				DesiredEvaluationRate = 3;
+			}
+
+			if (bUsingRootMotionFromEverything && DesiredEvaluationRate > 1)
+			{
+				//Use look ahead mode that allows us to rate limit updates even when using root motion
+				Tracker->UpdateRateParameters.SetLookAheadMode(DeltaTime, Tracker->GetAnimUpdateRateShiftTag(), TargetFrameTimeForUpdateRate*DesiredEvaluationRate);
+			}
+			else
+			{
+				Tracker->UpdateRateParameters.SetTrailMode(DeltaTime, Tracker->GetAnimUpdateRateShiftTag(), DesiredEvaluationRate, DesiredEvaluationRate, true);
+			}
+		}
+	}
+
+	void AnimUpdateRateTick(FAnimUpdateRateParametersTracker* Tracker, float DeltaTime, bool bNeedsValidRootMotion)
+	{
+		const float RecentlyRenderedTime = (Tracker->RegisteredComponents[0]->GetWorld()->TimeSeconds - 1.f);
+
+		// Go through components and figure out if they've been recently rendered, and the biggest MaxDistanceFactor
+		bool bRecentlyRendered = false;
+		bool bPlayingRootMotion = false;
+		bool bUsingRootMotionFromEverything = true;
+		float MaxDistanceFactor = 0.f;
+
+		const TArray<USkinnedMeshComponent*>& SkinnedComponents = Tracker->RegisteredComponents;
+
+		for (USkinnedMeshComponent* Component : SkinnedComponents)
+		{
+			bRecentlyRendered |= (Component->LastRenderTime > RecentlyRenderedTime);
+			MaxDistanceFactor = FMath::Max(MaxDistanceFactor, Component->MaxDistanceFactor);
+			bPlayingRootMotion |= Component->IsPlayingRootMotion();
+			bUsingRootMotionFromEverything &= Component->IsPlayingRootMotionFromEverything();
+		}
+
+		bNeedsValidRootMotion &= bPlayingRootMotion;
+
+		// Figure out which update rate should be used.
+		AnimUpdateRateSetParams(Tracker, DeltaTime, bRecentlyRendered, MaxDistanceFactor, bNeedsValidRootMotion, bUsingRootMotionFromEverything);
+	}
+
+	const TCHAR* B(bool b)
+	{
+		return b ? TEXT("true") : TEXT("false");
+	}
+
+	void TickUpdateRateParameters(USkinnedMeshComponent* SkinnedComponent, float DeltaTime, bool bNeedsValidRootMotion)
+	{
+		// Convert current frame counter from 64 to 32 bits.
+		const uint32 CurrentFrame32 = uint32(GFrameCounter % MAX_uint32);
+
+		UObject* TrackerIndex = GetMapIndexForComponent(SkinnedComponent);
+		FAnimUpdateRateParametersTracker* Tracker = ActorToUpdateRateParams.FindChecked(TrackerIndex);
+
+		if (CurrentFrame32 != Tracker->AnimUpdateRateFrameCount)
+		{
+			Tracker->AnimUpdateRateFrameCount = CurrentFrame32;
+			AnimUpdateRateTick(Tracker, DeltaTime, bNeedsValidRootMotion);
+		}
+	}
+}
+
+void USkinnedMeshComponent::OnRegister()
+{
+	// The reason this happens before register
+	// is so that any transform update (or children transform update)
+	// won't result in any issues of accessing SpaceBases
+	// This isn't really ideal solution because these transform won't have 
+	// any valid data yet. 
+
+	AnimUpdateRateParams = FAnimUpdateRateManager::GetUpdateRateParameters(this);
+
+	AllocateTransformData();
+	UpdateMasterBoneMap();	
+
+	Super::OnRegister();
+
+	UpdateLODStatus();
+}
+
+void USkinnedMeshComponent::OnUnregister()
+{
+	DeallocateTransformData();
+	Super::OnUnregister();
+	FAnimUpdateRateManager::CleanupUpdateRateParametersRef(this);
+	AnimUpdateRateParams = NULL;
+}
+
+void USkinnedMeshComponent::CreateRenderState_Concurrent()
+{
+	if( SkeletalMesh )
+	{
+		// Initialize the alternate weight tracks if present BEFORE creating the new mesh object
+		InitLODInfos();
+
+		// No need to create the mesh object if we aren't actually rendering anything (see UPrimitiveComponent::Attach)
+		if ( FApp::CanEverRender() && ShouldComponentAddToScene() )
+		{
+			ERHIFeatureLevel::Type SceneFeatureLevel = GetWorld()->FeatureLevel;
+			FSkeletalMeshResource* SkelMeshResource = SkeletalMesh->GetResourceForRendering();
+
+			// Also check if skeletal mesh has too many bones/chunk for GPU skinning.
+			const bool bIsCPUSkinned = SkelMeshResource->RequiresCPUSkinning(SceneFeatureLevel) || (GIsEditor && ShouldCPUSkin());
+			if(bIsCPUSkinned)
+			{
+				MeshObject = ::new FSkeletalMeshObjectCPUSkin(this, SkelMeshResource, SceneFeatureLevel);
+			}
+			else
+			{
+				MeshObject = ::new FSkeletalMeshObjectGPUSkin(this, SkelMeshResource, SceneFeatureLevel);
+			}
+
+			//Allow the editor a chance to manipulate it before its added to the scene
+			PostInitMeshObject(MeshObject);
+		}
+	}
+
+	Super::CreateRenderState_Concurrent();
+
+	if (SkeletalMesh)
+	{
+		// Update dynamic data
+
+		if(MeshObject)
+		{
+			int32 UseLOD = PredictedLODLevel;
+			if(MasterPoseComponent.IsValid())
+			{
+				UseLOD = FMath::Clamp(MasterPoseComponent->PredictedLODLevel, 0, MeshObject->GetSkeletalMeshResource().LODModels.Num()-1);
+			}
+
+			// We just recreated RenderState, that means it could be new mesh or different mesh
+			// Clear ActiveVertexAnims, so that it doesn't use from previous state
+			ActiveVertexAnims.Empty();
+			MeshObject->Update(UseLOD,this,ActiveVertexAnims);  // send to rendering thread
+		}
+
+		// scene proxy update of material usage based on active morphs
+		UpdateMorphMaterialUsageOnProxy();
+	}
+}
+
+void USkinnedMeshComponent::DestroyRenderState_Concurrent()
+{
+	Super::DestroyRenderState_Concurrent();
+
+	if(MeshObject)
+	{
+		// Begin releasing the RHI resources used by this skeletal mesh component.
+		// This doesn't immediately destroy anything, since the rendering thread may still be using the resources.
+		MeshObject->ReleaseResources();
+
+		// Begin a deferred delete of MeshObject.  BeginCleanup will call MeshObject->FinishDestroy after the above release resource
+		// commands execute in the rendering thread.
+		BeginCleanup(MeshObject);
+		MeshObject = NULL;
+	}
+}
+
+FString USkinnedMeshComponent::GetDetailedInfoInternal() const
+{
+	FString Result;  
+
+	if( SkeletalMesh != NULL )
+	{
+		Result = SkeletalMesh->GetDetailedInfoInternal();
+	}
+	else
+	{
+		Result = TEXT("No_SkeletalMesh");
+	}
+
+	return Result;  
+}
+
+void USkinnedMeshComponent::SendRenderDynamicData_Concurrent()
+{
+	SCOPE_CYCLE_COUNTER(STAT_SkelCompUpdateTransform);
+
+	Super::SendRenderDynamicData_Concurrent();
+
+	// if we have not updated the transforms then no need to send them to the rendering thread
+	// @todo GIsEditor used to be bUpdateSkelWhenNotRendered. Look into it further to find out why it doesn't update animations in the AnimSetViewer, when a level is loaded in UED (like POC_Cover.gear).
+	if( MeshObject && SkeletalMesh && (bForceMeshObjectUpdate || (bRecentlyRendered || MeshComponentUpdateFlag == EMeshComponentUpdateFlag::AlwaysTickPoseAndRefreshBones || GIsEditor || MeshObject->bHasBeenUpdatedAtLeastOnce == false)) )
+	{
+		SCOPE_CYCLE_COUNTER(STAT_MeshObjectUpdate);
+
+		int32 UseLOD = PredictedLODLevel;
+		// If we have a MasterPoseComponent - force this component to render at that LOD, so all bones are present for it.
+		// Note that this currently relies on the behaviour where this mesh is rendered at the LOD we pass in here for all viewports. We should
+		// be able to render it at lower LOD on viewports where it is further away. That will make the MasterPoseComponent case a bit harder to solve.
+		if(MasterPoseComponent.IsValid())
+		{
+			UseLOD = FMath::Clamp(MasterPoseComponent->PredictedLODLevel, 0, MeshObject->GetSkeletalMeshResource().LODModels.Num()-1);
+		}
+
+		// Are morph targets disabled for this LOD?
+		if ( SkeletalMesh->LODInfo[ UseLOD ].bHasBeenSimplified || bDisableMorphTarget )
+		{
+			ActiveVertexAnims.Empty();
+		}
+
+		MeshObject->Update(UseLOD,this,ActiveVertexAnims);  // send to rendering thread
+		MeshObject->bHasBeenUpdatedAtLeastOnce = true;
+		
+		// scene proxy update of material usage based on active morphs
+		UpdateMorphMaterialUsageOnProxy();
+	}
+
+}
+
+#if WITH_EDITOR
+void USkinnedMeshComponent::PostEditChangeProperty(FPropertyChangedEvent& PropertyChangedEvent)
+{
+	Super::PostEditChangeProperty(PropertyChangedEvent);
+
+	UProperty* PropertyThatChanged = PropertyChangedEvent.Property;
+
+	if ( PropertyChangedEvent.Property != NULL )
+	{
+		if ( GIsEditor && PropertyThatChanged->GetFName() == GET_MEMBER_NAME_CHECKED(USkinnedMeshComponent, StreamingDistanceMultiplier) )
+		{
+			// Recalculate in a few seconds.
+			GEngine->TriggerStreamingDataRebuild();
+		}
+	}
+}
+#endif // WITH_EDITOR
+
+void USkinnedMeshComponent::InitLODInfos()
+{
+	if (SkeletalMesh != NULL)
+	{
+		if (SkeletalMesh->LODInfo.Num() != LODInfo.Num())
+		{
+			LODInfo.Empty(SkeletalMesh->LODInfo.Num());
+			for (int32 Idx=0; Idx < SkeletalMesh->LODInfo.Num(); Idx++)
+			{
+				new(LODInfo) FSkelMeshComponentLODInfo();
+			}
+		}
+	}	
+}
+
+
+bool USkinnedMeshComponent::ShouldTickPose() const
+{
+	return ((MeshComponentUpdateFlag < EMeshComponentUpdateFlag::OnlyTickPoseWhenRendered) || bRecentlyRendered);
+}
+
+bool USkinnedMeshComponent::ShouldUpdateTransform(bool bLODHasChanged) const
+{
+	return (bLODHasChanged || bRecentlyRendered || (MeshComponentUpdateFlag == EMeshComponentUpdateFlag::AlwaysTickPoseAndRefreshBones));
+}
+
+
+void USkinnedMeshComponent::TickUpdateRate(float DeltaTime, bool bNeedsValidRootMotion)
+{
+	SCOPE_CYCLE_COUNTER(STAT_TickUpdateRate);
+	if( bEnableUpdateRateOptimizations )
+	{
+		if (GetOwner())
+		{
+			// Tick Owner once per frame. All attached SkinnedMeshComponents will share the same settings.
+			FAnimUpdateRateManager::TickUpdateRateParameters(this, DeltaTime, bNeedsValidRootMotion);
+
+			// debug -- @todo: hook this up to a console command.
+			if (bDisplayDebugUpdateRateOptimizations)
+			{
+				FColor DrawColor = AnimUpdateRateParams->GetUpdateRateDebugColor();
+				DrawDebugBox(GetWorld(), Bounds.Origin, Bounds.BoxExtent, FQuat::Identity, DrawColor, false);
+			}
+		}
+	}
+}
+
+void USkinnedMeshComponent::TickComponent(float DeltaTime, enum ELevelTick TickType, FActorComponentTickFunction *ThisTickFunction)
+{
+	SCOPE_CYCLE_COUNTER(STAT_SkinnedMeshCompTick);
+
+	// Tick ActorComponent first.
+	Super::TickComponent(DeltaTime, TickType, ThisTickFunction);
+
+	// See if this mesh was rendered recently. This has to happen first because other data will rely on this
+	bRecentlyRendered = (LastRenderTime > GetWorld()->TimeSeconds - 1.0f);
+
+	// Tick Pose first
+	if (ShouldTickPose())
+	{
+		TickPose(DeltaTime, false);
+	}
+
+	// Update component's LOD settings
+	const bool bLODHasChanged = UpdateLODStatus();
+
+	// If we have been recently rendered, and bForceRefPose has been on for at least a frame, or the LOD changed, update bone matrices.
+	if( ShouldUpdateTransform(bLODHasChanged) )
+	{
+		// Do not update bones if we are taking bone transforms from another SkelMeshComp
+		if( MasterPoseComponent.IsValid() )
+		{
+			UpdateSlaveComponent();
+		}
+		else 
+		{
+			RefreshBoneTransforms(ThisTickFunction);
+		}
+	}
+}
+
+UObject const* USkinnedMeshComponent::AdditionalStatObject() const
+{
+	return SkeletalMesh;
+}
+
+
+void USkinnedMeshComponent::UpdateSlaveComponent()
+{
+	MarkRenderDynamicDataDirty();
+}
+
+TArray<class UMaterialInterface*> USkinnedMeshComponent::GetMaterials() const
+{
+	TArray<class UMaterialInterface*> OutMaterials = Super::GetMaterials();
+
+	// if no material is overriden, look for mesh material;
+	if(OutMaterials.Num() == 0 && SkeletalMesh)
+	{
+		for (auto Material : SkeletalMesh->Materials)
+		{
+			OutMaterials.Add(Material.MaterialInterface);
+		}
+	}
+
+	return OutMaterials;
+}
+
+int32 USkinnedMeshComponent::GetNumMaterials() const
+{
+	if (SkeletalMesh)
+	{
+		return SkeletalMesh->Materials.Num();
+	}
+
+	return 0;
+}
+
+UMaterialInterface* USkinnedMeshComponent::GetMaterial(int32 MaterialIndex) const
+{
+	if(MaterialIndex < OverrideMaterials.Num() && OverrideMaterials[MaterialIndex])
+	{
+		return OverrideMaterials[MaterialIndex];
+	}
+	else if (SkeletalMesh && MaterialIndex < SkeletalMesh->Materials.Num() && SkeletalMesh->Materials[MaterialIndex].MaterialInterface)
+	{
+		return SkeletalMesh->Materials[MaterialIndex].MaterialInterface;
+	}
+
+	return NULL;
+}
+
+
+bool USkinnedMeshComponent::ShouldCPUSkin()
+{
+	return false;
+}
+
+void USkinnedMeshComponent::GetStreamingTextureInfo(TArray<FStreamingTexturePrimitiveInfo>& OutStreamingTextures) const
+{
+	if( SkeletalMesh )
+	{
+		const float LocalTexelFactor = SkeletalMesh->GetStreamingTextureFactor(0) * FMath::Max(0.0f, StreamingDistanceMultiplier);
+		const float WorldTexelFactor = LocalTexelFactor * ComponentToWorld.GetMaximumAxisScale();
+		const int32 NumMaterials = FMath::Max(SkeletalMesh->Materials.Num(), OverrideMaterials.Num());
+		for( int32 MatIndex = 0; MatIndex < NumMaterials; MatIndex++ )
+		{
+			UMaterialInterface* const MaterialInterface = GetMaterial(MatIndex);
+			if(MaterialInterface)
+			{
+				TArray<UTexture*> Textures;
+				
+				auto World = GetWorld();
+				MaterialInterface->GetUsedTextures(Textures, EMaterialQualityLevel::Num, false, World ? World->FeatureLevel : GMaxRHIFeatureLevel, false);
+				for(int32 TextureIndex = 0; TextureIndex < Textures.Num(); TextureIndex++ )
+				{
+					FStreamingTexturePrimitiveInfo& StreamingTexture = *new(OutStreamingTextures) FStreamingTexturePrimitiveInfo;
+					StreamingTexture.Bounds = Bounds.GetSphere();
+					StreamingTexture.TexelFactor = WorldTexelFactor;
+					StreamingTexture.Texture = Textures[TextureIndex];
+				}
+			}
+		}
+	}
+}
+
+bool USkinnedMeshComponent::ShouldUpdateBoneVisibility() const
+{
+	// do not update if it has MasterPoseComponent
+	return !MasterPoseComponent.IsValid();
+}
+void USkinnedMeshComponent::RebuildVisibilityArray()
+{
+	// BoneVisibility needs update if MasterComponent == NULL
+	// if MaterComponent, it should follow MaterPoseComponent
+	if ( ShouldUpdateBoneVisibility())
+	{
+		// If the BoneVisibilityStates array has a 0 for a parent bone, all children bones are meant to be hidden as well
+		// (as the concatenated matrix will have scale 0).  This code propagates explicitly hidden parents to children.
+
+		// On the first read of any cell of BoneVisibilityStates, BVS_HiddenByParent and BVS_Visible are treated as visible.
+		// If it starts out visible, the value written back will be BVS_Visible if the parent is visible; otherwise BVS_HiddenByParent.
+		// If it starts out hidden, the BVS_ExplicitlyHidden value stays in place
+
+		// The following code relies on a complete hierarchy sorted from parent to children
+		check(BoneVisibilityStates.Num() == SkeletalMesh->RefSkeleton.GetNum());
+		for (int32 BoneId=0; BoneId < BoneVisibilityStates.Num(); ++BoneId)
+		{
+			uint8 VisState = BoneVisibilityStates[BoneId];
+
+			// if not exclusively hidden, consider if parent is hidden
+			if (VisState != BVS_ExplicitlyHidden)
+			{
+				// Check direct parent (only need to do one deep, since we have already processed the parent and written to BoneVisibilityStates previously)
+				const int32 ParentIndex = SkeletalMesh->RefSkeleton.GetParentIndex(BoneId);
+				if ((ParentIndex == -1) || (BoneVisibilityStates[ParentIndex] == BVS_Visible))
+				{
+					BoneVisibilityStates[BoneId] = BVS_Visible;
+				}
+				else
+				{
+					BoneVisibilityStates[BoneId] = BVS_HiddenByParent;
+				}
+			}
+		}
+	}
+}
+
+FBoxSphereBounds USkinnedMeshComponent::CalcBounds(const FTransform& LocalToWorld) const
+{
+	return CalcMeshBound( FVector::ZeroVector, false, LocalToWorld );
+}
+
+
+class UPhysicsAsset* USkinnedMeshComponent::GetPhysicsAsset() const
+{
+	if (PhysicsAssetOverride)
+	{
+		return PhysicsAssetOverride;
+	}
+
+	if (SkeletalMesh && SkeletalMesh->PhysicsAsset)
+	{
+		return SkeletalMesh->PhysicsAsset;
+	}
+
+	return NULL;
+}
+
+FBoxSphereBounds USkinnedMeshComponent::CalcMeshBound(const FVector& RootOffset, bool UsePhysicsAsset, const FTransform& LocalToWorld) const
+{
+	SCOPE_CYCLE_COUNTER(STAT_UpdateSkelMeshBounds);
+
+	FBoxSphereBounds NewBounds;
+
+	// If physics are asleep, and actor is using physics to move, skip updating the bounds.
+	AActor* Owner = GetOwner();
+	FVector DrawScale = LocalToWorld.GetScale3D();	
+
+	UPhysicsAsset * const PhysicsAsset = GetPhysicsAsset();
+	UPhysicsAsset * const MasterPhysicsAsset = (MasterPoseComponent != NULL)? MasterPoseComponent->GetPhysicsAsset() : NULL;
+
+	// Can only use the PhysicsAsset to calculate the bounding box if we are not non-uniformly scaling the mesh.
+	const bool bCanUsePhysicsAsset = DrawScale.IsUniform() && (SkeletalMesh != NULL)
+		// either space base exists or child component
+		&& ( (GetNumSpaceBases() == SkeletalMesh->RefSkeleton.GetNum()) || (MasterPhysicsAsset) );
+
+	const bool bDetailModeAllowsRendering = (DetailMode <= GetCachedScalabilityCVars().DetailMode);
+	const bool bVisible = ( bDetailModeAllowsRendering && (ShouldRender() || bCastHiddenShadow));
+
+	const bool bHasPhysBodies = PhysicsAsset && PhysicsAsset->BodySetup.Num();
+	const bool bMasterHasPhysBodies = MasterPhysicsAsset && MasterPhysicsAsset->BodySetup.Num();
+
+	// if not visible, or we were told to use fixed bounds, use skelmesh bounds
+	if ( (!bVisible || bComponentUseFixedSkelBounds) && SkeletalMesh ) 
+	{
+		FBoxSphereBounds RootAdjustedBounds = SkeletalMesh->Bounds;
+		RootAdjustedBounds.Origin += RootOffset; // Adjust bounds by root bone translation
+		NewBounds = RootAdjustedBounds.TransformBy(LocalToWorld);
+	}
+	else if(MasterPoseComponent.IsValid() && MasterPoseComponent->SkeletalMesh && MasterPoseComponent->bComponentUseFixedSkelBounds)
+	{
+		FBoxSphereBounds RootAdjustedBounds = MasterPoseComponent->SkeletalMesh->Bounds;
+		RootAdjustedBounds.Origin += RootOffset; // Adjust bounds by root bone translation
+		NewBounds = RootAdjustedBounds.TransformBy(LocalToWorld);
+	}
+	// Use MasterPoseComponent's PhysicsAsset if told to
+	else if (MasterPoseComponent.IsValid() && bCanUsePhysicsAsset && bUseBoundsFromMasterPoseComponent)
+	{
+		NewBounds = MasterPoseComponent->Bounds;
+	}
+#if WITH_EDITOR
+	// For AnimSet Viewer, use 'bounds preview' physics asset if present.
+	else if(SkeletalMesh && bHasPhysBodies && bCanUsePhysicsAsset)
+	{
+		NewBounds = FBoxSphereBounds(PhysicsAsset->CalcAABB(this, LocalToWorld));
+	}
+#endif // WITH_EDITOR
+	// If we have a PhysicsAsset (with at least one matching bone), and we can use it, do so to calc bounds.
+	else if( bHasPhysBodies && bCanUsePhysicsAsset && UsePhysicsAsset )
+	{
+		NewBounds = FBoxSphereBounds(PhysicsAsset->CalcAABB(this, LocalToWorld));
+	}
+	// Use MasterPoseComponent's PhysicsAsset, if we don't have one and it does
+	else if(MasterPoseComponent.IsValid() && bCanUsePhysicsAsset && bMasterHasPhysBodies)
+	{
+		NewBounds = FBoxSphereBounds(MasterPhysicsAsset->CalcAABB(this, LocalToWorld));
+	}
+	// Fallback is to use the one from the skeletal mesh. Usually pretty bad in terms of Accuracy of where the SkelMesh Bounds are located (i.e. usually bigger than it needs to be)
+	else if( SkeletalMesh )
+	{
+		FBoxSphereBounds RootAdjustedBounds = SkeletalMesh->Bounds;
+
+		// Adjust bounds by root bone translation
+		RootAdjustedBounds.Origin += RootOffset;
+		NewBounds = RootAdjustedBounds.TransformBy(LocalToWorld);
+	}
+	else
+	{
+		NewBounds = FBoxSphereBounds(LocalToWorld.GetLocation(), FVector::ZeroVector, 0.f);
+	}
+
+	// Add bounds of any per-poly collision data.
+	// TODO UE4
+
+	NewBounds.BoxExtent *= BoundsScale;
+	NewBounds.SphereRadius *= BoundsScale;
+
+	return NewBounds;
+}
+
+
+FMatrix USkinnedMeshComponent::GetBoneMatrix(int32 BoneIdx) const
+{
+	if ( !IsRegistered() )
+	{
+		// if not registered, we don't have SpaceBases yet. 
+		// also ComponentToWorld isn't set yet (They're set from relativetranslation, relativerotation, relativescale)
+		return FMatrix::Identity;
+	}
+
+	// Handle case of use a MasterPoseComponent - get bone matrix from there.
+	if(MasterPoseComponent.IsValid())
+	{
+		if(BoneIdx < MasterBoneMap.Num())
+		{
+			int32 ParentBoneIndex = MasterBoneMap[BoneIdx];
+
+			// If ParentBoneIndex is valid, grab matrix from MasterPoseComponent.
+			if(	ParentBoneIndex != INDEX_NONE && 
+				ParentBoneIndex < MasterPoseComponent->GetNumSpaceBases())
+			{
+				return MasterPoseComponent->GetSpaceBases()[ParentBoneIndex].ToMatrixWithScale() * ComponentToWorld.ToMatrixWithScale();
+			}
+			else
+			{
+				UE_LOG(LogSkinnedMeshComp, Warning, TEXT("GetBoneMatrix : ParentBoneIndex(%d) out of range of MasterPoseComponent->SpaceBases for %s"), BoneIdx, *GetPathName());
+				return FMatrix::Identity;
+			}
+		}
+		else
+		{
+			UE_LOG(LogSkinnedMeshComp, Warning, TEXT("GetBoneMatrix : BoneIndex(%d) out of range of MasterBoneMap for %s (%s)"), BoneIdx, *this->GetFName().ToString(), SkeletalMesh ? *SkeletalMesh->GetFName().ToString() : TEXT("NULL") );
+			return FMatrix::Identity;
+		}
+	}
+	else
+	{
+		if( GetNumSpaceBases() && BoneIdx < GetNumSpaceBases() )
+		{
+			return GetSpaceBases()[BoneIdx].ToMatrixWithScale() * ComponentToWorld.ToMatrixWithScale();
+		}
+		else
+		{
+			UE_LOG(LogSkinnedMeshComp, Warning, TEXT("GetBoneMatrix : BoneIndex(%d) out of range of SpaceBases for %s (%s)"), BoneIdx, *GetPathName(), SkeletalMesh ? *SkeletalMesh->GetFullName() : TEXT("NULL") );
+			return FMatrix::Identity;
+		}
+	}
+}
+
+FTransform USkinnedMeshComponent::GetBoneTransform(int32 BoneIdx) const
+{
+	if (!IsRegistered())
+	{
+		// if not registered, we don't have SpaceBases yet. 
+		// also ComponentToWorld isn't set yet (They're set from relativelocation, relativerotation, relativescale)
+		return FTransform::Identity;
+	}
+
+	return GetBoneTransform(BoneIdx, ComponentToWorld);
+}
+
+FTransform USkinnedMeshComponent::GetBoneTransform(int32 BoneIdx, const FTransform& LocalToWorld) const
+{
+	// Handle case of use a MasterPoseComponent - get bone matrix from there.
+	if(MasterPoseComponent.IsValid())
+	{
+		if(BoneIdx < MasterBoneMap.Num())
+		{
+			int32 ParentBoneIndex = MasterBoneMap[BoneIdx];
+
+			// If ParentBoneIndex is valid, grab matrix from MasterPoseComponent.
+			if(	ParentBoneIndex != INDEX_NONE && 
+				ParentBoneIndex < MasterPoseComponent->GetNumSpaceBases())
+			{
+				return MasterPoseComponent->GetSpaceBases()[ParentBoneIndex] * LocalToWorld;
+			}
+			else
+			{
+				UE_LOG(LogSkinnedMeshComp, Warning, TEXT("GetBoneTransform : ParentBoneIndex(%d) out of range of MasterPoseComponent->SpaceBases for %s"), BoneIdx, *this->GetFName().ToString() );
+				return FTransform::Identity;
+			}
+		}
+		else
+		{
+			UE_LOG(LogSkinnedMeshComp, Warning, TEXT("GetBoneTransform : BoneIndex(%d) out of range of MasterBoneMap for %s"), BoneIdx, *this->GetFName().ToString() );
+			return FTransform::Identity;
+		}
+	}
+	else
+	{
+		if( GetNumSpaceBases() && BoneIdx < GetNumSpaceBases() )
+		{
+			return GetSpaceBases()[BoneIdx] * LocalToWorld;
+		}
+		else
+		{
+			UE_LOG(LogSkinnedMeshComp, Warning, TEXT("GetBoneTransform : BoneIndex(%d) out of range of SpaceBases for %s (%s)"), BoneIdx, *GetPathName(), SkeletalMesh ? *SkeletalMesh->GetFullName() : TEXT("NULL") );
+			return FTransform::Identity;
+		}
+	}
+}
+
+
+int32 USkinnedMeshComponent::GetBoneIndex( FName BoneName) const
+{
+	int32 BoneIndex = INDEX_NONE;
+	if ( BoneName != NAME_None && SkeletalMesh )
+	{
+		BoneIndex = SkeletalMesh->RefSkeleton.FindBoneIndex( BoneName );
+	}
+
+	return BoneIndex;
+}
+
+
+FName USkinnedMeshComponent::GetBoneName(int32 BoneIndex) const
+{
+	return (SkeletalMesh != NULL && SkeletalMesh->RefSkeleton.IsValidIndex(BoneIndex)) ? SkeletalMesh->RefSkeleton.GetBoneName(BoneIndex) : NAME_None;
+}
+
+
+FName USkinnedMeshComponent::GetParentBone( FName BoneName ) const
+{
+	FName Result = NAME_None;
+
+	int32 BoneIndex = GetBoneIndex(BoneName);
+	if ((BoneIndex != INDEX_NONE) && (BoneIndex > 0)) // This checks that this bone is not the root (ie no parent), and that BoneIndex != INDEX_NONE (ie bone name was found)
+	{
+		Result = SkeletalMesh->RefSkeleton.GetBoneName(SkeletalMesh->RefSkeleton.GetParentIndex(BoneIndex));
+	}
+	return Result;
+}
+
+
+void USkinnedMeshComponent::GetBoneNames(TArray<FName>& BoneNames)
+{
+	if (SkeletalMesh == NULL)
+	{
+		// no mesh, so no bones
+		BoneNames.Empty();
+	}
+	else
+	{
+		// pre-size the array to avoid unnecessary reallocation
+		BoneNames.Empty(SkeletalMesh->RefSkeleton.GetNum());
+		BoneNames.AddUninitialized(SkeletalMesh->RefSkeleton.GetNum());
+		for (int32 i = 0; i < SkeletalMesh->RefSkeleton.GetNum(); i++)
+		{
+			BoneNames[i] = SkeletalMesh->RefSkeleton.GetBoneName(i);
+		}
+	}
+}
+
+
+bool USkinnedMeshComponent::BoneIsChildOf(FName BoneName, FName ParentBoneName) const
+{
+	bool bResult = false;
+
+	if( SkeletalMesh )
+	{
+		const int32 BoneIndex = SkeletalMesh->RefSkeleton.FindBoneIndex(BoneName);
+		if(BoneIndex == INDEX_NONE)
+		{
+			UE_LOG(LogSkinnedMeshComp, Log, TEXT("execBoneIsChildOf: BoneName '%s' not found in SkeletalMesh '%s'"), *BoneName.ToString(), *SkeletalMesh->GetName());
+			return bResult;
+		}
+
+		const int32 ParentBoneIndex = SkeletalMesh->RefSkeleton.FindBoneIndex(ParentBoneName);
+		if(ParentBoneIndex == INDEX_NONE)
+		{
+			UE_LOG(LogSkinnedMeshComp, Log, TEXT("execBoneIsChildOf: ParentBoneName '%s' not found in SkeletalMesh '%s'"), *ParentBoneName.ToString(), *SkeletalMesh->GetName());
+			return bResult;
+		}
+
+		bResult = SkeletalMesh->RefSkeleton.BoneIsChildOf(BoneIndex, ParentBoneIndex);
+	}
+
+	return bResult;
+}
+
+
+FVector USkinnedMeshComponent::GetRefPosePosition(int32 BoneIndex)
+{
+	if(SkeletalMesh && (BoneIndex >= 0) && (BoneIndex < SkeletalMesh->RefSkeleton.GetNum()))
+	{
+		return SkeletalMesh->RefSkeleton.GetRefBonePose()[BoneIndex].GetTranslation();
+	}
+	else
+	{
+		return FVector::ZeroVector;
+	}
+}
+
+
+void USkinnedMeshComponent::SetSkeletalMesh(USkeletalMesh* InSkelMesh)
+{
+	// NOTE: InSkelMesh may be NULL (useful in the editor for removing the skeletal mesh associated with
+	//   this component on-the-fly)
+
+	if (InSkelMesh == SkeletalMesh)
+	{
+		// do nothing if the input mesh is the same mesh we're already using.
+		return;
+	}
+
+	// Unregister the component, swap the skeletal mesh, and reregister.
+	{
+		FComponentReregisterContext ReregisterContext(this);
+		check(MeshObject == NULL);
+		SkeletalMesh = InSkelMesh;
+		InvalidateCachedBounds();
+	}
+	
+	// TODO: (LH) find better way to call this 
+	// Right now SkeletalMeshComponent(child) needs to call this after animation update
+	// But some of functions of SkinnedMeshComponent has to be called before
+	// So it's hard to split both separate, but do we really need to call this when skeletalmesh is set
+	// Think about this
+	// RefreshBoneTransforms(); 
+	
+	// Notify the streaming system. Don't use Update(), because this may be the first time the mesh has been set
+	// and the component may have to be added to the streaming system for the first time.
+	IStreamingManager::Get().NotifyPrimitiveAttached( this, DPT_Spawned );
+}
+
+FSkeletalMeshResource* USkinnedMeshComponent::GetSkeletalMeshResource() const
+{
+	if (MeshObject)
+	{
+		return &MeshObject->GetSkeletalMeshResource();
+	}
+	else if (SkeletalMesh)
+	{
+		ERHIFeatureLevel::Type SceneFeatureLevel = GetWorld()->FeatureLevel;
+		return SkeletalMesh->GetResourceForRendering();
+	}
+	else
+	{
+		return NULL;
+	}
+}
+
+bool USkinnedMeshComponent::AllocateTransformData()
+{
+	// Allocate transforms if not present.
+	if ( SkeletalMesh != NULL && MasterPoseComponent == NULL )
+	{
+		if( GetNumSpaceBases() != SkeletalMesh->RefSkeleton.GetNum() )
+		{
+			for (int32 BaseIndex = 0; BaseIndex < 2; ++BaseIndex)
+			{
+				SpaceBasesArray[BaseIndex].Empty(SkeletalMesh->RefSkeleton.GetNum());
+				SpaceBasesArray[BaseIndex].AddUninitialized(SkeletalMesh->RefSkeleton.GetNum());
+
+				for (int32 I = 0; I < SkeletalMesh->RefSkeleton.GetNum(); ++I)
+				{
+					SpaceBasesArray[BaseIndex][I].SetIdentity();
+				}
+			}
+ 
+			BoneVisibilityStates.Empty( SkeletalMesh->RefSkeleton.GetNum() );
+			if( SkeletalMesh->RefSkeleton.GetNum() )
+			{
+				BoneVisibilityStates.AddUninitialized( SkeletalMesh->RefSkeleton.GetNum() );
+				for (int32 BoneIndex = 0; BoneIndex < SkeletalMesh->RefSkeleton.GetNum(); BoneIndex++)
+				{
+					BoneVisibilityStates[ BoneIndex ] = BVS_Visible;
+				}
+			}
+		}
+
+		// if it's same, do not touch, and return
+		return true;
+	}
+	
+	// Reset the animation stuff when changing mesh.
+	SpaceBasesArray[0].Empty();
+	SpaceBasesArray[1].Empty();
+
+	return false;
+}
+
+void USkinnedMeshComponent::DeallocateTransformData()
+{
+	SpaceBasesArray[0].Empty();
+	SpaceBasesArray[1].Empty();
+	BoneVisibilityStates.Empty();
+}
+
+void USkinnedMeshComponent::SetPhysicsAsset(class UPhysicsAsset* InPhysicsAsset, bool bForceReInit)
+{
+	PhysicsAssetOverride = InPhysicsAsset;
+}
+
+
+void USkinnedMeshComponent::SetMasterPoseComponent(class USkinnedMeshComponent* NewMasterBoneComponent)
+{
+	MasterPoseComponent = NewMasterBoneComponent;
+
+	// now add to slave components list, 
+	if(MasterPoseComponent.IsValid())
+	{
+		bool bAddNew=true;
+		// make sure no empty element is there, this is weak obj ptr, so it will go away unless there is 
+		// other reference, this is intentional as master to slave reference is weak
+		for ( auto Iter = MasterPoseComponent->SlavePoseComponents.CreateIterator(); Iter; ++Iter )
+		{
+			TWeakObjectPtr<USkinnedMeshComponent> Comp = (*Iter);
+			if (Comp.IsValid() == false)
+			{
+				// remove
+				MasterPoseComponent->SlavePoseComponents.RemoveAt(Iter.GetIndex());
+				--Iter;
+			}
+			// if it has same as me, ignore to add
+			else if (Comp.Get() == this)
+			{
+				bAddNew	= false;
+			}
+		}
+
+		if (bAddNew)
+		{
+			MasterPoseComponent->SlavePoseComponents.Add(this);
+		}
+	}
+
+	AllocateTransformData();
+	UpdateMasterBoneMap();
+}
+
+void USkinnedMeshComponent::InvalidateCachedBounds()
+{
+	bCachedLocalBoundsUpToDate = false;
+
+	// Also invalidate all slave components.
+	if( SlavePoseComponents.Num() > 0 )
+	{
+		for (auto Iter=SlavePoseComponents.CreateIterator(); Iter; ++Iter)
+		{
+			TWeakObjectPtr<USkinnedMeshComponent> SkinnedMeshComp = (*Iter);
+			if( SkinnedMeshComp.IsValid() )
+			{
+				SkinnedMeshComp->bCachedLocalBoundsUpToDate = false;
+			}
+		}
+	}
+}
+
+void USkinnedMeshComponent::RefreshSlaveComponents()
+{
+	if ( SlavePoseComponents.Num() > 0 )
+	{
+		for (auto Iter=SlavePoseComponents.CreateIterator(); Iter; ++Iter)
+		{
+			TWeakObjectPtr<USkinnedMeshComponent> MeshComp = (*Iter);
+			if (MeshComp.IsValid())
+			{
+				MeshComp->MarkRenderDynamicDataDirty();
+			}
+		}
+	}
+}
+
+void USkinnedMeshComponent::SetForceWireframe(bool InForceWireframe)
+{
+	if(bForceWireframe != InForceWireframe)
+	{
+		bForceWireframe = InForceWireframe;
+		MarkRenderStateDirty();
+	}
+}
+
+
+UMorphTarget* USkinnedMeshComponent::FindMorphTarget( FName MorphTargetName ) const
+{
+	if( SkeletalMesh != NULL )
+	{
+		return SkeletalMesh->FindMorphTarget(MorphTargetName);
+	}
+
+	return NULL;
+}
+
+
+void USkinnedMeshComponent::UpdateMasterBoneMap()
+{
+	MasterBoneMap.Empty();
+
+	if(SkeletalMesh && MasterPoseComponent.IsValid() && MasterPoseComponent->SkeletalMesh)
+	{
+		USkeletalMesh* ParentMesh = MasterPoseComponent->SkeletalMesh;
+
+		MasterBoneMap.Empty( SkeletalMesh->RefSkeleton.GetNum() );
+		MasterBoneMap.AddUninitialized( SkeletalMesh->RefSkeleton.GetNum() );
+		if (SkeletalMesh == ParentMesh)
+		{
+			// if the meshes are the same, the indices must match exactly so we don't need to look them up
+			for (int32 i = 0; i < MasterBoneMap.Num(); i++)
+			{
+				MasterBoneMap[i] = i;
+			}
+		}
+		else
+		{
+			for(int32 i=0; i<MasterBoneMap.Num(); i++)
+			{
+				FName BoneName = SkeletalMesh->RefSkeleton.GetBoneName(i);
+				MasterBoneMap[i] = ParentMesh->RefSkeleton.FindBoneIndex( BoneName );
+			}
+		}
+	}
+}
+
+FTransform USkinnedMeshComponent::GetSocketTransform(FName InSocketName, ERelativeTransformSpace TransformSpace) const
+{
+	FTransform OutSocketTransform = ComponentToWorld;
+
+	if (InSocketName != NAME_None)
+	{
+		int32 BoneIndex = GetBoneIndex(InSocketName);
+		if( BoneIndex != INDEX_NONE )
+		{
+			OutSocketTransform = GetBoneTransform(BoneIndex);
+		}
+		else
+		{
+			USkeletalMeshSocket const* const Socket = GetSocketByName(InSocketName);
+			if (Socket)
+			{
+				FTransform SocketLocalTransform = Socket->GetSocketLocalTransform();
+
+				BoneIndex = GetBoneIndex(Socket->BoneName);
+				if( BoneIndex != INDEX_NONE )
+				{
+					FTransform BoneTransform = GetBoneTransform(BoneIndex);
+					OutSocketTransform = SocketLocalTransform * BoneTransform; 
+				}
+			}
+		}
+	}
+
+	switch(TransformSpace)
+	{
+		case RTS_Actor:
+		{
+			if( AActor* Actor = GetOwner() )
+			{
+				return OutSocketTransform.GetRelativeTransform( Actor->GetTransform() );
+			}
+			break;
+		}
+		case RTS_Component:
+		{
+			return OutSocketTransform.GetRelativeTransform( ComponentToWorld );
+		}
+	}
+
+	return OutSocketTransform;
+}
+
+class USkeletalMeshSocket const* USkinnedMeshComponent::GetSocketByName(FName InSocketName) const
+{
+	USkeletalMeshSocket const* Socket = NULL;
+
+	if( SkeletalMesh )
+	{
+		Socket = SkeletalMesh->FindSocket( InSocketName );
+	}
+	else
+	{
+		UE_LOG(LogSkinnedMeshComp, Warning,TEXT("GetSocketByName(): No SkeletalMesh for %s"), *GetName());
+	}
+
+	return Socket;
+}
+
+bool USkinnedMeshComponent::DoesSocketExist(FName InSocketName) const
+{
+	return (GetSocketBoneName(InSocketName) != NAME_None);
+}
+
+FName USkinnedMeshComponent::GetSocketBoneName(FName InSocketName) const
+{
+	if(!SkeletalMesh)
+	{
+		return NAME_None;
+	}
+
+	// First check for a socket
+	USkeletalMeshSocket const* TmpSocket = SkeletalMesh->FindSocket(InSocketName);
+	if( TmpSocket )
+	{
+		return TmpSocket->BoneName;
+	}
+
+	// If socket is not found, maybe it was just a bone name.
+	if( GetBoneIndex(InSocketName) != INDEX_NONE )
+	{
+		return InSocketName;
+	}
+
+	// Doesn't exist.
+	return NAME_None;
+}
+
+
+FQuat USkinnedMeshComponent::GetBoneQuaternion(FName BoneName, EBoneSpaces::Type Space) const
+{
+	int32 BoneIndex = GetBoneIndex(BoneName);
+
+	if( BoneIndex == INDEX_NONE )
+	{
+		UE_LOG(LogSkinnedMeshComp, Warning, TEXT("USkinnedMeshComponent::execGetBoneQuaternion : Could not find bone: %s"), *BoneName.ToString());
+		return FQuat::Identity;
+	}
+
+	FTransform BoneTransform;
+	if( Space == EBoneSpaces::ComponentSpace )
+	{
+		if(MasterPoseComponent.IsValid())
+		{
+			if(BoneIndex < MasterBoneMap.Num())
+			{
+				int32 ParentBoneIndex = MasterBoneMap[BoneIndex];
+				// If ParentBoneIndex is valid, grab matrix from MasterPoseComponent.
+				if(	ParentBoneIndex != INDEX_NONE && 
+					ParentBoneIndex < MasterPoseComponent->GetNumSpaceBases())
+				{
+					BoneTransform = MasterPoseComponent->GetSpaceBases()[ParentBoneIndex];
+				}
+				else
+				{
+					BoneTransform = FTransform::Identity;
+				}
+			}
+			else
+			{
+				BoneTransform = FTransform::Identity;
+			}
+		}
+		else
+		{
+			BoneTransform = GetSpaceBases()[BoneIndex];
+		}
+	}
+	else
+	{
+		BoneTransform = GetBoneTransform(BoneIndex);
+	}
+
+	BoneTransform.RemoveScaling();
+	return BoneTransform.GetRotation();
+}
+
+
+FVector USkinnedMeshComponent::GetBoneLocation(FName BoneName, EBoneSpaces::Type Space) const
+{
+	int32 BoneIndex = GetBoneIndex(BoneName);
+	if( BoneIndex == INDEX_NONE )
+	{
+		UE_LOG(LogAnimation, Log, TEXT("USkinnedMeshComponent::GetBoneLocation (%s %s): Could not find bone: %s"), *GetFullName(), *GetDetailedInfo(), *BoneName.ToString() );
+		return FVector::ZeroVector;
+	}
+
+	if( Space == EBoneSpaces::ComponentSpace )
+	{
+		if(MasterPoseComponent.IsValid())
+		{
+			if(BoneIndex < MasterBoneMap.Num())
+			{
+				int32 ParentBoneIndex = MasterBoneMap[BoneIndex];
+				// If ParentBoneIndex is valid, grab transform from MasterPoseComponent.
+				if(	ParentBoneIndex != INDEX_NONE && 
+					ParentBoneIndex < MasterPoseComponent->GetNumSpaceBases())
+				{
+					return MasterPoseComponent->GetSpaceBases()[ParentBoneIndex].GetLocation();
+				}
+			}
+			
+			// return empty vector
+			return FVector::ZeroVector;			
+		}
+		else
+		{
+			return GetSpaceBases()[BoneIndex].GetLocation();
+		}
+	}
+	else if (Space == EBoneSpaces::WorldSpace)
+	{
+		// To support non-uniform scale (via LocalToWorld), use GetBoneMatrix
+		return GetBoneMatrix(BoneIndex).GetOrigin();
+	}
+	else
+	{
+		check(false); // Unknown BoneSpace
+		return FVector::ZeroVector;
+	}
+}
+
+
+FVector USkinnedMeshComponent::GetBoneAxis( FName BoneName, EAxis::Type Axis ) const
+{
+	int32 BoneIndex = GetBoneIndex(BoneName);
+	if (BoneIndex == INDEX_NONE)
+	{
+		UE_LOG(LogSkinnedMeshComp, Warning, TEXT("USkinnedMeshComponent::execGetBoneAxis : Could not find bone: %s"), *BoneName.ToString());
+		return FVector::ZeroVector;
+	}
+	else if (Axis == EAxis::None)
+	{
+		UE_LOG(LogSkinnedMeshComp, Warning, TEXT("USkinnedMeshComponent::execGetBoneAxis: Invalid axis specified"));
+		return FVector::ZeroVector;
+	}
+	else
+	{
+		return GetBoneMatrix(BoneIndex).GetUnitAxis(Axis);
+	}
+}
+
+bool USkinnedMeshComponent::HasAnySockets() const
+{
+	return (SkeletalMesh != NULL) && (
+#if WITH_EDITOR
+		(SkeletalMesh->GetActiveSocketList().Num() > 0) ||
+#endif
+		(SkeletalMesh->RefSkeleton.GetNum() > 0));
+}
+
+void USkinnedMeshComponent::QuerySupportedSockets(TArray<FComponentSocketDescription>& OutSockets) const
+{
+	if (SkeletalMesh != NULL)
+	{
+		//@TODO: need to make this work in the game too
+#if WITH_EDITOR
+		// Grab all the mesh and skeleton sockets
+		const TArray<USkeletalMeshSocket*> AllSockets = SkeletalMesh->GetActiveSocketList();
+
+		for (int32 SocketIdx = 0; SocketIdx < AllSockets.Num(); ++SocketIdx)
+		{
+			if (USkeletalMeshSocket* Socket = AllSockets[SocketIdx])
+			{
+				new (OutSockets) FComponentSocketDescription(Socket->SocketName, EComponentSocketType::Socket);
+			}
+		}
+#endif
+
+		// Now grab the bones, which can behave exactly like sockets
+		for (int32 BoneIdx = 0; BoneIdx < SkeletalMesh->RefSkeleton.GetNum(); ++BoneIdx)
+		{
+			const FName BoneName = SkeletalMesh->RefSkeleton.GetBoneName(BoneIdx);
+			new (OutSockets) FComponentSocketDescription(BoneName, EComponentSocketType::Bone);
+		}
+	}
+}
+
+void USkinnedMeshComponent::UpdateOverlaps(TArray<FOverlapInfo> const* PendingOverlaps, bool bDoNotifies, const TArray<FOverlapInfo>* OverlapsAtEndLocation)
+{
+	// we don't support overlap test on destructible or physics asset
+	// so use SceneComponent::UpdateOverlaps to handle children
+	USceneComponent::UpdateOverlaps(PendingOverlaps, bDoNotifies, OverlapsAtEndLocation);
+}
+
+void USkinnedMeshComponent::TransformToBoneSpace(FName BoneName, FVector InPosition, FRotator InRotation, FVector& OutPosition, FRotator& OutRotation) const
+{
+	int32 BoneIndex = GetBoneIndex(BoneName);
+	if(BoneIndex != INDEX_NONE)
+	{
+		FMatrix BoneToWorldTM = GetBoneMatrix(BoneIndex);
+		FMatrix WorldTM = FRotationTranslationMatrix(InRotation, InPosition);
+		FMatrix LocalTM = WorldTM * BoneToWorldTM.Inverse();
+
+		OutPosition = LocalTM.GetOrigin();
+		OutRotation = LocalTM.Rotator();
+	}
+}
+
+
+void USkinnedMeshComponent::TransformFromBoneSpace(FName BoneName, FVector InPosition, FRotator InRotation, FVector& OutPosition, FRotator& OutRotation)
+{
+	int32 BoneIndex = GetBoneIndex(BoneName);
+	if(BoneIndex != INDEX_NONE)
+	{
+		FMatrix BoneToWorldTM = GetBoneMatrix(BoneIndex);
+
+		FMatrix LocalTM = FRotationTranslationMatrix(InRotation, InPosition);
+		FMatrix WorldTM = LocalTM * BoneToWorldTM;
+
+		OutPosition = WorldTM.GetOrigin();
+		OutRotation = WorldTM.Rotator();
+	}
+}
+
+
+
+FName USkinnedMeshComponent::FindClosestBone(FVector TestLocation, FVector* BoneLocation, float IgnoreScale) const
+{
+	if (SkeletalMesh == NULL)
+	{
+		if (BoneLocation != NULL)
+		{
+			*BoneLocation = FVector::ZeroVector;
+		}
+		return NAME_None;
+	}
+	else
+	{
+		// transform the TestLocation into mesh local space so we don't have to transform the (mesh local) bone locations
+		TestLocation = ComponentToWorld.InverseTransformPosition(TestLocation);
+		
+		float IgnoreScaleSquared = FMath::Square(IgnoreScale);
+		float BestDistSquared = BIG_NUMBER;
+		int32 BestIndex = -1;
+		for (int32 i = 0; i < GetNumSpaceBases(); i++)
+		{
+			if (IgnoreScale < 0.f || GetSpaceBases()[i].GetScaledAxis( EAxis::X ).SizeSquared() > IgnoreScaleSquared)
+			{
+				float DistSquared = (TestLocation - GetSpaceBases()[i].GetLocation()).SizeSquared();
+				if (DistSquared < BestDistSquared)
+				{
+					BestIndex = i;
+					BestDistSquared = DistSquared;
+				}
+			}
+		}
+
+		if (BestIndex == -1)
+		{
+			if (BoneLocation != NULL)
+			{
+				*BoneLocation = FVector::ZeroVector;
+			}
+			return NAME_None;
+		}
+		else
+		{
+			// transform the bone location into world space
+			if (BoneLocation != NULL)
+			{
+				*BoneLocation = (GetSpaceBases()[BestIndex] * ComponentToWorld).GetLocation();
+			}
+			return SkeletalMesh->RefSkeleton.GetBoneName(BestIndex);
+		}
+	}
+}
+
+void USkinnedMeshComponent::ShowMaterialSection(int32 MaterialID, bool bShow, int32 LODIndex)
+{
+	if (!SkeletalMesh)
+	{
+		// no skeletalmesh, then nothing to do. 
+		return;
+	}
+	// Make sure LOD info for this component has been initialized
+	InitLODInfos();
+	if (LODInfo.IsValidIndex(LODIndex))
+	{
+		const FSkeletalMeshLODInfo& SkelLODInfo = SkeletalMesh->LODInfo[LODIndex];
+		FSkelMeshComponentLODInfo& SkelCompLODInfo = LODInfo[LODIndex];
+		TArray<bool>& HiddenMaterials = SkelCompLODInfo.HiddenMaterials;
+	
+		// allocate if not allocated yet
+		if ( HiddenMaterials.Num() != SkeletalMesh->Materials.Num() )
+		{
+			// Using skeletalmesh component because Materials.Num() should be <= SkeletalMesh->Materials.Num()		
+			HiddenMaterials.Empty(SkeletalMesh->Materials.Num());
+			HiddenMaterials.AddZeroed(SkeletalMesh->Materials.Num());
+		}
+		// If we are at a dropped LOD, route material index through the LODMaterialMap in the LODInfo struct.
+		int32 UseMaterialIndex = MaterialID;			
+		if(LODIndex > 0)
+		{
+			if(SkelLODInfo.LODMaterialMap.IsValidIndex(MaterialID))
+			{
+				UseMaterialIndex = SkelLODInfo.LODMaterialMap[MaterialID];
+				UseMaterialIndex = FMath::Clamp( UseMaterialIndex, 0, HiddenMaterials.Num() );
+			}
+		}
+		// Mark the mapped section material entry as visible/hidden
+		if (HiddenMaterials.IsValidIndex(UseMaterialIndex))
+		{
+			HiddenMaterials[UseMaterialIndex] = !bShow;
+		}
+
+		if ( MeshObject )
+		{
+			// need to send render thread for updated hidden section
+			ENQUEUE_UNIQUE_RENDER_COMMAND_THREEPARAMETER(
+				FUpdateHiddenSectionCommand, 
+				FSkeletalMeshObject*, MeshObject, MeshObject, 
+				TArray<bool>, HiddenMaterials, HiddenMaterials, 
+				int32, LODIndex, LODIndex,
+			{
+				MeshObject->SetHiddenMaterials(LODIndex,HiddenMaterials);
+			});
+		}
+	}
+}
+
+
+void USkinnedMeshComponent::GetUsedMaterials( TArray<UMaterialInterface*>& OutMaterials ) const
+{
+	if( SkeletalMesh )
+	{
+		// The max number of materials used is the max of the materials on the skeletal mesh and the materials on the mesh component
+		const int32 NumMaterials = FMath::Max( SkeletalMesh->Materials.Num(), OverrideMaterials.Num() );
+		for( int32 MatIdx = 0; MatIdx < NumMaterials; ++MatIdx )
+		{
+			// GetMaterial will determine the correct material to use for this index.  
+			UMaterialInterface* MaterialInterface = GetMaterial( MatIdx );
+			OutMaterials.Add( MaterialInterface );
+		}
+	}
+}
+
+template <bool bExtraBoneInfluencesT, bool bCachedMatrices>
+FORCEINLINE FVector USkinnedMeshComponent::GetTypedSkinnedVertexPosition(const FSkelMeshChunk& Chunk, const FSkeletalMeshVertexBuffer& VertexBufferGPUSkin, int32 VertIndex, bool bSoftVertex, const TArray<FMatrix> & RefToLocals) const
+{
+	FVector SkinnedPos(0,0,0);
+
+	const USkinnedMeshComponent* BaseComponent = MasterPoseComponent.IsValid() ? MasterPoseComponent.Get() : this;
+
+	// Do soft skinning for this vertex.
+	if(bSoftVertex)
+	{
+		const TGPUSkinVertexBase<bExtraBoneInfluencesT>* SrcSoftVertex = VertexBufferGPUSkin.GetVertexPtr<bExtraBoneInfluencesT>(Chunk.GetSoftVertexBufferIndex()+VertIndex);
+
+#if !PLATFORM_LITTLE_ENDIAN
+		// uint8[] elements in LOD.VertexBufferGPUSkin have been swapped for VET_UBYTE4 vertex stream use
+		for(int32 InfluenceIndex = MAX_INFLUENCES-1;InfluenceIndex >=  MAX_INFLUENCES-Chunk.MaxBoneInfluences;InfluenceIndex--)
+#else
+		for(int32 InfluenceIndex = 0;InfluenceIndex < Chunk.MaxBoneInfluences;InfluenceIndex++)
+#endif
+		{
+			int32 BoneIndex = Chunk.BoneMap[SrcSoftVertex->InfluenceBones[InfluenceIndex]];
+			if(MasterPoseComponent.IsValid())
+			{		
+				check(MasterBoneMap.Num() == SkeletalMesh->RefSkeleton.GetNum());
+				BoneIndex = MasterBoneMap[BoneIndex];
+			}
+
+			const float	Weight = (float)SrcSoftVertex->InfluenceWeights[InfluenceIndex] / 255.0f;
+			{
+				if (bCachedMatrices)
+				{
+					const FMatrix& RefToLocal = RefToLocals[BoneIndex];
+					SkinnedPos += RefToLocal.TransformPosition(VertexBufferGPUSkin.GetVertexPositionFast(SrcSoftVertex)) * Weight;
+				}
+				else
+				{
+					const FMatrix RefToLocal = SkeletalMesh->RefBasesInvMatrix[BoneIndex] * BaseComponent->GetSpaceBases()[BoneIndex].ToMatrixWithScale();
+					SkinnedPos += RefToLocal.TransformPosition(VertexBufferGPUSkin.GetVertexPositionFast(SrcSoftVertex)) * Weight;
+				}
+			}
+		}
+	}
+	// Do rigid (one-influence) skinning for this vertex.
+	else
+	{
+		const TGPUSkinVertexBase<false>* SrcRigidVertex = VertexBufferGPUSkin.GetVertexPtr<false>(Chunk.GetRigidVertexBufferIndex()+VertIndex);
+		const int32 RigidInfluenceIndex = SkinningTools::GetRigidInfluenceIndex();
+		int32 BoneIndex = Chunk.BoneMap[SrcRigidVertex->InfluenceBones[RigidInfluenceIndex]];
+		if(MasterPoseComponent.IsValid())
+		{
+			check(MasterBoneMap.Num() == SkeletalMesh->RefSkeleton.GetNum());
+			BoneIndex = MasterBoneMap[BoneIndex];
+		}
+
+		if (bCachedMatrices)
+		{
+			const FMatrix& RefToLocal = RefToLocals[BoneIndex];
+			SkinnedPos = RefToLocal.TransformPosition(VertexBufferGPUSkin.GetVertexPositionFast(SrcRigidVertex));
+		}
+		else
+		{
+			const FMatrix RefToLocal = SkeletalMesh->RefBasesInvMatrix[BoneIndex] * BaseComponent->GetSpaceBases()[BoneIndex].ToMatrixWithScale();
+			SkinnedPos = RefToLocal.TransformPosition(VertexBufferGPUSkin.GetVertexPositionFast(SrcRigidVertex));
+		}
+	}
+
+	return SkinnedPos;
+}
+
+FVector USkinnedMeshComponent::GetSkinnedVertexPosition(int32 VertexIndex) const
+{
+	FVector SkinnedPos(0,0,0);
+
+	// Fail if no mesh
+	if(!SkeletalMesh || !MeshObject)
+	{
+		return SkinnedPos;
+	}
+
+	FStaticLODModel& Model = MeshObject->GetSkeletalMeshResource().LODModels[0];
+
+	//cache RefToLocal matrices
+	int32 ChunkIndex;
+	int32 VertIndex;
+	bool bSoftVertex;
+	bool bHasExtraBoneInfluences;
+	Model.GetChunkAndSkinType(VertexIndex, ChunkIndex, VertIndex, bSoftVertex, bHasExtraBoneInfluences);
+
+	//update positions
+	check(ChunkIndex < Model.Chunks.Num());
+	const FSkelMeshChunk& Chunk = Model.Chunks[ChunkIndex];
+			//rigid
+
+	return bHasExtraBoneInfluences
+		? GetTypedSkinnedVertexPosition<true, false>(Chunk, Model.VertexBufferGPUSkin, VertIndex, bSoftVertex)
+			//soft
+		: GetTypedSkinnedVertexPosition<false, false>(Chunk, Model.VertexBufferGPUSkin, VertIndex, bSoftVertex);
+
+
+}
+
+void USkinnedMeshComponent::ComputeSkinnedPositions(TArray<FVector> & OutPositions) const
+{
+	OutPositions.Empty();
+
+	// Fail if no mesh
+	if (!SkeletalMesh || !MeshObject)
+	{
+		return;
+	}
+
+	FStaticLODModel& Model = MeshObject->GetSkeletalMeshResource().LODModels[0];
+	OutPositions.AddUninitialized(Model.NumVertices);
+
+	//cache RefToLocal matrices
+	const USkinnedMeshComponent* BaseComponent = MasterPoseComponent.IsValid() ? MasterPoseComponent.Get() : this;
+	TArray<FMatrix> RefToLocals;
+	RefToLocals.AddUninitialized(SkeletalMesh->RefBasesInvMatrix.Num());
+	{
+		for (int32 MatrixIdx = 0; MatrixIdx < RefToLocals.Num(); ++MatrixIdx)
+		{
+			RefToLocals[MatrixIdx] = SkeletalMesh->RefBasesInvMatrix[MatrixIdx] * BaseComponent->GetSpaceBases()[MatrixIdx].ToMatrixWithScale();
+		}
+	}
+
+	//update positions
+	for (int32 ChunkIdx = 0; ChunkIdx < Model.Chunks.Num(); ++ChunkIdx)
+	{
+		const FSkelMeshChunk& Chunk = Model.Chunks[ChunkIdx];
+		const uint32 NumChunkVerts = Chunk.GetNumVertices();
+		bool bHasExtraBoneInfluences = Chunk.HasExtraBoneInfluences();
+		{
+			//rigid
+			const uint32 RigidOffset = Chunk.GetRigidVertexBufferIndex();
+			const uint32 NumRigidVerts = Chunk.GetNumRigidVertices();
+			for (uint32 RigidIdx = 0; RigidIdx < NumRigidVerts; ++RigidIdx)
+			{
+				FVector SkinnedPosition = bHasExtraBoneInfluences ? GetTypedSkinnedVertexPosition<true, true>(Chunk, Model.VertexBufferGPUSkin, RigidIdx, false, RefToLocals) :
+																	GetTypedSkinnedVertexPosition<false, true>(Chunk, Model.VertexBufferGPUSkin, RigidIdx, false, RefToLocals);
+				OutPositions[RigidOffset + RigidIdx] = SkinnedPosition;
+			}
+		}
+		{
+			//soft
+			const uint32 SoftOffset = Chunk.GetSoftVertexBufferIndex();
+			const uint32 NumSoftVerts = Chunk.GetNumSoftVertices();
+			for (uint32 SoftIdx = 0; SoftIdx < NumSoftVerts; ++SoftIdx)
+			{
+				FVector SkinnedPosition = bHasExtraBoneInfluences ? GetTypedSkinnedVertexPosition<true, true>(Chunk, Model.VertexBufferGPUSkin, SoftIdx, true, RefToLocals) :
+																	GetTypedSkinnedVertexPosition<false,true>(Chunk, Model.VertexBufferGPUSkin, SoftIdx, true, RefToLocals);
+				OutPositions[SoftOffset + SoftIdx] = SkinnedPosition;
+			}
+		}
+	}
+}
+
+FColor USkinnedMeshComponent::GetVertexColor(int32 VertexIndex) const
+{
+	// Fail if no mesh
+	if (!SkeletalMesh || !MeshObject)
+	{
+		return FColor(255, 255, 255, 255);
+	}
+
+	FStaticLODModel& Model = MeshObject->GetSkeletalMeshResource().LODModels[0];
+
+	// Find the chunk and vertex within that chunk, and skinning type, for this vertex.
+	int32 ChunkIndex;
+	int32 VertIndex;
+	bool bSoftVertex;
+	bool bHasExtraBoneInfluences;
+	Model.GetChunkAndSkinType(VertexIndex, ChunkIndex, VertIndex, bSoftVertex, bHasExtraBoneInfluences);
+
+	check(ChunkIndex < Model.Chunks.Num());
+	const FSkelMeshChunk& Chunk = Model.Chunks[ChunkIndex];
+	
+	int32 VertexBase = bSoftVertex ? Chunk.GetSoftVertexBufferIndex() : Chunk.GetRigidVertexBufferIndex();
+
+	return Model.ColorVertexBuffer.VertexColor(VertexBase + VertIndex);
+}
+
+
+void USkinnedMeshComponent::HideBone( int32 BoneIndex, EPhysBodyOp PhysBodyOption)
+{
+	if (ShouldUpdateBoneVisibility() && BoneIndex < BoneVisibilityStates.Num())
+	{
+		checkSlow ( BoneIndex != INDEX_NONE );
+		BoneVisibilityStates[ BoneIndex ] = BVS_ExplicitlyHidden;
+		RebuildVisibilityArray();
+	}
+}
+
+
+void USkinnedMeshComponent::UnHideBone( int32 BoneIndex )
+{
+	if (ShouldUpdateBoneVisibility() && BoneIndex < BoneVisibilityStates.Num())
+	{
+		checkSlow ( BoneIndex != INDEX_NONE );
+		//@TODO: If unhiding the child of a still hidden bone (coming in, BoneVisibilityStates(RefSkel(BoneIndex).ParentIndex) != BVS_Visible),
+		// should we be re-enabling collision bodies?
+		// Setting visible to true here is OK in either case as it will be reset to BVS_HiddenByParent in RecalcRequiredBones later if needed.
+		BoneVisibilityStates[ BoneIndex ] = BVS_Visible;
+		RebuildVisibilityArray();
+	}
+}
+
+
+bool USkinnedMeshComponent::IsBoneHidden( int32 BoneIndex ) const
+{
+	if (ShouldUpdateBoneVisibility() && BoneIndex < BoneVisibilityStates.Num())
+	{
+		if ( BoneIndex != INDEX_NONE )
+		{
+			return BoneVisibilityStates[ BoneIndex ] != BVS_Visible;
+		}
+	}
+	else if ( MasterPoseComponent != NULL )
+	{
+		return MasterPoseComponent->IsBoneHidden( BoneIndex );
+	}
+
+	return false;
+}
+
+
+bool USkinnedMeshComponent::IsBoneHiddenByName( FName BoneName )
+{
+	// Find appropriate BoneIndex
+	int32 BoneIndex = GetBoneIndex(BoneName);
+	if(BoneIndex != INDEX_NONE)
+	{
+		return IsBoneHidden(BoneIndex);
+	}
+
+	return false;
+}
+
+void USkinnedMeshComponent::HideBoneByName( FName BoneName, EPhysBodyOp PhysBodyOption )
+{
+	// Find appropriate BoneIndex
+	int32 BoneIndex = GetBoneIndex(BoneName);
+	if ( BoneIndex != INDEX_NONE )
+	{
+		HideBone(BoneIndex, PhysBodyOption);
+	}
+}
+
+
+void USkinnedMeshComponent::UnHideBoneByName( FName BoneName )
+{
+	int32 BoneIndex = GetBoneIndex(BoneName);
+	if ( BoneIndex != INDEX_NONE )
+	{
+		UnHideBone(BoneIndex);
+	}
+}
+
+bool USkinnedMeshComponent::UpdateLODStatus()
+{
+	// Predict the best (min) LOD level we are going to need. Basically we use the Min (best) LOD the renderer desired last frame.
+	// Because we update bones based on this LOD level, we have to update bones to this LOD before we can allow rendering at it.
+
+	if(SkeletalMesh != NULL)
+	{
+		int32 MaxLODIndex = 0;
+		if (MeshObject)
+		{
+			MaxLODIndex = MeshObject->GetSkeletalMeshResource().LODModels.Num() - 1;
+		}
+
+		// Support forcing to a particular LOD.
+		if(ForcedLodModel > 0)
+		{
+			PredictedLODLevel = FMath::Clamp(ForcedLodModel - 1, 0, MaxLODIndex);
+		}
+		else
+		{
+			// If no MeshObject - just assume lowest LOD.
+			if(MeshObject)
+			{
+				PredictedLODLevel = FMath::Clamp(MeshObject->MinDesiredLODLevel + GSkeletalMeshLODBias, 0, MaxLODIndex);
+			}
+			else
+			{
+				PredictedLODLevel = MaxLODIndex;
+			}
+		}
+
+		// now check to see if we have a MinLODLevel
+		if( ( MinLodModel > 0 ) && ( MinLodModel <= MaxLODIndex ) )
+		{
+			PredictedLODLevel = FMath::Clamp(PredictedLODLevel, MinLodModel, MaxLODIndex);
+		}
+	}
+	else
+	{
+		PredictedLODLevel = 0;
+	}
+
+	// See if LOD has changed. 
+	bool bLODChanged = (PredictedLODLevel != OldPredictedLODLevel);
+	OldPredictedLODLevel = PredictedLODLevel;
+
+	// Read back MaxDistanceFactor from the render object.
+	if(MeshObject)
+	{
+		MaxDistanceFactor = MeshObject->MaxDistanceFactor;
+
+#if CHART_DISTANCE_FACTORS
+		// Only chart DistanceFactor if it was actually rendered recently
+		if(bChartDistanceFactor && ((LastRenderTime > GetWorld()->TimeSeconds - 1.0f) || bUpdateSkelWhenNotRendered))
+		{
+			AddDistanceFactorToChart(MaxDistanceFactor);
+		}
+#endif // CHART_DISTANCE_FACTORS
+	}
+
+	return bLODChanged;
+}
+
+/** See if an array of ActiveVertexAnims already contains the supplied anim */
+int32 FindVertexAnim(const TArray<FActiveVertexAnim>& ActiveAnims, UVertexAnimBase* Anim)
+{
+	for(int32 i=0; i<ActiveAnims.Num(); i++)
+	{
+		if(ActiveAnims[i].VertAnim == Anim)
+		{
+			return i;
+		}
+	}
+
+	return INDEX_NONE;
+}
+
+TArray<FActiveVertexAnim> USkinnedMeshComponent::UpdateActiveVertexAnims(const USkeletalMesh* InSkeletalMesh, const TMap<FName, float>& MorphCurveAnims, const TArray<FActiveVertexAnim>& ActiveAnims)
+{
+	TArray<struct FActiveVertexAnim> OutVertexAnims;
+
+	// First copy ActiveAnims
+	for(int32 AnimIdx=0; AnimIdx < ActiveAnims.Num(); AnimIdx++)
+	{
+		const FActiveVertexAnim& ActiveAnim = ActiveAnims[AnimIdx];
+		const float ActiveAnimAbsWeight = FMath::Abs(ActiveAnim.Weight);
+
+		// Check it has valid weight, and works on this SkeletalMesh
+		if (	ActiveAnimAbsWeight > MinVertexAnimBlendWeight &&
+			ActiveAnim.VertAnim != NULL &&
+			ActiveAnim.VertAnim->BaseSkelMesh == InSkeletalMesh)
+		{
+			OutVertexAnims.Add(ActiveAnim);
+		}
+		// @TODO Need to check for duplicates here?
+	}
+
+	// Then go over the CurveKeys finding morph targets by name
+	for(auto CurveIter=MorphCurveAnims.CreateConstIterator(); CurveIter; ++CurveIter)
+	{
+		const FName& CurveName	= (CurveIter).Key();
+		const float& Weight	= (CurveIter).Value();
+
+		// If it has a valid weight
+		if(FMath::Abs(Weight) > MinVertexAnimBlendWeight)
+		{
+			// Find morph reference
+			UMorphTarget* Target = InSkeletalMesh ? InSkeletalMesh->FindMorphTarget(CurveName) : NULL;
+			if(Target != NULL)				
+			{
+				// See if this morph target already has an entry
+				int32 AnimIndex = FindVertexAnim(OutVertexAnims, Target);
+				// If not, add it
+				if(AnimIndex == INDEX_NONE)
+				{
+					OutVertexAnims.Add(FActiveVertexAnim(Target, Weight));
+				}
+				// If it does, use the max weight
+				else
+				{
+					const float CurrentWeight = OutVertexAnims[AnimIndex].Weight;
+					OutVertexAnims[AnimIndex].Weight = FMath::Max<float>(CurrentWeight, Weight);
+				}
+			}
+		}
+	}
+
+	return OutVertexAnims;
+}
+
+void USkinnedMeshComponent::FlipEditableSpaceBases()
+{
+	if (bNeedToFlipSpaceBaseBuffers)
+	{
+		bNeedToFlipSpaceBaseBuffers = false;
+		if (bDoubleBufferedBlendSpaces)
+		{
+			CurrentReadSpaceBases = CurrentEditableSpaceBases;
+			CurrentEditableSpaceBases = 1 - CurrentEditableSpaceBases;
+		}
+		else
+		{
+			CurrentReadSpaceBases = CurrentEditableSpaceBases = 0;
+		}
+	}
+}
+
+void USkinnedMeshComponent::SetSpaceBaseDoubleBuffering(bool bInDoubleBufferedBlendSpaces)
+{
+	bDoubleBufferedBlendSpaces = bInDoubleBufferedBlendSpaces;
+
+	if (bDoubleBufferedBlendSpaces)
+	{
+		CurrentEditableSpaceBases = 1 - CurrentReadSpaceBases;
+	}
+	else
+	{
+		CurrentEditableSpaceBases = CurrentReadSpaceBases;
+	}
+}
+
+void FAnimUpdateRateParameters::SetTrailMode(float DeltaTime, uint8 UpdateRateShift, int32 NewUpdateRate, int32 NewEvaluationRate, bool bNewInterpSkippedFrames)
+{
+	OptimizeMode = TrailMode;
+	ThisTickDelta = DeltaTime;
+
+	UpdateRate = FMath::Max(NewUpdateRate, 1);
+	// Make sure EvaluationRate is a multiple of UpdateRate.
+	EvaluationRate = FMath::Max((NewEvaluationRate / UpdateRate) * UpdateRate, 1);
+	bInterpolateSkippedFrames = bNewInterpSkippedFrames;
+
+	// Make sure we don't overflow. we don't need very large numbers.
+	const uint32 Counter = (GFrameCounter + UpdateRateShift)% MAX_uint8;
+
+	bSkipUpdate = ((Counter % UpdateRate) > 0);
+	bSkipEvaluation = ((Counter % EvaluationRate) > 0);
+
+	AdditionalTime = 0.f;
+
+	if (bSkipUpdate)
+	{
+		TickedPoseOffestTime -= DeltaTime;
+	}
+	else
+	{
+		if (TickedPoseOffestTime < 0.f)
+		{
+			AdditionalTime = -TickedPoseOffestTime;
+			TickedPoseOffestTime = 0.f;
+		}
+	}
+}
+
+void FAnimUpdateRateParameters::SetLookAheadMode(float DeltaTime, uint8 UpdateRateShift, float LookAheadAmount)
+{
+	float OriginalTickedPoseOffestTime = TickedPoseOffestTime;
+	if (OptimizeMode == TrailMode)
+	{
+		TickedPoseOffestTime = 0.f;
+	}
+	OptimizeMode = LookAheadMode;
+	ThisTickDelta = DeltaTime;
+
+	bInterpolateSkippedFrames = true;
+
+	TickedPoseOffestTime -= DeltaTime;
+
+	if (TickedPoseOffestTime < 0.f)
+	{
+		LookAheadAmount = FMath::Max(TickedPoseOffestTime*-1.f, LookAheadAmount);
+		AdditionalTime = LookAheadAmount;
+		TickedPoseOffestTime += LookAheadAmount;
+
+		bool bValid = (TickedPoseOffestTime >= 0.f);
+		if (!bValid)
+		{
+			FPlatformMisc::LowLevelOutputDebugStringf(TEXT("TPO Time: %.3f | Orig TPO Time: %.3f | DT: %.3f | LookAheadAmount: %.3f\n"), TickedPoseOffestTime, OriginalTickedPoseOffestTime, DeltaTime, LookAheadAmount);
+		}
+		check(bValid);
+		bSkipUpdate = bSkipEvaluation = false;
+	}
+	else
+	{
+		AdditionalTime = 0.f;
+		bSkipUpdate = bSkipEvaluation = true;
+	}
+}
+
+float FAnimUpdateRateParameters::GetInterpolationAlpha() const
+{
+	if (OptimizeMode == TrailMode)
+	{
+		return 0.25f + (1.f / float(FMath::Max(EvaluationRate, 2) * 2));
+	}
+	else if (OptimizeMode == LookAheadMode)
+	{
+		return FMath::Clamp(ThisTickDelta / (TickedPoseOffestTime + ThisTickDelta), 0.f, 1.f);
+	}
+	check(false); // Unknown mode
+	return 0.f;
+}
+
+float FAnimUpdateRateParameters::GetRootMotionInterp() const
+{
+	if (OptimizeMode == LookAheadMode)
+	{
+		return FMath::Clamp(ThisTickDelta / (TickedPoseOffestTime + ThisTickDelta), 0.f, 1.f);
+	}
+	return 1.f;
+}
\ No newline at end of file
diff --git a/Engine/Source/Runtime/Engine/Private/Components/SkyLightComponent.cpp b/Engine/Source/Runtime/Engine/Private/Components/SkyLightComponent.cpp
new file mode 100644
--- /dev/null
+++ b/Engine/Source/Runtime/Engine/Private/Components/SkyLightComponent.cpp
@@ -0,0 +1,607 @@
+// Copyright 1998-2015 Epic Games, Inc. All Rights Reserved.
+
+/*=============================================================================
+	SkyLightComponent.cpp: SkyLightComponent implementation.
+=============================================================================*/
+
+#include "EnginePrivate.h"
+#if WITH_EDITOR
+#include "ObjectEditorUtils.h"
+#endif
+#include "Engine/SkyLight.h"
+#include "MessageLog.h"
+#include "UObjectToken.h"
+#include "Net/UnrealNetwork.h"
+#include "MapErrors.h"
+#include "ComponentInstanceDataCache.h"
+#include "ShaderCompiler.h"
+#include "Components/SkyLightComponent.h"
+
+#define LOCTEXT_NAMESPACE "SkyLightComponent"
+
+void FSkyTextureCubeResource::InitRHI()
+{
+	if (GetFeatureLevel() >= ERHIFeatureLevel::SM4)
+	{
+		FRHIResourceCreateInfo CreateInfo;
+		TextureCubeRHI = RHICreateTextureCube(Size, Format, NumMips, 0, CreateInfo);
+		TextureRHI = TextureCubeRHI;
+
+		// Create the sampler state RHI resource.
+		FSamplerStateInitializerRHI SamplerStateInitializer
+		(
+			SF_Trilinear,
+			AM_Clamp,
+			AM_Clamp,
+			AM_Clamp
+		);
+		SamplerStateRHI = RHICreateSamplerState(SamplerStateInitializer);
+	}
+}
+
+void FSkyTextureCubeResource::Release()
+{
+	check( IsInGameThread() );
+	checkSlow(NumRefs > 0);
+	if(--NumRefs == 0)
+	{
+		BeginReleaseResource(this);
+		// Have to defer actual deletion until above rendering command has been processed, we will use the deferred cleanup interface for that
+		BeginCleanup(this);
+	}
+}
+
+void UWorld::UpdateAllSkyCaptures()
+{
+	TArray<USkyLightComponent*> UpdatedComponents;
+
+	for (TObjectIterator<USkyLightComponent> It; It; ++It)
+	{
+		USkyLightComponent* CaptureComponent = *It;
+
+		if (ContainsActor(CaptureComponent->GetOwner()) && !CaptureComponent->IsPendingKill())
+		{
+			// Purge cached derived data and force an update
+			CaptureComponent->SetCaptureIsDirty();
+			UpdatedComponents.Add(CaptureComponent);
+		}
+	}
+
+	USkyLightComponent::UpdateSkyCaptureContents(this);
+}
+
+FSkyLightSceneProxy::FSkyLightSceneProxy(const USkyLightComponent* InLightComponent)
+	: LightComponent(InLightComponent)
+	, ProcessedTexture(InLightComponent->ProcessedSkyTexture)
+	, SkyDistanceThreshold(InLightComponent->SkyDistanceThreshold)
+	, bCastShadows(InLightComponent->CastShadows)
+	, bWantsStaticShadowing(InLightComponent->Mobility == EComponentMobility::Stationary)
+	, bPrecomputedLightingIsValid(InLightComponent->bPrecomputedLightingIsValid)
+	, bHasStaticLighting(InLightComponent->HasStaticLighting())
+	, LightColor(FLinearColor(InLightComponent->LightColor) * InLightComponent->Intensity)
+	, IrradianceEnvironmentMap(InLightComponent->IrradianceEnvironmentMap)
+	, IndirectLightingIntensity(InLightComponent->IndirectLightingIntensity)
+	, OcclusionMaxDistance(InLightComponent->OcclusionMaxDistance)
+	, Contrast(InLightComponent->Contrast)
+	, MinOcclusion(InLightComponent->MinOcclusion)
+	, OcclusionTint(InLightComponent->OcclusionTint)
+{
+}
+
+USkyLightComponent::USkyLightComponent(const FObjectInitializer& ObjectInitializer)
+	: Super(ObjectInitializer)
+{
+#if WITH_EDITORONLY_DATA
+	if (!IsRunningCommandlet())
+	{
+		static ConstructorHelpers::FObjectFinder<UTexture2D> StaticTexture(TEXT("/Engine/EditorResources/LightIcons/SkyLight"));
+		StaticEditorTexture = StaticTexture.Object;
+		StaticEditorTextureScale = 1.0f;
+		DynamicEditorTexture = StaticTexture.Object;
+		DynamicEditorTextureScale = 1.0f;
+	}
+#endif
+
+	Brightness_DEPRECATED = 1;
+	Intensity = 1;
+	IndirectLightingIntensity = 1.0f;
+	SkyDistanceThreshold = 150000;
+	Mobility = EComponentMobility::Stationary;
+	bCaptureDirty = false;
+	bLowerHemisphereIsBlack = true;
+	bSavedConstructionScriptValuesValid = true;
+	bHasEverCaptured = false;
+	OcclusionMaxDistance = 1000;
+	MinOcclusion = 0;
+	OcclusionTint = FColor::Black;
+}
+
+FSkyLightSceneProxy* USkyLightComponent::CreateSceneProxy() const
+{
+	if (ProcessedSkyTexture)
+	{
+		return new FSkyLightSceneProxy(this);
+	}
+	
+	return NULL;
+}
+
+void USkyLightComponent::SetCaptureIsDirty()
+{ 
+	if (bVisible && bAffectsWorld)
+	{
+		SkyCapturesToUpdate.AddUnique(this);
+		bCaptureDirty = true; 
+
+		// Mark saved values as invalid, in case a sky recapture is requested in a construction script between a save / restore of sky capture state
+		bSavedConstructionScriptValuesValid = false;
+	}
+}
+
+TArray<USkyLightComponent*> USkyLightComponent::SkyCapturesToUpdate;
+
+void USkyLightComponent::CreateRenderState_Concurrent()
+{
+	Super::CreateRenderState_Concurrent();
+
+	bool bHidden = false;
+#if WITH_EDITORONLY_DATA
+	bHidden = GetOwner() ? GetOwner()->bHiddenEdLevel : false;
+#endif // WITH_EDITORONLY_DATA
+
+	if(!ShouldComponentAddToScene())
+	{
+		bHidden = true;
+	}
+
+	const bool bIsValid = SourceType != SLS_SpecifiedCubemap || Cubemap != NULL;
+
+	if (bAffectsWorld && bVisible && !bHidden && bIsValid)
+	{
+		// Create the light's scene proxy.
+		SceneProxy = CreateSceneProxy();
+
+		if (SceneProxy)
+		{
+			// Add the light to the scene.
+			World->Scene->SetSkyLight(SceneProxy);
+		}
+	}
+}
+
+extern ENGINE_API int32 GReflectionCaptureSize;
+
+void USkyLightComponent::PostInitProperties()
+{
+	if (!HasAnyFlags(RF_ClassDefaultObject))
+	{
+		// Enqueue an update by default, so that newly placed components will get an update
+		// PostLoad will undo this for components loaded from disk
+		bCaptureDirty = true;
+		SkyCapturesToUpdate.AddUnique(this);
+	}
+
+	Super::PostInitProperties();
+}
+
+void USkyLightComponent::PostLoad()
+{
+	Super::PostLoad();
+
+	// All components are queued for update on creation by default, remove if needed
+	if (!bVisible || HasAnyFlags(RF_ClassDefaultObject))
+	{
+		SkyCapturesToUpdate.Remove(this);
+		bCaptureDirty = false;
+	}
+}
+
+/** 
+ * Fast path for updating light properties that doesn't require a re-register,
+ * Which would otherwise cause the scene's static draw lists to be recreated.
+ */
+void USkyLightComponent::UpdateLimitedRenderingStateFast()
+{
+	if (SceneProxy)
+	{
+		ENQUEUE_UNIQUE_RENDER_COMMAND_THREEPARAMETER(
+			FFastUpdateSkyLightCommand,
+			FSkyLightSceneProxy*,LightSceneProxy,SceneProxy,
+			FLinearColor,LightColor,FLinearColor(LightColor) * Intensity,
+			float,IndirectLightingIntensity,IndirectLightingIntensity,
+		{
+			LightSceneProxy->LightColor = LightColor;
+			LightSceneProxy->IndirectLightingIntensity = IndirectLightingIntensity;
+		});
+	}
+}
+
+/** 
+* This is called when property is modified by InterpPropertyTracks
+*
+* @param PropertyThatChanged	Property that changed
+*/
+void USkyLightComponent::PostInterpChange(UProperty* PropertyThatChanged)
+{
+	static FName LightColorName(TEXT("LightColor"));
+	static FName IntensityName(TEXT("Intensity"));
+	static FName IndirectLightingIntensityName(TEXT("IndirectLightingIntensity"));
+
+	FName PropertyName = PropertyThatChanged->GetFName();
+	if (PropertyName == LightColorName
+		|| PropertyName == IntensityName
+		|| PropertyName == IndirectLightingIntensityName)
+	{
+		UpdateLimitedRenderingStateFast();
+	}
+	else
+	{
+		Super::PostInterpChange(PropertyThatChanged);
+	}
+}
+
+void USkyLightComponent::DestroyRenderState_Concurrent()
+{
+	Super::DestroyRenderState_Concurrent();
+
+	if (SceneProxy)
+	{
+		World->Scene->DisableSkyLight(SceneProxy);
+
+		ENQUEUE_UNIQUE_RENDER_COMMAND_ONEPARAMETER(
+			FDestroySkyLightCommand,
+			FSkyLightSceneProxy*,LightSceneProxy,SceneProxy,
+		{
+			delete LightSceneProxy;
+		});
+
+		SceneProxy = NULL;
+	}
+}
+
+#if WITH_EDITOR
+void USkyLightComponent::PostEditChangeProperty(FPropertyChangedEvent& PropertyChangedEvent)
+{
+	Super::PostEditChangeProperty(PropertyChangedEvent);
+
+	SetCaptureIsDirty();
+}
+
+bool USkyLightComponent::CanEditChange(const UProperty* InProperty) const
+{
+	if (InProperty)
+	{
+		FString PropertyName = InProperty->GetName();
+
+		if (FCString::Strcmp(*PropertyName, TEXT("Cubemap")) == 0)
+		{
+			return SourceType == SLS_SpecifiedCubemap;
+		}
+
+		if (FCString::Strcmp(*PropertyName, TEXT("Contrast")) == 0
+			|| FCString::Strcmp(*PropertyName, TEXT("OcclusionMaxDistance")) == 0
+			|| FCString::Strcmp(*PropertyName, TEXT("MinOcclusion")) == 0)
+		{
+			static const auto CVar = IConsoleManager::Get().FindTConsoleVariableDataInt(TEXT("r.GenerateMeshDistanceFields"));
+			return Mobility == EComponentMobility::Movable && CastShadows && CVar->GetValueOnGameThread() != 0;
+		}
+	}
+
+	return Super::CanEditChange(InProperty);
+}
+
+void USkyLightComponent::CheckForErrors()
+{
+	AActor* Owner = GetOwner();
+
+	if (Owner && bVisible && bAffectsWorld)
+	{
+		UWorld* ThisWorld = Owner->GetWorld();
+		bool bMultipleFound = false;
+
+		if (ThisWorld)
+		{
+			for (TObjectIterator<USkyLightComponent> ComponentIt; ComponentIt; ++ComponentIt)
+			{
+				USkyLightComponent* Component = *ComponentIt;
+
+				if (Component != this 
+					&& !Component->HasAnyFlags(RF_PendingKill)
+					&& Component->bVisible
+					&& Component->bAffectsWorld
+					&& Component->GetOwner() 
+					&& ThisWorld->ContainsActor(Component->GetOwner())
+					&& !Component->GetOwner()->IsPendingKill())
+				{
+					bMultipleFound = true;
+					break;
+				}
+			}
+		}
+
+		if (bMultipleFound)
+		{
+			FMessageLog("MapCheck").Error()
+				->AddToken(FUObjectToken::Create(Owner))
+				->AddToken(FTextToken::Create(LOCTEXT( "MapCheck_Message_MultipleSkyLights", "Multiple sky lights are active, only one can be enabled per world." )))
+				->AddToken(FMapErrorToken::Create(FMapErrors::MultipleSkyLights));
+		}
+	}
+}
+
+#endif // WITH_EDITOR
+
+void USkyLightComponent::BeginDestroy()
+{
+	// Deregister the component from the update queue
+	if (bCaptureDirty)
+	{
+		SkyCapturesToUpdate.Remove(this);
+	}
+
+	// Release reference
+	ProcessedSkyTexture = NULL;
+
+	// Begin a fence to track the progress of the above BeginReleaseResource being completed on the RT
+	ReleaseResourcesFence.BeginFence();
+
+	Super::BeginDestroy();
+}
+
+bool USkyLightComponent::IsReadyForFinishDestroy()
+{
+	// Wait until the fence is complete before allowing destruction
+	return Super::IsReadyForFinishDestroy() && ReleaseResourcesFence.IsFenceComplete();
+}
+
+/** Used to store lightmap data during RerunConstructionScripts */
+class FPrecomputedSkyLightInstanceData : public FSceneComponentInstanceData
+{
+public:
+	FPrecomputedSkyLightInstanceData(const USkyLightComponent* SourceComponent)
+		: FSceneComponentInstanceData(SourceComponent)
+	{}
+
+	virtual void ApplyToComponent(UActorComponent* Component, const ECacheApplyPhase CacheApplyPhase) override
+	{
+		FSceneComponentInstanceData::ApplyToComponent(Component, CacheApplyPhase);
+		CastChecked<USkyLightComponent>(Component)->ApplyComponentInstanceData(this);
+	}
+
+	FGuid LightGuid;
+	bool bPrecomputedLightingIsValid;
+	// This has to be refcounted to keep it alive during the handoff without doing a deep copy
+	TRefCountPtr<FSkyTextureCubeResource> ProcessedSkyTexture;
+	FSHVectorRGB3 IrradianceEnvironmentMap;
+};
+
+// Init type name static
+FName USkyLightComponent::GetComponentInstanceDataType() const
+{
+	static const FName PrecomputedSkyLightInstanceDataTypeName(TEXT("PrecomputedSkyLightInstanceData"));
+	return PrecomputedSkyLightInstanceDataTypeName;
+}
+
+FActorComponentInstanceData* USkyLightComponent::GetComponentInstanceData() const
+{
+	FPrecomputedSkyLightInstanceData* InstanceData = new FPrecomputedSkyLightInstanceData(this);
+	InstanceData->LightGuid = LightGuid;
+	InstanceData->bPrecomputedLightingIsValid = bPrecomputedLightingIsValid;
+	InstanceData->ProcessedSkyTexture = ProcessedSkyTexture;
+	InstanceData->IrradianceEnvironmentMap = IrradianceEnvironmentMap;
+
+	return InstanceData;
+}
+
+void USkyLightComponent::ApplyComponentInstanceData(FPrecomputedSkyLightInstanceData* LightMapData)
+{
+	check(LightMapData);
+
+	LightGuid = LightMapData->LightGuid;
+	bPrecomputedLightingIsValid = LightMapData->bPrecomputedLightingIsValid;
+	ProcessedSkyTexture = LightMapData->ProcessedSkyTexture;
+	IrradianceEnvironmentMap = LightMapData->IrradianceEnvironmentMap;
+
+	if (ProcessedSkyTexture && bSavedConstructionScriptValuesValid)
+	{
+		// We have valid capture state, remove the queued update
+		bCaptureDirty = false;
+		SkyCapturesToUpdate.Remove(this);
+	}
+
+	MarkRenderStateDirty();
+}
+
+void USkyLightComponent::UpdateSkyCaptureContents(UWorld* WorldToUpdate)
+{
+	if (WorldToUpdate->Scene)
+	{
+		const bool bIsCompilingShaders = GShaderCompilingManager != NULL && GShaderCompilingManager->IsCompiling();
+
+		// Iterate backwards so we can remove elements without changing the index
+		for (int32 CaptureIndex = SkyCapturesToUpdate.Num() - 1; CaptureIndex >= 0; CaptureIndex--)
+		{
+			USkyLightComponent* CaptureComponent = SkyCapturesToUpdate[CaptureIndex];
+
+			if ((!CaptureComponent->GetOwner() || WorldToUpdate->ContainsActor(CaptureComponent->GetOwner()))
+				// Only process sky capture requests once async shader compiling completes, otherwise we will capture the scene with temporary shaders
+				&& (!bIsCompilingShaders || CaptureComponent->SourceType == SLS_SpecifiedCubemap))
+			{
+				// Only capture valid sky light components
+				if (CaptureComponent->SourceType != SLS_SpecifiedCubemap || CaptureComponent->Cubemap)
+				{
+					// Allocate the needed texture on first capture
+					if (!CaptureComponent->ProcessedSkyTexture)
+					{
+						CaptureComponent->ProcessedSkyTexture = new FSkyTextureCubeResource();
+						CaptureComponent->ProcessedSkyTexture->SetupParameters(GReflectionCaptureSize, FMath::CeilLogTwo(GReflectionCaptureSize) + 1, PF_FloatRGBA);
+						BeginInitResource(CaptureComponent->ProcessedSkyTexture);
+						CaptureComponent->MarkRenderStateDirty();
+					}
+
+					WorldToUpdate->Scene->UpdateSkyCaptureContents(CaptureComponent, false, CaptureComponent->ProcessedSkyTexture, CaptureComponent->IrradianceEnvironmentMap);
+
+					CaptureComponent->bHasEverCaptured = true;
+					CaptureComponent->MarkRenderStateDirty();
+				}
+
+				// Only remove queued update requests if we processed it for the right world
+				SkyCapturesToUpdate.RemoveAt(CaptureIndex);
+			}
+		}
+	}
+}
+
+void USkyLightComponent::CaptureEmissiveIrradianceEnvironmentMap(FSHVectorRGB3& OutIrradianceMap) const
+{
+	OutIrradianceMap = FSHVectorRGB3();
+
+	if (GetScene() && (SourceType != SLS_SpecifiedCubemap || Cubemap))
+	{
+		// Capture emissive scene lighting only for the lighting build
+		// This is necessary to avoid a feedback loop with the last lighting build results
+		GetScene()->UpdateSkyCaptureContents(this, true, NULL, OutIrradianceMap);
+	}
+}
+
+/** Set brightness of the light */
+void USkyLightComponent::SetIntensity(float NewIntensity)
+{
+	// Can't set brightness on a static light
+	if (AreDynamicDataChangesAllowed()
+		&& Intensity != NewIntensity)
+	{
+		Intensity = NewIntensity;
+		UpdateLimitedRenderingStateFast();
+	}
+}
+
+void USkyLightComponent::SetIndirectLightingIntensity(float NewIntensity)
+{
+	// Can't set brightness on a static light
+	if (AreDynamicDataChangesAllowed()
+		&& IndirectLightingIntensity != NewIntensity)
+	{
+		IndirectLightingIntensity = NewIntensity;
+		UpdateLimitedRenderingStateFast();
+	}
+}
+
+/** Set color of the light */
+void USkyLightComponent::SetLightColor(FLinearColor NewLightColor)
+{
+	FColor NewColor(NewLightColor);
+
+	// Can't set color on a static light
+	if (AreDynamicDataChangesAllowed()
+		&& LightColor != NewColor)
+	{
+		LightColor = NewColor;
+		UpdateLimitedRenderingStateFast();
+	}
+}
+
+void USkyLightComponent::SetCubemap(UTextureCube* NewCubemap)
+{
+	// Can't set color on a static light
+	if (AreDynamicDataChangesAllowed()
+		&& Cubemap != NewCubemap)
+	{
+		Cubemap = NewCubemap;
+		MarkRenderStateDirty();
+		SetCaptureIsDirty();
+	}
+}
+
+void USkyLightComponent::SetOcclusionTint(const FColor& InTint)
+{
+	// Can't set on a static light
+	if (AreDynamicDataChangesAllowed()
+		&& OcclusionTint != InTint)
+	{
+		OcclusionTint = InTint;
+		MarkRenderStateDirty();
+	}
+}
+
+void USkyLightComponent::SetMinOcclusion(float InMinOcclusion)
+{
+	// Can't set on a static light
+	if (AreDynamicDataChangesAllowed()
+		&& MinOcclusion != InMinOcclusion)
+	{
+		MinOcclusion = InMinOcclusion;
+		MarkRenderStateDirty();
+	}
+}
+
+void USkyLightComponent::SetVisibility(bool bNewVisibility, bool bPropagateToChildren)
+{
+	const bool bOldWasVisible = bVisible;
+
+	Super::SetVisibility(bNewVisibility, bPropagateToChildren);
+
+	if (bVisible && !bOldWasVisible && !bHasEverCaptured)
+	{
+		// Capture if we are being enabled for the first time
+		SetCaptureIsDirty();
+	}
+}
+
+void USkyLightComponent::RecaptureSky()
+{
+	SetCaptureIsDirty();
+}
+
+ASkyLight::ASkyLight(const FObjectInitializer& ObjectInitializer)
+	: Super(ObjectInitializer)
+{
+	LightComponent = CreateDefaultSubobject<USkyLightComponent>(TEXT("SkyLightComponent0"));
+	RootComponent = LightComponent;
+
+#if WITH_EDITORONLY_DATA
+	if (!IsRunningCommandlet())
+	{
+	// Structure to hold one-time initialization
+	struct FConstructorStatics
+	{
+			ConstructorHelpers::FObjectFinderOptional<UTexture2D> SkyLightTextureObject;
+		FName ID_Sky;
+		FText NAME_Sky;
+
+		FConstructorStatics()
+				: SkyLightTextureObject(TEXT("/Engine/EditorResources/LightIcons/SkyLight"))
+				, ID_Sky(TEXT("Sky"))
+			, NAME_Sky(NSLOCTEXT( "SpriteCategory", "Sky", "Sky" ))
+		{
+		}
+	};
+	static FConstructorStatics ConstructorStatics;
+
+		if (GetSpriteComponent())
+		{
+			GetSpriteComponent()->Sprite = ConstructorStatics.SkyLightTextureObject.Get();
+			GetSpriteComponent()->SpriteInfo.Category = ConstructorStatics.ID_Sky;
+			GetSpriteComponent()->SpriteInfo.DisplayName = ConstructorStatics.NAME_Sky;
+			GetSpriteComponent()->AttachParent = LightComponent;
+		}
+	}
+#endif // WITH_EDITORONLY_DATA
+}
+
+void ASkyLight::GetLifetimeReplicatedProps( TArray< FLifetimeProperty > & OutLifetimeProps ) const
+{
+	Super::GetLifetimeReplicatedProps( OutLifetimeProps );
+
+	DOREPLIFETIME( ASkyLight, bEnabled );
+}
+
+void ASkyLight::OnRep_bEnabled()
+{
+	LightComponent->SetVisibility(bEnabled);
+}
+
+#undef LOCTEXT_NAMESPACE
+
+/** Returns LightComponent subobject **/
+USkyLightComponent* ASkyLight::GetLightComponent() const { return LightComponent; }
diff --git a/Engine/Source/Runtime/Engine/Private/Components/SphereComponent.cpp b/Engine/Source/Runtime/Engine/Private/Components/SphereComponent.cpp
new file mode 100644
--- /dev/null
+++ b/Engine/Source/Runtime/Engine/Private/Components/SphereComponent.cpp
@@ -0,0 +1,171 @@
+// Copyright 1998-2015 Epic Games, Inc. All Rights Reserved.
+
+
+#include "EnginePrivate.h"
+#include "Components/SphereComponent.h"
+
+USphereComponent::USphereComponent(const FObjectInitializer& ObjectInitializer)
+	: Super(ObjectInitializer)
+{
+	SphereRadius = 32.0f;
+	ShapeColor = FColor(255, 0, 0, 255);
+
+	bUseEditorCompositing = true;
+}
+
+FBoxSphereBounds USphereComponent::CalcBounds(const FTransform& LocalToWorld) const
+{
+	return FBoxSphereBounds( FVector::ZeroVector, FVector(SphereRadius), SphereRadius ).TransformBy(LocalToWorld);
+}
+
+void USphereComponent::CalcBoundingCylinder(float& CylinderRadius, float& CylinderHalfHeight) const
+{
+	CylinderRadius = SphereRadius * ComponentToWorld.GetMaximumAxisScale();
+	CylinderHalfHeight = CylinderRadius;
+}
+
+void USphereComponent::UpdateBodySetup()
+{
+	if (ShapeBodySetup == NULL || ShapeBodySetup->IsPendingKill())
+	{
+		ShapeBodySetup = NewObject<UBodySetup>(this);
+		ShapeBodySetup->CollisionTraceFlag = CTF_UseSimpleAsComplex;
+		ShapeBodySetup->AggGeom.SphereElems.Add(FKSphereElem());
+	}
+
+	check (ShapeBodySetup->AggGeom.SphereElems.Num() == 1);
+	FKSphereElem* se = ShapeBodySetup->AggGeom.SphereElems.GetData();
+
+	// check for mal formed values
+	float Radius = SphereRadius;
+	if( Radius < KINDA_SMALL_NUMBER )
+	{
+		Radius = 0.1f;
+	}
+
+	// now set the PhysX data values
+	se->Center = FVector::ZeroVector;
+	se->Radius = Radius;
+}
+
+void USphereComponent::SetSphereRadius( float InSphereRadius, bool bUpdateOverlaps )
+{
+	SphereRadius = InSphereRadius;
+	MarkRenderStateDirty();
+
+	if (bPhysicsStateCreated)
+	{
+		DestroyPhysicsState();
+		UpdateBodySetup();
+		CreatePhysicsState();
+
+		if ( bUpdateOverlaps && IsCollisionEnabled() && GetOwner() )
+		{
+			UpdateOverlaps();
+		}
+	}
+}
+
+bool USphereComponent::IsZeroExtent() const
+{
+	return SphereRadius == 0.f;
+}
+
+
+FPrimitiveSceneProxy* USphereComponent::CreateSceneProxy()
+{
+	/** Represents a DrawLightRadiusComponent to the scene manager. */
+	class FSphereSceneProxy : public FPrimitiveSceneProxy
+	{
+	public:
+
+		/** Initialization constructor. */
+		FSphereSceneProxy(const USphereComponent* InComponent)
+			:	FPrimitiveSceneProxy(InComponent)
+			,	bDrawOnlyIfSelected( InComponent->bDrawOnlyIfSelected )
+			,	SphereColor(InComponent->ShapeColor)
+			,	ShapeMaterial(InComponent->ShapeMaterial)
+			,	SphereRadius(InComponent->SphereRadius)
+		{
+			bWillEverBeLit = false;
+		}
+
+		  // FPrimitiveSceneProxy interface.
+		
+		virtual void GetDynamicMeshElements(const TArray<const FSceneView*>& Views, const FSceneViewFamily& ViewFamily, uint32 VisibilityMap, FMeshElementCollector& Collector) const override
+		{
+			QUICK_SCOPE_CYCLE_COUNTER( STAT_SphereSceneProxy_GetDynamicMeshElements );
+
+			for (int32 ViewIndex = 0; ViewIndex < Views.Num(); ViewIndex++)
+			{
+				if (VisibilityMap & (1 << ViewIndex))
+				{
+					const FSceneView* View = Views[ViewIndex];
+					FPrimitiveDrawInterface* PDI = Collector.GetPDI(ViewIndex);
+
+					const FMatrix& LocalToWorld = GetLocalToWorld();
+					int32 SphereSides =  FMath::Clamp<int32>(SphereRadius/4.f, 16, 64);
+					if(ShapeMaterial && !View->Family->EngineShowFlags.Wireframe)
+					{
+						GetSphereMesh(LocalToWorld.GetOrigin(), FVector(SphereRadius), SphereSides, SphereSides/2, ShapeMaterial->GetRenderProxy(false), SDPG_World, false, ViewIndex, Collector);
+					}
+					else
+					{
+						const FLinearColor DrawSphereColor = GetViewSelectionColor(SphereColor, *View, IsSelected(), IsHovered(), false, IsIndividuallySelected() );
+
+						float AbsScaleX = LocalToWorld.GetScaledAxis(EAxis::X).Size();
+						float AbsScaleY = LocalToWorld.GetScaledAxis(EAxis::Y).Size();
+						float AbsScaleZ = LocalToWorld.GetScaledAxis(EAxis::Z).Size();
+						float MinAbsScale = FMath::Min3(AbsScaleX, AbsScaleY, AbsScaleZ);
+
+						FVector ScaledX = LocalToWorld.GetUnitAxis(EAxis::X) * MinAbsScale;
+						FVector ScaledY = LocalToWorld.GetUnitAxis(EAxis::Y) * MinAbsScale;
+						FVector ScaledZ = LocalToWorld.GetUnitAxis(EAxis::Z) * MinAbsScale;
+
+						DrawCircle(PDI, LocalToWorld.GetOrigin(), ScaledX, ScaledY, DrawSphereColor, SphereRadius, SphereSides, SDPG_World);
+						DrawCircle(PDI, LocalToWorld.GetOrigin(), ScaledX, ScaledZ, DrawSphereColor, SphereRadius, SphereSides, SDPG_World);
+						DrawCircle(PDI, LocalToWorld.GetOrigin(), ScaledY, ScaledZ, DrawSphereColor, SphereRadius, SphereSides, SDPG_World);
+					}
+				}
+			}
+		}
+
+		virtual FPrimitiveViewRelevance GetViewRelevance(const FSceneView* View)  override
+		{
+			const bool bVisibleForSelection = !bDrawOnlyIfSelected || IsSelected();
+			const bool bVisibleForShowFlags = true; // @TODO
+
+			FPrimitiveViewRelevance Result;
+			Result.bDrawRelevance = IsShown(View) && bVisibleForSelection && bVisibleForShowFlags;
+			Result.bDynamicRelevance = true;
+			Result.bShadowRelevance = IsShadowCast(View);
+			Result.bEditorPrimitiveRelevance = UseEditorCompositing(View);
+			return Result;
+		}
+
+		virtual uint32 GetMemoryFootprint( void ) const override { return( sizeof( *this ) + GetAllocatedSize() ); }
+		uint32 GetAllocatedSize( void ) const { return( FPrimitiveSceneProxy::GetAllocatedSize() ); }
+
+	private:
+		const uint32				bDrawOnlyIfSelected:1;
+		const FColor				SphereColor;
+		const UMaterialInterface*	ShapeMaterial;
+		const float					SphereRadius;
+	};
+
+	return new FSphereSceneProxy( this );
+}
+
+
+FCollisionShape USphereComponent::GetCollisionShape(float Inflation) const
+{
+	const float Radius = FMath::Max(0.f, GetScaledSphereRadius() + Inflation);
+	return FCollisionShape::MakeSphere(Radius);
+}
+
+bool USphereComponent::AreSymmetricRotations(const FQuat& A, const FQuat& B, const FVector& Scale3D) const
+{
+	// All rotations are equal when scale is uniform.
+	// Not detecting rotations around non-uniform scale.
+	return Scale3D.GetAbs().AllComponentsEqual() || A.Equals(B);
+}
diff --git a/Engine/Source/Runtime/Engine/Private/Components/SpotLightComponent.cpp b/Engine/Source/Runtime/Engine/Private/Components/SpotLightComponent.cpp
new file mode 100644
--- /dev/null
+++ b/Engine/Source/Runtime/Engine/Private/Components/SpotLightComponent.cpp
@@ -0,0 +1,278 @@
+// Copyright 1998-2015 Epic Games, Inc. All Rights Reserved.
+
+/*=============================================================================
+	SpotLightComponent.cpp: LightComponent implementation.
+=============================================================================*/
+
+#include "EnginePrivate.h"
+#include "PointLightSceneProxy.h"
+#include "Components/SpotLightComponent.h"
+
+
+/**
+ * The spot light policy for TMeshLightingDrawingPolicy.
+ */
+class FSpotLightPolicy
+{
+public:
+	typedef class FSpotLightSceneProxy SceneInfoType;
+};
+
+/**
+ * The scene info for a spot light.
+ */
+class FSpotLightSceneProxy : public TPointLightSceneProxy<FSpotLightPolicy>
+{
+public:
+
+	/** Outer cone angle in radians, clamped to a valid range. */
+	float OuterConeAngle;
+
+	/** Cosine of the spot light's inner cone angle. */
+	float CosInnerCone;
+
+	/** Cosine of the spot light's outer cone angle. */
+	float CosOuterCone;
+
+	/** 1 / (CosInnerCone - CosOuterCone) */
+	float InvCosConeDifference;
+
+	/** Sine of the spot light's outer cone angle. */
+	float SinOuterCone;
+
+	/** 1 / Tangent of the spot light's outer cone angle. */
+	float InvTanOuterCone;
+
+	/** Cosine of the spot light's outer light shaft cone angle. */
+	float CosLightShaftConeAngle;
+
+	/** 1 / (appCos(ClampedInnerLightShaftConeAngle) - CosLightShaftConeAngle) */
+	float InvCosLightShaftConeDifference;
+
+	/** Initialization constructor. */
+	FSpotLightSceneProxy(const USpotLightComponent* Component):
+		TPointLightSceneProxy<FSpotLightPolicy>(Component)
+	{
+		const float ClampedInnerConeAngle = FMath::Clamp(Component->InnerConeAngle,0.0f,89.0f) * (float)PI / 180.0f;
+		const float ClampedOuterConeAngle = FMath::Clamp(Component->OuterConeAngle * (float)PI / 180.0f,ClampedInnerConeAngle + 0.001f,89.0f * (float)PI / 180.0f + 0.001f);
+		OuterConeAngle = ClampedOuterConeAngle;
+		CosOuterCone = FMath::Cos(ClampedOuterConeAngle);
+		SinOuterCone = FMath::Sin(ClampedOuterConeAngle);
+		CosInnerCone = FMath::Cos(ClampedInnerConeAngle);
+		InvCosConeDifference = 1.0f / (CosInnerCone - CosOuterCone);
+		InvTanOuterCone = 1.0f / FMath::Tan(ClampedOuterConeAngle);
+		const float ClampedOuterLightShaftConeAngle = FMath::Clamp(Component->LightShaftConeAngle * (float)PI / 180.0f, 0.001f, 89.0f * (float)PI / 180.0f + 0.001f);
+		// Use half the outer light shaft cone angle as the inner angle to provide a nice falloff
+		// Not exposing the inner light shaft cone angle as it is probably not needed
+		const float ClampedInnerLightShaftConeAngle = .5f * ClampedOuterLightShaftConeAngle;
+		CosLightShaftConeAngle = FMath::Cos(ClampedOuterLightShaftConeAngle);
+		InvCosLightShaftConeDifference = 1.0f / (FMath::Cos(ClampedInnerLightShaftConeAngle) - CosLightShaftConeAngle);
+	}
+
+	/** Accesses parameters needed for rendering the light. */
+	virtual void GetParameters(FVector4& LightPositionAndInvRadius, FVector4& LightColorAndFalloffExponent, FVector& NormalizedLightDirection, FVector2D& SpotAngles, float& LightSourceRadius, float& LightSourceLength, float& LightMinRoughness) const
+	{
+		LightPositionAndInvRadius = FVector4(
+			GetOrigin(),
+			InvRadius);
+
+		LightColorAndFalloffExponent = FVector4(
+			GetColor().R,
+			GetColor().G,
+			GetColor().B,
+			FalloffExponent);
+
+		NormalizedLightDirection = -GetDirection();
+		SpotAngles = FVector2D(CosOuterCone, InvCosConeDifference);
+		LightSourceRadius = SourceRadius;
+		LightSourceLength = SourceLength;
+		LightMinRoughness = MinRoughness;
+	}
+
+	// FLightSceneInfo interface.
+	virtual bool AffectsBounds(const FBoxSphereBounds& Bounds) const
+	{
+		if(!TPointLightSceneProxy<FSpotLightPolicy>::AffectsBounds(Bounds))
+		{
+			return false;
+		}
+
+		FVector	U = GetOrigin() - (Bounds.SphereRadius / SinOuterCone) * GetDirection(),
+				D = Bounds.Origin - U;
+		float	dsqr = D | D,
+				E = GetDirection() | D;
+		if(E > 0.0f && E * E >= dsqr * FMath::Square(CosOuterCone))
+		{
+			D = Bounds.Origin - GetOrigin();
+			dsqr = D | D;
+			E = -(GetDirection() | D);
+			if(E > 0.0f && E * E >= dsqr * FMath::Square(SinOuterCone))
+				return dsqr <= FMath::Square(Bounds.SphereRadius);
+			else
+				return true;
+		}
+
+		return false;
+	}
+	
+	/**
+	 * Sets up a projected shadow initializer for shadows from the entire scene.
+	 * @return True if the whole-scene projected shadow should be used.
+	 */
+	virtual bool GetWholeSceneProjectedShadowInitializer(const FSceneViewFamily& ViewFamily, TArray<FWholeSceneProjectedShadowInitializer, TInlineAllocator<6> >& OutInitializers) const
+	{
+		FWholeSceneProjectedShadowInitializer& OutInitializer = *new(OutInitializers) FWholeSceneProjectedShadowInitializer;
+		OutInitializer.PreShadowTranslation = -GetLightToWorld().GetOrigin();
+		OutInitializer.WorldToLight = GetWorldToLight().RemoveTranslation();
+		OutInitializer.Scales = FVector(1.0f,InvTanOuterCone,InvTanOuterCone);
+		OutInitializer.FaceDirection = FVector(1,0,0);
+		OutInitializer.SubjectBounds = FBoxSphereBounds(
+			GetLightToWorld().RemoveTranslation().TransformPosition(FVector(Radius, 0, 0)),
+			FVector(Radius, Radius, Radius),
+			Radius
+			);
+		OutInitializer.WAxis = FVector4(0,0,1,0);
+		OutInitializer.MinLightW = 0.1f;
+		OutInitializer.MaxDistanceToCastInLightW = Radius;
+		OutInitializer.CascadeSettings.bRayTracedDistanceField = UseRayTracedDistanceFieldShadows() && DoesPlatformSupportDistanceFieldShadowing(ViewFamily.GetShaderPlatform());
+		return true;
+	}
+
+	virtual float GetOuterConeAngle() const { return OuterConeAngle; }
+
+	virtual FVector2D GetLightShaftConeParams() const override { return FVector2D(CosLightShaftConeAngle, InvCosLightShaftConeDifference); }
+
+	virtual FSphere GetBoundingSphere() const
+	{
+		// Use the law of cosines to find the distance to the furthest edge of the spotlight cone from a position that is halfway down the spotlight direction
+		const float BoundsRadius = FMath::Sqrt(1.25f * Radius * Radius - Radius * Radius * CosOuterCone);
+		return FSphere(GetOrigin() + .5f * GetDirection() * Radius, BoundsRadius);
+	}
+};
+
+USpotLightComponent::USpotLightComponent(const FObjectInitializer& ObjectInitializer)
+	: Super(ObjectInitializer)
+{
+#if WITH_EDITORONLY_DATA
+	if (!IsRunningCommandlet())
+	{
+		static ConstructorHelpers::FObjectFinder<UTexture2D> StaticTexture(TEXT("/Engine/EditorResources/LightIcons/S_LightSpot"));
+		static ConstructorHelpers::FObjectFinder<UTexture2D> DynamicTexture(TEXT("/Engine/EditorResources/LightIcons/S_LightSpotMove"));
+
+		StaticEditorTexture = StaticTexture.Object;
+		StaticEditorTextureScale = 0.5f;
+		DynamicEditorTexture = DynamicTexture.Object;
+		DynamicEditorTextureScale = 0.5f;
+	}
+#endif
+
+	InnerConeAngle = 0.0f;
+	OuterConeAngle = 44.0f;
+}
+
+void USpotLightComponent::SetInnerConeAngle(float NewInnerConeAngle)
+{
+	if (AreDynamicDataChangesAllowed(false)
+		&& NewInnerConeAngle != InnerConeAngle)
+	{
+		InnerConeAngle = NewInnerConeAngle;
+		MarkRenderStateDirty();
+	}
+}
+
+void USpotLightComponent::SetOuterConeAngle(float NewOuterConeAngle)
+{
+	if (AreDynamicDataChangesAllowed(false)
+		&& NewOuterConeAngle != OuterConeAngle)
+	{
+		OuterConeAngle = NewOuterConeAngle;
+		MarkRenderStateDirty();
+	}
+}
+
+// Disable for now
+//void USpotLightComponent::SetLightShaftConeAngle(float NewLightShaftConeAngle)
+//{
+//	if (NewLightShaftConeAngle != LightShaftConeAngle)
+//	{
+//		LightShaftConeAngle = NewLightShaftConeAngle;
+//		MarkRenderStateDirty();
+//	}
+//}
+
+FLightSceneProxy* USpotLightComponent::CreateSceneProxy() const
+{
+	return new FSpotLightSceneProxy(this);
+}
+
+FSphere USpotLightComponent::GetBoundingSphere() const
+{
+	float ClampedInnerConeAngle = FMath::Clamp(InnerConeAngle,0.0f,89.0f) * (float)PI / 180.0f;
+	float ClampedOuterConeAngle = FMath::Clamp(OuterConeAngle * (float)PI / 180.0f,ClampedInnerConeAngle + 0.001f,89.0f * (float)PI / 180.0f + 0.001f);
+
+	float CosOuterCone = FMath::Cos(ClampedOuterConeAngle);
+
+	// Use the law of cosines to find the distance to the furthest edge of the spotlight cone from a position that is halfway down the spotlight direction
+	const float BoundsRadius = FMath::Sqrt(1.25f * AttenuationRadius * AttenuationRadius - AttenuationRadius * AttenuationRadius * CosOuterCone);
+	return FSphere(ComponentToWorld.GetLocation() + .5f * GetDirection() * AttenuationRadius, BoundsRadius);
+}
+
+bool USpotLightComponent::AffectsBounds(const FBoxSphereBounds& Bounds) const
+{
+	if(!Super::AffectsBounds(Bounds))
+	{
+		return false;
+	}
+
+	float	ClampedInnerConeAngle = FMath::Clamp(InnerConeAngle,0.0f,89.0f) * (float)PI / 180.0f,
+			ClampedOuterConeAngle = FMath::Clamp(OuterConeAngle * (float)PI / 180.0f,ClampedInnerConeAngle + 0.001f,89.0f * (float)PI / 180.0f + 0.001f);
+
+	float	Sin = FMath::Sin(ClampedOuterConeAngle),
+			Cos = FMath::Cos(ClampedOuterConeAngle);
+
+	FVector	U = GetComponentLocation() - (Bounds.SphereRadius / Sin) * GetDirection(),
+			D = Bounds.Origin - U;
+	float	dsqr = D | D,
+			E = GetDirection() | D;
+	if(E > 0.0f && E * E >= dsqr * FMath::Square(Cos))
+	{
+		D = Bounds.Origin - GetComponentLocation();
+		dsqr = D | D;
+		E = -(GetDirection() | D);
+		if(E > 0.0f && E * E >= dsqr * FMath::Square(Sin))
+			return dsqr <= FMath::Square(Bounds.SphereRadius);
+		else
+			return true;
+	}
+
+	return false;
+}
+
+/**
+* @return ELightComponentType for the light component class 
+*/
+ELightComponentType USpotLightComponent::GetLightType() const
+{
+	return LightType_Spot;
+}
+
+#if WITH_EDITOR
+
+void USpotLightComponent::PostEditChangeProperty( FPropertyChangedEvent& PropertyChangedEvent)
+{
+	if (PropertyChangedEvent.Property)
+	{
+		if (PropertyChangedEvent.Property->GetFName() == TEXT("InnerConeAngle"))
+		{
+			OuterConeAngle = FMath::Max(InnerConeAngle, OuterConeAngle);
+		}
+		else if (PropertyChangedEvent.Property->GetFName() == TEXT("OuterConeAngle"))
+		{
+			InnerConeAngle = FMath::Min(InnerConeAngle, OuterConeAngle);
+		}
+	}
+
+	UPointLightComponent::PostEditChangeProperty(PropertyChangedEvent);
+}
+
+#endif	// WITH_EDITOR
\ No newline at end of file
