Engine: Moved actor components into Components folder

diff --git a/Engine/Source/Runtime/Engine/Private/Components/StaticMeshComponent.cpp b/Engine/Source/Runtime/Engine/Private/Components/StaticMeshComponent.cpp
new file mode 100644
--- /dev/null
+++ b/Engine/Source/Runtime/Engine/Private/Components/StaticMeshComponent.cpp
@@ -0,0 +1,1962 @@
+// Copyright 1998-2015 Epic Games, Inc. All Rights Reserved.
+
+#include "EnginePrivate.h"
+#include "StaticMeshResources.h"
+#include "Components/SplineMeshComponent.h"
+#include "Net/UnrealNetwork.h"
+#include "MessageLog.h"
+#include "UObjectToken.h"
+#include "MapErrors.h"
+#if WITH_EDITOR
+#include "ShowFlags.h"
+#include "Collision.h"
+#include "ConvexVolume.h"
+#endif
+#include "ComponentInstanceDataCache.h"
+#include "LightMap.h"
+#include "ShadowMap.h"
+#include "ComponentReregisterContext.h"
+#include "Engine/ShadowMapTexture2D.h"
+#include "AI/Navigation/NavCollision.h"
+#include "Engine/StaticMeshSocket.h"
+#include "NavigationSystemHelpers.h"
+
+#define LOCTEXT_NAMESPACE "StaticMeshComponent"
+
+class FStaticMeshComponentInstanceData : public FSceneComponentInstanceData
+{
+public:
+	FStaticMeshComponentInstanceData(const UStaticMeshComponent* SourceComponent)
+		: FSceneComponentInstanceData(SourceComponent)
+		, StaticMesh(SourceComponent->StaticMesh)
+		, bHasCachedStaticLighting(false)
+	{
+	}
+
+	virtual void ApplyToComponent(UActorComponent* Component, const ECacheApplyPhase CacheApplyPhase) override
+	{
+		FSceneComponentInstanceData::ApplyToComponent(Component, CacheApplyPhase);
+		CastChecked<UStaticMeshComponent>(Component)->ApplyComponentInstanceData(this);
+	}
+
+	/** Add vertex color data for a specified LOD before RerunConstructionScripts is called */
+	void AddVertexColorData(const struct FStaticMeshComponentLODInfo& LODInfo, uint32 LODIndex)
+	{
+		if (VertexColorLODs.Num() <= (int32)LODIndex)
+		{
+			VertexColorLODs.SetNum(LODIndex + 1);
+		}
+		FVertexColorLODData& VertexColorData = VertexColorLODs[LODIndex];
+		VertexColorData.LODIndex = LODIndex;
+		VertexColorData.PaintedVertices = LODInfo.PaintedVertices;
+		LODInfo.OverrideVertexColors->GetVertexColors(VertexColorData.VertexBufferColors);
+	}
+
+	/** Re-apply vertex color data after RerunConstructionScripts is called */
+	bool ApplyVertexColorData(UStaticMeshComponent* StaticMeshComponent) const
+	{
+		bool bAppliedAnyData = false;
+
+		if (StaticMeshComponent != NULL)
+		{
+			StaticMeshComponent->SetLODDataCount(VertexColorLODs.Num(), StaticMeshComponent->LODData.Num());
+
+			for (int32 LODDataIndex = 0; LODDataIndex < VertexColorLODs.Num(); ++LODDataIndex)
+			{
+				const FVertexColorLODData& VertexColorLODData = VertexColorLODs[LODDataIndex];
+				uint32 LODIndex = VertexColorLODData.LODIndex;
+
+				if (StaticMeshComponent->LODData.IsValidIndex(LODIndex))
+				{
+					FStaticMeshComponentLODInfo& LODInfo = StaticMeshComponent->LODData[LODIndex];
+					// this component could have been constructed from a template
+					// that had its own vert color overrides; so before we apply
+					// the instance's color data, we need to clear the old
+					// vert colors (so we can properly call InitFromColorArray())
+					StaticMeshComponent->RemoveInstanceVertexColorsFromLOD(LODIndex);
+					// may not be null at the start (could have been initialized 
+					// from a  component template with vert coloring), but should
+					// be null at this point, after RemoveInstanceVertexColorsFromLOD()
+					if (LODInfo.OverrideVertexColors == NULL)
+					{
+						LODInfo.PaintedVertices = VertexColorLODData.PaintedVertices;
+
+						LODInfo.OverrideVertexColors = new FColorVertexBuffer;
+						LODInfo.OverrideVertexColors->InitFromColorArray(VertexColorLODData.VertexBufferColors);
+
+						BeginInitResource(LODInfo.OverrideVertexColors);
+						bAppliedAnyData = true;
+					}
+				}
+			}
+		}
+
+		return bAppliedAnyData;
+	}
+
+	/** Used to store lightmap data during RerunConstructionScripts */
+	struct FLightMapInstanceData
+	{
+		/** Transform of instance */
+		FTransform		Transform;
+		/** Lightmaps from LODData */
+		TArray<FLightMapRef>	LODDataLightMap;
+		TArray<FShadowMapRef>	LODDataShadowMap;
+		TArray<FGuid> IrrelevantLights;
+	};
+
+	/** Vertex data stored per-LOD */
+	struct FVertexColorLODData
+	{
+		/** copy of painted vertex data */
+		TArray<FPaintedVertex> PaintedVertices;
+
+		/** Copy of vertex buffer colors */
+		TArray<FColor> VertexBufferColors;
+
+		/** Index of the LOD that this data came from */
+		uint32 LODIndex;
+
+		/** Check whether this contains valid data */
+		bool IsValid() const
+		{
+			return PaintedVertices.Num() > 0 && VertexBufferColors.Num() > 0;
+		}
+	};
+
+	/** Mesh being used by component */
+	class UStaticMesh* StaticMesh;
+
+	/** Array of cached vertex colors for each LOD */
+	TArray<FVertexColorLODData> VertexColorLODs;
+
+	bool bHasCachedStaticLighting;
+	FLightMapInstanceData CachedStaticLighting;
+};
+
+UStaticMeshComponent::UStaticMeshComponent(const FObjectInitializer& ObjectInitializer)
+	: Super(ObjectInitializer)
+{
+	PrimaryComponentTick.bCanEverTick = false;
+
+	// check BaseEngine.ini for profile setup
+	SetCollisionProfileName(UCollisionProfile::BlockAllDynamic_ProfileName);
+
+	WireframeColorOverride = FColor(255, 255, 255, 255);
+
+	MinLOD = 0;
+	bOverrideLightMapRes = false;
+	OverriddenLightMapRes = 64;
+	SubDivisionStepSize = 32;
+	bUseSubDivisions = true;
+	StreamingDistanceMultiplier = 1.0f;
+	bBoundsChangeTriggersStreamingDataRebuild = true;
+	bHasCustomNavigableGeometry = EHasCustomNavigableGeometry::Yes;
+
+	GetBodyInstance()->bAutoWeld = true;	//static mesh by default has auto welding
+
+#if WITH_EDITORONLY_DATA
+	SelectedEditorSection = INDEX_NONE;
+#endif
+}
+
+UStaticMeshComponent::~UStaticMeshComponent()
+{
+	// Empty, but required because we don't want to have to include LightMap.h and ShadowMap.h in StaticMeshComponent.h, and they are required to compile FLightMapRef and FShadowMapRef
+}
+
+void UStaticMeshComponent::GetLifetimeReplicatedProps(TArray< FLifetimeProperty > & OutLifetimeProps) const
+{
+	Super::GetLifetimeReplicatedProps(OutLifetimeProps);
+
+	DOREPLIFETIME( UStaticMeshComponent, StaticMesh );
+}
+
+void UStaticMeshComponent::OnRep_StaticMesh(class UStaticMesh *OldStaticMesh)
+{
+	// Only do stuff if this actually changed from the last local value
+	if (OldStaticMesh!= StaticMesh)
+	{
+		// We have to force a call to SetStaticMesh with a new StaticMesh
+		UStaticMesh *NewStaticMesh = StaticMesh;
+		StaticMesh = NULL;
+		
+		SetStaticMesh(NewStaticMesh);
+	}
+}
+
+bool UStaticMeshComponent::HasAnySockets() const
+{
+	return (StaticMesh != NULL) && (StaticMesh->Sockets.Num() > 0);
+}
+
+void UStaticMeshComponent::QuerySupportedSockets(TArray<FComponentSocketDescription>& OutSockets) const
+{
+	if (StaticMesh != NULL)
+	{
+		for (int32 SocketIdx = 0; SocketIdx < StaticMesh->Sockets.Num(); ++SocketIdx)
+		{
+			if (UStaticMeshSocket* Socket = StaticMesh->Sockets[SocketIdx])
+			{
+				new (OutSockets) FComponentSocketDescription(Socket->SocketName, EComponentSocketType::Socket);
+			}
+		}
+	}
+}
+
+FString UStaticMeshComponent::GetDetailedInfoInternal() const
+{
+	FString Result;  
+
+	if( StaticMesh != NULL )
+	{
+		Result = StaticMesh->GetPathName( NULL );
+	}
+	else
+	{
+		Result = TEXT("No_StaticMesh");
+	}
+
+	return Result;  
+}
+
+
+void UStaticMeshComponent::AddReferencedObjects(UObject* InThis, FReferenceCollector& Collector)
+{	
+	UStaticMeshComponent* This = CastChecked<UStaticMeshComponent>(InThis);
+	for (int32 LODIndex = 0; LODIndex < This->LODData.Num(); LODIndex++)
+	{
+		if (This->LODData[LODIndex].LightMap != NULL)
+		{
+			This->LODData[LODIndex].LightMap->AddReferencedObjects(Collector);
+		}
+
+		if (This->LODData[LODIndex].ShadowMap != NULL)
+		{
+			This->LODData[LODIndex].ShadowMap->AddReferencedObjects(Collector);
+		}
+	}
+
+	Super::AddReferencedObjects(This, Collector);
+}
+
+
+void UStaticMeshComponent::Serialize(FArchive& Ar)
+{
+	Super::Serialize(Ar);
+
+#if WITH_EDITORONLY_DATA
+	if( Ar.IsCooking() )
+	{
+		// This is needed for data stripping
+		for( int32 LODIndex = 0; LODIndex < LODData.Num(); LODIndex++ )
+		{
+			LODData[ LODIndex ].Owner = this;
+		}
+	}
+#endif // WITH_EDITORONLY_DATA
+	Ar << LODData;
+
+	if (Ar.UE4Ver() < VER_UE4_COMBINED_LIGHTMAP_TEXTURES)
+	{
+		check(AttachmentCounter.GetValue() == 0);
+		// Irrelevant lights were incorrect before VER_UE4_TOSS_IRRELEVANT_LIGHTS
+		IrrelevantLights.Empty();
+	}
+
+	if (Ar.UE4Ver() < VER_UE4_AUTO_WELDING)
+	{
+		GetBodyInstance()->bAutoWeld = false;	//existing content may rely on no auto welding
+	}
+}
+
+
+
+bool UStaticMeshComponent::AreNativePropertiesIdenticalTo( UObject* Other ) const
+{
+	bool bNativePropertiesAreIdentical = Super::AreNativePropertiesIdenticalTo( Other );
+	UStaticMeshComponent* OtherSMC = CastChecked<UStaticMeshComponent>(Other);
+
+	if( bNativePropertiesAreIdentical )
+	{
+		// Components are not identical if they have lighting information.
+		if( LODData.Num() || OtherSMC->LODData.Num() )
+		{
+			bNativePropertiesAreIdentical = false;
+		}
+	}
+	
+	return bNativePropertiesAreIdentical;
+}
+
+void UStaticMeshComponent::PreSave()
+{
+	Super::PreSave();
+#if WITH_EDITORONLY_DATA
+	CachePaintedDataIfNecessary();
+#endif // WITH_EDITORONLY_DATA
+}
+
+#if WITH_EDITOR
+void UStaticMeshComponent::CheckForErrors()
+{
+	Super::CheckForErrors();
+
+	// Get the mesh owner's name.
+	AActor* Owner = GetOwner();
+	FString OwnerName(*(GNone.ToString()));
+	if ( Owner )
+	{
+		OwnerName = Owner->GetName();
+	}
+
+	// Make sure any simplified meshes can still find their high res source mesh
+	if( StaticMesh != NULL && StaticMesh->RenderData )
+	{
+		int32 ZeroTriangleElements = 0;
+
+		// Check for element material index/material mismatches
+		for (int32 LODIndex = 0; LODIndex < StaticMesh->RenderData->LODResources.Num(); ++LODIndex)
+		{
+			FStaticMeshLODResources& MeshLODData = StaticMesh->RenderData->LODResources[LODIndex];
+			for (int32 SectionIndex = 0; SectionIndex < MeshLODData.Sections.Num(); SectionIndex++)
+			{
+				FStaticMeshSection& Element = MeshLODData.Sections[SectionIndex];
+				if (Element.NumTriangles == 0)
+				{
+					ZeroTriangleElements++;
+				}
+			}
+		}
+
+		if (OverrideMaterials.Num() > StaticMesh->Materials.Num())
+		{
+			FFormatNamedArguments Arguments;
+			Arguments.Add(TEXT("OverridenCount"), OverrideMaterials.Num());
+			Arguments.Add(TEXT("ReferencedCount"), StaticMesh->Materials.Num());
+			Arguments.Add(TEXT("MeshName"), FText::FromString(StaticMesh->GetName()));
+			FMessageLog("MapCheck").Warning()
+				->AddToken(FUObjectToken::Create(Owner))
+				->AddToken(FTextToken::Create(FText::Format(LOCTEXT( "MapCheck_Message_MoreMaterialsThanReferenced", "More overriden materials ({OverridenCount}) on static mesh component than are referenced ({ReferencedCount}) in source mesh '{MeshName}'" ), Arguments ) ))
+				->AddToken(FMapErrorToken::Create(FMapErrors::MoreMaterialsThanReferenced));
+		}
+		if (ZeroTriangleElements > 0)
+		{
+			FFormatNamedArguments Arguments;
+			Arguments.Add(TEXT("ElementCount"), ZeroTriangleElements);
+			Arguments.Add(TEXT("MeshName"), FText::FromString(StaticMesh->GetName()));
+			FMessageLog("MapCheck").Warning()
+				->AddToken(FUObjectToken::Create(Owner))
+				->AddToken(FTextToken::Create(FText::Format(LOCTEXT( "MapCheck_Message_ElementsWithZeroTriangles", "{ElementCount} element(s) with zero triangles in static mesh '{MeshName}'" ), Arguments ) ))
+				->AddToken(FMapErrorToken::Create(FMapErrors::ElementsWithZeroTriangles));
+		}
+	}
+
+	if (!StaticMesh && (!Owner || !Owner->IsA(AWorldSettings::StaticClass())))	// Ignore worldsettings
+	{
+		FMessageLog("MapCheck").Warning()
+			->AddToken(FUObjectToken::Create(Owner))
+			->AddToken(FTextToken::Create(LOCTEXT("MapCheck_Message_StaticMeshNull", "Static mesh actor has NULL StaticMesh property")))
+			->AddToken(FMapErrorToken::Create(FMapErrors::StaticMeshNull));
+	}
+
+	// Make sure any non uniform scaled meshes have appropriate collision
+	if ( IsCollisionEnabled() && StaticMesh != NULL && StaticMesh->BodySetup != NULL && Owner != NULL )
+	{
+		// Overall scale factor for this mesh.
+		const FVector& TotalScale3D = ComponentToWorld.GetScale3D();
+		if ( !TotalScale3D.IsUniform() &&
+			 (StaticMesh->BodySetup->AggGeom.BoxElems.Num() > 0   ||
+			  StaticMesh->BodySetup->AggGeom.SphylElems.Num() > 0 ||
+			  StaticMesh->BodySetup->AggGeom.SphereElems.Num() > 0) )
+
+		{
+			FFormatNamedArguments Arguments;
+			Arguments.Add(TEXT("MeshName"), FText::FromString(StaticMesh->GetName()));
+			FMessageLog("MapCheck").Warning()
+				->AddToken(FUObjectToken::Create(Owner))
+				->AddToken(FTextToken::Create(FText::Format(LOCTEXT( "MapCheck_Message_SimpleCollisionButNonUniformScale", "'{MeshName}' has simple collision but is being scaled non-uniformly - collision creation will fail" ), Arguments)))
+				->AddToken(FMapErrorToken::Create(FMapErrors::SimpleCollisionButNonUniformScale));
+		}
+	}
+
+	if ( BodyInstance.bSimulatePhysics && StaticMesh != NULL && StaticMesh->BodySetup != NULL && StaticMesh->BodySetup->AggGeom.GetElementCount() == 0) 
+	{
+		FMessageLog("MapCheck").Warning()
+			->AddToken(FUObjectToken::Create(this))
+			->AddToken(FTextToken::Create(FText::Format(LOCTEXT( "MapCheck_Message_SimulatePhyNoSimpleCollision", "{0} : Using bSimulatePhysics but StaticMesh has not simple collision."), FText::FromString(GetName()) ) ));
+	}
+
+	if( Mobility == EComponentMobility::Movable &&
+		CastShadow && 
+		bCastDynamicShadow && 
+		IsRegistered() && 
+		Bounds.SphereRadius > 2000.0f )
+	{
+		// Large shadow casting objects that create preshadows will cause a massive performance hit, since preshadows are meant for small shadow casters.
+		FMessageLog("MapCheck").PerformanceWarning()
+			->AddToken(FUObjectToken::Create(this))
+			->AddToken(FTextToken::Create(LOCTEXT( "MapCheck_Message_ActorLargeShadowCaster", "Large actor receives a pre-shadow and will cause an extreme performance hit unless bCastDynamicShadow is set to false.") ))
+			->AddToken(FMapErrorToken::Create(FMapErrors::ActorLargeShadowCaster));
+	}
+}
+#endif
+
+FBoxSphereBounds UStaticMeshComponent::CalcBounds(const FTransform& LocalToWorld) const
+{
+	if(StaticMesh)
+	{
+		// Graphics bounds.
+		FBoxSphereBounds NewBounds = StaticMesh->GetBounds().TransformBy(LocalToWorld);
+		
+		// Add bounds of collision geometry (if present).
+		if(StaticMesh->BodySetup)
+		{
+			FBox AggGeomBox = StaticMesh->BodySetup->AggGeom.CalcAABB(LocalToWorld);
+			if (AggGeomBox.IsValid)
+			{
+				NewBounds = Union(NewBounds,FBoxSphereBounds(AggGeomBox));
+			}
+		}
+
+		NewBounds.BoxExtent *= BoundsScale;
+		NewBounds.SphereRadius *= BoundsScale;
+
+		return NewBounds;
+	}
+	else
+	{
+		return FBoxSphereBounds(LocalToWorld.GetLocation(), FVector::ZeroVector, 0.f);
+	}
+}
+
+void UStaticMeshComponent::OnRegister()
+{
+	if(StaticMesh != NULL && StaticMesh->RenderData)
+	{
+		// Check that the static-mesh hasn't been changed to be incompatible with the cached light-map.
+		int32 NumLODs = StaticMesh->RenderData->LODResources.Num();
+		bool bLODsShareStaticLighting = StaticMesh->RenderData->bLODsShareStaticLighting;
+		if (!bLODsShareStaticLighting && NumLODs != LODData.Num())
+		{
+			for(int32 i=0;i<LODData.Num();i++)
+			{
+				LODData[i].LightMap = NULL;
+				LODData[i].ShadowMap = NULL;
+			}
+		}
+	}
+
+	if (StaticMesh != NULL && StaticMesh->SpeedTreeWind.IsValid() && GetScene())
+	{
+		for (int32 LODIndex = 0; LODIndex < StaticMesh->RenderData->LODResources.Num(); ++LODIndex)
+		{
+			GetScene()->AddSpeedTreeWind(&StaticMesh->RenderData->LODResources[LODIndex].VertexFactory, StaticMesh);
+		}
+	}
+
+	Super::OnRegister();
+}
+
+void UStaticMeshComponent::OnUnregister()
+{
+	if (StaticMesh != NULL && StaticMesh->SpeedTreeWind.IsValid() && GetScene())
+	{
+		for (int32 LODIndex = 0; LODIndex < StaticMesh->RenderData->LODResources.Num(); ++LODIndex)
+		{
+			GetScene()->RemoveSpeedTreeWind(&StaticMesh->RenderData->LODResources[LODIndex].VertexFactory, StaticMesh);
+		}
+	}
+
+	Super::OnUnregister();
+}
+
+void UStaticMeshComponent::GetStreamingTextureInfo(TArray<FStreamingTexturePrimitiveInfo>& OutStreamingTextures) const
+{
+	if ( StaticMesh && !bIgnoreInstanceForTextureStreaming )
+	{
+		const auto FeatureLevel = GetWorld() ? GetWorld()->FeatureLevel : GMaxRHIFeatureLevel;
+
+		bool bHasValidLightmapCoordinates = ((StaticMesh->LightMapCoordinateIndex >= 0)
+			&& StaticMesh->RenderData
+			&& StaticMesh->RenderData->LODResources.Num() > 0
+			&& ((uint32)StaticMesh->LightMapCoordinateIndex < StaticMesh->RenderData->LODResources[0].VertexBuffer.GetNumTexCoords()));
+
+		// We need to come up with a compensation factor for spline deformed meshes
+		float SplineDeformFactor = 1.f;
+		const USplineMeshComponent* SplineComp = Cast<const USplineMeshComponent>(this);
+		if (SplineComp)
+		{
+			// We do this by looking at the ratio between current bounds (including deformation) and undeformed (straight from staticmesh)
+			const float MinExtent = 1.0f;
+			FBoxSphereBounds UndeformedBounds = StaticMesh->GetBounds().TransformBy(ComponentToWorld);
+			if (UndeformedBounds.BoxExtent.X >= MinExtent)
+			{
+				SplineDeformFactor = FMath::Max(SplineDeformFactor, Bounds.BoxExtent.X / UndeformedBounds.BoxExtent.X);
+			}
+			if (UndeformedBounds.BoxExtent.Y >= MinExtent)
+			{
+				SplineDeformFactor = FMath::Max(SplineDeformFactor, Bounds.BoxExtent.Y / UndeformedBounds.BoxExtent.Y);
+			}
+			if (UndeformedBounds.BoxExtent.Z >= MinExtent)
+			{
+				SplineDeformFactor = FMath::Max(SplineDeformFactor, Bounds.BoxExtent.Z / UndeformedBounds.BoxExtent.Z);
+			}
+		}
+
+		const FSphere BoundingSphere	= Bounds.GetSphere();
+		const float LocalTexelFactor	= StaticMesh->GetStreamingTextureFactor(0) * FMath::Max(0.0f, StreamingDistanceMultiplier);
+		const float LocalLightmapFactor	= bHasValidLightmapCoordinates ? StaticMesh->GetStreamingTextureFactor(StaticMesh->LightMapCoordinateIndex) : 1.0f;
+		const float WorldTexelFactor	= SplineDeformFactor * LocalTexelFactor * ComponentToWorld.GetMaximumAxisScale();
+		const float WorldLightmapFactor	= SplineDeformFactor * LocalLightmapFactor * ComponentToWorld.GetMaximumAxisScale();
+
+		for (int32 LODIndex = 0; LODIndex < StaticMesh->RenderData->LODResources.Num(); ++LODIndex)
+		{
+			FStaticMeshLODResources& LOD = StaticMesh->RenderData->LODResources[LODIndex];
+
+			// Process each material applied to the top LOD of the mesh.
+			for( int32 ElementIndex = 0; ElementIndex < LOD.Sections.Num(); ElementIndex++ )
+			{
+				const FStaticMeshSection& Element = LOD.Sections[ElementIndex];
+				UMaterialInterface* Material = GetMaterial(Element.MaterialIndex);
+				if(!Material)
+				{
+					Material = UMaterial::GetDefaultMaterial(MD_Surface);
+				}
+
+				// Enumerate the textures used by the material.
+				TArray<UTexture*> Textures;
+
+				Material->GetUsedTextures(Textures, EMaterialQualityLevel::Num, false, FeatureLevel, false);
+
+				// Add each texture to the output with the appropriate parameters.
+				// TODO: Take into account which UVIndex is being used.
+				for(int32 TextureIndex = 0;TextureIndex < Textures.Num();TextureIndex++)
+				{
+					FStreamingTexturePrimitiveInfo& StreamingTexture = *new(OutStreamingTextures) FStreamingTexturePrimitiveInfo;
+					StreamingTexture.Bounds = BoundingSphere;
+					StreamingTexture.TexelFactor = WorldTexelFactor;
+					StreamingTexture.Texture = Textures[TextureIndex];
+				}
+			}
+
+			// Add in the lightmaps and shadowmaps.
+			if ( LODData.IsValidIndex(LODIndex) && bHasValidLightmapCoordinates )
+			{
+				const FStaticMeshComponentLODInfo& LODInfo = LODData[LODIndex];
+				FLightMap2D* Lightmap = LODInfo.LightMap ? LODInfo.LightMap->GetLightMap2D() : NULL;
+				uint32 LightmapIndex = AllowHighQualityLightmaps(FeatureLevel) ? 0 : 1;
+				if ( Lightmap != NULL && Lightmap->IsValid(LightmapIndex) )
+				{
+					const FVector2D& Scale = Lightmap->GetCoordinateScale();
+					if ( Scale.X > SMALL_NUMBER && Scale.Y > SMALL_NUMBER )
+					{
+						float LightmapFactorX		 = WorldLightmapFactor / Scale.X;
+						float LightmapFactorY		 = WorldLightmapFactor / Scale.Y;
+						FStreamingTexturePrimitiveInfo& StreamingTexture = *new(OutStreamingTextures) FStreamingTexturePrimitiveInfo;
+						StreamingTexture.Bounds		 = BoundingSphere;
+						StreamingTexture.TexelFactor = FMath::Max(LightmapFactorX, LightmapFactorY);
+						StreamingTexture.Texture	 = Lightmap->GetTexture(LightmapIndex);
+					}
+				}
+
+				FShadowMap2D* Shadowmap = LODInfo.ShadowMap ? LODInfo.ShadowMap->GetShadowMap2D() : NULL;
+
+				if ( Shadowmap != NULL && Shadowmap->IsValid() )
+				{
+					const FVector2D& Scale = Shadowmap->GetCoordinateScale();
+					if ( Scale.X > SMALL_NUMBER && Scale.Y > SMALL_NUMBER )
+					{
+						float ShadowmapFactorX		 = WorldLightmapFactor / Scale.X;
+						float ShadowmapFactorY		 = WorldLightmapFactor / Scale.Y;
+						FStreamingTexturePrimitiveInfo& StreamingTexture = *new(OutStreamingTextures) FStreamingTexturePrimitiveInfo;
+						StreamingTexture.Bounds		 = BoundingSphere;
+						StreamingTexture.TexelFactor = FMath::Max(ShadowmapFactorX, ShadowmapFactorY);
+						StreamingTexture.Texture	 = Shadowmap->GetTexture();
+					}
+				}
+			}
+		}
+	}
+}
+
+UBodySetup* UStaticMeshComponent::GetBodySetup()
+{
+	if(StaticMesh)
+	{
+		return StaticMesh->BodySetup;
+	}
+
+	return NULL;
+}
+
+bool UStaticMeshComponent::CanEditSimulatePhysics()
+{
+	if (UBodySetup* BodySetup = GetBodySetup())
+	{
+		return (BodySetup->AggGeom.GetElementCount() > 0) || (BodySetup->CollisionTraceFlag == CTF_UseComplexAsSimple);
+	}
+	else
+	{
+		return false;
+	}
+}
+
+FColor UStaticMeshComponent::GetWireframeColor() const
+{
+	if(bOverrideWireframeColor)
+	{
+		return WireframeColorOverride;
+	}
+	else
+	{
+		if(Mobility == EComponentMobility::Static)
+		{
+			return FColor(0, 255, 255, 255);
+		}
+		else if(Mobility == EComponentMobility::Stationary)
+		{
+			return FColor(128, 128, 255, 255);
+		}
+		else // Movable
+		{
+			if(BodyInstance.bSimulatePhysics)
+			{
+				return FColor(0, 255, 128, 255);
+			}
+			else
+			{
+				return FColor(255, 0, 255, 255);
+			}
+		}
+	}
+}
+
+
+bool UStaticMeshComponent::DoesSocketExist(FName InSocketName) const 
+{
+	return (GetSocketByName(InSocketName)  != NULL);
+}
+
+class UStaticMeshSocket const* UStaticMeshComponent::GetSocketByName(FName InSocketName) const
+{
+	UStaticMeshSocket const* Socket = NULL;
+
+	if( StaticMesh )
+	{
+		Socket = StaticMesh->FindSocket( InSocketName );
+	}
+
+	return Socket;
+}
+
+FTransform UStaticMeshComponent::GetSocketTransform(FName InSocketName, ERelativeTransformSpace TransformSpace) const
+{
+	if (InSocketName != NAME_None)
+	{
+		UStaticMeshSocket const* const Socket = GetSocketByName(InSocketName);
+		if (Socket)
+		{
+			FTransform SocketWorldTransform;
+			if ( Socket->GetSocketTransform(SocketWorldTransform, this) )
+			{
+				switch(TransformSpace)
+				{
+					case RTS_World:
+					{
+						return SocketWorldTransform;
+					}
+					case RTS_Actor:
+					{
+						if( const AActor* Actor = GetOwner() )
+						{
+							return SocketWorldTransform.GetRelativeTransform( GetOwner()->GetTransform() );
+						}
+						break;
+					}
+					case RTS_Component:
+					{
+						return SocketWorldTransform.GetRelativeTransform(ComponentToWorld);
+					}
+				}
+			}
+		}
+	}
+
+	return ComponentToWorld;
+}
+
+
+bool UStaticMeshComponent::RequiresOverrideVertexColorsFixup( TArray<int32>& OutLODIndices )
+{
+	OutLODIndices.Empty();
+	bool bFixupRequired = false;
+
+#if WITH_EDITORONLY_DATA
+	if ( StaticMesh && StaticMesh->RenderData
+		&& StaticMesh->RenderData->DerivedDataKey != StaticMeshDerivedDataKey )
+	{
+		int32 NumLODs = StaticMesh->GetNumLODs();
+		// Iterate over each LOD to confirm which ones, if any, actually need to have their colors updated
+		for ( TArray<FStaticMeshComponentLODInfo>::TConstIterator LODIter( LODData ); LODIter; ++LODIter )
+		{
+			const FStaticMeshComponentLODInfo& CurCompLODInfo = *LODIter;
+
+			// Confirm that the LOD has override colors and cached vertex data. If it doesn't have both, it can't be fixed up.
+			if ( CurCompLODInfo.OverrideVertexColors && 
+				CurCompLODInfo.OverrideVertexColors->GetNumVertices() > 0 && 
+				CurCompLODInfo.PaintedVertices.Num() > 0 && 
+				LODIter.GetIndex() < NumLODs )
+			{
+				OutLODIndices.Add( LODIter.GetIndex() );
+				bFixupRequired = true;
+			}
+		}
+	}
+#endif // WITH_EDITORONLY_DATA
+
+	return bFixupRequired;
+}
+
+void UStaticMeshComponent::RemoveInstanceVertexColorsFromLOD( int32 LODToRemoveColorsFrom )
+{
+#if WITH_EDITORONLY_DATA
+	if (StaticMesh && LODToRemoveColorsFrom < StaticMesh->GetNumLODs() && LODToRemoveColorsFrom < LODData.Num())
+	{
+		FStaticMeshComponentLODInfo& CurrentLODInfo = LODData[LODToRemoveColorsFrom];
+
+		CurrentLODInfo.ReleaseOverrideVertexColorsAndBlock();
+		CurrentLODInfo.PaintedVertices.Empty();
+		StaticMeshDerivedDataKey = StaticMesh->RenderData->DerivedDataKey;
+	}
+#endif
+}
+
+void UStaticMeshComponent::RemoveInstanceVertexColors()
+{
+#if WITH_EDITORONLY_DATA
+	for ( int32 i=0; i < StaticMesh->GetNumLODs(); i++ )
+	{
+		RemoveInstanceVertexColorsFromLOD( i );
+	}
+#endif
+}
+
+void UStaticMeshComponent::CopyInstanceVertexColorsIfCompatible( UStaticMeshComponent* SourceComponent )
+{
+#if WITH_EDITORONLY_DATA
+	// The static mesh assets have to match, currently.
+	if (( StaticMesh->GetPathName() == SourceComponent->StaticMesh->GetPathName() ) &&
+		( SourceComponent->LODData.Num() != 0 ))
+	{
+		Modify();
+
+		bool bRegistered = IsRegistered();
+		FComponentReregisterContext ReregisterContext(this);
+		if (bRegistered)
+		{
+			FlushRenderingCommands(); // don't sync threads unless we have to
+		}
+		// Remove any and all vertex colors from the target static mesh, if they exist.
+		RemoveInstanceVertexColors();
+
+		int32 NumSourceLODs = SourceComponent->StaticMesh->GetNumLODs();
+
+		// This this will set up the LODData for all the LODs
+		SetLODDataCount( NumSourceLODs, NumSourceLODs );
+
+		// Copy vertex colors from Source to Target (this)
+		for ( int32 CurrentLOD = 0; CurrentLOD != NumSourceLODs; CurrentLOD++ )
+		{
+			FStaticMeshLODResources& SourceLODModel = SourceComponent->StaticMesh->RenderData->LODResources[CurrentLOD];
+			FStaticMeshComponentLODInfo& SourceLODInfo = SourceComponent->LODData[CurrentLOD];
+
+			FStaticMeshLODResources& TargetLODModel = StaticMesh->RenderData->LODResources[CurrentLOD];
+			FStaticMeshComponentLODInfo& TargetLODInfo = LODData[CurrentLOD];
+
+			if ( SourceLODInfo.OverrideVertexColors != NULL )
+			{
+				// Copy vertex colors from source to target.
+				FColorVertexBuffer* SourceColorBuffer = SourceLODInfo.OverrideVertexColors;
+
+				TArray< FColor > CopiedColors;
+				for ( uint32 ColorVertexIndex = 0; ColorVertexIndex < SourceColorBuffer->GetNumVertices(); ColorVertexIndex++ )
+				{
+					CopiedColors.Add( SourceColorBuffer->VertexColor( ColorVertexIndex ) );
+				}
+
+				FColorVertexBuffer* TargetColorBuffer = &TargetLODModel.ColorVertexBuffer;
+
+				if ( TargetLODInfo.OverrideVertexColors != NULL )
+				{
+					TargetLODInfo.BeginReleaseOverrideVertexColors();
+					FlushRenderingCommands();
+				}
+				else
+				{
+					TargetLODInfo.OverrideVertexColors = new FColorVertexBuffer;
+					TargetLODInfo.OverrideVertexColors->InitFromColorArray( CopiedColors );
+				}
+				BeginInitResource( TargetLODInfo.OverrideVertexColors );
+			}
+		}
+
+		CachePaintedDataIfNecessary();
+		StaticMeshDerivedDataKey = StaticMesh->RenderData->DerivedDataKey;
+
+		MarkRenderStateDirty();
+	}
+#endif
+}
+
+void UStaticMeshComponent::CachePaintedDataIfNecessary()
+{
+#if WITH_EDITORONLY_DATA
+	// Only cache the vertex positions if we're in the editor
+	if ( GIsEditor && StaticMesh )
+	{
+		// Iterate over each component LOD info checking for the existence of override colors
+		int32 NumLODs = StaticMesh->GetNumLODs();
+		for ( TArray<FStaticMeshComponentLODInfo>::TIterator LODIter( LODData ); LODIter; ++LODIter )
+		{
+			FStaticMeshComponentLODInfo& CurCompLODInfo = *LODIter;
+
+			// Workaround for a copy-paste bug. If the number of painted vertices is <= 1 we know the data is garbage.
+			if ( CurCompLODInfo.PaintedVertices.Num() <= 1 )
+			{
+				CurCompLODInfo.PaintedVertices.Empty();
+			}
+
+			// If the mesh has override colors but no cached vertex positions, then the current vertex positions should be cached to help preserve instanced vertex colors during mesh tweaks
+			// NOTE: We purposefully do *not* cache the positions if cached positions already exist, as this would result in the loss of the ability to fixup the component if the source mesh
+			// were changed multiple times before a fix-up operation was attempted
+			if ( CurCompLODInfo.OverrideVertexColors && 
+				 CurCompLODInfo.OverrideVertexColors->GetNumVertices() > 0 &&
+				 CurCompLODInfo.PaintedVertices.Num() == 0 &&
+				 LODIter.GetIndex() < NumLODs ) 
+			{
+				FStaticMeshLODResources* CurRenderData = &(StaticMesh->RenderData->LODResources[ LODIter.GetIndex() ]);
+				if ( CurRenderData->GetNumVertices() == CurCompLODInfo.OverrideVertexColors->GetNumVertices() )
+				{
+					// Cache the data.
+					CurCompLODInfo.PaintedVertices.Reserve( CurRenderData->GetNumVertices() );
+					for ( int32 VertIndex = 0; VertIndex < CurRenderData->GetNumVertices(); ++VertIndex )
+					{
+						FPaintedVertex* Vertex = new( CurCompLODInfo.PaintedVertices ) FPaintedVertex;
+						Vertex->Position = CurRenderData->PositionVertexBuffer.VertexPosition( VertIndex );
+						Vertex->Normal = CurRenderData->VertexBuffer.VertexTangentZ( VertIndex );
+						Vertex->Color = CurCompLODInfo.OverrideVertexColors->VertexColor( VertIndex );
+					}
+				}
+				else
+				{
+					UE_LOG(LogStaticMesh, Warning, TEXT("Unable to cache painted data for mesh component. Vertex color overrides will be lost if the mesh is modified. %s %s LOD%d."), *GetFullName(), *StaticMesh->GetFullName(), LODIter.GetIndex() );
+				}
+			}
+		}
+	}
+#endif // WITH_EDITORONLY_DATA
+}
+
+bool UStaticMeshComponent::FixupOverrideColorsIfNecessary( bool bRebuildingStaticMesh )
+{
+#if WITH_EDITORONLY_DATA
+
+	// Detect if there is a version mismatch between the source mesh and the component. If so, the component's LODs potentially
+	// need to have their override colors updated to match changes in the source mesh.
+	TArray<int32> LODsToUpdate;
+
+	if ( RequiresOverrideVertexColorsFixup( LODsToUpdate ) )
+	{
+		// Check if we are building the static mesh.  If so we dont need to reregister this component as its already unregistered and will be reregistered
+		// when the static mesh is done building.  Having nested reregister contexts is not supported.
+		if( bRebuildingStaticMesh )
+		{
+			PrivateFixupOverrideColors( LODsToUpdate );
+		}
+		else
+		{
+			// Detach this component because rendering changes are about to be applied
+			FComponentReregisterContext ReregisterContext( this );
+			PrivateFixupOverrideColors( LODsToUpdate );
+		}
+	}
+	return LODsToUpdate.Num() > 0 ? true : false;
+#else
+	return false;
+#endif // WITH_EDITORONLY_DATA
+
+}
+
+void UStaticMeshComponent::PrivateFixupOverrideColors( const TArray<int32>& LODsToUpdate )
+{
+#if WITH_EDITOR
+	UE_LOG(LogStaticMesh,Verbose,TEXT("Fixing up override colors for %s [%s]"),*GetPathName(),*StaticMesh->GetPathName());
+	for ( TArray<int32>::TConstIterator LODIdxIter( LODsToUpdate ); LODIdxIter; ++LODIdxIter )
+	{
+		FStaticMeshComponentLODInfo& CurCompLODInfo = LODData[ *LODIdxIter ];
+		FStaticMeshLODResources& CurRenderData = StaticMesh->RenderData->LODResources[ *LODIdxIter ];
+		TArray<FColor> NewOverrideColors;
+
+		CurCompLODInfo.BeginReleaseOverrideVertexColors();
+		FlushRenderingCommands();
+		RemapPaintedVertexColors(
+			CurCompLODInfo.PaintedVertices,
+			*CurCompLODInfo.OverrideVertexColors,
+			CurRenderData.PositionVertexBuffer,
+			&CurRenderData.VertexBuffer,
+			NewOverrideColors
+			);
+		if (NewOverrideColors.Num())
+		{
+			CurCompLODInfo.OverrideVertexColors->InitFromColorArray(NewOverrideColors);
+		}
+		// Note: In order to avoid data loss we dont clear the colors if they cannot be fixed up.
+
+		// Initialize the vert. colors
+		BeginInitResource( CurCompLODInfo.OverrideVertexColors );
+	}
+	StaticMeshDerivedDataKey = StaticMesh->RenderData->DerivedDataKey;
+#endif // WITH_EDITOR
+}
+
+void UStaticMeshComponent::InitResources()
+{
+	for(int32 LODIndex = 0; LODIndex < LODData.Num(); LODIndex++)
+	{
+		FStaticMeshComponentLODInfo &LODInfo = LODData[LODIndex];
+		if(LODInfo.OverrideVertexColors)
+		{
+			BeginInitResource(LODInfo.OverrideVertexColors);
+			INC_DWORD_STAT_BY( STAT_InstVertexColorMemory, LODInfo.OverrideVertexColors->GetAllocatedSize() );
+		}
+	}
+}
+
+
+
+void UStaticMeshComponent::ReleaseResources()
+{
+	for(int32 LODIndex = 0;LODIndex < LODData.Num();LODIndex++)
+	{
+		LODData[LODIndex].BeginReleaseOverrideVertexColors();
+	}
+
+	DetachFence.BeginFence();
+}
+
+void UStaticMeshComponent::BeginDestroy()
+{
+	Super::BeginDestroy();
+	ReleaseResources();
+}
+
+void UStaticMeshComponent::ExportCustomProperties(FOutputDevice& Out, uint32 Indent)
+{
+	for (int32 LODIdx = 0; LODIdx < LODData.Num(); ++LODIdx)
+	{
+		Out.Logf(TEXT("%sCustomProperties "), FCString::Spc(Indent));
+
+		FStaticMeshComponentLODInfo& LODInfo = LODData[LODIdx];
+
+		if ((LODInfo.PaintedVertices.Num() > 0) || LODInfo.OverrideVertexColors)
+		{
+			Out.Logf( TEXT("CustomLODData LOD=%d "), LODIdx );
+		}
+
+		// Export the PaintedVertices array
+		if (LODInfo.PaintedVertices.Num() > 0)
+		{
+			FString	ValueStr;
+			LODInfo.ExportText(ValueStr);
+			Out.Log(ValueStr);
+		}
+		
+		// Export the OverrideVertexColors buffer
+		if(LODInfo.OverrideVertexColors && LODInfo.OverrideVertexColors->GetNumVertices() > 0)
+		{
+			FString	Value;
+			LODInfo.OverrideVertexColors->ExportText(Value);
+
+			Out.Log(Value);
+		}
+		Out.Logf(TEXT("\r\n"));
+	}
+}
+
+void UStaticMeshComponent::ImportCustomProperties(const TCHAR* SourceText, FFeedbackContext* Warn)
+{
+	check(SourceText);
+	check(Warn);
+
+	if (FParse::Command(&SourceText,TEXT("CustomLODData")))
+	{
+		int32 MaxLODIndex = -1;
+		int32 LODIndex;
+		FString TmpStr;
+
+		if (FParse::Value(SourceText, TEXT("LOD="), LODIndex))
+		{
+			TmpStr = FString::Printf(TEXT("%d"), LODIndex);
+			SourceText += TmpStr.Len() + 5;
+
+			// See if we need to add a new element to the LODData array
+			if (LODIndex > MaxLODIndex)
+			{
+				SetLODDataCount(LODIndex + 1, LODIndex + 1);
+				MaxLODIndex = LODIndex;
+			}
+		}
+
+		FStaticMeshComponentLODInfo& LODInfo = LODData[LODIndex];
+
+		// Populate the PaintedVertices array
+		LODInfo.ImportText(&SourceText);
+
+		// Populate the OverrideVertexColors buffer
+		if (const TCHAR* VertColorStr = FCString::Stristr(SourceText, TEXT("ColorVertexData")))
+		{
+			SourceText = VertColorStr;
+
+			// this component could have been constructed from a template that
+			// had its own vert color overrides; so before we apply the
+			// custom color data, we need to clear the old vert colors (so
+			// we can properly call ImportText())
+			RemoveInstanceVertexColorsFromLOD(LODIndex);
+
+			// may not be null at the start (could have been initialized 
+			// from a blueprint component template with vert coloring), but 
+			// should be null by this point, after RemoveInstanceVertexColorsFromLOD()
+			check(LODInfo.OverrideVertexColors == nullptr);
+
+			LODInfo.OverrideVertexColors = new FColorVertexBuffer;
+			LODInfo.OverrideVertexColors->ImportText(SourceText);
+		}
+	}
+}
+
+#if WITH_EDITOR
+
+void UStaticMeshComponent::PreEditUndo()
+{
+	Super::PreEditUndo();
+
+	// Undo can result in a resize of LODData which can calls ~FStaticMeshComponentLODInfo.
+	// To safely delete FStaticMeshComponentLODInfo::OverrideVertexColors we
+	// need to make sure the RT thread has no access to it any more.
+	check(!IsRegistered());
+	ReleaseResources();
+	DetachFence.Wait();
+}
+
+void UStaticMeshComponent::PostEditUndo()
+{
+	// If the StaticMesh was also involved in this transaction, it may need reinitialization first
+	if (StaticMesh)
+	{
+		StaticMesh->InitResources();
+	}
+
+	// The component's light-maps are loaded from the transaction, so their resources need to be reinitialized.
+	InitResources();
+
+	// Debug check command trying to track down undo related uninitialized resource
+	if (StaticMesh != NULL && StaticMesh->RenderData->LODResources.Num() > 0)
+	{
+		ENQUEUE_UNIQUE_RENDER_COMMAND_ONEPARAMETER(
+			ResourceCheckCommand,
+			FRenderResource*, Resource, &StaticMesh->RenderData->LODResources[0].IndexBuffer,
+			{
+				check( Resource->IsInitialized() );
+			}
+		);
+	}
+
+	Super::PostEditUndo();
+}
+
+void UStaticMeshComponent::PostEditChangeProperty(FPropertyChangedEvent& PropertyChangedEvent)
+{
+	// Ensure that OverriddenLightMapRes is a factor of 4
+	OverriddenLightMapRes = FMath::Max(OverriddenLightMapRes + 3 & ~3,4);
+
+	UProperty* PropertyThatChanged = PropertyChangedEvent.Property;
+	if (PropertyThatChanged)
+	{
+		if (((PropertyThatChanged->GetName().Contains(TEXT("OverriddenLightMapRes")) ) && (bOverrideLightMapRes == true)) ||
+			(PropertyThatChanged->GetName().Contains(TEXT("bOverrideLightMapRes")) ))
+		{
+			InvalidateLightingCache();
+		}
+
+		if ( PropertyThatChanged->GetName().Contains(TEXT("bIgnoreInstanceForTextureStreaming")) ||
+			 PropertyThatChanged->GetName().Contains(TEXT("StreamingDistanceMultiplier")) )
+		{
+			GEngine->TriggerStreamingDataRebuild();
+		}
+	}
+
+	FBodyInstanceEditorHelpers::EnsureConsistentMobilitySimulationSettingsOnPostEditChange(this, PropertyChangedEvent);
+
+	LightmassSettings.EmissiveBoost = FMath::Max(LightmassSettings.EmissiveBoost, 0.0f);
+	LightmassSettings.DiffuseBoost = FMath::Max(LightmassSettings.DiffuseBoost, 0.0f);
+
+	// Ensure properties are in sane range.
+	SubDivisionStepSize = FMath::Clamp( SubDivisionStepSize, 1, 128 );
+
+	Super::PostEditChangeProperty(PropertyChangedEvent);
+}
+#endif // WITH_EDITOR
+
+
+void UStaticMeshComponent::PostLoad()
+{
+	Super::PostLoad();
+
+	if ( StaticMesh )
+	{
+		StaticMesh->ConditionalPostLoad();
+		CachePaintedDataIfNecessary();
+		if ( FixupOverrideColorsIfNecessary() )
+		{
+#if WITH_EDITORONLY_DATA
+			FFormatNamedArguments Arguments;
+			Arguments.Add(TEXT("MeshName"), FText::FromString(GetName()));
+			FMessageLog("MapCheck").Info()
+				->AddToken(FUObjectToken::Create(GetOuter()))
+				->AddToken(FTextToken::Create(FText::Format( LOCTEXT( "MapCheck_Message_RepairedPaintedVertexColors", "{MeshName} : Repaired painted vertex colors" ), Arguments ) ))
+				->AddToken(FMapErrorToken::Create(FMapErrors::RepairedPaintedVertexColors));
+#endif
+		}
+	}
+
+#if WITH_EDITORONLY_DATA
+	// Remap the materials array if the static mesh materials may have been remapped to remove zero triangle sections.
+	if (StaticMesh && GetLinkerUE4Version() < VER_UE4_REMOVE_ZERO_TRIANGLE_SECTIONS && OverrideMaterials.Num())
+	{
+		StaticMesh->ConditionalPostLoad();
+		if (StaticMesh->HasValidRenderData()
+			&& StaticMesh->RenderData->MaterialIndexToImportIndex.Num())
+		{
+			TArray<UMaterialInterface*> OldMaterials;
+			const TArray<int32>& MaterialIndexToImportIndex = StaticMesh->RenderData->MaterialIndexToImportIndex;
+
+			Exchange(OverrideMaterials,OldMaterials);
+			OverrideMaterials.Empty(MaterialIndexToImportIndex.Num());
+			for (int32 MaterialIndex = 0; MaterialIndex < MaterialIndexToImportIndex.Num(); ++MaterialIndex)
+			{
+				UMaterialInterface* Material = NULL;
+				int32 OldMaterialIndex = MaterialIndexToImportIndex[MaterialIndex];
+				if (OldMaterials.IsValidIndex(OldMaterialIndex))
+				{
+					Material = OldMaterials[OldMaterialIndex];
+				}
+				OverrideMaterials.Add(Material);
+			}
+		}
+
+		if (OverrideMaterials.Num() > StaticMesh->Materials.Num())
+		{
+			OverrideMaterials.RemoveAt(StaticMesh->Materials.Num(), OverrideMaterials.Num() - StaticMesh->Materials.Num());
+		}
+	}
+#endif // #if WITH_EDITORONLY_DATA
+
+	// Legacy content may contain a lightmap resolution of 0, which was valid when vertex lightmaps were supported, but not anymore with only texture lightmaps
+	OverriddenLightMapRes = FMath::Max(OverriddenLightMapRes, 4);
+
+	// Initialize the resources for the freshly loaded component.
+	InitResources();
+}
+
+bool UStaticMeshComponent::SetStaticMesh(UStaticMesh* NewMesh)
+{
+	// Do nothing if we are already using the supplied static mesh
+	if(NewMesh == StaticMesh)
+	{
+		return false;
+	}
+
+	// Don't allow changing static meshes if "static" and registered
+	AActor* Owner = GetOwner();
+	if(!AreDynamicDataChangesAllowed() && Owner != NULL)
+	{
+		FMessageLog("PIE").Warning(FText::Format(LOCTEXT("SetMeshOnStatic", "Calling SetStaticMesh on '{0}' but Mobility is Static."), 
+			FText::FromString(GetPathName())));
+		return false;
+	}
+
+
+	StaticMesh = NewMesh;
+
+	// Need to send this to render thread at some point
+	MarkRenderStateDirty();
+
+	// Update physics representation right away
+	RecreatePhysicsState();
+
+	// update navigation relevancy
+	bNavigationRelevant = IsNavigationRelevant();
+
+	// Notify the streaming system. Don't use Update(), because this may be the first time the mesh has been set
+	// and the component may have to be added to the streaming system for the first time.
+	IStreamingManager::Get().NotifyPrimitiveAttached( this, DPT_Spawned );
+
+	// Since we have new mesh, we need to update bounds
+	UpdateBounds();
+	return true;
+}
+
+void UStaticMeshComponent::SetForcedLodModel(int32 NewForcedLodModel)
+{
+	if (ForcedLodModel != NewForcedLodModel)
+	{
+		ForcedLodModel = NewForcedLodModel;
+		MarkRenderStateDirty();
+	}
+}
+
+void UStaticMeshComponent::GetLocalBounds(FVector& Min, FVector& Max) const
+{
+	if (StaticMesh)
+	{
+		FBoxSphereBounds MeshBounds = StaticMesh->GetBounds();
+		Min = MeshBounds.Origin - MeshBounds.BoxExtent;
+		Max = MeshBounds.Origin + MeshBounds.BoxExtent;
+	}
+}
+
+bool UStaticMeshComponent::UsesOnlyUnlitMaterials() const
+{
+	if (StaticMesh && StaticMesh->RenderData)
+	{
+		// Figure out whether any of the sections has a lit material assigned.
+		bool bUsesOnlyUnlitMaterials = true;
+		for (int32 LODIndex = 0; bUsesOnlyUnlitMaterials && LODIndex < StaticMesh->RenderData->LODResources.Num(); ++LODIndex)
+		{
+			FStaticMeshLODResources& LOD = StaticMesh->RenderData->LODResources[LODIndex];
+			for (int32 ElementIndex=0; bUsesOnlyUnlitMaterials && ElementIndex<LOD.Sections.Num(); ElementIndex++)
+			{
+				UMaterialInterface*	MaterialInterface	= GetMaterial(LOD.Sections[ElementIndex].MaterialIndex);
+				UMaterial*			Material			= MaterialInterface ? MaterialInterface->GetMaterial() : NULL;
+
+				bUsesOnlyUnlitMaterials = Material && Material->GetShadingModel() == MSM_Unlit;
+			}
+		}
+		return bUsesOnlyUnlitMaterials;
+	}
+	else
+	{
+		return false;
+	}
+}
+
+
+bool UStaticMeshComponent::GetLightMapResolution( int32& Width, int32& Height ) const
+{
+	bool bPadded = false;
+	if( StaticMesh )
+	{
+		// Use overriden per component lightmap resolution.
+		if( bOverrideLightMapRes )
+		{
+			Width	= OverriddenLightMapRes;
+			Height	= OverriddenLightMapRes;
+		}
+		// Use the lightmap resolution defined in the static mesh.
+		else
+		{
+			Width	= StaticMesh->LightMapResolution;
+			Height	= StaticMesh->LightMapResolution;
+		}
+		bPadded = true;
+	}
+	// No associated static mesh!
+	else
+	{
+		Width	= 0;
+		Height	= 0;
+	}
+
+	return bPadded;
+}
+
+
+void UStaticMeshComponent::GetEstimatedLightMapResolution(int32& Width, int32& Height) const
+{
+	if (StaticMesh)
+	{
+		ELightMapInteractionType LMIType = GetStaticLightingType();
+
+		bool bUseSourceMesh = false;
+
+		// Use overriden per component lightmap resolution.
+		// If the overridden LM res is > 0, then this is what would be used...
+		if (bOverrideLightMapRes == true)
+		{
+			if (OverriddenLightMapRes != 0)
+			{
+				Width	= OverriddenLightMapRes;
+				Height	= OverriddenLightMapRes;
+			}
+		}
+		else
+		{
+			bUseSourceMesh = true;
+		}
+
+		// Use the lightmap resolution defined in the static mesh.
+		if (bUseSourceMesh == true)
+		{
+			Width	= StaticMesh->LightMapResolution;
+			Height	= StaticMesh->LightMapResolution;
+		}
+
+		// If it was not set by anything, give it a default value...
+		if (Width == 0)
+		{
+			int32 TempInt = 0;
+			verify(GConfig->GetInt(TEXT("DevOptions.StaticLighting"), TEXT("DefaultStaticMeshLightingRes"), TempInt, GLightmassIni));
+
+			Width	= TempInt;
+			Height	= TempInt;
+		}
+	}
+	else
+	{
+		Width	= 0;
+		Height	= 0;
+	}
+}
+
+
+int32 UStaticMeshComponent::GetStaticLightMapResolution() const
+{
+	int32 Width, Height;
+	GetLightMapResolution(Width, Height);
+	return FMath::Max<int32>(Width, Height);
+}
+
+bool UStaticMeshComponent::HasValidSettingsForStaticLighting() const
+{
+	int32 LightMapWidth = 0;
+	int32 LightMapHeight = 0;
+	GetLightMapResolution(LightMapWidth, LightMapHeight);
+
+	return Super::HasValidSettingsForStaticLighting() 
+		&& StaticMesh
+		&& UsesTextureLightmaps(LightMapWidth, LightMapHeight);
+}
+
+bool UStaticMeshComponent::UsesTextureLightmaps(int32 InWidth, int32 InHeight) const
+{
+	return (
+		(HasLightmapTextureCoordinates()) &&
+		(InWidth > 0) && 
+		(InHeight > 0)
+		);
+}
+
+
+bool UStaticMeshComponent::HasLightmapTextureCoordinates() const
+{
+	if ((StaticMesh != NULL) &&
+		(StaticMesh->LightMapCoordinateIndex >= 0) &&
+		(StaticMesh->RenderData != NULL) &&
+		(StaticMesh->RenderData->LODResources.Num() > 0) &&
+		(StaticMesh->LightMapCoordinateIndex >= 0) &&	
+		((uint32)StaticMesh->LightMapCoordinateIndex < StaticMesh->RenderData->LODResources[0].VertexBuffer.GetNumTexCoords()))
+	{
+		return true;
+	}
+	return false;
+}
+
+
+void UStaticMeshComponent::GetTextureLightAndShadowMapMemoryUsage(int32 InWidth, int32 InHeight, int32& OutLightMapMemoryUsage, int32& OutShadowMapMemoryUsage) const
+{
+	// Stored in texture.
+	const float MIP_FACTOR = 1.33f;
+	OutShadowMapMemoryUsage = FMath::TruncToInt(MIP_FACTOR * InWidth * InHeight); // G8
+
+	auto FeatureLevel = GetWorld() ? GetWorld()->FeatureLevel : GMaxRHIFeatureLevel;
+
+	if (AllowHighQualityLightmaps(FeatureLevel))
+	{
+		OutLightMapMemoryUsage = FMath::TruncToInt(NUM_HQ_LIGHTMAP_COEF * MIP_FACTOR * InWidth * InHeight); // DXT5
+	}
+	else
+	{
+		OutLightMapMemoryUsage = FMath::TruncToInt(NUM_LQ_LIGHTMAP_COEF * MIP_FACTOR * InWidth * InHeight / 2); // DXT1
+	}
+}
+
+
+void UStaticMeshComponent::GetLightAndShadowMapMemoryUsage( int32& LightMapMemoryUsage, int32& ShadowMapMemoryUsage ) const
+{
+	// Zero initialize.
+	ShadowMapMemoryUsage = 0;
+	LightMapMemoryUsage = 0;
+
+	// Cache light/ shadow map resolution.
+	int32 LightMapWidth = 0;
+	int32	LightMapHeight = 0;
+	GetLightMapResolution(LightMapWidth, LightMapHeight);
+
+	// Determine whether static mesh/ static mesh component has static shadowing.
+	if (HasStaticLighting() && StaticMesh)
+	{
+		// Determine whether we are using a texture or vertex buffer to store precomputed data.
+		if (UsesTextureLightmaps(LightMapWidth, LightMapHeight) == true)
+		{
+			GetTextureLightAndShadowMapMemoryUsage(LightMapWidth, LightMapHeight, LightMapMemoryUsage, ShadowMapMemoryUsage);
+		}
+	}
+}
+
+
+bool UStaticMeshComponent::GetEstimatedLightAndShadowMapMemoryUsage( 
+	int32& TextureLightMapMemoryUsage, int32& TextureShadowMapMemoryUsage,
+	int32& VertexLightMapMemoryUsage, int32& VertexShadowMapMemoryUsage,
+	int32& StaticLightingResolution, bool& bIsUsingTextureMapping, bool& bHasLightmapTexCoords) const
+{
+	TextureLightMapMemoryUsage	= 0;
+	TextureShadowMapMemoryUsage	= 0;
+	VertexLightMapMemoryUsage	= 0;
+	VertexShadowMapMemoryUsage	= 0;
+	bIsUsingTextureMapping		= false;
+	bHasLightmapTexCoords		= false;
+
+	// Cache light/ shadow map resolution.
+	int32 LightMapWidth			= 0;
+	int32 LightMapHeight		= 0;
+	GetEstimatedLightMapResolution(LightMapWidth, LightMapHeight);
+	StaticLightingResolution = LightMapWidth;
+
+	int32 TrueLightMapWidth		= 0;
+	int32 TrueLightMapHeight	= 0;
+	GetLightMapResolution(TrueLightMapWidth, TrueLightMapHeight);
+
+	// Determine whether static mesh/ static mesh component has static shadowing.
+	if (HasStaticLighting() && StaticMesh)
+	{
+		// Determine whether we are using a texture or vertex buffer to store precomputed data.
+		bHasLightmapTexCoords = HasLightmapTextureCoordinates();
+		// Determine whether we are using a texture or vertex buffer to store precomputed data.
+		bIsUsingTextureMapping = UsesTextureLightmaps(TrueLightMapWidth, TrueLightMapHeight);
+		// Stored in texture.
+		GetTextureLightAndShadowMapMemoryUsage(LightMapWidth, LightMapHeight, TextureLightMapMemoryUsage, TextureShadowMapMemoryUsage);
+
+		return true;
+	}
+
+	return false;
+}
+
+int32 UStaticMeshComponent::GetNumMaterials() const
+{
+	// @note : you don't have to consider Materials.Num()
+	// that only counts if overriden and it can't be more than StaticMesh->Materials. 
+	if(StaticMesh)
+	{
+		return StaticMesh->Materials.Num();
+	}
+	else
+	{
+		return 0;
+	}
+}
+
+TArray<class UMaterialInterface*> UStaticMeshComponent::GetMaterials() const
+{
+	TArray<class UMaterialInterface*> OutMaterials = Super::GetMaterials();
+
+	// if no material is overriden, look for mesh material;
+	if(OutMaterials.Num() == 0 && StaticMesh)
+	{
+		OutMaterials = StaticMesh->Materials;
+	}
+
+	return OutMaterials;
+}
+
+UMaterialInterface* UStaticMeshComponent::GetMaterial(int32 MaterialIndex) const
+{
+	// If we have a base materials array, use that
+	if(MaterialIndex < OverrideMaterials.Num() && OverrideMaterials[MaterialIndex])
+	{
+		return OverrideMaterials[MaterialIndex];
+	}
+	// Otherwise get from static mesh
+	else
+	{
+		return StaticMesh ? StaticMesh->GetMaterial(MaterialIndex) : NULL;
+	}
+}
+
+void UStaticMeshComponent::GetUsedMaterials(TArray<UMaterialInterface*>& OutMaterials) const
+{
+	if( StaticMesh && StaticMesh->RenderData )
+	{
+		for (int32 LODIndex = 0; LODIndex < StaticMesh->RenderData->LODResources.Num(); LODIndex++)
+		{
+			FStaticMeshLODResources& LODResources = StaticMesh->RenderData->LODResources[LODIndex];
+			for (int32 SectionIndex = 0; SectionIndex < LODResources.Sections.Num(); SectionIndex++)
+			{
+				// Get the material for each element at the current lod index
+				OutMaterials.AddUnique(GetMaterial(LODResources.Sections[SectionIndex].MaterialIndex));
+			}
+		}
+	}
+}
+
+int32 UStaticMeshComponent::GetBlueprintCreatedComponentIndex() const
+{
+	int32 ComponentIndex = 0;
+	for(const auto& Component : GetOwner()->BlueprintCreatedComponents)
+	{
+		if(Component == this)
+		{
+			return ComponentIndex;
+		}
+
+		ComponentIndex++;
+	}
+
+	return INDEX_NONE;
+}
+
+FName UStaticMeshComponent::GetComponentInstanceDataType() const
+{
+	static const FName StaticMeshComponentInstanceDataName(TEXT("StaticMeshInstanceData"));
+	return StaticMeshComponentInstanceDataName;
+}
+
+FActorComponentInstanceData* UStaticMeshComponent::GetComponentInstanceData() const
+{
+	FStaticMeshComponentInstanceData* StaticMeshInstanceData = nullptr;
+
+	// Don't back up static lighting if there isn't any
+	if(bHasCachedStaticLighting)
+	{
+		StaticMeshInstanceData = new FStaticMeshComponentInstanceData(this);
+
+		// Fill in info
+		StaticMeshInstanceData->bHasCachedStaticLighting = true;
+		StaticMeshInstanceData->CachedStaticLighting.Transform = ComponentToWorld;
+		StaticMeshInstanceData->CachedStaticLighting.IrrelevantLights = IrrelevantLights;
+		StaticMeshInstanceData->CachedStaticLighting.LODDataLightMap.Empty(LODData.Num());
+		for (const FStaticMeshComponentLODInfo& LODDataEntry : LODData)
+		{
+			StaticMeshInstanceData->CachedStaticLighting.LODDataLightMap.Add(LODDataEntry.LightMap);
+			StaticMeshInstanceData->CachedStaticLighting.LODDataShadowMap.Add(LODDataEntry.ShadowMap);
+		}
+	}
+
+
+	// Cache instance vertex colors
+	for( int32 LODIndex = 0; LODIndex < LODData.Num(); ++LODIndex )
+	{
+		const FStaticMeshComponentLODInfo& LODInfo = LODData[LODIndex];
+
+		if ( LODInfo.OverrideVertexColors && LODInfo.OverrideVertexColors->GetNumVertices() > 0 && LODInfo.PaintedVertices.Num() > 0 )
+		{
+			if (!StaticMeshInstanceData)
+			{
+				StaticMeshInstanceData = new FStaticMeshComponentInstanceData(this);
+			}
+
+			StaticMeshInstanceData->AddVertexColorData(LODInfo, LODIndex);
+		}
+	}
+
+	return (StaticMeshInstanceData ? StaticMeshInstanceData : Super::GetComponentInstanceData());
+}
+
+void UStaticMeshComponent::ApplyComponentInstanceData(FStaticMeshComponentInstanceData* StaticMeshInstanceData)
+{
+	check(StaticMeshInstanceData);
+
+	// Note: ApplyComponentInstanceData is called while the component is registered so the rendering thread is already using this component
+	// That means all component state that is modified here must be mirrored on the scene proxy, which will be recreated to receive the changes later due to MarkRenderStateDirty.
+
+	if (StaticMesh != StaticMeshInstanceData->StaticMesh)
+	{
+		return;
+	}
+
+	// See if data matches current state
+	if(	StaticMeshInstanceData->bHasCachedStaticLighting && StaticMeshInstanceData->CachedStaticLighting.Transform.Equals(ComponentToWorld, 1.e-3f) )
+	{
+		const int32 NumLODLightMaps = StaticMeshInstanceData->CachedStaticLighting.LODDataLightMap.Num();
+		SetLODDataCount(NumLODLightMaps, NumLODLightMaps);
+		for (int32 i = 0; i < NumLODLightMaps; ++i)
+		{
+			LODData[i].LightMap = StaticMeshInstanceData->CachedStaticLighting.LODDataLightMap[i];
+			LODData[i].ShadowMap = StaticMeshInstanceData->CachedStaticLighting.LODDataShadowMap[i];
+		}
+
+		IrrelevantLights = StaticMeshInstanceData->CachedStaticLighting.IrrelevantLights;
+		bHasCachedStaticLighting = true;
+	}
+
+	{
+		FComponentReregisterContext ReregisterStaticMesh(this);
+		StaticMeshInstanceData->ApplyVertexColorData(this);
+	}
+}
+
+#include "AI/Navigation/RecastHelpers.h"
+bool UStaticMeshComponent::DoCustomNavigableGeometryExport(FNavigableGeometryExport& GeomExport) const
+{
+	if (StaticMesh != NULL && StaticMesh->NavCollision != NULL)
+	{
+		UNavCollision* NavCollision = StaticMesh->NavCollision;
+		if (NavCollision->bIsDynamicObstacle)
+		{
+			return false;
+		}
+		
+		if (NavCollision->bHasConvexGeometry)
+		{
+			const FVector Scale3D = ComponentToWorld.GetScale3D();
+			// if any of scales is 0 there's no point in exporting it
+			if (!Scale3D.IsZero())
+			{
+				GeomExport.ExportCustomMesh(NavCollision->ConvexCollision.VertexBuffer.GetData(), NavCollision->ConvexCollision.VertexBuffer.Num(),
+					NavCollision->ConvexCollision.IndexBuffer.GetData(), NavCollision->ConvexCollision.IndexBuffer.Num(), ComponentToWorld);
+
+				GeomExport.ExportCustomMesh(NavCollision->TriMeshCollision.VertexBuffer.GetData(), NavCollision->TriMeshCollision.VertexBuffer.Num(),
+					NavCollision->TriMeshCollision.IndexBuffer.GetData(), NavCollision->TriMeshCollision.IndexBuffer.Num(), ComponentToWorld);
+			}
+
+			// regardless of above we don't want "regular" collision export for this mesh instance
+			return false;
+		}
+	}
+
+	return true;
+}
+
+bool UStaticMeshComponent::IsNavigationRelevant() const
+{
+	return StaticMesh != nullptr && Super::IsNavigationRelevant();
+}
+
+void UStaticMeshComponent::GetNavigationData(FNavigationRelevantData& Data) const
+{
+	Super::GetNavigationData(Data);
+
+	if (StaticMesh && StaticMesh->NavCollision)	
+	{
+		UNavCollision* NavCollision = StaticMesh->NavCollision;
+		if (NavCollision->bIsDynamicObstacle)
+		{
+			NavCollision->GetNavigationModifier(Data.Modifiers, ComponentToWorld);
+		}
+	}
+}
+
+#if WITH_EDITOR
+bool UStaticMeshComponent::ComponentIsTouchingSelectionBox(const FBox& InSelBBox, const FEngineShowFlags& ShowFlags, const bool bConsiderOnlyBSP, const bool bMustEncompassEntireComponent) const
+{
+	if (!bConsiderOnlyBSP && ShowFlags.StaticMeshes && StaticMesh != nullptr && StaticMesh->HasValidRenderData())
+	{
+		// Check if we are even inside it's bounding box, if we are not, there is no way we colliding via the more advanced checks we will do.
+		if (Super::ComponentIsTouchingSelectionBox(InSelBBox, ShowFlags, bConsiderOnlyBSP, false))
+		{
+			TArray<FVector> Vertex;
+
+			FStaticMeshLODResources& LODModel = StaticMesh->RenderData->LODResources[0];
+			FIndexArrayView Indices = LODModel.IndexBuffer.GetArrayView();
+
+			for (const auto& Section : LODModel.Sections)
+			{
+				// Iterate over each triangle.
+				for (int32 TriangleIndex = 0; TriangleIndex < (int32)Section.NumTriangles; TriangleIndex++)
+				{
+					Vertex.Empty(3);
+
+					int32 FirstIndex = TriangleIndex * 3 + Section.FirstIndex;
+					for (int32 i = 0; i < 3; i++)
+					{
+						int32 VertexIndex = Indices[FirstIndex + i];
+						FVector LocalPosition = LODModel.PositionVertexBuffer.VertexPosition(VertexIndex);
+						Vertex.Emplace(ComponentToWorld.TransformPosition(LocalPosition));
+					}
+
+					// Check if the triangle is colliding with the bounding box.
+					FSeparatingAxisPointCheck ThePointCheck(Vertex, InSelBBox.GetCenter(), InSelBBox.GetExtent(), false);
+					if (!bMustEncompassEntireComponent && ThePointCheck.bHit)
+					{
+						// Needn't encompass entire component: any intersection, we consider as touching
+						return true;
+					}
+					else if (bMustEncompassEntireComponent && !ThePointCheck.bHit)
+					{
+						// Must encompass entire component: any non intersection, we consider as not touching
+						return false;
+					}
+				}
+			}
+
+			// If the selection box has to encompass all of the component and none of the component's verts failed the intersection test, this component
+			// is consider touching
+			return true;
+		}
+	}
+
+	return false;
+}
+
+
+bool UStaticMeshComponent::ComponentIsTouchingSelectionFrustum(const FConvexVolume& InFrustum, const FEngineShowFlags& ShowFlags, const bool bConsiderOnlyBSP, const bool bMustEncompassEntireComponent) const
+{
+	if (!bConsiderOnlyBSP && ShowFlags.StaticMeshes && StaticMesh != nullptr && StaticMesh->HasValidRenderData())
+	{
+		// Check if we are even inside it's bounding box, if we are not, there is no way we colliding via the more advanced checks we will do.
+		if (Super::ComponentIsTouchingSelectionFrustum(InFrustum, ShowFlags, bConsiderOnlyBSP, false))
+		{
+			TArray<FVector> Vertex;
+
+			FStaticMeshLODResources& LODModel = StaticMesh->RenderData->LODResources[0];
+
+			uint32 NumVertices = LODModel.VertexBuffer.GetNumVertices();
+			for (uint32 VertexIndex = 0; VertexIndex < NumVertices; ++VertexIndex)
+			{
+				const FVector& LocalPosition = LODModel.PositionVertexBuffer.VertexPosition(VertexIndex);
+				const FVector WorldPosition = ComponentToWorld.TransformPosition(LocalPosition);
+				bool bLocationIntersected = InFrustum.IntersectSphere(WorldPosition, 0.0f);
+				if (bLocationIntersected && !bMustEncompassEntireComponent)
+				{
+					return true;
+				}
+				else if (!bLocationIntersected && bMustEncompassEntireComponent)
+				{
+					return false;
+				}
+			}
+
+			return true;
+		}
+	}
+
+	return false;
+}
+#endif
+
+
+//////////////////////////////////////////////////////////////////////////
+// StaticMeshComponentLODInfo
+
+
+
+/** Default constructor */
+FStaticMeshComponentLODInfo::FStaticMeshComponentLODInfo()
+	: OverrideVertexColors(NULL)
+#if WITH_EDITORONLY_DATA
+	, Owner(NULL)
+#endif
+{
+}
+
+/** Copy constructor */
+FStaticMeshComponentLODInfo::FStaticMeshComponentLODInfo(const FStaticMeshComponentLODInfo &rhs)
+	: OverrideVertexColors(NULL)
+#if WITH_EDITORONLY_DATA
+	, Owner(NULL)
+#endif
+{
+	if(rhs.OverrideVertexColors)
+	{
+		OverrideVertexColors = new FColorVertexBuffer;
+		OverrideVertexColors->InitFromColorArray(&rhs.OverrideVertexColors->VertexColor(0), rhs.OverrideVertexColors->GetNumVertices());
+	}
+}
+
+/** Destructor */
+FStaticMeshComponentLODInfo::~FStaticMeshComponentLODInfo()
+{
+	// Note: OverrideVertexColors had BeginReleaseResource called in UStaticMeshComponent::BeginDestroy, 
+	// And waits on a fence for that command to complete in UStaticMeshComponent::IsReadyForFinishDestroy,
+	// So we know it is safe to delete OverrideVertexColors here (RT can't be referencing it anymore)
+	CleanUp();
+}
+
+void FStaticMeshComponentLODInfo::CleanUp()
+{
+	if(OverrideVertexColors)
+	{
+		DEC_DWORD_STAT_BY( STAT_InstVertexColorMemory, OverrideVertexColors->GetAllocatedSize() );
+	}
+
+	delete OverrideVertexColors;
+	OverrideVertexColors = NULL;
+
+	PaintedVertices.Empty();
+}
+
+
+void FStaticMeshComponentLODInfo::BeginReleaseOverrideVertexColors()
+{
+	if(OverrideVertexColors)
+	{
+		// enqueue a rendering command to release
+		BeginReleaseResource(OverrideVertexColors);
+	}
+}
+
+void FStaticMeshComponentLODInfo::ReleaseOverrideVertexColorsAndBlock()
+{
+	if(OverrideVertexColors)
+	{
+		// enqueue a rendering command to release
+		BeginReleaseResource(OverrideVertexColors);
+		// Ensure the RT no longer accessed the data, might slow down
+		FlushRenderingCommands();
+		// The RT thread has no access to it any more so it's safe to delete it.
+		CleanUp();
+	}
+}
+
+void FStaticMeshComponentLODInfo::ExportText(FString& ValueStr)
+{
+	ValueStr += FString::Printf(TEXT("PaintedVertices(%d)="), PaintedVertices.Num());
+
+	// Rough approximation
+	ValueStr.Reserve(ValueStr.Len() + PaintedVertices.Num() * 125);
+
+	// Export the Position, Normal and Color info for each vertex
+	for(int32 i = 0; i < PaintedVertices.Num(); ++i)
+	{
+		FPaintedVertex& Vert = PaintedVertices[i];
+
+		ValueStr += FString::Printf(TEXT("((Position=(X=%.6f,Y=%.6f,Z=%.6f),"), Vert.Position.X, Vert.Position.Y, Vert.Position.Z);
+		ValueStr += FString::Printf(TEXT("(Normal=(X=%d,Y=%d,Z=%d,W=%d),"), Vert.Normal.Vector.X, Vert.Normal.Vector.Y, Vert.Normal.Vector.Z, Vert.Normal.Vector.W);
+		ValueStr += FString::Printf(TEXT("(Color=(B=%d,G=%d,R=%d,A=%d))"), Vert.Color.B, Vert.Color.G, Vert.Color.R, Vert.Color.A);
+
+		// Seperate each vertex entry with a comma
+		if ((i + 1) != PaintedVertices.Num())
+		{
+			ValueStr += TEXT(",");
+		}
+	}
+
+	ValueStr += TEXT(" ");
+}
+
+void FStaticMeshComponentLODInfo::ImportText(const TCHAR** SourceText)
+{
+	FString TmpStr;
+	int32 VertCount;
+	if (FParse::Value(*SourceText, TEXT("PaintedVertices("), VertCount))
+	{
+		TmpStr = FString::Printf(TEXT("%d"), VertCount);
+		*SourceText += TmpStr.Len() + 18;
+
+		FString SourceTextStr = *SourceText;
+		TArray<FString> Tokens;
+		int32 TokenIdx = 0;
+		bool bValidInput = true;
+
+		// Tokenize the text
+		SourceTextStr.ParseIntoArray(Tokens, TEXT(","), false);
+
+		// There should be 11 tokens per vertex
+		check(Tokens.Num() * 11 >= VertCount);
+
+		PaintedVertices.AddUninitialized(VertCount);
+
+		for (int32 Idx = 0; Idx < VertCount; ++Idx)
+		{
+			// Position
+			bValidInput &= FParse::Value(*Tokens[TokenIdx++], TEXT("X="), PaintedVertices[Idx].Position.X);
+			bValidInput &= FParse::Value(*Tokens[TokenIdx++], TEXT("Y="), PaintedVertices[Idx].Position.Y);
+			bValidInput &= FParse::Value(*Tokens[TokenIdx++], TEXT("Z="), PaintedVertices[Idx].Position.Z);
+			// Normal
+			bValidInput &= FParse::Value(*Tokens[TokenIdx++], TEXT("X="), PaintedVertices[Idx].Normal.Vector.X);
+			bValidInput &= FParse::Value(*Tokens[TokenIdx++], TEXT("Y="), PaintedVertices[Idx].Normal.Vector.Y);
+			bValidInput &= FParse::Value(*Tokens[TokenIdx++], TEXT("Z="), PaintedVertices[Idx].Normal.Vector.Z);
+			bValidInput &= FParse::Value(*Tokens[TokenIdx++], TEXT("W="), PaintedVertices[Idx].Normal.Vector.W);
+			// Color
+			bValidInput &= FParse::Value(*Tokens[TokenIdx++], TEXT("B="), PaintedVertices[Idx].Color.B);
+			bValidInput &= FParse::Value(*Tokens[TokenIdx++], TEXT("G="), PaintedVertices[Idx].Color.G);
+			bValidInput &= FParse::Value(*Tokens[TokenIdx++], TEXT("R="), PaintedVertices[Idx].Color.R);
+			bValidInput &= FParse::Value(*Tokens[TokenIdx++], TEXT("A="), PaintedVertices[Idx].Color.A);
+
+			// Verify that the info for this vertex was read correctly
+			check(bValidInput);
+		}
+
+		// Advance the text pointer past all of the data we just read
+		int32 LODDataStrLen = 0;
+		for (int32 Idx = 0; Idx < TokenIdx - 1; ++Idx)
+		{
+			LODDataStrLen += Tokens[Idx].Len() + 1;
+		}
+		*SourceText += LODDataStrLen;
+	}
+}
+
+FArchive& operator<<(FArchive& Ar,FStaticMeshComponentLODInfo& I)
+{
+	const uint8 OverrideColorsStripFlag = 1;
+	bool bStrippedOverrideColors = false;
+#if WITH_EDITORONLY_DATA
+	if( Ar.IsCooking() )
+	{
+		// Check if override color should be stripped too	
+		int32 LODIndex = 0;
+		for( ; LODIndex < I.Owner->LODData.Num() && &I != &I.Owner->LODData[ LODIndex ]; LODIndex++ )
+		{}
+		check( LODIndex < I.Owner->LODData.Num() );
+
+		bStrippedOverrideColors = !I.OverrideVertexColors || 
+			( I.Owner->StaticMesh == NULL || 
+			I.Owner->StaticMesh->RenderData == NULL ||
+			LODIndex >= I.Owner->StaticMesh->RenderData->LODResources.Num() ||
+			I.OverrideVertexColors->GetNumVertices() != I.Owner->StaticMesh->RenderData->LODResources[LODIndex].VertexBuffer.GetNumVertices() );
+	}
+#endif // WITH_EDITORONLY_DATA
+	FStripDataFlags StripFlags( Ar, bStrippedOverrideColors ? OverrideColorsStripFlag : 0 );
+
+	if( !StripFlags.IsDataStrippedForServer() )
+	{
+		Ar << I.LightMap;
+		Ar << I.ShadowMap;
+	}
+
+	if( !StripFlags.IsClassDataStripped( OverrideColorsStripFlag ) )
+	{
+		// Bulk serialization (new method)
+		uint8 bLoadVertexColorData = (I.OverrideVertexColors != NULL);
+		Ar << bLoadVertexColorData;
+
+		if(bLoadVertexColorData)
+		{
+			if(Ar.IsLoading())
+			{
+				check(!I.OverrideVertexColors);
+				I.OverrideVertexColors = new FColorVertexBuffer;
+			}
+
+			I.OverrideVertexColors->Serialize( Ar, true );
+		}
+	}
+
+	// Serialize out cached vertex information if necessary.
+	if ( !StripFlags.IsEditorDataStripped() )
+	{
+		Ar << I.PaintedVertices;
+	}
+
+	// Empty when loading and we don't care about saving it again, like e.g. a client.
+	if( Ar.IsLoading() && ( !GIsEditor && !IsRunningCommandlet() ) )
+	{
+		I.PaintedVertices.Empty();
+	}
+
+	return Ar;
+}
+
+#undef LOCTEXT_NAMESPACE
diff --git a/Engine/Source/Runtime/Engine/Private/Components/TextRenderComponent.cpp b/Engine/Source/Runtime/Engine/Private/Components/TextRenderComponent.cpp
new file mode 100644
--- /dev/null
+++ b/Engine/Source/Runtime/Engine/Private/Components/TextRenderComponent.cpp
@@ -0,0 +1,967 @@
+// Copyright 1998-2015 Epic Games, Inc. All Rights Reserved.
+
+#include "EnginePrivate.h"
+#include "Components/TextRenderComponent.h"
+#include "DynamicMeshBuilder.h"
+#include "Engine/Font.h"
+#include "Engine/TextRenderActor.h"
+#include "LocalVertexFactory.h"
+
+ATextRenderActor::ATextRenderActor(const FObjectInitializer& ObjectInitializer)
+	: Super(ObjectInitializer)
+{
+	TextRender = CreateDefaultSubobject<UTextRenderComponent>(TEXT("NewTextRenderComponent"));
+	RootComponent = TextRender;
+
+#if WITH_EDITORONLY_DATA
+	SpriteComponent = CreateEditorOnlyDefaultSubobject<UBillboardComponent>(TEXT("Sprite"));
+
+	if (!IsRunningCommandlet() && (SpriteComponent != nullptr))
+	{
+		// Structure to hold one-time initialization
+		struct FConstructorStatics
+		{
+			ConstructorHelpers::FObjectFinderOptional<UTexture2D> TextRenderTexture;
+			FConstructorStatics()
+				: TextRenderTexture(TEXT("/Engine/EditorResources/S_TextRenderActorIcon"))
+			{
+			}
+		};
+		static FConstructorStatics ConstructorStatics;
+
+		SpriteComponent->Sprite = ConstructorStatics.TextRenderTexture.Get();
+		SpriteComponent->RelativeScale3D = FVector(0.5f, 0.5f, 0.5f);
+		SpriteComponent->AttachParent = TextRender;
+		SpriteComponent->bIsScreenSizeScaled = true;
+		SpriteComponent->bAbsoluteScale = true;
+		SpriteComponent->bReceivesDecals = false;
+	}
+#endif
+}
+
+// ---------------------------------------------------------------
+// Text parsing class which understands line breaks using the '<br>' characters
+
+struct FTextIterator
+{
+	const TCHAR* SourceString;
+	const TCHAR* CurrentPosition;
+
+	FTextIterator(const TCHAR* InSourceString)
+		: SourceString(InSourceString)
+		, CurrentPosition(InSourceString)
+	{
+		check(InSourceString);
+	}
+
+	bool NextLine()
+	{
+		check(CurrentPosition);
+		return (CurrentPosition[0] != '\0');
+	}
+
+	bool NextCharacterInLine(int32& Ch)
+	{	
+		bool bRet = false;
+		check(CurrentPosition);
+
+		if (CurrentPosition[0] == '\0')
+		{
+			// Leave current position on the null terminator
+		}
+		else if (CurrentPosition[0] == '<' && CurrentPosition[1] == 'b' && CurrentPosition[2] == 'r' && CurrentPosition[3] == '>')
+		{
+			CurrentPosition += 4;
+		}
+		else if (CurrentPosition[0] == '\n')
+		{
+			++CurrentPosition;
+		}
+		else
+		{
+			Ch = *CurrentPosition;
+			CurrentPosition++;
+			bRet = true;
+		}
+
+		return bRet;
+	}
+
+	bool Peek(int32& Ch)
+	{
+		check(CurrentPosition);
+		if ( CurrentPosition[0] =='\0' || (CurrentPosition[0] == '<' && CurrentPosition[1] == 'b' && CurrentPosition[2] == 'r' && CurrentPosition[3] == '>') || (CurrentPosition[0] == '\n') )
+		{
+			return false;
+		}
+		else
+		{
+			Ch = CurrentPosition[0];
+			return true;
+		}
+	}
+};
+
+// ---------------------------------------------------------------
+
+/** Vertex Buffer */
+class FTextRenderVertexBuffer : public FVertexBuffer 
+{
+public:
+	TArray<FDynamicMeshVertex> Vertices;
+
+	virtual void InitRHI() override
+	{
+		FRHIResourceCreateInfo CreateInfo;
+		VertexBufferRHI = RHICreateVertexBuffer(Vertices.Num() * sizeof(FDynamicMeshVertex),BUF_Static,CreateInfo);
+
+		// Copy the vertex data into the vertex buffer.
+		void* VertexBufferData = RHILockVertexBuffer(VertexBufferRHI,0,Vertices.Num() * sizeof(FDynamicMeshVertex), RLM_WriteOnly);
+		FMemory::Memcpy(VertexBufferData,Vertices.GetData(),Vertices.Num() * sizeof(FDynamicMeshVertex));
+		RHIUnlockVertexBuffer(VertexBufferRHI);
+	}
+};
+
+/** Index Buffer */
+class FTextRenderIndexBuffer : public FIndexBuffer 
+{
+public:
+	TArray<uint16> Indices;
+
+	void InitRHI()
+	{
+		FRHIResourceCreateInfo CreateInfo;
+		IndexBufferRHI = RHICreateIndexBuffer(sizeof(uint16), Indices.Num() * sizeof(uint16), BUF_Static, CreateInfo);
+
+		// Copy the index data into the index buffer.
+		void* Buffer = RHILockIndexBuffer(IndexBufferRHI, 0, Indices.Num() * sizeof(uint16), RLM_WriteOnly);
+		FMemory::Memcpy(Buffer, Indices.GetData(), Indices.Num() * sizeof(uint16));
+		RHIUnlockIndexBuffer(IndexBufferRHI);
+	}
+};
+
+/** Vertex Factory */
+class FTextRenderVertexFactory : public FLocalVertexFactory
+{
+public:
+
+	FTextRenderVertexFactory()
+	{}
+
+	/** Initialization */
+	void Init(const FTextRenderVertexBuffer* VertexBuffer)
+	{
+		check(IsInRenderingThread())
+
+		// Initialize the vertex factory's stream components.
+		DataType NewData;
+		NewData.PositionComponent = STRUCTMEMBER_VERTEXSTREAMCOMPONENT(VertexBuffer,FDynamicMeshVertex,Position,VET_Float3);
+		NewData.TextureCoordinates.Add(
+			FVertexStreamComponent(VertexBuffer,STRUCT_OFFSET(FDynamicMeshVertex,TextureCoordinate),sizeof(FDynamicMeshVertex),VET_Float2)
+			);
+		NewData.TangentBasisComponents[0] = STRUCTMEMBER_VERTEXSTREAMCOMPONENT(VertexBuffer,FDynamicMeshVertex,TangentX,VET_PackedNormal);
+		NewData.TangentBasisComponents[1] = STRUCTMEMBER_VERTEXSTREAMCOMPONENT(VertexBuffer,FDynamicMeshVertex,TangentZ,VET_PackedNormal);
+		NewData.ColorComponent = STRUCTMEMBER_VERTEXSTREAMCOMPONENT(VertexBuffer, FDynamicMeshVertex, Color, VET_Color);
+		SetData(NewData);
+	}
+};
+
+// ---------------------------------------------------------------
+
+// @param It must be a valid initialized text iterator
+// @param Font 0 is silently ignored
+FVector2D ComputeTextSize(FTextIterator It, class UFont* Font,
+	float XScale, float YScale, float HorizSpacingAdjust)
+{
+	FVector2D Ret(0, 0);
+
+	if(!Font)
+	{
+		return Ret;
+	}
+
+	const float CharIncrement = ( (float)Font->Kerning + HorizSpacingAdjust ) * XScale;
+
+	float LineX = 0;
+
+	const int32 PageIndex = 0;
+
+	int32 Ch;
+	while (It.NextCharacterInLine(Ch))
+	{
+		Ch = (int32)Font->RemapChar(Ch);
+
+		if(!Font->Characters.IsValidIndex(Ch + PageIndex))
+		{
+			continue;
+		}
+
+		FFontCharacter& Char = Font->Characters[Ch + PageIndex];
+
+		if(!Font->Textures.IsValidIndex(Char.TextureIndex))
+		{
+			continue;
+		}
+
+		UTexture2D* Tex = Font->Textures[Char.TextureIndex];
+
+		if(Tex && Tex->Resource)
+		{
+			FIntPoint ImportedTextureSize = Tex->GetImportedSize();
+			FVector2D InvTextureSize(1.0f / (float)ImportedTextureSize.X, 1.0f / (float)ImportedTextureSize.Y);
+
+			const float X		= LineX;
+			const float Y		= Char.VerticalOffset * YScale;
+			float SizeX			= Char.USize * XScale;
+			const float SizeY	= Char.VSize * YScale;
+			const float U		= Char.StartU * InvTextureSize.X;
+			const float V		= Char.StartV * InvTextureSize.Y;
+			const float SizeU	= Char.USize * InvTextureSize.X;
+			const float SizeV	= Char.VSize * InvTextureSize.Y;
+
+			float Right = X + SizeX;
+			float Bottom = Y + SizeY;
+
+			Ret.X = FMath::Max(Ret.X, Right);
+			Ret.Y = FMath::Max(Ret.Y, Bottom);
+
+			// if we have another non-whitespace character to render, add the font's kerning.
+			int32 NextCh;
+			if( It.Peek(NextCh) && !FChar::IsWhitespace(NextCh) )
+			{
+				SizeX += CharIncrement;
+			}
+
+			LineX += SizeX;
+		}
+	}
+
+	return Ret;
+}
+
+// compute the left top depending on the alignment
+static float ComputeHorizontalAlignmentOffset(FVector2D Size, EHorizTextAligment HorizontalAlignment)
+{
+	float Ret = 0;
+
+	switch (HorizontalAlignment)
+	{
+	case EHTA_Left:
+		{
+			// X is already 0
+			break;
+		}
+
+	case EHTA_Center:
+		{
+			Ret = -Size.X * 0.5f;
+			break;
+		}
+
+	case EHTA_Right:
+		{
+			Ret = -Size.X;
+			break;
+		}
+
+	default:
+		{
+			// internal error
+			check(0);
+		}
+	}
+
+	return Ret;
+}
+
+float ComputeVerticalAlignmentOffset(float SizeY, EVerticalTextAligment VerticalAlignment, float LegacyVerticalOffset)
+{
+	switch (VerticalAlignment)
+	{
+	case EVRTA_QuadTop:
+		{
+			return LegacyVerticalOffset;
+		}
+	case EVRTA_TextBottom:
+		{
+			return -SizeY;
+		}
+	case EVRTA_TextTop:
+		{
+			return 0;
+		}
+	case EVRTA_TextCenter:
+		{
+			return -SizeY / 2.0f;
+		}
+	default:
+		{
+			check(0);
+			return 0;
+		}
+	}
+}
+
+/**
+ * For the given text info, calculate the vertical offset that needs to be applied to the component
+ * in order to vertically align it to the requested alignment.
+ **/
+float CalculateVerticalAlignmentOffset(
+	const TCHAR* Text,
+	class UFont* Font,
+	float XScale, 
+	float YScale, 
+	float HorizSpacingAdjust,
+	EVerticalTextAligment VerticalAlignment)
+{
+	check(Text);
+
+	if(!Font)
+	{
+		return 0;
+	}
+
+	float FirstLineHeight = -1; // Only kept around for legacy positioning support
+	float StartY = 0;
+
+	FTextIterator It(Text);
+
+	while (It.NextLine())
+	{
+		FVector2D LineSize = ComputeTextSize(It, Font, XScale, YScale, HorizSpacingAdjust);
+
+		if (FirstLineHeight < 0)
+		{
+			FirstLineHeight = LineSize.Y;
+		}
+
+		int32 Ch;
+
+		// Iterate to end of line
+		while (It.NextCharacterInLine(Ch))
+		{
+		}
+
+		// Move Y position down to next line. If the current line is empty, move by max char height in font
+		StartY += LineSize.Y > 0 ? LineSize.Y : Font->GetMaxCharHeight();
+	}
+
+	// Calculate a vertical translation to create the correct vertical alignment
+	return -ComputeVerticalAlignmentOffset(StartY, VerticalAlignment, FirstLineHeight);
+}
+	
+// ------------------------------------------------------
+
+/** Represents a URenderTextComponent to the scene manager. */
+class FTextRenderSceneProxy : public FPrimitiveSceneProxy
+{
+public:
+	// constructor / destructor
+	FTextRenderSceneProxy(UTextRenderComponent* Component);
+	virtual ~FTextRenderSceneProxy();
+
+	// Begin FPrimitiveSceneProxy interface
+	virtual void GetDynamicMeshElements(const TArray<const FSceneView*>& Views, const FSceneViewFamily& ViewFamily, uint32 VisibilityMap, FMeshElementCollector& Collector) const override;
+	virtual void DrawStaticElements(FStaticPrimitiveDrawInterface* PDI) override;
+	virtual FPrimitiveViewRelevance GetViewRelevance(const FSceneView* View) override;
+	virtual bool CanBeOccluded() const override;
+	virtual uint32 GetMemoryFootprint() const override;
+	uint32 GetAllocatedSize() const;
+	// End FPrimitiveSceneProxy interface
+
+private:
+	// Begin FPrimitiveSceneProxy interface
+	virtual void CreateRenderThreadResources() override;
+	// End FPrimitiveSceneProxy interface
+
+	void ReleaseRenderThreadResources();
+	bool BuildStringMesh( TArray<FDynamicMeshVertex>& OutVertices, TArray<uint16>& OutIndices );
+
+private:
+	FMaterialRelevance MaterialRelevance;
+	FTextRenderVertexBuffer VertexBuffer;
+	FTextRenderIndexBuffer IndexBuffer;
+	FTextRenderVertexFactory VertexFactory;
+	const FColor TextRenderColor;
+	UMaterialInterface* TextMaterial;
+	UFont* Font;
+	FString Text;
+	float XScale;
+	float YScale;
+	float HorizSpacingAdjust;
+	EHorizTextAligment HorizontalAlignment;
+	EVerticalTextAligment VerticalAlignment;
+	bool bAlwaysRenderAsText;
+};
+
+FTextRenderSceneProxy::FTextRenderSceneProxy( UTextRenderComponent* Component) :
+	FPrimitiveSceneProxy(Component),
+	TextRenderColor(Component->TextRenderColor),
+	Font(Component->Font),
+	Text(Component->Text),
+	XScale(Component->WorldSize * Component->XScale * Component->InvDefaultSize),
+	YScale(Component->WorldSize * Component->YScale * Component->InvDefaultSize),
+	HorizSpacingAdjust(Component->HorizSpacingAdjust),
+	HorizontalAlignment(Component->HorizontalAlignment),
+	VerticalAlignment(Component->VerticalAlignment),
+	bAlwaysRenderAsText(Component->bAlwaysRenderAsText)
+{
+	WireframeColor = FLinearColor(1, 0, 0);
+	UMaterialInterface* EffectiveMaterial = 0;
+
+	if(Component->TextMaterial)
+	{
+		UMaterial* BaseMaterial = Component->TextMaterial->GetMaterial();
+
+		if(BaseMaterial->MaterialDomain == MD_Surface)
+		{
+			EffectiveMaterial = Component->TextMaterial;
+		}
+	}
+
+	if(!EffectiveMaterial)
+	{
+		EffectiveMaterial = UMaterial::GetDefaultMaterial(MD_Surface);
+	}
+
+	TextMaterial = EffectiveMaterial;
+	MaterialRelevance |= TextMaterial->GetMaterial()->GetRelevance(GetScene().GetFeatureLevel());
+}
+
+FTextRenderSceneProxy::~FTextRenderSceneProxy()
+{
+	ReleaseRenderThreadResources();
+}
+
+void FTextRenderSceneProxy::CreateRenderThreadResources()
+{
+	if(Font && Font->FontCacheType == EFontCacheType::Runtime)
+	{
+		// Runtime fonts can't currently be used here as they use the font cache from Slate application
+		// which can only be used on the game thread
+		return;
+	}
+
+	if(BuildStringMesh(VertexBuffer.Vertices, IndexBuffer.Indices))
+	{
+		// Init vertex factory
+		VertexFactory.Init(&VertexBuffer);
+
+		// Enqueue initialization of render resources
+		VertexBuffer.InitResource();
+		IndexBuffer.InitResource();
+		VertexFactory.InitResource();
+	}
+}
+
+void FTextRenderSceneProxy::ReleaseRenderThreadResources()
+{
+	VertexBuffer.ReleaseResource();
+	IndexBuffer.ReleaseResource();
+	VertexFactory.ReleaseResource();
+}
+
+void FTextRenderSceneProxy::GetDynamicMeshElements(const TArray<const FSceneView*>& Views, const FSceneViewFamily& ViewFamily, uint32 VisibilityMap, FMeshElementCollector& Collector) const 
+{
+	QUICK_SCOPE_CYCLE_COUNTER( STAT_TextRenderSceneProxy_GetDynamicMeshElements );
+
+	// Vertex factory will not been initialized when the text string is empty or font is invalid.
+	if(VertexFactory.IsInitialized())
+	{
+		for (int32 ViewIndex = 0; ViewIndex < Views.Num(); ViewIndex++)
+		{
+			if (VisibilityMap & (1 << ViewIndex))
+			{
+				const FSceneView* View = Views[ViewIndex];
+				// Draw the mesh.
+				FMeshBatch& Mesh = Collector.AllocateMesh();
+				FMeshBatchElement& BatchElement = Mesh.Elements[0];
+				BatchElement.IndexBuffer = &IndexBuffer;
+				Mesh.VertexFactory = &VertexFactory;
+				BatchElement.PrimitiveUniformBufferResource = &GetUniformBuffer();
+				BatchElement.FirstIndex = 0;
+				BatchElement.NumPrimitives = IndexBuffer.Indices.Num() / 3;
+				BatchElement.MinVertexIndex = 0;
+				BatchElement.MaxVertexIndex = VertexBuffer.Vertices.Num() - 1;
+				Mesh.ReverseCulling = IsLocalToWorldDeterminantNegative();
+				Mesh.bDisableBackfaceCulling = false;
+				Mesh.Type = PT_TriangleList;
+				Mesh.DepthPriorityGroup = SDPG_World;
+				const bool bUseSelectedMaterial = GIsEditor && (View->Family->EngineShowFlags.Selection) ? IsSelected() : false;
+				Mesh.MaterialRenderProxy = TextMaterial->GetRenderProxy(bUseSelectedMaterial);
+				Mesh.bCanApplyViewModeOverrides = !bAlwaysRenderAsText;
+
+				Collector.AddMesh(ViewIndex, Mesh);
+
+#if !(UE_BUILD_SHIPPING || UE_BUILD_TEST)
+				RenderBounds(Collector.GetPDI(ViewIndex), View->Family->EngineShowFlags, GetBounds(), IsSelected());
+#endif
+			}
+		}
+	}
+}
+
+void FTextRenderSceneProxy::DrawStaticElements(FStaticPrimitiveDrawInterface* PDI)
+{
+	// Vertex factory will not been initialized when the font is invalid or the text string is empty.
+	if(VertexFactory.IsInitialized())
+	{
+		// Draw the mesh.
+		FMeshBatch Mesh;
+		FMeshBatchElement& BatchElement = Mesh.Elements[0];
+		BatchElement.IndexBuffer = &IndexBuffer;
+		Mesh.VertexFactory = &VertexFactory;
+		Mesh.MaterialRenderProxy = TextMaterial->GetRenderProxy(false);
+		BatchElement.PrimitiveUniformBufferResource = &GetUniformBuffer();
+		BatchElement.FirstIndex = 0;
+		BatchElement.NumPrimitives = IndexBuffer.Indices.Num() / 3;
+		BatchElement.MinVertexIndex = 0;
+		BatchElement.MaxVertexIndex = VertexBuffer.Vertices.Num() - 1;
+		Mesh.ReverseCulling = IsLocalToWorldDeterminantNegative();
+		Mesh.bDisableBackfaceCulling = false;
+		Mesh.Type = PT_TriangleList;
+		Mesh.DepthPriorityGroup = SDPG_World;
+		PDI->DrawMesh(Mesh, 1.0f);
+	}
+}
+
+bool FTextRenderSceneProxy::CanBeOccluded() const
+{
+	return !MaterialRelevance.bDisableDepthTest;
+}
+
+FPrimitiveViewRelevance FTextRenderSceneProxy::GetViewRelevance(const FSceneView* View)
+{
+	FPrimitiveViewRelevance Result;
+	Result.bDrawRelevance = IsShown(View) && View->Family->EngineShowFlags.TextRender;
+	Result.bShadowRelevance = IsShadowCast(View);
+	Result.bRenderCustomDepth = ShouldRenderCustomDepth();
+	Result.bRenderInMainPass = ShouldRenderInMainPass();
+
+	if( IsRichView(*View->Family) 
+		|| View->Family->EngineShowFlags.Bounds 
+		|| View->Family->EngineShowFlags.Collision 
+		|| IsSelected() 
+		|| IsHovered()
+		)
+	{
+		Result.bDynamicRelevance = true;
+	}
+	else
+	{
+		Result.bStaticRelevance = true;
+	}
+
+	MaterialRelevance.SetPrimitiveViewRelevance(Result);
+	return Result;
+}
+
+uint32 FTextRenderSceneProxy::GetMemoryFootprint() const
+{
+	return( sizeof( *this ) + GetAllocatedSize() );
+}
+
+uint32 FTextRenderSceneProxy::GetAllocatedSize() const
+{
+	return( FPrimitiveSceneProxy::GetAllocatedSize() ); 
+}
+
+/**
+* For the given text, constructs a mesh to be used by the vertex factory for rendering.
+*/
+bool  FTextRenderSceneProxy::BuildStringMesh( TArray<FDynamicMeshVertex>& OutVertices, TArray<uint16>& OutIndices )
+{
+	if(!Font || Text.IsEmpty())
+	{
+		return false;
+	}
+
+	float FirstLineHeight = -1; // Only kept around for legacy positioning support
+	float StartY = 0;
+
+	FLinearColor ActualColor = TextRenderColor;
+
+	const float CharIncrement = ( (float)Font->Kerning + HorizSpacingAdjust ) * XScale;
+
+	float LineX = 0;
+
+	const int32 PageIndex = 0;
+
+	FTextIterator It(*Text);
+	while (It.NextLine())
+	{
+		FVector2D LineSize = ComputeTextSize(It, Font, XScale, YScale, HorizSpacingAdjust);
+		float StartX = ComputeHorizontalAlignmentOffset(LineSize, HorizontalAlignment);
+
+		if (FirstLineHeight < 0)
+		{
+			FirstLineHeight = LineSize.Y;
+		}
+
+		LineX = 0;
+		int32 Ch;
+
+		while (It.NextCharacterInLine(Ch))
+		{
+			Ch = (int32)Font->RemapChar(Ch);
+
+			if(!Font->Characters.IsValidIndex(Ch + PageIndex))
+			{
+				continue;
+			}
+
+			FFontCharacter& Char = Font->Characters[Ch + PageIndex];
+
+			if(!Font->Textures.IsValidIndex(Char.TextureIndex))
+			{
+				continue;
+			}
+
+			UTexture2D* Tex = Font->Textures[Char.TextureIndex];
+
+			if(Tex && Tex->Resource)
+			{
+				FIntPoint ImportedTextureSize = Tex->GetImportedSize();
+				FVector2D InvTextureSize(1.0f / (float)ImportedTextureSize.X, 1.0f / (float)ImportedTextureSize.Y);
+
+				const float X      = LineX + StartX;
+				const float Y      = StartY + Char.VerticalOffset * YScale;
+				float SizeX = Char.USize * XScale;
+				const float SizeY = Char.VSize * YScale;
+				const float U		= Char.StartU * InvTextureSize.X;
+				const float V		= Char.StartV * InvTextureSize.Y;
+				const float SizeU	= Char.USize * InvTextureSize.X;
+				const float SizeV	= Char.VSize * InvTextureSize.Y;			
+
+				float Left = X;
+				float Top = Y;
+				float Right = X + SizeX;
+				float Bottom = Y + SizeY;
+
+				// axis choice and sign to get good alignment when placed on surface
+				FVector4 V0 = FVector4(0, -Left, -Top, 0);
+				FVector4 V1 = FVector4(0, -Right, -Top, 0);
+				FVector4 V2 = FVector4(0, -Left, -Bottom, 0);
+				FVector4 V3 = FVector4(0, -Right, -Bottom, 0);
+
+				FVector TangentX(0, -1, 0);
+				FVector TangentY(0, 0, -1);
+				FVector TangentZ(1, 0, 0);
+
+				int32 V00 = OutVertices.Add( FDynamicMeshVertex( V0, TangentX, TangentZ, FVector2D(U, V), ActualColor));
+				int32 V10 = OutVertices.Add( FDynamicMeshVertex( V1, TangentX, TangentZ, FVector2D(U + SizeU,	V),	ActualColor));
+				int32 V01 = OutVertices.Add( FDynamicMeshVertex( V2, TangentX, TangentZ, FVector2D(U,	V + SizeV), ActualColor));
+				int32 V11 = OutVertices.Add( FDynamicMeshVertex( V3, TangentX, TangentZ, FVector2D(U + SizeU,	V + SizeV), ActualColor));
+
+				check(V00 < 65536);
+				check(V10 < 65536);
+				check(V01 < 65536);
+				check(V11 < 65536);
+
+				OutIndices.Add(V00);
+				OutIndices.Add(V11);
+				OutIndices.Add(V10);
+
+				OutIndices.Add(V00);
+				OutIndices.Add(V01);
+				OutIndices.Add(V11);
+
+				// if we have another non-whitespace character to render, add the font's kerning.
+				int32 NextChar;
+				if( It.Peek(NextChar) && !FChar::IsWhitespace(NextChar) )
+				{
+					SizeX += CharIncrement;
+				}
+
+				LineX += SizeX;
+			}
+		}
+
+		// Move Y position down to next line. If the current line is empty, move by max char height in font
+		StartY += LineSize.Y > 0 ? LineSize.Y : Font->GetMaxCharHeight();
+	}
+
+	// Avoid initializing RHI resources when no vertices are generated.
+	return (OutVertices.Num() > 0);
+}
+
+// ------------------------------------------------------
+
+UTextRenderComponent::UTextRenderComponent(const FObjectInitializer& ObjectInitializer)
+	: Super(ObjectInitializer)
+{
+	// Structure to hold one-time initialization
+	struct FConstructorStatics
+	{
+		ConstructorHelpers::FObjectFinderOptional<UFont> Font;
+		ConstructorHelpers::FObjectFinderOptional<UMaterial> TextMaterial;
+		FConstructorStatics()
+			: Font(TEXT("/Engine/EngineFonts/RobotoDistanceField"))
+			, TextMaterial(TEXT("/Engine/EngineMaterials/DefaultTextMaterialOpaque"))
+		{
+		}
+	};
+	static FConstructorStatics ConstructorStatics;
+
+	Text = "Text";
+	Font = ConstructorStatics.Font.Get();
+	TextMaterial = ConstructorStatics.TextMaterial.Get();
+
+	SetCollisionProfileName(UCollisionProfile::NoCollision_ProfileName);
+	TextRenderColor = FColor::White;
+	XScale = 1;
+	YScale = 1;
+	HorizSpacingAdjust = 0;
+	HorizontalAlignment = EHTA_Left;
+	VerticalAlignment = EVRTA_TextBottom;
+
+	bGenerateOverlapEvents = false;
+
+	if( Font )
+	{
+		Font->ConditionalPostLoad();
+		WorldSize = Font->GetMaxCharHeight();
+		InvDefaultSize = 1.0f / WorldSize;
+	}
+	else
+	{
+		WorldSize = 30.0f;
+		InvDefaultSize = 1.0f / 30.0f;
+	}
+}
+
+FPrimitiveSceneProxy* UTextRenderComponent::CreateSceneProxy()
+{
+	return new FTextRenderSceneProxy(this);
+}
+
+void UTextRenderComponent::GetUsedMaterials( TArray<UMaterialInterface*>& OutMaterials ) const
+{
+	OutMaterials.Add(TextMaterial);
+}
+
+int32 UTextRenderComponent::GetNumMaterials() const
+{
+	return 1;
+}
+
+void UTextRenderComponent::SetMaterial(int32 ElementIndex, UMaterialInterface* InMaterial)
+{
+	if (ElementIndex == 0)
+	{
+		SetTextMaterial(InMaterial);
+	}
+}
+
+UMaterialInterface* UTextRenderComponent::GetMaterial(int32 ElementIndex) const
+{
+	return (ElementIndex == 0) ? TextMaterial : NULL;
+}
+
+bool UTextRenderComponent::ShouldRecreateProxyOnUpdateTransform() const
+{
+	return true;
+}
+
+FBoxSphereBounds UTextRenderComponent::CalcBounds(const FTransform& LocalToWorld) const
+{
+	if(!Text.IsEmpty() && Font)
+	{
+		FVector2D Size(FLT_MIN, 0);
+		FVector2D LeftTop(FLT_MAX, FLT_MAX);
+		float FirstLineHeight = -1;
+
+		FTextIterator It(*Text);
+
+		float AdjustedXScale = WorldSize * XScale * InvDefaultSize;
+		float AdjustedYScale = WorldSize * YScale * InvDefaultSize;
+
+		while (It.NextLine())
+		{
+			FVector2D LineSize = ComputeTextSize(It, Font, AdjustedXScale, AdjustedYScale, HorizSpacingAdjust);
+			float LineLeft = ComputeHorizontalAlignmentOffset(LineSize, HorizontalAlignment);
+
+			Size.X = FMath::Max(LineSize.X, Size.X);
+			Size.Y += LineSize.Y > 0 ? LineSize.Y : Font->GetMaxCharHeight();
+			LeftTop.X = FMath::Min(LeftTop.X, LineLeft);
+
+			if (FirstLineHeight < 0)
+			{
+				FirstLineHeight = LineSize.Y;
+			}
+
+			int32 Ch;
+			while (It.NextCharacterInLine(Ch));
+		}
+
+		LeftTop.Y = ComputeVerticalAlignmentOffset(Size.Y, VerticalAlignment, FirstLineHeight);
+		FBox LocalBox(FVector(0, -LeftTop.X, -LeftTop.Y), FVector(0, -(LeftTop.X + Size.X), -(LeftTop.Y + Size.Y)));
+
+		FBoxSphereBounds Ret(LocalBox.TransformBy(LocalToWorld));
+
+		Ret.BoxExtent *= BoundsScale;
+		Ret.SphereRadius *= BoundsScale;
+
+		return Ret;
+	}
+	else
+	{
+		return FBoxSphereBounds(ForceInit);
+	}
+}
+
+FMatrix UTextRenderComponent::GetRenderMatrix() const
+{
+	// Adjust LocalToWorld transform to account for vertical text alignment when rendering.
+	if(!Text.IsEmpty() && Font)
+	{
+		float SizeY = 0;
+		float FirstLineHeight = -1;
+		float AdjustedXScale = WorldSize * XScale * InvDefaultSize;
+		float AdjustedYScale = WorldSize * YScale * InvDefaultSize;
+
+		FTextIterator It(*Text);
+		while (It.NextLine())
+		{
+			FVector2D LineSize = ComputeTextSize(It, Font, AdjustedXScale, AdjustedYScale, HorizSpacingAdjust);
+			SizeY += LineSize.Y > 0 ? LineSize.Y : Font->GetMaxCharHeight();
+
+			if (FirstLineHeight < 0)
+			{
+				FirstLineHeight = LineSize.Y;
+			}
+
+			int32 Ch;
+			while (It.NextCharacterInLine(Ch));
+		}
+
+		// Calculate a vertical translation to create the correct vertical alignment
+		FMatrix VerticalTransform = FMatrix::Identity;
+		float VerticalAlignmentOffset = -ComputeVerticalAlignmentOffset(SizeY, VerticalAlignment, FirstLineHeight);
+		VerticalTransform = VerticalTransform.ConcatTranslation(FVector(0, 0, VerticalAlignmentOffset));
+
+		return VerticalTransform *  ComponentToWorld.ToMatrixWithScale();
+
+	}
+	return UPrimitiveComponent::GetRenderMatrix();
+}
+
+void UTextRenderComponent::SetText(const FString& Value)
+{
+	Text = Value;
+	MarkRenderStateDirty();	
+}
+
+void UTextRenderComponent::SetTextMaterial(class UMaterialInterface* Value)
+{
+	TextMaterial = Value;
+	MarkRenderStateDirty();	
+}
+
+void UTextRenderComponent::SetFont(UFont* Value)
+{
+	Font = Value;
+	if( Font )
+	{
+		InvDefaultSize = 1.0f / Font->GetMaxCharHeight();
+	}
+	else
+	{
+		InvDefaultSize = 1.0f;
+	}
+
+	MarkRenderStateDirty();	
+}
+
+void UTextRenderComponent::SetHorizontalAlignment(EHorizTextAligment Value)
+{
+	HorizontalAlignment = Value;
+	MarkRenderStateDirty();	
+}
+
+void UTextRenderComponent::SetTextRenderColor(FColor Value)
+{
+	TextRenderColor = Value;
+	MarkRenderStateDirty();	
+}
+
+void UTextRenderComponent::SetXScale(float Value)
+{
+	XScale = Value;
+	MarkRenderStateDirty();	
+}
+
+void UTextRenderComponent::SetYScale(float Value)
+{
+	YScale = Value;
+	MarkRenderStateDirty();	
+}
+
+void UTextRenderComponent::SetHorizSpacingAdjust(float Value)
+{
+	HorizSpacingAdjust = Value;
+	MarkRenderStateDirty();	
+}
+
+void UTextRenderComponent::SetWorldSize(float Value)
+{
+	WorldSize = Value;
+	MarkRenderStateDirty();	
+}
+
+FVector UTextRenderComponent::GetTextLocalSize() const
+{
+	FBoxSphereBounds Bounds = CalcBounds(FTransform::Identity);
+	return Bounds.GetBox().GetSize();
+}
+
+FVector UTextRenderComponent::GetTextWorldSize() const
+{
+	FBoxSphereBounds Bounds = CalcBounds(ComponentToWorld);
+	return Bounds.GetBox().GetSize();
+}
+
+void UTextRenderComponent::PostLoad()
+{
+	// Try and fix up assets created before the vertical alignment fix was implemented. Because we didn't flag that
+	// fix with its own version, use the version number closest to that CL
+	if (GetLinkerUE4Version() < VER_UE4_PACKAGE_REQUIRES_LOCALIZATION_GATHER_FLAGGING)
+	{
+		float Offset = CalculateVerticalAlignmentOffset(*Text, Font, XScale, YScale, HorizSpacingAdjust, VerticalAlignment);
+		FTransform RelativeTransform = GetRelativeTransform();
+		FTransform CorrectionLeft = FTransform::Identity;
+		FTransform CorrectionRight = FTransform::Identity;
+		CorrectionLeft.SetTranslation(FVector(0.0f, 0.0f, -Offset));
+		CorrectionRight.SetTranslation(FVector(0.0f, 0.0f, Offset));
+		SetRelativeTransform(CorrectionLeft * RelativeTransform * CorrectionRight);
+	}
+
+	if (GetLinkerUE4Version() < VER_UE4_ADD_TEXT_COMPONENT_VERTICAL_ALIGNMENT)
+	{
+		VerticalAlignment = EVRTA_QuadTop;
+	}
+
+	if( GetLinkerUE4Version() < VER_UE4_TEXT_RENDER_COMPONENTS_WORLD_SPACE_SIZING )
+	{
+		if( Font )
+		{
+			WorldSize = Font->GetMaxCharHeight();
+			InvDefaultSize = 1.0f / WorldSize;
+		}
+		else
+		{
+			//Just guess I suppose? If there is no font then there's no text to break so it's ok.
+			WorldSize = 30.0f;
+			InvDefaultSize = 1.0f / 30.0f;
+		}
+	}
+
+	Super::PostLoad();
+}
+
+/** Returns TextRender subobject **/
+UTextRenderComponent* ATextRenderActor::GetTextRender() const { return TextRender; }
+#if WITH_EDITORONLY_DATA
+/** Returns SpriteComponent subobject **/
+UBillboardComponent* ATextRenderActor::GetSpriteComponent() const { return SpriteComponent; }
+#endif
