Fix Mac compile issue

diff --git a/Engine/Source/Runtime/Engine/Classes/Engine/MemberReference.h b/Engine/Source/Runtime/Engine/Classes/Engine/MemberReference.h
--- a/Engine/Source/Runtime/Engine/Classes/Engine/MemberReference.h
+++ b/Engine/Source/Runtime/Engine/Classes/Engine/MemberReference.h
@@ -2,6 +2,7 @@
 
 #pragma once
 
+#include "EngineLogs.h"
 #include "Engine/Blueprint.h"
 #include "MemberReference.generated.h"
 
