Don't replicate packet/channel sequence during replays. Instead, assume packets are in-order. This saves bandwidth + simplifies checkpoints

diff --git a/Engine/Source/Runtime/Engine/Private/DataChannel.cpp b/Engine/Source/Runtime/Engine/Private/DataChannel.cpp
--- a/Engine/Source/Runtime/Engine/Private/DataChannel.cpp
+++ b/Engine/Source/Runtime/Engine/Private/DataChannel.cpp
@@ -120,7 +120,7 @@ bool UChannel::CleanUp( const bool bForDestroy )
 	}
 
 	// remember sequence number of first non-acked outgoing reliable bunch for this slot
-	if (OutRec != NULL)
+	if (OutRec != NULL && !Connection->InternalAck)
 	{
 		Connection->PendingOutRec[ChIndex] = OutRec->ChSequence;
 		//UE_LOG(LogNetTraffic, Log, TEXT("%i save pending out bunch %i"),ChIndex,Connection->PendingOutRec[ChIndex]);
@@ -300,6 +300,8 @@ void UChannel::ReceivedRawBunch( FInBunch & Bunch, bool & bOutSkipAck )
 
 	if ( Bunch.bReliable && Bunch.ChSequence != Connection->InReliable[ChIndex] + 1 )
 	{
+		// We shouldn't hit this path on 100% reliable connections
+		check( !Connection->InternalAck );
 		// If this bunch has a dependency on a previous unreceived bunch, buffer it.
 		checkSlow(!Bunch.bOpen);
 
@@ -355,6 +357,9 @@ void UChannel::ReceivedRawBunch( FInBunch & Bunch, bool & bOutSkipAck )
 		// Dispatch any waiting bunches.
 		while( InRec )
 		{
+			// We shouldn't hit this path on 100% reliable connections
+			check( !Connection->InternalAck );
+
 			if( InRec->ChSequence!=Connection->InReliable[ChIndex]+1 )
 				break;
 			UE_LOG(LogNetTraffic, Log, TEXT("      Channel %d Unleashing queued bunch"), ChIndex );
@@ -510,6 +515,9 @@ bool UChannel::ReceivedNextBunch( FInBunch & Bunch, bool & bOutSkipAck )
 			else
 			{
 				// Merge problem - delete InPartialBunch. This is mainly so that in the unlikely chance that ChSequence wraps around, we wont merge two completely separate partial bunches.
+
+				// We shouldn't hit this path on 100% reliable connections
+				check( !Connection->InternalAck );
 				
 				bOutSkipAck = true;	// Don't ack the packet, since we didn't process the bunch
 
@@ -536,7 +544,7 @@ bool UChannel::ReceivedNextBunch( FInBunch & Bunch, bool & bOutSkipAck )
 		// Fairly large number, and probably a bad idea to even have a bunch this size, but want to be safe for now and not throw out legitimate data
 		static const int32 MAX_CONSTRUCTED_PARTIAL_SIZE_IN_BYTES = 1024 * 64;		
 
-		if ( InPartialBunch != NULL && InPartialBunch->GetNumBytes() > MAX_CONSTRUCTED_PARTIAL_SIZE_IN_BYTES )
+		if ( !Connection->InternalAck && InPartialBunch != NULL && InPartialBunch->GetNumBytes() > MAX_CONSTRUCTED_PARTIAL_SIZE_IN_BYTES )
 		{
 			UE_LOG( LogNetPartialBunch, Error, TEXT( "Final partial bunch too large" ) );
 			Bunch.SetError();
diff --git a/Engine/Source/Runtime/Engine/Private/NetConnection.cpp b/Engine/Source/Runtime/Engine/Private/NetConnection.cpp
--- a/Engine/Source/Runtime/Engine/Private/NetConnection.cpp
+++ b/Engine/Source/Runtime/Engine/Private/NetConnection.cpp
@@ -647,7 +647,7 @@ void UNetConnection::ReceivedPacket( FBitReader& Reader )
 	LastReceiveTime = Driver->Time;
 
 	// Check packet ordering.
-	const int32 PacketId = MakeRelative(Reader.ReadInt(MAX_PACKETID),InPacketId,MAX_PACKETID);
+	const int32 PacketId = InternalAck ? InPacketId + 1 : MakeRelative(Reader.ReadInt(MAX_PACKETID),InPacketId,MAX_PACKETID);
 	if( PacketId > InPacketId )
 	{
 		const int32 PacketsLost = PacketId - InPacketId - 1;
@@ -821,8 +821,16 @@ void UNetConnection::ReceivedPacket( FBitReader& Reader )
 
 			if ( Bunch.bReliable )
 			{
-				// If this is a reliable bunch, use the last processed reliable sequence to read the new reliable sequence
-				Bunch.ChSequence = MakeRelative( Reader.ReadInt( MAX_CHSEQUENCE ), InReliable[Bunch.ChIndex], MAX_CHSEQUENCE );
+				if ( InternalAck )
+				{
+					// We can derive the sequence for 100% reliable connections
+					Bunch.ChSequence = InReliable[Bunch.ChIndex] + 1;
+				}
+				else
+				{
+					// If this is a reliable bunch, use the last processed reliable sequence to read the new reliable sequence
+					Bunch.ChSequence = MakeRelative( Reader.ReadInt( MAX_CHSEQUENCE ), InReliable[Bunch.ChIndex], MAX_CHSEQUENCE );
+				}
 			} 
 			else if ( Bunch.bPartial )
 			{
@@ -1037,7 +1045,7 @@ int32 UNetConnection::WriteBitsToSendBuffer(
 	LastStart = FBitWriterMark( SendBuffer );
 
 	// If this is the start of the queue, make sure to add the packet id
-	if ( SendBuffer.GetNumBits() == 0 )
+	if ( SendBuffer.GetNumBits() == 0 && !InternalAck )
 	{
 		SendBuffer.WriteIntWrapped( OutPacketId, MAX_PACKETID );
 		ValidateSendBuffer();
@@ -1089,7 +1097,7 @@ int64 UNetConnection::GetFreeSendBufferBits()
 	// Otherwise, we only need to account for trailer size
 	const int32 ExtraBits = ( SendBuffer.GetNumBits() > 0 ) ? MAX_PACKET_TRAILER_BITS : MAX_PACKET_HEADER_BITS + MAX_PACKET_TRAILER_BITS;
 
-	const int32 NumberOfFreeBits = ( MaxPacket * 8 ) - ( SendBuffer.GetNumBits() + ExtraBits );
+	const int32 NumberOfFreeBits = SendBuffer.GetMaxBits() - ( SendBuffer.GetNumBits() + ExtraBits );
 
 	check( NumberOfFreeBits >= 0 );
 
@@ -1181,7 +1189,7 @@ int32 UNetConnection::SendRawBunch( FOutBunch& Bunch, bool InAllowMerge )
 	Header.WriteBit( Bunch.bHasMustBeMappedGUIDs );
 	Header.WriteBit( Bunch.bPartial );
 
-	if ( Bunch.bReliable )
+	if ( Bunch.bReliable && !InternalAck )
 	{
 		Header.WriteIntWrapped(Bunch.ChSequence, MAX_CHSEQUENCE);
 	}
diff --git a/Engine/Source/Runtime/Engine/Private/UnrealNetwork.cpp b/Engine/Source/Runtime/Engine/Private/UnrealNetwork.cpp
--- a/Engine/Source/Runtime/Engine/Private/UnrealNetwork.cpp
+++ b/Engine/Source/Runtime/Engine/Private/UnrealNetwork.cpp
@@ -6,7 +6,14 @@
 
 FNetworkVersion::FGetLocalNetworkVersionOverride FNetworkVersion::GetLocalNetworkVersionOverride;
 FNetworkVersion::FIsNetworkCompatibleOverride FNetworkVersion::IsNetworkCompatibleOverride;
-const uint32 FNetworkVersion::InternalProtocolVersion = 2;
+
+enum ENetworkVersionHistory
+{
+	HISTORY_INITIAL				= 1,
+	HISTORY_INTERNAL_ACK		= 3				// We no longer save packet/channel sequence in stream. We can derive this for 100% reliable connections.
+};
+
+const uint32 FNetworkVersion::InternalProtocolVersion = HISTORY_INTERNAL_ACK;
 
 uint32 FNetworkVersion::GetLocalNetworkVersion()
 {
