Checking for null on World in UGameViewportClient::Init

diff --git a/Engine/Source/Runtime/Engine/Private/GameViewportClient.cpp b/Engine/Source/Runtime/Engine/Private/GameViewportClient.cpp
--- a/Engine/Source/Runtime/Engine/Private/GameViewportClient.cpp
+++ b/Engine/Source/Runtime/Engine/Private/GameViewportClient.cpp
@@ -250,14 +250,17 @@ void UGameViewportClient::Init(struct FWorldContext& WorldContext, UGameInstance
 			if (AudioDevice)
 			{
 				// Set the base mix of the new device based on the world settings of the world
-				AudioDevice->SetDefaultBaseSoundMix(World->GetWorldSettings()->DefaultBaseSoundMix);
+				if (World)
+				{
+					AudioDevice->SetDefaultBaseSoundMix(World->GetWorldSettings()->DefaultBaseSoundMix);
+
+					// Set the world's audio device handle to use so that sounds which play in that world will use the correct audio device
+					World->SetAudioDeviceHandle(AudioDeviceHandle);
+				}
 
 				// Set this audio device handle on the world context so future world's set onto the world context
 				// will pass the audio device handle to them and audio will play on the correct audio device
 				WorldContext.AudioDeviceHandle = AudioDeviceHandle;
-
-				// Set the world's audio device handle to use so that sounds which play in that world will use the correct audio device
-				World->SetAudioDeviceHandle(AudioDeviceHandle);
 			}
 		}
 	}
