UE4 - fix newly spawned tick group not running a second time

diff --git a/Engine/Source/Runtime/Engine/Private/TickTaskManager.cpp b/Engine/Source/Runtime/Engine/Private/TickTaskManager.cpp
--- a/Engine/Source/Runtime/Engine/Private/TickTaskManager.cpp
+++ b/Engine/Source/Runtime/Engine/Private/TickTaskManager.cpp
@@ -254,7 +254,7 @@ public:
 					}
 				}
 			}
-			WaitForTickGroup = ETickingGroup(WorldTickGroup + 1);
+			WaitForTickGroup = ETickingGroup(WorldTickGroup + (WorldTickGroup == TG_NewlySpawned ? 0 : 1)); // don't advance for newly spawned
 		}
 		else
 		{
