Fixed 'show volumes' and 'show collision' console commands so that volumes / BSPs are correctly shown.

diff --git a/Engine/Source/Runtime/Engine/Classes/Engine/GameViewportClient.h b/Engine/Source/Runtime/Engine/Classes/Engine/GameViewportClient.h
--- a/Engine/Source/Runtime/Engine/Classes/Engine/GameViewportClient.h
+++ b/Engine/Source/Runtime/Engine/Classes/Engine/GameViewportClient.h
@@ -614,6 +614,9 @@ private:
 		return EnabledStats.Num();
 	}
 
+	/** Process the 'show volumes' console command */
+	void ToggleShowVolumes();
+
 	/** Process the 'show collision' console command */
 	void ToggleShowCollision();
 
diff --git a/Engine/Source/Runtime/Engine/Private/Components/BrushComponent.cpp b/Engine/Source/Runtime/Engine/Private/Components/BrushComponent.cpp
--- a/Engine/Source/Runtime/Engine/Private/Components/BrushComponent.cpp
+++ b/Engine/Source/Runtime/Engine/Private/Components/BrushComponent.cpp
@@ -343,7 +343,7 @@ public:
 			return Result;
 		}
 
-		const bool bInCollisionView = (View->Family->EngineShowFlags.CollisionVisibility || View->Family->EngineShowFlags.CollisionPawn);
+		const bool bInCollisionView = (View->Family->EngineShowFlags.Collision || View->Family->EngineShowFlags.CollisionVisibility || View->Family->EngineShowFlags.CollisionPawn);
 
 		if(IsShown(View))
 		{
@@ -409,7 +409,7 @@ public:
 		}
 
 		// Don't render on top in 'collision view' modes
-		if(!bInCollisionView)
+		if(!bInCollisionView && !View->bIsGameView)
 		{
 			Result.bEditorPrimitiveRelevance = true;
 		}
diff --git a/Engine/Source/Runtime/Engine/Private/GameViewportClient.cpp b/Engine/Source/Runtime/Engine/Private/GameViewportClient.cpp
--- a/Engine/Source/Runtime/Engine/Private/GameViewportClient.cpp
+++ b/Engine/Source/Runtime/Engine/Private/GameViewportClient.cpp
@@ -2188,31 +2188,7 @@ bool UGameViewportClient::HandleShowCommand( const TCHAR* Cmd, FOutputDevice& Ar
 					// TODO: Investigate why this is doesn't appear to work
 					if (AllowDebugViewmodes())
 					{
-						// Iterate over all brushes
-						for( TObjectIterator<UBrushComponent> It; It; ++It )
-						{
-							UBrushComponent* BrushComponent = *It;
-							AVolume* Owner = Cast<AVolume>( BrushComponent->GetOwner() );
-
-							// Only bother with volume brushes that belong to the world's scene
-							if( Owner && BrushComponent->GetScene() == GetWorld()->Scene )
-							{
-								// We're expecting this to be in the game at this point
-								check( Owner->GetWorld()->IsGameWorld() );
-
-								// Toggle visibility of this volume
-								if( BrushComponent->IsVisible() )
-								{
-									Owner->bHidden = true;
-									BrushComponent->SetVisibility( false );
-								}
-								else
-								{
-									Owner->bHidden = false;
-									BrushComponent->SetVisibility( true );
-								}
-							}
-						}
+						ToggleShowVolumes();
 					}
 					else
 					{
@@ -2272,6 +2248,42 @@ TOptional<EPopupMethod> UGameViewportClient::OnQueryPopupMethod() const
 	return EPopupMethod::UseCurrentWindow;
 }
 
+void UGameViewportClient::ToggleShowVolumes()
+{
+	// Don't allow 'show collision' and 'show volumes' at the same time, so turn collision off
+	if (EngineShowFlags.Volumes && EngineShowFlags.Collision)
+	{
+		EngineShowFlags.Collision = false;
+		ToggleShowCollision();
+	}
+
+	// Iterate over all brushes
+	for (TObjectIterator<UBrushComponent> It; It; ++It)
+	{
+		UBrushComponent* BrushComponent = *It;
+		AVolume* Owner = Cast<AVolume>(BrushComponent->GetOwner());
+
+		// Only bother with volume brushes that belong to the world's scene
+		if (Owner && BrushComponent->GetScene() == GetWorld()->Scene)
+		{
+			// We're expecting this to be in the game at this point
+			check(Owner->GetWorld()->IsGameWorld());
+
+			// Toggle visibility of this volume
+			if (BrushComponent->IsVisible())
+			{
+				BrushComponent->SetVisibility(false);
+				BrushComponent->SetHiddenInGame(true);
+			}
+			else
+			{
+				BrushComponent->SetVisibility(true);
+				BrushComponent->SetHiddenInGame(false);
+			}
+		}
+	}
+}
+
 void UGameViewportClient::ToggleShowCollision()
 {
 	// special case: for the Engine.Collision flag, we need to un-hide any primitive components that collide so their collision geometry gets rendered
@@ -2279,6 +2291,13 @@ void UGameViewportClient::ToggleShowCollision()
 
 	if (bIsShowingCollision)
 	{
+		// Don't allow 'show collision' and 'show volumes' at the same time, so turn collision off
+		if (EngineShowFlags.Volumes)
+		{
+			EngineShowFlags.Volumes = false;
+			ToggleShowVolumes();
+		}
+
 		NumViewportsShowingCollision++;
 		ShowCollisionOnSpawnedActorsDelegateHandle = GetWorld()->AddOnActorSpawnedHandler(FOnActorSpawned::FDelegate::CreateUObject(this, &UGameViewportClient::ShowCollisionOnSpawnedActors));
 	}
