Making a new method which shuts down all audio devices in UEngine::ShutdownAudioDeviceManager

diff --git a/Engine/Source/Runtime/Engine/Private/AudioDeviceManager.cpp b/Engine/Source/Runtime/Engine/Private/AudioDeviceManager.cpp
--- a/Engine/Source/Runtime/Engine/Private/AudioDeviceManager.cpp
+++ b/Engine/Source/Runtime/Engine/Private/AudioDeviceManager.cpp
@@ -170,6 +170,18 @@ bool FAudioDeviceManager::ShutdownAudioDevice(uint32 Handle)
 	return true;
 }
 
+bool FAudioDeviceManager::ShutdownAllAudioDevices()
+{
+	for (FAudioDevice* AudioDevice : Devices)
+	{
+		if (AudioDevice)
+		{
+			ShutdownAudioDevice(AudioDevice->DeviceHandle);
+		}
+	}
+	return true;
+}
+
 class FAudioDevice* FAudioDeviceManager::GetAudioDevice(uint32 Handle)
 {
 	if (!IsValidAudioDeviceHandle(Handle))
diff --git a/Engine/Source/Runtime/Engine/Private/UnrealEngine.cpp b/Engine/Source/Runtime/Engine/Private/UnrealEngine.cpp
--- a/Engine/Source/Runtime/Engine/Private/UnrealEngine.cpp
+++ b/Engine/Source/Runtime/Engine/Private/UnrealEngine.cpp
@@ -947,8 +947,7 @@ void UEngine::ShutdownAudioDeviceManager()
 	// Shutdown the main audio device in the UEEngine
 	if (AudioDeviceManager)
 	{
-		AudioDeviceManager->ShutdownAudioDevice(MainAudioDeviceHandle);
-		check(AudioDeviceManager->GetNumActiveAudioDevices() == 0);
+		AudioDeviceManager->ShutdownAllAudioDevices();
 		delete AudioDeviceManager;
 		AudioDeviceManager = NULL;
 	}
diff --git a/Engine/Source/Runtime/Engine/Public/AudioDeviceManager.h b/Engine/Source/Runtime/Engine/Public/AudioDeviceManager.h
--- a/Engine/Source/Runtime/Engine/Public/AudioDeviceManager.h
+++ b/Engine/Source/Runtime/Engine/Public/AudioDeviceManager.h
@@ -43,6 +43,11 @@ public:
 	*/
 	bool ShutdownAudioDevice(uint32 Handle);
 
+	/**
+	* Shuts down all active audio devices
+	*/
+	bool ShutdownAllAudioDevices();
+
 	/**
 	* Returns a ptr to the audio device associated with the handle. If the
 	* handle is invalid then a NULL device ptr will be returned.
