#ue4 - Delete Pawn methods IsWalking(), IsFalling(), IsCrouched(). Properly mark LaunchPawn() as deprecated in C++.

- The deleted methods have been private and BP-deprecated since 4.0 release, and did nothing. This allows subclasses to implement functions with these names and expose them to blueprints if they wish.
- LaunchPawn() BP function was already tagged with deprecation metadata.

diff --git a/Engine/Source/Runtime/Engine/Classes/GameFramework/Pawn.h b/Engine/Source/Runtime/Engine/Classes/GameFramework/Pawn.h
--- a/Engine/Source/Runtime/Engine/Classes/GameFramework/Pawn.h
+++ b/Engine/Source/Runtime/Engine/Classes/GameFramework/Pawn.h
@@ -158,20 +158,6 @@ public:
 	/** Return Physics Volume for this Pawn **/
 	virtual APhysicsVolume* GetPawnPhysicsVolume() const;
 
-private:
-	/** (DEPRECATED) @RETURN true if Pawn is currently walking (moving along the ground) */
-	UFUNCTION(BlueprintCallable, Category="Character|Movement", meta=(DeprecatedFunction, DeprecationMessage="Query the movement component instead using IsMovingOnGround()."))
-	virtual bool IsWalking() const;
-
-	/** (DEPRECATED) @RETURN true if Pawn is currently falling */
-	UFUNCTION(BlueprintCallable, Category="Character|Movement", meta=(DeprecatedFunction, DeprecationMessage="Query the movement component instead."))
-	virtual bool IsFalling() const;
-
-	/** (DEPRECATED) @RETURN true is Pawn is currently crouched */
-	UFUNCTION(BlueprintCallable, Category="Character|Movement", meta=(DeprecatedFunction, DeprecationMessage="Query the movement component instead."))
-	virtual bool IsCrouched() const;
-
-public:
 	/** Gets the owning actor of the Movement Base Component on which the pawn is standing. */
 	UFUNCTION(BlueprintPure, Category="Pawn")
 	static AActor* GetMovementBaseActor(const APawn* Pawn);
@@ -443,6 +429,7 @@ public:
 	// DEPRECATED FUNCTIONS
 
 	/** (Deprecated) Launch Character with LaunchVelocity  */
+	DEPRECATED(4.8, "LaunchPawn is deprecated. For Characters, use LaunchCharacter() instead.")
 	UFUNCTION(BlueprintCallable, Category="Pawn", meta=(DeprecatedFunction, DeprecationMessage="Use Character.LaunchCharacter instead"))
 	void LaunchPawn(FVector LaunchVelocity, bool bXYOverride, bool bZOverride);
 
diff --git a/Engine/Source/Runtime/Engine/Private/Pawn.cpp b/Engine/Source/Runtime/Engine/Private/Pawn.cpp
--- a/Engine/Source/Runtime/Engine/Private/Pawn.cpp
+++ b/Engine/Source/Runtime/Engine/Private/Pawn.cpp
@@ -948,24 +948,6 @@ void APawn::LaunchPawn(FVector LaunchVelocity, bool bXYOverride, bool bZOverride
 	}
 }
 
-bool APawn::IsWalking() const
-{
-	// @todo: Deprecated, remove
-	return GetMovementComponent() ? GetMovementComponent()->IsMovingOnGround() : false;
-}
-
-bool APawn::IsFalling() const
-{
-	// @todo: Deprecated, remove
-	return GetMovementComponent() ? GetMovementComponent()->IsFalling() : false;
-}
-
-bool APawn::IsCrouched() const
-{
-	// @todo: Deprecated, remove
-	return GetMovementComponent() ? GetMovementComponent()->IsCrouching() : false;
-}
-
 // REPLICATION
 
 void APawn::PostNetReceiveVelocity(const FVector& NewVelocity)
