Removing FAudioDevice::Init call from inside FAudioDeviceManager::CreateAudioDevice() to outside

- This is way that the audio device used to be created and initialized.
- This change fixes assets not loading upfront due to the the main audio device not yet being set up

diff --git a/Engine/Source/Runtime/Engine/Private/AudioDeviceManager.cpp b/Engine/Source/Runtime/Engine/Private/AudioDeviceManager.cpp
--- a/Engine/Source/Runtime/Engine/Private/AudioDeviceManager.cpp
+++ b/Engine/Source/Runtime/Engine/Private/AudioDeviceManager.cpp
@@ -65,13 +65,6 @@ FAudioDevice* FAudioDeviceManager::CreateAudioDevice(uint32& HandleOut)
 		return nullptr;
 	}
 
-	if (!NewAudioDevice->Init())
-	{
-		delete NewAudioDevice;
-		HandleOut = AUDIO_DEVICE_HANDLE_INVALID;
-		return nullptr;
-	}
-
 	// Now generation a new audio device handle for the device and store the
 	// ptr to the new device in the array of audio devices.
 
diff --git a/Engine/Source/Runtime/Engine/Private/GameViewportClient.cpp b/Engine/Source/Runtime/Engine/Private/GameViewportClient.cpp
--- a/Engine/Source/Runtime/Engine/Private/GameViewportClient.cpp
+++ b/Engine/Source/Runtime/Engine/Private/GameViewportClient.cpp
@@ -246,21 +246,24 @@ void UGameViewportClient::Init(struct FWorldContext& WorldContext, UGameInstance
 		FAudioDeviceManager* AudioDeviceManager = GEngine->GetAudioDeviceManager();
 		if (AudioDeviceManager)
 		{
-			FAudioDevice* AudioDevice = AudioDeviceManager->CreateAudioDevice(AudioDeviceHandle);
-			if (AudioDevice)
+			FAudioDevice* NewAudioDevice = AudioDeviceManager->CreateAudioDevice(AudioDeviceHandle);
+			if (NewAudioDevice)
 			{
-				// Set the base mix of the new device based on the world settings of the world
-				if (World)
+				if (NewAudioDevice->Init())
 				{
-					AudioDevice->SetDefaultBaseSoundMix(World->GetWorldSettings()->DefaultBaseSoundMix);
+					// Set the base mix of the new device based on the world settings of the world
+					if (World)
+					{ 
+						NewAudioDevice->SetDefaultBaseSoundMix(World->GetWorldSettings()->DefaultBaseSoundMix);
 
-					// Set the world's audio device handle to use so that sounds which play in that world will use the correct audio device
-					World->SetAudioDeviceHandle(AudioDeviceHandle);
-				}
+						// Set the world's audio device handle to use so that sounds which play in that world will use the correct audio device
+						World->SetAudioDeviceHandle(AudioDeviceHandle);
+					}
 
-				// Set this audio device handle on the world context so future world's set onto the world context
-				// will pass the audio device handle to them and audio will play on the correct audio device
-				WorldContext.AudioDeviceHandle = AudioDeviceHandle;
+					// Set this audio device handle on the world context so future world's set onto the world context
+					// will pass the audio device handle to them and audio will play on the correct audio device
+					WorldContext.AudioDeviceHandle = AudioDeviceHandle;
+				}
 			}
 		}
 	}
diff --git a/Engine/Source/Runtime/Engine/Private/UnrealEngine.cpp b/Engine/Source/Runtime/Engine/Private/UnrealEngine.cpp
--- a/Engine/Source/Runtime/Engine/Private/UnrealEngine.cpp
+++ b/Engine/Source/Runtime/Engine/Private/UnrealEngine.cpp
@@ -1670,11 +1670,22 @@ bool UEngine::InitializeAudioDeviceManager()
 					// Create the audio device manager and register the platform module to the device manager
 					AudioDeviceManager = new FAudioDeviceManager();
 					AudioDeviceManager->RegisterAudioDeviceModule(AudioDeviceModule);
-					FAudioDevice* NewAudioDevice = AudioDeviceManager->CreateAudioDevice(MainAudioDeviceHandle);
 
-					// Set the new default device as the active device
-					AudioDeviceManager->SetActiveDevice(MainAudioDeviceHandle);
-					return NewAudioDevice != nullptr;
+					// Create a new audio device.
+					FAudioDevice* NewAudioDevice = AudioDeviceManager->CreateAudioDevice(MainAudioDeviceHandle);
+					if (NewAudioDevice)
+					{
+						// Initialize the audio device
+						if (NewAudioDevice->Init())
+						{
+							AudioDeviceManager->SetActiveDevice(MainAudioDeviceHandle);
+						}
+						else
+						{
+							// Shut it down if we failed to initialize
+							AudioDeviceManager->ShutdownAudioDevice(MainAudioDeviceHandle);
+						}
+					}
 				}
 			}
 		}
