fixed initial setup of navigation system
#ue4 UE-11739

diff --git a/Engine/Source/Runtime/Engine/Private/AI/Navigation/NavigationSystem.cpp b/Engine/Source/Runtime/Engine/Private/AI/Navigation/NavigationSystem.cpp
--- a/Engine/Source/Runtime/Engine/Private/AI/Navigation/NavigationSystem.cpp
+++ b/Engine/Source/Runtime/Engine/Private/AI/Navigation/NavigationSystem.cpp
@@ -523,6 +523,7 @@ void UNavigationSystem::OnInitializeActors()
 void UNavigationSystem::OnWorldInitDone(FNavigationSystem::EMode Mode)
 {
 	OperationMode = Mode;
+	DoInitialSetup();
 	
 	UWorld* World = GetWorld();
 
