Reduce the amount of stack memory used by FTextureLayoutNode3d::AddSurfaceInner()

diff --git a/Engine/Source/Runtime/Engine/Public/TextureLayout3d.h b/Engine/Source/Runtime/Engine/Public/TextureLayout3d.h
--- a/Engine/Source/Runtime/Engine/Public/TextureLayout3d.h
+++ b/Engine/Source/Runtime/Engine/Public/TextureLayout3d.h
@@ -221,9 +221,6 @@ private:
 	int32 AddSurfaceInner(int32 NodeIndex, uint32 ElementSizeX, uint32 ElementSizeY, uint32 ElementSizeZ, bool bAllowTextureEnlargement)
 	{
 		checkSlow(NodeIndex != INDEX_NONE);
-		// Store a copy of the current node on the stack for easier debugging.
-		// Can't store a pointer to the current node since Nodes may be reallocated in this function.
-		const FTextureLayoutNode3d CurrentNode = Nodes[NodeIndex];
 		// But do access this node via a pointer until the first recursive call. Prevents a ton of LHS.
 		const FTextureLayoutNode3d* CurrentNodePtr = &Nodes[NodeIndex];
 		if (CurrentNodePtr->ChildA != INDEX_NONE)
@@ -242,7 +239,7 @@ private:
 				return Result;
 			}
 
-			return AddSurfaceInner(CurrentNode.ChildB, ElementSizeX, ElementSizeY, ElementSizeZ, bAllowTextureEnlargement);
+			return AddSurfaceInner(Nodes[NodeIndex].ChildB, ElementSizeX, ElementSizeY, ElementSizeZ, bAllowTextureEnlargement);
 		}
 		// Node has no children, it is a leaf
 		else
@@ -285,6 +282,10 @@ private:
 			// Add new nodes, and link them as children of the current node.
 			if (ExcessWidth > ExcessHeight)
 			{
+				// Store a copy of the current node on the stack for easier debugging.
+				// Can't store a pointer to the current node since Nodes may be reallocated in this function.
+				const FTextureLayoutNode3d CurrentNode = Nodes[NodeIndex];
+
 				if (ExcessWidth > ExcessDepth)
 				{
 					// Update the child indices
@@ -338,6 +339,10 @@ private:
 			}
 			else
 			{
+				// Store a copy of the current node on the stack for easier debugging.
+				// Can't store a pointer to the current node since Nodes may be reallocated in this function.
+				const FTextureLayoutNode3d CurrentNode = Nodes[NodeIndex];
+
 				if (ExcessHeight > ExcessDepth)
 				{
 					Nodes[NodeIndex].ChildA = Nodes.Num();
