Build - Fixing it.

diff --git a/Engine/Source/Runtime/Engine/Private/Slate/SGameLayerManager.cpp b/Engine/Source/Runtime/Engine/Private/Slate/SGameLayerManager.cpp
--- a/Engine/Source/Runtime/Engine/Private/Slate/SGameLayerManager.cpp
+++ b/Engine/Source/Runtime/Engine/Private/Slate/SGameLayerManager.cpp
@@ -5,6 +5,7 @@
 #include "SGameLayerManager.h"
 
 #include "Engine/UserInterfaceSettings.h"
+#include "Slate/SceneViewport.h"
 
 #include "STooltipPresenter.h"
 #include "Widgets/Layout/SScissorRectBox.h"
