Merging 2479091

//depot/UE4-UT/...

to //depot/UE4/...

diff --git a/Engine/Source/Runtime/Engine/Private/DataChannel.cpp b/Engine/Source/Runtime/Engine/Private/DataChannel.cpp
--- a/Engine/Source/Runtime/Engine/Private/DataChannel.cpp
+++ b/Engine/Source/Runtime/Engine/Private/DataChannel.cpp
@@ -1905,9 +1905,10 @@ void UActorChannel::ProcessBunch( FInBunch & Bunch )
 
 		// Check to see if the actor was destroyed
 		// If so, don't continue processing packets on this channel, or we'll trigger an error otherwise
+		// note that this is a legitimate occurrence, particularly on client to server RPCs
 		if ( !Actor || Actor->IsPendingKill() )
 		{
-			UE_LOG( LogNet, Warning, TEXT( "UActorChannel::ProcessBunch: Actor was destroyed during Replicator.ReceivedBunch processing" ) );
+			UE_LOG( LogNet, Verbose, TEXT( "UActorChannel::ProcessBunch: Actor was destroyed during Replicator.ReceivedBunch processing" ) );
 			// If we lose the actor on this channel, we can no longer process bunches, so consider this channel broken
 			Broken = 1;		
 			break;
