#ue4 - Added tunable control of FloatingPawnMovement's turning responsiveness.

- Simplified compact statement that was computing the turning vector in order to make it easier to read.

Based on forum discussion here:
https://forums.unrealengine.com/showthread.php?62245-Question-about-floatingPawnMovement-component&p=243658&viewfull=1#post243658

diff --git a/Engine/Source/Runtime/Engine/Classes/GameFramework/FloatingPawnMovement.h b/Engine/Source/Runtime/Engine/Classes/GameFramework/FloatingPawnMovement.h
--- a/Engine/Source/Runtime/Engine/Classes/GameFramework/FloatingPawnMovement.h
+++ b/Engine/Source/Runtime/Engine/Classes/GameFramework/FloatingPawnMovement.h
@@ -38,6 +38,14 @@ class ENGINE_API UFloatingPawnMovement : public UPawnMovementComponent
 	UPROPERTY(EditAnywhere, BlueprintReadWrite, Category=FloatingPawnMovement)
 	float Deceleration;
 
+	/**
+	 * Setting affecting extra force applied when changing direction, making turns have less drift and become more responsive.
+	 * Velocity magnitude is not allowed to increase, that only happens due to normal acceleration. It may decrease with large direction changes.
+	 * Larger values apply extra force to reach the target direction more quickly, while a zero value disables any extra turn force.
+	 */
+	UPROPERTY(EditAnywhere, BlueprintReadWrite, Category=FloatingPawnMovement, meta=(ClampMin="0", UIMin="0"))
+	float TurningBoost;
+
 protected:
 
 	/** Update Velocity based on input. Also applies gravity. */
diff --git a/Engine/Source/Runtime/Engine/Private/FloatingPawnMovement.cpp b/Engine/Source/Runtime/Engine/Private/FloatingPawnMovement.cpp
--- a/Engine/Source/Runtime/Engine/Private/FloatingPawnMovement.cpp
+++ b/Engine/Source/Runtime/Engine/Private/FloatingPawnMovement.cpp
@@ -9,6 +9,7 @@ UFloatingPawnMovement::UFloatingPawnMovement(const FObjectInitializer& ObjectIni
 	MaxSpeed = 1200.f;
 	Acceleration = 4000.f;
 	Deceleration = 8000.f;
+	TurningBoost = 8.0f;
 	bPositionCorrected = false;
 
 	ResetMoveState();
@@ -106,7 +107,9 @@ void UFloatingPawnMovement::ApplyControlInputToVelocity(float DeltaTime)
 		// Apply change in velocity direction
 		if (Velocity.SizeSquared() > 0.f)
 		{
-			Velocity -= (Velocity - ControlAcceleration * Velocity.Size()) * FMath::Min(DeltaTime * 8.f, 1.f);
+			// Change direction faster than only using acceleration, but never increase velocity magnitude.
+			const float TimeScale = FMath::Clamp(DeltaTime * TurningBoost, 0.f, 1.f);
+			Velocity = Velocity + (ControlAcceleration * Velocity.Size() - Velocity) * TimeScale;
 		}
 	}
 	else
