Fixed warning display when a stereo sound plays with spatialization activated

diff --git a/Engine/Source/Runtime/Engine/Private/SoundWave.cpp b/Engine/Source/Runtime/Engine/Private/SoundWave.cpp
--- a/Engine/Source/Runtime/Engine/Private/SoundWave.cpp
+++ b/Engine/Source/Runtime/Engine/Private/SoundWave.cpp
@@ -629,7 +629,7 @@ void USoundWave::Parse( FAudioDevice* AudioDevice, const UPTRINT NodeWaveInstanc
 		ActiveSound.bFinished = false;
 
 		// Sanity check
-		if( NumChannels > 2 && WaveInstance->bUseSpatialization && !WaveInstance->bReportedSpatializationWarning)
+		if( NumChannels >= 2 && WaveInstance->bUseSpatialization && !WaveInstance->bReportedSpatializationWarning)
 		{
 			FString SoundWarningInfo = FString::Printf(TEXT("Spatialisation on stereo and multichannel sounds is not supported. SoundWave: %s"), *GetName());
 			if (ActiveSound.Sound != this)
