Remove virtual from ActorComponent BlueprintImplementableEvents that already have C++ equivalent virtuals

diff --git a/Engine/Source/Runtime/Engine/Classes/Components/ActorComponent.h b/Engine/Source/Runtime/Engine/Classes/Components/ActorComponent.h
--- a/Engine/Source/Runtime/Engine/Classes/Components/ActorComponent.h
+++ b/Engine/Source/Runtime/Engine/Classes/Components/ActorComponent.h
@@ -338,7 +338,7 @@ public:
 	 * or when the component is dynamically created if the Actor has already BegunPlay. 
 	 */
 	UFUNCTION(BlueprintImplementableEvent, meta=(FriendlyName = "Begin Play"))
-	virtual void ReceiveBeginPlay();
+	void ReceiveBeginPlay();
 
 	/**
 	 * Ends gameplay for this component.
@@ -354,7 +354,7 @@ public:
 
 	/** Blueprint implementable event for when the component ends play, generally via destruction or its Actor's EndPlay. */
 	UFUNCTION(BlueprintImplementableEvent, meta=(Keywords = "delete", FriendlyName = "End Play"))
-	virtual void ReceiveEndPlay(EEndPlayReason::Type EndPlayReason);
+	void ReceiveEndPlay(EEndPlayReason::Type EndPlayReason);
 	
 	/**
 	 * When called, will call the virtual call chain to register all of the tick functions
@@ -608,7 +608,7 @@ public:
 
 	/** Event called every frame */
 	UFUNCTION(BlueprintImplementableEvent, meta=(FriendlyName = "Tick"))
-	virtual void ReceiveTick(float DeltaSeconds);
+	void ReceiveTick(float DeltaSeconds);
 	
 	/** 
 	 *  Called by owner actor on position shifting
