On OS X the crash reporter can be fired from any thread, even a system thread that is not registered with UE4, so UConsole::Serialize can't rely on calling TConsoleVariableData<int>::GetValueOnGameThread , instead it must use TConsoleVariableData<int>::GetValueOnAnyThread.

diff --git a/Engine/Source/Runtime/Engine/Private/UserInterface/Console.cpp b/Engine/Source/Runtime/Engine/Private/UserInterface/Console.cpp
--- a/Engine/Source/Runtime/Engine/Private/UserInterface/Console.cpp
+++ b/Engine/Source/Runtime/Engine/Private/UserInterface/Console.cpp
@@ -1324,7 +1324,7 @@ void UConsole::Serialize( const TCHAR* V, ELogVerbosity::Type Verbosity, const c
 
 		if(CVar)
 		{
-			int MinVerbosity = CVar->GetValueOnGameThread();
+			int MinVerbosity = CVar->GetValueOnAnyThread();
 
 			if((int)Verbosity <= MinVerbosity)
 			{
