Implemented SetMaterial for destructible meshes correctly to handle pre and post fracture phys materials to the apex actor and the chunks.

diff --git a/Engine/Source/Runtime/Engine/Classes/Components/DestructibleComponent.h b/Engine/Source/Runtime/Engine/Classes/Components/DestructibleComponent.h
--- a/Engine/Source/Runtime/Engine/Classes/Components/DestructibleComponent.h
+++ b/Engine/Source/Runtime/Engine/Classes/Components/DestructibleComponent.h
@@ -138,6 +138,8 @@ public:
 	virtual bool LineTraceComponent( FHitResult& OutHit, const FVector Start, const FVector End, const FCollisionQueryParams& Params ) override;
 	virtual bool SweepComponent( FHitResult& OutHit, const FVector Start, const FVector End, const FCollisionShape& CollisionShape, bool bTraceComplex=false) override;
 	virtual void SetEnableGravity(bool bGravityEnabled) override;
+
+	virtual void SetMaterial(int32 ElementIndex, UMaterialInterface* Material) override;
 	// End UPrimitiveComponent interface.
 
 	// Begin SkinnedMeshComponent interface.
diff --git a/Engine/Source/Runtime/Engine/Private/Components/DestructibleComponent.cpp b/Engine/Source/Runtime/Engine/Private/Components/DestructibleComponent.cpp
--- a/Engine/Source/Runtime/Engine/Private/Components/DestructibleComponent.cpp
+++ b/Engine/Source/Runtime/Engine/Private/Components/DestructibleComponent.cpp
@@ -1364,4 +1364,36 @@ void UDestructibleComponent::SetCollisionResponseForActor(PxRigidDynamic* Actor,
 	}
 }
 
+void UDestructibleComponent::SetMaterial(int32 ElementIndex, UMaterialInterface* Material)
+{
+	// Mesh component handles render side materials
+	Super::SetMaterial(ElementIndex, Material);
+
+	// Update physical properties of the chunks in the mesh
+	GetBodyInstance()->UpdatePhysicalMaterials();
+	int32 NumBones = SkeletalMesh->RefSkeleton.GetNum();
+	for(int32 BoneIdx = 0 ; BoneIdx < NumBones ; ++BoneIdx)
+	{
+		FName BoneName = SkeletalMesh->RefSkeleton.GetBoneName(BoneIdx);
+		if(FBodyInstance* Instance = GetBodyInstance(BoneName))
+		{
+			Instance->UpdatePhysicalMaterials();
+		}
+	}
+
+	// Set new template parameters for the apex actor, so they take effect before fracturing too.
+	physx::NxPhysX3DescTemplate* Template = ApexDestructibleActor->createPhysX3DescTemplate();
+	if(ApexDestructibleActor->getPhysX3Template(*Template))
+	{
+		UPhysicalMaterial* SimpleMaterial = GetBodyInstance()->GetSimplePhysicalMaterial();
+		check(SimpleMaterial);
+		PxMaterial* PhysxMat = SimpleMaterial->GetPhysXMaterial();
+
+		Template->setMaterials(&PhysxMat, 1);
+
+		ApexDestructibleActor->setPhysX3Template(Template);
+	}
+	Template->release();
+}
+
 #endif // WITH_PHYSX
\ No newline at end of file
