Added a Min/Max boundary property to SplineMeshComponent to define the bounds of the curve numerically instead of by the mesh bounds.
#github Pull request #701
#change Fixed USplineComponent::CalcBounds which was not correctly calculating the bounding box. Now CalcBounds and CalcSliceTransform also take into account the Min/Max boundary value.

diff --git a/Engine/Source/Editor/UnrealEd/Private/Lightmass/Lightmass.cpp b/Engine/Source/Editor/UnrealEd/Private/Lightmass/Lightmass.cpp
--- a/Engine/Source/Editor/UnrealEd/Private/Lightmass/Lightmass.cpp
+++ b/Engine/Source/Editor/UnrealEd/Private/Lightmass/Lightmass.cpp
@@ -1486,8 +1486,20 @@ void FLightmassExporter::WriteMeshInstances( int32 Channel )
 			Copy(*SplineParams, SMInstanceMeshData.SplineParameters);
 			SMInstanceMeshData.SplineParameters.SplineUpDir = SplineComponent->SplineUpDir;
 			SMInstanceMeshData.SplineParameters.bSmoothInterpRollScale = SplineComponent->bSmoothInterpRollScale;
-			SMInstanceMeshData.SplineParameters.MeshMinZ = MeshBounds.Origin[SplineComponent->ForwardAxis] - MeshBounds.BoxExtent[SplineComponent->ForwardAxis];
-			SMInstanceMeshData.SplineParameters.MeshRangeZ = 2.f * MeshBounds.BoxExtent[SplineComponent->ForwardAxis];
+
+			if (FMath::IsNearlyEqual(SplineComponent->SplineBoundaryMin, SplineComponent->SplineBoundaryMax))
+			{
+				// set ranges according to the extents of the mesh
+				SMInstanceMeshData.SplineParameters.MeshMinZ = MeshBounds.Origin[SplineComponent->ForwardAxis] - MeshBounds.BoxExtent[SplineComponent->ForwardAxis];
+				SMInstanceMeshData.SplineParameters.MeshRangeZ = 2.f * MeshBounds.BoxExtent[SplineComponent->ForwardAxis];
+			}
+			else
+			{
+				// set ranges according to the custom boundary min/max
+				SMInstanceMeshData.SplineParameters.MeshMinZ = SplineComponent->SplineBoundaryMin;
+				SMInstanceMeshData.SplineParameters.MeshRangeZ = SplineComponent->SplineBoundaryMax - SplineComponent->SplineBoundaryMin;
+			}
+
 			SMInstanceMeshData.SplineParameters.ForwardAxis = (Lightmass::ESplineMeshAxis::Type)SplineComponent->ForwardAxis.GetValue();
 		}
 		else
diff --git a/Engine/Source/Runtime/Core/Private/Math/UnrealMath.cpp b/Engine/Source/Runtime/Core/Private/Math/UnrealMath.cpp
--- a/Engine/Source/Runtime/Core/Private/Math/UnrealMath.cpp
+++ b/Engine/Source/Runtime/Core/Private/Math/UnrealMath.cpp
@@ -1191,7 +1191,7 @@ static void FindBounds( float& OutMin, float& OutMax,  float Start, float StartL
 		const float c = StartLeaveTan;
 
 		const float Discriminant = (b*b) - (4.f*a*c);
-		if(Discriminant > 0.f && a > 0.f) // Solving doesn't work if a is zero, which usually indicates co-incident start and end, and zero tangents anyway
+		if(Discriminant > 0.f && !FMath::IsNearlyZero(a)) // Solving doesn't work if a is zero, which usually indicates co-incident start and end, and zero tangents anyway
 		{
 			const float SqrtDisc = FMath::Sqrt( Discriminant );
 
diff --git a/Engine/Source/Runtime/Engine/Classes/Components/SplineMeshComponent.h b/Engine/Source/Runtime/Engine/Classes/Components/SplineMeshComponent.h
--- a/Engine/Source/Runtime/Engine/Classes/Components/SplineMeshComponent.h
+++ b/Engine/Source/Runtime/Engine/Classes/Components/SplineMeshComponent.h
@@ -113,6 +113,14 @@ class ENGINE_API USplineMeshComponent : public UStaticMeshComponent, public IInt
 	UPROPERTY(EditAnywhere, Category=SplineMesh)
 	TEnumAsByte<ESplineMeshAxis::Type> ForwardAxis;
 
+	/** Minimum coordinate along the spline forward axis which corresponds to start of spline. If set to 0.0, will use bounding box to determine bounds */
+	UPROPERTY(EditAnywhere, Category = SplineMesh, AdvancedDisplay)
+	float SplineBoundaryMin;
+
+	/** Maximum coordinate along the spline forward axis which corresponds to end of spline. If set to 0.0, will use bounding box to determine bounds */
+	UPROPERTY(EditAnywhere, Category = SplineMesh, AdvancedDisplay)
+	float SplineBoundaryMax;
+	
 	// Physics data.
 	UPROPERTY()
 	class UBodySetup* BodySetup;
@@ -269,6 +277,22 @@ class ENGINE_API USplineMeshComponent : public UStaticMeshComponent, public IInt
 	UFUNCTION(BlueprintCallable, Category = SplineMesh)
 	void SetSplineUpDir(const FVector& InSplineUpDir);
 
+	/** Get the boundary min */
+	UFUNCTION(BlueprintCallable, Category = SplineMesh)
+	float GetBoundaryMin() const;
+
+	/** Set the boundary min */
+	UFUNCTION(BlueprintCallable, Category = SplineMesh)
+	void SetBoundaryMin(float InBoundaryMin);
+
+	/** Get the boundary max */
+	UFUNCTION(BlueprintCallable, Category = SplineMesh)
+	float GetBoundaryMax() const;
+
+	/** Set the boundary max */
+	UFUNCTION(BlueprintCallable, Category = SplineMesh)
+	void SetBoundaryMax(float InBoundaryMax);
+
 	// Destroys the body setup, used to clear collision if the mesh goes missing
 	void DestroyBodySetup();
 #if WITH_EDITOR
@@ -278,12 +302,20 @@ class ENGINE_API USplineMeshComponent : public UStaticMeshComponent, public IInt
 
 	/**
 	 * Calculates the spline transform, including roll, scale, and offset along the spline at a specified distance
-	 * @Note:  This is mirrored to Lightmass::CalcSliceTransform() and LocalVertexShader.usf.  If you update one of these, please update them all! 
 	 */
 	FTransform CalcSliceTransform(const float DistanceAlong) const;
 
+	/**
+	 * Calculates the spline transform, including roll, scale, and offset along the spline at a specified alpha interpolation parameter along the spline
+	 * @Note:  This is mirrored to Lightmass::CalcSliceTransform() and LocalVertexShader.usf.  If you update one of these, please update them all!
+	 */
+	FTransform CalcSliceTransformAtSplineOffset(const float Alpha) const;
+
 	inline static const float& GetAxisValue(const FVector& InVector, ESplineMeshAxis::Type InAxis);
 	inline static float& GetAxisValue(FVector& InVector, ESplineMeshAxis::Type InAxis);
+
+	/** Returns a vector which, when componentwise-multiplied by another vector, will zero all the components not corresponding to the supplied ESplineMeshAxis */
+	inline static FVector GetAxisMask(ESplineMeshAxis::Type InAxis);
 };
 
 const float& USplineMeshComponent::GetAxisValue(const FVector& InVector, ESplineMeshAxis::Type InAxis)
@@ -318,3 +350,18 @@ float& USplineMeshComponent::GetAxisValue(FVector& InVector, ESplineMeshAxis::Ty
 	}
 }
 
+FVector USplineMeshComponent::GetAxisMask(ESplineMeshAxis::Type InAxis)
+{
+	switch (InAxis)
+	{
+	case ESplineMeshAxis::X:
+		return FVector(0, 1, 1);
+	case ESplineMeshAxis::Y:
+		return FVector(1, 0, 1);
+	case ESplineMeshAxis::Z:
+		return FVector(1, 1, 0);
+	default:
+		check(0);
+		return FVector::ZeroVector;
+	}
+}
diff --git a/Engine/Source/Runtime/Engine/Private/Components/SplineMeshComponent.cpp b/Engine/Source/Runtime/Engine/Private/Components/SplineMeshComponent.cpp
--- a/Engine/Source/Runtime/Engine/Private/Components/SplineMeshComponent.cpp
+++ b/Engine/Source/Runtime/Engine/Private/Components/SplineMeshComponent.cpp
@@ -236,6 +236,9 @@ USplineMeshComponent::USplineMeshComponent(const FObjectInitializer& ObjectIniti
 	SplineParams.EndPos = FVector(100.f, 0.f, 0.f);
 	SplineParams.EndTangent = FVector(100.f, 0.f, 0.f);
 	SplineParams.EndScale = FVector2D(1.f, 1.f);
+
+	SplineBoundaryMin = 0;
+	SplineBoundaryMax = 0;
 }
 
 FVector USplineMeshComponent::GetStartPosition() const
@@ -379,6 +382,28 @@ void USplineMeshComponent::SetSplineUpDir(const FVector& InSplineUpDir)
 	MarkSplineParamsDirty();
 }
 
+float USplineMeshComponent::GetBoundaryMin() const
+{
+	return SplineBoundaryMin;
+}
+
+void USplineMeshComponent::SetBoundaryMin(float InBoundaryMin)
+{
+	SplineBoundaryMin = InBoundaryMin;
+	MarkSplineParamsDirty();
+}
+
+float USplineMeshComponent::GetBoundaryMax() const
+{
+	return SplineBoundaryMax;
+}
+
+void USplineMeshComponent::SetBoundaryMax(float InBoundaryMax)
+{
+	SplineBoundaryMax = InBoundaryMax;
+	MarkSplineParamsDirty();
+}
+
 
 void USplineMeshComponent::MarkSplineParamsDirty()
 {
@@ -448,40 +473,11 @@ FPrimitiveSceneProxy* USplineMeshComponent::CreateSceneProxy()
 	}
 }
 
-FBoxSphereBounds USplineMeshComponent::CalcBounds(const FTransform& LocalToWorld) const
-{
-	// Use util to generate bounds of spline
-	FInterpCurvePoint<FVector> Start(0.f, SplineParams.StartPos, SplineParams.StartTangent, SplineParams.StartTangent, CIM_CurveUser);
-	FInterpCurvePoint<FVector> End(1.f, SplineParams.EndPos, SplineParams.EndTangent, SplineParams.EndTangent, CIM_CurveUser);
-
-	FVector CurveMax(-BIG_NUMBER, -BIG_NUMBER, -BIG_NUMBER);
-	FVector CurveMin(BIG_NUMBER, BIG_NUMBER, BIG_NUMBER);
-	CurveVectorFindIntervalBounds(Start, End, CurveMin, CurveMax);
-
-	FBox LocalBox(CurveMin, CurveMax);
-
-	// Find largest extent of mesh in XY, and add on all around
-	if(StaticMesh)
-	{
-		FVector MinMeshExtent = StaticMesh->GetBounds().Origin - StaticMesh->GetBounds().BoxExtent;
-		FVector MaxMeshExtent = StaticMesh->GetBounds().Origin + StaticMesh->GetBounds().BoxExtent;
-		GetAxisValue(MinMeshExtent, ForwardAxis) = 0;
-		GetAxisValue(MaxMeshExtent, ForwardAxis) = 0;
-		float MaxDim = FMath::Max<float>(MinMeshExtent.GetAbsMax(), MaxMeshExtent.GetAbsMax());
-
-		float MaxScale = FMath::Max(SplineParams.StartScale.GetAbsMax(), SplineParams.EndScale.GetAbsMax());
-
-		LocalBox = LocalBox.ExpandBy(MaxScale * MaxDim);
-	}
-
-	return FBoxSphereBounds( LocalBox.TransformBy(LocalToWorld) );
-}
-
 
-/** 
- * Functions used for transforming a static mesh component based on a spline.  
- * This needs to be updated if the spline functionality changes!
- */
+/**
+* Functions used for transforming a static mesh component based on a spline.
+* This needs to be updated if the spline functionality changes!
+*/
 static float SmoothStep(float A, float B, float X)
 {
 	if (X < A)
@@ -501,13 +497,13 @@ static FVector SplineEvalPos(const FVector& StartPos, const FVector& StartTangen
 	const float A2 = A  * A;
 	const float A3 = A2 * A;
 
-	return (((2*A3)-(3*A2)+1) * StartPos) + ((A3-(2*A2)+A) * StartTangent) + ((A3-A2) * EndTangent) + (((-2*A3)+(3*A2)) * EndPos);
+	return (((2 * A3) - (3 * A2) + 1) * StartPos) + ((A3 - (2 * A2) + A) * StartTangent) + ((A3 - A2) * EndTangent) + (((-2 * A3) + (3 * A2)) * EndPos);
 }
 
 static FVector SplineEvalDir(const FVector& StartPos, const FVector& StartTangent, const FVector& EndPos, const FVector& EndTangent, float A)
 {
-	const FVector C = (6*StartPos) + (3*StartTangent) + (3*EndTangent) - (6*EndPos);
-	const FVector D = (-6*StartPos) - (4*StartTangent) - (2*EndTangent) + (6*EndPos);
+	const FVector C = (6 * StartPos) + (3 * StartTangent) + (3 * EndTangent) - (6 * EndPos);
+	const FVector D = (-6 * StartPos) - (4 * StartTangent) - (2 * EndTangent) + (6 * EndPos);
 	const FVector E = StartTangent;
 
 	const float A2 = A  * A;
@@ -516,14 +512,137 @@ static FVector SplineEvalDir(const FVector& StartPos, const FVector& StartTangen
 }
 
 
+FBoxSphereBounds USplineMeshComponent::CalcBounds(const FTransform& LocalToWorld) const
+{
+	if (!StaticMesh)
+	{
+		return FBoxSphereBounds(FBox(0));
+	}
+
+	float MinT = 0.0f;
+	float MaxT = 1.0f;
+
+	const FBoxSphereBounds MeshBounds = StaticMesh->GetBounds();
+
+	const bool bHasCustomBoundary = !FMath::IsNearlyEqual(SplineBoundaryMin, SplineBoundaryMax);
+	if (bHasCustomBoundary)
+	{
+		// If there's a custom boundary, alter the min/max of the spline we need to evaluate
+		const float MeshMin = GetAxisValue(MeshBounds.Origin - MeshBounds.BoxExtent, ForwardAxis);
+		const float MeshMax = GetAxisValue(MeshBounds.Origin + MeshBounds.BoxExtent, ForwardAxis);
+
+		const float MeshMinT = (MeshMin - SplineBoundaryMin) / (SplineBoundaryMax - SplineBoundaryMin);
+		const float MeshMaxT = (MeshMax - SplineBoundaryMin) / (SplineBoundaryMax - SplineBoundaryMin);
+
+		// Disallow extrapolation beyond a certain value; enormous bounding boxes cause the render thread to crash
+		const float MaxSplineExtrapolation = 4.0f;
+		if (FMath::Abs(MeshMinT) < MaxSplineExtrapolation && FMath::Abs(MeshMaxT) < MaxSplineExtrapolation)
+		{
+			MinT = MeshMinT;
+			MaxT = MeshMaxT;
+		}
+	}
+
+	const FVector AxisMask = GetAxisMask(ForwardAxis);
+	const FVector FlattenedMeshOrigin = MeshBounds.Origin * AxisMask;
+	const FVector FlattenedMeshExtent = MeshBounds.BoxExtent * AxisMask;
+	const FBox MeshBoundingBox = FBox(FlattenedMeshOrigin - FlattenedMeshExtent, FlattenedMeshOrigin + FlattenedMeshExtent);
+
+	FBox BoundingBox(0);
+	BoundingBox += MeshBoundingBox.TransformBy(CalcSliceTransformAtSplineOffset(MinT));
+	BoundingBox += MeshBoundingBox.TransformBy(CalcSliceTransformAtSplineOffset(MaxT));
+
+	// Work out coefficients of the cubic spline derivative equation dx/dt
+	const FVector A = 6.0f * SplineParams.StartPos + 3.0f * SplineParams.StartTangent + 3.0f * SplineParams.EndTangent - 6.0f * SplineParams.EndPos;
+	const FVector B = -6.0f * SplineParams.StartPos - 4.0f * SplineParams.StartTangent - 2.0f * SplineParams.EndTangent + 6.0f * SplineParams.EndPos;
+	const FVector C = SplineParams.StartTangent;
+
+	// Minima/maxima happen where dx/dt == 0, calculate t values
+	const FVector Discriminant = B * B - 4.0f * A * C;
+
+	// Work out minima/maxima component-by-component.
+	// Negative discriminant means no solution; A == 0 implies coincident start/end points
+	if (Discriminant.X > 0.0f && !FMath::IsNearlyZero(A.X))
+	{
+		const float SqrtDiscriminant = FMath::Sqrt(Discriminant.X);
+		const float Denominator = 0.5f / A.X;
+		const float T0 = (-B.X + SqrtDiscriminant) * Denominator;
+		const float T1 = (-B.X - SqrtDiscriminant) * Denominator;
+
+		if (T0 >= MinT && T0 <= MaxT)
+		{
+			BoundingBox += MeshBoundingBox.TransformBy(CalcSliceTransformAtSplineOffset(T0));
+		}
+
+		if (T1 >= MinT && T1 <= MaxT)
+		{
+			BoundingBox += MeshBoundingBox.TransformBy(CalcSliceTransformAtSplineOffset(T1));
+		}
+	}
+
+	if (Discriminant.Y > 0.0f && !FMath::IsNearlyZero(A.Y))
+	{
+		const float SqrtDiscriminant = FMath::Sqrt(Discriminant.Y);
+		const float Denominator = 0.5f / A.Y;
+		const float T0 = (-B.Y + SqrtDiscriminant) * Denominator;
+		const float T1 = (-B.Y - SqrtDiscriminant) * Denominator;
+
+		if (T0 >= MinT && T0 <= MaxT)
+		{
+			BoundingBox += MeshBoundingBox.TransformBy(CalcSliceTransformAtSplineOffset(T0));
+		}
+
+		if (T1 >= MinT && T1 <= MaxT)
+		{
+			BoundingBox += MeshBoundingBox.TransformBy(CalcSliceTransformAtSplineOffset(T1));
+		}
+	}
+
+	if (Discriminant.Z > 0.0f && !FMath::IsNearlyZero(A.Z))
+	{
+		const float SqrtDiscriminant = FMath::Sqrt(Discriminant.Z);
+		const float Denominator = 0.5f / A.Z;
+		const float T0 = (-B.Z + SqrtDiscriminant) * Denominator;
+		const float T1 = (-B.Z - SqrtDiscriminant) * Denominator;
+
+		if (T0 >= MinT && T0 <= MaxT)
+		{
+			BoundingBox += MeshBoundingBox.TransformBy(CalcSliceTransformAtSplineOffset(T0));
+		}
+
+		if (T1 >= MinT && T1 <= MaxT)
+		{
+			BoundingBox += MeshBoundingBox.TransformBy(CalcSliceTransformAtSplineOffset(T1));
+		}
+	}
+
+	return FBoxSphereBounds(BoundingBox.TransformBy(LocalToWorld));
+}
+
+
 FTransform USplineMeshComponent::CalcSliceTransform(const float DistanceAlong) const
 {
+	const bool bHasCustomBoundary = !FMath::IsNearlyEqual(SplineBoundaryMin, SplineBoundaryMax);
+
 	// Find how far 'along' mesh we are
-	FBoxSphereBounds StaticMeshBounds = StaticMesh->GetBounds();
-	const float MeshMinZ = GetAxisValue(StaticMeshBounds.Origin, ForwardAxis) - GetAxisValue(StaticMeshBounds.BoxExtent, ForwardAxis);
-	const float MeshRangeZ = 2.0f * GetAxisValue(StaticMeshBounds.BoxExtent, ForwardAxis);
-	const float Alpha = (DistanceAlong - MeshMinZ) / MeshRangeZ;
+	float Alpha;
+	if (bHasCustomBoundary)
+	{
+		Alpha = (DistanceAlong - SplineBoundaryMin) / (SplineBoundaryMax - SplineBoundaryMin);
+	}
+	else
+	{
+		const FBoxSphereBounds StaticMeshBounds = StaticMesh->GetBounds();
+		const float MeshMinZ = GetAxisValue(StaticMeshBounds.Origin, ForwardAxis) - GetAxisValue(StaticMeshBounds.BoxExtent, ForwardAxis);
+		const float MeshRangeZ = 2.0f * GetAxisValue(StaticMeshBounds.BoxExtent, ForwardAxis);
+		Alpha = (DistanceAlong - MeshMinZ) / MeshRangeZ;
+	}
 
+	return CalcSliceTransformAtSplineOffset(Alpha);
+}
+
+FTransform USplineMeshComponent::CalcSliceTransformAtSplineOffset(const float Alpha) const
+{
 	// Apply hermite interp to Alpha if desired
 	const float HermiteAlpha = bSmoothInterpRollScale ? SmoothStep(0.0, 1.0, Alpha) : Alpha;
 
diff --git a/Engine/Source/Runtime/Engine/Private/SplineMeshSceneProxy.cpp b/Engine/Source/Runtime/Engine/Private/SplineMeshSceneProxy.cpp
--- a/Engine/Source/Runtime/Engine/Private/SplineMeshSceneProxy.cpp
+++ b/Engine/Source/Runtime/Engine/Private/SplineMeshSceneProxy.cpp
@@ -27,9 +27,17 @@ FSplineMeshSceneProxy::FSplineMeshSceneProxy(USplineMeshComponent* InComponent)
 	ForwardAxis = InComponent->ForwardAxis;
 
 	// Fill in info about the mesh
-	FBoxSphereBounds StaticMeshBounds = StaticMesh->GetBounds();
-	SplineMeshScaleZ = 0.5f / USplineMeshComponent::GetAxisValue(StaticMeshBounds.BoxExtent, ForwardAxis); // 1/(2 * Extent)
-	SplineMeshMinZ = USplineMeshComponent::GetAxisValue(StaticMeshBounds.Origin, ForwardAxis) * SplineMeshScaleZ - 0.5f;
+	if (FMath::IsNearlyEqual(InComponent->SplineBoundaryMin, InComponent->SplineBoundaryMax))
+	{
+		FBoxSphereBounds StaticMeshBounds = StaticMesh->GetBounds();
+		SplineMeshScaleZ = 0.5f / USplineMeshComponent::GetAxisValue(StaticMeshBounds.BoxExtent, ForwardAxis); // 1/(2 * Extent)
+		SplineMeshMinZ = USplineMeshComponent::GetAxisValue(StaticMeshBounds.Origin, ForwardAxis) * SplineMeshScaleZ - 0.5f;
+	}
+	else
+	{
+		SplineMeshScaleZ = 1.0f / (InComponent->SplineBoundaryMax - InComponent->SplineBoundaryMin);
+		SplineMeshMinZ = InComponent->SplineBoundaryMin * SplineMeshScaleZ;
+	}
 
 	LODResources.Reset(InComponent->StaticMesh->RenderData->LODResources.Num());
 
