Fix compile on non APEX platforms

diff --git a/Engine/Source/Runtime/Engine/Private/Components/DestructibleComponent.cpp b/Engine/Source/Runtime/Engine/Private/Components/DestructibleComponent.cpp
--- a/Engine/Source/Runtime/Engine/Private/Components/DestructibleComponent.cpp
+++ b/Engine/Source/Runtime/Engine/Private/Components/DestructibleComponent.cpp
@@ -1381,8 +1381,9 @@ void UDestructibleComponent::SetMaterial(int32 ElementIndex, UMaterialInterface*
 		}
 	}
 
+#if WITH_APEX
 	// Set new template parameters for the apex actor, so they take effect before fracturing too.
-	physx::NxPhysX3DescTemplate* Template = ApexDestructibleActor->createPhysX3DescTemplate();
+	physx::apex::NxPhysX3DescTemplate* Template = ApexDestructibleActor->createPhysX3DescTemplate();
 	if(ApexDestructibleActor->getPhysX3Template(*Template))
 	{
 		UPhysicalMaterial* SimpleMaterial = GetBodyInstance()->GetSimplePhysicalMaterial();
@@ -1394,6 +1395,7 @@ void UDestructibleComponent::SetMaterial(int32 ElementIndex, UMaterialInterface*
 		ApexDestructibleActor->setPhysX3Template(Template);
 	}
 	Template->release();
+#endif
 }
 
 #endif // WITH_PHYSX
\ No newline at end of file
