Fix editor only compile

diff --git a/Engine/Source/Runtime/Engine/Classes/Animation/AnimSequence.h b/Engine/Source/Runtime/Engine/Classes/Animation/AnimSequence.h
--- a/Engine/Source/Runtime/Engine/Classes/Animation/AnimSequence.h
+++ b/Engine/Source/Runtime/Engine/Classes/Animation/AnimSequence.h
@@ -712,30 +712,6 @@ public:
 		return TrackToSkeletonMapTable[TrackIndex].BoneTreeIndex; 
 	}
 
-	
-	/**
-	 * Crops the raw anim data either from Start to CurrentTime or CurrentTime to End depending on
-	 * value of bFromStart.  Can't be called against cooked data.
-	 *
-	 * @param	CurrentTime		marker for cropping (either beginning or end)
-	 * @param	bFromStart		whether marker is begin or end marker
-	 * @return					true if the operation was successful.
-	 */
-	ENGINE_API bool CropRawAnimData( float CurrentTime, bool bFromStart );
-
-		
-	/**
-	 * Crops the raw anim data either from Start to CurrentTime or CurrentTime to End depending on
-	 * value of bFromStart.  Can't be called against cooked data.
-	 *
-	 * @param	StartFrame		StartFrame to insert (0-based)
-	 * @param	EndFrame		EndFrame to insert (0-based
-	 * @param	CopyFrame		A frame that we copy from (0-based)
-	 * @return					true if the operation was successful.
-	 */
-	ENGINE_API bool InsertFramesToRawAnimData( int32 StartFrame, int32 EndFrame, int32 CopyFrame);
-
-
 	/** Clears any data in the AnimSequence, so it can be recycled when importing a new animation with same name over it. */
 	ENGINE_API void RecycleAnimSequence();
 
@@ -802,11 +778,28 @@ public:
 	 * Create Animation Sequence from the given animation
 	 */
 	ENGINE_API bool CreateAnimation(class UAnimSequence * Sequence);
+
 	/**
-	 * Resize Animation to Start/End. If End is bigger, it repeats last frame upto End. 
-	 * @todo implement
+	 * Crops the raw anim data either from Start to CurrentTime or CurrentTime to End depending on
+	 * value of bFromStart.  Can't be called against cooked data.
+	 *
+	 * @param	CurrentTime		marker for cropping (either beginning or end)
+	 * @param	bFromStart		whether marker is begin or end marker
+	 * @return					true if the operation was successful.
 	 */
-	ENGINE_API bool Resize(int32 Start, int32 End);
+	ENGINE_API bool CropRawAnimData( float CurrentTime, bool bFromStart );
+
+		
+	/**
+	 * Crops the raw anim data either from Start to CurrentTime or CurrentTime to End depending on
+	 * value of bFromStart.  Can't be called against cooked data.
+	 *
+	 * @param	StartFrame		StartFrame to insert (0-based)
+	 * @param	EndFrame		EndFrame to insert (0-based
+	 * @param	CopyFrame		A frame that we copy from (0-based)
+	 * @return					true if the operation was successful.
+	 */
+	ENGINE_API bool InsertFramesToRawAnimData( int32 StartFrame, int32 EndFrame, int32 CopyFrame);
 #endif
 
 	/** 
@@ -846,7 +839,6 @@ private:
 	void ResetAnimation();
 	/** Refresh Track Map from Animation Track Names **/
 	void RefreshTrackMapFromAnimTrackNames();
-#endif
 
 	/**
 	 * Utility function that helps to remove track, you can't just remove RawAnimationData
@@ -857,9 +849,10 @@ private:
 	 */
 	int32 InsertTrack(const FName& BoneName);
 
-	/** Resize Sequence Length and Num of Frames */
 	void ResizeSequence(float NewLength, int32 NewNumFrames);
 
+#endif
+
 	friend class UAnimationAsset;
 };
 
diff --git a/Engine/Source/Runtime/Engine/Private/Animation/AnimSequence.cpp b/Engine/Source/Runtime/Engine/Private/Animation/AnimSequence.cpp
--- a/Engine/Source/Runtime/Engine/Private/Animation/AnimSequence.cpp
+++ b/Engine/Source/Runtime/Engine/Private/Animation/AnimSequence.cpp
@@ -1205,6 +1205,7 @@ void UAnimSequence::RetargetBoneTransform(FTransform& BoneTransform, const int32
 	}
 }
 
+#if WITH_EDITOR
 /** Utility function to crop data from a RawAnimSequenceTrack */
 static int32 CropRawTrack(FRawAnimSequenceTrack& RawTrack, int32 StartKey, int32 NumKeys, int32 TotalNumOfFrames)
 {
@@ -1634,6 +1635,7 @@ void UAnimSequence::FlipRotationWForNonRoot(USkeletalMesh * SkelMesh)
 	// Apply compression
 	FAnimationUtils::CompressAnimSequence(this, false, false);
 }
+#endif 
 
 void UAnimSequence::RecycleAnimSequence()
 {
@@ -2510,6 +2512,7 @@ void UAnimSequence::RemoveNaNTracks()
 	}
 }
 
+
 void UAnimSequence::RemoveTrack(int32 TrackIndex)
 {
 	if (RawAnimationData.IsValidIndex(TrackIndex))
@@ -3494,10 +3497,6 @@ bool UAnimSequence::CreateAnimation(UAnimSequence * Sequence)
 	return false;
 }
 
-bool UAnimSequence::Resize(int32 Start, int32 End)
-{
-	return false;
-}
 #endif
 /*-----------------------------------------------------------------------------
 	AnimNotify& subclasses
diff --git a/Engine/Source/Runtime/Engine/Private/Animation/AnimSequenceBase.cpp b/Engine/Source/Runtime/Engine/Private/Animation/AnimSequenceBase.cpp
--- a/Engine/Source/Runtime/Engine/Private/Animation/AnimSequenceBase.cpp
+++ b/Engine/Source/Runtime/Engine/Private/Animation/AnimSequenceBase.cpp
@@ -333,6 +333,7 @@ void FRawCurveTracks::Resize(float TotalLength)
 		Curve.Resize(TotalLength);
 	}
 
+#if WITH_EDITORONLY_DATA
 	for(auto& Curve: VectorCurves)
 	{
 		Curve.Resize(TotalLength);
@@ -342,6 +343,7 @@ void FRawCurveTracks::Resize(float TotalLength)
 	{
 		Curve.Resize(TotalLength);
 	}
+#endif
 }
 
 void FRawCurveTracks::Serialize(FArchive& Ar)
