#ue4 - Reverse normals in component overlap event when object is not moving and something else hits this component. This is needed so that normal is always pushing the object away from the impact point.

diff --git a/Engine/Source/Runtime/Engine/Classes/Engine/EngineTypes.h b/Engine/Source/Runtime/Engine/Classes/Engine/EngineTypes.h
--- a/Engine/Source/Runtime/Engine/Classes/Engine/EngineTypes.h
+++ b/Engine/Source/Runtime/Engine/Classes/Engine/EngineTypes.h
@@ -1448,9 +1448,9 @@ struct ENGINE_API FHitResult
 	float Time;
 	
 	/**
-	 * The location in world space where the shape would end up against the impacted object, if there is a hit. Equal to the point of impact for line tests.
+	 * The location in world space where the moving shape would end up against the impacted object, if there is a hit. Equal to the point of impact for line tests.
 	 * Example: for a sphere trace test, this is the point where the center of the sphere would be located when it touched the other object.
-	 * For swept movement (but not queries) this may not equal the final location of the shape since hits are pulled back slightly to prevent precision issues.
+	 * For swept movement (but not queries) this may not equal the final location of the shape since hits are pulled back slightly to prevent precision issues from overlapping another surface.
 	 */
 	UPROPERTY()
 	FVector_NetQuantize Location;
@@ -1606,6 +1606,18 @@ struct ENGINE_API FHitResult
 	{
 		return (InHits.Num() - GetNumBlockingHits(InHits));
 	}
+
+	/**
+	 * Get a copy of the HitResult with relevant information reversed.
+	 * For example when receiving a hit from another object, we reverse the normals.
+	 */
+	static FHitResult GetReversedHit(const FHitResult& Hit)
+	{
+		FHitResult Result(Hit);
+		Result.Normal = -Result.Normal;
+		Result.ImpactNormal = -Result.ImpactNormal;
+		return Result;
+	}
 };
 
 
diff --git a/Engine/Source/Runtime/Engine/Classes/GameFramework/Actor.h b/Engine/Source/Runtime/Engine/Classes/GameFramework/Actor.h
--- a/Engine/Source/Runtime/Engine/Classes/GameFramework/Actor.h
+++ b/Engine/Source/Runtime/Engine/Classes/GameFramework/Actor.h
@@ -1087,9 +1087,13 @@ public:
 	void GetOverlappingComponents(TArray<UPrimitiveComponent*>& OverlappingComponents) const;
 
 	/** 
-	 *	Event when this actor bumps into a blocking object, or blocks another actor that bumps into it. This could happen due to things like Character movement, using Set Location with 'sweep' enabled, or physics simulation.
-	 *	For events when objects overlap (e.g. walking into a trigger) see the 'Overlap' event.
-	 *	@note For collisions during physics simulation to generate hit events, 'Simulation Generates Hit Events' must be enabled.
+	 * Event when this actor bumps into a blocking object, or blocks another actor that bumps into it.
+	 * This could happen due to things like Character movement, using Set Location with 'sweep' enabled, or physics simulation.
+	 * For events when objects overlap (e.g. walking into a trigger) see the 'Overlap' event.
+	 *
+	 * @note When receiving a hit from another object (bSelfMoved is false), the 'Hit.Normal' and 'Hit.ImpactNormal' will
+	 * have their directions reversed to indicate opposition to this object's potential response and deflection.
+	 * @note For collisions during physics simulation to generate hit events, 'Simulation Generates Hit Events' must be enabled.
 	 */
 	UFUNCTION(BlueprintImplementableEvent, meta=(FriendlyName = "Hit"), Category="Collision")
 	virtual void ReceiveHit(class UPrimitiveComponent* MyComp, AActor* Other, class UPrimitiveComponent* OtherComp, bool bSelfMoved, FVector HitLocation, FVector HitNormal, FVector NormalImpulse, const FHitResult& Hit);
@@ -2375,6 +2379,11 @@ public:
 	/** Get the class of this Actor. */
 	DEPRECATED(4.8, "Use GetClass() instead of GetActorClass()")
 	class UClass* GetActorClass() const;
+
+private:
+
+	// Helper that already assumes the Hit info is reversed, and avoids creating a temp FHitResult if possible.
+	void InternalDispatchBlockingHit(UPrimitiveComponent* MyComp, UPrimitiveComponent* OtherComp, bool bSelfMoved, FHitResult const& Hit);
 };
 
 /**
diff --git a/Engine/Source/Runtime/Engine/Private/Actor.cpp b/Engine/Source/Runtime/Engine/Private/Actor.cpp
--- a/Engine/Source/Runtime/Engine/Private/Actor.cpp
+++ b/Engine/Source/Runtime/Engine/Private/Actor.cpp
@@ -1774,41 +1774,37 @@ bool IsActorValidToNotify(AActor* Actor)
 	return (Actor != NULL) && !Actor->IsPendingKill() && !Actor->GetClass()->HasAnyClassFlags(CLASS_NewerVersionExists);
 }
 
-void AActor::DispatchBlockingHit(UPrimitiveComponent* MyComp, UPrimitiveComponent* OtherComp, bool bSelfMoved, FHitResult const& Hit)
+void AActor::InternalDispatchBlockingHit(UPrimitiveComponent* MyComp, UPrimitiveComponent* OtherComp, bool bSelfMoved, FHitResult const& Hit)
 {
 	check(MyComp);
 
 	AActor* OtherActor = (OtherComp != NULL) ? OtherComp->GetOwner() : NULL;
-	const FHitResult* HitResult = &Hit;
 
-	FHitResult FlippedHit(Hit);
-	if(bSelfMoved == false)	//if we didn't move we flip the normals so they push us away from moving object
-	{
-		FlippedHit.Normal = -FlippedHit.Normal;
-		FlippedHit.ImpactNormal = -FlippedHit.ImpactNormal;
-		HitResult = &FlippedHit;
-	}
-	
 	// Call virtual
 	if(IsActorValidToNotify(this))
 	{
-		
-		ReceiveHit(MyComp, OtherActor, OtherComp, bSelfMoved, Hit.ImpactPoint, HitResult->ImpactNormal, FVector(0,0,0), *HitResult);
+		ReceiveHit(MyComp, OtherActor, OtherComp, bSelfMoved, Hit.ImpactPoint, Hit.ImpactNormal, FVector(0,0,0), Hit);
 	}
 
 	// If we are still ok, call delegate on actor
 	if(IsActorValidToNotify(this))
 	{
-		OnActorHit.Broadcast(this, OtherActor, FVector(0,0,0), *HitResult);
+		OnActorHit.Broadcast(this, OtherActor, FVector(0,0,0), Hit);
 	}
 
 	// If component is still alive, call delegate on component
 	if(!MyComp->IsPendingKill())
 	{
-		MyComp->OnComponentHit.Broadcast(OtherActor, OtherComp, FVector(0,0,0), *HitResult);
+		MyComp->OnComponentHit.Broadcast(OtherActor, OtherComp, FVector(0,0,0), Hit);
 	}
 }
 
+void AActor::DispatchBlockingHit(UPrimitiveComponent* MyComp, UPrimitiveComponent* OtherComp, bool bSelfMoved, FHitResult const& Hit)
+{
+	InternalDispatchBlockingHit(MyComp, OtherComp, bSelfMoved, bSelfMoved ? Hit : FHitResult::GetReversedHit(Hit));
+}
+
+
 FString AActor::GetHumanReadableName() const
 {
 	return GetName();
diff --git a/Engine/Source/Runtime/Engine/Private/Components/PrimitiveComponent.cpp b/Engine/Source/Runtime/Engine/Private/Components/PrimitiveComponent.cpp
--- a/Engine/Source/Runtime/Engine/Private/Components/PrimitiveComponent.cpp
+++ b/Engine/Source/Runtime/Engine/Private/Components/PrimitiveComponent.cpp
@@ -1973,7 +1973,8 @@ void UPrimitiveComponent::BeginComponentOverlap(const FOverlapInfo& OtherOverlap
 
 				if (!OtherComp->IsPendingKill())
 				{
-					OtherComp->OnComponentBeginOverlap.Broadcast(MyActor, this, INDEX_NONE, OtherOverlap.bFromSweep, OtherOverlap.OverlapInfo);
+					// Reverse normals for other component. When it's a sweep, we are the one that moved.
+					OtherComp->OnComponentBeginOverlap.Broadcast(MyActor, this, INDEX_NONE, OtherOverlap.bFromSweep, OtherOverlap.bFromSweep ? FHitResult::GetReversedHit(OtherOverlap.OverlapInfo) : OtherOverlap.OverlapInfo);
 				}
 
 				// then execute actor notification if this is a new actor touch
