#ue4 - Improved documentation for "reversed normals" in ReceiveHit, OnComponentHit, and OnComponentBeginOverlap.

diff --git a/Engine/Source/Runtime/Engine/Classes/Components/PrimitiveComponent.h b/Engine/Source/Runtime/Engine/Classes/Components/PrimitiveComponent.h
--- a/Engine/Source/Runtime/Engine/Classes/Components/PrimitiveComponent.h
+++ b/Engine/Source/Runtime/Engine/Classes/Components/PrimitiveComponent.h
@@ -586,7 +586,10 @@ public:
 	/** 
 	 *	Event called when a component hits (or is hit by) something solid. This could happen due to things like Character movement, using Set Location with 'sweep' enabled, or physics simulation.
 	 *	For events when objects overlap (e.g. walking into a trigger) see the 'Overlap' event.
-	 *	@note For collisions during physics simulation to generate hit events, 'Simulation Generates Hit Events' must be enabled for this component
+	 *
+	 *	@note For collisions during physics simulation to generate hit events, 'Simulation Generates Hit Events' must be enabled for this component.
+	 *	@note When receiving a hit from another object's movement, the directions of 'Hit.Normal' and 'Hit.ImpactNormal'
+	 *	will be adjusted to indicate force from the other object against this object.
 	 */
 	UPROPERTY(BlueprintAssignable, Category="Collision")
 	FComponentHitSignature OnComponentHit;
@@ -594,7 +597,10 @@ public:
 	/** 
 	 *	Event called when something starts to overlaps this component, for example a player walking into a trigger.
 	 *	For events when objects have a blocking collision, for example a player hitting a wall, see 'Hit' events.
+	 *
 	 *	@note Both this component and the other one must have bGenerateOverlapEvents set to true to generate overlap events.
+	 *	@note When receiving an overlap from another object's movement, the directions of 'Hit.Normal' and 'Hit.ImpactNormal'
+	 *	will be adjusted to indicate force from the other object against this object.
 	 */
 	UPROPERTY(BlueprintAssignable, Category="Collision")
 	FComponentBeginOverlapSignature OnComponentBeginOverlap;
diff --git a/Engine/Source/Runtime/Engine/Classes/GameFramework/Actor.h b/Engine/Source/Runtime/Engine/Classes/GameFramework/Actor.h
--- a/Engine/Source/Runtime/Engine/Classes/GameFramework/Actor.h
+++ b/Engine/Source/Runtime/Engine/Classes/GameFramework/Actor.h
@@ -1091,9 +1091,9 @@ public:
 	 * This could happen due to things like Character movement, using Set Location with 'sweep' enabled, or physics simulation.
 	 * For events when objects overlap (e.g. walking into a trigger) see the 'Overlap' event.
 	 *
-	 * @note When receiving a hit from another object (bSelfMoved is false), the 'Hit.Normal' and 'Hit.ImpactNormal' will
-	 * have their directions reversed to indicate opposition to this object's potential response and deflection.
 	 * @note For collisions during physics simulation to generate hit events, 'Simulation Generates Hit Events' must be enabled.
+	 * @note When receiving a hit from another object's movement (bSelfMoved is false), the directions of 'Hit.Normal' and 'Hit.ImpactNormal'
+	 * will be adjusted to indicate force from the other object against this object.
 	 */
 	UFUNCTION(BlueprintImplementableEvent, meta=(FriendlyName = "Hit"), Category="Collision")
 	virtual void ReceiveHit(class UPrimitiveComponent* MyComp, AActor* Other, class UPrimitiveComponent* OtherComp, bool bSelfMoved, FVector HitLocation, FVector HitNormal, FVector NormalImpulse, const FHitResult& Hit);
