Fix for crash opening a level that contains instanced meshes with built lighting (UE-11865)

diff --git a/Engine/Source/Runtime/Engine/Private/InstancedStaticMesh.cpp b/Engine/Source/Runtime/Engine/Private/InstancedStaticMesh.cpp
--- a/Engine/Source/Runtime/Engine/Private/InstancedStaticMesh.cpp
+++ b/Engine/Source/Runtime/Engine/Private/InstancedStaticMesh.cpp
@@ -827,11 +827,14 @@ void UInstancedStaticMeshComponent::ApplyComponentInstanceData(FInstancedStaticM
 	// See if data matches current state
 	if (InstancedMeshData->bHasCachedStaticLighting)
 	{
-		bool bMatch = InstancedMeshData->CachedStaticLighting.Transform.Equals(ComponentToWorld);
+		bool bMatch = false;
 
 		// Check for any instance having moved as that would invalidate static lighting
-		if (PerInstanceSMData.Num() == InstancedMeshData->PerInstanceSMData.Num())
+		if (PerInstanceSMData.Num() == InstancedMeshData->PerInstanceSMData.Num() &&
+			InstancedMeshData->CachedStaticLighting.Transform.Equals(ComponentToWorld))
 		{
+			bMatch = true;
+
 			for (int32 InstanceIndex = 0; InstanceIndex < PerInstanceSMData.Num(); ++InstanceIndex)
 			{
 				if (PerInstanceSMData[InstanceIndex].Transform != InstancedMeshData->PerInstanceSMData[InstanceIndex].Transform)
