Enabled UHT to digest GENERATED_BODY instead of GENERATED_UCLASS_BODY, GENERATED_USTRUCT_BODY, GENERATED_UINTERFACE_BODY or GENERATED_IINTERFACE_BODY, changed every occurence to the new syntax and fix

diff --git a/Engine/Source/Runtime/Engine/Classes/Animation/AnimCompositeBase.h b/Engine/Source/Runtime/Engine/Classes/Animation/AnimCompositeBase.h
--- a/Engine/Source/Runtime/Engine/Classes/Animation/AnimCompositeBase.h
+++ b/Engine/Source/Runtime/Engine/Classes/Animation/AnimCompositeBase.h
@@ -20,7 +20,7 @@ class UAnimSequence;
 USTRUCT()
 struct FRootMotionExtractionStep
 {
-	GENERATED_USTRUCT_BODY()
+	GENERATED_BODY()
 
 	/** AnimSequence ref */
 	UPROPERTY()
@@ -53,7 +53,7 @@ struct FRootMotionExtractionStep
 USTRUCT()
 struct FAnimSegment
 {
-	GENERATED_USTRUCT_BODY()
+	GENERATED_BODY()
 
 	/** Anim Reference to play - only allow AnimSequence or AnimComposite **/
 	UPROPERTY(EditAnywhere, BlueprintReadWrite, Category=AnimSegment)
@@ -138,7 +138,7 @@ struct FAnimSegment
 USTRUCT()
 struct FAnimTrack
 {
-	GENERATED_USTRUCT_BODY()
+	GENERATED_BODY()
 
 	UPROPERTY(EditAnywhere, BlueprintReadWrite, Category=AnimTrack, EditFixedSize)
 	TArray<FAnimSegment>	AnimSegments;
@@ -186,7 +186,9 @@ struct FAnimTrack
 UCLASS(abstract, MinimalAPI)
 class UAnimCompositeBase : public UAnimSequenceBase
 {
-	GENERATED_UCLASS_BODY()
+	GENERATED_BODY()
+public:
+	ENGINE_API UAnimCompositeBase(const FObjectInitializer& ObjectInitializer = FObjectInitializer::Get());
 
 #if WITH_EDITOR
 	/** Set Sequence Length */
diff --git a/Engine/Source/Runtime/Engine/Classes/Animation/AnimCompress.h b/Engine/Source/Runtime/Engine/Classes/Animation/AnimCompress.h
--- a/Engine/Source/Runtime/Engine/Classes/Animation/AnimCompress.h
+++ b/Engine/Source/Runtime/Engine/Classes/Animation/AnimCompress.h
@@ -13,7 +13,9 @@
 UCLASS(abstract, hidecategories=Object, MinimalAPI)
 class UAnimCompress : public UObject
 {
-	GENERATED_UCLASS_BODY()
+	GENERATED_BODY()
+public:
+	ENGINE_API UAnimCompress(const FObjectInitializer& ObjectInitializer = FObjectInitializer::Get());
 
 	/** Name of Compression Scheme used for this asset*/
 	UPROPERTY(Category=Compression, VisibleAnywhere)
diff --git a/Engine/Source/Runtime/Engine/Classes/Animation/AnimCompress_Automatic.h b/Engine/Source/Runtime/Engine/Classes/Animation/AnimCompress_Automatic.h
--- a/Engine/Source/Runtime/Engine/Classes/Animation/AnimCompress_Automatic.h
+++ b/Engine/Source/Runtime/Engine/Classes/Animation/AnimCompress_Automatic.h
@@ -13,7 +13,9 @@
 UCLASS(hidecategories=AnimationCompressionAlgorithm)
 class UAnimCompress_Automatic : public UAnimCompress
 {
-	GENERATED_UCLASS_BODY()
+	GENERATED_BODY()
+public:
+	UAnimCompress_Automatic(const FObjectInitializer& ObjectInitializer = FObjectInitializer::Get());
 
 	/** Maximum amount of error that a compression technique can introduce in an end effector */
 	UPROPERTY(EditAnywhere, Category=AnimationCompressionAlgorithm_Automatic)
diff --git a/Engine/Source/Runtime/Engine/Classes/Animation/AnimCompress_BitwiseCompressOnly.h b/Engine/Source/Runtime/Engine/Classes/Animation/AnimCompress_BitwiseCompressOnly.h
--- a/Engine/Source/Runtime/Engine/Classes/Animation/AnimCompress_BitwiseCompressOnly.h
+++ b/Engine/Source/Runtime/Engine/Classes/Animation/AnimCompress_BitwiseCompressOnly.h
@@ -12,7 +12,9 @@
 UCLASS(MinimalAPI)
 class UAnimCompress_BitwiseCompressOnly : public UAnimCompress
 {
-	GENERATED_UCLASS_BODY()
+	GENERATED_BODY()
+public:
+	ENGINE_API UAnimCompress_BitwiseCompressOnly(const FObjectInitializer& ObjectInitializer = FObjectInitializer::Get());
 
 
 protected:
diff --git a/Engine/Source/Runtime/Engine/Classes/Animation/AnimCompress_LeastDestructive.h b/Engine/Source/Runtime/Engine/Classes/Animation/AnimCompress_LeastDestructive.h
--- a/Engine/Source/Runtime/Engine/Classes/Animation/AnimCompress_LeastDestructive.h
+++ b/Engine/Source/Runtime/Engine/Classes/Animation/AnimCompress_LeastDestructive.h
@@ -12,7 +12,9 @@
 UCLASS()
 class UAnimCompress_LeastDestructive : public UAnimCompress
 {
-	GENERATED_UCLASS_BODY()
+	GENERATED_BODY()
+public:
+	UAnimCompress_LeastDestructive(const FObjectInitializer& ObjectInitializer = FObjectInitializer::Get());
 
 
 protected:
diff --git a/Engine/Source/Runtime/Engine/Classes/Animation/AnimCompress_PerTrackCompression.h b/Engine/Source/Runtime/Engine/Classes/Animation/AnimCompress_PerTrackCompression.h
--- a/Engine/Source/Runtime/Engine/Classes/Animation/AnimCompress_PerTrackCompression.h
+++ b/Engine/Source/Runtime/Engine/Classes/Animation/AnimCompress_PerTrackCompression.h
@@ -13,7 +13,9 @@
 UCLASS(hidecategories=AnimCompress)
 class UAnimCompress_PerTrackCompression : public UAnimCompress_RemoveLinearKeys
 {
-	GENERATED_UCLASS_BODY()
+	GENERATED_BODY()
+public:
+	UAnimCompress_PerTrackCompression(const FObjectInitializer& ObjectInitializer = FObjectInitializer::Get());
 
 	/** Maximum threshold to use when replacing a component with zero. Lower values retain more keys, but yield less compression. */
 	UPROPERTY(EditAnywhere, Category=PerTrack)
diff --git a/Engine/Source/Runtime/Engine/Classes/Animation/AnimCompress_RemoveEverySecondKey.h b/Engine/Source/Runtime/Engine/Classes/Animation/AnimCompress_RemoveEverySecondKey.h
--- a/Engine/Source/Runtime/Engine/Classes/Animation/AnimCompress_RemoveEverySecondKey.h
+++ b/Engine/Source/Runtime/Engine/Classes/Animation/AnimCompress_RemoveEverySecondKey.h
@@ -12,7 +12,9 @@
 UCLASS(MinimalAPI)
 class UAnimCompress_RemoveEverySecondKey : public UAnimCompress
 {
-	GENERATED_UCLASS_BODY()
+	GENERATED_BODY()
+public:
+	ENGINE_API UAnimCompress_RemoveEverySecondKey(const FObjectInitializer& ObjectInitializer = FObjectInitializer::Get());
 
 	/** Animations with fewer than MinKeys will not lose any keys. */
 	UPROPERTY(EditAnywhere, Category=AnimationCompressionAlgorithm_RemoveEverySecondKey)
diff --git a/Engine/Source/Runtime/Engine/Classes/Animation/AnimCompress_RemoveLinearKeys.h b/Engine/Source/Runtime/Engine/Classes/Animation/AnimCompress_RemoveLinearKeys.h
--- a/Engine/Source/Runtime/Engine/Classes/Animation/AnimCompress_RemoveLinearKeys.h
+++ b/Engine/Source/Runtime/Engine/Classes/Animation/AnimCompress_RemoveLinearKeys.h
@@ -12,7 +12,9 @@
 UCLASS(MinimalAPI)
 class UAnimCompress_RemoveLinearKeys : public UAnimCompress
 {
-	GENERATED_UCLASS_BODY()
+	GENERATED_BODY()
+public:
+	ENGINE_API UAnimCompress_RemoveLinearKeys(const FObjectInitializer& ObjectInitializer = FObjectInitializer::Get());
 
 	/** Maximum position difference to use when testing if an animation key may be removed. Lower values retain more keys, but yield less compression. */
 	UPROPERTY(EditAnywhere, Category=LinearKeyRemoval)
diff --git a/Engine/Source/Runtime/Engine/Classes/Animation/AnimCompress_RemoveTrivialKeys.h b/Engine/Source/Runtime/Engine/Classes/Animation/AnimCompress_RemoveTrivialKeys.h
--- a/Engine/Source/Runtime/Engine/Classes/Animation/AnimCompress_RemoveTrivialKeys.h
+++ b/Engine/Source/Runtime/Engine/Classes/Animation/AnimCompress_RemoveTrivialKeys.h
@@ -14,7 +14,9 @@
 UCLASS(MinimalAPI)
 class UAnimCompress_RemoveTrivialKeys : public UAnimCompress
 {
-	GENERATED_UCLASS_BODY()
+	GENERATED_BODY()
+public:
+	ENGINE_API UAnimCompress_RemoveTrivialKeys(const FObjectInitializer& ObjectInitializer = FObjectInitializer::Get());
 
 	UPROPERTY(EditAnywhere, Category=AnimationCompressionAlgorithm_RemoveTrivialKeys)
 	float MaxPosDiff;
diff --git a/Engine/Source/Runtime/Engine/Classes/Animation/AnimData/BoneMaskFilter.h b/Engine/Source/Runtime/Engine/Classes/Animation/AnimData/BoneMaskFilter.h
--- a/Engine/Source/Runtime/Engine/Classes/Animation/AnimData/BoneMaskFilter.h
+++ b/Engine/Source/Runtime/Engine/Classes/Animation/AnimData/BoneMaskFilter.h
@@ -12,7 +12,7 @@
 USTRUCT()
 struct FBranchFilter
 {
-	GENERATED_USTRUCT_BODY()
+	GENERATED_BODY()
 
 	/** Bone Name to filter **/
 	UPROPERTY(EditAnywhere, Category=Branch)
@@ -26,7 +26,7 @@ struct FBranchFilter
 USTRUCT()
 struct FInputBlendPose
 {
-	GENERATED_USTRUCT_BODY()
+	GENERATED_BODY()
 
 	/** Bone Name to filter **/
 	UPROPERTY(EditAnywhere, Category=Filter)
@@ -36,7 +36,9 @@ struct FInputBlendPose
 UCLASS(hidecategories=Object, BlueprintType,MinimalAPI)
 class UBoneMaskFilter : public UObject
 {
-	GENERATED_UCLASS_BODY()
+	GENERATED_BODY()
+public:
+	ENGINE_API UBoneMaskFilter(const FObjectInitializer& ObjectInitializer = FObjectInitializer::Get());
 
 	UPROPERTY(EditAnywhere, Category=BoneMask)
 	TArray<FInputBlendPose>	BlendPoses;
diff --git a/Engine/Source/Runtime/Engine/Classes/Animation/AnimInstance.h b/Engine/Source/Runtime/Engine/Classes/Animation/AnimInstance.h
--- a/Engine/Source/Runtime/Engine/Classes/Animation/AnimInstance.h
+++ b/Engine/Source/Runtime/Engine/Classes/Animation/AnimInstance.h
@@ -83,7 +83,7 @@ enum EBoneRotationSource
 USTRUCT()
 struct FA2Pose
 {
-	GENERATED_USTRUCT_BODY()
+	GENERATED_BODY()
 
 	UPROPERTY()
 	TArray<FTransform> Bones;
@@ -95,7 +95,7 @@ struct FA2Pose
 USTRUCT()
 struct ENGINE_API FA2CSPose : public FA2Pose
 {
-	GENERATED_USTRUCT_BODY()
+	GENERATED_BODY()
 
 private:
 	/** Pointer to current BoneContainer. */
@@ -173,7 +173,7 @@ private:
 USTRUCT(BlueprintType)
 struct FBoneTransform
 {
-	GENERATED_USTRUCT_BODY()
+	GENERATED_BODY()
 
 	/** @todo anim: should be Skeleton bone index in the future, but right now it's Mesh BoneIndex **/
 	UPROPERTY()
@@ -196,7 +196,7 @@ struct FBoneTransform
 USTRUCT(BlueprintType)
 struct FPerBoneBlendWeight
 {
-	GENERATED_USTRUCT_BODY()
+	GENERATED_BODY()
 	UPROPERTY()
 	int32 SourceIndex; // source index of the buffer
 	UPROPERTY()
@@ -212,7 +212,7 @@ struct FPerBoneBlendWeight
 USTRUCT(BlueprintType)
 struct FPerBoneBlendWeights
 {
-	GENERATED_USTRUCT_BODY()
+	GENERATED_BODY()
 
 	UPROPERTY()
 	TArray<FPerBoneBlendWeight> BoneBlendWeights;
@@ -226,7 +226,7 @@ struct FPerBoneBlendWeights
 USTRUCT()
 struct FSlotEvaluationPose
 {
-	GENERATED_USTRUCT_BODY()
+	GENERATED_BODY()
 
 	/** Type of additive for pose */
 	UPROPERTY()
@@ -341,7 +341,9 @@ struct FNativeStateBinding
 UCLASS(transient, Blueprintable, hideCategories=AnimInstance, BlueprintType)
 class ENGINE_API UAnimInstance : public UObject
 {
-	GENERATED_UCLASS_BODY()
+	GENERATED_BODY()
+public:
+	UAnimInstance(const FObjectInitializer& ObjectInitializer = FObjectInitializer::Get());
 
 	/** DeltaTime **/
 	UPROPERTY()
diff --git a/Engine/Source/Runtime/Engine/Classes/Animation/AnimMontage.h b/Engine/Source/Runtime/Engine/Classes/Animation/AnimMontage.h
--- a/Engine/Source/Runtime/Engine/Classes/Animation/AnimMontage.h
+++ b/Engine/Source/Runtime/Engine/Classes/Animation/AnimMontage.h
@@ -21,7 +21,7 @@ class UAnimMontage;
 USTRUCT()
 struct FCompositeSection : public FAnimLinkableElement
 {
-	GENERATED_USTRUCT_BODY()
+	GENERATED_BODY()
 
 	/** Section Name */
 	UPROPERTY(EditAnywhere, Category=Section)
@@ -50,7 +50,7 @@ struct FCompositeSection : public FAnimLinkableElement
 USTRUCT()
 struct FSlotAnimationTrack
 {
-	GENERATED_USTRUCT_BODY()
+	GENERATED_BODY()
 
 	UPROPERTY(EditAnywhere, Category=Slot)
 	FName SlotName;
@@ -69,7 +69,7 @@ struct FSlotAnimationTrack
 USTRUCT()
 struct FBranchingPoint : public FAnimLinkableElement
 {
-	GENERATED_USTRUCT_BODY()
+	GENERATED_BODY()
 
 	UPROPERTY(EditAnywhere, Category=BranchingPoint)
 	FName EventName;
@@ -103,7 +103,7 @@ namespace EAnimNotifyEventType
 USTRUCT()
 struct FBranchingPointMarker
 {
-	GENERATED_USTRUCT_BODY()
+	GENERATED_BODY()
 
 	UPROPERTY()
 	int32 NotifyIndex;
@@ -147,7 +147,7 @@ DECLARE_DELEGATE_TwoParams( FOnMontageBlendingOutStarted, class UAnimMontage*, b
 USTRUCT()
 struct FAnimMontageInstance
 {
-	GENERATED_USTRUCT_BODY()
+	GENERATED_BODY()
 
 	// Montage reference
 	UPROPERTY()
@@ -309,7 +309,9 @@ private:
 UCLASS(config=Engine, hidecategories=(UObject, Length), MinimalAPI, BlueprintType)
 class UAnimMontage : public UAnimCompositeBase
 {
-	GENERATED_UCLASS_BODY()
+	GENERATED_BODY()
+public:
+	ENGINE_API UAnimMontage(const FObjectInitializer& ObjectInitializer = FObjectInitializer::Get());
 
 	/** Default blend in time. */
 	UPROPERTY(EditAnywhere, Category=Montage)
diff --git a/Engine/Source/Runtime/Engine/Classes/Animation/AnimNodeBase.h b/Engine/Source/Runtime/Engine/Classes/Animation/AnimNodeBase.h
--- a/Engine/Source/Runtime/Engine/Classes/Animation/AnimNodeBase.h
+++ b/Engine/Source/Runtime/Engine/Classes/Animation/AnimNodeBase.h
@@ -261,7 +261,7 @@ namespace EPinHidingMode
 USTRUCT()
 struct ENGINE_API FPoseLinkBase
 {
-	GENERATED_USTRUCT_BODY()
+	GENERATED_BODY()
 
 	// Serialized link ID, used to build the non-serialized pointer map
 	UPROPERTY()
@@ -305,7 +305,7 @@ public:
 USTRUCT()
 struct ENGINE_API FPoseLink : public FPoseLinkBase
 {
-	GENERATED_USTRUCT_BODY()
+	GENERATED_BODY()
 
 public:
 	// Interface
@@ -316,7 +316,7 @@ public:
 USTRUCT()
 struct ENGINE_API FComponentSpacePoseLink : public FPoseLinkBase
 {
-	GENERATED_USTRUCT_BODY()
+	GENERATED_BODY()
 
 public:
 	// Interface
@@ -327,7 +327,7 @@ public:
 USTRUCT()
 struct FExposedValueHandler
 {
-	GENERATED_USTRUCT_BODY()
+	GENERATED_BODY()
 
 	// The function to call to update associated properties (can be NULL)
 	UPROPERTY()
@@ -354,7 +354,7 @@ struct FExposedValueHandler
 USTRUCT()
 struct ENGINE_API FAnimNode_Base
 {
-	GENERATED_USTRUCT_BODY()
+	GENERATED_BODY()
 
 	// The default handler for graph-exposed inputs
 	UPROPERTY()
diff --git a/Engine/Source/Runtime/Engine/Classes/Animation/AnimNode_StateMachine.h b/Engine/Source/Runtime/Engine/Classes/Animation/AnimNode_StateMachine.h
--- a/Engine/Source/Runtime/Engine/Classes/Animation/AnimNode_StateMachine.h
+++ b/Engine/Source/Runtime/Engine/Classes/Animation/AnimNode_StateMachine.h
@@ -12,7 +12,7 @@ struct FAnimNode_TransitionPoseEvaluator;
 USTRUCT()
 struct FAnimationActiveTransitionEntry
 {
-	GENERATED_USTRUCT_BODY()
+	GENERATED_BODY()
 
 	// Elapsed time for this transition
 	float ElapsedTime;
@@ -77,7 +77,7 @@ protected:
 USTRUCT()
 struct FAnimationPotentialTransition
 {
-	GENERATED_USTRUCT_BODY()
+	GENERATED_BODY()
 
 	int32 TargetState;
 
@@ -99,7 +99,7 @@ public:
 USTRUCT()
 struct ENGINE_API FAnimNode_StateMachine : public FAnimNode_Base
 {
-	GENERATED_USTRUCT_BODY()
+	GENERATED_BODY()
 
 public:
 	// Index into the BakedStateMachines array in the owning UAnimBlueprintGeneratedClass
diff --git a/Engine/Source/Runtime/Engine/Classes/Animation/AnimNotifies/AnimNotify.h b/Engine/Source/Runtime/Engine/Classes/Animation/AnimNotifies/AnimNotify.h
--- a/Engine/Source/Runtime/Engine/Classes/Animation/AnimNotifies/AnimNotify.h
+++ b/Engine/Source/Runtime/Engine/Classes/Animation/AnimNotifies/AnimNotify.h
@@ -12,13 +12,16 @@ struct FAnimNotifyEvent;
 UCLASS(abstract, Blueprintable, const, hidecategories=Object, collapsecategories)
 class ENGINE_API UAnimNotify : public UObject
 {
-	GENERATED_UCLASS_BODY()
+	GENERATED_BODY()
+public:
+	UAnimNotify(const FObjectInitializer& ObjectInitializer = FObjectInitializer::Get());
 
 	/** 
 	 * Implementable event to get a custom name for the notify
 	 */
 	UFUNCTION(BlueprintNativeEvent)
 	FString GetNotifyName() const;
+	virtual FString GetNotifyName_Implementation() const;
 
 	UFUNCTION(BlueprintImplementableEvent)
 	virtual bool Received_Notify(USkeletalMeshComponent* MeshComp, UAnimSequenceBase* Animation) const;
diff --git a/Engine/Source/Runtime/Engine/Classes/Animation/AnimNotifies/AnimNotifyState.h b/Engine/Source/Runtime/Engine/Classes/Animation/AnimNotifies/AnimNotifyState.h
--- a/Engine/Source/Runtime/Engine/Classes/Animation/AnimNotifies/AnimNotifyState.h
+++ b/Engine/Source/Runtime/Engine/Classes/Animation/AnimNotifies/AnimNotifyState.h
@@ -12,13 +12,16 @@ struct FAnimNotifyEvent;
 UCLASS(abstract, editinlinenew, Blueprintable, const, hidecategories=Object, collapsecategories, meta=(ShowWorldContextPin))
 class ENGINE_API UAnimNotifyState : public UObject
 {
-	GENERATED_UCLASS_BODY()
+	GENERATED_BODY()
+public:
+	UAnimNotifyState(const FObjectInitializer& ObjectInitializer = FObjectInitializer::Get());
 
 	/** 
 	 * Implementable event to get a custom name for the notify
 	 */
 	UFUNCTION(BlueprintNativeEvent)
 	FString GetNotifyName() const;
+	virtual FString GetNotifyName_Implementation() const;
 
 	UFUNCTION(BlueprintImplementableEvent)
 	virtual bool Received_NotifyBegin(USkeletalMeshComponent * MeshComp, UAnimSequenceBase * Animation, float TotalDuration) const;
diff --git a/Engine/Source/Runtime/Engine/Classes/Animation/AnimNotifies/AnimNotifyState_TimedParticleEffect.h b/Engine/Source/Runtime/Engine/Classes/Animation/AnimNotifies/AnimNotifyState_TimedParticleEffect.h
--- a/Engine/Source/Runtime/Engine/Classes/Animation/AnimNotifies/AnimNotifyState_TimedParticleEffect.h
+++ b/Engine/Source/Runtime/Engine/Classes/Animation/AnimNotifies/AnimNotifyState_TimedParticleEffect.h
@@ -10,7 +10,9 @@
 UCLASS(MinimalAPI, Blueprintable, meta = (DisplayName = "Timed Particle Effect"))
 class UAnimNotifyState_TimedParticleEffect : public UAnimNotifyState
 {
-	GENERATED_UCLASS_BODY()
+	GENERATED_BODY()
+public:
+	ENGINE_API UAnimNotifyState_TimedParticleEffect(const FObjectInitializer& ObjectInitializer = FObjectInitializer::Get());
 
 	// The particle system template to use when spawning the particle component
 	UPROPERTY(EditAnywhere, Category = ParticleSystem, meta = (ToolTip = "The particle system to spawn for the notify state"))
@@ -60,4 +62,4 @@ class UAnimNotifyState_TimedParticleEffect : public UAnimNotifyState
 
 private:
 	bool ValidateParameters(USkeletalMeshComponent* MeshComp);
-};
\ No newline at end of file
+};
diff --git a/Engine/Source/Runtime/Engine/Classes/Animation/AnimNotifies/AnimNotifyState_Trail.h b/Engine/Source/Runtime/Engine/Classes/Animation/AnimNotifies/AnimNotifyState_Trail.h
--- a/Engine/Source/Runtime/Engine/Classes/Animation/AnimNotifies/AnimNotifyState_Trail.h
+++ b/Engine/Source/Runtime/Engine/Classes/Animation/AnimNotifies/AnimNotifyState_Trail.h
@@ -9,7 +9,9 @@ DECLARE_LOG_CATEGORY_EXTERN(LogAnimTrails, Log, All);
 UCLASS(editinlinenew, Blueprintable, const, hidecategories = Object, collapsecategories, MinimalAPI, meta = (ShowWorldContextPin, DisplayName = "Trail"))
 class UAnimNotifyState_Trail : public UAnimNotifyState
 {
-	GENERATED_UCLASS_BODY()
+	GENERATED_BODY()
+public:
+	ENGINE_API UAnimNotifyState_Trail(const FObjectInitializer& ObjectInitializer = FObjectInitializer::Get());
 	
 	/** The particle system to use for this trail. */
 	UPROPERTY(EditAnywhere, BlueprintReadWrite, Category = Trail)
diff --git a/Engine/Source/Runtime/Engine/Classes/Animation/AnimSequence.h b/Engine/Source/Runtime/Engine/Classes/Animation/AnimSequence.h
--- a/Engine/Source/Runtime/Engine/Classes/Animation/AnimSequence.h
+++ b/Engine/Source/Runtime/Engine/Classes/Animation/AnimSequence.h
@@ -84,7 +84,7 @@ enum EAdditiveBasePoseType
 USTRUCT()
 struct FRawAnimSequenceTrack
 {
-	GENERATED_USTRUCT_BODY()
+	GENERATED_BODY()
 
 	/** Position keys. */
 	UPROPERTY()
@@ -121,7 +121,7 @@ struct FRawAnimSequenceTrack
 USTRUCT()
 struct FAnimSequenceTrackContainer
 {
-	GENERATED_USTRUCT_BODY()
+	GENERATED_BODY()
 
 	UPROPERTY()
 	TArray<struct FRawAnimSequenceTrack> AnimationTracks;
@@ -159,7 +159,7 @@ struct FAnimSequenceTrackContainer
 USTRUCT()
 struct FTrackToSkeletonMap
 {
-	GENERATED_USTRUCT_BODY()
+	GENERATED_BODY()
 
 	UPROPERTY()
 	int32 SkeletonIndex_DEPRECATED;    // 0 as current Skeleton (the one above), and N is the N'th parent
@@ -185,7 +185,7 @@ struct FTrackToSkeletonMap
 USTRUCT()
 struct FTranslationTrack
 {
-	GENERATED_USTRUCT_BODY()
+	GENERATED_BODY()
 
 	UPROPERTY()
 	TArray<FVector> PosKeys;
@@ -200,7 +200,7 @@ struct FTranslationTrack
 USTRUCT()
 struct FRotationTrack
 {
-	GENERATED_USTRUCT_BODY()
+	GENERATED_BODY()
 
 	UPROPERTY()
 	TArray<FQuat> RotKeys;
@@ -215,7 +215,7 @@ struct FRotationTrack
 USTRUCT()
 struct FScaleTrack
 {
-	GENERATED_USTRUCT_BODY()
+	GENERATED_BODY()
 
 	UPROPERTY()
 	TArray<FVector> ScaleKeys;
@@ -233,7 +233,7 @@ struct FScaleTrack
 USTRUCT()
 struct FCurveTrack
 {
-	GENERATED_USTRUCT_BODY()
+	GENERATED_BODY()
 
 	UPROPERTY()
 	FName CurveName;
@@ -253,7 +253,7 @@ struct FCurveTrack
 USTRUCT()
 struct FCompressedTrack
 {
-	GENERATED_USTRUCT_BODY()
+	GENERATED_BODY()
 
 	UPROPERTY()
 	TArray<uint8> ByteStream;
@@ -285,7 +285,7 @@ struct FCompressedTrack
 USTRUCT()
 struct FCompressedOffsetData
 {
-	GENERATED_USTRUCT_BODY()
+	GENERATED_BODY()
 
 	UPROPERTY()
 	TArray<int32> OffsetData;
@@ -345,7 +345,9 @@ struct FCompressedOffsetData
 UCLASS(config=Engine, hidecategories=(UObject, Length), MinimalAPI, BlueprintType)
 class UAnimSequence : public UAnimSequenceBase
 {
-	GENERATED_UCLASS_BODY()
+	GENERATED_BODY()
+public:
+	ENGINE_API UAnimSequence(const FObjectInitializer& ObjectInitializer = FObjectInitializer::Get());
 
 	/** Number of raw frames in this sequence (not used by engine - just for informational purposes). */
 	UPROPERTY(AssetRegistrySearchable)
diff --git a/Engine/Source/Runtime/Engine/Classes/Animation/AnimSequenceBase.h b/Engine/Source/Runtime/Engine/Classes/Animation/AnimSequenceBase.h
--- a/Engine/Source/Runtime/Engine/Classes/Animation/AnimSequenceBase.h
+++ b/Engine/Source/Runtime/Engine/Classes/Animation/AnimSequenceBase.h
@@ -63,7 +63,7 @@ namespace ENotifyFilterType
 USTRUCT()
 struct FAnimNotifyEvent : public FAnimLinkableElement
 {
-	GENERATED_USTRUCT_BODY()
+	GENERATED_BODY()
 
 	/** The user requested time for this notify */
 	UPROPERTY()
@@ -193,7 +193,7 @@ struct FAnimNotifyEvent : public FAnimLinkableElement
 USTRUCT()
 struct FAnimNotifyTrack
 {
-	GENERATED_USTRUCT_BODY()
+	GENERATED_BODY()
 
 	UPROPERTY()
 	FName TrackName;
@@ -247,7 +247,7 @@ enum EAnimCurveFlags
 USTRUCT()
 struct FAnimCurveBase
 {
-	GENERATED_USTRUCT_BODY()
+	GENERATED_BODY()
 
 	// Last observed name of the curve. We store this so we can recover from situations that
 	// mean the skeleton doesn't have a mapped name for our UID (such as a user saving the an
@@ -316,7 +316,7 @@ public:
 USTRUCT()
 struct FFloatCurve : public FAnimCurveBase
 {
-	GENERATED_USTRUCT_BODY()
+	GENERATED_BODY()
 
 	/** Curve data for float. */
 	UPROPERTY()
@@ -348,7 +348,7 @@ struct FFloatCurve : public FAnimCurveBase
 USTRUCT()
 struct FVectorCurve : public FAnimCurveBase
 {
-	GENERATED_USTRUCT_BODY()
+	GENERATED_BODY()
 
 	enum EIndex
 	{
@@ -389,7 +389,7 @@ struct FVectorCurve : public FAnimCurveBase
 USTRUCT()
 struct FTransformCurve: public FAnimCurveBase
 {
-	GENERATED_USTRUCT_BODY()
+	GENERATED_BODY()
 
 	/** Curve data for each transform. */
 	UPROPERTY()
@@ -433,7 +433,7 @@ struct FTransformCurve: public FAnimCurveBase
 USTRUCT()
 struct FRawCurveTracks
 {
-	GENERATED_USTRUCT_BODY()
+	GENERATED_BODY()
 
 	enum ESupportedCurveType
 	{
@@ -564,7 +564,9 @@ enum ETypeAdvanceAnim
 UCLASS(abstract, MinimalAPI, BlueprintType)
 class UAnimSequenceBase : public UAnimationAsset
 {
-	GENERATED_UCLASS_BODY()
+	GENERATED_BODY()
+public:
+	ENGINE_API UAnimSequenceBase(const FObjectInitializer& ObjectInitializer = FObjectInitializer::Get());
 
 	/** Animation notifies, sorted by time (earliest notification first). */
 	UPROPERTY()
diff --git a/Engine/Source/Runtime/Engine/Classes/Animation/AnimSet.h b/Engine/Source/Runtime/Engine/Classes/Animation/AnimSet.h
--- a/Engine/Source/Runtime/Engine/Classes/Animation/AnimSet.h
+++ b/Engine/Source/Runtime/Engine/Classes/Animation/AnimSet.h
@@ -13,7 +13,7 @@
 USTRUCT()
 struct FAnimSetMeshLinkup
 {
-	GENERATED_USTRUCT_BODY()
+	GENERATED_BODY()
 
 	/** 
 	 * Mapping table. Size must be same as size of SkelMesh reference skeleton. 
@@ -33,7 +33,9 @@ struct FAnimSetMeshLinkup
 UCLASS(hidecategories=Object, MinimalAPI)
 class UAnimSet : public UObject
 {
-	GENERATED_UCLASS_BODY()
+	GENERATED_BODY()
+public:
+	ENGINE_API UAnimSet(const FObjectInitializer& ObjectInitializer = FObjectInitializer::Get());
 
 private:
 	/** 
diff --git a/Engine/Source/Runtime/Engine/Classes/Animation/AnimSingleNodeInstance.h b/Engine/Source/Runtime/Engine/Classes/Animation/AnimSingleNodeInstance.h
--- a/Engine/Source/Runtime/Engine/Classes/Animation/AnimSingleNodeInstance.h
+++ b/Engine/Source/Runtime/Engine/Classes/Animation/AnimSingleNodeInstance.h
@@ -15,7 +15,9 @@ DECLARE_DYNAMIC_DELEGATE(FPostEvaluateAnimEvent);
 UCLASS(transient, NotBlueprintable)
 class ENGINE_API UAnimSingleNodeInstance : public UAnimInstance
 {
-	GENERATED_UCLASS_BODY()
+	GENERATED_BODY()
+public:
+	UAnimSingleNodeInstance(const FObjectInitializer& ObjectInitializer = FObjectInitializer::Get());
 
 	/** Current Asset being played **/
 	UPROPERTY(transient)
diff --git a/Engine/Source/Runtime/Engine/Classes/Animation/AnimStateMachineTypes.h b/Engine/Source/Runtime/Engine/Classes/Animation/AnimStateMachineTypes.h
--- a/Engine/Source/Runtime/Engine/Classes/Animation/AnimStateMachineTypes.h
+++ b/Engine/Source/Runtime/Engine/Classes/Animation/AnimStateMachineTypes.h
@@ -30,7 +30,7 @@ namespace ETransitionLogicType
 USTRUCT()
 struct FAnimationTransitionRule
 {
-	GENERATED_USTRUCT_BODY()
+	GENERATED_BODY()
 
 	UPROPERTY()
 	FName RuleToExecute;
@@ -57,7 +57,7 @@ struct FAnimationTransitionRule
 USTRUCT()
 struct FAnimationStateBase
 {
-	GENERATED_USTRUCT_BODY()
+	GENERATED_BODY()
 
 	// The name of this state
 	UPROPERTY()
@@ -71,7 +71,7 @@ struct FAnimationStateBase
 USTRUCT()
 struct FAnimationState : public FAnimationStateBase
 {
-	GENERATED_USTRUCT_BODY()
+	GENERATED_BODY()
 
 	// Set of legal transitions out of this state; already in priority order
 	UPROPERTY()
@@ -105,7 +105,7 @@ struct FAnimationState : public FAnimationStateBase
 USTRUCT()
 struct FAnimationTransitionBetweenStates : public FAnimationStateBase
 {
-	GENERATED_USTRUCT_BODY()
+	GENERATED_BODY()
 
 	// Transition-only: State being transitioned from
 	UPROPERTY()
@@ -157,7 +157,7 @@ struct FAnimationTransitionBetweenStates : public FAnimationStateBase
 USTRUCT()
 struct FBakedStateExitTransition
 {
-	GENERATED_USTRUCT_BODY()
+	GENERATED_BODY()
 
 	// The node property index for this rule
 	UPROPERTY()
@@ -197,7 +197,7 @@ struct FBakedStateExitTransition
 USTRUCT()
 struct FBakedAnimationState
 {
-	GENERATED_USTRUCT_BODY()
+	GENERATED_BODY()
 
 	// The name of this state
 	UPROPERTY()
@@ -240,7 +240,7 @@ public:
 USTRUCT()
 struct FBakedAnimationStateMachine
 {
-	GENERATED_USTRUCT_BODY()
+	GENERATED_BODY()
 
 	// Name of this machine (primarily for debugging purposes)
 	UPROPERTY()
@@ -275,6 +275,8 @@ public:
 UCLASS()
 class UAnimStateMachineTypes : public UObject
 {
-	GENERATED_UCLASS_BODY()
+	GENERATED_BODY()
+public:
+	UAnimStateMachineTypes(const FObjectInitializer& ObjectInitializer = FObjectInitializer::Get());
 };
 
diff --git a/Engine/Source/Runtime/Engine/Classes/Animation/AnimationAsset.h b/Engine/Source/Runtime/Engine/Classes/Animation/AnimationAsset.h
--- a/Engine/Source/Runtime/Engine/Classes/Animation/AnimationAsset.h
+++ b/Engine/Source/Runtime/Engine/Classes/Animation/AnimationAsset.h
@@ -17,7 +17,7 @@
 USTRUCT(BlueprintType)
 struct FBlendSampleData
 {
-	GENERATED_USTRUCT_BODY()
+	GENERATED_BODY()
 
 	UPROPERTY()
 	int32 SampleDataIndex;
@@ -59,7 +59,7 @@ struct FBlendSampleData
 USTRUCT()
 struct FBlendFilter
 {
-	GENERATED_USTRUCT_BODY()
+	GENERATED_BODY()
 
 	FFIRFilterTimeBased FilterPerAxis[3];
 
@@ -113,7 +113,7 @@ namespace ERootMotionMode
 USTRUCT()
 struct FAnimExtractContext
 {
-	GENERATED_USTRUCT_BODY()
+	GENERATED_BODY()
 
 	/** Is Root Motion being extracted? */
 	UPROPERTY()
@@ -149,7 +149,7 @@ struct FAnimExtractContext
 USTRUCT()
 struct FAnimTickRecord
 {
-	GENERATED_USTRUCT_BODY()
+	GENERATED_BODY()
 
 	UPROPERTY()
 	class UAnimationAsset* SourceAsset;
@@ -187,7 +187,7 @@ namespace EAnimGroupRole
 USTRUCT()
 struct FAnimGroupInstance
 {
-	GENERATED_USTRUCT_BODY()
+	GENERATED_BODY()
 
 public:
 	// The list of animation players in this group which are going to be evaluated this frame
@@ -240,7 +240,7 @@ public:
 USTRUCT()
 struct FRootMotionMovementParams
 {
-	GENERATED_USTRUCT_BODY()
+	GENERATED_BODY()
 
 	UPROPERTY()
 	bool bHasRootMotion;
@@ -411,7 +411,7 @@ private:
 USTRUCT()
 struct FAnimationGroupReference
 {
-	GENERATED_USTRUCT_BODY()
+	GENERATED_BODY()
 	
 	// The name of the group
 	UPROPERTY(EditAnywhere, Category=Settings)
@@ -430,7 +430,9 @@ struct FAnimationGroupReference
 UCLASS(abstract, MinimalAPI)
 class UAnimationAsset : public UObject
 {
-	GENERATED_UCLASS_BODY()
+	GENERATED_BODY()
+public:
+	ENGINE_API UAnimationAsset(const FObjectInitializer& ObjectInitializer = FObjectInitializer::Get());
 
 private:
 	/** Pointer to the Skeleton this asset can be played on .	*/
diff --git a/Engine/Source/Runtime/Engine/Classes/Animation/AnimationSettings.h b/Engine/Source/Runtime/Engine/Classes/Animation/AnimationSettings.h
--- a/Engine/Source/Runtime/Engine/Classes/Animation/AnimationSettings.h
+++ b/Engine/Source/Runtime/Engine/Classes/Animation/AnimationSettings.h
@@ -16,7 +16,9 @@ UCLASS(config=Engine, defaultconfig)
 class ENGINE_API UAnimationSettings
 	: public UObject
 {
-	GENERATED_UCLASS_BODY()
+	GENERATED_BODY()
+public:
+	UAnimationSettings(const FObjectInitializer& ObjectInitializer = FObjectInitializer::Get());
 
 	// compression upgrade version
 	UPROPERTY(config, VisibleAnywhere, Category = Compression)
diff --git a/Engine/Source/Runtime/Engine/Classes/Animation/BlendSpace.h b/Engine/Source/Runtime/Engine/Classes/Animation/BlendSpace.h
--- a/Engine/Source/Runtime/Engine/Classes/Animation/BlendSpace.h
+++ b/Engine/Source/Runtime/Engine/Classes/Animation/BlendSpace.h
@@ -12,7 +12,9 @@
 UCLASS(config=Engine, hidecategories=Object, MinimalAPI, BlueprintType)
 class UBlendSpace : public UBlendSpaceBase
 {
-	GENERATED_UCLASS_BODY()
+	GENERATED_BODY()
+public:
+	ENGINE_API UBlendSpace(const FObjectInitializer& ObjectInitializer = FObjectInitializer::Get());
 
 public:
 
diff --git a/Engine/Source/Runtime/Engine/Classes/Animation/BlendSpace1D.h b/Engine/Source/Runtime/Engine/Classes/Animation/BlendSpace1D.h
--- a/Engine/Source/Runtime/Engine/Classes/Animation/BlendSpace1D.h
+++ b/Engine/Source/Runtime/Engine/Classes/Animation/BlendSpace1D.h
@@ -13,7 +13,9 @@
 UCLASS(config=Engine, hidecategories=Object, MinimalAPI, BlueprintType)
 class UBlendSpace1D : public UBlendSpaceBase
 {
-	GENERATED_UCLASS_BODY()
+	GENERATED_BODY()
+public:
+	ENGINE_API UBlendSpace1D(const FObjectInitializer& ObjectInitializer = FObjectInitializer::Get());
 
 public:
 
diff --git a/Engine/Source/Runtime/Engine/Classes/Animation/BlendSpaceBase.h b/Engine/Source/Runtime/Engine/Classes/Animation/BlendSpaceBase.h
--- a/Engine/Source/Runtime/Engine/Classes/Animation/BlendSpaceBase.h
+++ b/Engine/Source/Runtime/Engine/Classes/Animation/BlendSpaceBase.h
@@ -23,7 +23,7 @@ enum EBlendSpaceAxis
 USTRUCT()
 struct FInterpolationParameter
 {
-	GENERATED_USTRUCT_BODY()
+	GENERATED_BODY()
 
 	/** Interpolation Time for input, when it gets input, it will use this time to interpolate to target, used for smoother interpolation **/
 	UPROPERTY(EditAnywhere, Category=Parameter)
@@ -37,7 +37,7 @@ struct FInterpolationParameter
 USTRUCT()
 struct FBlendParameter
 {
-	GENERATED_USTRUCT_BODY()
+	GENERATED_BODY()
 
 	UPROPERTY(EditAnywhere, Category=BlendParameter)
 	FString DisplayName;
@@ -78,7 +78,7 @@ struct FBlendParameter
 USTRUCT()
 struct FBlendSample
 {
-	GENERATED_USTRUCT_BODY()
+	GENERATED_BODY()
 
 	UPROPERTY(EditAnywhere, Category=BlendSample)
 	class UAnimSequence* Animation;
@@ -115,7 +115,7 @@ struct FBlendSample
 USTRUCT()
 struct FEditorElement
 {
-	GENERATED_USTRUCT_BODY()
+	GENERATED_BODY()
 
 	UPROPERTY(EditAnywhere, Category=EditorElement)
 	int32 Indices[3];    /*MAX_VERTICES @fixmeconst*/
@@ -144,7 +144,7 @@ struct FEditorElement
 USTRUCT()
 struct FGridBlendSample
 {
-	GENERATED_USTRUCT_BODY()
+	GENERATED_BODY()
 
 	UPROPERTY()
 	struct FEditorElement GridElement;
@@ -163,7 +163,7 @@ struct FGridBlendSample
 USTRUCT()
 struct FPerBoneInterpolation
 {
-	GENERATED_USTRUCT_BODY()
+	GENERATED_BODY()
 
 	UPROPERTY(EditAnywhere, Category=FPerBoneInterpolation)
 	FBoneReference BoneReference;
@@ -198,7 +198,9 @@ namespace ENotifyTriggerMode
 UCLASS(config=Engine, hidecategories=Object, MinimalAPI, BlueprintType)
 class UBlendSpaceBase : public UAnimationAsset
 {
-	GENERATED_UCLASS_BODY()
+	GENERATED_BODY()
+public:
+	ENGINE_API UBlendSpaceBase(const FObjectInitializer& ObjectInitializer = FObjectInitializer::Get());
 
 protected:
 	/** Blend Parameters for each axis. **/
diff --git a/Engine/Source/Runtime/Engine/Classes/Animation/BoneControllers/AnimNode_BoneDrivenController.h b/Engine/Source/Runtime/Engine/Classes/Animation/BoneControllers/AnimNode_BoneDrivenController.h
--- a/Engine/Source/Runtime/Engine/Classes/Animation/BoneControllers/AnimNode_BoneDrivenController.h
+++ b/Engine/Source/Runtime/Engine/Classes/Animation/BoneControllers/AnimNode_BoneDrivenController.h
@@ -26,7 +26,7 @@ namespace EComponentType
 USTRUCT()
 struct ENGINE_API FAnimNode_BoneDrivenController : public FAnimNode_SkeletalControlBase
 {
-	GENERATED_USTRUCT_BODY()
+	GENERATED_BODY()
 
 	// Bone to use as controller input
 	UPROPERTY(EditAnywhere, BlueprintReadWrite, Category = Bones)
@@ -76,4 +76,4 @@ protected:
 	// FAnimNode_SkeletalControlBase interface
 	virtual void InitializeBoneReferences(const FBoneContainer& RequiredBones) override;
 	//////////////////////////////////////////////////////////////////////////
-};
\ No newline at end of file
+};
diff --git a/Engine/Source/Runtime/Engine/Classes/Animation/BoneControllers/AnimNode_Fabrik.h b/Engine/Source/Runtime/Engine/Classes/Animation/BoneControllers/AnimNode_Fabrik.h
--- a/Engine/Source/Runtime/Engine/Classes/Animation/BoneControllers/AnimNode_Fabrik.h
+++ b/Engine/Source/Runtime/Engine/Classes/Animation/BoneControllers/AnimNode_Fabrik.h
@@ -13,7 +13,7 @@
 USTRUCT()
 struct FABRIKChainLink
 {
-	GENERATED_USTRUCT_BODY()
+	GENERATED_BODY()
 
 public:
 	/** Position of bone in component space. */
@@ -53,7 +53,7 @@ public:
 USTRUCT()
 struct ENGINE_API FAnimNode_Fabrik : public FAnimNode_SkeletalControlBase
 {
-	GENERATED_USTRUCT_BODY()
+	GENERATED_BODY()
 
 	/** Coordinates for target location of tip bone - if EffectorLocationSpace is bone, this is the offset from Target Bone to use as target location*/
 	UPROPERTY(EditAnywhere, BlueprintReadWrite, Category = EndEffector, meta = (PinShownByDefault))
@@ -106,4 +106,4 @@ private:
 
 	// Convenience function to get current (pre-translation iteration) component space location of bone by bone index
 	FVector GetCurrentLocation(FA2CSPose & MeshBases, int32 const & BoneIndex);
-};
\ No newline at end of file
+};
diff --git a/Engine/Source/Runtime/Engine/Classes/Animation/BoneControllers/AnimNode_Trail.h b/Engine/Source/Runtime/Engine/Classes/Animation/BoneControllers/AnimNode_Trail.h
--- a/Engine/Source/Runtime/Engine/Classes/Animation/BoneControllers/AnimNode_Trail.h
+++ b/Engine/Source/Runtime/Engine/Classes/Animation/BoneControllers/AnimNode_Trail.h
@@ -12,7 +12,7 @@
 USTRUCT()
 struct ENGINE_API FAnimNode_Trail : public FAnimNode_SkeletalControlBase
 {
-	GENERATED_USTRUCT_BODY()
+	GENERATED_BODY()
 
 	/** Reference to the active bone in the hierarchy to modify. */
 	UPROPERTY(EditAnywhere, BlueprintReadWrite, Category=Trail)
@@ -81,4 +81,4 @@ private:
 	// End of FAnimNode_SkeletalControlBase interface
 
 	FVector GetAlignVector(EAxis::Type AxisOption, bool bInvert);
-};
\ No newline at end of file
+};
diff --git a/Engine/Source/Runtime/Engine/Classes/Animation/Rig.h b/Engine/Source/Runtime/Engine/Classes/Animation/Rig.h
--- a/Engine/Source/Runtime/Engine/Classes/Animation/Rig.h
+++ b/Engine/Source/Runtime/Engine/Classes/Animation/Rig.h
@@ -17,7 +17,7 @@
 USTRUCT()
 struct FNode
 {
-	GENERATED_USTRUCT_BODY()
+	GENERATED_BODY()
 
 	/** Name of the original node. We don't allow to change this. This is used for identity.**/
 	UPROPERTY(VisibleAnywhere, Category="FNode")
@@ -86,7 +86,7 @@ namespace EConstraintTransform
 USTRUCT()
 struct FRigTransformConstraint
 {
-	GENERATED_USTRUCT_BODY()
+	GENERATED_BODY()
 
 	/** What transform type **/
 	UPROPERTY(/*EditAnywhere, Category="FTransformBaseConstraint"*/)
@@ -105,7 +105,7 @@ struct FRigTransformConstraint
 USTRUCT()
 struct FTransformBaseConstraint
 {
-	GENERATED_USTRUCT_BODY()
+	GENERATED_BODY()
 
 	/** What transform type **/
 	UPROPERTY(EditAnywhere, Category="FTransformBaseConstraint")
@@ -117,7 +117,7 @@ struct FTransformBaseConstraint
 USTRUCT()
 struct FTransformBase
 {
-	GENERATED_USTRUCT_BODY()
+	GENERATED_BODY()
 
 	UPROPERTY(EditAnywhere, Category="FTransformBase")
 	FName					Node;
@@ -136,7 +136,9 @@ DECLARE_DELEGATE_RetVal_OneParam(int32, FGetParentIndex, FName);
 UCLASS(hidecategories=Object, MinimalAPI)
 class URig : public UObject
 {
-	GENERATED_UCLASS_BODY()
+	GENERATED_BODY()
+public:
+	ENGINE_API URig(const FObjectInitializer& ObjectInitializer = FObjectInitializer::Get());
 
 private:
 
diff --git a/Engine/Source/Runtime/Engine/Classes/Animation/SkeletalMeshActor.h b/Engine/Source/Runtime/Engine/Classes/Animation/SkeletalMeshActor.h
--- a/Engine/Source/Runtime/Engine/Classes/Animation/SkeletalMeshActor.h
+++ b/Engine/Source/Runtime/Engine/Classes/Animation/SkeletalMeshActor.h
@@ -15,7 +15,9 @@
 UCLASS(ClassGroup=ISkeletalMeshes, Blueprintable, ComponentWrapperClass, ConversionRoot, meta=(ChildCanTick))
 class ENGINE_API ASkeletalMeshActor : public AActor, public IMatineeAnimInterface
 {
-	GENERATED_UCLASS_BODY()
+	GENERATED_BODY()
+public:
+	ASkeletalMeshActor(const FObjectInitializer& ObjectInitializer = FObjectInitializer::Get());
 
 	/** Whether or not this actor should respond to anim notifies - CURRENTLY ONLY AFFECTS PlayParticleEffect NOTIFIES**/
 	UPROPERTY(EditAnywhere, BlueprintReadWrite, Category=Animation, AdvancedDisplay)
diff --git a/Engine/Source/Runtime/Engine/Classes/Animation/Skeleton.h b/Engine/Source/Runtime/Engine/Classes/Animation/Skeleton.h
--- a/Engine/Source/Runtime/Engine/Classes/Animation/Skeleton.h
+++ b/Engine/Source/Runtime/Engine/Classes/Animation/Skeleton.h
@@ -17,7 +17,7 @@ class USkeletalMesh;
 USTRUCT()
 struct FSkeletonToMeshLinkup
 {
-	GENERATED_USTRUCT_BODY()
+	GENERATED_BODY()
 
 	/** 
 	 * Mapping table. Size must be same as size of bone tree (not Mesh Ref Pose). 
@@ -58,7 +58,7 @@ namespace EBoneTranslationRetargetingMode
 USTRUCT()
 struct FBoneNode
 {
-	GENERATED_USTRUCT_BODY()
+	GENERATED_BODY()
 
 	/** Name of bone, this is the search criteria to match with mesh bone. This will be NAME_None if deleted **/
 	UPROPERTY()
@@ -90,7 +90,7 @@ struct FBoneNode
 USTRUCT()
 struct FReferencePose
 {
-	GENERATED_USTRUCT_BODY()
+	GENERATED_BODY()
 
 	UPROPERTY()
 	FName	PoseName;
@@ -117,7 +117,7 @@ struct FReferencePose
 USTRUCT()
 struct FBoneReductionSetting
 {
-	GENERATED_USTRUCT_BODY()
+	GENERATED_BODY()
 
 	UPROPERTY()
 	TArray<FName> BonesToRemove;
@@ -147,7 +147,7 @@ struct FBoneReductionSetting
 USTRUCT()
 struct FNameMapping
 {
-	GENERATED_USTRUCT_BODY()
+	GENERATED_BODY()
 
 	UPROPERTY()
 	FName NodeName;
@@ -177,7 +177,7 @@ struct FNameMapping
 USTRUCT()
 struct FRigConfiguration
 {
-	GENERATED_USTRUCT_BODY()
+	GENERATED_BODY()
 
 	UPROPERTY()
 	class URig * Rig;
@@ -190,7 +190,7 @@ struct FRigConfiguration
 USTRUCT()
 struct FAnimSlotGroup
 {
-	GENERATED_USTRUCT_BODY()
+	GENERATED_BODY()
 
 public:
 	static ENGINE_API const FName DefaultGroupName;
@@ -223,7 +223,9 @@ public:
 UCLASS(hidecategories=Object, MinimalAPI)
 class USkeleton : public UObject
 {
-	GENERATED_UCLASS_BODY()
+	GENERATED_BODY()
+public:
+	ENGINE_API USkeleton(const FObjectInitializer& ObjectInitializer = FObjectInitializer::Get());
 
 private:
 	/** Skeleton bone tree - each contains name and parent index**/
diff --git a/Engine/Source/Runtime/Engine/Classes/Animation/SmartName.h b/Engine/Source/Runtime/Engine/Classes/Animation/SmartName.h
--- a/Engine/Source/Runtime/Engine/Classes/Animation/SmartName.h
+++ b/Engine/Source/Runtime/Engine/Classes/Animation/SmartName.h
@@ -7,7 +7,7 @@
 USTRUCT()
 struct ENGINE_API FSmartNameMapping
 {
-	GENERATED_USTRUCT_BODY();
+	GENERATED_BODY()
 
 	FSmartNameMapping();
 
@@ -83,7 +83,7 @@ private:
 USTRUCT()
 struct ENGINE_API FSmartNameContainer
 {
-	GENERATED_USTRUCT_BODY();
+	GENERATED_BODY()
 
 	// Add a new smartname container with the provided name
 	void AddContainer(FName NewContainerName);
@@ -99,4 +99,4 @@ struct ENGINE_API FSmartNameContainer
 
 private:
 	TMap<FName, FSmartNameMapping> NameMappings;	// List of smartname mappings
-};
\ No newline at end of file
+};
diff --git a/Engine/Source/Runtime/Engine/Classes/Animation/VertexAnim/MorphTarget.h b/Engine/Source/Runtime/Engine/Classes/Animation/VertexAnim/MorphTarget.h
--- a/Engine/Source/Runtime/Engine/Classes/Animation/VertexAnim/MorphTarget.h
+++ b/Engine/Source/Runtime/Engine/Classes/Animation/VertexAnim/MorphTarget.h
@@ -66,7 +66,9 @@ struct FMorphTargetLODModel
 UCLASS(hidecategories=Object, MinimalAPI)
 class UMorphTarget : public UVertexAnimBase
 {
-	GENERATED_UCLASS_BODY()
+	GENERATED_BODY()
+public:
+	ENGINE_API UMorphTarget(const FObjectInitializer& ObjectInitializer = FObjectInitializer::Get());
 
 public:
 	/** morph mesh vertex data for each LOD */
diff --git a/Engine/Source/Runtime/Engine/Classes/Animation/VertexAnim/VertexAnimBase.h b/Engine/Source/Runtime/Engine/Classes/Animation/VertexAnim/VertexAnimBase.h
--- a/Engine/Source/Runtime/Engine/Classes/Animation/VertexAnim/VertexAnimBase.h
+++ b/Engine/Source/Runtime/Engine/Classes/Animation/VertexAnim/VertexAnimBase.h
@@ -47,7 +47,9 @@ struct FVertexAnimEvalStateBase
 UCLASS(hidecategories=Object, MinimalAPI, abstract)
 class UVertexAnimBase : public UObject
 {
-	GENERATED_UCLASS_BODY()
+	GENERATED_BODY()
+public:
+	ENGINE_API UVertexAnimBase(const FObjectInitializer& ObjectInitializer = FObjectInitializer::Get());
 
 	/** USkeletalMesh that this vertex animation works on. */
 	UPROPERTY(AssetRegistrySearchable)
diff --git a/Engine/Source/Runtime/Engine/Classes/Animation/VertexAnim/VertexAnimation.h b/Engine/Source/Runtime/Engine/Classes/Animation/VertexAnim/VertexAnimation.h
--- a/Engine/Source/Runtime/Engine/Classes/Animation/VertexAnim/VertexAnimation.h
+++ b/Engine/Source/Runtime/Engine/Classes/Animation/VertexAnim/VertexAnimation.h
@@ -27,7 +27,9 @@ struct FVertexAnimFrame
 UCLASS(hidecategories=Object, MinimalAPI)
 class UVertexAnimation : public UVertexAnimBase
 {
-	GENERATED_UCLASS_BODY()
+	GENERATED_BODY()
+public:
+	ENGINE_API UVertexAnimation(const FObjectInitializer& ObjectInitializer = FObjectInitializer::Get());
 
 	/** Number of verts animated by this animation, should be size of Deltas array in each frame */
 	UPROPERTY()
diff --git a/Engine/Source/Runtime/Engine/Classes/Atmosphere/AtmosphericFog.h b/Engine/Source/Runtime/Engine/Classes/Atmosphere/AtmosphericFog.h
--- a/Engine/Source/Runtime/Engine/Classes/Atmosphere/AtmosphericFog.h
+++ b/Engine/Source/Runtime/Engine/Classes/Atmosphere/AtmosphericFog.h
@@ -12,7 +12,9 @@
 UCLASS(showcategories=(Movement, Rendering, "Utilities|Transformation", "Input|MouseInput", "Input|TouchInput"), ClassGroup=Fog, hidecategories=(Info,Object,Input), MinimalAPI)
 class AAtmosphericFog : public AInfo
 {
-	GENERATED_UCLASS_BODY()
+	GENERATED_BODY()
+public:
+	ENGINE_API AAtmosphericFog(const FObjectInitializer& ObjectInitializer = FObjectInitializer::Get());
 
 private_subobject:
 	/** Main fog component */
diff --git a/Engine/Source/Runtime/Engine/Classes/Atmosphere/AtmosphericFogComponent.h b/Engine/Source/Runtime/Engine/Classes/Atmosphere/AtmosphericFogComponent.h
--- a/Engine/Source/Runtime/Engine/Classes/Atmosphere/AtmosphericFogComponent.h
+++ b/Engine/Source/Runtime/Engine/Classes/Atmosphere/AtmosphericFogComponent.h
@@ -7,7 +7,7 @@
 USTRUCT()
 struct FAtmospherePrecomputeParameters
 {
-	GENERATED_USTRUCT_BODY()
+	GENERATED_BODY()
 	
 	/** Rayleigh scattering density height scale, ranges from [0...1] */
 	UPROPERTY(EditAnywhere, BlueprintReadOnly, Category=AtmosphereParam)
@@ -89,7 +89,9 @@ struct FAtmospherePrecomputeParameters
 UCLASS(ClassGroup=Rendering, collapsecategories, hidecategories=(Object, Mobility, Activation, "Components|Activation"), editinlinenew, meta=(BlueprintSpawnableComponent), MinimalAPI)
 class UAtmosphericFogComponent : public USceneComponent
 {
-	GENERATED_UCLASS_BODY()
+	GENERATED_BODY()
+public:
+	ENGINE_API UAtmosphericFogComponent(const FObjectInitializer& ObjectInitializer = FObjectInitializer::Get());
 
 	/** Global scattering factor. */
 	UPROPERTY(EditAnywhere, BlueprintReadWrite, interp, Category=Atmosphere)
diff --git a/Engine/Source/Runtime/Engine/Classes/Camera/CameraActor.h b/Engine/Source/Runtime/Engine/Classes/Camera/CameraActor.h
--- a/Engine/Source/Runtime/Engine/Classes/Camera/CameraActor.h
+++ b/Engine/Source/Runtime/Engine/Classes/Camera/CameraActor.h
@@ -12,7 +12,9 @@
 UCLASS(ClassGroup=Common, hideCategories=(Input, Rendering), showcategories=("Input|MouseInput", "Input|TouchInput"), Blueprintable)
 class ENGINE_API ACameraActor : public AActor
 {
-	GENERATED_UCLASS_BODY()
+	GENERATED_BODY()
+public:
+	ACameraActor(const FObjectInitializer& ObjectInitializer = FObjectInitializer::Get());
 
 private:
 
diff --git a/Engine/Source/Runtime/Engine/Classes/Camera/CameraAnim.h b/Engine/Source/Runtime/Engine/Classes/Camera/CameraAnim.h
--- a/Engine/Source/Runtime/Engine/Classes/Camera/CameraAnim.h
+++ b/Engine/Source/Runtime/Engine/Classes/Camera/CameraAnim.h
@@ -10,7 +10,9 @@
 UCLASS(notplaceable, MinimalAPI)
 class UCameraAnim : public UObject
 {
-	GENERATED_UCLASS_BODY()
+	GENERATED_BODY()
+public:
+	ENGINE_API UCameraAnim(const FObjectInitializer& ObjectInitializer = FObjectInitializer::Get());
 
 	/** The UInterpGroup that holds our actual interpolation data. */
 	UPROPERTY()
diff --git a/Engine/Source/Runtime/Engine/Classes/Camera/CameraAnimInst.h b/Engine/Source/Runtime/Engine/Classes/Camera/CameraAnimInst.h
--- a/Engine/Source/Runtime/Engine/Classes/Camera/CameraAnimInst.h
+++ b/Engine/Source/Runtime/Engine/Classes/Camera/CameraAnimInst.h
@@ -13,7 +13,9 @@
 UCLASS(notplaceable, BlueprintType, transient)
 class ENGINE_API UCameraAnimInst : public UObject
 {
-	GENERATED_UCLASS_BODY()
+	GENERATED_BODY()
+public:
+	UCameraAnimInst(const FObjectInitializer& ObjectInitializer = FObjectInitializer::Get());
 
 	/** which CameraAnim this is an instance of */
 	UPROPERTY()
diff --git a/Engine/Source/Runtime/Engine/Classes/Camera/CameraComponent.h b/Engine/Source/Runtime/Engine/Classes/Camera/CameraComponent.h
--- a/Engine/Source/Runtime/Engine/Classes/Camera/CameraComponent.h
+++ b/Engine/Source/Runtime/Engine/Classes/Camera/CameraComponent.h
@@ -12,7 +12,9 @@
 UCLASS(HideCategories=(Mobility, Rendering, LOD), Blueprintable, ClassGroup=Camera, meta=(BlueprintSpawnableComponent), MinimalAPI)
 class UCameraComponent : public USceneComponent
 {
-	GENERATED_UCLASS_BODY()
+	GENERATED_BODY()
+public:
+	ENGINE_API UCameraComponent(const FObjectInitializer& ObjectInitializer = FObjectInitializer::Get());
 
 	/** The horizontal field of view (in degrees) in perspective mode (ignored in Orthographic mode) */
 	UPROPERTY(Interp, EditAnywhere, BlueprintReadWrite, Category=CameraSettings, meta=(UIMin = "5.0", UIMax = "170", ClampMin = "0.001", ClampMax = "360.0"))
diff --git a/Engine/Source/Runtime/Engine/Classes/Camera/CameraModifier.h b/Engine/Source/Runtime/Engine/Classes/Camera/CameraModifier.h
--- a/Engine/Source/Runtime/Engine/Classes/Camera/CameraModifier.h
+++ b/Engine/Source/Runtime/Engine/Classes/Camera/CameraModifier.h
@@ -13,7 +13,9 @@
 UCLASS()
 class ENGINE_API UCameraModifier : public UObject
 {
-	GENERATED_UCLASS_BODY()
+	GENERATED_BODY()
+public:
+	UCameraModifier(const FObjectInitializer& ObjectInitializer = FObjectInitializer::Get());
 
 protected:
 	/** If true, do not apply this modifier to the camera. */
diff --git a/Engine/Source/Runtime/Engine/Classes/Camera/CameraModifier_CameraShake.h b/Engine/Source/Runtime/Engine/Classes/Camera/CameraModifier_CameraShake.h
--- a/Engine/Source/Runtime/Engine/Classes/Camera/CameraModifier_CameraShake.h
+++ b/Engine/Source/Runtime/Engine/Classes/Camera/CameraModifier_CameraShake.h
@@ -16,7 +16,7 @@
 USTRUCT()
 struct ENGINE_API FCameraShakeInstance
 {
-	GENERATED_USTRUCT_BODY()
+	GENERATED_BODY()
 
 	/** The source camera shake data. */
 	UPROPERTY()
@@ -100,7 +100,9 @@ struct ENGINE_API FCameraShakeInstance
 UCLASS(config=Camera)
 class ENGINE_API UCameraModifier_CameraShake : public UCameraModifier
 {
-	GENERATED_UCLASS_BODY()
+	GENERATED_BODY()
+public:
+	UCameraModifier_CameraShake(const FObjectInitializer& ObjectInitializer = FObjectInitializer::Get());
 
 public:
 	/** List of active CameraShakes */
diff --git a/Engine/Source/Runtime/Engine/Classes/Camera/CameraShake.h b/Engine/Source/Runtime/Engine/Classes/Camera/CameraShake.h
--- a/Engine/Source/Runtime/Engine/Classes/Camera/CameraShake.h
+++ b/Engine/Source/Runtime/Engine/Classes/Camera/CameraShake.h
@@ -23,7 +23,7 @@ enum EInitialOscillatorOffset
 USTRUCT()
 struct FFOscillator
 {
-	GENERATED_USTRUCT_BODY()
+	GENERATED_BODY()
 
 	/** Amplitude of the sinusoidal oscillation. */
 	UPROPERTY(EditAnywhere, Category=FOscillator)
@@ -48,7 +48,7 @@ struct FFOscillator
 USTRUCT()
 struct FROscillator
 {
-	GENERATED_USTRUCT_BODY()
+	GENERATED_BODY()
 
 	/** Pitch oscillation. */
 	UPROPERTY(EditAnywhere, Category=ROscillator)
@@ -68,7 +68,7 @@ struct FROscillator
 USTRUCT()
 struct FVOscillator
 {
-	GENERATED_USTRUCT_BODY()
+	GENERATED_BODY()
 
 	/** Oscillation in the X axis. */
 	UPROPERTY(EditAnywhere, Category = VOscillator)
@@ -103,7 +103,9 @@ struct FVOscillator
 UCLASS(Blueprintable, editinlinenew)
 class ENGINE_API UCameraShake : public UObject
 {
-	GENERATED_UCLASS_BODY()
+	GENERATED_BODY()
+public:
+	UCameraShake(const FObjectInitializer& ObjectInitializer = FObjectInitializer::Get());
 
 	/** 
 	 *  If true to only allow a single instance of this shake to play at any given time. 
diff --git a/Engine/Source/Runtime/Engine/Classes/Camera/PlayerCameraManager.h b/Engine/Source/Runtime/Engine/Classes/Camera/PlayerCameraManager.h
--- a/Engine/Source/Runtime/Engine/Classes/Camera/PlayerCameraManager.h
+++ b/Engine/Source/Runtime/Engine/Classes/Camera/PlayerCameraManager.h
@@ -34,7 +34,7 @@ enum EViewTargetBlendFunction
 USTRUCT()
 struct FCameraCacheEntry
 {
-	GENERATED_USTRUCT_BODY()
+	GENERATED_BODY()
 public:
 
 	/** World time this entry was created. */
@@ -54,7 +54,7 @@ public:
 USTRUCT()
 struct ENGINE_API FTViewTarget
 {
-	GENERATED_USTRUCT_BODY()
+	GENERATED_BODY()
 public:
 
 	/** Target Actor used to compute POV */
@@ -93,7 +93,7 @@ public:
 USTRUCT()
 struct FViewTargetTransitionParams
 {
-	GENERATED_USTRUCT_BODY()
+	GENERATED_BODY()
 public:
 
 	/** Total duration of blend to pending view target. 0 means no blending. */
@@ -163,7 +163,9 @@ public:
 UCLASS(notplaceable, transient, BlueprintType, Blueprintable)
 class ENGINE_API APlayerCameraManager : public AActor
 {
-	GENERATED_UCLASS_BODY()
+	GENERATED_BODY()
+public:
+	APlayerCameraManager(const FObjectInitializer& ObjectInitializer = FObjectInitializer::Get());
 
 	/** PlayerController that owns this Camera actor */
 	UPROPERTY()
diff --git a/Engine/Source/Runtime/Engine/Classes/Commandlets/Commandlet.h b/Engine/Source/Runtime/Engine/Classes/Commandlets/Commandlet.h
--- a/Engine/Source/Runtime/Engine/Classes/Commandlets/Commandlet.h
+++ b/Engine/Source/Runtime/Engine/Classes/Commandlets/Commandlet.h
@@ -32,7 +32,9 @@
 UCLASS(abstract, transient,MinimalAPI)
 class UCommandlet : public UObject
 {
-	GENERATED_UCLASS_BODY()
+	GENERATED_BODY()
+public:
+	ENGINE_API UCommandlet(const FObjectInitializer& ObjectInitializer = FObjectInitializer::Get());
 
 	/** Description of the commandlet's purpose */
 	UPROPERTY()
diff --git a/Engine/Source/Runtime/Engine/Classes/Commandlets/PluginCommandlet.h b/Engine/Source/Runtime/Engine/Classes/Commandlets/PluginCommandlet.h
--- a/Engine/Source/Runtime/Engine/Classes/Commandlets/PluginCommandlet.h
+++ b/Engine/Source/Runtime/Engine/Classes/Commandlets/PluginCommandlet.h
@@ -7,7 +7,9 @@
 UCLASS()
 class UPluginCommandlet : public UCommandlet
 {
-	GENERATED_UCLASS_BODY()
+	GENERATED_BODY()
+public:
+	UPluginCommandlet(const FObjectInitializer& ObjectInitializer = FObjectInitializer::Get());
 
 	/** Parsed commandline tokens */
 	TArray<FString> CmdLineTokens;
diff --git a/Engine/Source/Runtime/Engine/Classes/Commandlets/SmokeTestCommandlet.h b/Engine/Source/Runtime/Engine/Classes/Commandlets/SmokeTestCommandlet.h
--- a/Engine/Source/Runtime/Engine/Classes/Commandlets/SmokeTestCommandlet.h
+++ b/Engine/Source/Runtime/Engine/Classes/Commandlets/SmokeTestCommandlet.h
@@ -7,7 +7,9 @@
 UCLASS()
 class USmokeTestCommandlet : public UCommandlet
 {
-	GENERATED_UCLASS_BODY()
+	GENERATED_BODY()
+public:
+	USmokeTestCommandlet(const FObjectInitializer& ObjectInitializer = FObjectInitializer::Get());
 
 
 	// Begin UCommandlet Interface
diff --git a/Engine/Source/Runtime/Engine/Classes/Components/ApplicationLifecycleComponent.h b/Engine/Source/Runtime/Engine/Classes/Components/ApplicationLifecycleComponent.h
--- a/Engine/Source/Runtime/Engine/Classes/Components/ApplicationLifecycleComponent.h
+++ b/Engine/Source/Runtime/Engine/Classes/Components/ApplicationLifecycleComponent.h
@@ -9,7 +9,9 @@
 UCLASS(ClassGroup=Utility, HideCategories=(Activation, "Components|Activation", Collision), meta=(BlueprintSpawnableComponent))
 class ENGINE_API UApplicationLifecycleComponent : public UActorComponent
 {
-	GENERATED_UCLASS_BODY()
+	GENERATED_BODY()
+public:
+	UApplicationLifecycleComponent(const FObjectInitializer& ObjectInitializer = FObjectInitializer::Get());
 
 	DECLARE_DYNAMIC_MULTICAST_DELEGATE(FApplicationLifetimeDelegate);
 
diff --git a/Engine/Source/Runtime/Engine/Classes/Components/ArrowComponent.h b/Engine/Source/Runtime/Engine/Classes/Components/ArrowComponent.h
--- a/Engine/Source/Runtime/Engine/Classes/Components/ArrowComponent.h
+++ b/Engine/Source/Runtime/Engine/Classes/Components/ArrowComponent.h
@@ -10,7 +10,9 @@
 UCLASS(ClassGroup=Utility, hidecategories=(Object,LOD,Physics,Lighting,TextureStreaming,Activation,"Components|Activation",Collision), editinlinenew, meta=(BlueprintSpawnableComponent), MinimalAPI)
 class UArrowComponent : public UPrimitiveComponent
 {
-	GENERATED_UCLASS_BODY()
+	GENERATED_BODY()
+public:
+	ENGINE_API UArrowComponent(const FObjectInitializer& ObjectInitializer = FObjectInitializer::Get());
 
 	UPROPERTY(EditAnywhere, BlueprintReadWrite, Category=ArrowComponent)
 	FColor ArrowColor;
diff --git a/Engine/Source/Runtime/Engine/Classes/Components/AudioComponent.h b/Engine/Source/Runtime/Engine/Classes/Components/AudioComponent.h
--- a/Engine/Source/Runtime/Engine/Classes/Components/AudioComponent.h
+++ b/Engine/Source/Runtime/Engine/Classes/Components/AudioComponent.h
@@ -21,7 +21,7 @@ DECLARE_DYNAMIC_DELEGATE_TwoParams( FOnQueueSubtitles, const TArray<struct FSubt
 USTRUCT()
 struct FAudioComponentParam
 {
-	GENERATED_USTRUCT_BODY()
+	GENERATED_BODY()
 
 	// Name of the parameter
 	UPROPERTY(EditAnywhere, BlueprintReadWrite, Category=AudioComponentParam)
@@ -63,7 +63,9 @@ struct FAudioComponentParam
 UCLASS(ClassGroup=(Audio, Common), hidecategories=(Object, ActorComponent, Physics, Rendering, Mobility, LOD), ShowCategories=Trigger, meta=(BlueprintSpawnableComponent))
 class ENGINE_API UAudioComponent : public USceneComponent
 {
-	GENERATED_UCLASS_BODY()
+	GENERATED_BODY()
+public:
+	UAudioComponent(const FObjectInitializer& ObjectInitializer = FObjectInitializer::Get());
 
 	/** The sound to be played */
 	UPROPERTY(EditAnywhere, BlueprintReadWrite, Category=Sound)
diff --git a/Engine/Source/Runtime/Engine/Classes/Components/BillboardComponent.h b/Engine/Source/Runtime/Engine/Classes/Components/BillboardComponent.h
--- a/Engine/Source/Runtime/Engine/Classes/Components/BillboardComponent.h
+++ b/Engine/Source/Runtime/Engine/Classes/Components/BillboardComponent.h
@@ -13,7 +13,9 @@ struct FConvexVolume;
 UCLASS(ClassGroup=Rendering, collapsecategories, hidecategories=(Object,Activation,"Components|Activation",Physics,Collision,Lighting,LOD,Mesh,PhysicsVolume), editinlinenew, meta=(BlueprintSpawnableComponent),MinimalAPI)
 class UBillboardComponent : public UPrimitiveComponent
 {
-	GENERATED_UCLASS_BODY()
+	GENERATED_BODY()
+public:
+	ENGINE_API UBillboardComponent(const FObjectInitializer& ObjectInitializer = FObjectInitializer::Get());
 
 	UPROPERTY(EditAnywhere, BlueprintReadWrite, Category=Sprite)
 	class UTexture2D* Sprite;
diff --git a/Engine/Source/Runtime/Engine/Classes/Components/BoxComponent.h b/Engine/Source/Runtime/Engine/Classes/Components/BoxComponent.h
--- a/Engine/Source/Runtime/Engine/Classes/Components/BoxComponent.h
+++ b/Engine/Source/Runtime/Engine/Classes/Components/BoxComponent.h
@@ -11,7 +11,9 @@
 UCLASS(ClassGroup="Collision", hidecategories=(Object,LOD,Lighting,TextureStreaming), editinlinenew, meta=(DisplayName="Box Collision", BlueprintSpawnableComponent))
 class ENGINE_API UBoxComponent : public UShapeComponent
 {
-	GENERATED_UCLASS_BODY()
+	GENERATED_BODY()
+public:
+	UBoxComponent(const FObjectInitializer& ObjectInitializer = FObjectInitializer::Get());
 
 protected:
 	/** The extents of the box **/
diff --git a/Engine/Source/Runtime/Engine/Classes/Components/BoxReflectionCaptureComponent.h b/Engine/Source/Runtime/Engine/Classes/Components/BoxReflectionCaptureComponent.h
--- a/Engine/Source/Runtime/Engine/Classes/Components/BoxReflectionCaptureComponent.h
+++ b/Engine/Source/Runtime/Engine/Classes/Components/BoxReflectionCaptureComponent.h
@@ -9,7 +9,9 @@
 UCLASS(hidecategories=(Collision, Object, Physics, SceneComponent), MinimalAPI)
 class UBoxReflectionCaptureComponent : public UReflectionCaptureComponent
 {
-	GENERATED_UCLASS_BODY()
+	GENERATED_BODY()
+public:
+	ENGINE_API UBoxReflectionCaptureComponent(const FObjectInitializer& ObjectInitializer = FObjectInitializer::Get());
 
 	/** Adjust capture transition distance */
 	UPROPERTY(EditAnywhere, BlueprintReadWrite, Category=ReflectionCapture, meta=(UIMin = "1", UIMax = "1000"))
diff --git a/Engine/Source/Runtime/Engine/Classes/Components/BrushComponent.h b/Engine/Source/Runtime/Engine/Classes/Components/BrushComponent.h
--- a/Engine/Source/Runtime/Engine/Classes/Components/BrushComponent.h
+++ b/Engine/Source/Runtime/Engine/Classes/Components/BrushComponent.h
@@ -14,7 +14,9 @@ struct FConvexVolume;
 UCLASS(editinlinenew, MinimalAPI, hidecategories=(Physics, Lighting, LOD, Rendering, TextureStreaming, Transform, Activation, "Components|Activation"), showcategories=(Mobility, "Rendering|Material"))
 class UBrushComponent : public UPrimitiveComponent
 {
-	GENERATED_UCLASS_BODY()
+	GENERATED_BODY()
+public:
+	ENGINE_API UBrushComponent(const FObjectInitializer& ObjectInitializer = FObjectInitializer::Get());
 
 	UPROPERTY()
 	class UModel* Brush;
diff --git a/Engine/Source/Runtime/Engine/Classes/Components/CapsuleComponent.h b/Engine/Source/Runtime/Engine/Classes/Components/CapsuleComponent.h
--- a/Engine/Source/Runtime/Engine/Classes/Components/CapsuleComponent.h
+++ b/Engine/Source/Runtime/Engine/Classes/Components/CapsuleComponent.h
@@ -10,7 +10,9 @@
 UCLASS(ClassGroup="Collision", editinlinenew, hidecategories=(Object,LOD,Lighting,TextureStreaming), meta=(DisplayName="Capsule Collision", BlueprintSpawnableComponent))
 class ENGINE_API UCapsuleComponent : public UShapeComponent
 {
-	GENERATED_UCLASS_BODY()
+	GENERATED_BODY()
+public:
+	UCapsuleComponent(const FObjectInitializer& ObjectInitializer = FObjectInitializer::Get());
 
 protected:
 	/** 
diff --git a/Engine/Source/Runtime/Engine/Classes/Components/ChildActorComponent.h b/Engine/Source/Runtime/Engine/Classes/Components/ChildActorComponent.h
--- a/Engine/Source/Runtime/Engine/Classes/Components/ChildActorComponent.h
+++ b/Engine/Source/Runtime/Engine/Classes/Components/ChildActorComponent.h
@@ -11,7 +11,9 @@ class FChildActorComponentInstanceData;
 UCLASS(ClassGroup=Utility, hidecategories=(Object,LOD,Physics,Lighting,TextureStreaming,Activation,"Components|Activation",Collision), meta=(BlueprintSpawnableComponent), MinimalAPI)
 class UChildActorComponent : public USceneComponent
 {
-	GENERATED_UCLASS_BODY()
+	GENERATED_BODY()
+public:
+	ENGINE_API UChildActorComponent(const FObjectInitializer& ObjectInitializer = FObjectInitializer::Get());
 
 	UFUNCTION(BlueprintCallable, Category=ChildActorComponent)
 	ENGINE_API void SetChildActorClass(TSubclassOf<AActor> Class);
diff --git a/Engine/Source/Runtime/Engine/Classes/Components/DecalComponent.h b/Engine/Source/Runtime/Engine/Classes/Components/DecalComponent.h
--- a/Engine/Source/Runtime/Engine/Classes/Components/DecalComponent.h
+++ b/Engine/Source/Runtime/Engine/Classes/Components/DecalComponent.h
@@ -13,7 +13,9 @@ class FDeferredDecalProxy;
 UCLASS(hidecategories=(Collision, Object, Physics, SceneComponent, Activation, "Components|Activation", Mobility), ClassGroup=Rendering, meta=(BlueprintSpawnableComponent))
 class ENGINE_API UDecalComponent : public USceneComponent
 {
-	GENERATED_UCLASS_BODY()
+	GENERATED_BODY()
+public:
+	UDecalComponent(const FObjectInitializer& ObjectInitializer = FObjectInitializer::Get());
 
 	/** Decal material. */
 	//editable(Decal) private{private} const MaterialInterface	DecalMaterial;
diff --git a/Engine/Source/Runtime/Engine/Classes/Components/DestructibleComponent.h b/Engine/Source/Runtime/Engine/Classes/Components/DestructibleComponent.h
--- a/Engine/Source/Runtime/Engine/Classes/Components/DestructibleComponent.h
+++ b/Engine/Source/Runtime/Engine/Classes/Components/DestructibleComponent.h
@@ -44,7 +44,9 @@ DECLARE_DYNAMIC_MULTICAST_DELEGATE_TwoParams(FComponentFractureSignature, const
 UCLASS(ClassGroup=Physics, hidecategories=(Object,Mesh,"Components|SkinnedMesh",Mirroring,Activation,"Components|Activation"), config=Engine, editinlinenew, meta=(BlueprintSpawnableComponent))
 class ENGINE_API UDestructibleComponent : public USkinnedMeshComponent
 {
-	GENERATED_UCLASS_BODY()
+	GENERATED_BODY()
+public:
+	UDestructibleComponent(const FObjectInitializer& ObjectInitializer = FObjectInitializer::Get());
 
 	/** If set, use this actor's fracture effects instead of the asset's fracture effects. */
 	UPROPERTY(EditAnywhere, BlueprintReadWrite, Category=DestructibleComponent)
diff --git a/Engine/Source/Runtime/Engine/Classes/Components/DirectionalLightComponent.h b/Engine/Source/Runtime/Engine/Classes/Components/DirectionalLightComponent.h
--- a/Engine/Source/Runtime/Engine/Classes/Components/DirectionalLightComponent.h
+++ b/Engine/Source/Runtime/Engine/Classes/Components/DirectionalLightComponent.h
@@ -13,7 +13,9 @@
 UCLASS(ClassGroup=Lights, hidecategories=(Object, LightProfiles), editinlinenew, meta=(BlueprintSpawnableComponent), MinimalAPI)
 class UDirectionalLightComponent : public ULightComponent
 {
-	GENERATED_UCLASS_BODY()
+	GENERATED_BODY()
+public:
+	ENGINE_API UDirectionalLightComponent(const FObjectInitializer& ObjectInitializer = FObjectInitializer::Get());
 
 	/** Whether to occlude fog and atmosphere inscattering with screenspace blurred occlusion from this light. */
 	UPROPERTY(EditAnywhere, BlueprintReadOnly, Category=LightShafts, meta=(DisplayName = "Light Shaft Occlusion"))
diff --git a/Engine/Source/Runtime/Engine/Classes/Components/DrawFrustumComponent.h b/Engine/Source/Runtime/Engine/Classes/Components/DrawFrustumComponent.h
--- a/Engine/Source/Runtime/Engine/Classes/Components/DrawFrustumComponent.h
+++ b/Engine/Source/Runtime/Engine/Classes/Components/DrawFrustumComponent.h
@@ -11,7 +11,9 @@
 UCLASS(collapsecategories, hidecategories=Object, editinlinenew, MinimalAPI)
 class UDrawFrustumComponent : public UPrimitiveComponent
 {
-	GENERATED_UCLASS_BODY()
+	GENERATED_BODY()
+public:
+	ENGINE_API UDrawFrustumComponent(const FObjectInitializer& ObjectInitializer = FObjectInitializer::Get());
 	
 	/** Color to draw the wireframe frustum. */
 	UPROPERTY(EditAnywhere, BlueprintReadWrite, Category=DrawFrustumComponent)
diff --git a/Engine/Source/Runtime/Engine/Classes/Components/DrawSphereComponent.h b/Engine/Source/Runtime/Engine/Classes/Components/DrawSphereComponent.h
--- a/Engine/Source/Runtime/Engine/Classes/Components/DrawSphereComponent.h
+++ b/Engine/Source/Runtime/Engine/Classes/Components/DrawSphereComponent.h
@@ -11,7 +11,9 @@
 UCLASS(collapsecategories, hidecategories=Object, editinlinenew, MinimalAPI)
 class UDrawSphereComponent : public USphereComponent
 {
-	GENERATED_UCLASS_BODY()
+	GENERATED_BODY()
+public:
+	ENGINE_API UDrawSphereComponent(const FObjectInitializer& ObjectInitializer = FObjectInitializer::Get());
 
 public:
 #if WITH_EDITOR
diff --git a/Engine/Source/Runtime/Engine/Classes/Components/ExponentialHeightFogComponent.h b/Engine/Source/Runtime/Engine/Classes/Components/ExponentialHeightFogComponent.h
--- a/Engine/Source/Runtime/Engine/Classes/Components/ExponentialHeightFogComponent.h
+++ b/Engine/Source/Runtime/Engine/Classes/Components/ExponentialHeightFogComponent.h
@@ -10,7 +10,9 @@
 UCLASS(ClassGroup=Rendering, collapsecategories, hidecategories=(Object, Mobility), editinlinenew, meta=(BlueprintSpawnableComponent))
 class ENGINE_API UExponentialHeightFogComponent : public USceneComponent
 {
-	GENERATED_UCLASS_BODY()
+	GENERATED_BODY()
+public:
+	UExponentialHeightFogComponent(const FObjectInitializer& ObjectInitializer = FObjectInitializer::Get());
 
 	/** Global density factor. */
 	UPROPERTY(BlueprintReadOnly, interp, Category=ExponentialHeightFogComponent, meta=(UIMin = "0", UIMax = ".05"))
diff --git a/Engine/Source/Runtime/Engine/Classes/Components/HierarchicalInstancedStaticMeshComponent.h b/Engine/Source/Runtime/Engine/Classes/Components/HierarchicalInstancedStaticMeshComponent.h
--- a/Engine/Source/Runtime/Engine/Classes/Components/HierarchicalInstancedStaticMeshComponent.h
+++ b/Engine/Source/Runtime/Engine/Classes/Components/HierarchicalInstancedStaticMeshComponent.h
@@ -12,7 +12,7 @@ class FClusterBuilder;
 USTRUCT()
 struct FClusterNode
 {
-	GENERATED_USTRUCT_BODY()
+	GENERATED_BODY()
 
 	UPROPERTY()
 	FVector BoundMin;
@@ -76,7 +76,9 @@ public:
 UCLASS(ClassGroup=Rendering, meta=(BlueprintSpawnableComponent))
 class ENGINE_API UHierarchicalInstancedStaticMeshComponent : public UInstancedStaticMeshComponent
 {
-	GENERATED_UCLASS_BODY()
+	GENERATED_BODY()
+public:
+	UHierarchicalInstancedStaticMeshComponent(const FObjectInitializer& ObjectInitializer = FObjectInitializer::Get());
 
 	~UHierarchicalInstancedStaticMeshComponent();
 
diff --git a/Engine/Source/Runtime/Engine/Classes/Components/InputComponent.h b/Engine/Source/Runtime/Engine/Classes/Components/InputComponent.h
--- a/Engine/Source/Runtime/Engine/Classes/Components/InputComponent.h
+++ b/Engine/Source/Runtime/Engine/Classes/Components/InputComponent.h
@@ -80,7 +80,7 @@ protected:
 USTRUCT()
 struct FInputChord
 {
-	GENERATED_USTRUCT_BODY()
+	GENERATED_BODY()
 
 	/** The Key is the core of the chord. */
 	UPROPERTY()
@@ -530,7 +530,9 @@ UCLASS(transient, config=Input, hidecategories=(Activation, "Components|Activati
 class ENGINE_API UInputComponent
 	: public UActorComponent
 {
-	GENERATED_UCLASS_BODY()
+	GENERATED_BODY()
+public:
+	UInputComponent(const FObjectInitializer& ObjectInitializer = FObjectInitializer::Get());
 
 	/** The collection of key bindings. */
 	TArray<struct FInputKeyBinding> KeyBindings;
diff --git a/Engine/Source/Runtime/Engine/Classes/Components/InstancedStaticMeshComponent.h b/Engine/Source/Runtime/Engine/Classes/Components/InstancedStaticMeshComponent.h
--- a/Engine/Source/Runtime/Engine/Classes/Components/InstancedStaticMeshComponent.h
+++ b/Engine/Source/Runtime/Engine/Classes/Components/InstancedStaticMeshComponent.h
@@ -16,7 +16,7 @@ struct FPerInstanceRenderData;
 USTRUCT()
 struct FInstancedStaticMeshInstanceData
 {
-	GENERATED_USTRUCT_BODY()
+	GENERATED_BODY()
 
 	UPROPERTY(EditAnywhere, Category=Instances)
 	FMatrix Transform;
@@ -49,7 +49,7 @@ struct FInstancedStaticMeshInstanceData
 USTRUCT()
 struct FInstancedStaticMeshMappingInfo
 {
-	GENERATED_USTRUCT_BODY()
+	GENERATED_BODY()
 
 	FStaticLightingTextureMapping_InstancedStaticMesh* Mapping;
 
@@ -63,7 +63,9 @@ struct FInstancedStaticMeshMappingInfo
 UCLASS(ClassGroup=Rendering, meta=(BlueprintSpawnableComponent))
 class ENGINE_API UInstancedStaticMeshComponent : public UStaticMeshComponent
 {
-	GENERATED_UCLASS_BODY()
+	GENERATED_BODY()
+public:
+	UInstancedStaticMeshComponent(const FObjectInitializer& ObjectInitializer = FObjectInitializer::Get());
 
 	/** Array of instances, bulk serialized */
 	UPROPERTY(EditAnywhere, Transient, DuplicateTransient, DisplayName="Instances", Category=Instances, meta=(MakeEditWidget=true))
diff --git a/Engine/Source/Runtime/Engine/Classes/Components/LightComponent.h b/Engine/Source/Runtime/Engine/Classes/Components/LightComponent.h
--- a/Engine/Source/Runtime/Engine/Classes/Components/LightComponent.h
+++ b/Engine/Source/Runtime/Engine/Classes/Components/LightComponent.h
@@ -42,7 +42,9 @@ public:
 UCLASS(abstract, HideCategories=(Trigger,Activation,"Components|Activation",Physics), ShowCategories=(Mobility))
 class ENGINE_API ULightComponent : public ULightComponentBase
 {
-	GENERATED_UCLASS_BODY()
+	GENERATED_BODY()
+public:
+	ULightComponent(const FObjectInitializer& ObjectInitializer = FObjectInitializer::Get());
 
 	/** 
 	 * Shadow map channel which is used to match up with the appropriate static shadowing during a deferred shading pass.
diff --git a/Engine/Source/Runtime/Engine/Classes/Components/LightComponentBase.h b/Engine/Source/Runtime/Engine/Classes/Components/LightComponentBase.h
--- a/Engine/Source/Runtime/Engine/Classes/Components/LightComponentBase.h
+++ b/Engine/Source/Runtime/Engine/Classes/Components/LightComponentBase.h
@@ -6,7 +6,9 @@
 UCLASS(abstract, HideCategories=(Trigger,Activation,"Components|Activation",Physics), ShowCategories=(Mobility))
 class ENGINE_API ULightComponentBase : public USceneComponent
 {
-	GENERATED_UCLASS_BODY()
+	GENERATED_BODY()
+public:
+	ULightComponentBase(const FObjectInitializer& ObjectInitializer = FObjectInitializer::Get());
 
 	/**
 	 * GUID used to associate a light component with precomputed shadowing information across levels.
diff --git a/Engine/Source/Runtime/Engine/Classes/Components/LineBatchComponent.h b/Engine/Source/Runtime/Engine/Classes/Components/LineBatchComponent.h
--- a/Engine/Source/Runtime/Engine/Classes/Components/LineBatchComponent.h
+++ b/Engine/Source/Runtime/Engine/Classes/Components/LineBatchComponent.h
@@ -7,7 +7,7 @@
 USTRUCT()
 struct FBatchedLine
 {
-	GENERATED_USTRUCT_BODY()
+	GENERATED_BODY()
 
 	UPROPERTY()
 	FVector Start;
@@ -48,7 +48,7 @@ struct FBatchedLine
 USTRUCT()
 struct FBatchedPoint
 {
-	GENERATED_USTRUCT_BODY()
+	GENERATED_BODY()
 
 	UPROPERTY()
 	FVector Position;
@@ -111,7 +111,9 @@ public:
 UCLASS(MinimalAPI)
 class ULineBatchComponent : public UPrimitiveComponent
 {
-	GENERATED_UCLASS_BODY()
+	GENERATED_BODY()
+public:
+	ENGINE_API ULineBatchComponent(const FObjectInitializer& ObjectInitializer = FObjectInitializer::Get());
 
 	/** Buffer of lines to draw */
 	TArray<struct FBatchedLine> BatchedLines;
diff --git a/Engine/Source/Runtime/Engine/Classes/Components/MaterialBillboardComponent.h b/Engine/Source/Runtime/Engine/Classes/Components/MaterialBillboardComponent.h
--- a/Engine/Source/Runtime/Engine/Classes/Components/MaterialBillboardComponent.h
+++ b/Engine/Source/Runtime/Engine/Classes/Components/MaterialBillboardComponent.h
@@ -8,7 +8,7 @@
 USTRUCT()
 struct FMaterialSpriteElement
 {
-	GENERATED_USTRUCT_BODY()
+	GENERATED_BODY()
 
 	/** The material that the sprite is rendered with. */
 	UPROPERTY(EditAnywhere, BlueprintReadWrite, Category=MaterialSpriteElement)
@@ -53,7 +53,9 @@ struct FMaterialSpriteElement
 UCLASS(ClassGroup=Rendering, collapsecategories, hidecategories=(Object,Activation,"Components|Activation",Physics,Collision,Lighting,LOD,Mesh,PhysicsVolume), editinlinenew, meta=(BlueprintSpawnableComponent))
 class ENGINE_API UMaterialBillboardComponent : public UPrimitiveComponent
 {
-	GENERATED_UCLASS_BODY()
+	GENERATED_BODY()
+public:
+	UMaterialBillboardComponent(const FObjectInitializer& ObjectInitializer = FObjectInitializer::Get());
 
 	UPROPERTY(EditAnywhere, BlueprintReadWrite, Category=Sprite)
 	TArray<struct FMaterialSpriteElement> Elements;
diff --git a/Engine/Source/Runtime/Engine/Classes/Components/MeshComponent.h b/Engine/Source/Runtime/Engine/Classes/Components/MeshComponent.h
--- a/Engine/Source/Runtime/Engine/Classes/Components/MeshComponent.h
+++ b/Engine/Source/Runtime/Engine/Classes/Components/MeshComponent.h
@@ -17,7 +17,9 @@ struct FMaterialRelevance;
 UCLASS(abstract)
 class ENGINE_API UMeshComponent : public UPrimitiveComponent
 {
-	GENERATED_UCLASS_BODY()
+	GENERATED_BODY()
+public:
+	UMeshComponent(const FObjectInitializer& ObjectInitializer = FObjectInitializer::Get());
 
 	/** Per-Component material overrides.  These must NOT be set directly or a race condition can occur between GC and the rendering thread. */
 	UPROPERTY(EditAnywhere, Category=Rendering, Meta=(ToolTip="Material overrides."))
diff --git a/Engine/Source/Runtime/Engine/Classes/Components/ModelComponent.h b/Engine/Source/Runtime/Engine/Classes/Components/ModelComponent.h
--- a/Engine/Source/Runtime/Engine/Classes/Components/ModelComponent.h
+++ b/Engine/Source/Runtime/Engine/Classes/Components/ModelComponent.h
@@ -18,7 +18,9 @@ class FModelElement;
 UCLASS(MinimalAPI)
 class UModelComponent : public UPrimitiveComponent, public IInterface_CollisionDataProvider
 {
-	GENERATED_UCLASS_BODY()
+	GENERATED_BODY()
+public:
+	ENGINE_API UModelComponent(const FObjectInitializer& ObjectInitializer = FObjectInitializer::Get());
 
 private:
 	/** The BSP tree. */
diff --git a/Engine/Source/Runtime/Engine/Classes/Components/NiagaraComponent.h b/Engine/Source/Runtime/Engine/Classes/Components/NiagaraComponent.h
--- a/Engine/Source/Runtime/Engine/Classes/Components/NiagaraComponent.h
+++ b/Engine/Source/Runtime/Engine/Classes/Components/NiagaraComponent.h
@@ -93,7 +93,9 @@ private:
 UCLASS()
 class ENGINE_API UNiagaraComponent : public UPrimitiveComponent
 {
-	GENERATED_UCLASS_BODY()
+	GENERATED_BODY()
+public:
+	UNiagaraComponent(const FObjectInitializer& ObjectInitializer = FObjectInitializer::Get());
 private:
 	UPROPERTY()
 	//class UNiagaraEffect *Effect;
diff --git a/Engine/Source/Runtime/Engine/Classes/Components/PawnNoiseEmitterComponent.h b/Engine/Source/Runtime/Engine/Classes/Components/PawnNoiseEmitterComponent.h
--- a/Engine/Source/Runtime/Engine/Classes/Components/PawnNoiseEmitterComponent.h
+++ b/Engine/Source/Runtime/Engine/Classes/Components/PawnNoiseEmitterComponent.h
@@ -12,7 +12,9 @@
 UCLASS(ClassGroup=AI, meta=(BlueprintSpawnableComponent))
 class ENGINE_API UPawnNoiseEmitterComponent : public UActorComponent
 {
-	GENERATED_UCLASS_BODY()
+	GENERATED_BODY()
+public:
+	UPawnNoiseEmitterComponent(const FObjectInitializer& ObjectInitializer = FObjectInitializer::Get());
 
 public:
 	// Most recent noise made by this pawn not at its own location
diff --git a/Engine/Source/Runtime/Engine/Classes/Components/PlaneReflectionCaptureComponent.h b/Engine/Source/Runtime/Engine/Classes/Components/PlaneReflectionCaptureComponent.h
--- a/Engine/Source/Runtime/Engine/Classes/Components/PlaneReflectionCaptureComponent.h
+++ b/Engine/Source/Runtime/Engine/Classes/Components/PlaneReflectionCaptureComponent.h
@@ -9,7 +9,9 @@
 UCLASS(hidecategories=(Collision, Object, Physics, SceneComponent), MinimalAPI)
 class UPlaneReflectionCaptureComponent : public UReflectionCaptureComponent
 {
-	GENERATED_UCLASS_BODY()
+	GENERATED_BODY()
+public:
+	ENGINE_API UPlaneReflectionCaptureComponent(const FObjectInitializer& ObjectInitializer = FObjectInitializer::Get());
 
 	/** Radius of the area that can receive reflections from this capture. */
 	UPROPERTY(EditAnywhere, BlueprintReadWrite, Category=ReflectionCapture, meta=(UIMin = "8.0", UIMax = "16384.0"))
diff --git a/Engine/Source/Runtime/Engine/Classes/Components/PointLightComponent.h b/Engine/Source/Runtime/Engine/Classes/Components/PointLightComponent.h
--- a/Engine/Source/Runtime/Engine/Classes/Components/PointLightComponent.h
+++ b/Engine/Source/Runtime/Engine/Classes/Components/PointLightComponent.h
@@ -11,7 +11,9 @@
 UCLASS(ClassGroup=(Lights,Common), hidecategories=(Object, LightShafts), editinlinenew, meta=(BlueprintSpawnableComponent))
 class ENGINE_API UPointLightComponent : public ULightComponent
 {
-	GENERATED_UCLASS_BODY()
+	GENERATED_BODY()
+public:
+	UPointLightComponent(const FObjectInitializer& ObjectInitializer = FObjectInitializer::Get());
 
 	UPROPERTY()
 	float Radius_DEPRECATED;
diff --git a/Engine/Source/Runtime/Engine/Classes/Components/PoseableMeshComponent.h b/Engine/Source/Runtime/Engine/Classes/Components/PoseableMeshComponent.h
--- a/Engine/Source/Runtime/Engine/Classes/Components/PoseableMeshComponent.h
+++ b/Engine/Source/Runtime/Engine/Classes/Components/PoseableMeshComponent.h
@@ -10,7 +10,9 @@
 UCLASS(ClassGroup=Rendering, hidecategories=(Object,Physics), config=Engine, editinlinenew, meta=(BlueprintSpawnableComponent))
 class ENGINE_API UPoseableMeshComponent : public USkinnedMeshComponent
 {
-	GENERATED_UCLASS_BODY()
+	GENERATED_BODY()
+public:
+	UPoseableMeshComponent(const FObjectInitializer& ObjectInitializer = FObjectInitializer::Get());
 
 	/** Temporary array of local-space (ie relative to parent bone) rotation/translation/scale for each bone. */
 	TArray<FTransform> LocalAtoms;
diff --git a/Engine/Source/Runtime/Engine/Classes/Components/PostProcessComponent.h b/Engine/Source/Runtime/Engine/Classes/Components/PostProcessComponent.h
--- a/Engine/Source/Runtime/Engine/Classes/Components/PostProcessComponent.h
+++ b/Engine/Source/Runtime/Engine/Classes/Components/PostProcessComponent.h
@@ -11,7 +11,9 @@
 UCLASS(ClassGroup = Rendering, collapsecategories, hidecategories = (Object), editinlinenew, meta = (BlueprintSpawnableComponent), MinimalAPI)
 class UPostProcessComponent : public USceneComponent, public IInterface_PostProcessVolume
 {
-	GENERATED_UCLASS_BODY()
+	GENERATED_BODY()
+public:
+	ENGINE_API UPostProcessComponent(const FObjectInitializer& ObjectInitializer = FObjectInitializer::Get());
 
 	/** Post process settings to use for this volume. */
 	UPROPERTY(interp, Category = PostProcessVolume, meta = (ShowOnlyInnerProperties))
diff --git a/Engine/Source/Runtime/Engine/Classes/Components/ReflectionCaptureComponent.h b/Engine/Source/Runtime/Engine/Classes/Components/ReflectionCaptureComponent.h
--- a/Engine/Source/Runtime/Engine/Classes/Components/ReflectionCaptureComponent.h
+++ b/Engine/Source/Runtime/Engine/Classes/Components/ReflectionCaptureComponent.h
@@ -71,7 +71,9 @@ private:
 UCLASS(abstract, hidecategories=(Collision, Object, Physics, SceneComponent, Activation, "Components|Activation", Mobility), MinimalAPI)
 class UReflectionCaptureComponent : public USceneComponent
 {
-	GENERATED_UCLASS_BODY()
+	GENERATED_BODY()
+public:
+	ENGINE_API UReflectionCaptureComponent(const FObjectInitializer& ObjectInitializer = FObjectInitializer::Get());
 
 	/** A brightness control to scale the captured scene's reflection intensity. */
 	UPROPERTY(EditAnywhere, BlueprintReadOnly, Category=ReflectionCapture, meta=(UIMin = ".5", UIMax = "4"))
diff --git a/Engine/Source/Runtime/Engine/Classes/Components/SceneCaptureComponent.h b/Engine/Source/Runtime/Engine/Classes/Components/SceneCaptureComponent.h
--- a/Engine/Source/Runtime/Engine/Classes/Components/SceneCaptureComponent.h
+++ b/Engine/Source/Runtime/Engine/Classes/Components/SceneCaptureComponent.h
@@ -8,7 +8,7 @@
 USTRUCT()
 struct FEngineShowFlagsSetting
 {
-	GENERATED_USTRUCT_BODY()
+	GENERATED_BODY()
 
 	UPROPERTY(EditAnywhere, BlueprintReadWrite, Category = SceneCapture)
 	FString ShowFlagName;
@@ -21,7 +21,9 @@ struct FEngineShowFlagsSetting
 UCLASS(hidecategories=(abstract, Collision, Object, Physics, SceneComponent, Mobility), MinimalAPI)
 class USceneCaptureComponent : public USceneComponent
 {
-	GENERATED_UCLASS_BODY()
+	GENERATED_BODY()
+public:
+	ENGINE_API USceneCaptureComponent(const FObjectInitializer& ObjectInitializer = FObjectInitializer::Get());
 
 	/** The components won't rendered by current component.*/
  	UPROPERTY()
diff --git a/Engine/Source/Runtime/Engine/Classes/Components/SceneCaptureComponent2D.h b/Engine/Source/Runtime/Engine/Classes/Components/SceneCaptureComponent2D.h
--- a/Engine/Source/Runtime/Engine/Classes/Components/SceneCaptureComponent2D.h
+++ b/Engine/Source/Runtime/Engine/Classes/Components/SceneCaptureComponent2D.h
@@ -20,7 +20,9 @@ enum ESceneCaptureSource
 UCLASS(hidecategories=(Collision, Object, Physics, SceneComponent), ClassGroup=Rendering, MinimalAPI, editinlinenew, meta=(BlueprintSpawnableComponent))
 class USceneCaptureComponent2D : public USceneCaptureComponent
 {
-	GENERATED_UCLASS_BODY()
+	GENERATED_BODY()
+public:
+	ENGINE_API USceneCaptureComponent2D(const FObjectInitializer& ObjectInitializer = FObjectInitializer::Get());
 
 	/** Camera field of view (in degrees) */
 	UPROPERTY(interp, Category=SceneCapture, meta=(DisplayName = "Field of View", UIMin = "5.0", UIMax = "170", ClampMin = "0.001", ClampMax = "360.0"))
@@ -64,4 +66,4 @@ public:
 	ENGINE_API void UpdateContent();
 
 	ENGINE_API static void UpdateDeferredCaptures( FSceneInterface* Scene );
-};
\ No newline at end of file
+};
diff --git a/Engine/Source/Runtime/Engine/Classes/Components/SceneCaptureComponentCube.h b/Engine/Source/Runtime/Engine/Classes/Components/SceneCaptureComponentCube.h
--- a/Engine/Source/Runtime/Engine/Classes/Components/SceneCaptureComponentCube.h
+++ b/Engine/Source/Runtime/Engine/Classes/Components/SceneCaptureComponentCube.h
@@ -10,7 +10,9 @@
 UCLASS(hidecategories = (Collision, Object, Physics, SceneComponent), ClassGroup=Rendering, MinimalAPI, editinlinenew, meta = (BlueprintSpawnableComponent))
 class USceneCaptureComponentCube : public USceneCaptureComponent
 {
-	GENERATED_UCLASS_BODY()
+	GENERATED_BODY()
+public:
+	ENGINE_API USceneCaptureComponentCube(const FObjectInitializer& ObjectInitializer = FObjectInitializer::Get());
 
 	/** Temporary render target that can be used by the editor. */
 	UPROPERTY(EditAnywhere, BlueprintReadWrite, Category = SceneCapture)
@@ -32,4 +34,4 @@ public:
 	void UpdateContent();
 
 	ENGINE_API static void UpdateDeferredCaptures(FSceneInterface* Scene);
-};
\ No newline at end of file
+};
diff --git a/Engine/Source/Runtime/Engine/Classes/Components/ShapeComponent.h b/Engine/Source/Runtime/Engine/Classes/Components/ShapeComponent.h
--- a/Engine/Source/Runtime/Engine/Classes/Components/ShapeComponent.h
+++ b/Engine/Source/Runtime/Engine/Classes/Components/ShapeComponent.h
@@ -10,7 +10,9 @@
 UCLASS(abstract, hidecategories=(Object,LOD,Lighting,TextureStreaming,Activation,"Components|Activation"), editinlinenew, meta=(BlueprintSpawnableComponent), showcategories=(Mobility))
 class ENGINE_API UShapeComponent : public UPrimitiveComponent
 {
-	GENERATED_UCLASS_BODY()
+	GENERATED_BODY()
+public:
+	UShapeComponent(const FObjectInitializer& ObjectInitializer = FObjectInitializer::Get());
 
 	/** Color used to draw the shape. */
 	UPROPERTY(EditAnywhere, AdvancedDisplay, BlueprintReadOnly, Category=Shape)
diff --git a/Engine/Source/Runtime/Engine/Classes/Components/SkeletalMeshComponent.h b/Engine/Source/Runtime/Engine/Classes/Components/SkeletalMeshComponent.h
--- a/Engine/Source/Runtime/Engine/Classes/Components/SkeletalMeshComponent.h
+++ b/Engine/Source/Runtime/Engine/Classes/Components/SkeletalMeshComponent.h
@@ -193,7 +193,7 @@ namespace EAnimationMode
 USTRUCT()
 struct FSingleAnimationPlayData
 {
-	GENERATED_USTRUCT_BODY()
+	GENERATED_BODY()
 
 	// @todo in the future, we should make this one UObject
 	// and have detail customization to display different things
@@ -245,7 +245,7 @@ struct FSingleAnimationPlayData
 USTRUCT()
 struct FSkeletalMeshComponentPreClothTickFunction : public FTickFunction
 {
-	GENERATED_USTRUCT_BODY()
+	GENERATED_BODY()
 
 	/** World this tick function belongs to **/
 	class USkeletalMeshComponent*	Target;
@@ -272,7 +272,9 @@ struct FSkeletalMeshComponentPreClothTickFunction : public FTickFunction
 UCLASS(ClassGroup=(Rendering, Common), hidecategories=Object, config=Engine, editinlinenew, meta=(BlueprintSpawnableComponent))
 class ENGINE_API USkeletalMeshComponent : public USkinnedMeshComponent, public IInterface_CollisionDataProvider
 {
-	GENERATED_UCLASS_BODY()
+	GENERATED_BODY()
+public:
+	USkeletalMeshComponent(const FObjectInitializer& ObjectInitializer = FObjectInitializer::Get());
 	
 	/**
 	 * Animation 
diff --git a/Engine/Source/Runtime/Engine/Classes/Components/SkinnedMeshComponent.h b/Engine/Source/Runtime/Engine/Classes/Components/SkinnedMeshComponent.h
--- a/Engine/Source/Runtime/Engine/Classes/Components/SkinnedMeshComponent.h
+++ b/Engine/Source/Runtime/Engine/Classes/Components/SkinnedMeshComponent.h
@@ -76,7 +76,7 @@ namespace EBoneSpaces
 USTRUCT()
 struct FActiveVertexAnim
 {
-	GENERATED_USTRUCT_BODY()
+	GENERATED_BODY()
 
 	/** The anim that we want to apply. */
 	UPROPERTY()
@@ -122,7 +122,7 @@ struct FActiveVertexAnim
 USTRUCT()
 struct FSkelMeshComponentLODInfo
 {
-	GENERATED_USTRUCT_BODY()
+	GENERATED_BODY()
 
 	/** Material corresponds to section. To show/hide each section, use this **/
 	UPROPERTY()
@@ -147,7 +147,9 @@ struct FSkelMeshComponentLODInfo
 UCLASS(hidecategories=Object, config=Engine, editinlinenew, abstract)
 class ENGINE_API USkinnedMeshComponent : public UMeshComponent
 {
-	GENERATED_UCLASS_BODY()
+	GENERATED_BODY()
+public:
+	USkinnedMeshComponent(const FObjectInitializer& ObjectInitializer = FObjectInitializer::Get());
 
 	/** The skeletal mesh used by this component. */
 	UPROPERTY(EditAnywhere, BlueprintReadOnly, Category="Mesh")
diff --git a/Engine/Source/Runtime/Engine/Classes/Components/SkyLightComponent.h b/Engine/Source/Runtime/Engine/Classes/Components/SkyLightComponent.h
--- a/Engine/Source/Runtime/Engine/Classes/Components/SkyLightComponent.h
+++ b/Engine/Source/Runtime/Engine/Classes/Components/SkyLightComponent.h
@@ -85,7 +85,9 @@ enum ESkyLightSourceType
 UCLASS(ClassGroup=Lights, HideCategories=(Trigger,Activation,"Components|Activation",Physics), meta=(BlueprintSpawnableComponent))
 class ENGINE_API USkyLightComponent : public ULightComponentBase
 {
-	GENERATED_UCLASS_BODY()
+	GENERATED_BODY()
+public:
+	USkyLightComponent(const FObjectInitializer& ObjectInitializer = FObjectInitializer::Get());
 
 	/** Indicates where to get the light contribution from. */
 	UPROPERTY(EditAnywhere, BlueprintReadOnly, Category=Light)
diff --git a/Engine/Source/Runtime/Engine/Classes/Components/SphereComponent.h b/Engine/Source/Runtime/Engine/Classes/Components/SphereComponent.h
--- a/Engine/Source/Runtime/Engine/Classes/Components/SphereComponent.h
+++ b/Engine/Source/Runtime/Engine/Classes/Components/SphereComponent.h
@@ -11,7 +11,9 @@
 UCLASS(ClassGroup="Collision", editinlinenew, hidecategories=(Object,LOD,Lighting,TextureStreaming), meta=(DisplayName="Sphere Collision", BlueprintSpawnableComponent))
 class ENGINE_API USphereComponent : public UShapeComponent
 {
-	GENERATED_UCLASS_BODY()
+	GENERATED_BODY()
+public:
+	USphereComponent(const FObjectInitializer& ObjectInitializer = FObjectInitializer::Get());
 
 protected:
 	/** The radius of the sphere **/
diff --git a/Engine/Source/Runtime/Engine/Classes/Components/SphereReflectionCaptureComponent.h b/Engine/Source/Runtime/Engine/Classes/Components/SphereReflectionCaptureComponent.h
--- a/Engine/Source/Runtime/Engine/Classes/Components/SphereReflectionCaptureComponent.h
+++ b/Engine/Source/Runtime/Engine/Classes/Components/SphereReflectionCaptureComponent.h
@@ -9,7 +9,9 @@
 UCLASS(hidecategories=(Collision, Object, Physics, SceneComponent), MinimalAPI)
 class USphereReflectionCaptureComponent : public UReflectionCaptureComponent
 {
-	GENERATED_UCLASS_BODY()
+	GENERATED_BODY()
+public:
+	ENGINE_API USphereReflectionCaptureComponent(const FObjectInitializer& ObjectInitializer = FObjectInitializer::Get());
 
 	/** Radius of the area that can receive reflections from this capture. */
 	UPROPERTY(EditAnywhere, BlueprintReadWrite, Category=ReflectionCapture, meta=(UIMin = "8.0", UIMax = "16384.0"))
diff --git a/Engine/Source/Runtime/Engine/Classes/Components/SplineComponent.h b/Engine/Source/Runtime/Engine/Classes/Components/SplineComponent.h
--- a/Engine/Source/Runtime/Engine/Classes/Components/SplineComponent.h
+++ b/Engine/Source/Runtime/Engine/Classes/Components/SplineComponent.h
@@ -27,7 +27,9 @@ namespace ESplinePointType
 UCLASS(ClassGroup=Utility, meta=(BlueprintSpawnableComponent))
 class ENGINE_API USplineComponent : public UPrimitiveComponent
 {
-	GENERATED_UCLASS_BODY()
+	GENERATED_BODY()
+public:
+	USplineComponent(const FObjectInitializer& ObjectInitializer = FObjectInitializer::Get());
 
 	/** Actual data for spline. Locations and tangents are in world space. */
 	UPROPERTY()
diff --git a/Engine/Source/Runtime/Engine/Classes/Components/SplineMeshComponent.h b/Engine/Source/Runtime/Engine/Classes/Components/SplineMeshComponent.h
--- a/Engine/Source/Runtime/Engine/Classes/Components/SplineMeshComponent.h
+++ b/Engine/Source/Runtime/Engine/Classes/Components/SplineMeshComponent.h
@@ -25,7 +25,7 @@ namespace ESplineMeshAxis
 USTRUCT()
 struct FSplineMeshParams
 {
-	GENERATED_USTRUCT_BODY()
+	GENERATED_BODY()
 
 	/** Start location of spline, in component space */
 	UPROPERTY(EditAnywhere, Category=SplineMesh)
@@ -91,7 +91,9 @@ struct FSplineMeshParams
 UCLASS(ClassGroup=Rendering, hidecategories=(Physics), meta=(BlueprintSpawnableComponent))
 class ENGINE_API USplineMeshComponent : public UStaticMeshComponent, public IInterface_CollisionDataProvider
 {
-	GENERATED_UCLASS_BODY()
+	GENERATED_BODY()
+public:
+	USplineMeshComponent(const FObjectInitializer& ObjectInitializer = FObjectInitializer::Get());
 
 	/** Spline that is used to deform mesh */
 	UPROPERTY(EditAnywhere, Category=SplineMesh, meta=(ShowOnlyInnerProperties))
diff --git a/Engine/Source/Runtime/Engine/Classes/Components/SpotLightComponent.h b/Engine/Source/Runtime/Engine/Classes/Components/SpotLightComponent.h
--- a/Engine/Source/Runtime/Engine/Classes/Components/SpotLightComponent.h
+++ b/Engine/Source/Runtime/Engine/Classes/Components/SpotLightComponent.h
@@ -10,7 +10,9 @@
 UCLASS(ClassGroup=Lights, hidecategories=Object, editinlinenew, meta=(BlueprintSpawnableComponent))
 class ENGINE_API USpotLightComponent : public UPointLightComponent
 {
-	GENERATED_UCLASS_BODY()
+	GENERATED_BODY()
+public:
+	USpotLightComponent(const FObjectInitializer& ObjectInitializer = FObjectInitializer::Get());
 
 	/** Degrees. */
 	UPROPERTY(EditAnywhere, BlueprintReadOnly, Category=Light, meta=(UIMin = "1.0", UIMax = "80.0"))
diff --git a/Engine/Source/Runtime/Engine/Classes/Components/StaticMeshComponent.h b/Engine/Source/Runtime/Engine/Classes/Components/StaticMeshComponent.h
--- a/Engine/Source/Runtime/Engine/Classes/Components/StaticMeshComponent.h
+++ b/Engine/Source/Runtime/Engine/Classes/Components/StaticMeshComponent.h
@@ -19,7 +19,7 @@ struct FConvexVolume;
 USTRUCT()
 struct FPaintedVertex
 {
-	GENERATED_USTRUCT_BODY()
+	GENERATED_BODY()
 
 	UPROPERTY()
 	FVector Position;
@@ -51,7 +51,7 @@ struct FPaintedVertex
 USTRUCT()
 struct FStaticMeshComponentLODInfo
 {
-	GENERATED_USTRUCT_BODY()
+	GENERATED_BODY()
 
 	FLightMapRef LightMap;
 
@@ -120,7 +120,9 @@ struct TStructOpsTypeTraits<FStaticMeshComponentLODInfo> : public TStructOpsType
 UCLASS(ClassGroup=(Rendering, Common), hidecategories=(Object,Activation,"Components|Activation"), ShowCategories=(Mobility), editinlinenew, meta=(BlueprintSpawnableComponent))
 class ENGINE_API UStaticMeshComponent : public UMeshComponent
 {
-	GENERATED_UCLASS_BODY()
+	GENERATED_BODY()
+public:
+	UStaticMeshComponent(const FObjectInitializer& ObjectInitializer = FObjectInitializer::Get());
 
 	/** If 0, auto-select LOD level. if >0, force to (ForcedLodModel-1). */
 	UPROPERTY(EditAnywhere, AdvancedDisplay, BlueprintReadOnly, Category=LOD)
diff --git a/Engine/Source/Runtime/Engine/Classes/Components/TextRenderComponent.h b/Engine/Source/Runtime/Engine/Classes/Components/TextRenderComponent.h
--- a/Engine/Source/Runtime/Engine/Classes/Components/TextRenderComponent.h
+++ b/Engine/Source/Runtime/Engine/Classes/Components/TextRenderComponent.h
@@ -35,7 +35,9 @@ enum EVerticalTextAligment
 UCLASS(ClassGroup=Rendering, hidecategories=(Object,LOD,Physics,TextureStreaming,Activation,"Components|Activation",Collision), editinlinenew, meta=(BlueprintSpawnableComponent = ""))
 class ENGINE_API UTextRenderComponent : public UPrimitiveComponent
 {
-	GENERATED_UCLASS_BODY()
+	GENERATED_BODY()
+public:
+	UTextRenderComponent(const FObjectInitializer& ObjectInitializer = FObjectInitializer::Get());
 
 	/** Text content, can be multi line using <br> as line separator */
 	UPROPERTY(EditAnywhere, BlueprintReadOnly, Category=Text, meta=(MultiLine=true))
diff --git a/Engine/Source/Runtime/Engine/Classes/Components/TimelineComponent.h b/Engine/Source/Runtime/Engine/Classes/Components/TimelineComponent.h
--- a/Engine/Source/Runtime/Engine/Classes/Components/TimelineComponent.h
+++ b/Engine/Source/Runtime/Engine/Classes/Components/TimelineComponent.h
@@ -45,7 +45,7 @@ namespace ETimelineDirection
 USTRUCT()
 struct FTimelineEventEntry
 {
-	GENERATED_USTRUCT_BODY()
+	GENERATED_BODY()
 
 	/** Time at which event should fire */
 	UPROPERTY()
@@ -67,7 +67,7 @@ struct FTimelineEventEntry
 USTRUCT()
 struct FTimelineVectorTrack
 {
-	GENERATED_USTRUCT_BODY()
+	GENERATED_BODY()
 
 	/** Vector curve to be evaluated */
 	UPROPERTY()
@@ -101,7 +101,7 @@ struct FTimelineVectorTrack
 USTRUCT()
 struct FTimelineFloatTrack
 {
-	GENERATED_USTRUCT_BODY()
+	GENERATED_BODY()
 
 	/** Float curve to be evaluated */
 	UPROPERTY()
@@ -136,7 +136,7 @@ struct FTimelineFloatTrack
 USTRUCT()
 struct FTimelineLinearColorTrack
 {
-	GENERATED_USTRUCT_BODY()
+	GENERATED_BODY()
 
 	/** Float curve to be evaluated */
 	UPROPERTY()
@@ -169,7 +169,7 @@ struct FTimelineLinearColorTrack
 USTRUCT()
 struct FTimeline
 {
-	GENERATED_USTRUCT_BODY()
+	GENERATED_BODY()
 		
 private:
 	/** Specified how the timeline determines its own length (e.g. specified length, last keyframe) */
@@ -361,7 +361,9 @@ private:
 UCLASS(MinimalAPI)
 class UTimelineComponent : public UActorComponent
 {
-	GENERATED_UCLASS_BODY()
+	GENERATED_BODY()
+public:
+	ENGINE_API UTimelineComponent(const FObjectInitializer& ObjectInitializer = FObjectInitializer::Get());
 
 private:
 	/** The actual timeline structure */
diff --git a/Engine/Source/Runtime/Engine/Classes/Components/VectorFieldComponent.h b/Engine/Source/Runtime/Engine/Classes/Components/VectorFieldComponent.h
--- a/Engine/Source/Runtime/Engine/Classes/Components/VectorFieldComponent.h
+++ b/Engine/Source/Runtime/Engine/Classes/Components/VectorFieldComponent.h
@@ -12,7 +12,9 @@
 UCLASS(ClassGroup=Rendering, hidecategories=Object, meta=(BlueprintSpawnableComponent), MinimalAPI)
 class UVectorFieldComponent : public UPrimitiveComponent
 {
-	GENERATED_UCLASS_BODY()
+	GENERATED_BODY()
+public:
+	ENGINE_API UVectorFieldComponent(const FObjectInitializer& ObjectInitializer = FObjectInitializer::Get());
 
 	/** The vector field asset. */
 	UPROPERTY(EditAnywhere, BlueprintReadWrite, Category=VectorFieldComponent)
diff --git a/Engine/Source/Runtime/Engine/Classes/Components/WindDirectionalSourceComponent.h b/Engine/Source/Runtime/Engine/Classes/Components/WindDirectionalSourceComponent.h
--- a/Engine/Source/Runtime/Engine/Classes/Components/WindDirectionalSourceComponent.h
+++ b/Engine/Source/Runtime/Engine/Classes/Components/WindDirectionalSourceComponent.h
@@ -7,7 +7,9 @@
 UCLASS(collapsecategories, hidecategories=(Object, Mobility), editinlinenew)
 class UWindDirectionalSourceComponent : public USceneComponent
 {
-	GENERATED_UCLASS_BODY()
+	GENERATED_BODY()
+public:
+	UWindDirectionalSourceComponent(const FObjectInitializer& ObjectInitializer = FObjectInitializer::Get());
 
 	UPROPERTY(interp, Category=WindDirectionalSourceComponent)
 	float Strength;
diff --git a/Engine/Source/Runtime/Engine/Classes/Curves/CurveBase.h b/Engine/Source/Runtime/Engine/Classes/Curves/CurveBase.h
--- a/Engine/Source/Runtime/Engine/Classes/Curves/CurveBase.h
+++ b/Engine/Source/Runtime/Engine/Classes/Curves/CurveBase.h
@@ -42,7 +42,7 @@ private:
 USTRUCT()
 struct FKeyHandleMap
 {
-	GENERATED_USTRUCT_BODY()
+	GENERATED_BODY()
 public:
 	FKeyHandleMap() {}
 
@@ -92,7 +92,7 @@ struct TStructOpsTypeTraits< FKeyHandleMap > : public TStructOpsTypeTraitsBase
 USTRUCT()
 struct ENGINE_API FIndexedCurve
 {
-	GENERATED_USTRUCT_BODY()
+	GENERATED_BODY()
 public:
 	FIndexedCurve() {}
 
@@ -160,7 +160,7 @@ enum ERichCurveTangentWeightMode
 USTRUCT()
 struct ENGINE_API FRichCurveKey
 {
-	GENERATED_USTRUCT_BODY()
+	GENERATED_BODY()
 
 	/** Interpolation mode between this key and the next */
 	UPROPERTY()
@@ -268,7 +268,7 @@ struct TStructOpsTypeTraits< FRichCurveKey > : public TStructOpsTypeTraitsBase
 USTRUCT()
 struct ENGINE_API FRichCurve : public FIndexedCurve
 {
-	GENERATED_USTRUCT_BODY()
+	GENERATED_BODY()
 
 public:
 	virtual ~FRichCurve()
@@ -462,7 +462,9 @@ public:
 UCLASS(abstract)
 class ENGINE_API UCurveBase : public UObject, public FCurveOwnerInterface
 {
-	GENERATED_UCLASS_BODY()
+	GENERATED_BODY()
+public:
+	UCurveBase(const FObjectInitializer& ObjectInitializer = FObjectInitializer::Get());
 
 	/** The filename imported to create this object. Relative to this object's package, BaseDir() or absolute */
 	UPROPERTY()
@@ -524,7 +526,7 @@ public:
 USTRUCT()
 struct FIntegralKey
 {
-	GENERATED_USTRUCT_BODY()
+	GENERATED_BODY()
 public:
 	FIntegralKey(float InTime = 0.f, int32 InValue = 0)
 		: Time(InTime)
@@ -543,7 +545,7 @@ public:
 USTRUCT()
 struct ENGINE_API FIntegralCurve : public FIndexedCurve
 {
-	GENERATED_USTRUCT_BODY()
+	GENERATED_BODY()
 public:
 	
 	/** Get number of keys in curve. */
diff --git a/Engine/Source/Runtime/Engine/Classes/Curves/CurveEdPresetCurve.h b/Engine/Source/Runtime/Engine/Classes/Curves/CurveEdPresetCurve.h
--- a/Engine/Source/Runtime/Engine/Classes/Curves/CurveEdPresetCurve.h
+++ b/Engine/Source/Runtime/Engine/Classes/Curves/CurveEdPresetCurve.h
@@ -11,6 +11,8 @@
 UCLASS(deprecated)
 class UDEPRECATED_CurveEdPresetCurve : public UObject
 {
-	GENERATED_UCLASS_BODY()
+	GENERATED_BODY()
+public:
+	UDEPRECATED_CurveEdPresetCurve(const FObjectInitializer& ObjectInitializer = FObjectInitializer::Get());
 };
 
diff --git a/Engine/Source/Runtime/Engine/Classes/Curves/CurveFloat.h b/Engine/Source/Runtime/Engine/Classes/Curves/CurveFloat.h
--- a/Engine/Source/Runtime/Engine/Classes/Curves/CurveFloat.h
+++ b/Engine/Source/Runtime/Engine/Classes/Curves/CurveFloat.h
@@ -8,7 +8,7 @@
 USTRUCT()
 struct ENGINE_API FRuntimeFloatCurve
 {
-	GENERATED_USTRUCT_BODY()
+	GENERATED_BODY()
 
 	UPROPERTY()
 	FRichCurve EditorCurveData;
@@ -26,7 +26,9 @@ struct ENGINE_API FRuntimeFloatCurve
 UCLASS(BlueprintType)
 class ENGINE_API UCurveFloat : public UCurveBase
 {
-	GENERATED_UCLASS_BODY()
+	GENERATED_BODY()
+public:
+	UCurveFloat(const FObjectInitializer& ObjectInitializer = FObjectInitializer::Get());
 
 	/** Keyframe data */
 	UPROPERTY()
diff --git a/Engine/Source/Runtime/Engine/Classes/Curves/CurveLinearColor.h b/Engine/Source/Runtime/Engine/Classes/Curves/CurveLinearColor.h
--- a/Engine/Source/Runtime/Engine/Classes/Curves/CurveLinearColor.h
+++ b/Engine/Source/Runtime/Engine/Classes/Curves/CurveLinearColor.h
@@ -7,7 +7,7 @@
 USTRUCT()
 struct ENGINE_API FRuntimeCurveLinearColor
 {
-	GENERATED_USTRUCT_BODY()
+	GENERATED_BODY()
 public:
 	UPROPERTY()
 	FRichCurve ColorCurves[4];
@@ -21,7 +21,9 @@ public:
 UCLASS(BlueprintType)
 class ENGINE_API UCurveLinearColor : public UCurveBase
 {
-	GENERATED_UCLASS_BODY()
+	GENERATED_BODY()
+public:
+	UCurveLinearColor(const FObjectInitializer& ObjectInitializer = FObjectInitializer::Get());
 
 	/** Keyframe data, one curve for red, green, blue, and alpha */
 	UPROPERTY()
