Enabled UHT to digest GENERATED_BODY instead of GENERATED_UCLASS_BODY, GENERATED_USTRUCT_BODY, GENERATED_UINTERFACE_BODY or GENERATED_IINTERFACE_BODY, changed every occurence to the new syntax and fix

diff --git a/Engine/Source/Runtime/Engine/Classes/Curves/CurveVector.h b/Engine/Source/Runtime/Engine/Classes/Curves/CurveVector.h
--- a/Engine/Source/Runtime/Engine/Classes/Curves/CurveVector.h
+++ b/Engine/Source/Runtime/Engine/Classes/Curves/CurveVector.h
@@ -7,7 +7,9 @@
 UCLASS(BlueprintType, MinimalAPI)
 class UCurveVector : public UCurveBase
 {
-	GENERATED_UCLASS_BODY()
+	GENERATED_BODY()
+public:
+	ENGINE_API UCurveVector(const FObjectInitializer& ObjectInitializer = FObjectInitializer::Get());
 
 	/** Keyframe data, one curve for X, Y and Z */
 	UPROPERTY()
diff --git a/Engine/Source/Runtime/Engine/Classes/Debug/DebugDrawService.h b/Engine/Source/Runtime/Engine/Classes/Debug/DebugDrawService.h
--- a/Engine/Source/Runtime/Engine/Classes/Debug/DebugDrawService.h
+++ b/Engine/Source/Runtime/Engine/Classes/Debug/DebugDrawService.h
@@ -15,7 +15,9 @@ DECLARE_DELEGATE_TwoParams(FDebugDrawDelegate, class UCanvas*, class APlayerCont
 UCLASS(config=Engine)
 class ENGINE_API UDebugDrawService : public UBlueprintFunctionLibrary
 {
-	GENERATED_UCLASS_BODY()
+	GENERATED_BODY()
+public:
+	UDebugDrawService(const FObjectInitializer& ObjectInitializer = FObjectInitializer::Get());
 
 	//void Register
 
diff --git a/Engine/Source/Runtime/Engine/Classes/Debug/LogVisualizerCameraController.h b/Engine/Source/Runtime/Engine/Classes/Debug/LogVisualizerCameraController.h
--- a/Engine/Source/Runtime/Engine/Classes/Debug/LogVisualizerCameraController.h
+++ b/Engine/Source/Runtime/Engine/Classes/Debug/LogVisualizerCameraController.h
@@ -8,11 +8,13 @@
 UCLASS(config=Input, hidedropdown)
 class ALogVisualizerCameraController : public ADebugCameraController
 {
+	GENERATED_BODY()
+public:
+	ALogVisualizerCameraController(const FObjectInitializer& ObjectInitializer = FObjectInitializer::Get());
 public:
 	DECLARE_DELEGATE_OneParam(FActorSelectedDelegate, class AActor*);
 	DECLARE_DELEGATE_OneParam(FLogEntryIterationDelegate, int32);
 
-	GENERATED_UCLASS_BODY()
 
 	UPROPERTY()
 	class AActor* PickedActor;
diff --git a/Engine/Source/Runtime/Engine/Classes/Debug/LogVisualizerHUD.h b/Engine/Source/Runtime/Engine/Classes/Debug/LogVisualizerHUD.h
--- a/Engine/Source/Runtime/Engine/Classes/Debug/LogVisualizerHUD.h
+++ b/Engine/Source/Runtime/Engine/Classes/Debug/LogVisualizerHUD.h
@@ -7,7 +7,9 @@
 UCLASS()
 class ALogVisualizerHUD : public ADebugCameraHUD
 {
-	GENERATED_UCLASS_BODY()
+	GENERATED_BODY()
+public:
+	ALogVisualizerHUD(const FObjectInitializer& ObjectInitializer = FObjectInitializer::Get());
 
 	FFontRenderInfo TextRenderInfo;
 
diff --git a/Engine/Source/Runtime/Engine/Classes/Debug/ReporterBase.h b/Engine/Source/Runtime/Engine/Classes/Debug/ReporterBase.h
--- a/Engine/Source/Runtime/Engine/Classes/Debug/ReporterBase.h
+++ b/Engine/Source/Runtime/Engine/Classes/Debug/ReporterBase.h
@@ -17,7 +17,9 @@ namespace EReporterLineStyle
 UCLASS(Abstract)
 class UReporterBase : public UObject
 {
-	GENERATED_UCLASS_BODY()
+	GENERATED_BODY()
+public:
+	UReporterBase(const FObjectInitializer& ObjectInitializer = FObjectInitializer::Get());
 
 public:
 	
@@ -53,4 +55,4 @@ protected:
 	 * @param Color The color to draw the triangle
 	 */
 	void DrawTriangle(UCanvas* Canvas, const FVector2D& Vertex1, const FVector2D& Vertex2, const FVector2D& Vertex3, const FLinearColor& Color);
-};
\ No newline at end of file
+};
diff --git a/Engine/Source/Runtime/Engine/Classes/Debug/ReporterGraph.h b/Engine/Source/Runtime/Engine/Classes/Debug/ReporterGraph.h
--- a/Engine/Source/Runtime/Engine/Classes/Debug/ReporterGraph.h
+++ b/Engine/Source/Runtime/Engine/Classes/Debug/ReporterGraph.h
@@ -96,7 +96,9 @@ struct FGraphLine
 UCLASS()
 class ENGINE_API UReporterGraph : public UReporterBase
 {
-	GENERATED_UCLASS_BODY()
+	GENERATED_BODY()
+public:
+	UReporterGraph(const FObjectInitializer& ObjectInitializer = FObjectInitializer::Get());
 
 public:
 
diff --git a/Engine/Source/Runtime/Engine/Classes/DeviceProfiles/DeviceProfile.h b/Engine/Source/Runtime/Engine/Classes/DeviceProfiles/DeviceProfile.h
--- a/Engine/Source/Runtime/Engine/Classes/DeviceProfiles/DeviceProfile.h
+++ b/Engine/Source/Runtime/Engine/Classes/DeviceProfiles/DeviceProfile.h
@@ -15,7 +15,9 @@ DECLARE_DELEGATE(FOnCVarsUpdated);
 UCLASS(config=DeviceProfiles, perobjectconfig, Blueprintable)
 class ENGINE_API UDeviceProfile : public UObject
 {
-	GENERATED_UCLASS_BODY()
+	GENERATED_BODY()
+public:
+	UDeviceProfile(const FObjectInitializer& ObjectInitializer = FObjectInitializer::Get());
 
 	/** The type of this profile, I.e. IOS, Windows, PS4 etc */
 	UPROPERTY(VisibleAnywhere, config, Category=DeviceSettings)
diff --git a/Engine/Source/Runtime/Engine/Classes/DeviceProfiles/DeviceProfileManager.h b/Engine/Source/Runtime/Engine/Classes/DeviceProfiles/DeviceProfileManager.h
--- a/Engine/Source/Runtime/Engine/Classes/DeviceProfiles/DeviceProfileManager.h
+++ b/Engine/Source/Runtime/Engine/Classes/DeviceProfiles/DeviceProfileManager.h
@@ -19,9 +19,11 @@ DECLARE_MULTICAST_DELEGATE( FOnManagerUpdated );
  UCLASS( config=DeviceProfiles, transient )
 class ENGINE_API UDeviceProfileManager : public UObject
 {
+	GENERATED_BODY()
+public:
+	UDeviceProfileManager(const FObjectInitializer& ObjectInitializer = FObjectInitializer::Get());
 public:
 
-	GENERATED_UCLASS_BODY()
 
 	/**
 	 * Startup and select the active device profile
@@ -128,4 +130,4 @@ private:
 	static FString DeviceProfileFileName;
 
 	int32 RenameIndex;
-};
\ No newline at end of file
+};
diff --git a/Engine/Source/Runtime/Engine/Classes/Distributions/Distribution.h b/Engine/Source/Runtime/Engine/Classes/Distributions/Distribution.h
--- a/Engine/Source/Runtime/Engine/Classes/Distributions/Distribution.h
+++ b/Engine/Source/Runtime/Engine/Classes/Distributions/Distribution.h
@@ -58,7 +58,9 @@ struct FRawDistribution
 UCLASS(DefaultToInstanced, collapsecategories, hidecategories=Object, editinlinenew, abstract,MinimalAPI)
 class UDistribution : public UObject, public FCurveEdInterface
 {
-	GENERATED_UCLASS_BODY()
+	GENERATED_BODY()
+public:
+	ENGINE_API UDistribution(const FObjectInitializer& ObjectInitializer = FObjectInitializer::Get());
 
 	/** Default value for initializing and checking correct values on UDistributions. */
 	static const float DefaultValue;
diff --git a/Engine/Source/Runtime/Engine/Classes/Distributions/DistributionFloat.h b/Engine/Source/Runtime/Engine/Classes/Distributions/DistributionFloat.h
--- a/Engine/Source/Runtime/Engine/Classes/Distributions/DistributionFloat.h
+++ b/Engine/Source/Runtime/Engine/Classes/Distributions/DistributionFloat.h
@@ -20,7 +20,7 @@ struct FFloatDistribution
 USTRUCT()
 struct FRawDistributionFloat : public FRawDistribution
 {
-	GENERATED_USTRUCT_BODY()
+	GENERATED_BODY()
 
 private:
 	UPROPERTY()
@@ -71,10 +71,11 @@ public:
 	
 };
 
-UCLASS(abstract, customconstructor,MinimalAPI)
+UCLASS(abstract,MinimalAPI)
 class UDistributionFloat : public UDistribution
 {
-	GENERATED_UCLASS_BODY()
+	GENERATED_BODY()
+public:
 
 	/** Can this variable be baked out to a FRawDistribution? Should be true 99% of the time*/
 	UPROPERTY(EditAnywhere, Category=Baked)
diff --git a/Engine/Source/Runtime/Engine/Classes/Distributions/DistributionFloatConstant.h b/Engine/Source/Runtime/Engine/Classes/Distributions/DistributionFloatConstant.h
--- a/Engine/Source/Runtime/Engine/Classes/Distributions/DistributionFloatConstant.h
+++ b/Engine/Source/Runtime/Engine/Classes/Distributions/DistributionFloatConstant.h
@@ -8,7 +8,9 @@
 UCLASS(collapsecategories, hidecategories=Object, editinlinenew,MinimalAPI)
 class UDistributionFloatConstant : public UDistributionFloat
 {
-	GENERATED_UCLASS_BODY()
+	GENERATED_BODY()
+public:
+	ENGINE_API UDistributionFloatConstant(const FObjectInitializer& ObjectInitializer = FObjectInitializer::Get());
 
 	/** This float will be returned for all values of time. */
 	UPROPERTY(EditAnywhere, Category=DistributionFloatConstant)
diff --git a/Engine/Source/Runtime/Engine/Classes/Distributions/DistributionFloatConstantCurve.h b/Engine/Source/Runtime/Engine/Classes/Distributions/DistributionFloatConstantCurve.h
--- a/Engine/Source/Runtime/Engine/Classes/Distributions/DistributionFloatConstantCurve.h
+++ b/Engine/Source/Runtime/Engine/Classes/Distributions/DistributionFloatConstantCurve.h
@@ -8,7 +8,9 @@
 UCLASS(collapsecategories, hidecategories=Object, editinlinenew,MinimalAPI)
 class UDistributionFloatConstantCurve : public UDistributionFloat
 {
-	GENERATED_UCLASS_BODY()
+	GENERATED_BODY()
+public:
+	ENGINE_API UDistributionFloatConstantCurve(const FObjectInitializer& ObjectInitializer = FObjectInitializer::Get());
 
 	/** Keyframe data for how output constant varies over time. */
 	UPROPERTY(EditAnywhere, Category=DistributionFloatConstantCurve)
diff --git a/Engine/Source/Runtime/Engine/Classes/Distributions/DistributionFloatParameterBase.h b/Engine/Source/Runtime/Engine/Classes/Distributions/DistributionFloatParameterBase.h
--- a/Engine/Source/Runtime/Engine/Classes/Distributions/DistributionFloatParameterBase.h
+++ b/Engine/Source/Runtime/Engine/Classes/Distributions/DistributionFloatParameterBase.h
@@ -8,7 +8,9 @@
 UCLASS(abstract, collapsecategories, hidecategories=Object, editinlinenew)
 class UDistributionFloatParameterBase : public UDistributionFloatConstant
 {
-	GENERATED_UCLASS_BODY()
+	GENERATED_BODY()
+public:
+	UDistributionFloatParameterBase(const FObjectInitializer& ObjectInitializer = FObjectInitializer::Get());
 
 	/** todo document */
 	UPROPERTY(EditAnywhere, Category=DistributionFloatParameterBase)
diff --git a/Engine/Source/Runtime/Engine/Classes/Distributions/DistributionFloatParticleParameter.h b/Engine/Source/Runtime/Engine/Classes/Distributions/DistributionFloatParticleParameter.h
--- a/Engine/Source/Runtime/Engine/Classes/Distributions/DistributionFloatParticleParameter.h
+++ b/Engine/Source/Runtime/Engine/Classes/Distributions/DistributionFloatParticleParameter.h
@@ -8,7 +8,9 @@
 UCLASS(collapsecategories, hidecategories=Object, editinlinenew, MinimalAPI)
 class UDistributionFloatParticleParameter : public UDistributionFloatParameterBase
 {
-	GENERATED_UCLASS_BODY()
+	GENERATED_BODY()
+public:
+	ENGINE_API UDistributionFloatParticleParameter(const FObjectInitializer& ObjectInitializer = FObjectInitializer::Get());
 
 
 	// Begin UDistributionFloatParameterBase Interface
diff --git a/Engine/Source/Runtime/Engine/Classes/Distributions/DistributionFloatUniform.h b/Engine/Source/Runtime/Engine/Classes/Distributions/DistributionFloatUniform.h
--- a/Engine/Source/Runtime/Engine/Classes/Distributions/DistributionFloatUniform.h
+++ b/Engine/Source/Runtime/Engine/Classes/Distributions/DistributionFloatUniform.h
@@ -9,7 +9,9 @@
 UCLASS(collapsecategories, hidecategories=Object, editinlinenew, MinimalAPI)
 class UDistributionFloatUniform : public UDistributionFloat
 {
-	GENERATED_UCLASS_BODY()
+	GENERATED_BODY()
+public:
+	ENGINE_API UDistributionFloatUniform(const FObjectInitializer& ObjectInitializer = FObjectInitializer::Get());
 
 	/** Low end of output float distribution. */
 	UPROPERTY(EditAnywhere, Category=DistributionFloatUniform)
diff --git a/Engine/Source/Runtime/Engine/Classes/Distributions/DistributionFloatUniformCurve.h b/Engine/Source/Runtime/Engine/Classes/Distributions/DistributionFloatUniformCurve.h
--- a/Engine/Source/Runtime/Engine/Classes/Distributions/DistributionFloatUniformCurve.h
+++ b/Engine/Source/Runtime/Engine/Classes/Distributions/DistributionFloatUniformCurve.h
@@ -9,7 +9,9 @@
 UCLASS(collapsecategories, hidecategories=Object, editinlinenew, MinimalAPI)
 class UDistributionFloatUniformCurve : public UDistributionFloat
 {
-	GENERATED_UCLASS_BODY()
+	GENERATED_BODY()
+public:
+	ENGINE_API UDistributionFloatUniformCurve(const FObjectInitializer& ObjectInitializer = FObjectInitializer::Get());
 
 	/** Keyframe data for how output constant varies over time. */
 	UPROPERTY(EditAnywhere, Category=DistributionFloatUniformCurve)
diff --git a/Engine/Source/Runtime/Engine/Classes/Distributions/DistributionVector.h b/Engine/Source/Runtime/Engine/Classes/Distributions/DistributionVector.h
--- a/Engine/Source/Runtime/Engine/Classes/Distributions/DistributionVector.h
+++ b/Engine/Source/Runtime/Engine/Classes/Distributions/DistributionVector.h
@@ -50,7 +50,7 @@ struct FVector4Distribution
 USTRUCT()
 struct FRawDistributionVector : public FRawDistribution
 {
-	GENERATED_USTRUCT_BODY()
+	GENERATED_BODY()
 
 private:
 	UPROPERTY()
@@ -101,10 +101,11 @@ public:
 	
 };
 
-UCLASS(abstract, customconstructor)
+UCLASS(abstract)
 class ENGINE_API UDistributionVector : public UDistribution
 {
-	GENERATED_UCLASS_BODY()
+	GENERATED_BODY()
+public:
 
 	/** Can this variable be baked out to a FRawDistribution? Should be true 99% of the time*/
 	UPROPERTY(EditAnywhere, Category=Baked)
diff --git a/Engine/Source/Runtime/Engine/Classes/Distributions/DistributionVectorConstant.h b/Engine/Source/Runtime/Engine/Classes/Distributions/DistributionVectorConstant.h
--- a/Engine/Source/Runtime/Engine/Classes/Distributions/DistributionVectorConstant.h
+++ b/Engine/Source/Runtime/Engine/Classes/Distributions/DistributionVectorConstant.h
@@ -8,7 +8,9 @@
 UCLASS(collapsecategories, hidecategories=Object, editinlinenew)
 class ENGINE_API UDistributionVectorConstant : public UDistributionVector
 {
-	GENERATED_UCLASS_BODY()
+	GENERATED_BODY()
+public:
+	UDistributionVectorConstant(const FObjectInitializer& ObjectInitializer = FObjectInitializer::Get());
 
 	/** This FVector will be returned for all input times. */
 	UPROPERTY(EditAnywhere, Category=DistributionVectorConstant)
diff --git a/Engine/Source/Runtime/Engine/Classes/Distributions/DistributionVectorConstantCurve.h b/Engine/Source/Runtime/Engine/Classes/Distributions/DistributionVectorConstantCurve.h
--- a/Engine/Source/Runtime/Engine/Classes/Distributions/DistributionVectorConstantCurve.h
+++ b/Engine/Source/Runtime/Engine/Classes/Distributions/DistributionVectorConstantCurve.h
@@ -8,7 +8,9 @@
 UCLASS(collapsecategories, hidecategories=Object, editinlinenew,MinimalAPI)
 class UDistributionVectorConstantCurve : public UDistributionVector
 {
-	GENERATED_UCLASS_BODY()
+	GENERATED_BODY()
+public:
+	ENGINE_API UDistributionVectorConstantCurve(const FObjectInitializer& ObjectInitializer = FObjectInitializer::Get());
 
 	/** Keyframe data for each component (X,Y,Z) over time. */
 	UPROPERTY(EditAnywhere, Category=DistributionVectorConstantCurve)
diff --git a/Engine/Source/Runtime/Engine/Classes/Distributions/DistributionVectorParameterBase.h b/Engine/Source/Runtime/Engine/Classes/Distributions/DistributionVectorParameterBase.h
--- a/Engine/Source/Runtime/Engine/Classes/Distributions/DistributionVectorParameterBase.h
+++ b/Engine/Source/Runtime/Engine/Classes/Distributions/DistributionVectorParameterBase.h
@@ -8,7 +8,9 @@
 UCLASS(abstract, collapsecategories, hidecategories=Object, editinlinenew)
 class UDistributionVectorParameterBase : public UDistributionVectorConstant
 {
-	GENERATED_UCLASS_BODY()
+	GENERATED_BODY()
+public:
+	UDistributionVectorParameterBase(const FObjectInitializer& ObjectInitializer = FObjectInitializer::Get());
 
 	UPROPERTY(EditAnywhere, Category=DistributionVectorParameterBase)
 	FName ParameterName;
diff --git a/Engine/Source/Runtime/Engine/Classes/Distributions/DistributionVectorParticleParameter.h b/Engine/Source/Runtime/Engine/Classes/Distributions/DistributionVectorParticleParameter.h
--- a/Engine/Source/Runtime/Engine/Classes/Distributions/DistributionVectorParticleParameter.h
+++ b/Engine/Source/Runtime/Engine/Classes/Distributions/DistributionVectorParticleParameter.h
@@ -8,7 +8,9 @@
 UCLASS(collapsecategories, hidecategories=Object, editinlinenew, MinimalAPI)
 class UDistributionVectorParticleParameter : public UDistributionVectorParameterBase
 {
-	GENERATED_UCLASS_BODY()
+	GENERATED_BODY()
+public:
+	ENGINE_API UDistributionVectorParticleParameter(const FObjectInitializer& ObjectInitializer = FObjectInitializer::Get());
 
 
 	// Begin UDistributionVectorParameterBase Interface
diff --git a/Engine/Source/Runtime/Engine/Classes/Distributions/DistributionVectorUniform.h b/Engine/Source/Runtime/Engine/Classes/Distributions/DistributionVectorUniform.h
--- a/Engine/Source/Runtime/Engine/Classes/Distributions/DistributionVectorUniform.h
+++ b/Engine/Source/Runtime/Engine/Classes/Distributions/DistributionVectorUniform.h
@@ -8,7 +8,9 @@
 UCLASS(collapsecategories, hidecategories=Object, editinlinenew, MinimalAPI)
 class UDistributionVectorUniform : public UDistributionVector
 {
-	GENERATED_UCLASS_BODY()
+	GENERATED_BODY()
+public:
+	ENGINE_API UDistributionVectorUniform(const FObjectInitializer& ObjectInitializer = FObjectInitializer::Get());
 
 	/** Upper end of FVector magnitude range. */
 	UPROPERTY(EditAnywhere, Category=DistributionVectorUniform)
diff --git a/Engine/Source/Runtime/Engine/Classes/Distributions/DistributionVectorUniformCurve.h b/Engine/Source/Runtime/Engine/Classes/Distributions/DistributionVectorUniformCurve.h
--- a/Engine/Source/Runtime/Engine/Classes/Distributions/DistributionVectorUniformCurve.h
+++ b/Engine/Source/Runtime/Engine/Classes/Distributions/DistributionVectorUniformCurve.h
@@ -8,7 +8,9 @@
 UCLASS(collapsecategories, hidecategories=Object, editinlinenew, MinimalAPI)
 class UDistributionVectorUniformCurve : public UDistributionVector
 {
-	GENERATED_UCLASS_BODY()
+	GENERATED_BODY()
+public:
+	ENGINE_API UDistributionVectorUniformCurve(const FObjectInitializer& ObjectInitializer = FObjectInitializer::Get());
 
 	/** Keyframe data for how output constant varies over time. */
 	UPROPERTY(EditAnywhere, Category=DistributionVectorUniformCurve)
diff --git a/Engine/Source/Runtime/Engine/Classes/EdGraph/EdGraph.h b/Engine/Source/Runtime/Engine/Classes/EdGraph/EdGraph.h
--- a/Engine/Source/Runtime/Engine/Classes/EdGraph/EdGraph.h
+++ b/Engine/Source/Runtime/Engine/Classes/EdGraph/EdGraph.h
@@ -9,7 +9,7 @@ struct FEdGraphEditAction;
 USTRUCT()
 struct ENGINE_API FGraphReference
 {
-	GENERATED_USTRUCT_BODY()
+	GENERATED_BODY()
 protected:
 	// Reference to the actual graph
 	UPROPERTY()
@@ -55,7 +55,9 @@ struct TStructOpsTypeTraits<FGraphReference> : public TStructOpsTypeTraitsBase
 UCLASS()
 class ENGINE_API UEdGraph : public UObject
 {
-	GENERATED_UCLASS_BODY()
+	GENERATED_BODY()
+public:
+	UEdGraph(const FObjectInitializer& ObjectInitializer = FObjectInitializer::Get());
 
 	/** The schema that this graph obeys */
 	UPROPERTY()
diff --git a/Engine/Source/Runtime/Engine/Classes/EdGraph/EdGraphNode.h b/Engine/Source/Runtime/Engine/Classes/EdGraph/EdGraphNode.h
--- a/Engine/Source/Runtime/Engine/Classes/EdGraph/EdGraphNode.h
+++ b/Engine/Source/Runtime/Engine/Classes/EdGraph/EdGraphNode.h
@@ -85,7 +85,9 @@ struct FGraphNodeContextMenuBuilder
 UCLASS()
 class ENGINE_API UEdGraphNode : public UObject
 {
-	GENERATED_UCLASS_BODY()
+	GENERATED_BODY()
+public:
+	UEdGraphNode(const FObjectInitializer& ObjectInitializer = FObjectInitializer::Get());
 
 	/** List of connector pins */
 	UPROPERTY()
diff --git a/Engine/Source/Runtime/Engine/Classes/EdGraph/EdGraphNode_Documentation.h b/Engine/Source/Runtime/Engine/Classes/EdGraph/EdGraphNode_Documentation.h
--- a/Engine/Source/Runtime/Engine/Classes/EdGraph/EdGraphNode_Documentation.h
+++ b/Engine/Source/Runtime/Engine/Classes/EdGraph/EdGraphNode_Documentation.h
@@ -6,7 +6,9 @@
 UCLASS(MinimalAPI)
 class UEdGraphNode_Documentation : public UEdGraphNode
 {
-	GENERATED_UCLASS_BODY()
+	GENERATED_BODY()
+public:
+	ENGINE_API UEdGraphNode_Documentation(const FObjectInitializer& ObjectInitializer = FObjectInitializer::Get());
 
 	/** Documentation Link */
 	UPROPERTY(meta=(FriendlyName="Documentation Link"))
diff --git a/Engine/Source/Runtime/Engine/Classes/EdGraph/EdGraphPin.h b/Engine/Source/Runtime/Engine/Classes/EdGraph/EdGraphPin.h
--- a/Engine/Source/Runtime/Engine/Classes/EdGraph/EdGraphPin.h
+++ b/Engine/Source/Runtime/Engine/Classes/EdGraph/EdGraphPin.h
@@ -8,7 +8,7 @@
 USTRUCT()
 struct FSimpleMemberReference
 {
-	GENERATED_USTRUCT_BODY()
+	GENERATED_BODY()
 
 	/** Class that this member is defined in. */
 	UPROPERTY()
@@ -48,7 +48,7 @@ FORCEINLINE FArchive& operator<<(FArchive& Ar, FSimpleMemberReference& Data)
 USTRUCT()
 struct FEdGraphPinType
 {
-	GENERATED_USTRUCT_BODY()
+	GENERATED_BODY()
 
 	/** Category of pin type */
 	UPROPERTY()
@@ -152,7 +152,9 @@ enum EBlueprintPinStyleType
 UCLASS(MinimalAPI)
 class UEdGraphPin : public UObject
 {
-	GENERATED_UCLASS_BODY()
+	GENERATED_BODY()
+public:
+	ENGINE_API UEdGraphPin(const FObjectInitializer& ObjectInitializer = FObjectInitializer::Get());
 
 	/** Name of this pin */
 	UPROPERTY()
diff --git a/Engine/Source/Runtime/Engine/Classes/EdGraph/EdGraphSchema.h b/Engine/Source/Runtime/Engine/Classes/EdGraph/EdGraphSchema.h
--- a/Engine/Source/Runtime/Engine/Classes/EdGraph/EdGraphSchema.h
+++ b/Engine/Source/Runtime/Engine/Classes/EdGraph/EdGraphSchema.h
@@ -50,7 +50,7 @@ enum ECanCreateConnectionResponse
 USTRUCT()
 struct ENGINE_API FEdGraphSchemaAction
 {
-	GENERATED_USTRUCT_BODY()
+	GENERATED_BODY()
 
 	// Simple type info
 	static FName StaticGetTypeId() {static FName Type("FEdGraphSchemaAction"); return Type;}
@@ -146,7 +146,7 @@ struct ENGINE_API FEdGraphSchemaAction
 USTRUCT()
 struct ENGINE_API FEdGraphSchemaAction_NewNode : public FEdGraphSchemaAction
 {
-	GENERATED_USTRUCT_BODY()
+	GENERATED_BODY()
 
 	// Simple type info
 	static FName StaticGetTypeId() {static FName Type("FEdGraphSchemaAction_NewNode"); return Type;}
@@ -398,7 +398,9 @@ public:
 UCLASS(abstract)
 class ENGINE_API UEdGraphSchema : public UObject
 {
-	GENERATED_UCLASS_BODY()
+	GENERATED_BODY()
+public:
+	UEdGraphSchema(const FObjectInitializer& ObjectInitializer = FObjectInitializer::Get());
 
 
 	/**
diff --git a/Engine/Source/Runtime/Engine/Classes/EditorFramework/AssetImportData.h b/Engine/Source/Runtime/Engine/Classes/EditorFramework/AssetImportData.h
--- a/Engine/Source/Runtime/Engine/Classes/EditorFramework/AssetImportData.h
+++ b/Engine/Source/Runtime/Engine/Classes/EditorFramework/AssetImportData.h
@@ -9,7 +9,9 @@
 UCLASS(MinimalAPI)
 class UAssetImportData : public UObject
 {
-	GENERATED_UCLASS_BODY()
+	GENERATED_BODY()
+public:
+	ENGINE_API UAssetImportData(const FObjectInitializer& ObjectInitializer = FObjectInitializer::Get());
 
 	/** Path to the resource used to construct this static mesh. Relative to the object's package, BaseDir() or absolute */
 	UPROPERTY(EditAnywhere, Category=ImportSettings)
diff --git a/Engine/Source/Runtime/Engine/Classes/EditorFramework/ThumbnailInfo.h b/Engine/Source/Runtime/Engine/Classes/EditorFramework/ThumbnailInfo.h
--- a/Engine/Source/Runtime/Engine/Classes/EditorFramework/ThumbnailInfo.h
+++ b/Engine/Source/Runtime/Engine/Classes/EditorFramework/ThumbnailInfo.h
@@ -10,7 +10,9 @@
 UCLASS(MinimalAPI)
 class UThumbnailInfo : public UObject
 {
-	GENERATED_UCLASS_BODY()
+	GENERATED_BODY()
+public:
+	ENGINE_API UThumbnailInfo(const FObjectInitializer& ObjectInitializer = FObjectInitializer::Get());
 };
 
 
diff --git a/Engine/Source/Runtime/Engine/Classes/Engine/ActorChannel.h b/Engine/Source/Runtime/Engine/Classes/Engine/ActorChannel.h
--- a/Engine/Source/Runtime/Engine/Classes/Engine/ActorChannel.h
+++ b/Engine/Source/Runtime/Engine/Classes/Engine/ActorChannel.h
@@ -36,11 +36,12 @@
  *		| </End Tag>           |                                                                           |
  *		|----------------------|---------------------------------------------------------------------------|
  */
-UCLASS(transient, customConstructor)
+UCLASS(transient)
 class ENGINE_API UActorChannel
 	: public UChannel
 {
-	GENERATED_UCLASS_BODY()
+	GENERATED_BODY()
+public:
 
 	friend class FObjectReplicator;
 
diff --git a/Engine/Source/Runtime/Engine/Classes/Engine/AssetUserData.h b/Engine/Source/Runtime/Engine/Classes/Engine/AssetUserData.h
--- a/Engine/Source/Runtime/Engine/Classes/Engine/AssetUserData.h
+++ b/Engine/Source/Runtime/Engine/Classes/Engine/AssetUserData.h
@@ -12,7 +12,9 @@ UCLASS(DefaultToInstanced)
 class ENGINE_API UAssetUserData
 	: public UObject
 {
-	GENERATED_UCLASS_BODY()
+	GENERATED_BODY()
+public:
+	UAssetUserData(const FObjectInitializer& ObjectInitializer = FObjectInitializer::Get());
 
 	/** used for debugging UAssetUserData data in editor */
 	virtual void Draw(class FPrimitiveDrawInterface* PDI, const class FSceneView* View) const {}
diff --git a/Engine/Source/Runtime/Engine/Classes/Engine/BlendableInterface.h b/Engine/Source/Runtime/Engine/Classes/Engine/BlendableInterface.h
--- a/Engine/Source/Runtime/Engine/Classes/Engine/BlendableInterface.h
+++ b/Engine/Source/Runtime/Engine/Classes/Engine/BlendableInterface.h
@@ -31,7 +31,9 @@ enum EBlendableLocation
 UINTERFACE()
 class UBlendableInterface : public UInterface
 {
-	GENERATED_UINTERFACE_BODY()
+	GENERATED_BODY()
+public:
+	UBlendableInterface(const FObjectInitializer& ObjectInitializer = FObjectInitializer::Get());
 };
 
 /**
@@ -39,7 +41,8 @@ class UBlendableInterface : public UInterface
  */
 class IBlendableInterface
 {
-	GENERATED_IINTERFACE_BODY()
+	GENERATED_BODY()
+public:
 		
 	// @param Weight 0..1, 1=fully take the values from this object, crash if outside the valid range
 	virtual void OverrideBlendableSettings(class FSceneView& View, float Weight) const = 0;
diff --git a/Engine/Source/Runtime/Engine/Classes/Engine/BlockingVolume.h b/Engine/Source/Runtime/Engine/Classes/Engine/BlockingVolume.h
--- a/Engine/Source/Runtime/Engine/Classes/Engine/BlockingVolume.h
+++ b/Engine/Source/Runtime/Engine/Classes/Engine/BlockingVolume.h
@@ -11,7 +11,9 @@ UCLASS(MinimalAPI)
 class ABlockingVolume
 	: public AVolume
 {
-	GENERATED_UCLASS_BODY()
+	GENERATED_BODY()
+public:
+	ENGINE_API ABlockingVolume(const FObjectInitializer& ObjectInitializer = FObjectInitializer::Get());
 
 public:
 
diff --git a/Engine/Source/Runtime/Engine/Classes/Engine/Blueprint.h b/Engine/Source/Runtime/Engine/Classes/Engine/Blueprint.h
--- a/Engine/Source/Runtime/Engine/Classes/Engine/Blueprint.h
+++ b/Engine/Source/Runtime/Engine/Classes/Engine/Blueprint.h
@@ -113,7 +113,7 @@ public:
 USTRUCT()
 struct FBPVariableMetaDataEntry
 {
-	GENERATED_USTRUCT_BODY()
+	GENERATED_BODY()
 
 	/** Name of metadata key */
 	UPROPERTY(EditAnywhere, Category=BPVariableMetaDataEntry)
@@ -138,7 +138,7 @@ struct FBPVariableMetaDataEntry
 USTRUCT()
 struct FBPVariableDescription
 {
-	GENERATED_USTRUCT_BODY()
+	GENERATED_BODY()
 
 	/** Name of the variable */
 	UPROPERTY(EditAnywhere, Category=BPVariableDescription)
@@ -198,7 +198,7 @@ struct FBPVariableDescription
 USTRUCT()
 struct FBPInterfaceDescription
 {
-	GENERATED_USTRUCT_BODY()
+	GENERATED_BODY()
 
 	/** Reference to the interface class we're adding to this blueprint */
 	UPROPERTY()
@@ -218,7 +218,7 @@ struct FBPInterfaceDescription
 USTRUCT()
 struct FEditedDocumentInfo
 {
-	GENERATED_USTRUCT_BODY()
+	GENERATED_BODY()
 
 	UPROPERTY()
 	UObject* EditedObject;
@@ -264,7 +264,9 @@ struct FEditedDocumentInfo
 UCLASS(config=Engine, BlueprintType)
 class ENGINE_API UBlueprint : public UBlueprintCore
 {
-	GENERATED_UCLASS_BODY()
+	GENERATED_BODY()
+public:
+	UBlueprint(const FObjectInitializer& ObjectInitializer = FObjectInitializer::Get());
 
 	/** Whether or not this blueprint should recompile itself on load */
 	UPROPERTY(config)
diff --git a/Engine/Source/Runtime/Engine/Classes/Engine/BlueprintCore.h b/Engine/Source/Runtime/Engine/Classes/Engine/BlueprintCore.h
--- a/Engine/Source/Runtime/Engine/Classes/Engine/BlueprintCore.h
+++ b/Engine/Source/Runtime/Engine/Classes/Engine/BlueprintCore.h
@@ -9,7 +9,9 @@ UCLASS()
 class ENGINE_API UBlueprintCore
 	: public UObject
 {
-	GENERATED_UCLASS_BODY()
+	GENERATED_BODY()
+public:
+	UBlueprintCore(const FObjectInitializer& ObjectInitializer = FObjectInitializer::Get());
 
 	/** Pointer to the skeleton class; this is regenerated any time a member variable or function is added but  
 	is usually incomplete (no code or hidden autogenerated variables are added to it) */
diff --git a/Engine/Source/Runtime/Engine/Classes/Engine/BlueprintGeneratedClass.h b/Engine/Source/Runtime/Engine/Classes/Engine/BlueprintGeneratedClass.h
--- a/Engine/Source/Runtime/Engine/Classes/Engine/BlueprintGeneratedClass.h
+++ b/Engine/Source/Runtime/Engine/Classes/Engine/BlueprintGeneratedClass.h
@@ -12,7 +12,7 @@ class UEdGraphPin;
 USTRUCT()
 struct FNodeToCodeAssociation
 {
-	GENERATED_USTRUCT_BODY()
+	GENERATED_BODY()
 
 public:
 	TWeakObjectPtr<UEdGraphNode> Node;
@@ -38,7 +38,7 @@ public:
 USTRUCT()
 struct FDebuggingInfoForSingleFunction
 {
-	GENERATED_USTRUCT_BODY()
+	GENERATED_BODY()
 
 public:
 	// Reverse map from code offset to source node
@@ -63,7 +63,7 @@ public:
 USTRUCT()
 struct FPointerToUberGraphFrame
 {
-	GENERATED_USTRUCT_BODY()
+	GENERATED_BODY()
 
 public:
 	uint8* RawPointer;
@@ -180,7 +180,7 @@ public:
 USTRUCT()
 struct ENGINE_API FBlueprintDebugData
 {
-	GENERATED_USTRUCT_BODY()
+	GENERATED_BODY()
 	FBlueprintDebugData()
 	{
 	}
@@ -420,7 +420,7 @@ public:
 USTRUCT()
 struct ENGINE_API FEventGraphFastCallPair
 {
-	GENERATED_USTRUCT_BODY()
+	GENERATED_BODY()
 
 	UPROPERTY()
 	UFunction* FunctionToPatch;
@@ -432,7 +432,9 @@ struct ENGINE_API FEventGraphFastCallPair
 UCLASS()
 class ENGINE_API UBlueprintGeneratedClass : public UClass
 {
-	GENERATED_UCLASS_BODY()
+	GENERATED_BODY()
+public:
+	UBlueprintGeneratedClass(const FObjectInitializer& ObjectInitializer = FObjectInitializer::Get());
 
 public:
 
diff --git a/Engine/Source/Runtime/Engine/Classes/Engine/BookMark.h b/Engine/Source/Runtime/Engine/Classes/Engine/BookMark.h
--- a/Engine/Source/Runtime/Engine/Classes/Engine/BookMark.h
+++ b/Engine/Source/Runtime/Engine/Classes/Engine/BookMark.h
@@ -11,7 +11,9 @@
 UCLASS(hidecategories=Object, MinimalAPI)
 class UBookMark : public UObject
 {
-	GENERATED_UCLASS_BODY()
+	GENERATED_BODY()
+public:
+	ENGINE_API UBookMark(const FObjectInitializer& ObjectInitializer = FObjectInitializer::Get());
 
 	/** Camera position */
 	UPROPERTY(EditAnywhere, Category=BookMark)
diff --git a/Engine/Source/Runtime/Engine/Classes/Engine/BookMark2D.h b/Engine/Source/Runtime/Engine/Classes/Engine/BookMark2D.h
--- a/Engine/Source/Runtime/Engine/Classes/Engine/BookMark2D.h
+++ b/Engine/Source/Runtime/Engine/Classes/Engine/BookMark2D.h
@@ -12,7 +12,9 @@ UCLASS(hidecategories=Object)
 class UBookMark2D
 	: public UObject
 {
-	GENERATED_UCLASS_BODY()
+	GENERATED_BODY()
+public:
+	UBookMark2D(const FObjectInitializer& ObjectInitializer = FObjectInitializer::Get());
 
 	/** Zoom of the camera */
 	UPROPERTY(EditAnywhere, Category=BookMark2D)
diff --git a/Engine/Source/Runtime/Engine/Classes/Engine/BoxReflectionCapture.h b/Engine/Source/Runtime/Engine/Classes/Engine/BoxReflectionCapture.h
--- a/Engine/Source/Runtime/Engine/Classes/Engine/BoxReflectionCapture.h
+++ b/Engine/Source/Runtime/Engine/Classes/Engine/BoxReflectionCapture.h
@@ -14,5 +14,7 @@ UCLASS(hidecategories=(Collision, Attachment, Actor), MinimalAPI)
 class ABoxReflectionCapture
 	: public AReflectionCapture
 {
-	GENERATED_UCLASS_BODY()
+	GENERATED_BODY()
+public:
+	ENGINE_API ABoxReflectionCapture(const FObjectInitializer& ObjectInitializer = FObjectInitializer::Get());
 };
diff --git a/Engine/Source/Runtime/Engine/Classes/Engine/Breakpoint.h b/Engine/Source/Runtime/Engine/Classes/Engine/Breakpoint.h
--- a/Engine/Source/Runtime/Engine/Classes/Engine/Breakpoint.h
+++ b/Engine/Source/Runtime/Engine/Classes/Engine/Breakpoint.h
@@ -9,7 +9,9 @@ UCLASS()
 class ENGINE_API UBreakpoint
 	: public UObject
 {
-	GENERATED_UCLASS_BODY()
+	GENERATED_BODY()
+public:
+	UBreakpoint(const FObjectInitializer& ObjectInitializer = FObjectInitializer::Get());
 
 private:
 	// Is the breakpoint currently enabled?
diff --git a/Engine/Source/Runtime/Engine/Classes/Engine/Brush.h b/Engine/Source/Runtime/Engine/Classes/Engine/Brush.h
--- a/Engine/Source/Runtime/Engine/Classes/Engine/Brush.h
+++ b/Engine/Source/Runtime/Engine/Classes/Engine/Brush.h
@@ -45,7 +45,7 @@ enum EBrushType
 USTRUCT()
 struct FGeomSelection
 {
-	GENERATED_USTRUCT_BODY()
+	GENERATED_BODY()
 
 	UPROPERTY()
 	int32 Type;    // EGeometrySelectionType_
@@ -74,7 +74,9 @@ UCLASS(hidecategories=(Object, Collision, Display, Rendering, Physics, Input, Bl
 class ENGINE_API ABrush
 	: public AActor
 {
-	GENERATED_UCLASS_BODY()
+	GENERATED_BODY()
+public:
+	ABrush(const FObjectInitializer& ObjectInitializer = FObjectInitializer::Get());
 
 	/** Type of brush */
 	UPROPERTY(EditAnywhere, Category=Brush)
diff --git a/Engine/Source/Runtime/Engine/Classes/Engine/BrushBuilder.h b/Engine/Source/Runtime/Engine/Classes/Engine/BrushBuilder.h
--- a/Engine/Source/Runtime/Engine/Classes/Engine/BrushBuilder.h
+++ b/Engine/Source/Runtime/Engine/Classes/Engine/BrushBuilder.h
@@ -9,7 +9,7 @@
 USTRUCT()
 struct FBuilderPoly
 {
-	GENERATED_USTRUCT_BODY()
+	GENERATED_BODY()
 
 	UPROPERTY()
 	TArray<int32> VertexIndices;
@@ -48,7 +48,9 @@ UCLASS(abstract, hidecategories=(Object), MinimalAPI)
 class UBrushBuilder
 	: public UObject
 {
-	GENERATED_UCLASS_BODY()
+	GENERATED_BODY()
+public:
+	ENGINE_API UBrushBuilder(const FObjectInitializer& ObjectInitializer = FObjectInitializer::Get());
 
 	UPROPERTY()
 	FString BitmapFilename;
diff --git a/Engine/Source/Runtime/Engine/Classes/Engine/BrushShape.h b/Engine/Source/Runtime/Engine/Classes/Engine/BrushShape.h
--- a/Engine/Source/Runtime/Engine/Classes/Engine/BrushShape.h
+++ b/Engine/Source/Runtime/Engine/Classes/Engine/BrushShape.h
@@ -12,7 +12,9 @@ UCLASS(MinimalAPI)
 class ABrushShape
 	: public ABrush
 {
-	GENERATED_UCLASS_BODY()
+	GENERATED_BODY()
+public:
+	ENGINE_API ABrushShape(const FObjectInitializer& ObjectInitializer = FObjectInitializer::Get());
 
 public:
 
diff --git a/Engine/Source/Runtime/Engine/Classes/Engine/Canvas.h b/Engine/Source/Runtime/Engine/Classes/Engine/Canvas.h
--- a/Engine/Source/Runtime/Engine/Classes/Engine/Canvas.h
+++ b/Engine/Source/Runtime/Engine/Classes/Engine/Canvas.h
@@ -16,7 +16,7 @@ class UFont;
 USTRUCT()
 struct FCanvasIcon
 {
-	GENERATED_USTRUCT_BODY()
+	GENERATED_BODY()
 
 	/** Source texture */
 	UPROPERTY(EditAnywhere, BlueprintReadWrite, Category=CanvasIcon)
@@ -53,7 +53,9 @@ UCLASS(transient)
 class ENGINE_API UCanvas
 	: public UObject
 {
-	GENERATED_UCLASS_BODY()
+	GENERATED_BODY()
+public:
+	UCanvas(const FObjectInitializer& ObjectInitializer = FObjectInitializer::Get());
 
 	// Modifiable properties.
 	UPROPERTY()
diff --git a/Engine/Source/Runtime/Engine/Classes/Engine/CanvasRenderTarget2D.h b/Engine/Source/Runtime/Engine/Classes/Engine/CanvasRenderTarget2D.h
--- a/Engine/Source/Runtime/Engine/Classes/Engine/CanvasRenderTarget2D.h
+++ b/Engine/Source/Runtime/Engine/Classes/Engine/CanvasRenderTarget2D.h
@@ -24,7 +24,9 @@ UCLASS(BlueprintType, Blueprintable)
 class ENGINE_API UCanvasRenderTarget2D
 	: public UTextureRenderTarget2D
 {
-	GENERATED_UCLASS_BODY()
+	GENERATED_BODY()
+public:
+	UCanvasRenderTarget2D(const FObjectInitializer& ObjectInitializer = FObjectInitializer::Get());
 
 public:
 
diff --git a/Engine/Source/Runtime/Engine/Classes/Engine/Channel.h b/Engine/Source/Runtime/Engine/Classes/Engine/Channel.h
--- a/Engine/Source/Runtime/Engine/Classes/Engine/Channel.h
+++ b/Engine/Source/Runtime/Engine/Classes/Engine/Channel.h
@@ -33,7 +33,9 @@ UCLASS(transient)
 class ENGINE_API UChannel
 	: public UObject
 {
-	GENERATED_UCLASS_BODY()
+	GENERATED_BODY()
+public:
+	UChannel(const FObjectInitializer& ObjectInitializer = FObjectInitializer::Get());
 	
 	UPROPERTY()
 	class UNetConnection*	Connection;		// Owner connection.
diff --git a/Engine/Source/Runtime/Engine/Classes/Engine/ChildConnection.h b/Engine/Source/Runtime/Engine/Classes/Engine/ChildConnection.h
--- a/Engine/Source/Runtime/Engine/Classes/Engine/ChildConnection.h
+++ b/Engine/Source/Runtime/Engine/Classes/Engine/ChildConnection.h
@@ -12,7 +12,9 @@ UCLASS(MinimalAPI,transient,config=Engine)
 class UChildConnection
 	: public UNetConnection
 {
-	GENERATED_UCLASS_BODY()
+	GENERATED_BODY()
+public:
+	ENGINE_API UChildConnection(const FObjectInitializer& ObjectInitializer = FObjectInitializer::Get());
 
 	UPROPERTY(transient)
 	class UNetConnection* Parent;
diff --git a/Engine/Source/Runtime/Engine/Classes/Engine/CloudStorageBase.h b/Engine/Source/Runtime/Engine/Classes/Engine/CloudStorageBase.h
--- a/Engine/Source/Runtime/Engine/Classes/Engine/CloudStorageBase.h
+++ b/Engine/Source/Runtime/Engine/Classes/Engine/CloudStorageBase.h
@@ -33,7 +33,9 @@ UCLASS()
 class UCloudStorageBase
 	: public UPlatformInterfaceBase
 {
-	GENERATED_UCLASS_BODY()
+	GENERATED_BODY()
+public:
+	UCloudStorageBase(const FObjectInitializer& ObjectInitializer = FObjectInitializer::Get());
 
 	/** When using local storage (aka "cloud emulation"), this maintains a list of the file paths */
 	UPROPERTY()
diff --git a/Engine/Source/Runtime/Engine/Classes/Engine/CollisionProfile.h b/Engine/Source/Runtime/Engine/Classes/Engine/CollisionProfile.h
--- a/Engine/Source/Runtime/Engine/Classes/Engine/CollisionProfile.h
+++ b/Engine/Source/Runtime/Engine/Classes/Engine/CollisionProfile.h
@@ -13,7 +13,7 @@
 USTRUCT(BlueprintType)
 struct FCollisionProfileName
 {
-	GENERATED_USTRUCT_BODY()
+	GENERATED_BODY()
 
 	FCollisionProfileName()
 		: Name(NAME_None)
@@ -30,7 +30,7 @@ struct FCollisionProfileName
 USTRUCT()
 struct ENGINE_API FCollisionResponseTemplate
 {
-	GENERATED_USTRUCT_BODY()
+	GENERATED_BODY()
 
 	UPROPERTY()
 	FName Name;
@@ -77,7 +77,7 @@ struct ENGINE_API FCollisionResponseTemplate
 USTRUCT()
 struct FCustomChannelSetup
 {
-	GENERATED_USTRUCT_BODY()
+	GENERATED_BODY()
 
 	/** Which channel you'd like to customize **/
 	UPROPERTY()
@@ -122,7 +122,7 @@ struct FCustomChannelSetup
 USTRUCT()
 struct ENGINE_API FCustomProfile
 {
-	GENERATED_USTRUCT_BODY()
+	GENERATED_BODY()
 
 	UPROPERTY()
 	FName Name;
@@ -140,7 +140,9 @@ UCLASS(abstract, config=Engine, defaultconfig, MinimalAPI)
 class UCollisionProfile
 	: public UObject
 {
-	GENERATED_UCLASS_BODY()
+	GENERATED_BODY()
+public:
+	ENGINE_API UCollisionProfile(const FObjectInitializer& ObjectInitializer = FObjectInitializer::Get());
 
 private:
 
diff --git a/Engine/Source/Runtime/Engine/Classes/Engine/ComponentDelegateBinding.h b/Engine/Source/Runtime/Engine/Classes/Engine/ComponentDelegateBinding.h
--- a/Engine/Source/Runtime/Engine/Classes/Engine/ComponentDelegateBinding.h
+++ b/Engine/Source/Runtime/Engine/Classes/Engine/ComponentDelegateBinding.h
@@ -10,7 +10,7 @@
 USTRUCT()
 struct ENGINE_API FBlueprintComponentDelegateBinding
 {
-	GENERATED_USTRUCT_BODY()
+	GENERATED_BODY()
 
 	/** Name of component property that contains delegate we want to assign to. */
 	UPROPERTY()
@@ -36,7 +36,9 @@ UCLASS()
 class ENGINE_API UComponentDelegateBinding
 	: public UDynamicBlueprintBinding
 {
-	GENERATED_UCLASS_BODY()
+	GENERATED_BODY()
+public:
+	UComponentDelegateBinding(const FObjectInitializer& ObjectInitializer = FObjectInitializer::Get());
 
 	UPROPERTY()
 	TArray<FBlueprintComponentDelegateBinding> ComponentDelegateBindings;
diff --git a/Engine/Source/Runtime/Engine/Classes/Engine/Console.h b/Engine/Source/Runtime/Engine/Classes/Engine/Console.h
--- a/Engine/Source/Runtime/Engine/Classes/Engine/Console.h
+++ b/Engine/Source/Runtime/Engine/Classes/Engine/Console.h
@@ -12,7 +12,7 @@
 USTRUCT()
 struct FAutoCompleteNode
 {
-	GENERATED_USTRUCT_BODY()
+	GENERATED_BODY()
 
 	/** Char for node in the tree */
 	UPROPERTY()
@@ -55,7 +55,9 @@ class ENGINE_API UConsole
 	: public UObject
 	, public FOutputDevice
 {
-	GENERATED_UCLASS_BODY()
+	GENERATED_BODY()
+public:
+	UConsole(const FObjectInitializer& ObjectInitializer = FObjectInitializer::Get());
 
 	/** The player which the next console command should be executed in the context of.  If nullptr, execute in the viewport. */
 	UPROPERTY()
diff --git a/Engine/Source/Runtime/Engine/Classes/Engine/ControlChannel.h b/Engine/Source/Runtime/Engine/Classes/Engine/ControlChannel.h
--- a/Engine/Source/Runtime/Engine/Classes/Engine/ControlChannel.h
+++ b/Engine/Source/Runtime/Engine/Classes/Engine/ControlChannel.h
@@ -8,11 +8,12 @@
 /**
  * A channel for exchanging connection control messages.
  */
-UCLASS(transient, customConstructor)
+UCLASS(transient)
 class UControlChannel
 	: public UChannel
 {
-	GENERATED_UCLASS_BODY()
+	GENERATED_BODY()
+public:
 
 	/**
 	 * Used to interrogate the first packet received to determine endianess
diff --git a/Engine/Source/Runtime/Engine/Classes/Engine/CullDistanceVolume.h b/Engine/Source/Runtime/Engine/Classes/Engine/CullDistanceVolume.h
--- a/Engine/Source/Runtime/Engine/Classes/Engine/CullDistanceVolume.h
+++ b/Engine/Source/Runtime/Engine/Classes/Engine/CullDistanceVolume.h
@@ -12,7 +12,7 @@
 USTRUCT()
 struct FCullDistanceSizePair
 {
-	GENERATED_USTRUCT_BODY()
+	GENERATED_BODY()
 
 	/** Size to associate with cull distance. */
 	UPROPERTY(EditAnywhere, BlueprintReadOnly, Category=CullDistanceSizePair)
@@ -39,7 +39,9 @@ UCLASS(hidecategories=(Advanced, Attachment, Collision, Volume))
 class ACullDistanceVolume
 	: public AVolume
 {
-	GENERATED_UCLASS_BODY()
+	GENERATED_BODY()
+public:
+	ACullDistanceVolume(const FObjectInitializer& ObjectInitializer = FObjectInitializer::Get());
 
 	/**
 	 * Array of size and cull distance pairs. The code will calculate the sphere diameter of a primitive's BB and look for a best
diff --git a/Engine/Source/Runtime/Engine/Classes/Engine/CurveTable.h b/Engine/Source/Runtime/Engine/Classes/Engine/CurveTable.h
--- a/Engine/Source/Runtime/Engine/Classes/Engine/CurveTable.h
+++ b/Engine/Source/Runtime/Engine/Classes/Engine/CurveTable.h
@@ -22,7 +22,9 @@ UCLASS(MinimalAPI)
 class UCurveTable
 	: public UObject
 {
-	GENERATED_UCLASS_BODY()
+	GENERATED_BODY()
+public:
+	ENGINE_API UCurveTable(const FObjectInitializer& ObjectInitializer = FObjectInitializer::Get());
 
 	/** The filename that was used to create this object. Relative to the object's package, BaseDir() or absolute */
 	UPROPERTY()
@@ -99,7 +101,7 @@ protected:
 USTRUCT(BlueprintType)
 struct ENGINE_API FCurveTableRowHandle
 {
-	GENERATED_USTRUCT_BODY()
+	GENERATED_BODY()
 
 	FCurveTableRowHandle()
 		: CurveTable(nullptr)
diff --git a/Engine/Source/Runtime/Engine/Classes/Engine/DataAsset.h b/Engine/Source/Runtime/Engine/Classes/Engine/DataAsset.h
--- a/Engine/Source/Runtime/Engine/Classes/Engine/DataAsset.h
+++ b/Engine/Source/Runtime/Engine/Classes/Engine/DataAsset.h
@@ -12,7 +12,9 @@ UCLASS(abstract,MinimalAPI)
 class UDataAsset
 	: public UObject
 {
-	GENERATED_UCLASS_BODY()
+	GENERATED_BODY()
+public:
+	ENGINE_API UDataAsset(const FObjectInitializer& ObjectInitializer = FObjectInitializer::Get());
 
 public:
 
diff --git a/Engine/Source/Runtime/Engine/Classes/Engine/DataTable.h b/Engine/Source/Runtime/Engine/Classes/Engine/DataTable.h
--- a/Engine/Source/Runtime/Engine/Classes/Engine/DataTable.h
+++ b/Engine/Source/Runtime/Engine/Classes/Engine/DataTable.h
@@ -21,7 +21,7 @@ class TJsonWriter;
 USTRUCT()
 struct FTableRowBase
 {
-	GENERATED_USTRUCT_BODY()
+	GENERATED_BODY()
 
 	FTableRowBase() { }
 };
@@ -34,7 +34,9 @@ UCLASS(MinimalAPI)
 class UDataTable
 	: public UObject
 {
-	GENERATED_UCLASS_BODY()
+	GENERATED_BODY()
+public:
+	ENGINE_API UDataTable(const FObjectInitializer& ObjectInitializer = FObjectInitializer::Get());
 
 	/** Structure to use for each row of the table, must inherit from FTableRowBase */
 	UPROPERTY()
@@ -201,7 +203,7 @@ private:
 USTRUCT(BlueprintType)
 struct ENGINE_API FDataTableRowHandle
 {
-	GENERATED_USTRUCT_BODY()
+	GENERATED_BODY()
 
 	FDataTableRowHandle()
 		: DataTable(nullptr)
@@ -242,7 +244,7 @@ struct ENGINE_API FDataTableRowHandle
 USTRUCT()
 struct ENGINE_API FDataTableCategoryHandle
 {
-	GENERATED_USTRUCT_BODY()
+	GENERATED_BODY()
 
 	/** Pointer to table we want a row from */
 	UPROPERTY(EditAnywhere, BlueprintReadWrite, Category=DataTableCategoryHandle)
diff --git a/Engine/Source/Runtime/Engine/Classes/Engine/DebugCameraController.h b/Engine/Source/Runtime/Engine/Classes/Engine/DebugCameraController.h
--- a/Engine/Source/Runtime/Engine/Classes/Engine/DebugCameraController.h
+++ b/Engine/Source/Runtime/Engine/Classes/Engine/DebugCameraController.h
@@ -16,7 +16,9 @@ UCLASS(config=Game, hidedropdown)
 class ENGINE_API ADebugCameraController
 	: public APlayerController
 {
-	GENERATED_UCLASS_BODY()
+	GENERATED_BODY()
+public:
+	ADebugCameraController(const FObjectInitializer& ObjectInitializer = FObjectInitializer::Get());
 
 	/** Whether to show information about the selected actor on the debug camera HUD. */
 	UPROPERTY(globalconfig)
diff --git a/Engine/Source/Runtime/Engine/Classes/Engine/DebugCameraHUD.h b/Engine/Source/Runtime/Engine/Classes/Engine/DebugCameraHUD.h
--- a/Engine/Source/Runtime/Engine/Classes/Engine/DebugCameraHUD.h
+++ b/Engine/Source/Runtime/Engine/Classes/Engine/DebugCameraHUD.h
@@ -13,7 +13,9 @@ UCLASS(config=Game)
 class ENGINE_API ADebugCameraHUD
 	: public AHUD
 {
-	GENERATED_UCLASS_BODY()
+	GENERATED_BODY()
+public:
+	ADebugCameraHUD(const FObjectInitializer& ObjectInitializer = FObjectInitializer::Get());
 
 	/** @todo document */
 	virtual bool DisplayMaterials( float X, float& Y, float DY, class UMeshComponent* MeshComp );
diff --git a/Engine/Source/Runtime/Engine/Classes/Engine/DecalActor.h b/Engine/Source/Runtime/Engine/Classes/Engine/DecalActor.h
--- a/Engine/Source/Runtime/Engine/Classes/Engine/DecalActor.h
+++ b/Engine/Source/Runtime/Engine/Classes/Engine/DecalActor.h
@@ -15,7 +15,9 @@ UCLASS(hideCategories=(Collision, Attachment, Actor, Input, Replication), showCa
 class ADecalActor
 	: public AActor
 {
-	GENERATED_UCLASS_BODY()
+	GENERATED_BODY()
+public:
+	ENGINE_API ADecalActor(const FObjectInitializer& ObjectInitializer = FObjectInitializer::Get());
 
 private_subobject:
 	/** The decal component for this decal actor */
diff --git a/Engine/Source/Runtime/Engine/Classes/Engine/DemoNetConnection.h b/Engine/Source/Runtime/Engine/Classes/Engine/DemoNetConnection.h
--- a/Engine/Source/Runtime/Engine/Classes/Engine/DemoNetConnection.h
+++ b/Engine/Source/Runtime/Engine/Classes/Engine/DemoNetConnection.h
@@ -18,7 +18,9 @@ UCLASS(transient, config=Engine)
 class UDemoNetConnection
 	: public UNetConnection
 {
-	GENERATED_UCLASS_BODY()
+	GENERATED_BODY()
+public:
+	UDemoNetConnection(const FObjectInitializer& ObjectInitializer = FObjectInitializer::Get());
 
 public:
 
diff --git a/Engine/Source/Runtime/Engine/Classes/Engine/DemoNetDriver.h b/Engine/Source/Runtime/Engine/Classes/Engine/DemoNetDriver.h
--- a/Engine/Source/Runtime/Engine/Classes/Engine/DemoNetDriver.h
+++ b/Engine/Source/Runtime/Engine/Classes/Engine/DemoNetDriver.h
@@ -11,7 +11,9 @@ UCLASS(transient, config=Engine)
 class UDemoNetDriver
 	: public UNetDriver
 {
-	GENERATED_UCLASS_BODY()
+	GENERATED_BODY()
+public:
+	UDemoNetDriver(const FObjectInitializer& ObjectInitializer = FObjectInitializer::Get());
 
 	/** Name of the file to read/write from */
 	FString DemoFilename;
diff --git a/Engine/Source/Runtime/Engine/Classes/Engine/DemoPendingNetGame.h b/Engine/Source/Runtime/Engine/Classes/Engine/DemoPendingNetGame.h
--- a/Engine/Source/Runtime/Engine/Classes/Engine/DemoPendingNetGame.h
+++ b/Engine/Source/Runtime/Engine/Classes/Engine/DemoPendingNetGame.h
@@ -9,5 +9,7 @@ UCLASS(transient, config=Engine)
 class UDemoPendingNetGame
 	: public UPendingNetGame
 {
-	GENERATED_UCLASS_BODY()
+	GENERATED_BODY()
+public:
+	UDemoPendingNetGame(const FObjectInitializer& ObjectInitializer = FObjectInitializer::Get());
 };
diff --git a/Engine/Source/Runtime/Engine/Classes/Engine/DestructibleFractureSettings.h b/Engine/Source/Runtime/Engine/Classes/Engine/DestructibleFractureSettings.h
--- a/Engine/Source/Runtime/Engine/Classes/Engine/DestructibleFractureSettings.h
+++ b/Engine/Source/Runtime/Engine/Classes/Engine/DestructibleFractureSettings.h
@@ -44,7 +44,7 @@ namespace EDestructibleImportOptions
 USTRUCT()
 struct FFractureMaterial
 {
-	GENERATED_USTRUCT_BODY()
+	GENERATED_BODY()
 
 	/**
 	 * The UV scale (geometric distance/unit texture distance) for interior materials.
@@ -102,7 +102,7 @@ struct FFractureMaterial
 USTRUCT()
 struct FDestructibleChunkParameters
 {
-	GENERATED_USTRUCT_BODY()
+	GENERATED_BODY()
 
 	/**
 		Defines the chunk to be environmentally supported, if the appropriate NxDestructibleParametersFlag flags
@@ -153,7 +153,9 @@ UCLASS(MinimalAPI)
 class UDestructibleFractureSettings
 	: public UObject
 {
-	GENERATED_UCLASS_BODY()
+	GENERATED_BODY()
+public:
+	ENGINE_API UDestructibleFractureSettings(const FObjectInitializer& ObjectInitializer = FObjectInitializer::Get());
 
 #if WITH_HOT_RELOAD_CTORS
 	/** DO NOT USE. This constructor is for internal usage only for hot-reload purposes. */
diff --git a/Engine/Source/Runtime/Engine/Classes/Engine/DestructibleMesh.h b/Engine/Source/Runtime/Engine/Classes/Engine/DestructibleMesh.h
--- a/Engine/Source/Runtime/Engine/Classes/Engine/DestructibleMesh.h
+++ b/Engine/Source/Runtime/Engine/Classes/Engine/DestructibleMesh.h
@@ -39,7 +39,7 @@ enum EImpactDamageOverride
 USTRUCT()
 struct FDestructibleDepthParameters
 {
-	GENERATED_USTRUCT_BODY()
+	GENERATED_BODY()
 
 	/** Chunks up to the depth DefaultImpactDamageDepth will take impact damage, unless one of the override options (see EImpactDamageOverride) is chosen. */
 	UPROPERTY(EditAnywhere, Category=DestructibleDepthParameters)
@@ -62,7 +62,7 @@ struct FDestructibleDepthParameters
 USTRUCT()
 struct FDestructibleParametersFlag
 {
-	GENERATED_USTRUCT_BODY()
+	GENERATED_BODY()
 
 	/**
 		If set, chunks will "remember" damage applied to them, so that many applications of a damage amount
@@ -168,7 +168,7 @@ struct FDestructibleParametersFlag
 USTRUCT()
 struct FDestructibleDamageParameters
 {
-	GENERATED_USTRUCT_BODY()
+	GENERATED_BODY()
 
 	/**
 		The damage amount which will cause a chunk to fracture (break free).
@@ -238,7 +238,7 @@ struct FDestructibleDamageParameters
 USTRUCT()
 struct FDestructibleDebrisParameters
 {
-	GENERATED_USTRUCT_BODY()
+	GENERATED_BODY()
 
 	/**
 		"Debris chunks" (see debrisDepth, above) will be destroyed after a time (in seconds)
@@ -299,7 +299,7 @@ struct FDestructibleDebrisParameters
 USTRUCT()
 struct FDestructibleAdvancedParameters
 {
-	GENERATED_USTRUCT_BODY()
+	GENERATED_BODY()
 
 	/**
 		Limits the amount of damage applied to a chunk.  This is useful for preventing the entire destructible
@@ -352,7 +352,7 @@ struct FDestructibleAdvancedParameters
 USTRUCT()
 struct FDestructibleSpecialHierarchyDepths
 {
-	GENERATED_USTRUCT_BODY()
+	GENERATED_BODY()
 
 	// todo dynamicChunksDominanceGroup
 	// todo dynamicChunksGroupsMask
@@ -413,7 +413,7 @@ struct FDestructibleSpecialHierarchyDepths
 USTRUCT()
 struct FDestructibleParameters
 {
-	GENERATED_USTRUCT_BODY()
+	GENERATED_BODY()
 
 	/** Parameters that pertain to chunk damage.  See DestructibleDamageParameters. */
 	UPROPERTY(EditAnywhere, Category=DestructibleParameters)
@@ -452,7 +452,9 @@ UCLASS(hidecategories=(Object, Mesh, LevelOfDetail, Mirroring, Physics, Reimport
 class UDestructibleMesh
 	: public USkeletalMesh
 {
-	GENERATED_UCLASS_BODY()
+	GENERATED_BODY()
+public:
+	ENGINE_API UDestructibleMesh(const FObjectInitializer& ObjectInitializer = FObjectInitializer::Get());
 
 	/** Parameters controlling the destruction behavior. */
 	UPROPERTY(EditAnywhere, Category=DestructibleMesh)
diff --git a/Engine/Source/Runtime/Engine/Classes/Engine/DirectionalLight.h b/Engine/Source/Runtime/Engine/Classes/Engine/DirectionalLight.h
--- a/Engine/Source/Runtime/Engine/Classes/Engine/DirectionalLight.h
+++ b/Engine/Source/Runtime/Engine/Classes/Engine/DirectionalLight.h
@@ -16,7 +16,9 @@ UCLASS(ClassGroup=(Lights, DirectionalLights), MinimalAPI, meta=(ChildCanTick))
 class ADirectionalLight
 	: public ALight
 {
-	GENERATED_UCLASS_BODY()
+	GENERATED_BODY()
+public:
+	ENGINE_API ADirectionalLight(const FObjectInitializer& ObjectInitializer = FObjectInitializer::Get());
 
 #if WITH_EDITORONLY_DATA
 	// Reference to editor visualization arrow
diff --git a/Engine/Source/Runtime/Engine/Classes/Engine/DocumentationActor.h b/Engine/Source/Runtime/Engine/Classes/Engine/DocumentationActor.h
--- a/Engine/Source/Runtime/Engine/Classes/Engine/DocumentationActor.h
+++ b/Engine/Source/Runtime/Engine/Classes/Engine/DocumentationActor.h
@@ -25,7 +25,9 @@ UCLASS(hidecategories = (Sprite, MaterialSprite, Actor, Transform, Tags, Materia
 class ENGINE_API ADocumentationActor
 	: public AActor
 {
-	GENERATED_UCLASS_BODY()
+	GENERATED_BODY()
+public:
+	ADocumentationActor(const FObjectInitializer& ObjectInitializer = FObjectInitializer::Get());
 
 	/* Open the document link this actor relates to */
 	bool OpenDocumentLink() const;
diff --git a/Engine/Source/Runtime/Engine/Classes/Engine/DynamicBlueprintBinding.h b/Engine/Source/Runtime/Engine/Classes/Engine/DynamicBlueprintBinding.h
--- a/Engine/Source/Runtime/Engine/Classes/Engine/DynamicBlueprintBinding.h
+++ b/Engine/Source/Runtime/Engine/Classes/Engine/DynamicBlueprintBinding.h
@@ -9,7 +9,9 @@ UCLASS(abstract)
 class ENGINE_API UDynamicBlueprintBinding
 	: public UObject
 {
-	GENERATED_UCLASS_BODY()
+	GENERATED_BODY()
+public:
+	UDynamicBlueprintBinding(const FObjectInitializer& ObjectInitializer = FObjectInitializer::Get());
 
 	virtual void BindDynamicDelegates(UObject* InInstance) const { }
 };
diff --git a/Engine/Source/Runtime/Engine/Classes/Engine/Engine.h b/Engine/Source/Runtime/Engine/Classes/Engine/Engine.h
--- a/Engine/Source/Runtime/Engine/Classes/Engine/Engine.h
+++ b/Engine/Source/Runtime/Engine/Classes/Engine/Engine.h
@@ -63,7 +63,7 @@ enum EConsoleType
 USTRUCT()
 struct FFullyLoadedPackagesInfo
 {
-	GENERATED_USTRUCT_BODY()
+	GENERATED_BODY()
 
 	/** When to load these packages */
 	UPROPERTY()
@@ -94,7 +94,7 @@ struct FFullyLoadedPackagesInfo
 USTRUCT()
 struct FLevelStreamingStatus
 {
-	GENERATED_USTRUCT_BODY()
+	GENERATED_BODY()
 
 	UPROPERTY()
 	FName PackageName;
@@ -130,7 +130,7 @@ struct FLevelStreamingStatus
 USTRUCT()
 struct FNetDriverDefinition
 {
-	GENERATED_USTRUCT_BODY()
+	GENERATED_BODY()
 
 	/** Unique name of this net driver definition */
 	UPROPERTY()
@@ -160,7 +160,7 @@ struct FNetDriverDefinition
 USTRUCT()
 struct FNamedNetDriver
 {
-	GENERATED_USTRUCT_BODY()
+	GENERATED_BODY()
 
 	/** Instantiation of named net driver */
 	UPROPERTY(transient)
@@ -208,7 +208,7 @@ struct FNamedNetDriver
 USTRUCT()
 struct FWorldContext
 {
-	GENERATED_USTRUCT_BODY()
+	GENERATED_BODY()
 
 	/**************************************************************/
 	
@@ -341,7 +341,7 @@ private:
 USTRUCT()
 struct FStatColorMapEntry
 {
-	GENERATED_USTRUCT_BODY()
+	GENERATED_BODY()
 
 	UPROPERTY(globalconfig)
 	float In;
@@ -362,7 +362,7 @@ struct FStatColorMapEntry
 USTRUCT()
 struct FStatColorMapping
 {
-	GENERATED_USTRUCT_BODY()
+	GENERATED_BODY()
 
 	UPROPERTY(globalconfig)
 	FString StatName;
@@ -386,7 +386,7 @@ struct FStatColorMapping
 USTRUCT()
 struct FDropNoteInfo
 {
-	GENERATED_USTRUCT_BODY()
+	GENERATED_BODY()
 
 	/** Location to create Note actor in edited level. */
 	UPROPERTY()
@@ -415,7 +415,7 @@ struct FDropNoteInfo
 USTRUCT()
 struct FScreenMessageString
 {
-	GENERATED_USTRUCT_BODY()
+	GENERATED_BODY()
 
 	/** The 'key' for this message. */
 	UPROPERTY(transient)
@@ -466,7 +466,7 @@ namespace EMatineeCaptureType
 USTRUCT()
 struct FGameNameRedirect
 {
-	GENERATED_USTRUCT_BODY()
+	GENERATED_BODY()
 
 	UPROPERTY()
 	FString OldGameName;
@@ -479,7 +479,7 @@ struct FGameNameRedirect
 USTRUCT()
 struct FClassRedirect
 {
-	GENERATED_USTRUCT_BODY()
+	GENERATED_BODY()
 
 	UPROPERTY()
 	FString ObjectName;
@@ -504,7 +504,7 @@ struct FClassRedirect
 USTRUCT()
 struct FStructRedirect
 {
-	GENERATED_USTRUCT_BODY()
+	GENERATED_BODY()
 
 	UPROPERTY()
 	FString OldStructName;
@@ -517,7 +517,7 @@ struct FStructRedirect
 USTRUCT()
 struct FPluginRedirect
 {
-	GENERATED_USTRUCT_BODY()
+	GENERATED_BODY()
 
 	UPROPERTY()
 	FString OldPluginName;
@@ -535,7 +535,7 @@ DECLARE_DELEGATE(FEndStreamingPauseDelegate);
 USTRUCT()
 struct FMatineeScreenshotOptions
 {
-	GENERATED_USTRUCT_BODY()
+	GENERATED_BODY()
 
 	/** determines if we should start the matinee capture as soon as the game loads */
 	UPROPERTY(transient)
@@ -578,7 +578,9 @@ class ENGINE_API UEngine
 	: public UObject
 	, public FExec
 {
-	GENERATED_UCLASS_BODY()
+	GENERATED_BODY()
+public:
+	UEngine(const FObjectInitializer& ObjectInitializer = FObjectInitializer::Get());
 
 private:
 	// Fonts.
diff --git a/Engine/Source/Runtime/Engine/Classes/Engine/EngineBaseTypes.h b/Engine/Source/Runtime/Engine/Classes/Engine/EngineBaseTypes.h
--- a/Engine/Source/Runtime/Engine/Classes/Engine/EngineBaseTypes.h
+++ b/Engine/Source/Runtime/Engine/Classes/Engine/EngineBaseTypes.h
@@ -76,7 +76,7 @@ enum ETickingGroup
 USTRUCT()
 struct FTickPrerequisite
 {
-	GENERATED_USTRUCT_BODY()
+	GENERATED_BODY()
 
 	/** Tick functions live inside of UObjects, so we need a separate weak pointer to the UObject solely for the purpose of determining if PrerequisiteTickFunction is still valid **/
 	TWeakObjectPtr<class UObject> PrerequisiteObject;
@@ -126,7 +126,7 @@ struct FTickPrerequisite
 USTRUCT()
 struct ENGINE_API FTickFunction
 {
-	GENERATED_USTRUCT_BODY()
+	GENERATED_BODY()
 
 public:
 	// The following UPROPERTYs are for configuration and inherited from the CDO/archetype/blueprint etc
@@ -308,7 +308,7 @@ struct TStructOpsTypeTraits<FTickFunction> : public TStructOpsTypeTraitsBase
 USTRUCT()
 struct FActorTickFunction : public FTickFunction
 {
-	GENERATED_USTRUCT_BODY()
+	GENERATED_BODY()
 
 	/**  AActor  that is the target of this tick **/
 	class AActor*	Target;
@@ -340,7 +340,7 @@ struct TStructOpsTypeTraits<FActorTickFunction> : public TStructOpsTypeTraitsBas
 USTRUCT()
 struct FActorComponentTickFunction : public FTickFunction
 {
-	GENERATED_USTRUCT_BODY()
+	GENERATED_BODY()
 
 	/**  AActor  component that is the target of this tick **/
 	class UActorComponent*	Target;
@@ -372,7 +372,7 @@ struct TStructOpsTypeTraits<FActorComponentTickFunction> : public TStructOpsType
 USTRUCT()
 struct FPrimitiveComponentPostPhysicsTickFunction : public FTickFunction
 {
-	GENERATED_USTRUCT_BODY()
+	GENERATED_BODY()
 
 	/** PrimitiveComponent component that is the target of this tick **/
 	class UPrimitiveComponent*	Target;
@@ -575,7 +575,7 @@ namespace EDemoPlayFailure
 USTRUCT()
 struct ENGINE_API FURL
 {
-	GENERATED_USTRUCT_BODY()
+	GENERATED_BODY()
 
 	UPROPERTY()
 	FString Protocol;    // Protocol, i.e. "unreal" or "http".
@@ -704,7 +704,7 @@ enum EViewModeIndex
 USTRUCT()
 struct FExposureSettings
 {
-	GENERATED_USTRUCT_BODY()
+	GENERATED_BODY()
 
 	FExposureSettings()
 		: LogOffset(0), bFixed(false)
@@ -745,7 +745,9 @@ struct FExposureSettings
 UCLASS(abstract, config=Engine)
 class UEngineBaseTypes : public UObject
 {
-	GENERATED_UCLASS_BODY()
+	GENERATED_BODY()
+public:
+	UEngineBaseTypes(const FObjectInitializer& ObjectInitializer = FObjectInitializer::Get());
 
 };
 
diff --git a/Engine/Source/Runtime/Engine/Classes/Engine/ExponentialHeightFog.h b/Engine/Source/Runtime/Engine/Classes/Engine/ExponentialHeightFog.h
--- a/Engine/Source/Runtime/Engine/Classes/Engine/ExponentialHeightFog.h
+++ b/Engine/Source/Runtime/Engine/Classes/Engine/ExponentialHeightFog.h
@@ -12,7 +12,9 @@ UCLASS(showcategories=(Movement, Rendering, "Utilities|Transformation"), ClassGr
 class AExponentialHeightFog
 	: public AInfo
 {
-	GENERATED_UCLASS_BODY()
+	GENERATED_BODY()
+public:
+	ENGINE_API AExponentialHeightFog(const FObjectInitializer& ObjectInitializer = FObjectInitializer::Get());
 
 private_subobject:
 	/** @todo document */
diff --git a/Engine/Source/Runtime/Engine/Classes/Engine/Font.h b/Engine/Source/Runtime/Engine/Classes/Engine/Font.h
--- a/Engine/Source/Runtime/Engine/Classes/Engine/Font.h
+++ b/Engine/Source/Runtime/Engine/Classes/Engine/Font.h
@@ -1,4 +1,4 @@
-﻿// Copyright 1998-2015 Epic Games, Inc. All Rights Reserved.
+// Copyright 1998-2015 Epic Games, Inc. All Rights Reserved.
 
 #pragma once
 
@@ -26,7 +26,7 @@ enum class EFontCacheType
 USTRUCT()
 struct FFontCharacter
 {
-	GENERATED_USTRUCT_BODY()
+	GENERATED_BODY()
 
 	UPROPERTY(EditAnywhere, Category=FontCharacter)
 	int32 StartU;
@@ -97,7 +97,9 @@ struct TStructOpsTypeTraits<FFontCharacter> : public TStructOpsTypeTraitsBase
 UCLASS(hidecategories=Object, autoexpandcategories=Font, MinimalAPI, BlueprintType)
 class UFont : public UObject, public IFontProviderInterface
 {
-	GENERATED_UCLASS_BODY()
+	GENERATED_BODY()
+public:
+	ENGINE_API UFont(const FObjectInitializer& ObjectInitializer = FObjectInitializer::Get());
 
 	~UFont();
 
diff --git a/Engine/Source/Runtime/Engine/Classes/Engine/FontImportOptions.h b/Engine/Source/Runtime/Engine/Classes/Engine/FontImportOptions.h
--- a/Engine/Source/Runtime/Engine/Classes/Engine/FontImportOptions.h
+++ b/Engine/Source/Runtime/Engine/Classes/Engine/FontImportOptions.h
@@ -20,7 +20,7 @@ enum EFontImportCharacterSet
 USTRUCT()
 struct FFontImportOptionsData
 {
-	GENERATED_USTRUCT_BODY()
+	GENERATED_BODY()
 
 	/** Name of the typeface for the font to import */
 	UPROPERTY(EditAnywhere, Category=FontImportOptionsData)
@@ -179,7 +179,9 @@ UCLASS(hidecategories=Object, transient, MinimalAPI)
 class UFontImportOptions
 	: public UObject
 {
-	GENERATED_UCLASS_BODY()
+	GENERATED_BODY()
+public:
+	ENGINE_API UFontImportOptions(const FObjectInitializer& ObjectInitializer = FObjectInitializer::Get());
 
 	/** The actual data for this object.  We wrap it in a struct so that we can copy it around between objects. */
 	UPROPERTY(EditAnywhere, Category=FontImportOptions, meta=(FullyExpand = "true"))
diff --git a/Engine/Source/Runtime/Engine/Classes/Engine/GameEngine.h b/Engine/Source/Runtime/Engine/Classes/Engine/GameEngine.h
--- a/Engine/Source/Runtime/Engine/Classes/Engine/GameEngine.h
+++ b/Engine/Source/Runtime/Engine/Classes/Engine/GameEngine.h
@@ -13,7 +13,9 @@ UCLASS(config=Engine, transient)
 class ENGINE_API UGameEngine
 	: public UEngine
 {
-	GENERATED_UCLASS_BODY()
+	GENERATED_BODY()
+public:
+	UGameEngine(const FObjectInitializer& ObjectInitializer = FObjectInitializer::Get());
 
 	/** Check to see if we need to start a movie capture (used on the first tick when we want to record a Matinee) */
 	UPROPERTY(transient)
diff --git a/Engine/Source/Runtime/Engine/Classes/Engine/GameInstance.h b/Engine/Source/Runtime/Engine/Classes/Engine/GameInstance.h
--- a/Engine/Source/Runtime/Engine/Classes/Engine/GameInstance.h
+++ b/Engine/Source/Runtime/Engine/Classes/Engine/GameInstance.h
@@ -33,7 +33,9 @@ class FOnlineSessionSearchResult;
 UCLASS(config=Game, transient, BlueprintType, Blueprintable)
 class ENGINE_API UGameInstance : public UObject, public FExec
 {
-	GENERATED_UCLASS_BODY()
+	GENERATED_BODY()
+public:
+	UGameInstance(const FObjectInitializer& ObjectInitializer = FObjectInitializer::Get());
 
 protected:
 	struct FWorldContext* WorldContext;
@@ -157,4 +159,4 @@ public:
 	virtual void			HandleDemoPlaybackFailure( EDemoPlayFailure::Type FailureType, const FString& ErrorString = TEXT("") ) { }
 
 	static void				AddReferencedObjects(UObject* InThis, FReferenceCollector& Collector);
-};
\ No newline at end of file
+};
diff --git a/Engine/Source/Runtime/Engine/Classes/Engine/GameViewportClient.h b/Engine/Source/Runtime/Engine/Classes/Engine/GameViewportClient.h
--- a/Engine/Source/Runtime/Engine/Classes/Engine/GameViewportClient.h
+++ b/Engine/Source/Runtime/Engine/Classes/Engine/GameViewportClient.h
@@ -51,7 +51,9 @@ enum EStereoscopicPass
 UCLASS(Within=Engine, transient, config=Engine)
 class ENGINE_API UGameViewportClient : public UScriptViewportClient, public FExec
 {
-	GENERATED_UCLASS_BODY()
+	GENERATED_BODY()
+public:
+	UGameViewportClient(const FObjectInitializer& ObjectInitializer = FObjectInitializer::Get());
 
 public:
 	virtual ~UGameViewportClient();
diff --git a/Engine/Source/Runtime/Engine/Classes/Engine/GeneratedMeshAreaLight.h b/Engine/Source/Runtime/Engine/Classes/Engine/GeneratedMeshAreaLight.h
--- a/Engine/Source/Runtime/Engine/Classes/Engine/GeneratedMeshAreaLight.h
+++ b/Engine/Source/Runtime/Engine/Classes/Engine/GeneratedMeshAreaLight.h
@@ -12,5 +12,7 @@ UCLASS(notplaceable, MinimalAPI)
 class AGeneratedMeshAreaLight
 	: public ASpotLight
 {
-	GENERATED_UCLASS_BODY()
+	GENERATED_BODY()
+public:
+	ENGINE_API AGeneratedMeshAreaLight(const FObjectInitializer& ObjectInitializer = FObjectInitializer::Get());
 };
diff --git a/Engine/Source/Runtime/Engine/Classes/Engine/InGameAdManager.h b/Engine/Source/Runtime/Engine/Classes/Engine/InGameAdManager.h
--- a/Engine/Source/Runtime/Engine/Classes/Engine/InGameAdManager.h
+++ b/Engine/Source/Runtime/Engine/Classes/Engine/InGameAdManager.h
@@ -34,7 +34,9 @@ DECLARE_DYNAMIC_DELEGATE( FOnUserClosedAdvertisement );
 UCLASS()
 class UInGameAdManager : public UPlatformInterfaceBase
 {
-	GENERATED_UCLASS_BODY()
+	GENERATED_BODY()
+public:
+	UInGameAdManager(const FObjectInitializer& ObjectInitializer = FObjectInitializer::Get());
 
 	/** If true, the game will pause when the user clicks on the ad, which could take over the screen */
 	UPROPERTY()
diff --git a/Engine/Source/Runtime/Engine/Classes/Engine/InputActionDelegateBinding.h b/Engine/Source/Runtime/Engine/Classes/Engine/InputActionDelegateBinding.h
--- a/Engine/Source/Runtime/Engine/Classes/Engine/InputActionDelegateBinding.h
+++ b/Engine/Source/Runtime/Engine/Classes/Engine/InputActionDelegateBinding.h
@@ -7,7 +7,7 @@
 USTRUCT()
 struct ENGINE_API FBlueprintInputActionDelegateBinding : public FBlueprintInputDelegateBinding
 {
-	GENERATED_USTRUCT_BODY()
+	GENERATED_BODY()
 
 	UPROPERTY()
 	FName InputActionName;
@@ -30,7 +30,9 @@ struct ENGINE_API FBlueprintInputActionDelegateBinding : public FBlueprintInputD
 UCLASS()
 class ENGINE_API UInputActionDelegateBinding : public UInputDelegateBinding
 {
-	GENERATED_UCLASS_BODY()
+	GENERATED_BODY()
+public:
+	UInputActionDelegateBinding(const FObjectInitializer& ObjectInitializer = FObjectInitializer::Get());
 
 	UPROPERTY()
 	TArray<FBlueprintInputActionDelegateBinding> InputActionDelegateBindings;
diff --git a/Engine/Source/Runtime/Engine/Classes/Engine/InputAxisDelegateBinding.h b/Engine/Source/Runtime/Engine/Classes/Engine/InputAxisDelegateBinding.h
--- a/Engine/Source/Runtime/Engine/Classes/Engine/InputAxisDelegateBinding.h
+++ b/Engine/Source/Runtime/Engine/Classes/Engine/InputAxisDelegateBinding.h
@@ -7,7 +7,7 @@
 USTRUCT()
 struct ENGINE_API FBlueprintInputAxisDelegateBinding : public FBlueprintInputDelegateBinding
 {
-	GENERATED_USTRUCT_BODY()
+	GENERATED_BODY()
 
 	UPROPERTY()
 	FName InputAxisName;
@@ -26,7 +26,9 @@ struct ENGINE_API FBlueprintInputAxisDelegateBinding : public FBlueprintInputDel
 UCLASS()
 class ENGINE_API UInputAxisDelegateBinding : public UInputDelegateBinding
 {
-	GENERATED_UCLASS_BODY()
+	GENERATED_BODY()
+public:
+	UInputAxisDelegateBinding(const FObjectInitializer& ObjectInitializer = FObjectInitializer::Get());
 
 	UPROPERTY()
 	TArray<FBlueprintInputAxisDelegateBinding> InputAxisDelegateBindings;
diff --git a/Engine/Source/Runtime/Engine/Classes/Engine/InputAxisKeyDelegateBinding.h b/Engine/Source/Runtime/Engine/Classes/Engine/InputAxisKeyDelegateBinding.h
--- a/Engine/Source/Runtime/Engine/Classes/Engine/InputAxisKeyDelegateBinding.h
+++ b/Engine/Source/Runtime/Engine/Classes/Engine/InputAxisKeyDelegateBinding.h
@@ -7,7 +7,7 @@
 USTRUCT()
 struct ENGINE_API FBlueprintInputAxisKeyDelegateBinding : public FBlueprintInputDelegateBinding
 {
-	GENERATED_USTRUCT_BODY()
+	GENERATED_BODY()
 
 	UPROPERTY()
 	FKey AxisKey;
@@ -25,7 +25,9 @@ struct ENGINE_API FBlueprintInputAxisKeyDelegateBinding : public FBlueprintInput
 UCLASS()
 class ENGINE_API UInputAxisKeyDelegateBinding : public UInputDelegateBinding
 {
-	GENERATED_UCLASS_BODY()
+	GENERATED_BODY()
+public:
+	UInputAxisKeyDelegateBinding(const FObjectInitializer& ObjectInitializer = FObjectInitializer::Get());
 
 	UPROPERTY()
 	TArray<FBlueprintInputAxisKeyDelegateBinding> InputAxisKeyDelegateBindings;
diff --git a/Engine/Source/Runtime/Engine/Classes/Engine/InputDelegateBinding.h b/Engine/Source/Runtime/Engine/Classes/Engine/InputDelegateBinding.h
--- a/Engine/Source/Runtime/Engine/Classes/Engine/InputDelegateBinding.h
+++ b/Engine/Source/Runtime/Engine/Classes/Engine/InputDelegateBinding.h
@@ -7,7 +7,7 @@
 USTRUCT()
 struct ENGINE_API FBlueprintInputDelegateBinding
 {
-	GENERATED_USTRUCT_BODY()
+	GENERATED_BODY()
 
 	UPROPERTY()
 	uint32 bConsumeInput:1;
@@ -29,7 +29,9 @@ struct ENGINE_API FBlueprintInputDelegateBinding
 UCLASS(abstract)
 class ENGINE_API UInputDelegateBinding : public UDynamicBlueprintBinding
 {
-	GENERATED_UCLASS_BODY()
+	GENERATED_BODY()
+public:
+	UInputDelegateBinding(const FObjectInitializer& ObjectInitializer = FObjectInitializer::Get());
 
 	virtual void BindToInputComponent(UInputComponent* InputComponent) const { };
 
diff --git a/Engine/Source/Runtime/Engine/Classes/Engine/InputKeyDelegateBinding.h b/Engine/Source/Runtime/Engine/Classes/Engine/InputKeyDelegateBinding.h
--- a/Engine/Source/Runtime/Engine/Classes/Engine/InputKeyDelegateBinding.h
+++ b/Engine/Source/Runtime/Engine/Classes/Engine/InputKeyDelegateBinding.h
@@ -8,7 +8,7 @@
 USTRUCT()
 struct ENGINE_API FBlueprintInputKeyDelegateBinding : public FBlueprintInputDelegateBinding
 {
-	GENERATED_USTRUCT_BODY()
+	GENERATED_BODY()
 
 	UPROPERTY()
 	FInputChord InputChord;
@@ -30,7 +30,9 @@ struct ENGINE_API FBlueprintInputKeyDelegateBinding : public FBlueprintInputDele
 UCLASS()
 class ENGINE_API UInputKeyDelegateBinding : public UInputDelegateBinding
 {
-	GENERATED_UCLASS_BODY()
+	GENERATED_BODY()
+public:
+	UInputKeyDelegateBinding(const FObjectInitializer& ObjectInitializer = FObjectInitializer::Get());
 
 	UPROPERTY()
 	TArray<FBlueprintInputKeyDelegateBinding> InputKeyDelegateBindings;
diff --git a/Engine/Source/Runtime/Engine/Classes/Engine/InputTouchDelegateBinding.h b/Engine/Source/Runtime/Engine/Classes/Engine/InputTouchDelegateBinding.h
--- a/Engine/Source/Runtime/Engine/Classes/Engine/InputTouchDelegateBinding.h
+++ b/Engine/Source/Runtime/Engine/Classes/Engine/InputTouchDelegateBinding.h
@@ -7,7 +7,7 @@
 USTRUCT()
 struct ENGINE_API FBlueprintInputTouchDelegateBinding : public FBlueprintInputDelegateBinding
 {
-	GENERATED_USTRUCT_BODY()
+	GENERATED_BODY()
 
 	UPROPERTY()
 	TEnumAsByte<EInputEvent> InputKeyEvent;
@@ -27,7 +27,9 @@ struct ENGINE_API FBlueprintInputTouchDelegateBinding : public FBlueprintInputDe
 UCLASS()
 class ENGINE_API UInputTouchDelegateBinding : public UInputDelegateBinding
 {
-	GENERATED_UCLASS_BODY()
+	GENERATED_BODY()
+public:
+	UInputTouchDelegateBinding(const FObjectInitializer& ObjectInitializer = FObjectInitializer::Get());
 
 	UPROPERTY()
 	TArray<FBlueprintInputTouchDelegateBinding> InputTouchDelegateBindings;
diff --git a/Engine/Source/Runtime/Engine/Classes/Engine/InputVectorAxisDelegateBinding.h b/Engine/Source/Runtime/Engine/Classes/Engine/InputVectorAxisDelegateBinding.h
--- a/Engine/Source/Runtime/Engine/Classes/Engine/InputVectorAxisDelegateBinding.h
+++ b/Engine/Source/Runtime/Engine/Classes/Engine/InputVectorAxisDelegateBinding.h
@@ -7,7 +7,9 @@
 UCLASS()
 class ENGINE_API UInputVectorAxisDelegateBinding : public UInputAxisKeyDelegateBinding
 {
-	GENERATED_UCLASS_BODY()
+	GENERATED_BODY()
+public:
+	UInputVectorAxisDelegateBinding(const FObjectInitializer& ObjectInitializer = FObjectInitializer::Get());
 
 	// Begin UInputDelegateBinding interface
 	virtual void BindToInputComponent(UInputComponent* InputComponent) const override;
diff --git a/Engine/Source/Runtime/Engine/Classes/Engine/IntSerialization.h b/Engine/Source/Runtime/Engine/Classes/Engine/IntSerialization.h
--- a/Engine/Source/Runtime/Engine/Classes/Engine/IntSerialization.h
+++ b/Engine/Source/Runtime/Engine/Classes/Engine/IntSerialization.h
@@ -7,7 +7,9 @@
 UCLASS()
 class ENGINE_API UIntSerialization	: public UObject
 {
-	GENERATED_UCLASS_BODY()
+	GENERATED_BODY()
+public:
+	UIntSerialization(const FObjectInitializer& ObjectInitializer = FObjectInitializer::Get());
 
 public:
 //	ENGINE_API UIntSerialization(const FObjectInitializer& ObjectInitializer, uint8 InUnsignedInt8Variable, uint16 InUnsignedInt16Variable, uint32 InUnsignedInt32Variable, uint64 InUnsignedInt64Variable, int8 InSignedInt8Variable, int16 InSignedInt16Variable, int InSignedInt32Variable, int64 InSignedInt64Variable);
@@ -55,4 +57,4 @@ public:
 /*FArchive& operator<< (FArchive& Ar, FIntSerilizationTest& IntSerilizationTest)
 {
 	return Ar << IntSerilizationTest.UnsignedInt16Variable;
-}*/
\ No newline at end of file
+}*/
diff --git a/Engine/Source/Runtime/Engine/Classes/Engine/InterpCurveEdSetup.h b/Engine/Source/Runtime/Engine/Classes/Engine/InterpCurveEdSetup.h
--- a/Engine/Source/Runtime/Engine/Classes/Engine/InterpCurveEdSetup.h
+++ b/Engine/Source/Runtime/Engine/Classes/Engine/InterpCurveEdSetup.h
@@ -9,7 +9,7 @@
 USTRUCT()
 struct FCurveEdEntry
 {
-	GENERATED_USTRUCT_BODY()
+	GENERATED_BODY()
 
 	UPROPERTY()
 	class UObject* CurveObject;
@@ -56,7 +56,7 @@ struct FCurveEdEntry
 USTRUCT()
 struct FCurveEdTab
 {
-	GENERATED_USTRUCT_BODY()
+	GENERATED_BODY()
 
 	UPROPERTY()
 	FString TabName;
@@ -99,7 +99,9 @@ struct FCurveEdTab
 UCLASS(MinimalAPI)
 class UInterpCurveEdSetup : public UObject
 {
-	GENERATED_UCLASS_BODY()
+	GENERATED_BODY()
+public:
+	ENGINE_API UInterpCurveEdSetup(const FObjectInitializer& ObjectInitializer = FObjectInitializer::Get());
 
 	UPROPERTY()
 	TArray<struct FCurveEdTab> Tabs;
diff --git a/Engine/Source/Runtime/Engine/Classes/Engine/LODActor.h b/Engine/Source/Runtime/Engine/Classes/Engine/LODActor.h
--- a/Engine/Source/Runtime/Engine/Classes/Engine/LODActor.h
+++ b/Engine/Source/Runtime/Engine/Classes/Engine/LODActor.h
@@ -13,7 +13,9 @@
 UCLASS()
 class ENGINE_API ALODActor : public AActor
 {
-	GENERATED_UCLASS_BODY()
+	GENERATED_BODY()
+public:
+	ALODActor(const FObjectInitializer& ObjectInitializer = FObjectInitializer::Get());
 
 private_subobject:
 	UPROPERTY(Category=LODActor, VisibleAnywhere)
diff --git a/Engine/Source/Runtime/Engine/Classes/Engine/LevelBounds.h b/Engine/Source/Runtime/Engine/Classes/Engine/LevelBounds.h
--- a/Engine/Source/Runtime/Engine/Classes/Engine/LevelBounds.h
+++ b/Engine/Source/Runtime/Engine/Classes/Engine/LevelBounds.h
@@ -12,7 +12,9 @@
 UCLASS(hidecategories=(Advanced, Collision, Display, Rendering, Physics, Input), showcategories=("Input|MouseInput", "Input|TouchInput"), MinimalAPI)
 class ALevelBounds : public AActor
 {
-	GENERATED_UCLASS_BODY()
+	GENERATED_BODY()
+public:
+	ENGINE_API ALevelBounds(const FObjectInitializer& ObjectInitializer = FObjectInitializer::Get());
 		
 	/** Whether to automatically update actor bounds based on all relevant actors bounds belonging to the same level */
 	UPROPERTY(EditAnywhere, Category=LevelBounds)
diff --git a/Engine/Source/Runtime/Engine/Classes/Engine/LevelScriptActor.h b/Engine/Source/Runtime/Engine/Classes/Engine/LevelScriptActor.h
--- a/Engine/Source/Runtime/Engine/Classes/Engine/LevelScriptActor.h
+++ b/Engine/Source/Runtime/Engine/Classes/Engine/LevelScriptActor.h
@@ -21,7 +21,9 @@
 UCLASS(notplaceable, meta=(KismetHideOverrides = "ReceiveAnyDamage,ReceivePointDamage,ReceiveRadialDamage,ReceiveActorBeginOverlap,ReceiveActorEndOverlap,ReceiveHit,ReceiveDestroyed,ReceiveActorBeginCursorOver,ReceiveActorEndCursorOver,ReceiveActorOnClicked,ReceiveActorOnReleased,ReceiveActorOnInputTouchBegin,ReceiveActorOnInputTouchEnd,ReceiveActorOnInputTouchEnter,ReceiveActorOnInputTouchLeave"), HideCategories=(Collision,Rendering,"Utilities|Transformation"))
 class ENGINE_API ALevelScriptActor : public AActor
 {
-	GENERATED_UCLASS_BODY()
+	GENERATED_BODY()
+public:
+	ALevelScriptActor(const FObjectInitializer& ObjectInitializer = FObjectInitializer::Get());
 
 	// --- Utility Functions ----------------------------
 	
diff --git a/Engine/Source/Runtime/Engine/Classes/Engine/LevelScriptBlueprint.h b/Engine/Source/Runtime/Engine/Classes/Engine/LevelScriptBlueprint.h
--- a/Engine/Source/Runtime/Engine/Classes/Engine/LevelScriptBlueprint.h
+++ b/Engine/Source/Runtime/Engine/Classes/Engine/LevelScriptBlueprint.h
@@ -20,7 +20,9 @@
 UCLASS(MinimalAPI, NotBlueprintType)
 class ULevelScriptBlueprint : public UBlueprint
 {
-	GENERATED_UCLASS_BODY()
+	GENERATED_BODY()
+public:
+	ENGINE_API ULevelScriptBlueprint(const FObjectInitializer& ObjectInitializer = FObjectInitializer::Get());
 
 #if WITH_EDITORONLY_DATA
 	/** The friendly name to use for UI */
diff --git a/Engine/Source/Runtime/Engine/Classes/Engine/LevelStreaming.h b/Engine/Source/Runtime/Engine/Classes/Engine/LevelStreaming.h
--- a/Engine/Source/Runtime/Engine/Classes/Engine/LevelStreaming.h
+++ b/Engine/Source/Runtime/Engine/Classes/Engine/LevelStreaming.h
@@ -69,7 +69,9 @@ DECLARE_DYNAMIC_MULTICAST_DELEGATE( FLevelStreamingVisibilityStatus );
 UCLASS(abstract, editinlinenew,MinimalAPI)
 class ULevelStreaming : public UObject
 {
-	GENERATED_UCLASS_BODY()
+	GENERATED_BODY()
+public:
+	ENGINE_API ULevelStreaming(const FObjectInitializer& ObjectInitializer = FObjectInitializer::Get());
 
 	/** Deprecated name of the package containing the level to load. Use WorldAsset or GetWorldAssetPackageFName instead.		*/
 	UPROPERTY()
diff --git a/Engine/Source/Runtime/Engine/Classes/Engine/LevelStreamingAlwaysLoaded.h b/Engine/Source/Runtime/Engine/Classes/Engine/LevelStreamingAlwaysLoaded.h
--- a/Engine/Source/Runtime/Engine/Classes/Engine/LevelStreamingAlwaysLoaded.h
+++ b/Engine/Source/Runtime/Engine/Classes/Engine/LevelStreamingAlwaysLoaded.h
@@ -13,7 +13,9 @@
 UCLASS(MinimalAPI)
 class ULevelStreamingAlwaysLoaded : public ULevelStreaming
 {
-	GENERATED_UCLASS_BODY()
+	GENERATED_BODY()
+public:
+	ENGINE_API ULevelStreamingAlwaysLoaded(const FObjectInitializer& ObjectInitializer = FObjectInitializer::Get());
 
 	// Begin ULevelStreaming Interface
 	virtual bool ShouldBeLoaded() const override;
diff --git a/Engine/Source/Runtime/Engine/Classes/Engine/LevelStreamingKismet.h b/Engine/Source/Runtime/Engine/Classes/Engine/LevelStreamingKismet.h
--- a/Engine/Source/Runtime/Engine/Classes/Engine/LevelStreamingKismet.h
+++ b/Engine/Source/Runtime/Engine/Classes/Engine/LevelStreamingKismet.h
@@ -14,7 +14,9 @@
 UCLASS(MinimalAPI, BlueprintType)
 class ULevelStreamingKismet : public ULevelStreaming
 {
-	GENERATED_UCLASS_BODY()
+	GENERATED_BODY()
+public:
+	ENGINE_API ULevelStreamingKismet(const FObjectInitializer& ObjectInitializer = FObjectInitializer::Get());
 
 	/** Whether the level should be loaded at startup																			*/
 	UPROPERTY(Category=LevelStreaming, EditAnywhere)
diff --git a/Engine/Source/Runtime/Engine/Classes/Engine/LevelStreamingPersistent.h b/Engine/Source/Runtime/Engine/Classes/Engine/LevelStreamingPersistent.h
--- a/Engine/Source/Runtime/Engine/Classes/Engine/LevelStreamingPersistent.h
+++ b/Engine/Source/Runtime/Engine/Classes/Engine/LevelStreamingPersistent.h
@@ -13,7 +13,9 @@
 UCLASS(transient)
 class ULevelStreamingPersistent : public ULevelStreaming
 {
-	GENERATED_UCLASS_BODY()
+	GENERATED_BODY()
+public:
+	ULevelStreamingPersistent(const FObjectInitializer& ObjectInitializer = FObjectInitializer::Get());
 
 
 	// Begin ULevelStreaming Interface
diff --git a/Engine/Source/Runtime/Engine/Classes/Engine/LevelStreamingVolume.h b/Engine/Source/Runtime/Engine/Classes/Engine/LevelStreamingVolume.h
--- a/Engine/Source/Runtime/Engine/Classes/Engine/LevelStreamingVolume.h
+++ b/Engine/Source/Runtime/Engine/Classes/Engine/LevelStreamingVolume.h
@@ -25,7 +25,9 @@ enum EStreamingVolumeUsage
 UCLASS(hidecategories=(Advanced, Attachment, Collision, Volume), MinimalAPI)
 class ALevelStreamingVolume : public AVolume
 {
-	GENERATED_UCLASS_BODY()
+	GENERATED_BODY()
+public:
+	ENGINE_API ALevelStreamingVolume(const FObjectInitializer& ObjectInitializer = FObjectInitializer::Get());
 
 	/** Levels names affected by this level streaming volume. */
 	UPROPERTY(Category=LevelStreamingVolume, VisibleAnywhere, BlueprintReadOnly, meta=(DisplayName = "Streaming Levels"))
diff --git a/Engine/Source/Runtime/Engine/Classes/Engine/Light.h b/Engine/Source/Runtime/Engine/Classes/Engine/Light.h
--- a/Engine/Source/Runtime/Engine/Classes/Engine/Light.h
+++ b/Engine/Source/Runtime/Engine/Classes/Engine/Light.h
@@ -8,7 +8,9 @@
 UCLASS(Abstract, ClassGroup=Lights, hideCategories=(Input,Collision,Replication), showCategories=("Input|MouseInput", "Input|TouchInput"), ComponentWrapperClass, ConversionRoot, meta=(ChildCanTick))
 class ENGINE_API ALight : public AActor
 {
-	GENERATED_UCLASS_BODY()
+	GENERATED_BODY()
+public:
+	ALight(const FObjectInitializer& ObjectInitializer = FObjectInitializer::Get());
 
 private_subobject:
 	/** @todo document */
diff --git a/Engine/Source/Runtime/Engine/Classes/Engine/LightMapTexture2D.h b/Engine/Source/Runtime/Engine/Classes/Engine/LightMapTexture2D.h
--- a/Engine/Source/Runtime/Engine/Classes/Engine/LightMapTexture2D.h
+++ b/Engine/Source/Runtime/Engine/Classes/Engine/LightMapTexture2D.h
@@ -21,7 +21,9 @@ enum ELightMapFlags
 UCLASS(MinimalAPI)
 class ULightMapTexture2D : public UTexture2D
 {
-	GENERATED_UCLASS_BODY()
+	GENERATED_BODY()
+public:
+	ENGINE_API ULightMapTexture2D(const FObjectInitializer& ObjectInitializer = FObjectInitializer::Get());
 
 	// Begin UObject interface.
 	virtual void Serialize( FArchive& Ar ) override;
diff --git a/Engine/Source/Runtime/Engine/Classes/Engine/LocalPlayer.h b/Engine/Source/Runtime/Engine/Classes/Engine/LocalPlayer.h
--- a/Engine/Source/Runtime/Engine/Classes/Engine/LocalPlayer.h
+++ b/Engine/Source/Runtime/Engine/Classes/Engine/LocalPlayer.h
@@ -138,7 +138,9 @@ private:
 UCLASS(Within=Engine, config=Engine, transient)
 class ENGINE_API ULocalPlayer : public UPlayer
 {
-	GENERATED_UCLASS_BODY()
+	GENERATED_BODY()
+public:
+	ULocalPlayer(const FObjectInitializer& ObjectInitializer = FObjectInitializer::Get());
 
 #if WITH_HOT_RELOAD_CTORS
 	/** DO NOT USE. This constructor is for internal usage only for hot-reload purposes. */
diff --git a/Engine/Source/Runtime/Engine/Classes/Engine/MicroTransactionBase.h b/Engine/Source/Runtime/Engine/Classes/Engine/MicroTransactionBase.h
--- a/Engine/Source/Runtime/Engine/Classes/Engine/MicroTransactionBase.h
+++ b/Engine/Source/Runtime/Engine/Classes/Engine/MicroTransactionBase.h
@@ -49,7 +49,7 @@ enum EMicroTransactionResult
 USTRUCT()
 struct FPurchaseInfo
 {
-	GENERATED_USTRUCT_BODY()
+	GENERATED_BODY()
 
 // 	/** What kind of microtransaction purchase is this? */
 // 	var EPurchaseType Type;
@@ -75,7 +75,9 @@ struct FPurchaseInfo
 UCLASS()
 class UMicroTransactionBase : public UPlatformInterfaceBase
 {
-	GENERATED_UCLASS_BODY()
+	GENERATED_BODY()
+public:
+	UMicroTransactionBase(const FObjectInitializer& ObjectInitializer = FObjectInitializer::Get());
 
 	/** The list of products available to purchase, filled out by the time a MTD_PurchaseQueryComplete is fired */
 	UPROPERTY()
diff --git a/Engine/Source/Runtime/Engine/Classes/Engine/NavigationObjectBase.h b/Engine/Source/Runtime/Engine/Classes/Engine/NavigationObjectBase.h
--- a/Engine/Source/Runtime/Engine/Classes/Engine/NavigationObjectBase.h
+++ b/Engine/Source/Runtime/Engine/Classes/Engine/NavigationObjectBase.h
@@ -8,7 +8,9 @@
 UCLASS(hidecategories=(Lighting, LightColor, Force), ClassGroup=Navigation, NotBlueprintable, abstract)
 class ENGINE_API ANavigationObjectBase : public AActor, public INavAgentInterface
 {
-	GENERATED_UCLASS_BODY()
+	GENERATED_BODY()
+public:
+	ANavigationObjectBase(const FObjectInitializer& ObjectInitializer = FObjectInitializer::Get());
 
 private_subobject:
 	DEPRECATED_FORGAME(4.6, "CapsuleComponent should not be accessed directly, please use GetCapsuleComponent() function instead. CapsuleComponent will soon be private and your code will not compile.")
diff --git a/Engine/Source/Runtime/Engine/Classes/Engine/NetConnection.h b/Engine/Source/Runtime/Engine/Classes/Engine/NetConnection.h
--- a/Engine/Source/Runtime/Engine/Classes/Engine/NetConnection.h
+++ b/Engine/Source/Runtime/Engine/Classes/Engine/NetConnection.h
@@ -97,10 +97,11 @@ struct DelayedPacket
 #define PING_ACK_DELAY 0.5
 
 
-UCLASS(customConstructor, Abstract, MinimalAPI, transient, config=Engine)
+UCLASS(Abstract, MinimalAPI, transient, config=Engine)
 class UNetConnection : public UPlayer
 {
-	GENERATED_UCLASS_BODY()
+	GENERATED_BODY()
+public:
 
 	/** child connections for secondary viewports */
 	UPROPERTY(transient)
diff --git a/Engine/Source/Runtime/Engine/Classes/Engine/NetDriver.h b/Engine/Source/Runtime/Engine/Classes/Engine/NetDriver.h
--- a/Engine/Source/Runtime/Engine/Classes/Engine/NetDriver.h
+++ b/Engine/Source/Runtime/Engine/Classes/Engine/NetDriver.h
@@ -92,10 +92,11 @@ struct FActorDestructionInfo
 };
 
 
-UCLASS(Abstract, customConstructor, transient, MinimalAPI, config=Engine)
+UCLASS(Abstract, transient, MinimalAPI, config=Engine)
 class UNetDriver : public UObject, public FExec
 {
-	GENERATED_UCLASS_BODY()
+	GENERATED_BODY()
+public:
 
 protected:
 
diff --git a/Engine/Source/Runtime/Engine/Classes/Engine/NetSerialization.h b/Engine/Source/Runtime/Engine/Classes/Engine/NetSerialization.h
--- a/Engine/Source/Runtime/Engine/Classes/Engine/NetSerialization.h
+++ b/Engine/Source/Runtime/Engine/Classes/Engine/NetSerialization.h
@@ -26,7 +26,7 @@
 USTRUCT()
 struct FExampleStruct
 {
-	GENERATED_USTRUCT_BODY()
+	GENERATED_BODY()
 	
 	/**
 	 * @param Ar			FArchive to read or write from.
@@ -90,7 +90,7 @@ struct TStructOpsTypeTraits< FExampleStruct > : public TStructOpsTypeTraitsBase
 USTRUCT()
 struct FExampleItemEntry : public FFastArraySerializerItem
 {
-	GENERATED_USTRUCT_BODY()
+	GENERATED_BODY()
 
 	// Your data:
 	UPROPERTY()
@@ -116,7 +116,7 @@ struct FExampleItemEntry : public FFastArraySerializerItem
 USTRUCT()
 struct FExampleArray: public FFastArraySerializer
 {
-	GENERATED_USTRUCT_BODY()
+	GENERATED_BODY()
 
 	UPROPERTY()
 	TArray<FExampleItemEntry>	Items;	/** Step 3: You MUST have a TArray named Items of the struct you made in step 1. */
@@ -274,7 +274,7 @@ public:
 USTRUCT()
 struct FFastArraySerializerItem
 {
-	GENERATED_USTRUCT_BODY()
+	GENERATED_BODY()
 
 	FFastArraySerializerItem()
 		: ReplicationID(INDEX_NONE), ReplicationKey(INDEX_NONE)
@@ -341,7 +341,7 @@ struct FFastArraySerializerUnmappedItem
 USTRUCT()
 struct FFastArraySerializer
 {
-	GENERATED_USTRUCT_BODY()
+	GENERATED_BODY()
 
 	FFastArraySerializer() : IDCounter(0), ArrayReplicationKey(0) { }
 
@@ -1060,7 +1060,7 @@ bool SerializeFixedVector(FVector &Vector, FArchive& Ar)
 USTRUCT()
 struct FVector_NetQuantize : public FVector
 {
-	GENERATED_USTRUCT_BODY()
+	GENERATED_BODY()
 
 	FORCEINLINE FVector_NetQuantize()
 	{}
@@ -1105,7 +1105,7 @@ struct TStructOpsTypeTraits< FVector_NetQuantize > : public TStructOpsTypeTraits
 USTRUCT()
 struct FVector_NetQuantize10 : public FVector
 {
-	GENERATED_USTRUCT_BODY()
+	GENERATED_BODY()
 
 	FORCEINLINE FVector_NetQuantize10()
 	{}
@@ -1150,7 +1150,7 @@ struct TStructOpsTypeTraits< FVector_NetQuantize10 > : public TStructOpsTypeTrai
 USTRUCT()
 struct FVector_NetQuantize100 : public FVector
 {
-	GENERATED_USTRUCT_BODY()
+	GENERATED_BODY()
 
 	FORCEINLINE FVector_NetQuantize100()
 	{}
@@ -1193,7 +1193,7 @@ struct TStructOpsTypeTraits< FVector_NetQuantize100 > : public TStructOpsTypeTra
 USTRUCT()
 struct FVector_NetQuantizeNormal : public FVector
 {
-	GENERATED_USTRUCT_BODY()
+	GENERATED_BODY()
 
 	FORCEINLINE FVector_NetQuantizeNormal()
 	{}
diff --git a/Engine/Source/Runtime/Engine/Classes/Engine/NetworkSettings.h b/Engine/Source/Runtime/Engine/Classes/Engine/NetworkSettings.h
--- a/Engine/Source/Runtime/Engine/Classes/Engine/NetworkSettings.h
+++ b/Engine/Source/Runtime/Engine/Classes/Engine/NetworkSettings.h
@@ -12,7 +12,9 @@ UCLASS(config=Engine, defaultconfig)
 class ENGINE_API UNetworkSettings
 	: public UObject
 {
-	GENERATED_UCLASS_BODY()
+	GENERATED_BODY()
+public:
+	UNetworkSettings(const FObjectInitializer& ObjectInitializer = FObjectInitializer::Get());
 
 	UPROPERTY(config, EditAnywhere, Category=libcurl, meta=(
 		ConsoleVariable="n.VerifyPeer",DisplayName="Verify Peer",
diff --git a/Engine/Source/Runtime/Engine/Classes/Engine/NiagaraActor.h b/Engine/Source/Runtime/Engine/Classes/Engine/NiagaraActor.h
--- a/Engine/Source/Runtime/Engine/Classes/Engine/NiagaraActor.h
+++ b/Engine/Source/Runtime/Engine/Classes/Engine/NiagaraActor.h
@@ -6,7 +6,9 @@
 UCLASS()
 class ENGINE_API ANiagaraActor : public AActor
 {
-	GENERATED_UCLASS_BODY()
+	GENERATED_BODY()
+public:
+	ANiagaraActor(const FObjectInitializer& ObjectInitializer = FObjectInitializer::Get());
 
 private_subobject:
 	/** Pointer to effect component */
diff --git a/Engine/Source/Runtime/Engine/Classes/Engine/NiagaraEffect.h b/Engine/Source/Runtime/Engine/Classes/Engine/NiagaraEffect.h
--- a/Engine/Source/Runtime/Engine/Classes/Engine/NiagaraEffect.h
+++ b/Engine/Source/Runtime/Engine/Classes/Engine/NiagaraEffect.h
@@ -9,7 +9,9 @@
 UCLASS(MinimalAPI)
 class UNiagaraEffect : public UObject
 {
-	GENERATED_UCLASS_BODY()
+	GENERATED_BODY()
+public:
+	ENGINE_API UNiagaraEffect(const FObjectInitializer& ObjectInitializer = FObjectInitializer::Get());
 
 public:
 	ENGINE_API FNiagaraEmitterProperties *AddEmitterProperties();
@@ -170,4 +172,4 @@ private:
 	TArray< TSharedPtr<FNiagaraSimulation> > Emitters;
 
 
-};
\ No newline at end of file
+};
diff --git a/Engine/Source/Runtime/Engine/Classes/Engine/NiagaraEffectRendererProperties.h b/Engine/Source/Runtime/Engine/Classes/Engine/NiagaraEffectRendererProperties.h
--- a/Engine/Source/Runtime/Engine/Classes/Engine/NiagaraEffectRendererProperties.h
+++ b/Engine/Source/Runtime/Engine/Classes/Engine/NiagaraEffectRendererProperties.h
@@ -12,7 +12,9 @@
 UCLASS()
 class UNiagaraEffectRendererProperties : public UObject
 {
-	GENERATED_UCLASS_BODY()
+	GENERATED_BODY()
+public:
+	UNiagaraEffectRendererProperties(const FObjectInitializer& ObjectInitializer);
 	UNiagaraEffectRendererProperties()
 	{}
 
diff --git a/Engine/Source/Runtime/Engine/Classes/Engine/NiagaraRibbonRendererProperties.h b/Engine/Source/Runtime/Engine/Classes/Engine/NiagaraRibbonRendererProperties.h
--- a/Engine/Source/Runtime/Engine/Classes/Engine/NiagaraRibbonRendererProperties.h
+++ b/Engine/Source/Runtime/Engine/Classes/Engine/NiagaraRibbonRendererProperties.h
@@ -8,8 +8,10 @@
 UCLASS()
 class UNiagaraRibbonRendererProperties : public UNiagaraEffectRendererProperties
 {
+	GENERATED_BODY()
+public:
+	UNiagaraRibbonRendererProperties(const FObjectInitializer& ObjectInitializer);
 public:
-	GENERATED_UCLASS_BODY()
 
 	UNiagaraRibbonRendererProperties()
 	{
diff --git a/Engine/Source/Runtime/Engine/Classes/Engine/NiagaraScript.h b/Engine/Source/Runtime/Engine/Classes/Engine/NiagaraScript.h
--- a/Engine/Source/Runtime/Engine/Classes/Engine/NiagaraScript.h
+++ b/Engine/Source/Runtime/Engine/Classes/Engine/NiagaraScript.h
@@ -10,7 +10,9 @@
 UCLASS(MinimalAPI)
 class UNiagaraScript : public UObject
 {
-	GENERATED_UCLASS_BODY()
+	GENERATED_BODY()
+public:
+	ENGINE_API UNiagaraScript(const FObjectInitializer& ObjectInitializer = FObjectInitializer::Get());
 
 	/** Byte code to execute for this system */
 	UPROPERTY()
@@ -29,4 +31,4 @@ class UNiagaraScript : public UObject
 	UPROPERTY()
 	class UNiagaraScriptSourceBase*	Source;
 #endif
-};
\ No newline at end of file
+};
diff --git a/Engine/Source/Runtime/Engine/Classes/Engine/NiagaraScriptSourceBase.h b/Engine/Source/Runtime/Engine/Classes/Engine/NiagaraScriptSourceBase.h
--- a/Engine/Source/Runtime/Engine/Classes/Engine/NiagaraScriptSourceBase.h
+++ b/Engine/Source/Runtime/Engine/Classes/Engine/NiagaraScriptSourceBase.h
@@ -20,9 +20,11 @@ struct EditorExposedVectorCurveConstant
 UCLASS(MinimalAPI)
 class UNiagaraScriptSourceBase : public UObject
 {
-	GENERATED_UCLASS_BODY()
+	GENERATED_BODY()
+public:
+	ENGINE_API UNiagaraScriptSourceBase(const FObjectInitializer& ObjectInitializer = FObjectInitializer::Get());
 
 	TArray<TSharedPtr<EditorExposedVectorConstant> > ExposedVectorConstants;
 	TArray<TSharedPtr<EditorExposedVectorCurveConstant> > ExposedVectorCurveConstants;
 	virtual void Compile() {};
-};
\ No newline at end of file
+};
diff --git a/Engine/Source/Runtime/Engine/Classes/Engine/NiagaraSpriteRendererProperties.h b/Engine/Source/Runtime/Engine/Classes/Engine/NiagaraSpriteRendererProperties.h
--- a/Engine/Source/Runtime/Engine/Classes/Engine/NiagaraSpriteRendererProperties.h
+++ b/Engine/Source/Runtime/Engine/Classes/Engine/NiagaraSpriteRendererProperties.h
@@ -8,8 +8,10 @@
 UCLASS()
 class UNiagaraSpriteRendererProperties : public UNiagaraEffectRendererProperties
 {
+	GENERATED_BODY()
+public:
+	UNiagaraSpriteRendererProperties(const FObjectInitializer& ObjectInitializer);
 public:
-	GENERATED_UCLASS_BODY()
 
 	UNiagaraSpriteRendererProperties()
 	{
diff --git a/Engine/Source/Runtime/Engine/Classes/Engine/Note.h b/Engine/Source/Runtime/Engine/Classes/Engine/Note.h
--- a/Engine/Source/Runtime/Engine/Classes/Engine/Note.h
+++ b/Engine/Source/Runtime/Engine/Classes/Engine/Note.h
@@ -11,7 +11,9 @@
 UCLASS(MinimalAPI, hidecategories = (Input), showcategories=("Input|MouseInput", "Input|TouchInput"))
 class ANote : public AActor
 {
-	GENERATED_UCLASS_BODY()
+	GENERATED_BODY()
+public:
+	ENGINE_API ANote(const FObjectInitializer& ObjectInitializer = FObjectInitializer::Get());
 
 #if WITH_EDITORONLY_DATA
 	UPROPERTY(EditAnywhere, BlueprintReadWrite, Category=Note)
diff --git a/Engine/Source/Runtime/Engine/Classes/Engine/ObjectLibrary.h b/Engine/Source/Runtime/Engine/Classes/Engine/ObjectLibrary.h
--- a/Engine/Source/Runtime/Engine/Classes/Engine/ObjectLibrary.h
+++ b/Engine/Source/Runtime/Engine/Classes/Engine/ObjectLibrary.h
@@ -8,7 +8,9 @@
 UCLASS(MinimalAPI)
 class UObjectLibrary : public UObject
 {
-	GENERATED_UCLASS_BODY()
+	GENERATED_BODY()
+public:
+	ENGINE_API UObjectLibrary(const FObjectInitializer& ObjectInitializer = FObjectInitializer::Get());
 
 	/** Class that Objects must be of. If ContainsBlueprints is true, this is the native class that the blueprints are instances of and not UClass  */
 	UPROPERTY(EditAnywhere, Category=ObjectLibrary, meta=(AllowAbstract = ""))
diff --git a/Engine/Source/Runtime/Engine/Classes/Engine/ObjectReferencer.h b/Engine/Source/Runtime/Engine/Classes/Engine/ObjectReferencer.h
--- a/Engine/Source/Runtime/Engine/Classes/Engine/ObjectReferencer.h
+++ b/Engine/Source/Runtime/Engine/Classes/Engine/ObjectReferencer.h
@@ -11,7 +11,9 @@
 UCLASS()
 class UObjectReferencer : public UObject
 {
-	GENERATED_UCLASS_BODY()
+	GENERATED_BODY()
+public:
+	UObjectReferencer(const FObjectInitializer& ObjectInitializer = FObjectInitializer::Get());
 
 	/** Array of objects being referenced. */
 	UPROPERTY(EditAnywhere, Category=ObjectReferencer)
diff --git a/Engine/Source/Runtime/Engine/Classes/Engine/PendingNetGame.h b/Engine/Source/Runtime/Engine/Classes/Engine/PendingNetGame.h
--- a/Engine/Source/Runtime/Engine/Classes/Engine/PendingNetGame.h
+++ b/Engine/Source/Runtime/Engine/Classes/Engine/PendingNetGame.h
@@ -61,12 +61,13 @@ public:
 	virtual void NotifyControlMessage(UNetConnection* Connection, uint8 MessageType, class FInBunch& Bunch) PURE_VIRTUAL(FNetworkNotify::NotifyReceivedText,);
 };
 
-UCLASS(customConstructor, transient)
+UCLASS(transient)
 class UPendingNetGame :
 	public UObject,
 	public FNetworkNotify
 {
-	GENERATED_UCLASS_BODY()
+	GENERATED_BODY()
+public:
 
 	/** 
 	 * Net driver created for contacting the new server
diff --git a/Engine/Source/Runtime/Engine/Classes/Engine/PlaneReflectionCapture.h b/Engine/Source/Runtime/Engine/Classes/Engine/PlaneReflectionCapture.h
--- a/Engine/Source/Runtime/Engine/Classes/Engine/PlaneReflectionCapture.h
+++ b/Engine/Source/Runtime/Engine/Classes/Engine/PlaneReflectionCapture.h
@@ -10,7 +10,9 @@
 UCLASS(abstract, hidecategories=(Collision, Attachment, Actor), MinimalAPI)
 class APlaneReflectionCapture : public AReflectionCapture
 {
-	GENERATED_UCLASS_BODY()
+	GENERATED_BODY()
+public:
+	ENGINE_API APlaneReflectionCapture(const FObjectInitializer& ObjectInitializer = FObjectInitializer::Get());
 
 public:
 
diff --git a/Engine/Source/Runtime/Engine/Classes/Engine/PlatformInterfaceBase.h b/Engine/Source/Runtime/Engine/Classes/Engine/PlatformInterfaceBase.h
--- a/Engine/Source/Runtime/Engine/Classes/Engine/PlatformInterfaceBase.h
+++ b/Engine/Source/Runtime/Engine/Classes/Engine/PlatformInterfaceBase.h
@@ -34,7 +34,7 @@ enum EPlatformInterfaceDataType
 USTRUCT()
 struct FPlatformInterfaceData
 {
-	GENERATED_USTRUCT_BODY()
+	GENERATED_BODY()
 
 	/** An optional tag for this data */
 	UPROPERTY()
@@ -72,7 +72,7 @@ struct FPlatformInterfaceData
 USTRUCT()
 struct FPlatformInterfaceDelegateResult
 {
-	GENERATED_USTRUCT_BODY()
+	GENERATED_BODY()
 
 	/** This is always usable, no matter the type */
 	UPROPERTY()
@@ -100,7 +100,7 @@ DECLARE_DYNAMIC_DELEGATE_OneParam( FPlatformInterfaceDelegate, const struct FPla
 USTRUCT()
 struct FDelegateArray
 {
-	GENERATED_USTRUCT_BODY()
+	GENERATED_BODY()
 
 	UPROPERTY()
 	TArray<FPlatformInterfaceDelegate> Delegates;
@@ -110,7 +110,9 @@ struct FDelegateArray
 UCLASS(transient,MinimalAPI)
 class UPlatformInterfaceBase : public UObject
 {
-	GENERATED_UCLASS_BODY()
+	GENERATED_BODY()
+public:
+	ENGINE_API UPlatformInterfaceBase(const FObjectInitializer& ObjectInitializer = FObjectInitializer::Get());
 
 	/** Array of delegate arrays. Only add and remove via helper functions, and call via the helper delegate call function */
 	UPROPERTY()
diff --git a/Engine/Source/Runtime/Engine/Classes/Engine/PlatformInterfaceWebResponse.h b/Engine/Source/Runtime/Engine/Classes/Engine/PlatformInterfaceWebResponse.h
--- a/Engine/Source/Runtime/Engine/Classes/Engine/PlatformInterfaceWebResponse.h
+++ b/Engine/Source/Runtime/Engine/Classes/Engine/PlatformInterfaceWebResponse.h
@@ -13,7 +13,9 @@
 UCLASS(transient,MinimalAPI)
 class UPlatformInterfaceWebResponse : public UObject
 {
-	GENERATED_UCLASS_BODY()
+	GENERATED_BODY()
+public:
+	ENGINE_API UPlatformInterfaceWebResponse(const FObjectInitializer& ObjectInitializer = FObjectInitializer::Get());
 
 	/** This holds the original requested URL */
 	UPROPERTY()
diff --git a/Engine/Source/Runtime/Engine/Classes/Engine/Player.h b/Engine/Source/Runtime/Engine/Classes/Engine/Player.h
--- a/Engine/Source/Runtime/Engine/Classes/Engine/Player.h
+++ b/Engine/Source/Runtime/Engine/Classes/Engine/Player.h
@@ -10,7 +10,9 @@
 UCLASS(transient, config=Engine)
 class UPlayer : public UObject, public FExec
 {
-	GENERATED_UCLASS_BODY()
+	GENERATED_BODY()
+public:
+	UPlayer(const FObjectInitializer& ObjectInitializer = FObjectInitializer::Get());
 
 	/** The actor this player controls. */
 	UPROPERTY(transient)
diff --git a/Engine/Source/Runtime/Engine/Classes/Engine/PlayerStartPIE.h b/Engine/Source/Runtime/Engine/Classes/Engine/PlayerStartPIE.h
--- a/Engine/Source/Runtime/Engine/Classes/Engine/PlayerStartPIE.h
+++ b/Engine/Source/Runtime/Engine/Classes/Engine/PlayerStartPIE.h
@@ -11,7 +11,9 @@
 UCLASS(notplaceable,MinimalAPI)
 class APlayerStartPIE : public APlayerStart
 {
-	GENERATED_UCLASS_BODY()
+	GENERATED_BODY()
+public:
+	ENGINE_API APlayerStartPIE(const FObjectInitializer& ObjectInitializer = FObjectInitializer::Get());
 
 };
 
diff --git a/Engine/Source/Runtime/Engine/Classes/Engine/PointLight.h b/Engine/Source/Runtime/Engine/Classes/Engine/PointLight.h
--- a/Engine/Source/Runtime/Engine/Classes/Engine/PointLight.h
+++ b/Engine/Source/Runtime/Engine/Classes/Engine/PointLight.h
@@ -7,7 +7,9 @@
 UCLASS(ClassGroup=(Lights, PointLights), ComponentWrapperClass, MinimalAPI, meta=(ChildCanTick))
 class APointLight : public ALight
 {
-	GENERATED_UCLASS_BODY()
+	GENERATED_BODY()
+public:
+	ENGINE_API APointLight(const FObjectInitializer& ObjectInitializer = FObjectInitializer::Get());
 
 	UPROPERTY(BlueprintReadOnly, Category="Light", meta=(ExposeFunctionCategories="PointLight,Rendering|Lighting"))
 	class UPointLightComponent* PointLightComponent;
