Enabled UHT to digest GENERATED_BODY instead of GENERATED_UCLASS_BODY, GENERATED_USTRUCT_BODY, GENERATED_UINTERFACE_BODY or GENERATED_IINTERFACE_BODY, changed every occurence to the new syntax and fix

diff --git a/Engine/Source/Runtime/Engine/Classes/Engine/Polys.h b/Engine/Source/Runtime/Engine/Classes/Engine/Polys.h
--- a/Engine/Source/Runtime/Engine/Classes/Engine/Polys.h
+++ b/Engine/Source/Runtime/Engine/Classes/Engine/Polys.h
@@ -578,10 +578,11 @@ public:
 	}
 };
 
-UCLASS(customConstructor, MinimalAPI)
+UCLASS(MinimalAPI)
 class UPolys : public UObject
 {
-	GENERATED_UCLASS_BODY()
+	GENERATED_BODY()
+public:
 	// Elements.
 	TTransArray<FPoly> Element;
 
diff --git a/Engine/Source/Runtime/Engine/Classes/Engine/PostProcessVolume.h b/Engine/Source/Runtime/Engine/Classes/Engine/PostProcessVolume.h
--- a/Engine/Source/Runtime/Engine/Classes/Engine/PostProcessVolume.h
+++ b/Engine/Source/Runtime/Engine/Classes/Engine/PostProcessVolume.h
@@ -17,7 +17,9 @@
 UCLASS(autoexpandcategories=PostProcessVolume, hidecategories=(Advanced, Collision, Volume, Brush, Attachment), MinimalAPI)
 class APostProcessVolume : public AVolume, public IInterface_PostProcessVolume
 {
-	GENERATED_UCLASS_BODY()
+	GENERATED_BODY()
+public:
+	ENGINE_API APostProcessVolume(const FObjectInitializer& ObjectInitializer = FObjectInitializer::Get());
 
 	/** Post process settings to use for this volume. */
 	UPROPERTY(interp, Category=PostProcessVolume, meta=(ShowOnlyInnerProperties))
diff --git a/Engine/Source/Runtime/Engine/Classes/Engine/ReflectionCapture.h b/Engine/Source/Runtime/Engine/Classes/Engine/ReflectionCapture.h
--- a/Engine/Source/Runtime/Engine/Classes/Engine/ReflectionCapture.h
+++ b/Engine/Source/Runtime/Engine/Classes/Engine/ReflectionCapture.h
@@ -6,7 +6,9 @@
 UCLASS(abstract, hidecategories=(Collision, Attachment, Actor), MinimalAPI)
 class AReflectionCapture : public AActor
 {
-	GENERATED_UCLASS_BODY()
+	GENERATED_BODY()
+public:
+	ENGINE_API AReflectionCapture(const FObjectInitializer& ObjectInitializer = FObjectInitializer::Get());
 
 private_subobject:
 	/** Reflection capture component. */
diff --git a/Engine/Source/Runtime/Engine/Classes/Engine/RendererSettings.h b/Engine/Source/Runtime/Engine/Classes/Engine/RendererSettings.h
--- a/Engine/Source/Runtime/Engine/Classes/Engine/RendererSettings.h
+++ b/Engine/Source/Runtime/Engine/Classes/Engine/RendererSettings.h
@@ -91,7 +91,9 @@ UCLASS(config=Engine, defaultconfig)
 class ENGINE_API URendererSettings
 	: public UObject
 {
-	GENERATED_UCLASS_BODY()
+	GENERATED_BODY()
+public:
+	URendererSettings(const FObjectInitializer& ObjectInitializer = FObjectInitializer::Get());
 
 	UPROPERTY(config, EditAnywhere, Category=Mobile, meta=(
 		ConsoleVariable="r.MobileHDR",DisplayName="Mobile HDR",
diff --git a/Engine/Source/Runtime/Engine/Classes/Engine/SCS_Node.h b/Engine/Source/Runtime/Engine/Classes/Engine/SCS_Node.h
--- a/Engine/Source/Runtime/Engine/Classes/Engine/SCS_Node.h
+++ b/Engine/Source/Runtime/Engine/Classes/Engine/SCS_Node.h
@@ -10,7 +10,9 @@ DECLARE_DELEGATE_OneParam( FSCSNodeNameChanged, const FName& );
 UCLASS(MinimalAPI)
 class USCS_Node : public UObject
 {
-	GENERATED_UCLASS_BODY()
+	GENERATED_BODY()
+public:
+	ENGINE_API USCS_Node(const FObjectInitializer& ObjectInitializer = FObjectInitializer::Get());
 
 	/** Template for the component to create */
 	UPROPERTY()
diff --git a/Engine/Source/Runtime/Engine/Classes/Engine/Scene.h b/Engine/Source/Runtime/Engine/Classes/Engine/Scene.h
--- a/Engine/Source/Runtime/Engine/Classes/Engine/Scene.h
+++ b/Engine/Source/Runtime/Engine/Classes/Engine/Scene.h
@@ -34,7 +34,7 @@ enum EAntiAliasingMethod
 USTRUCT(BlueprintType)
 struct FPostProcessSettings
 {
-	GENERATED_USTRUCT_BODY()
+	GENERATED_BODY()
 
 	// first all bOverride_... as they get grouped together into bitfields
 
@@ -1040,9 +1040,11 @@ struct FPostProcessSettings
 UCLASS()
 class UScene : public UObject
 {
-	GENERATED_UCLASS_BODY()
+	GENERATED_BODY()
+public:
+	UScene(const FObjectInitializer& ObjectInitializer = FObjectInitializer::Get());
 
 
 	/** bits needed to store DPG value */
 	#define SDPG_NumBits	3
-};
\ No newline at end of file
+};
diff --git a/Engine/Source/Runtime/Engine/Classes/Engine/SceneCapture.h b/Engine/Source/Runtime/Engine/Classes/Engine/SceneCapture.h
--- a/Engine/Source/Runtime/Engine/Classes/Engine/SceneCapture.h
+++ b/Engine/Source/Runtime/Engine/Classes/Engine/SceneCapture.h
@@ -11,7 +11,9 @@
 UCLASS(abstract, hidecategories=(Collision, Attachment, Actor), MinimalAPI)
 class ASceneCapture : public AActor
 {
-	GENERATED_UCLASS_BODY()
+	GENERATED_BODY()
+public:
+	ENGINE_API ASceneCapture(const FObjectInitializer& ObjectInitializer = FObjectInitializer::Get());
 
 private_subobject:
 	/** To display the 3d camera in the editor. */
diff --git a/Engine/Source/Runtime/Engine/Classes/Engine/SceneCapture2D.h b/Engine/Source/Runtime/Engine/Classes/Engine/SceneCapture2D.h
--- a/Engine/Source/Runtime/Engine/Classes/Engine/SceneCapture2D.h
+++ b/Engine/Source/Runtime/Engine/Classes/Engine/SceneCapture2D.h
@@ -12,7 +12,9 @@
 UCLASS(hidecategories=(Collision, Material, Attachment, Actor), MinimalAPI)
 class ASceneCapture2D : public ASceneCapture
 {
-	GENERATED_UCLASS_BODY()
+	GENERATED_BODY()
+public:
+	ENGINE_API ASceneCapture2D(const FObjectInitializer& ObjectInitializer = FObjectInitializer::Get());
 
 private_subobject:
 	/** Scene capture component. */
diff --git a/Engine/Source/Runtime/Engine/Classes/Engine/SceneCaptureCube.h b/Engine/Source/Runtime/Engine/Classes/Engine/SceneCaptureCube.h
--- a/Engine/Source/Runtime/Engine/Classes/Engine/SceneCaptureCube.h
+++ b/Engine/Source/Runtime/Engine/Classes/Engine/SceneCaptureCube.h
@@ -12,7 +12,9 @@
 UCLASS(hidecategories = (Collision, Material, Attachment, Actor), MinimalAPI)
 class ASceneCaptureCube : public ASceneCapture
 {
-	GENERATED_UCLASS_BODY()
+	GENERATED_BODY()
+public:
+	ENGINE_API ASceneCaptureCube(const FObjectInitializer& ObjectInitializer = FObjectInitializer::Get());
 
 private_subobject:
 	/** Scene capture component. */
diff --git a/Engine/Source/Runtime/Engine/Classes/Engine/ScriptViewportClient.h b/Engine/Source/Runtime/Engine/Classes/Engine/ScriptViewportClient.h
--- a/Engine/Source/Runtime/Engine/Classes/Engine/ScriptViewportClient.h
+++ b/Engine/Source/Runtime/Engine/Classes/Engine/ScriptViewportClient.h
@@ -12,7 +12,9 @@
 UCLASS(transient)
 class UScriptViewportClient : public UObject, public FCommonViewportClient
 {
-	GENERATED_UCLASS_BODY()
+	GENERATED_BODY()
+public:
+	UScriptViewportClient(const FObjectInitializer& ObjectInitializer = FObjectInitializer::Get());
 
 };
 
diff --git a/Engine/Source/Runtime/Engine/Classes/Engine/ShadowMapTexture2D.h b/Engine/Source/Runtime/Engine/Classes/Engine/ShadowMapTexture2D.h
--- a/Engine/Source/Runtime/Engine/Classes/Engine/ShadowMapTexture2D.h
+++ b/Engine/Source/Runtime/Engine/Classes/Engine/ShadowMapTexture2D.h
@@ -7,7 +7,9 @@
 UCLASS(MinimalAPI)
 class UShadowMapTexture2D : public UTexture2D
 {
-	GENERATED_UCLASS_BODY()
+	GENERATED_BODY()
+public:
+	ENGINE_API UShadowMapTexture2D(const FObjectInitializer& ObjectInitializer = FObjectInitializer::Get());
 	
 	/** Bit-field with shadowmap flags. */
 	UPROPERTY()
diff --git a/Engine/Source/Runtime/Engine/Classes/Engine/SimpleConstructionScript.h b/Engine/Source/Runtime/Engine/Classes/Engine/SimpleConstructionScript.h
--- a/Engine/Source/Runtime/Engine/Classes/Engine/SimpleConstructionScript.h
+++ b/Engine/Source/Runtime/Engine/Classes/Engine/SimpleConstructionScript.h
@@ -9,7 +9,9 @@ class USCS_Node;
 UCLASS(MinimalAPI)
 class USimpleConstructionScript : public UObject
 {
-	GENERATED_UCLASS_BODY()
+	GENERATED_BODY()
+public:
+	ENGINE_API USimpleConstructionScript(const FObjectInitializer& ObjectInitializer = FObjectInitializer::Get());
 
 	// Begin UObject Interface
 	virtual void Serialize(FArchive& Ar) override;
diff --git a/Engine/Source/Runtime/Engine/Classes/Engine/SkeletalMesh.h b/Engine/Source/Runtime/Engine/Classes/Engine/SkeletalMesh.h
--- a/Engine/Source/Runtime/Engine/Classes/Engine/SkeletalMesh.h
+++ b/Engine/Source/Runtime/Engine/Classes/Engine/SkeletalMesh.h
@@ -56,7 +56,7 @@ enum SkeletalMeshOptimizationType
 USTRUCT()
 struct FBoneMirrorInfo
 {
-	GENERATED_USTRUCT_BODY()
+	GENERATED_BODY()
 
 	/** The bone to mirror. */
 	UPROPERTY(EditAnywhere, Category=BoneMirrorInfo, meta=(ArrayClamp = "RefSkeleton"))
@@ -78,7 +78,7 @@ struct FBoneMirrorInfo
 USTRUCT()
 struct FBoneMirrorExport
 {
-	GENERATED_USTRUCT_BODY()
+	GENERATED_BODY()
 
 	UPROPERTY(EditAnywhere, Category=BoneMirrorExport)
 	FName BoneName;
@@ -101,7 +101,7 @@ struct FBoneMirrorExport
 USTRUCT()
 struct FTriangleSortSettings
 {
-	GENERATED_USTRUCT_BODY()
+	GENERATED_BODY()
 
 	UPROPERTY(EditAnywhere, Category=TriangleSortSettings)
 	TEnumAsByte<enum ETriangleSortOption> TriangleSorting;
@@ -125,7 +125,7 @@ struct FTriangleSortSettings
 USTRUCT()
 struct FBoneReference
 {
-	GENERATED_USTRUCT_BODY()
+	GENERATED_BODY()
 
 	/** Name of bone to control. This is the main bone chain to modify from. **/
 	UPROPERTY(EditAnywhere, Category = BoneReference)
@@ -160,7 +160,7 @@ struct FBoneReference
 USTRUCT()
 struct FSkeletalMeshOptimizationSettings
 {
-	GENERATED_USTRUCT_BODY()
+	GENERATED_BODY()
 
 		/** The method to use when optimizing the skeletal mesh LOD */
 		UPROPERTY()
@@ -281,7 +281,7 @@ struct FSkeletalMeshOptimizationSettings
 USTRUCT()
 struct FSkeletalMeshLODInfo
 {
-	GENERATED_USTRUCT_BODY()
+	GENERATED_BODY()
 
 	/**	Indicates when to use this LOD. A smaller number means use this LOD when further away. */
 	UPROPERTY(EditAnywhere, Category=SkeletalMeshLODInfo)
@@ -322,7 +322,7 @@ struct FSkeletalMeshLODInfo
 USTRUCT()
 struct FMorphTargetMap
 {
-	GENERATED_USTRUCT_BODY()
+	GENERATED_BODY()
 
 	/** The bone to mirror. */
 	UPROPERTY(EditAnywhere, Category=MorphTargetMap)
@@ -403,7 +403,7 @@ struct FClothBonePlane
 USTRUCT()
 struct FClothPhysicsProperties
 {
-	GENERATED_USTRUCT_BODY()
+	GENERATED_BODY()
 
 	// Bending stiffness of the cloth in the range [0, 1]. 
 	UPROPERTY(EditAnywhere, Category = Stiffness, meta = (ClampMin = "0.0", ClampMax = "1.0", UIMin = "0.0", UIMax = "1.0"))
@@ -443,7 +443,7 @@ struct FClothPhysicsProperties
 USTRUCT()
 struct FClothingAssetData
 {
-	GENERATED_USTRUCT_BODY()
+	GENERATED_BODY()
 
 	/* User-defined asset name */
 	UPROPERTY(EditAnywhere, Category=ClothingAssetData)
@@ -491,7 +491,7 @@ struct FClothingAssetData
 USTRUCT()
 struct FSkeletalMaterial
 {
-	GENERATED_USTRUCT_BODY()
+	GENERATED_BODY()
 
 	FSkeletalMaterial()
 		: MaterialInterface( NULL )
@@ -531,7 +531,9 @@ class FSkeletalMeshResource;
 UCLASS(hidecategories=Object, MinimalAPI, BlueprintType)
 class USkeletalMesh : public UObject, public IInterface_CollisionDataProvider, public IInterface_AssetUserData
 {
-	GENERATED_UCLASS_BODY()
+	GENERATED_BODY()
+public:
+	ENGINE_API USkeletalMesh(const FObjectInitializer& ObjectInitializer = FObjectInitializer::Get());
 
 private:
 	/** Rendering resources created at import time. */
diff --git a/Engine/Source/Runtime/Engine/Classes/Engine/SkeletalMeshSocket.h b/Engine/Source/Runtime/Engine/Classes/Engine/SkeletalMeshSocket.h
--- a/Engine/Source/Runtime/Engine/Classes/Engine/SkeletalMeshSocket.h
+++ b/Engine/Source/Runtime/Engine/Classes/Engine/SkeletalMeshSocket.h
@@ -7,7 +7,9 @@
 UCLASS(hidecategories=Object, hidecategories=Actor, MinimalAPI)
 class USkeletalMeshSocket : public UObject
 {
-	GENERATED_UCLASS_BODY()
+	GENERATED_BODY()
+public:
+	ENGINE_API USkeletalMeshSocket(const FObjectInitializer& ObjectInitializer = FObjectInitializer::Get());
 
 	/** 
 	 *	Defines a named attachment location on the USkeletalMesh. 
diff --git a/Engine/Source/Runtime/Engine/Classes/Engine/SkyLight.h b/Engine/Source/Runtime/Engine/Classes/Engine/SkyLight.h
--- a/Engine/Source/Runtime/Engine/Classes/Engine/SkyLight.h
+++ b/Engine/Source/Runtime/Engine/Classes/Engine/SkyLight.h
@@ -7,7 +7,9 @@
 UCLASS(ClassGroup=Lights, hidecategories=(Input,Collision,Replication,Info), showcategories=("Rendering", "Input|MouseInput", "Input|TouchInput"), ComponentWrapperClass, ConversionRoot, Blueprintable)
 class ENGINE_API ASkyLight : public AInfo
 {
-	GENERATED_UCLASS_BODY()
+	GENERATED_BODY()
+public:
+	ASkyLight(const FObjectInitializer& ObjectInitializer = FObjectInitializer::Get());
 
 private_subobject:
 	/** @todo document */
diff --git a/Engine/Source/Runtime/Engine/Classes/Engine/SphereReflectionCapture.h b/Engine/Source/Runtime/Engine/Classes/Engine/SphereReflectionCapture.h
--- a/Engine/Source/Runtime/Engine/Classes/Engine/SphereReflectionCapture.h
+++ b/Engine/Source/Runtime/Engine/Classes/Engine/SphereReflectionCapture.h
@@ -13,7 +13,9 @@ class UDrawSphereComponent;
 UCLASS(hidecategories = (Collision, Attachment, Actor), MinimalAPI)
 class ASphereReflectionCapture : public AReflectionCapture
 {
-	GENERATED_UCLASS_BODY()
+	GENERATED_BODY()
+public:
+	ENGINE_API ASphereReflectionCapture(const FObjectInitializer& ObjectInitializer = FObjectInitializer::Get());
 
 private_subobject:
 	/** Sphere component used to visualize the capture radius */
diff --git a/Engine/Source/Runtime/Engine/Classes/Engine/SplineMeshActor.h b/Engine/Source/Runtime/Engine/Classes/Engine/SplineMeshActor.h
--- a/Engine/Source/Runtime/Engine/Classes/Engine/SplineMeshActor.h
+++ b/Engine/Source/Runtime/Engine/Classes/Engine/SplineMeshActor.h
@@ -11,7 +11,9 @@
 UCLASS(hideCategories = (Input), showCategories = ("Input|MouseInput", "Input|TouchInput"), ConversionRoot, ComponentWrapperClass, meta = (ChildCanTick))
 class ENGINE_API ASplineMeshActor : public AActor
 {
-	GENERATED_UCLASS_BODY()
+	GENERATED_BODY()
+public:
+	ASplineMeshActor(const FObjectInitializer& ObjectInitializer = FObjectInitializer::Get());
 
 private_subobject:
 	UPROPERTY(Category = SplineMeshActor, VisibleAnywhere, BlueprintReadOnly, meta = (ExposeFunctionCategories = "Mesh,Rendering,Physics,Components|StaticMesh,Components|SplineMesh", AllowPrivateAccess = "true"))
diff --git a/Engine/Source/Runtime/Engine/Classes/Engine/SpotLight.h b/Engine/Source/Runtime/Engine/Classes/Engine/SpotLight.h
--- a/Engine/Source/Runtime/Engine/Classes/Engine/SpotLight.h
+++ b/Engine/Source/Runtime/Engine/Classes/Engine/SpotLight.h
@@ -8,7 +8,9 @@
 UCLASS(ClassGroup=(Lights, SpotLights), MinimalAPI, ComponentWrapperClass, meta=(ChildCanTick))
 class ASpotLight : public ALight
 {
-	GENERATED_UCLASS_BODY()
+	GENERATED_BODY()
+public:
+	ENGINE_API ASpotLight(const FObjectInitializer& ObjectInitializer = FObjectInitializer::Get());
 
 	UPROPERTY(BlueprintReadOnly, Category="Light", meta=(ExposeFunctionCategories="SpotLight,Rendering|Lighting"))
 	class USpotLightComponent* SpotLightComponent;
diff --git a/Engine/Source/Runtime/Engine/Classes/Engine/StaticMesh.h b/Engine/Source/Runtime/Engine/Classes/Engine/StaticMesh.h
--- a/Engine/Source/Runtime/Engine/Classes/Engine/StaticMesh.h
+++ b/Engine/Source/Runtime/Engine/Classes/Engine/StaticMesh.h
@@ -53,7 +53,7 @@ enum EOptimizationType
 USTRUCT()
 struct FStaticMeshOptimizationSettings
 {
-	GENERATED_USTRUCT_BODY()
+	GENERATED_BODY()
 
 	/** The method to use when optimizing the skeletal mesh LOD */
 	UPROPERTY()
@@ -134,7 +134,7 @@ struct FStaticMeshOptimizationSettings
 USTRUCT()
 struct FStaticMeshSourceModel
 {
-	GENERATED_USTRUCT_BODY()
+	GENERATED_BODY()
 
 #if WITH_EDITOR
 	/** Imported raw mesh data. Optional for all but the first LOD. */
@@ -174,7 +174,7 @@ struct FStaticMeshSourceModel
 USTRUCT()
 struct FMeshSectionInfo
 {
-	GENERATED_USTRUCT_BODY()
+	GENERATED_BODY()
 
 	/** Index in to the Materials array on UStaticMesh. */
 	UPROPERTY()
@@ -215,7 +215,7 @@ bool operator!=(const FMeshSectionInfo& A, const FMeshSectionInfo& B);
 USTRUCT()
 struct FMeshSectionInfoMap
 {
-	GENERATED_USTRUCT_BODY()
+	GENERATED_BODY()
 
 	/** Maps an LOD+Section to the material it should render with. */
 	TMap<uint32,FMeshSectionInfo> Map;
@@ -245,7 +245,7 @@ struct FMeshSectionInfoMap
 USTRUCT()
 struct FAssetEditorOrbitCameraPosition
 {
-	GENERATED_USTRUCT_BODY()
+	GENERATED_BODY()
 
 	FAssetEditorOrbitCameraPosition()
 		: bIsSet(false)
@@ -285,10 +285,11 @@ struct FAssetEditorOrbitCameraPosition
  * @see https://docs.unrealengine.com/latest/INT/Engine/Content/Types/StaticMeshes/
  * @see AStaticMeshActor, UStaticMeshComponent
  */
-UCLASS(collapsecategories, hidecategories=Object, customconstructor, MinimalAPI, BlueprintType, config=Engine)
+UCLASS(collapsecategories, hidecategories=Object, MinimalAPI, BlueprintType, config=Engine)
 class UStaticMesh : public UObject, public IInterface_CollisionDataProvider, public IInterface_AssetUserData
 {
-	GENERATED_UCLASS_BODY()
+	GENERATED_BODY()
+public:
 
 	/** Pointer to the data used to render this static mesh. */
 	TScopedPointer<class FStaticMeshRenderData> RenderData;
diff --git a/Engine/Source/Runtime/Engine/Classes/Engine/StaticMeshActor.h b/Engine/Source/Runtime/Engine/Classes/Engine/StaticMeshActor.h
--- a/Engine/Source/Runtime/Engine/Classes/Engine/StaticMeshActor.h
+++ b/Engine/Source/Runtime/Engine/Classes/Engine/StaticMeshActor.h
@@ -15,7 +15,9 @@
 UCLASS(hidecategories=(Input), showcategories=("Input|MouseInput", "Input|TouchInput"), ConversionRoot, ComponentWrapperClass, meta=(ChildCanTick))
 class ENGINE_API AStaticMeshActor : public AActor
 {
-	GENERATED_UCLASS_BODY()
+	GENERATED_BODY()
+public:
+	AStaticMeshActor(const FObjectInitializer& ObjectInitializer = FObjectInitializer::Get());
 
 private_subobject:
 	DEPRECATED_FORGAME(4.6, "StaticMeshComponent should not be accessed directly, please use GetStaticMeshComponent() function instead. StaticMeshComponent will soon be private and your code will not compile.")
diff --git a/Engine/Source/Runtime/Engine/Classes/Engine/StaticMeshSocket.h b/Engine/Source/Runtime/Engine/Classes/Engine/StaticMeshSocket.h
--- a/Engine/Source/Runtime/Engine/Classes/Engine/StaticMeshSocket.h
+++ b/Engine/Source/Runtime/Engine/Classes/Engine/StaticMeshSocket.h
@@ -7,7 +7,9 @@
 UCLASS(hidecategories=Object, hidecategories=Actor, MinimalAPI)
 class UStaticMeshSocket : public UObject
 {
-	GENERATED_UCLASS_BODY()
+	GENERATED_BODY()
+public:
+	ENGINE_API UStaticMeshSocket(const FObjectInitializer& ObjectInitializer = FObjectInitializer::Get());
 
 	/** 
 	 *	Defines a named attachment location on the UStaticMesh. 
diff --git a/Engine/Source/Runtime/Engine/Classes/Engine/SubsurfaceProfile.h b/Engine/Source/Runtime/Engine/Classes/Engine/SubsurfaceProfile.h
--- a/Engine/Source/Runtime/Engine/Classes/Engine/SubsurfaceProfile.h
+++ b/Engine/Source/Runtime/Engine/Classes/Engine/SubsurfaceProfile.h
@@ -8,7 +8,7 @@
 USTRUCT(BlueprintType)
 struct FSubsurfaceProfileStruct
 {
-	GENERATED_USTRUCT_BODY()
+	GENERATED_BODY()
 	
 	/** in world/unreal units (cm) */
 	UPROPERTY(Category = "SubsurfaceProfileStruct", EditAnywhere, BlueprintReadOnly, meta = (ClampMin = "0.1", UIMax = "50.0", ClampMax = "1000.0"))
@@ -55,7 +55,9 @@ struct FSubsurfaceProfileStruct
 UCLASS(autoexpandcategories = SubsurfaceProfile, MinimalAPI)
 class USubsurfaceProfile : public UObject
 {
-	GENERATED_UCLASS_BODY()
+	GENERATED_BODY()
+public:
+	ENGINE_API USubsurfaceProfile(const FObjectInitializer& ObjectInitializer = FObjectInitializer::Get());
 
 	UPROPERTY(Category = USubsurfaceProfile, EditAnywhere, meta = (ShowOnlyInnerProperties))
 	struct FSubsurfaceProfileStruct Settings;
diff --git a/Engine/Source/Runtime/Engine/Classes/Engine/TargetPoint.h b/Engine/Source/Runtime/Engine/Classes/Engine/TargetPoint.h
--- a/Engine/Source/Runtime/Engine/Classes/Engine/TargetPoint.h
+++ b/Engine/Source/Runtime/Engine/Classes/Engine/TargetPoint.h
@@ -10,7 +10,9 @@
 UCLASS(MinimalAPI)
 class ATargetPoint : public AActor
 {
-	GENERATED_UCLASS_BODY()
+	GENERATED_BODY()
+public:
+	ENGINE_API ATargetPoint(const FObjectInitializer& ObjectInitializer = FObjectInitializer::Get());
 
 #if WITH_EDITORONLY_DATA
 private_subobject:
diff --git a/Engine/Source/Runtime/Engine/Classes/Engine/TextRenderActor.h b/Engine/Source/Runtime/Engine/Classes/Engine/TextRenderActor.h
--- a/Engine/Source/Runtime/Engine/Classes/Engine/TextRenderActor.h
+++ b/Engine/Source/Runtime/Engine/Classes/Engine/TextRenderActor.h
@@ -11,7 +11,9 @@
 UCLASS(MinimalAPI, ComponentWrapperClass, hideCategories = (Collision, Attachment, Actor))
 class ATextRenderActor : public AActor
 {
-	GENERATED_UCLASS_BODY()
+	GENERATED_BODY()
+public:
+	ENGINE_API ATextRenderActor(const FObjectInitializer& ObjectInitializer = FObjectInitializer::Get());
 
 	friend class UActorFactoryTextRender;
 
diff --git a/Engine/Source/Runtime/Engine/Classes/Engine/Texture.h b/Engine/Source/Runtime/Engine/Classes/Engine/Texture.h
--- a/Engine/Source/Runtime/Engine/Classes/Engine/Texture.h
+++ b/Engine/Source/Runtime/Engine/Classes/Engine/Texture.h
@@ -111,7 +111,7 @@ enum ETextureSourceFormat
 USTRUCT()
 struct FTextureSource
 {
-	GENERATED_USTRUCT_BODY()
+	GENERATED_BODY()
 
 	/** Default constructor. */
 	ENGINE_API FTextureSource();
@@ -277,7 +277,7 @@ private:
 USTRUCT()
 struct FTexturePlatformData
 {
-	GENERATED_USTRUCT_BODY()
+	GENERATED_BODY()
 
 	/** Width of the texture. */
 	int32 SizeX;
@@ -341,7 +341,9 @@ struct FTexturePlatformData
 UCLASS(abstract, MinimalAPI, BlueprintType)
 class UTexture : public UObject
 {
-	GENERATED_UCLASS_BODY()
+	GENERATED_BODY()
+public:
+	ENGINE_API UTexture(const FObjectInitializer& ObjectInitializer = FObjectInitializer::Get());
 
 	/*--------------------------------------------------------------------------
 		Editor only properties used to build the runtime texture data.
diff --git a/Engine/Source/Runtime/Engine/Classes/Engine/Texture2D.h b/Engine/Source/Runtime/Engine/Classes/Engine/Texture2D.h
--- a/Engine/Source/Runtime/Engine/Classes/Engine/Texture2D.h
+++ b/Engine/Source/Runtime/Engine/Classes/Engine/Texture2D.h
@@ -8,7 +8,9 @@
 UCLASS(hidecategories=Object, MinimalAPI, BlueprintType)
 class UTexture2D : public UTexture
 {
-	GENERATED_UCLASS_BODY()
+	GENERATED_BODY()
+public:
+	ENGINE_API UTexture2D(const FObjectInitializer& ObjectInitializer = FObjectInitializer::Get());
 
 public:
 	/** Number of miplevels the texture should have resident.					*/
diff --git a/Engine/Source/Runtime/Engine/Classes/Engine/Texture2DDynamic.h b/Engine/Source/Runtime/Engine/Classes/Engine/Texture2DDynamic.h
--- a/Engine/Source/Runtime/Engine/Classes/Engine/Texture2DDynamic.h
+++ b/Engine/Source/Runtime/Engine/Classes/Engine/Texture2DDynamic.h
@@ -7,7 +7,9 @@
 UCLASS(hidecategories=Object, MinimalAPI)
 class UTexture2DDynamic : public UTexture
 {
-	GENERATED_UCLASS_BODY()
+	GENERATED_BODY()
+public:
+	ENGINE_API UTexture2DDynamic(const FObjectInitializer& ObjectInitializer = FObjectInitializer::Get());
 
 	/** The width of the texture. */
 	int32 SizeX;
diff --git a/Engine/Source/Runtime/Engine/Classes/Engine/TextureCube.h b/Engine/Source/Runtime/Engine/Classes/Engine/TextureCube.h
--- a/Engine/Source/Runtime/Engine/Classes/Engine/TextureCube.h
+++ b/Engine/Source/Runtime/Engine/Classes/Engine/TextureCube.h
@@ -7,7 +7,9 @@
 UCLASS(hidecategories=Object, MinimalAPI)
 class UTextureCube : public UTexture
 {
-	GENERATED_UCLASS_BODY()
+	GENERATED_BODY()
+public:
+	ENGINE_API UTextureCube(const FObjectInitializer& ObjectInitializer = FObjectInitializer::Get());
 
 public:
 	/** Platform data. */
diff --git a/Engine/Source/Runtime/Engine/Classes/Engine/TextureLightProfile.h b/Engine/Source/Runtime/Engine/Classes/Engine/TextureLightProfile.h
--- a/Engine/Source/Runtime/Engine/Classes/Engine/TextureLightProfile.h
+++ b/Engine/Source/Runtime/Engine/Classes/Engine/TextureLightProfile.h
@@ -8,7 +8,9 @@
 UCLASS(hidecategories=(Object,CompositeTexture,Texture2D), MinimalAPI, BlueprintType)
 class UTextureLightProfile : public UTexture2D
 {
-	GENERATED_UCLASS_BODY()
+	GENERATED_BODY()
+public:
+	ENGINE_API UTextureLightProfile(const FObjectInitializer& ObjectInitializer = FObjectInitializer::Get());
 
 	/** Light brightness in Lumens, imported from IES profile, <= 0 if the profile is used for masking only. Use with InverseSquareFalloff. */
 	UPROPERTY(EditAnywhere, Category=TextureLightProfile, AssetRegistrySearchable)
diff --git a/Engine/Source/Runtime/Engine/Classes/Engine/TextureRenderTarget.h b/Engine/Source/Runtime/Engine/Classes/Engine/TextureRenderTarget.h
--- a/Engine/Source/Runtime/Engine/Classes/Engine/TextureRenderTarget.h
+++ b/Engine/Source/Runtime/Engine/Classes/Engine/TextureRenderTarget.h
@@ -13,7 +13,9 @@
 UCLASS(abstract, MinimalAPI)
 class UTextureRenderTarget : public UTexture
 {
-	GENERATED_UCLASS_BODY()	
+	GENERATED_BODY()
+public:
+	ENGINE_API UTextureRenderTarget(const FObjectInitializer& ObjectInitializer = FObjectInitializer::Get());
 
 	/** Will override FTextureRenderTarget2DResource::GetDisplayGamma if > 0. */
 	UPROPERTY(EditAnywhere, BlueprintReadWrite, Category=TextureRenderTarget)
diff --git a/Engine/Source/Runtime/Engine/Classes/Engine/TextureRenderTarget2D.h b/Engine/Source/Runtime/Engine/Classes/Engine/TextureRenderTarget2D.h
--- a/Engine/Source/Runtime/Engine/Classes/Engine/TextureRenderTarget2D.h
+++ b/Engine/Source/Runtime/Engine/Classes/Engine/TextureRenderTarget2D.h
@@ -15,7 +15,9 @@
 UCLASS(hidecategories=Object, hidecategories=Texture, MinimalAPI)
 class UTextureRenderTarget2D : public UTextureRenderTarget
 {
-	GENERATED_UCLASS_BODY()
+	GENERATED_BODY()
+public:
+	ENGINE_API UTextureRenderTarget2D(const FObjectInitializer& ObjectInitializer = FObjectInitializer::Get());
 
 	/** The width of the texture.												*/
 	UPROPERTY(EditAnywhere, BlueprintReadOnly, Category=TextureRenderTarget2D, AssetRegistrySearchable)
diff --git a/Engine/Source/Runtime/Engine/Classes/Engine/TextureRenderTargetCube.h b/Engine/Source/Runtime/Engine/Classes/Engine/TextureRenderTargetCube.h
--- a/Engine/Source/Runtime/Engine/Classes/Engine/TextureRenderTargetCube.h
+++ b/Engine/Source/Runtime/Engine/Classes/Engine/TextureRenderTargetCube.h
@@ -14,7 +14,9 @@
 UCLASS(hidecategories=Object, hidecategories=Texture, MinimalAPI)
 class UTextureRenderTargetCube : public UTextureRenderTarget
 {
-	GENERATED_UCLASS_BODY()
+	GENERATED_BODY()
+public:
+	ENGINE_API UTextureRenderTargetCube(const FObjectInitializer& ObjectInitializer = FObjectInitializer::Get());
 
 	/** The width of the texture.												*/
 	UPROPERTY(EditAnywhere, BlueprintReadWrite, Category=TextureRenderTargetCube, AssetRegistrySearchable)
diff --git a/Engine/Source/Runtime/Engine/Classes/Engine/TimelineTemplate.h b/Engine/Source/Runtime/Engine/Classes/Engine/TimelineTemplate.h
--- a/Engine/Source/Runtime/Engine/Classes/Engine/TimelineTemplate.h
+++ b/Engine/Source/Runtime/Engine/Classes/Engine/TimelineTemplate.h
@@ -7,7 +7,7 @@
 USTRUCT()
 struct FTTTrackBase
 {
-	GENERATED_USTRUCT_BODY()
+	GENERATED_BODY()
 
 	/** Name of this track */
 	UPROPERTY()
@@ -30,7 +30,7 @@ struct FTTTrackBase
 USTRUCT()
 struct FTTEventTrack : public FTTTrackBase
 {
-	GENERATED_USTRUCT_BODY()
+	GENERATED_BODY()
 
 	/** Curve object used to store keys */
 	UPROPERTY()
@@ -50,7 +50,7 @@ struct FTTEventTrack : public FTTTrackBase
 USTRUCT()
 struct FTTFloatTrack : public FTTTrackBase
 {
-	GENERATED_USTRUCT_BODY()
+	GENERATED_BODY()
 
 	/** Curve object used to define float value over time */
 	UPROPERTY()
@@ -70,7 +70,7 @@ struct FTTFloatTrack : public FTTTrackBase
 USTRUCT()
 struct FTTVectorTrack : public FTTTrackBase
 {
-	GENERATED_USTRUCT_BODY()
+	GENERATED_BODY()
 
 	/** Curve object used to define vector value over time */
 	UPROPERTY()
@@ -90,7 +90,7 @@ struct FTTVectorTrack : public FTTTrackBase
 USTRUCT()
 struct FTTLinearColorTrack : public FTTTrackBase
 {
-	GENERATED_USTRUCT_BODY()
+	GENERATED_BODY()
 
 	/** Curve object used to define color value over time */
 	UPROPERTY()
@@ -109,7 +109,9 @@ struct FTTLinearColorTrack : public FTTTrackBase
 UCLASS(MinimalAPI)
 class UTimelineTemplate : public UObject
 {
-	GENERATED_UCLASS_BODY()
+	GENERATED_BODY()
+public:
+	ENGINE_API UTimelineTemplate(const FObjectInitializer& ObjectInitializer = FObjectInitializer::Get());
 
 	/** Length of this timeline */
 	UPROPERTY(EditAnywhere, Category=TimelineTemplate)
diff --git a/Engine/Source/Runtime/Engine/Classes/Engine/TriggerBase.h b/Engine/Source/Runtime/Engine/Classes/Engine/TriggerBase.h
--- a/Engine/Source/Runtime/Engine/Classes/Engine/TriggerBase.h
+++ b/Engine/Source/Runtime/Engine/Classes/Engine/TriggerBase.h
@@ -11,7 +11,9 @@ class UBillboardComponent;
 UCLASS(ClassGroup=Common, abstract, ConversionRoot, MinimalAPI)
 class ATriggerBase : public AActor
 {
-	GENERATED_UCLASS_BODY()
+	GENERATED_BODY()
+public:
+	ENGINE_API ATriggerBase(const FObjectInitializer& ObjectInitializer = FObjectInitializer::Get());
 
 private_subobject:
 	/** Shape component used for collision */
diff --git a/Engine/Source/Runtime/Engine/Classes/Engine/TriggerBox.h b/Engine/Source/Runtime/Engine/Classes/Engine/TriggerBox.h
--- a/Engine/Source/Runtime/Engine/Classes/Engine/TriggerBox.h
+++ b/Engine/Source/Runtime/Engine/Classes/Engine/TriggerBox.h
@@ -9,7 +9,9 @@
 UCLASS(MinimalAPI)
 class ATriggerBox : public ATriggerBase
 {
-	GENERATED_UCLASS_BODY()
+	GENERATED_BODY()
+public:
+	ENGINE_API ATriggerBox(const FObjectInitializer& ObjectInitializer = FObjectInitializer::Get());
 
 
 #if WITH_EDITOR
diff --git a/Engine/Source/Runtime/Engine/Classes/Engine/TriggerCapsule.h b/Engine/Source/Runtime/Engine/Classes/Engine/TriggerCapsule.h
--- a/Engine/Source/Runtime/Engine/Classes/Engine/TriggerCapsule.h
+++ b/Engine/Source/Runtime/Engine/Classes/Engine/TriggerCapsule.h
@@ -9,7 +9,9 @@
 UCLASS(MinimalAPI)
 class ATriggerCapsule : public ATriggerBase
 {
-	GENERATED_UCLASS_BODY()
+	GENERATED_BODY()
+public:
+	ENGINE_API ATriggerCapsule(const FObjectInitializer& ObjectInitializer = FObjectInitializer::Get());
 
 
 #if WITH_EDITOR
diff --git a/Engine/Source/Runtime/Engine/Classes/Engine/TriggerSphere.h b/Engine/Source/Runtime/Engine/Classes/Engine/TriggerSphere.h
--- a/Engine/Source/Runtime/Engine/Classes/Engine/TriggerSphere.h
+++ b/Engine/Source/Runtime/Engine/Classes/Engine/TriggerSphere.h
@@ -9,7 +9,9 @@
 UCLASS(MinimalAPI)
 class ATriggerSphere : public ATriggerBase
 {
-	GENERATED_UCLASS_BODY()
+	GENERATED_BODY()
+public:
+	ENGINE_API ATriggerSphere(const FObjectInitializer& ObjectInitializer = FObjectInitializer::Get());
 
 
 #if WITH_EDITOR
diff --git a/Engine/Source/Runtime/Engine/Classes/Engine/TriggerVolume.h b/Engine/Source/Runtime/Engine/Classes/Engine/TriggerVolume.h
--- a/Engine/Source/Runtime/Engine/Classes/Engine/TriggerVolume.h
+++ b/Engine/Source/Runtime/Engine/Classes/Engine/TriggerVolume.h
@@ -8,7 +8,9 @@
 UCLASS()
 class ATriggerVolume : public AVolume
 {
-	GENERATED_UCLASS_BODY()
+	GENERATED_BODY()
+public:
+	ATriggerVolume(const FObjectInitializer& ObjectInitializer = FObjectInitializer::Get());
 
 	// Begin UObject interface.
 #if WITH_EDITOR
diff --git a/Engine/Source/Runtime/Engine/Classes/Engine/TwitterIntegrationBase.h b/Engine/Source/Runtime/Engine/Classes/Engine/TwitterIntegrationBase.h
--- a/Engine/Source/Runtime/Engine/Classes/Engine/TwitterIntegrationBase.h
+++ b/Engine/Source/Runtime/Engine/Classes/Engine/TwitterIntegrationBase.h
@@ -31,7 +31,9 @@ enum ETwitterIntegrationDelegate
 UCLASS(config=Engine)
 class UTwitterIntegrationBase : public UPlatformInterfaceBase
 {
-	GENERATED_UCLASS_BODY()
+	GENERATED_BODY()
+public:
+	UTwitterIntegrationBase(const FObjectInitializer& ObjectInitializer = FObjectInitializer::Get());
 
 	/**
 	 * Perform any needed initialization
diff --git a/Engine/Source/Runtime/Engine/Classes/Engine/UserDefinedEnum.h b/Engine/Source/Runtime/Engine/Classes/Engine/UserDefinedEnum.h
--- a/Engine/Source/Runtime/Engine/Classes/Engine/UserDefinedEnum.h
+++ b/Engine/Source/Runtime/Engine/Classes/Engine/UserDefinedEnum.h
@@ -17,7 +17,9 @@
 UCLASS()
 class ENGINE_API UUserDefinedEnum : public UEnum
 {
-	GENERATED_UCLASS_BODY()
+	GENERATED_BODY()
+public:
+	UUserDefinedEnum(const FObjectInitializer& ObjectInitializer = FObjectInitializer::Get());
 
 #if WITH_EDITORONLY_DATA
 	UPROPERTY()
diff --git a/Engine/Source/Runtime/Engine/Classes/Engine/UserDefinedStruct.h b/Engine/Source/Runtime/Engine/Classes/Engine/UserDefinedStruct.h
--- a/Engine/Source/Runtime/Engine/Classes/Engine/UserDefinedStruct.h
+++ b/Engine/Source/Runtime/Engine/Classes/Engine/UserDefinedStruct.h
@@ -22,7 +22,9 @@ enum EUserDefinedStructureStatus
 UCLASS()
 class ENGINE_API UUserDefinedStruct : public UScriptStruct
 {
-	GENERATED_UCLASS_BODY()
+	GENERATED_BODY()
+public:
+	UUserDefinedStruct(const FObjectInitializer& ObjectInitializer = FObjectInitializer::Get());
 
 public:
 #if WITH_EDITORONLY_DATA
diff --git a/Engine/Source/Runtime/Engine/Classes/Engine/UserInterfaceSettings.h b/Engine/Source/Runtime/Engine/Classes/Engine/UserInterfaceSettings.h
--- a/Engine/Source/Runtime/Engine/Classes/Engine/UserInterfaceSettings.h
+++ b/Engine/Source/Runtime/Engine/Classes/Engine/UserInterfaceSettings.h
@@ -44,7 +44,9 @@ UCLASS(config=Engine, defaultconfig)
 class ENGINE_API UUserInterfaceSettings
 	: public UObject
 {
-	GENERATED_UCLASS_BODY()
+	GENERATED_BODY()
+public:
+	UUserInterfaceSettings(const FObjectInitializer& ObjectInitializer = FObjectInitializer::Get());
 
 public:
 
diff --git a/Engine/Source/Runtime/Engine/Classes/Engine/WindDirectionalSource.h b/Engine/Source/Runtime/Engine/Classes/Engine/WindDirectionalSource.h
--- a/Engine/Source/Runtime/Engine/Classes/Engine/WindDirectionalSource.h
+++ b/Engine/Source/Runtime/Engine/Classes/Engine/WindDirectionalSource.h
@@ -6,7 +6,9 @@
 UCLASS(ClassGroup=Wind, showcategories=(Rendering, "Utilities|Transformation"))
 class ENGINE_API AWindDirectionalSource : public AInfo
 {
-	GENERATED_UCLASS_BODY()
+	GENERATED_BODY()
+public:
+	AWindDirectionalSource(const FObjectInitializer& ObjectInitializer = FObjectInitializer::Get());
 
 private_subobject:
 	DEPRECATED_FORGAME(4.6, "Component should not be accessed directly, please use GetComponent() function instead. Component will soon be private and your code will not compile.")
diff --git a/Engine/Source/Runtime/Engine/Classes/Engine/World.h b/Engine/Source/Runtime/Engine/Classes/Engine/World.h
--- a/Engine/Source/Runtime/Engine/Classes/Engine/World.h
+++ b/Engine/Source/Runtime/Engine/Classes/Engine/World.h
@@ -324,7 +324,7 @@ struct ENGINE_API FLevelViewportInfo
 USTRUCT()
 struct FStartPhysicsTickFunction : public FTickFunction
 {
-	GENERATED_USTRUCT_BODY()
+	GENERATED_BODY()
 
 	/** World this tick function belongs to **/
 	class UWorld*	Target;
@@ -347,7 +347,7 @@ struct FStartPhysicsTickFunction : public FTickFunction
 USTRUCT()
 struct FEndPhysicsTickFunction : public FTickFunction
 {
-	GENERATED_USTRUCT_BODY()
+	GENERATED_BODY()
 
 	/** World this tick function belongs to **/
 	class UWorld*	Target;
@@ -370,7 +370,7 @@ struct FEndPhysicsTickFunction : public FTickFunction
 USTRUCT()
 struct FStartClothSimulationFunction : public FTickFunction
 {
-	GENERATED_USTRUCT_BODY()
+	GENERATED_BODY()
 
 	/** World this tick function belongs to **/
 	class UWorld*	Target;
@@ -393,7 +393,7 @@ struct FStartClothSimulationFunction : public FTickFunction
 USTRUCT()
 struct FEndClothSimulationFunction : public FTickFunction
 {
-	GENERATED_USTRUCT_BODY()
+	GENERATED_BODY()
 
 	/** World this tick function belongs to **/
 	class UWorld*	Target;
@@ -537,10 +537,11 @@ public:
  *
  */
 
-UCLASS(customConstructor, config=Engine)
+UCLASS(config=Engine)
 class ENGINE_API UWorld : public UObject, public FNetworkNotify
 {
-	GENERATED_UCLASS_BODY()
+	GENERATED_BODY()
+public:
 
 	~UWorld();
 
diff --git a/Engine/Source/Runtime/Engine/Classes/Engine/WorldComposition.h b/Engine/Source/Runtime/Engine/Classes/Engine/WorldComposition.h
--- a/Engine/Source/Runtime/Engine/Classes/Engine/WorldComposition.h
+++ b/Engine/Source/Runtime/Engine/Classes/Engine/WorldComposition.h
@@ -64,7 +64,9 @@ struct FDistanceVisibleLevel
 UCLASS(config=Engine)
 class ENGINE_API UWorldComposition : public UObject
 {
-	GENERATED_UCLASS_BODY()
+	GENERATED_BODY()
+public:
+	UWorldComposition(const FObjectInitializer& ObjectInitializer = FObjectInitializer::Get());
 
 	typedef TArray<FWorldCompositionTile> FTilesList;
 
@@ -221,4 +223,4 @@ public:
 	UPROPERTY(config)
 	bool						bLoadAllTilesDuringCinematic;
 
-};
\ No newline at end of file
+};
diff --git a/Engine/Source/Runtime/Engine/Classes/Exporters/Exporter.h b/Engine/Source/Runtime/Engine/Classes/Exporters/Exporter.h
--- a/Engine/Source/Runtime/Engine/Classes/Exporters/Exporter.h
+++ b/Engine/Source/Runtime/Engine/Classes/Exporters/Exporter.h
@@ -14,7 +14,9 @@ class UActorComponent;
 UCLASS(abstract, transient, MinimalAPI)
 class UExporter : public UObject
 {
-	GENERATED_UCLASS_BODY()
+	GENERATED_BODY()
+public:
+	ENGINE_API UExporter(const FObjectInitializer& ObjectInitializer = FObjectInitializer::Get());
 
 	/** Supported class of this exporter */
 	UPROPERTY()
diff --git a/Engine/Source/Runtime/Engine/Classes/GameFramework/CameraBlockingVolume.h b/Engine/Source/Runtime/Engine/Classes/GameFramework/CameraBlockingVolume.h
--- a/Engine/Source/Runtime/Engine/Classes/GameFramework/CameraBlockingVolume.h
+++ b/Engine/Source/Runtime/Engine/Classes/GameFramework/CameraBlockingVolume.h
@@ -9,7 +9,9 @@
 UCLASS()
 class ENGINE_API ACameraBlockingVolume : public AVolume
 {
-	GENERATED_UCLASS_BODY()
+	GENERATED_BODY()
+public:
+	ACameraBlockingVolume(const FObjectInitializer& ObjectInitializer = FObjectInitializer::Get());
 };
 
 
diff --git a/Engine/Source/Runtime/Engine/Classes/GameFramework/CheatManager.h b/Engine/Source/Runtime/Engine/Classes/GameFramework/CheatManager.h
--- a/Engine/Source/Runtime/Engine/Classes/GameFramework/CheatManager.h
+++ b/Engine/Source/Runtime/Engine/Classes/GameFramework/CheatManager.h
@@ -57,7 +57,9 @@ struct FDebugTraceInfo
 UCLASS(Within=PlayerController)
 class ENGINE_API UCheatManager : public UObject
 {
-	GENERATED_UCLASS_BODY()
+	GENERATED_BODY()
+public:
+	UCheatManager(const FObjectInitializer& ObjectInitializer = FObjectInitializer::Get());
 
 	/** Debug camera - used to have independent camera without stopping gameplay */
 	UPROPERTY()
@@ -204,8 +206,10 @@ class ENGINE_API UCheatManager : public UObject
 	UFUNCTION(exec)
 	virtual void ToggleAILogging();
 
-	UFUNCTION(reliable, server, WithValidation)
+	UFUNCTION(reliable, server="ServerToggleAILogging_Implementation", WithValidation="ServerToggleAILogging_Validate")
 	virtual void ServerToggleAILogging();
+	virtual void ServerToggleAILogging_Implementation();
+	virtual bool ServerToggleAILogging_Validate();
 
 	/** Toggle capsule trace debugging. Will trace a capsule from current view point and show where it hits the world */
 	UFUNCTION(exec)
diff --git a/Engine/Source/Runtime/Engine/Classes/GameFramework/Controller.h b/Engine/Source/Runtime/Engine/Classes/GameFramework/Controller.h
--- a/Engine/Source/Runtime/Engine/Classes/GameFramework/Controller.h
+++ b/Engine/Source/Runtime/Engine/Classes/GameFramework/Controller.h
@@ -31,7 +31,9 @@ DECLARE_DYNAMIC_MULTICAST_DELEGATE_FourParams( FInstigatedAnyDamageSignature, fl
 UCLASS(abstract, notplaceable, NotBlueprintable, HideCategories=(Collision,Rendering,"Utilities|Transformation")) 
 class ENGINE_API AController : public AActor, public INavAgentInterface
 {
-	GENERATED_UCLASS_BODY()
+	GENERATED_BODY()
+public:
+	AController(const FObjectInitializer& ObjectInitializer = FObjectInitializer::Get());
 
 private:
 	/** Pawn currently being controlled by this controller.  Use Pawn.Possess() to take control of a pawn */
@@ -153,12 +155,14 @@ public:
 	class APlayerController* CastToPlayerController();
 
 	/** Replicated function to set the pawn location and rotation, allowing server to force (ex. teleports). */
-	UFUNCTION(reliable, client)
+	UFUNCTION(reliable, client="ClientSetLocation_Implementation")
 	void ClientSetLocation(FVector NewLocation, FRotator NewRotation);
+	virtual void ClientSetLocation_Implementation(FVector NewLocation, FRotator NewRotation);
 
 	/** Replicated function to set the pawn rotation, allowing the server to force. */
-	UFUNCTION(reliable, client)
+	UFUNCTION(reliable, client="ClientSetRotation_Implementation")
 	void ClientSetRotation(FRotator NewRotation, bool bResetCamera = false);
+	virtual void ClientSetRotation_Implementation(FRotator NewRotation, bool bResetCamera);
 
 	/** Return the Pawn that is currently 'controlled' by this PlayerController */
 	UFUNCTION(BlueprintCallable, Category="Pawn", meta=(FriendlyName="Get Controlled Pawn"))
diff --git a/Engine/Source/Runtime/Engine/Classes/GameFramework/DamageType.h b/Engine/Source/Runtime/Engine/Classes/GameFramework/DamageType.h
--- a/Engine/Source/Runtime/Engine/Classes/GameFramework/DamageType.h
+++ b/Engine/Source/Runtime/Engine/Classes/GameFramework/DamageType.h
@@ -15,7 +15,9 @@
 UCLASS(MinimalAPI, const, Blueprintable, BlueprintType)
 class UDamageType : public UObject
 {
-	GENERATED_UCLASS_BODY()
+	GENERATED_BODY()
+public:
+	ENGINE_API UDamageType(const FObjectInitializer& ObjectInitializer = FObjectInitializer::Get());
 
 	/** True if this damagetype is caused by the world (falling off level, into lava, etc). */
 	UPROPERTY(EditAnywhere, BlueprintReadWrite, Category=DamageType)
diff --git a/Engine/Source/Runtime/Engine/Classes/GameFramework/DefaultPawn.h b/Engine/Source/Runtime/Engine/Classes/GameFramework/DefaultPawn.h
--- a/Engine/Source/Runtime/Engine/Classes/GameFramework/DefaultPawn.h
+++ b/Engine/Source/Runtime/Engine/Classes/GameFramework/DefaultPawn.h
@@ -20,7 +20,9 @@ class UStaticMeshComponent;
 UCLASS(config=Game, Blueprintable, BlueprintType)
 class ENGINE_API ADefaultPawn : public APawn
 {
-	GENERATED_UCLASS_BODY()
+	GENERATED_BODY()
+public:
+	ADefaultPawn(const FObjectInitializer& ObjectInitializer = FObjectInitializer::Get());
 
 	// Begin Pawn overrides
 	virtual UPawnMovementComponent* GetMovementComponent() const override;
diff --git a/Engine/Source/Runtime/Engine/Classes/GameFramework/DefaultPhysicsVolume.h b/Engine/Source/Runtime/Engine/Classes/GameFramework/DefaultPhysicsVolume.h
--- a/Engine/Source/Runtime/Engine/Classes/GameFramework/DefaultPhysicsVolume.h
+++ b/Engine/Source/Runtime/Engine/Classes/GameFramework/DefaultPhysicsVolume.h
@@ -17,7 +17,9 @@
 UCLASS(notplaceable, transient, MinimalAPI)
 class ADefaultPhysicsVolume : public APhysicsVolume
 {
-	GENERATED_UCLASS_BODY()
+	GENERATED_BODY()
+public:
+	ENGINE_API ADefaultPhysicsVolume(const FObjectInitializer& ObjectInitializer = FObjectInitializer::Get());
 
 	// Begin AActor interface
 	ENGINE_API virtual void Destroyed() override;
diff --git a/Engine/Source/Runtime/Engine/Classes/GameFramework/EngineMessage.h b/Engine/Source/Runtime/Engine/Classes/GameFramework/EngineMessage.h
--- a/Engine/Source/Runtime/Engine/Classes/GameFramework/EngineMessage.h
+++ b/Engine/Source/Runtime/Engine/Classes/GameFramework/EngineMessage.h
@@ -7,7 +7,9 @@
 UCLASS(abstract, MinimalAPI)
 class  UEngineMessage : public ULocalMessage
 {
-	GENERATED_UCLASS_BODY()
+	GENERATED_BODY()
+public:
+	ENGINE_API UEngineMessage(const FObjectInitializer& ObjectInitializer = FObjectInitializer::Get());
 
 	/** Message displayed in message dialog when player pawn fails to spawn because no playerstart was available. */
 	UPROPERTY()
diff --git a/Engine/Source/Runtime/Engine/Classes/GameFramework/FloatingPawnMovement.h b/Engine/Source/Runtime/Engine/Classes/GameFramework/FloatingPawnMovement.h
--- a/Engine/Source/Runtime/Engine/Classes/GameFramework/FloatingPawnMovement.h
+++ b/Engine/Source/Runtime/Engine/Classes/GameFramework/FloatingPawnMovement.h
@@ -15,7 +15,9 @@
 UCLASS(ClassGroup = Movement, meta = (BlueprintSpawnableComponent))
 class ENGINE_API UFloatingPawnMovement : public UPawnMovementComponent
 {
-	GENERATED_UCLASS_BODY()
+	GENERATED_BODY()
+public:
+	UFloatingPawnMovement(const FObjectInitializer& ObjectInitializer = FObjectInitializer::Get());
 
 	//Begin UActorComponent Interface
 	virtual void TickComponent(float DeltaTime, enum ELevelTick TickType, FActorComponentTickFunction *ThisTickFunction) override;
diff --git a/Engine/Source/Runtime/Engine/Classes/GameFramework/ForceFeedbackEffect.h b/Engine/Source/Runtime/Engine/Classes/GameFramework/ForceFeedbackEffect.h
--- a/Engine/Source/Runtime/Engine/Classes/GameFramework/ForceFeedbackEffect.h
+++ b/Engine/Source/Runtime/Engine/Classes/GameFramework/ForceFeedbackEffect.h
@@ -6,7 +6,7 @@
 USTRUCT()
 struct FForceFeedbackChannelDetails
 {
-	GENERATED_USTRUCT_BODY()
+	GENERATED_BODY()
 
 	UPROPERTY(EditAnywhere, Category="ChannelDetails")
 	uint32 bAffectsLeftLarge:1;
@@ -35,7 +35,7 @@ struct FForceFeedbackChannelDetails
 USTRUCT()
 struct FActiveForceFeedbackEffect
 {
-	GENERATED_USTRUCT_BODY()
+	GENERATED_BODY()
 
 	UPROPERTY()
 	class UForceFeedbackEffect* ForceFeedbackEffect;
@@ -70,7 +70,9 @@ struct FActiveForceFeedbackEffect
 UCLASS(BlueprintType, MinimalAPI)
 class UForceFeedbackEffect : public UObject
 {
-	GENERATED_UCLASS_BODY()
+	GENERATED_BODY()
+public:
+	ENGINE_API UForceFeedbackEffect(const FObjectInitializer& ObjectInitializer = FObjectInitializer::Get());
 
 	UPROPERTY(EditAnywhere, Category="ForceFeedbackEffect")
 	TArray<FForceFeedbackChannelDetails> ChannelDetails;
@@ -86,4 +88,4 @@ class UForceFeedbackEffect : public UObject
 	float GetDuration();
 
 	void GetValues(const float EvalTime, FForceFeedbackValues& Values) const;
-};
\ No newline at end of file
+};
diff --git a/Engine/Source/Runtime/Engine/Classes/GameFramework/GameMode.h b/Engine/Source/Runtime/Engine/Classes/GameFramework/GameMode.h
--- a/Engine/Source/Runtime/Engine/Classes/GameFramework/GameMode.h
+++ b/Engine/Source/Runtime/Engine/Classes/GameFramework/GameMode.h
@@ -33,7 +33,7 @@ DECLARE_DELEGATE_RetVal( bool, FCanUnpause );
 USTRUCT()
 struct FGameClassShortName
 {
-	GENERATED_USTRUCT_BODY()
+	GENERATED_BODY()
 
 	/** Abbreviation that can be used as an alias for the class name */
 	UPROPERTY()
@@ -62,7 +62,9 @@ struct FGameClassShortName
 UCLASS(config=Game, notplaceable, BlueprintType, Blueprintable, Transient, hideCategories=(Info, Rendering, MovementReplication, Replication, Actor), meta=(ShortTooltip="Game Mode defines the game being played, its rules, scoring, and other facets of the game type."))
 class ENGINE_API AGameMode : public AInfo
 {
-	GENERATED_UCLASS_BODY()
+	GENERATED_BODY()
+public:
+	AGameMode(const FObjectInitializer& ObjectInitializer = FObjectInitializer::Get());
 
 	// Code to deal with the match state machine
 
diff --git a/Engine/Source/Runtime/Engine/Classes/GameFramework/GameNetworkManager.h b/Engine/Source/Runtime/Engine/Classes/GameFramework/GameNetworkManager.h
--- a/Engine/Source/Runtime/Engine/Classes/GameFramework/GameNetworkManager.h
+++ b/Engine/Source/Runtime/Engine/Classes/GameFramework/GameNetworkManager.h
@@ -22,7 +22,9 @@ enum EStandbyType
 UCLASS(config=Game, notplaceable)
 class ENGINE_API AGameNetworkManager : public AInfo
 {
-	GENERATED_UCLASS_BODY()
+	GENERATED_BODY()
+public:
+	AGameNetworkManager(const FObjectInitializer& ObjectInitializer = FObjectInitializer::Get());
 
 	//======================================================================================================================
 	// Listen server dynamic netspeed adjustment
diff --git a/Engine/Source/Runtime/Engine/Classes/GameFramework/GameSession.h b/Engine/Source/Runtime/Engine/Classes/GameFramework/GameSession.h
--- a/Engine/Source/Runtime/Engine/Classes/GameFramework/GameSession.h
+++ b/Engine/Source/Runtime/Engine/Classes/GameFramework/GameSession.h
@@ -14,7 +14,9 @@ class APlayerController;
 UCLASS(config=Game, notplaceable)
 class ENGINE_API AGameSession : public AInfo
 {
-	GENERATED_UCLASS_BODY()
+	GENERATED_BODY()
+public:
+	AGameSession(const FObjectInitializer& ObjectInitializer = FObjectInitializer::Get());
 
 	/** Maximum number of spectators allowed by this server. */
 	UPROPERTY(globalconfig)
diff --git a/Engine/Source/Runtime/Engine/Classes/GameFramework/GameState.h b/Engine/Source/Runtime/Engine/Classes/GameFramework/GameState.h
--- a/Engine/Source/Runtime/Engine/Classes/GameFramework/GameState.h
+++ b/Engine/Source/Runtime/Engine/Classes/GameFramework/GameState.h
@@ -11,7 +11,9 @@
 UCLASS(config=Game, notplaceable, BlueprintType, Blueprintable)
 class ENGINE_API AGameState : public AInfo
 {
-	GENERATED_UCLASS_BODY()
+	GENERATED_BODY()
+public:
+	AGameState(const FObjectInitializer& ObjectInitializer = FObjectInitializer::Get());
 
 	/** Class of the server's game mode, assigned by GameMode. */
 	UPROPERTY(replicatedUsing=OnRep_GameModeClass)
diff --git a/Engine/Source/Runtime/Engine/Classes/GameFramework/GameUserSettings.h b/Engine/Source/Runtime/Engine/Classes/GameFramework/GameUserSettings.h
--- a/Engine/Source/Runtime/Engine/Classes/GameFramework/GameUserSettings.h
+++ b/Engine/Source/Runtime/Engine/Classes/GameFramework/GameUserSettings.h
@@ -10,7 +10,9 @@
 UCLASS(config=GameUserSettings, configdonotcheckdefaults)
 class ENGINE_API UGameUserSettings : public UObject
 {
-	GENERATED_UCLASS_BODY()
+	GENERATED_BODY()
+public:
+	UGameUserSettings(const FObjectInitializer& ObjectInitializer = FObjectInitializer::Get());
 
 public:
 
diff --git a/Engine/Source/Runtime/Engine/Classes/GameFramework/HUD.h b/Engine/Source/Runtime/Engine/Classes/GameFramework/HUD.h
--- a/Engine/Source/Runtime/Engine/Classes/GameFramework/HUD.h
+++ b/Engine/Source/Runtime/Engine/Classes/GameFramework/HUD.h
@@ -25,7 +25,9 @@ class UFont;
 UCLASS(config=Game, hidecategories=(Rendering,Actor,Input,Replication), showcategories=("Input|MouseInput", "Input|TouchInput"), notplaceable, transient, BlueprintType, Blueprintable)
 class ENGINE_API AHUD : public AActor
 {
-	GENERATED_UCLASS_BODY()
+	GENERATED_BODY()
+public:
+	AHUD(const FObjectInitializer& ObjectInitializer = FObjectInitializer::Get());
 
 	/** Pre-defined FColors for convenience. */
 	UPROPERTY()
@@ -151,14 +153,16 @@ public:
 	 * @param InFont 				font to use
 	 * @param FontScale 			scale
 	 */
-	UFUNCTION(reliable, client, SealedEvent)
+	UFUNCTION(reliable, client="AddDebugText_Implementation", SealedEvent)
 	void AddDebugText(const FString& DebugText, AActor* SrcActor = NULL, float Duration = 0, FVector Offset = FVector(ForceInit), FVector DesiredOffset = FVector(ForceInit), FColor TextColor = FColor(ForceInit), bool bSkipOverwriteCheck = false, bool bAbsoluteLocation = false, bool bKeepAttachedToActor = false, UFont* InFont = NULL, float FontScale = 1.0);
+	void AddDebugText_Implementation(const FString& DebugText, AActor* SrcActor, float Duration, FVector Offset, FVector DesiredOffset, FColor TextColor, bool bSkipOverwriteCheck, bool bAbsoluteLocation, bool bKeepAttachedToActor, UFont* InFont, float FontScale);
 
 	/**
 	 * Remove all debug strings added via AddDebugText
 	 */
-	UFUNCTION(reliable, client, SealedEvent)
+	UFUNCTION(reliable, client="RemoveAllDebugStrings_Implementation", SealedEvent)
 	void RemoveAllDebugStrings();
+	void RemoveAllDebugStrings_Implementation();
 
 	/**
 	 * Remove debug strings for the given actor
@@ -166,8 +170,9 @@ public:
 	 * @param	SrcActor			Actor whose string you wish to remove
 	 * @param	bLeaveDurationText	when true text that has a finite duration will be removed, otherwise all will be removed for given actor
 	 */
-	UFUNCTION(reliable, client, SealedEvent)
+	UFUNCTION(reliable, client="RemoveDebugText_Implementation", SealedEvent)
 	void RemoveDebugText(AActor* SrcActor, bool bLeaveDurationText = false);
+	void RemoveDebugText_Implementation(AActor* SrcActor, bool bLeaveDurationText);
 
 	/** Hook to allow blueprints to do custom HUD drawing. @see bSuppressNativeHUD to control HUD drawing in base class. */
 	UFUNCTION(BlueprintImplementableEvent, BlueprintCosmetic)
diff --git a/Engine/Source/Runtime/Engine/Classes/GameFramework/InputSettings.h b/Engine/Source/Runtime/Engine/Classes/GameFramework/InputSettings.h
--- a/Engine/Source/Runtime/Engine/Classes/GameFramework/InputSettings.h
+++ b/Engine/Source/Runtime/Engine/Classes/GameFramework/InputSettings.h
@@ -14,7 +14,9 @@ UCLASS(config=Input, defaultconfig)
 class ENGINE_API UInputSettings
 	: public UObject
 {
-	GENERATED_UCLASS_BODY()
+	GENERATED_BODY()
+public:
+	UInputSettings(const FObjectInitializer& ObjectInitializer = FObjectInitializer::Get());
 
 	/** Properties of Axis controls */
 	UPROPERTY(config, EditAnywhere, EditFixedSize, Category="Bindings", meta=(ToolTip="List of Axis Properties"), AdvancedDisplay)
diff --git a/Engine/Source/Runtime/Engine/Classes/GameFramework/KillZVolume.h b/Engine/Source/Runtime/Engine/Classes/GameFramework/KillZVolume.h
--- a/Engine/Source/Runtime/Engine/Classes/GameFramework/KillZVolume.h
+++ b/Engine/Source/Runtime/Engine/Classes/GameFramework/KillZVolume.h
@@ -14,7 +14,9 @@
 UCLASS()
 class ENGINE_API AKillZVolume : public APhysicsVolume
 {
-	GENERATED_UCLASS_BODY()
+	GENERATED_BODY()
+public:
+	AKillZVolume(const FObjectInitializer& ObjectInitializer = FObjectInitializer::Get());
 	
 	//Begin PhysicsVolume Interface
 	virtual void ActorEnteredVolume(class AActor* Other) override;
diff --git a/Engine/Source/Runtime/Engine/Classes/GameFramework/LocalMessage.h b/Engine/Source/Runtime/Engine/Classes/GameFramework/LocalMessage.h
--- a/Engine/Source/Runtime/Engine/Classes/GameFramework/LocalMessage.h
+++ b/Engine/Source/Runtime/Engine/Classes/GameFramework/LocalMessage.h
@@ -19,8 +19,8 @@
 USTRUCT()
 struct ENGINE_API FClientReceiveData
 {
+	GENERATED_BODY()
 	//always need to be here
-	GENERATED_USTRUCT_BODY()
 
 	UPROPERTY()
 	APlayerController* LocalPC;
@@ -49,7 +49,9 @@ struct ENGINE_API FClientReceiveData
 UCLASS(abstract)
 class ENGINE_API ULocalMessage : public UObject
 {
-	GENERATED_UCLASS_BODY()
+	GENERATED_BODY()
+public:
+	ULocalMessage(const FObjectInitializer& ObjectInitializer = FObjectInitializer::Get());
 	/** send message to client */
 	virtual void ClientReceive(const FClientReceiveData& ClientData) const PURE_VIRTUAL(ULocalMessage::ClientReceive,);
 };
diff --git a/Engine/Source/Runtime/Engine/Classes/GameFramework/MovementComponent.h b/Engine/Source/Runtime/Engine/Classes/GameFramework/MovementComponent.h
--- a/Engine/Source/Runtime/Engine/Classes/GameFramework/MovementComponent.h
+++ b/Engine/Source/Runtime/Engine/Classes/GameFramework/MovementComponent.h
@@ -46,7 +46,9 @@ enum class EPlaneConstraintAxisSetting : uint8
 UCLASS(ClassGroup=Movement, abstract, BlueprintType)
 class ENGINE_API UMovementComponent : public UActorComponent
 {
-	GENERATED_UCLASS_BODY()
+	GENERATED_BODY()
+public:
+	UMovementComponent(const FObjectInitializer& ObjectInitializer = FObjectInitializer::Get());
 
 	/**
 	 * The component we move and update.
diff --git a/Engine/Source/Runtime/Engine/Classes/GameFramework/NavMovementComponent.h b/Engine/Source/Runtime/Engine/Classes/GameFramework/NavMovementComponent.h
--- a/Engine/Source/Runtime/Engine/Classes/GameFramework/NavMovementComponent.h
+++ b/Engine/Source/Runtime/Engine/Classes/GameFramework/NavMovementComponent.h
@@ -22,7 +22,9 @@ class UCapsuleComponent;
 UCLASS(abstract)
 class ENGINE_API UNavMovementComponent : public UMovementComponent
 {
-	GENERATED_UCLASS_BODY()
+	GENERATED_BODY()
+public:
+	UNavMovementComponent(const FObjectInitializer& ObjectInitializer = FObjectInitializer::Get());
 
 	/** Properties that define how the component can move. */
 	UPROPERTY(EditAnywhere, BlueprintReadWrite, Category=MovementComponent, meta=(DisplayName="Movement Capabilities", Keywords="Nav Agent"))
diff --git a/Engine/Source/Runtime/Engine/Classes/GameFramework/OnlineSession.h b/Engine/Source/Runtime/Engine/Classes/GameFramework/OnlineSession.h
--- a/Engine/Source/Runtime/Engine/Classes/GameFramework/OnlineSession.h
+++ b/Engine/Source/Runtime/Engine/Classes/GameFramework/OnlineSession.h
@@ -10,7 +10,9 @@
 UCLASS(config=Game)
 class ENGINE_API UOnlineSession : public UObject
 {
-	GENERATED_UCLASS_BODY()
+	GENERATED_BODY()
+public:
+	UOnlineSession(const FObjectInitializer& ObjectInitializer = FObjectInitializer::Get());
 
 public:
 
diff --git a/Engine/Source/Runtime/Engine/Classes/GameFramework/PainCausingVolume.h b/Engine/Source/Runtime/Engine/Classes/GameFramework/PainCausingVolume.h
--- a/Engine/Source/Runtime/Engine/Classes/GameFramework/PainCausingVolume.h
+++ b/Engine/Source/Runtime/Engine/Classes/GameFramework/PainCausingVolume.h
@@ -16,7 +16,9 @@
 UCLASS()
 class ENGINE_API APainCausingVolume : public APhysicsVolume
 {
-	GENERATED_UCLASS_BODY()
+	GENERATED_BODY()
+public:
+	APainCausingVolume(const FObjectInitializer& ObjectInitializer = FObjectInitializer::Get());
 
 	/** Whether volume currently causes damage. */
 	UPROPERTY(EditAnywhere, BlueprintReadWrite, Category=PainCausingVolume)
diff --git a/Engine/Source/Runtime/Engine/Classes/GameFramework/PhysicsVolume.h b/Engine/Source/Runtime/Engine/Classes/GameFramework/PhysicsVolume.h
--- a/Engine/Source/Runtime/Engine/Classes/GameFramework/PhysicsVolume.h
+++ b/Engine/Source/Runtime/Engine/Classes/GameFramework/PhysicsVolume.h
@@ -12,7 +12,9 @@
 UCLASS()
 class ENGINE_API APhysicsVolume : public AVolume
 {
-	GENERATED_UCLASS_BODY()
+	GENERATED_BODY()
+public:
+	APhysicsVolume(const FObjectInitializer& ObjectInitializer = FObjectInitializer::Get());
 
 	// Begin UObject interface.
 #if WITH_EDITOR
diff --git a/Engine/Source/Runtime/Engine/Classes/GameFramework/PlayerController.h b/Engine/Source/Runtime/Engine/Classes/GameFramework/PlayerController.h
--- a/Engine/Source/Runtime/Engine/Classes/GameFramework/PlayerController.h
+++ b/Engine/Source/Runtime/Engine/Classes/GameFramework/PlayerController.h
@@ -146,7 +146,9 @@ protected:
 UCLASS(config=Game, BlueprintType, Blueprintable, meta=(ShortTooltip="A Player Controller is an actor responsible for controlling a Pawn used by the player."))
 class ENGINE_API APlayerController : public AController
 {
-	GENERATED_UCLASS_BODY()
+	GENERATED_BODY()
+public:
+	APlayerController(const FObjectInitializer& ObjectInitializer = FObjectInitializer::Get());
 	
 	/** UPlayer associated with this PlayerController.  Could be a local player or a net connection. */
 	UPROPERTY()
@@ -337,12 +339,14 @@ public:
 	virtual void LocalTravel(const FString& URL);
 
 	/** Return the client to the main menu gracefully */
-	UFUNCTION(Reliable, Client)
+	UFUNCTION(Reliable, Client="ClientReturnToMainMenu_Implementation")
 	virtual void ClientReturnToMainMenu(const FString& ReturnReason);
+	virtual void ClientReturnToMainMenu_Implementation(const FString& ReturnReason);
 
 	/** Development RPC for testing object reference replication */
-	UFUNCTION(Reliable, Client)
+	UFUNCTION(Reliable, Client="ClientRepObjRef_Implementation")
 	virtual void ClientRepObjRef(UObject* Object);
+	virtual void ClientRepObjRef_Implementation(UObject* Object);
 
 	/**
 	 * Locally try to pause game (call serverpause to pause network game); returns success indicator.  Calls GameMode's SetPause().
@@ -507,8 +511,9 @@ public:
 	 * Tell the client to enable or disable voice chat (not muting)
 	 * @param bEnable enable or disable voice chat
 	 */
-	UFUNCTION(Reliable, Client)
+	UFUNCTION(Reliable, Client="ClientEnableNetworkVoice_Implementation")
 	virtual void ClientEnableNetworkVoice(bool bEnable);
+	virtual void ClientEnableNetworkVoice_Implementation(bool bEnable);
 
 	/** Enable voice chat transmission */
 	void StartTalking();
@@ -530,36 +535,43 @@ public:
 	 *
 	 * NOTE: This is done as an RPC instead of variable replication because ordering matters
 	 */
-	UFUNCTION(Reliable, Client)
+	UFUNCTION(Reliable, Client="ClientVoiceHandshakeComplete_Implementation")
 	virtual void ClientVoiceHandshakeComplete();
+	virtual void ClientVoiceHandshakeComplete_Implementation();
 
 	/**
 	 * Tell the server to mute a player for this controller
 	 * @param PlayerId player id to mute
 	 */
-	UFUNCTION(server, reliable, WithValidation)
+	UFUNCTION(server="ServerMutePlayer_Implementation", reliable, WithValidation="ServerMutePlayer_Validate")
 	virtual void ServerMutePlayer(FUniqueNetIdRepl PlayerId);
+	virtual void ServerMutePlayer_Implementation(FUniqueNetIdRepl PlayerId);
+	virtual bool ServerMutePlayer_Validate(FUniqueNetIdRepl );
 
 	/**
 	 * Tell the server to unmute a player for this controller
 	 * @param PlayerId player id to unmute
 	 */
-	UFUNCTION(server, reliable, WithValidation )
+	UFUNCTION(server="ServerUnmutePlayer_Implementation", reliable, WithValidation="ServerUnmutePlayer_Validate" )
 	virtual void ServerUnmutePlayer(FUniqueNetIdRepl PlayerId);
+	virtual void ServerUnmutePlayer_Implementation(FUniqueNetIdRepl PlayerId);
+	virtual bool ServerUnmutePlayer_Validate(FUniqueNetIdRepl );
 
 	/**
 	 * Tell the client to mute a player for this controller
 	 * @param PlayerId player id to mute
 	 */
-	UFUNCTION(Reliable, Client)
+	UFUNCTION(Reliable, Client="ClientMutePlayer_Implementation")
 	virtual void ClientMutePlayer(FUniqueNetIdRepl PlayerId);
+	virtual void ClientMutePlayer_Implementation(FUniqueNetIdRepl PlayerId);
 
 	/**
 	 * Tell the client to unmute a player for this controller
 	 * @param PlayerId player id to unmute
 	 */
-	UFUNCTION(Reliable, Client)
+	UFUNCTION(Reliable, Client="ClientUnmutePlayer_Implementation")
 	virtual void ClientUnmutePlayer(FUniqueNetIdRepl PlayerId);
+	virtual void ClientUnmutePlayer_Implementation(FUniqueNetIdRepl PlayerId);
 
 	/**
 	 * Mutes a remote player on the server and then tells the client to mute
@@ -596,54 +608,64 @@ public:
 	virtual void SendToConsole(const FString& Command);
 
 	/** Adds a location to the texture streaming system for the specified duration. */
-	UFUNCTION(reliable, client, SealedEvent)
+	UFUNCTION(reliable, client="ClientAddTextureStreamingLoc_Implementation", SealedEvent)
 	void ClientAddTextureStreamingLoc(FVector InLoc, float Duration, bool bOverrideLocation);
+	void ClientAddTextureStreamingLoc_Implementation(FVector InLoc, float Duration, bool bOverrideLocation);
 
 	/** Tells client to cancel any pending map change. */
-	UFUNCTION(Reliable, Client)
+	UFUNCTION(Reliable, Client="ClientCancelPendingMapChange_Implementation")
 	void ClientCancelPendingMapChange();
+	virtual void ClientCancelPendingMapChange_Implementation();
 
 	/** Set CurrentNetSpeed to the lower of its current value and Cap. */
-	UFUNCTION(Reliable, Client)
+	UFUNCTION(Reliable, Client="ClientCapBandwidth_Implementation")
 	void ClientCapBandwidth(int32 Cap);
+	virtual void ClientCapBandwidth_Implementation(int32 Cap);
 
 	/** Actually performs the level transition prepared by PrepareMapChange(). */
-	UFUNCTION(Reliable, Client)	
+	UFUNCTION(Reliable, Client="ClientCommitMapChange_Implementation")	
 	void ClientCommitMapChange();
+	virtual void ClientCommitMapChange_Implementation();
 
 	/** Tells the client to block until all pending level streaming actions are complete
 	 * happens at the end of the tick
 	 * primarily used to force update the client ASAP at join time
 	 */
-	UFUNCTION(reliable, client, SealedEvent)
+	UFUNCTION(reliable, client="ClientFlushLevelStreaming_Implementation", SealedEvent)
 	void ClientFlushLevelStreaming();
+	void ClientFlushLevelStreaming_Implementation();
 
 	/** Forces GC at the end of the tick on the client */
-	UFUNCTION(Reliable, Client)
+	UFUNCTION(Reliable, Client="ClientForceGarbageCollection_Implementation")
 	void ClientForceGarbageCollection();
+	virtual void ClientForceGarbageCollection_Implementation();
 
 	/** 
 	 * Replicated function called by GameHasEnded().
 	 * @param	EndGameFocus - actor to view with camera
 	 * @param	bIsWinner - true if this controller is on winning team
 	 */
-	UFUNCTION(Reliable, Client)
+	UFUNCTION(Reliable, Client="ClientGameEnded_Implementation")
 	void ClientGameEnded(class AActor* EndGameFocus, bool bIsWinner);
+	virtual void ClientGameEnded_Implementation(AActor* EndGameFocus, bool bIsWinner);
 
 	/** 
 	 * Server uses this to force client into NewState .
 	 * @Note ALL STATE NAMES NEED TO BE DEFINED IN name table in UnrealNames.h to be correctly replicated (so they are mapped to the same thing on client and server).
 	 */
-	UFUNCTION(Reliable, Client)
+	UFUNCTION(Reliable, Client="ClientGotoState_Implementation")
 	void ClientGotoState(FName NewState);
+	virtual void ClientGotoState_Implementation(FName NewState);
 
 	/** calls IgnoreLookInput on client */
-	UFUNCTION(Reliable, Client)
+	UFUNCTION(Reliable, Client="ClientIgnoreLookInput_Implementation")
 	void ClientIgnoreLookInput(bool bIgnore);
+	virtual void ClientIgnoreLookInput_Implementation(bool bIgnore);
 
 	/** calls IgnoreMoveInput on client */
-	UFUNCTION(Reliable, Client)
+	UFUNCTION(Reliable, Client="ClientIgnoreMoveInput_Implementation")
 	void ClientIgnoreMoveInput(bool bIgnore);
+	virtual void ClientIgnoreMoveInput_Implementation(bool bIgnore);
 
 	/**
 	 * Outputs a message to HUD
@@ -651,8 +673,9 @@ public:
 	 * @param Type - @todo document
 	 * @param MsgLifeTime - Optional length of time to display 0 = default time
 	 */
-	UFUNCTION(Reliable, Client)
+	UFUNCTION(Reliable, Client="ClientMessage_Implementation")
 	void ClientMessage(const FString& S, FName Type = NAME_None, float MsgLifeTime = 0);
+	virtual void ClientMessage_Implementation(const FString& S, FName Type, float MsgLifeTime);
 
 	/** Play the indicated CameraAnim on this camera.
 	 * @param AnimToPlay - Camera animation to play
@@ -665,8 +688,9 @@ public:
 	 * @param Space - Animation play area
 	 * @param CustomPlaySpace - Matrix used when Space = CAPS_UserDefined
 	 */
-	UFUNCTION(unreliable, client, BlueprintCallable, Category = "Game|Feedback")
+	UFUNCTION(unreliable, client="ClientPlayCameraAnim_Implementation", BlueprintCallable, Category = "Game|Feedback")
 	void ClientPlayCameraAnim(class UCameraAnim* AnimToPlay, float Scale=0, float Rate=0, float BlendInTime=0, float BlendOutTime=0, bool bLoop=false, bool bRandomStartTime=false, ECameraAnimPlaySpace::Type Space=ECameraAnimPlaySpace::CameraLocal, FRotator CustomPlaySpace=FRotator::ZeroRotator);
+	virtual void ClientPlayCameraAnim_Implementation(UCameraAnim* AnimToPlay, float Scale, float Rate, float BlendInTime, float BlendOutTime, bool bLoop, bool bRandomStartTime, ECameraAnimPlaySpace::Type Space, FRotator CustomPlaySpace);
 
 	/** 
 	 * Play Camera Shake 
@@ -675,8 +699,9 @@ public:
 	 * @param PlaySpace - Which coordinate system to play the shake in (used for CameraAnims within the shake).
 	 * @param UserPlaySpaceRot - Matrix used when PlaySpace = CAPS_UserDefined
 	 */
-	UFUNCTION(unreliable, client, BlueprintCallable, Category="Game|Feedback")
+	UFUNCTION(unreliable, client="ClientPlayCameraShake_Implementation", BlueprintCallable, Category="Game|Feedback")
 	void ClientPlayCameraShake(TSubclassOf<class UCameraShake> Shake, float Scale = 0, ECameraAnimPlaySpace::Type PlaySpace = ECameraAnimPlaySpace::CameraLocal, FRotator UserPlaySpaceRot = FRotator::ZeroRotator);
+	virtual void ClientPlayCameraShake_Implementation(TSubclassOf<UCameraShake>  Shake, float Scale, ECameraAnimPlaySpace::Type PlaySpace, FRotator UserPlaySpaceRot);
 
 	/**
 	 * Play sound client-side (so only the client will hear it)
@@ -684,8 +709,9 @@ public:
 	 * @param VolumeMultiplier - Volume multiplier to apply to the sound
 	 * @param PitchMultiplier - Pitch multiplier to apply to the sound
 	 */
-	UFUNCTION(unreliable, client)
+	UFUNCTION(unreliable, client="ClientPlaySound_Implementation")
 	void ClientPlaySound(class USoundBase* Sound, float VolumeMultiplier = 1.f, float PitchMultiplier = 1.f );
+	virtual void ClientPlaySound_Implementation(USoundBase* Sound, float VolumeMultiplier, float PitchMultiplier);
 
 	/**
 	 * Play sound client-side at the specified location
@@ -694,8 +720,9 @@ public:
 	 * @param VolumeMultiplier - Volume multiplier to apply to the sound
 	 * @param PitchMultiplier - Pitch multiplier to apply to the sound
 	 */
-	UFUNCTION(unreliable, client)
+	UFUNCTION(unreliable, client="ClientPlaySoundAtLocation_Implementation")
 	void ClientPlaySoundAtLocation(class USoundBase* Sound, FVector Location, float VolumeMultiplier = 1.f, float PitchMultiplier = 1.f);
+	virtual void ClientPlaySoundAtLocation_Implementation(USoundBase* Sound, FVector Location, float VolumeMultiplier, float PitchMultiplier);
 
 	/** Asynchronously loads the given level in preparation for a streaming map transition.
 	 * the server sends one function per level name since dynamic arrays can't be replicated
@@ -703,8 +730,9 @@ public:
 	 * @param bFirst - whether this is the first item in the list (so clear the list first)
 	 * @param bLast - whether this is the last item in the list (so start preparing the change after receiving it)
 	 */
-	UFUNCTION(Reliable, Client)
+	UFUNCTION(Reliable, Client="ClientPrepareMapChange_Implementation")
 	void ClientPrepareMapChange(FName LevelName, bool bFirst, bool bLast);
+	virtual void ClientPrepareMapChange_Implementation(FName LevelName, bool bFirst, bool bLast);
 
 	/**
 	 * Forces the streaming system to disregard the normal logic for the specified duration and
@@ -714,23 +742,27 @@ public:
 	 * @param bEnableStreaming	- Whether to start (true) or stop (false) streaming
 	 * @param CinematicTextureGroups	- Bitfield indicating which texture groups that use extra high-resolution mips
 	 */
-	UFUNCTION(Reliable, Client)
+	UFUNCTION(Reliable, Client="ClientPrestreamTextures_Implementation")
 	void ClientPrestreamTextures(class AActor* ForcedActor, float ForceDuration, bool bEnableStreaming, int32 CinematicTextureGroups = 0);
+	virtual void ClientPrestreamTextures_Implementation(AActor* ForcedActor, float ForceDuration, bool bEnableStreaming, int32 CinematicTextureGroups);
 
 	/** Tell client to reset the PlayerController */
-	UFUNCTION(Reliable, Client)
+	UFUNCTION(Reliable, Client="ClientReset_Implementation")
 	void ClientReset();
+	virtual void ClientReset_Implementation();
 
 	/** Tell client to restart the level */
-	UFUNCTION(Reliable, Client)
+	UFUNCTION(Reliable, Client="ClientRestart_Implementation")
 	void ClientRestart(class APawn* NewPawn);
+	virtual void ClientRestart_Implementation(APawn* NewPawn);
 
 	/** 
 	 * Tells the client to block until all pending level streaming actions are complete.
 	 * Happens at the end of the tick primarily used to force update the client ASAP at join time.
 	 */
-	UFUNCTION(Reliable, Client)
+	UFUNCTION(Reliable, Client="ClientSetBlockOnAsyncLoading_Implementation")
 	void ClientSetBlockOnAsyncLoading();
+	virtual void ClientSetBlockOnAsyncLoading_Implementation();
 
 	/** Tell client to fade camera
 	 * @Param bEnableFading - true if we should apply FadeColor/FadeAmount to the screen
@@ -739,19 +771,22 @@ public:
 	 * @Param FadeTime - length of time for fade to occur over
 	 * @Param bFadeAudio - true to apply fading of audio alongside the video
 	 */
-	UFUNCTION(Reliable, Client)
+	UFUNCTION(Reliable, Client="ClientSetCameraFade_Implementation")
 	void ClientSetCameraFade(bool bEnableFading, FColor FadeColor = FColor(ForceInit), FVector2D FadeAlpha = FVector2D(ForceInit), float FadeTime = 0, bool bFadeAudio = false);
+	virtual void ClientSetCameraFade_Implementation(bool bEnableFading, FColor FadeColor, FVector2D FadeAlpha, float FadeTime, bool bFadeAudio);
 
 	/**
 	 * Replicated function to set camera style on client
 	 * @param	NewCamMode, name defining the new camera mode
 	 */
-	UFUNCTION(Reliable, Client)
+	UFUNCTION(Reliable, Client="ClientSetCameraMode_Implementation")
 	void ClientSetCameraMode(FName NewCamMode);
+	virtual void ClientSetCameraMode_Implementation(FName NewCamMode);
 
 	/** Called by the server to synchronize cinematic transitions with the client */
-	UFUNCTION(Reliable, Client)
+	UFUNCTION(Reliable, Client="ClientSetCinematicMode_Implementation")
 	void ClientSetCinematicMode(bool bInCinematicMode, bool bAffectsMovement, bool bAffectsTurning, bool bAffectsHUD);
+	virtual void ClientSetCinematicMode_Implementation(bool bInCinematicMode, bool bAffectsMovement, bool bAffectsTurning, bool bAffectsHUD);
 
 	/**
 	 * Forces the streaming system to disregard the normal logic for the specified duration and
@@ -761,12 +796,14 @@ public:
 	 * @param ForceDuration	- Number of seconds to keep all mip-levels in memory, disregarding the normal priority logic.
 	 * @param CinematicTextureGroups	- Bitfield indicating which texture groups that use extra high-resolution mips
 	 */
-	UFUNCTION(Reliable, Client)
+	UFUNCTION(Reliable, Client="ClientSetForceMipLevelsToBeResident_Implementation")
 	void ClientSetForceMipLevelsToBeResident(class UMaterialInterface* Material, float ForceDuration, int32 CinematicTextureGroups = 0);
+	virtual void ClientSetForceMipLevelsToBeResident_Implementation(UMaterialInterface* Material, float ForceDuration, int32 CinematicTextureGroups);
 
 	/** Set the client's class of HUD and spawns a new instance of it. If there was already a HUD active, it is destroyed. */
-	UFUNCTION(BlueprintCallable, Category="HUD", Reliable, Client)
+	UFUNCTION(BlueprintCallable, Category="HUD", Reliable, Client="ClientSetHUD_Implementation")
 	void ClientSetHUD(TSubclassOf<class AHUD> NewHUDClass);
+	virtual void ClientSetHUD_Implementation(TSubclassOf<AHUD>  NewHUDClass);
 
 	/** Helper to get the size of the HUD canvas for this player controller.  Returns 0 if there is no HUD */
 	UFUNCTION(BlueprintCallable, Category="HUD")
@@ -780,24 +817,29 @@ public:
 	 * @param A - new actor to set as view target
 	 * @param TransitionParams - parameters to use for controlling the transition
 	 */
-	UFUNCTION(Reliable, Client)
+	UFUNCTION(Reliable, Client="ClientSetViewTarget_Implementation")
 	void ClientSetViewTarget(class AActor* A, struct FViewTargetTransitionParams TransitionParams = FViewTargetTransitionParams());
+	virtual void ClientSetViewTarget_Implementation(AActor* A, FViewTargetTransitionParams TransitionParams);
 
 	/** Spawn a camera lens effect (e.g. blood). */
-	UFUNCTION(unreliable, client, BlueprintCallable, Category="Game|Feedback")
+	UFUNCTION(unreliable, client="ClientSpawnCameraLensEffect_Implementation", BlueprintCallable, Category="Game|Feedback")
 	void ClientSpawnCameraLensEffect(TSubclassOf<class AEmitterCameraLensEffectBase>  LensEffectEmitterClass);
+	virtual void ClientSpawnCameraLensEffect_Implementation(TSubclassOf<AEmitterCameraLensEffectBase>  LensEffectEmitterClass);
 
 	/** Removes all Camera Lens Effects. */
-	UFUNCTION(reliable, client, BlueprintCallable, Category="Game|Feedback")
+	UFUNCTION(reliable, client="ClientClearCameraLensEffects_Implementation", BlueprintCallable, Category="Game|Feedback")
 	virtual void ClientClearCameraLensEffects();
+	virtual void ClientClearCameraLensEffects_Implementation();
 
 	/** Stop camera animation on client. */
-	UFUNCTION(Reliable, Client)
+	UFUNCTION(Reliable, Client="ClientStopCameraAnim_Implementation")
 	void ClientStopCameraAnim(class UCameraAnim* AnimToStop);
+	virtual void ClientStopCameraAnim_Implementation(UCameraAnim* AnimToStop);
 
 	/** Stop camera shake on client.  */
-	UFUNCTION(reliable, client, BlueprintCallable, Category="Game|Feedback")
+	UFUNCTION(reliable, client="ClientStopCameraShake_Implementation", BlueprintCallable, Category="Game|Feedback")
 	void ClientStopCameraShake(TSubclassOf<class UCameraShake> Shake);
+	virtual void ClientStopCameraShake_Implementation(TSubclassOf<UCameraShake>  Shake);
 
 	/** 
 	 * Play a force feedback pattern on the player's controller
@@ -805,16 +847,18 @@ public:
 	 * @param	bLooping				Whether the pattern should be played repeatedly or be a single one shot
 	 * @param	Tag						A tag that allows stopping of an effect.  If another effect with this Tag is playing, it will be stopped and replaced
 	 */
-	UFUNCTION(unreliable, client, BlueprintCallable, Category="Game|Feedback")
+	UFUNCTION(unreliable, client="ClientPlayForceFeedback_Implementation", BlueprintCallable, Category="Game|Feedback")
 	void ClientPlayForceFeedback(class UForceFeedbackEffect* ForceFeedbackEffect, bool bLooping, FName Tag);
+	virtual void ClientPlayForceFeedback_Implementation(UForceFeedbackEffect* ForceFeedbackEffect, bool bLooping, FName Tag);
 
 	/** 
 	 * Stops a playing force feedback pattern
 	 * @param	ForceFeedbackEffect		If set only patterns from that effect will be stopped
 	 * @param	Tag						If not none only the pattern with this tag will be stopped
 	 */
-	UFUNCTION(reliable, client, BlueprintCallable, Category="Game|Feedback")
+	UFUNCTION(reliable, client="ClientStopForceFeedback_Implementation", BlueprintCallable, Category="Game|Feedback")
 	void ClientStopForceFeedback(class UForceFeedbackEffect* ForceFeedbackEffect, FName Tag);
+	virtual void ClientStopForceFeedback_Implementation(UForceFeedbackEffect* ForceFeedbackEffect, FName Tag);
 
 	/** 
 	 * Latent action that controls the playing of force feedback 
@@ -855,8 +899,9 @@ public:
 	 * @param MapPackageGuid	The GUID of the map package to travel to - this is used to find the file when it has been autodownloaded,
 	 * 							so it is only needed for clients
 	 */
-	UFUNCTION(Reliable, Client)
+	UFUNCTION(Reliable, Client="ClientTravelInternal_Implementation")
 	void ClientTravelInternal(const FString& URL, enum ETravelType TravelType, bool bSeamless = false, FGuid MapPackageGuid = FGuid());
+	virtual void ClientTravelInternal_Implementation(const FString& URL, ETravelType TravelType, bool bSeamless, FGuid MapPackageGuid);
 
 	/**
 	 * Replicated Update streaming status
@@ -866,73 +911,97 @@ public:
 	 * @param bNewShouldBlockOnLoad - Whether we want to force a blocking load
 	 * @param LODIndex				- Current LOD index for a streaming level
 	 */
-	UFUNCTION(Reliable, Client)
+	UFUNCTION(Reliable, Client="ClientUpdateLevelStreamingStatus_Implementation")
 	void ClientUpdateLevelStreamingStatus(FName PackageName, bool bNewShouldBeLoaded, bool bNewShouldBeVisible, bool bNewShouldBlockOnLoad, int32 LODIndex);
+	virtual void ClientUpdateLevelStreamingStatus_Implementation(FName PackageName, bool bNewShouldBeLoaded, bool bNewShouldBeVisible, bool bNewShouldBlockOnLoad, int32 LODIndex);
 
 	/** Notify client they were kicked from the server */
-	UFUNCTION(Reliable, Client)
+	UFUNCTION(Reliable, Client="ClientWasKicked_Implementation")
 	void ClientWasKicked(const FText& KickReason);
+	virtual void ClientWasKicked_Implementation(const FText& KickReason);
 
 	/** Assign Pawn to player, but avoid calling ClientRestart if we have already accepted this pawn */
-	UFUNCTION(Reliable, Client)
+	UFUNCTION(Reliable, Client="ClientRetryClientRestart_Implementation")
 	void ClientRetryClientRestart(class APawn* NewPawn);
+	virtual void ClientRetryClientRestart_Implementation(APawn* NewPawn);
 
 	/** Call ClientRetryClientRestart, but only if the current pawn is not the currently acknowledged pawn (and throttled to avoid saturating the network). */
 	virtual void SafeRetryClientRestart();
 
 	/** send client localized message id */
-	UFUNCTION(Reliable, Client)
+	UFUNCTION(Reliable, Client="ClientReceiveLocalizedMessage_Implementation")
 	void ClientReceiveLocalizedMessage(TSubclassOf<ULocalMessage> Message, int32 Switch = 0, class APlayerState* RelatedPlayerState_1 = NULL, class APlayerState* RelatedPlayerState_2 = NULL, class UObject* OptionalObject = NULL);
+	virtual void ClientReceiveLocalizedMessage_Implementation(TSubclassOf<ULocalMessage>  Message, int32 Switch, APlayerState* RelatedPlayerState_1, APlayerState* RelatedPlayerState_2, UObject* OptionalObject);
 
 	/** acknowledge possession of pawn */
-	UFUNCTION(reliable, server, WithValidation)
+	UFUNCTION(reliable, server="ServerAcknowledgePossession_Implementation", WithValidation="ServerAcknowledgePossession_Validate")
 	void ServerAcknowledgePossession(class APawn* P);
+	virtual void ServerAcknowledgePossession_Implementation(APawn* P);
+	virtual bool ServerAcknowledgePossession_Validate(APawn* );
 
 	/** change mode of camera */
-	UFUNCTION(reliable, server, WithValidation)
+	UFUNCTION(reliable, server="ServerCamera_Implementation", WithValidation="ServerCamera_Validate")
 	void ServerCamera(FName NewMode);
+	virtual void ServerCamera_Implementation(FName NewMode);
+	virtual bool ServerCamera_Validate(FName );
 
 	/** Change name of server */
-	UFUNCTION(reliable, server, WithValidation)
+	UFUNCTION(reliable, server="ServerChangeName_Implementation", WithValidation="ServerChangeName_Validate")
 	void ServerChangeName(const FString& S);
+	virtual void ServerChangeName_Implementation(const FString& S);
+	virtual bool ServerChangeName_Validate(const FString& );
 
 	/** 
 	 * Called to notify the server when the client has loaded a new world via seamless traveling
 	 * @param WorldPackageName the name of the world package that was loaded
 	 */
-	UFUNCTION(reliable, server, WithValidation, SealedEvent)
+	UFUNCTION(reliable, server="ServerNotifyLoadedWorld_Implementation", WithValidation="ServerNotifyLoadedWorld_Validate", SealedEvent)
 	void ServerNotifyLoadedWorld(FName WorldPackageName);
+	void ServerNotifyLoadedWorld_Implementation(FName WorldPackageName);
+	bool ServerNotifyLoadedWorld_Validate(FName );
 
 	/** Replicate pause request to the server */
-	UFUNCTION(reliable, server, WithValidation)
+	UFUNCTION(reliable, server="ServerPause_Implementation", WithValidation="ServerPause_Validate")
 	void ServerPause();
+	virtual void ServerPause_Implementation();
+	virtual bool ServerPause_Validate();
 
 	/** Attempts to restart this player, generally called from the client upon respawn request. */
-	UFUNCTION(reliable, server, WithValidation)
+	UFUNCTION(reliable, server="ServerRestartPlayer_Implementation", WithValidation="ServerRestartPlayer_Validate")
 	void ServerRestartPlayer();
+	virtual void ServerRestartPlayer_Implementation();
+	virtual bool ServerRestartPlayer_Validate();
 
 	/** When spectating, pings the server to make sure spectating should continue. */
-	UFUNCTION(unreliable, server, WithValidation)
+	UFUNCTION(unreliable, server="ServerSetSpectatorLocation_Implementation", WithValidation="ServerSetSpectatorLocation_Validate")
 	void ServerSetSpectatorLocation(FVector NewLoc);
+	virtual void ServerSetSpectatorLocation_Implementation(FVector NewLoc);
+	virtual bool ServerSetSpectatorLocation_Validate(FVector );
 
 	/** Calls ServerSetSpectatorLocation but throttles it to reduce bandwidth and only calls it when necessary. */
 	void SafeServerUpdateSpectatorState();
 
 	/** Tells the server to make sure the possessed pawn is in sync with the client. */
-	UFUNCTION(unreliable, server, WithValidation)
+	UFUNCTION(unreliable, server="ServerCheckClientPossession_Implementation", WithValidation="ServerCheckClientPossession_Validate")
 	void ServerCheckClientPossession();
+	virtual void ServerCheckClientPossession_Implementation();
+	virtual bool ServerCheckClientPossession_Validate();
 
 	/** Call ServerCheckClientPossession on the server, but only if the current pawn is not the acknowledged pawn (and throttled to avoid saturating the network). */
 	virtual void SafeServerCheckClientPossession();
 
 	/** Notifies the server that the client has ticked gameplay code, and should no longer get the extended "still loading" timeout grace period */
-	UFUNCTION(reliable, server, WithValidation)
+	UFUNCTION(reliable, server="ServerShortTimeout_Implementation", WithValidation="ServerShortTimeout_Validate")
 	void ServerShortTimeout();
+	virtual void ServerShortTimeout_Implementation();
+	virtual bool ServerShortTimeout_Validate();
 
 	/** If PlayerCamera.bUseClientSideCameraUpdates is set, client will replicate camera positions to the server. */
 	// @TODO - combine pitch/yaw into one int, maybe also send location compressed
-	UFUNCTION(unreliable, server, WithValidation)
+	UFUNCTION(unreliable, server="ServerUpdateCamera_Implementation", WithValidation="ServerUpdateCamera_Validate")
 	void ServerUpdateCamera(FVector_NetQuantize CamLoc, int32 CamPitchAndYaw);
+	virtual void ServerUpdateCamera_Implementation(FVector_NetQuantize CamLoc, int32 CamPitchAndYaw);
+	virtual bool ServerUpdateCamera_Validate(FVector_NetQuantize , int32 );
 
 	/** 
 	 * Called when the client adds/removes a streamed level
@@ -940,32 +1009,45 @@ public:
 	 * Actors the client has not initialized
 	 * @param PackageName the name of the package for the level whose status changed
 	 */
-	UFUNCTION(reliable, server, WithValidation, SealedEvent)
+	UFUNCTION(reliable, server="ServerUpdateLevelVisibility_Implementation", WithValidation="ServerUpdateLevelVisibility_Validate", SealedEvent)
 	void ServerUpdateLevelVisibility(FName PackageName, bool bIsVisible);
+	void ServerUpdateLevelVisibility_Implementation(FName PackageName, bool bIsVisible);
+	bool ServerUpdateLevelVisibility_Validate(FName , bool );
 
 	/** Used by client to request server to confirm current viewtarget (server will respond with ClientSetViewTarget() ). */
-	UFUNCTION(reliable, server, WithValidation)
+	UFUNCTION(reliable, server="ServerVerifyViewTarget_Implementation", WithValidation="ServerVerifyViewTarget_Validate")
 	void ServerVerifyViewTarget();
+	virtual void ServerVerifyViewTarget_Implementation();
+	virtual bool ServerVerifyViewTarget_Validate();
 
 	/** Move camera to next player on round ended or spectating*/
-	UFUNCTION(unreliable, server, WithValidation)
+	UFUNCTION(unreliable, server="ServerViewNextPlayer_Implementation", WithValidation="ServerViewNextPlayer_Validate")
 	void ServerViewNextPlayer();
+	virtual void ServerViewNextPlayer_Implementation();
+	virtual bool ServerViewNextPlayer_Validate();
 
 	/** Move camera to previous player on round ended or spectating */
-	UFUNCTION(unreliable, server, WithValidation)
+	UFUNCTION(unreliable, server="ServerViewPrevPlayer_Implementation", WithValidation="ServerViewPrevPlayer_Validate")
 	void ServerViewPrevPlayer();
+	virtual void ServerViewPrevPlayer_Implementation();
+	virtual bool ServerViewPrevPlayer_Validate();
 
 	/** Move camera to current user */
-	UFUNCTION(unreliable, server, WithValidation)
+	UFUNCTION(unreliable, server="ServerViewSelf_Implementation", WithValidation="ServerViewSelf_Validate")
 	void ServerViewSelf(struct FViewTargetTransitionParams TransitionParams = FViewTargetTransitionParams());
+	virtual void ServerViewSelf_Implementation(FViewTargetTransitionParams TransitionParams);
+	virtual bool ServerViewSelf_Validate(FViewTargetTransitionParams );
 
 	/** @todo document */
-	UFUNCTION(Reliable, Client)
+	UFUNCTION(Reliable, Client="ClientTeamMessage_Implementation")
 	void ClientTeamMessage(class APlayerState* SenderPlayerState, const FString& S, FName Type, float MsgLifeTime = 0);
+	virtual void ClientTeamMessage_Implementation(APlayerState* SenderPlayerState, const FString& S, FName Type, float MsgLifeTime);
 
 	/** Used by UGameplayDebuggingControllerComponent to replicate messages for AI debugging in network games. */
-	UFUNCTION(reliable, server, WithValidation)
+	UFUNCTION(reliable, server="ServerToggleAILogging_Implementation", WithValidation="ServerToggleAILogging_Validate")
 	void ServerToggleAILogging();
+	virtual void ServerToggleAILogging_Implementation();
+	virtual bool ServerToggleAILogging_Validate();
 
 	/** Add Pitch (look up) input */
 	UFUNCTION(BlueprintCallable, Category="Game|Player", meta=(Keywords="up down"))
@@ -1412,8 +1494,9 @@ public:
 	/**
 	 * Notify from server that Visual Logger is recording, to show that information on client about possible performance issues 
 	 */
-	UFUNCTION(Reliable, Client)
+	UFUNCTION(Reliable, Client="OnServerStartedVisualLogger_Implementation")
 	void OnServerStartedVisualLogger(bool bIsLogging);
+	virtual void OnServerStartedVisualLogger_Implementation(bool bIsLogging);
 
 protected:
 
diff --git a/Engine/Source/Runtime/Engine/Classes/GameFramework/PlayerInput.h b/Engine/Source/Runtime/Engine/Classes/GameFramework/PlayerInput.h
--- a/Engine/Source/Runtime/Engine/Classes/GameFramework/PlayerInput.h
+++ b/Engine/Source/Runtime/Engine/Classes/GameFramework/PlayerInput.h
@@ -19,7 +19,7 @@ ENGINE_API DECLARE_LOG_CATEGORY_EXTERN(LogInput, Log, All);
 USTRUCT()
 struct FKeyBind
 {
-	GENERATED_USTRUCT_BODY()
+	GENERATED_BODY()
 
 	/** The key to be bound to the command */
 	UPROPERTY(config)
@@ -70,7 +70,7 @@ struct FKeyBind
 USTRUCT()
 struct FInputAxisProperties
 {
-	GENERATED_USTRUCT_BODY()
+	GENERATED_BODY()
 
 	/** What the dead zone of the axis is.  For control axes such as analog sticks. */
 	UPROPERTY(EditAnywhere, Category="Input")
@@ -101,7 +101,7 @@ struct FInputAxisProperties
 USTRUCT()
 struct FInputAxisConfigEntry
 {
-	GENERATED_USTRUCT_BODY()
+	GENERATED_BODY()
 
 	/** Axis Key these properties apply to */
 	UPROPERTY(VisibleAnywhere, Category="Input")
@@ -120,7 +120,7 @@ struct FInputAxisConfigEntry
 USTRUCT()
 struct FInputActionKeyMapping
 {
-	GENERATED_USTRUCT_BODY()
+	GENERATED_BODY()
 
 	/** Friendly name of action, e.g "jump" */
 	UPROPERTY(EditAnywhere, Category="Input")
@@ -188,7 +188,7 @@ struct FInputActionKeyMapping
 USTRUCT()
 struct FInputAxisKeyMapping
 {
-	GENERATED_USTRUCT_BODY()
+	GENERATED_BODY()
 
 	/** Friendly name of axis, e.g "MoveForward" */
 	UPROPERTY(EditAnywhere, Category="Input")
@@ -247,7 +247,9 @@ struct FInputAxisKeyMapping
 UCLASS(Within=PlayerController, config=Input, transient)
 class ENGINE_API UPlayerInput : public UObject
 {
-	GENERATED_UCLASS_BODY()
+	GENERATED_BODY()
+public:
+	UPlayerInput(const FObjectInitializer& ObjectInitializer = FObjectInitializer::Get());
 
 public:
 	// NOTE: These touch vectors are calculated and set directly, they do not go through the .ini Bindings
diff --git a/Engine/Source/Runtime/Engine/Classes/GameFramework/PlayerStart.h b/Engine/Source/Runtime/Engine/Classes/GameFramework/PlayerStart.h
--- a/Engine/Source/Runtime/Engine/Classes/GameFramework/PlayerStart.h
+++ b/Engine/Source/Runtime/Engine/Classes/GameFramework/PlayerStart.h
@@ -12,9 +12,11 @@
 UCLASS(ClassGroup=Common, hidecategories=Collision)
 class ENGINE_API APlayerStart : public ANavigationObjectBase
 {
+	GENERATED_BODY()
+public:
+	APlayerStart(const FObjectInitializer& ObjectInitializer = FObjectInitializer::Get());
 	/** To take more control over PlayerStart selection, you can override the virtual AGameMode::FindPlayerStart and AGameMode::ChoosePlayerStart function */
 
-	GENERATED_UCLASS_BODY()
 
 	/** Used when searching for which playerstart to use. */
 	UPROPERTY(EditAnywhere, BlueprintReadWrite, Category=Object)
diff --git a/Engine/Source/Runtime/Engine/Classes/GameFramework/PlayerState.h b/Engine/Source/Runtime/Engine/Classes/GameFramework/PlayerState.h
--- a/Engine/Source/Runtime/Engine/Classes/GameFramework/PlayerState.h
+++ b/Engine/Source/Runtime/Engine/Classes/GameFramework/PlayerState.h
@@ -31,7 +31,9 @@ struct PingAvgData
 UCLASS(BlueprintType, Blueprintable, notplaceable)
 class ENGINE_API APlayerState : public AInfo
 {
-	GENERATED_UCLASS_BODY()
+	GENERATED_BODY()
+public:
+	APlayerState(const FObjectInitializer& ObjectInitializer = FObjectInitializer::Get());
 
 	/** Player's current score. */
 	UPROPERTY(replicatedUsing=OnRep_Score, BlueprintReadOnly, Category=PlayerState)
diff --git a/Engine/Source/Runtime/Engine/Classes/GameFramework/ProjectileMovementComponent.h b/Engine/Source/Runtime/Engine/Classes/GameFramework/ProjectileMovementComponent.h
--- a/Engine/Source/Runtime/Engine/Classes/GameFramework/ProjectileMovementComponent.h
+++ b/Engine/Source/Runtime/Engine/Classes/GameFramework/ProjectileMovementComponent.h
@@ -18,7 +18,9 @@
 UCLASS(ClassGroup=Movement, meta=(BlueprintSpawnableComponent), ShowCategories=(Velocity))
 class ENGINE_API UProjectileMovementComponent : public UMovementComponent
 {
-	GENERATED_UCLASS_BODY()
+	GENERATED_BODY()
+public:
+	UProjectileMovementComponent(const FObjectInitializer& ObjectInitializer = FObjectInitializer::Get());
 
 	DECLARE_DYNAMIC_MULTICAST_DELEGATE_TwoParams( FOnProjectileBounceDelegate, const FHitResult&, ImpactResult, const FVector&, ImpactVelocity );
 	DECLARE_DYNAMIC_MULTICAST_DELEGATE_OneParam( FOnProjectileStopDelegate, const FHitResult&, ImpactResult );
diff --git a/Engine/Source/Runtime/Engine/Classes/GameFramework/RotatingMovementComponent.h b/Engine/Source/Runtime/Engine/Classes/GameFramework/RotatingMovementComponent.h
--- a/Engine/Source/Runtime/Engine/Classes/GameFramework/RotatingMovementComponent.h
+++ b/Engine/Source/Runtime/Engine/Classes/GameFramework/RotatingMovementComponent.h
@@ -12,7 +12,9 @@
 UCLASS(ClassGroup=Movement, meta=(BlueprintSpawnableComponent), HideCategories=(Velocity))
 class ENGINE_API URotatingMovementComponent : public UMovementComponent
 {
-	GENERATED_UCLASS_BODY()
+	GENERATED_BODY()
+public:
+	URotatingMovementComponent(const FObjectInitializer& ObjectInitializer = FObjectInitializer::Get());
 
 	/** How fast to update roll/pitch/yaw of the component we update. */
 	UPROPERTY(EditAnywhere, BlueprintReadWrite, Category=RotatingComponent)
diff --git a/Engine/Source/Runtime/Engine/Classes/GameFramework/SaveGame.h b/Engine/Source/Runtime/Engine/Classes/GameFramework/SaveGame.h
--- a/Engine/Source/Runtime/Engine/Classes/GameFramework/SaveGame.h
+++ b/Engine/Source/Runtime/Engine/Classes/GameFramework/SaveGame.h
@@ -14,6 +14,9 @@
 UCLASS(abstract, Blueprintable, BlueprintType)
 class ENGINE_API USaveGame : public UObject
 {
+	GENERATED_BODY()
+public:
+	USaveGame(const FObjectInitializer& ObjectInitializer = FObjectInitializer::Get());
 	/**
 	 *	@see UGameplayStatics::CreateSaveGameObject
 	 *	@see UGameplayStatics::SaveGameToSlot
@@ -22,5 +25,4 @@ class ENGINE_API USaveGame : public UObject
 	 *	@see UGameplayStatics::DeleteGameInSlot
 	 */
 
-	GENERATED_UCLASS_BODY()
 };
diff --git a/Engine/Source/Runtime/Engine/Classes/GameFramework/SpectatorPawn.h b/Engine/Source/Runtime/Engine/Classes/GameFramework/SpectatorPawn.h
--- a/Engine/Source/Runtime/Engine/Classes/GameFramework/SpectatorPawn.h
+++ b/Engine/Source/Runtime/Engine/Classes/GameFramework/SpectatorPawn.h
@@ -13,7 +13,9 @@
 UCLASS(config=Game, Blueprintable, BlueprintType, notplaceable)
 class ENGINE_API ASpectatorPawn : public ADefaultPawn
 {
-	GENERATED_UCLASS_BODY()
+	GENERATED_BODY()
+public:
+	ASpectatorPawn(const FObjectInitializer& ObjectInitializer = FObjectInitializer::Get());
 
 	// Begin Pawn overrides
 	/** Overridden to avoid changing network role. If subclasses want networked behavior, call the Pawn::PossessedBy() instead. */
diff --git a/Engine/Source/Runtime/Engine/Classes/GameFramework/SpectatorPawnMovement.h b/Engine/Source/Runtime/Engine/Classes/GameFramework/SpectatorPawnMovement.h
--- a/Engine/Source/Runtime/Engine/Classes/GameFramework/SpectatorPawnMovement.h
+++ b/Engine/Source/Runtime/Engine/Classes/GameFramework/SpectatorPawnMovement.h
@@ -11,7 +11,9 @@
 UCLASS()
 class ENGINE_API USpectatorPawnMovement : public UFloatingPawnMovement
 {
-	GENERATED_UCLASS_BODY()
+	GENERATED_BODY()
+public:
+	USpectatorPawnMovement(const FObjectInitializer& ObjectInitializer = FObjectInitializer::Get());
 
 	/** If true, component moves at full speed no matter the time dilation. Default is false. */
 	UPROPERTY()
diff --git a/Engine/Source/Runtime/Engine/Classes/GameFramework/SpringArmComponent.h b/Engine/Source/Runtime/Engine/Classes/GameFramework/SpringArmComponent.h
--- a/Engine/Source/Runtime/Engine/Classes/GameFramework/SpringArmComponent.h
+++ b/Engine/Source/Runtime/Engine/Classes/GameFramework/SpringArmComponent.h
@@ -14,7 +14,9 @@
 UCLASS(ClassGroup=Camera, meta=(BlueprintSpawnableComponent), hideCategories=(Mobility))
 class ENGINE_API USpringArmComponent : public USceneComponent
 {
-	GENERATED_UCLASS_BODY()
+	GENERATED_BODY()
+public:
+	USpringArmComponent(const FObjectInitializer& ObjectInitializer = FObjectInitializer::Get());
 
 	/** Natural length of the spring arm when there are no collisions */
 	UPROPERTY(EditAnywhere, BlueprintReadWrite, Category=Camera)
diff --git a/Engine/Source/Runtime/Engine/Classes/GameFramework/TouchInterface.h b/Engine/Source/Runtime/Engine/Classes/GameFramework/TouchInterface.h
--- a/Engine/Source/Runtime/Engine/Classes/GameFramework/TouchInterface.h
+++ b/Engine/Source/Runtime/Engine/Classes/GameFramework/TouchInterface.h
@@ -9,7 +9,7 @@ class SVirtualJoystick;
 USTRUCT()
 struct FTouchInputControl
 {
-	GENERATED_USTRUCT_BODY()
+	GENERATED_BODY()
 
 	// basically mirroring SVirtualJoystick::FControlInfo but as an editable class
 	UPROPERTY(EditAnywhere, Category="Control", meta=(ToolTip="For sticks, this is the Thumb"))
@@ -44,7 +44,9 @@ struct FTouchInputControl
 UCLASS(Blueprintable, BlueprintType)
 class ENGINE_API UTouchInterface : public UObject
 {
-	GENERATED_UCLASS_BODY()
+	GENERATED_BODY()
+public:
+	UTouchInterface(const FObjectInitializer& ObjectInitializer = FObjectInitializer::Get());
 
 	UPROPERTY(EditAnywhere, Category="TouchInterface")
 	TArray<FTouchInputControl> Controls;
diff --git a/Engine/Source/Runtime/Engine/Classes/GameFramework/Volume.h b/Engine/Source/Runtime/Engine/Classes/GameFramework/Volume.h
--- a/Engine/Source/Runtime/Engine/Classes/GameFramework/Volume.h
+++ b/Engine/Source/Runtime/Engine/Classes/GameFramework/Volume.h
@@ -13,7 +13,9 @@ DECLARE_LOG_CATEGORY_EXTERN(LogVolume, Log, All);
 UCLASS(showcategories=Collision, hidecategories=(Brush, Physics), abstract, ConversionRoot)
 class ENGINE_API AVolume : public ABrush
 {
-	GENERATED_UCLASS_BODY()
+	GENERATED_BODY()
+public:
+	AVolume(const FObjectInitializer& ObjectInitializer = FObjectInitializer::Get());
 
 #if WITH_EDITOR
 	/** Delegate used for notifications when a volumes initial shape changes */
diff --git a/Engine/Source/Runtime/Engine/Classes/GameFramework/WheeledVehicle.h b/Engine/Source/Runtime/Engine/Classes/GameFramework/WheeledVehicle.h
--- a/Engine/Source/Runtime/Engine/Classes/GameFramework/WheeledVehicle.h
+++ b/Engine/Source/Runtime/Engine/Classes/GameFramework/WheeledVehicle.h
@@ -19,7 +19,9 @@ class FDebugDisplayInfo;
 UCLASS(abstract, config=Game, BlueprintType)
 class ENGINE_API AWheeledVehicle : public APawn
 {
-	GENERATED_UCLASS_BODY()
+	GENERATED_BODY()
+public:
+	AWheeledVehicle(const FObjectInitializer& ObjectInitializer = FObjectInitializer::Get());
 
 private_subobject:
 	/**  The main skeletal mesh associated with this Vehicle */
diff --git a/Engine/Source/Runtime/Engine/Classes/GameFramework/WorldSettings.h b/Engine/Source/Runtime/Engine/Classes/GameFramework/WorldSettings.h
--- a/Engine/Source/Runtime/Engine/Classes/GameFramework/WorldSettings.h
+++ b/Engine/Source/Runtime/Engine/Classes/GameFramework/WorldSettings.h
@@ -18,7 +18,7 @@ enum EVisibilityAggressiveness
 USTRUCT()
 struct FGameModePrefix
 {
-	GENERATED_USTRUCT_BODY()
+	GENERATED_BODY()
 
 	/** map prefix, e.g. "DM" */
 	UPROPERTY()
@@ -32,7 +32,7 @@ struct FGameModePrefix
 USTRUCT()
 struct FLightmassWorldInfoSettings
 {
-	GENERATED_USTRUCT_BODY()
+	GENERATED_BODY()
 
 	/** 
 	 * Warning: Setting this to less than 1 will greatly increase build times!
@@ -162,7 +162,7 @@ struct FLightmassWorldInfoSettings
 USTRUCT()
 struct ENGINE_API FNetViewer
 {
-	GENERATED_USTRUCT_BODY()
+	GENERATED_BODY()
 
 	/** The player controller corresponding to this view */
 	UPROPERTY()
@@ -195,7 +195,7 @@ struct ENGINE_API FNetViewer
 USTRUCT()
 struct ENGINE_API FHierarchicalSimplification
 {
-	GENERATED_USTRUCT_BODY()
+	GENERATED_BODY()
 
 	/** If this is true, it will simplify mesh. However it is slower. If false, it will just merge actors but not simplify */
 	UPROPERTY(Category=FHierarchicalSimplification, EditAnywhere)
@@ -233,7 +233,9 @@ struct ENGINE_API FHierarchicalSimplification
 UCLASS(config=game, hidecategories=(Actor, Advanced, Display, Events, Object, Attachment, Info, Input, Blueprint, Layers, Tags, Replication), showcategories=("Input|MouseInput", "Input|TouchInput"), notplaceable)
 class ENGINE_API AWorldSettings : public AInfo, public IInterface_AssetUserData
 {
-	GENERATED_UCLASS_BODY()
+	GENERATED_BODY()
+public:
+	AWorldSettings(const FObjectInitializer& ObjectInitializer = FObjectInitializer::Get());
 
 	/** DEFAULT BASIC PHYSICS SETTINGS **/
 
