Enabled UHT to digest GENERATED_BODY instead of GENERATED_UCLASS_BODY, GENERATED_USTRUCT_BODY, GENERATED_UINTERFACE_BODY or GENERATED_IINTERFACE_BODY, changed every occurence to the new syntax and fix

diff --git a/Engine/Source/Runtime/Engine/Classes/Interfaces/Interface_AssetUserData.h b/Engine/Source/Runtime/Engine/Classes/Interfaces/Interface_AssetUserData.h
--- a/Engine/Source/Runtime/Engine/Classes/Interfaces/Interface_AssetUserData.h
+++ b/Engine/Source/Runtime/Engine/Classes/Interfaces/Interface_AssetUserData.h
@@ -9,12 +9,15 @@ class UAssetUserData;
 UINTERFACE(MinimalApi, meta=(CannotImplementInterfaceInBlueprint))
 class UInterface_AssetUserData : public UInterface
 {
-	GENERATED_UINTERFACE_BODY()
+	GENERATED_BODY()
+public:
+	ENGINE_API UInterface_AssetUserData(const FObjectInitializer& ObjectInitializer = FObjectInitializer::Get());
 };
 
 class IInterface_AssetUserData
 {
-	GENERATED_IINTERFACE_BODY()
+	GENERATED_BODY()
+public:
 
 	virtual void AddAssetUserData(UAssetUserData* InUserData) {}
 
diff --git a/Engine/Source/Runtime/Engine/Classes/Interfaces/Interface_CollisionDataProvider.h b/Engine/Source/Runtime/Engine/Classes/Interfaces/Interface_CollisionDataProvider.h
--- a/Engine/Source/Runtime/Engine/Classes/Interfaces/Interface_CollisionDataProvider.h
+++ b/Engine/Source/Runtime/Engine/Classes/Interfaces/Interface_CollisionDataProvider.h
@@ -12,7 +12,7 @@
 USTRUCT()
 struct FTriIndices
 {
-	GENERATED_USTRUCT_BODY()
+	GENERATED_BODY()
 
 	UPROPERTY()
 	int32 v0;
@@ -37,7 +37,7 @@ struct FTriIndices
 USTRUCT()
 struct FTriMeshCollisionData
 {
-	GENERATED_USTRUCT_BODY()
+	GENERATED_BODY()
 
 	/** Array of vertices included in the triangle mesh */
 	UPROPERTY(transient)
@@ -64,12 +64,15 @@ struct FTriMeshCollisionData
 UINTERFACE(meta=(CannotImplementInterfaceInBlueprint))
 class UInterface_CollisionDataProvider : public UInterface
 {
-	GENERATED_UINTERFACE_BODY()
+	GENERATED_BODY()
+public:
+	UInterface_CollisionDataProvider(const FObjectInitializer& ObjectInitializer = FObjectInitializer::Get());
 };
 
 class IInterface_CollisionDataProvider
 {
-	GENERATED_IINTERFACE_BODY()
+	GENERATED_BODY()
+public:
 
 
 	/**	 Interface for retrieving triangle mesh collision data from the implementing object 
diff --git a/Engine/Source/Runtime/Engine/Classes/Interfaces/Interface_PostProcessVolume.h b/Engine/Source/Runtime/Engine/Classes/Interfaces/Interface_PostProcessVolume.h
--- a/Engine/Source/Runtime/Engine/Classes/Interfaces/Interface_PostProcessVolume.h
+++ b/Engine/Source/Runtime/Engine/Classes/Interfaces/Interface_PostProcessVolume.h
@@ -7,7 +7,9 @@
 UINTERFACE()
 class UInterface_PostProcessVolume : public UInterface
 {
-	GENERATED_UINTERFACE_BODY()
+	GENERATED_BODY()
+public:
+	UInterface_PostProcessVolume(const FObjectInitializer& ObjectInitializer = FObjectInitializer::Get());
 };
 
 struct FPostProcessVolumeProperties
@@ -22,7 +24,8 @@ struct FPostProcessVolumeProperties
 
 class IInterface_PostProcessVolume
 {
-	GENERATED_IINTERFACE_BODY()
+	GENERATED_BODY()
+public:
 
 	ENGINE_API virtual bool EncompassesPoint(FVector Point, float SphereRadius/*=0.f*/, float* OutDistanceToPoint) = 0;
 	ENGINE_API virtual FPostProcessVolumeProperties GetProperties() const = 0;
diff --git a/Engine/Source/Runtime/Engine/Classes/Interfaces/NetworkPredictionInterface.h b/Engine/Source/Runtime/Engine/Classes/Interfaces/NetworkPredictionInterface.h
--- a/Engine/Source/Runtime/Engine/Classes/Interfaces/NetworkPredictionInterface.h
+++ b/Engine/Source/Runtime/Engine/Classes/Interfaces/NetworkPredictionInterface.h
@@ -12,12 +12,15 @@
 UINTERFACE(MinimalAPI, meta=(CannotImplementInterfaceInBlueprint))
 class UNetworkPredictionInterface : public UInterface
 {
-	GENERATED_UINTERFACE_BODY()
+	GENERATED_BODY()
+public:
+	ENGINE_API UNetworkPredictionInterface(const FObjectInitializer& ObjectInitializer = FObjectInitializer::Get());
 };
 
 class ENGINE_API INetworkPredictionInterface
 {
-	GENERATED_IINTERFACE_BODY()
+	GENERATED_BODY()
+public:
 
 	//--------------------------------
 	// Server hooks
diff --git a/Engine/Source/Runtime/Engine/Classes/Kismet/BlueprintFunctionLibrary.h b/Engine/Source/Runtime/Engine/Classes/Kismet/BlueprintFunctionLibrary.h
--- a/Engine/Source/Runtime/Engine/Classes/Kismet/BlueprintFunctionLibrary.h
+++ b/Engine/Source/Runtime/Engine/Classes/Kismet/BlueprintFunctionLibrary.h
@@ -8,7 +8,9 @@
 UCLASS(Abstract, MinimalAPI)
 class UBlueprintFunctionLibrary : public UObject
 {
-	GENERATED_UCLASS_BODY()
+	GENERATED_BODY()
+public:
+	ENGINE_API UBlueprintFunctionLibrary(const FObjectInitializer& ObjectInitializer = FObjectInitializer::Get());
 
 	// UObject interface
 	ENGINE_API virtual int32 GetFunctionCallspace(UFunction* Function, void* Parms, FFrame* Stack) override;
diff --git a/Engine/Source/Runtime/Engine/Classes/Kismet/DataTableFunctionLibrary.h b/Engine/Source/Runtime/Engine/Classes/Kismet/DataTableFunctionLibrary.h
--- a/Engine/Source/Runtime/Engine/Classes/Kismet/DataTableFunctionLibrary.h
+++ b/Engine/Source/Runtime/Engine/Classes/Kismet/DataTableFunctionLibrary.h
@@ -23,7 +23,9 @@ namespace EEvaluateCurveTableResult
 UCLASS(MinimalAPI)
 class UDataTableFunctionLibrary : public UBlueprintFunctionLibrary
 {
-	GENERATED_UCLASS_BODY()
+	GENERATED_BODY()
+public:
+	ENGINE_API UDataTableFunctionLibrary(const FObjectInitializer& ObjectInitializer = FObjectInitializer::Get());
 
 	UFUNCTION(BlueprintCallable, Category = "DataTable", meta = (ExpandEnumAsExecs="OutResult", DataTablePin="CurveTable"))
 	static void EvaluateCurveTableRow(UCurveTable* CurveTable, FName RowName, float InXY, TEnumAsByte<EEvaluateCurveTableResult::Type>& OutResult, float& OutXY);
diff --git a/Engine/Source/Runtime/Engine/Classes/Kismet/GameplayStatics.h b/Engine/Source/Runtime/Engine/Classes/Kismet/GameplayStatics.h
--- a/Engine/Source/Runtime/Engine/Classes/Kismet/GameplayStatics.h
+++ b/Engine/Source/Runtime/Engine/Classes/Kismet/GameplayStatics.h
@@ -1,4 +1,4 @@
-﻿// Copyright 1998-2015 Epic Games, Inc. All Rights Reserved.
+// Copyright 1998-2015 Epic Games, Inc. All Rights Reserved.
 
 #pragma once
 #include "Engine/LatentActionManager.h"
@@ -28,7 +28,9 @@ namespace ESuggestProjVelocityTraceOption
 UCLASS()
 class ENGINE_API UGameplayStatics : public UBlueprintFunctionLibrary
 {
-	GENERATED_UCLASS_BODY()
+	GENERATED_BODY()
+public:
+	UGameplayStatics(const FObjectInitializer& ObjectInitializer = FObjectInitializer::Get());
 
 	// --- Spawning functions ------------------------------
 
diff --git a/Engine/Source/Runtime/Engine/Classes/Kismet/HeadMountedDisplayFunctionLibrary.h b/Engine/Source/Runtime/Engine/Classes/Kismet/HeadMountedDisplayFunctionLibrary.h
--- a/Engine/Source/Runtime/Engine/Classes/Kismet/HeadMountedDisplayFunctionLibrary.h
+++ b/Engine/Source/Runtime/Engine/Classes/Kismet/HeadMountedDisplayFunctionLibrary.h
@@ -17,7 +17,9 @@ namespace EOrientPositionSelector
 UCLASS()
 class UHeadMountedDisplayFunctionLibrary : public UBlueprintFunctionLibrary
 {
-	GENERATED_UCLASS_BODY()
+	GENERATED_BODY()
+public:
+	UHeadMountedDisplayFunctionLibrary(const FObjectInitializer& ObjectInitializer = FObjectInitializer::Get());
 
 	/**
 	 * Returns whether or not we are currently using the head mounted display.
diff --git a/Engine/Source/Runtime/Engine/Classes/Kismet/KismetArrayLibrary.h b/Engine/Source/Runtime/Engine/Classes/Kismet/KismetArrayLibrary.h
--- a/Engine/Source/Runtime/Engine/Classes/Kismet/KismetArrayLibrary.h
+++ b/Engine/Source/Runtime/Engine/Classes/Kismet/KismetArrayLibrary.h
@@ -6,7 +6,9 @@
 UCLASS(MinimalAPI)
 class UKismetArrayLibrary : public UBlueprintFunctionLibrary
 {
-	GENERATED_UCLASS_BODY()
+	GENERATED_BODY()
+public:
+	ENGINE_API UKismetArrayLibrary(const FObjectInitializer& ObjectInitializer = FObjectInitializer::Get());
 	/** 
 	 *Add item to array
 	 *
diff --git a/Engine/Source/Runtime/Engine/Classes/Kismet/KismetGuidLibrary.h b/Engine/Source/Runtime/Engine/Classes/Kismet/KismetGuidLibrary.h
--- a/Engine/Source/Runtime/Engine/Classes/Kismet/KismetGuidLibrary.h
+++ b/Engine/Source/Runtime/Engine/Classes/Kismet/KismetGuidLibrary.h
@@ -9,7 +9,9 @@ UCLASS(MinimalAPI)
 class UKismetGuidLibrary
 	: public UBlueprintFunctionLibrary
 {
-	GENERATED_UCLASS_BODY()
+	GENERATED_BODY()
+public:
+	ENGINE_API UKismetGuidLibrary(const FObjectInitializer& ObjectInitializer = FObjectInitializer::Get());
 
 	/** Returns true if the values are equal (A == B) */
 	UFUNCTION(BlueprintPure, meta=(FriendlyName="Equal (Guid)", CompactNodeTitle="=="), Category="Guid")
diff --git a/Engine/Source/Runtime/Engine/Classes/Kismet/KismetInputLibrary.h b/Engine/Source/Runtime/Engine/Classes/Kismet/KismetInputLibrary.h
--- a/Engine/Source/Runtime/Engine/Classes/Kismet/KismetInputLibrary.h
+++ b/Engine/Source/Runtime/Engine/Classes/Kismet/KismetInputLibrary.h
@@ -9,7 +9,9 @@ typedef FKeyEvent FKeyEvent;
 UCLASS(MinimalAPI)
 class UKismetInputLibrary : public UBlueprintFunctionLibrary
 {
-	GENERATED_UCLASS_BODY()
+	GENERATED_BODY()
+public:
+	ENGINE_API UKismetInputLibrary(const FObjectInitializer& ObjectInitializer = FObjectInitializer::Get());
 
 	/** Calibrate the tilt for the input device */
 	UFUNCTION(BlueprintCallable, Category="Input")
diff --git a/Engine/Source/Runtime/Engine/Classes/Kismet/KismetMaterialLibrary.h b/Engine/Source/Runtime/Engine/Classes/Kismet/KismetMaterialLibrary.h
--- a/Engine/Source/Runtime/Engine/Classes/Kismet/KismetMaterialLibrary.h
+++ b/Engine/Source/Runtime/Engine/Classes/Kismet/KismetMaterialLibrary.h
@@ -7,7 +7,9 @@
 UCLASS(MinimalAPI)
 class UKismetMaterialLibrary : public UBlueprintFunctionLibrary
 {
-	GENERATED_UCLASS_BODY()
+	GENERATED_BODY()
+public:
+	ENGINE_API UKismetMaterialLibrary(const FObjectInitializer& ObjectInitializer = FObjectInitializer::Get());
 
 	/** Sets a scalar parameter value on the material collection instance. Logs if ParameterName is invalid. */
 	UFUNCTION(BlueprintCallable, Category="Rendering|Material", meta=(Keywords="SetFloatParameterValue", WorldContext="WorldContextObject", MaterialParameterCollectionFunction = "true"))
diff --git a/Engine/Source/Runtime/Engine/Classes/Kismet/KismetMathLibrary.h b/Engine/Source/Runtime/Engine/Classes/Kismet/KismetMathLibrary.h
--- a/Engine/Source/Runtime/Engine/Classes/Kismet/KismetMathLibrary.h
+++ b/Engine/Source/Runtime/Engine/Classes/Kismet/KismetMathLibrary.h
@@ -58,7 +58,9 @@ namespace EEasingFunc
 UCLASS(MinimalAPI)
 class UKismetMathLibrary : public UBlueprintFunctionLibrary
 {
-	GENERATED_UCLASS_BODY()
+	GENERATED_BODY()
+public:
+	ENGINE_API UKismetMathLibrary(const FObjectInitializer& ObjectInitializer = FObjectInitializer::Get());
 
 	//
 	// Boolean functions.
diff --git a/Engine/Source/Runtime/Engine/Classes/Kismet/KismetNodeHelperLibrary.h b/Engine/Source/Runtime/Engine/Classes/Kismet/KismetNodeHelperLibrary.h
--- a/Engine/Source/Runtime/Engine/Classes/Kismet/KismetNodeHelperLibrary.h
+++ b/Engine/Source/Runtime/Engine/Classes/Kismet/KismetNodeHelperLibrary.h
@@ -12,7 +12,9 @@
 UCLASS(MinimalAPI)
 class UKismetNodeHelperLibrary : public UBlueprintFunctionLibrary
 {
-	GENERATED_UCLASS_BODY()
+	GENERATED_BODY()
+public:
+	ENGINE_API UKismetNodeHelperLibrary(const FObjectInitializer& ObjectInitializer = FObjectInitializer::Get());
 
 	// --- Functions for handling bitmasking an integer as an array of bools ------------------------------
 
diff --git a/Engine/Source/Runtime/Engine/Classes/Kismet/KismetStringLibrary.h b/Engine/Source/Runtime/Engine/Classes/Kismet/KismetStringLibrary.h
--- a/Engine/Source/Runtime/Engine/Classes/Kismet/KismetStringLibrary.h
+++ b/Engine/Source/Runtime/Engine/Classes/Kismet/KismetStringLibrary.h
@@ -7,7 +7,9 @@
 UCLASS(MinimalAPI)
 class UKismetStringLibrary : public UBlueprintFunctionLibrary
 {
-	GENERATED_UCLASS_BODY()
+	GENERATED_BODY()
+public:
+	ENGINE_API UKismetStringLibrary(const FObjectInitializer& ObjectInitializer = FObjectInitializer::Get());
 
 	/** Converts a float value to a string */
 	UFUNCTION(BlueprintPure, meta=(FriendlyName = "ToString (float)", CompactNodeTitle = "->"), Category="Utilities|String")
diff --git a/Engine/Source/Runtime/Engine/Classes/Kismet/KismetSystemLibrary.h b/Engine/Source/Runtime/Engine/Classes/Kismet/KismetSystemLibrary.h
--- a/Engine/Source/Runtime/Engine/Classes/Kismet/KismetSystemLibrary.h
+++ b/Engine/Source/Runtime/Engine/Classes/Kismet/KismetSystemLibrary.h
@@ -49,7 +49,7 @@ namespace EQuitPreference
 USTRUCT()
 struct FGenericStruct
 {
-	GENERATED_USTRUCT_BODY()
+	GENERATED_BODY()
 
 	UPROPERTY()
 	int32 Data;
@@ -58,7 +58,9 @@ struct FGenericStruct
 UCLASS(MinimalAPI)
 class UKismetSystemLibrary : public UBlueprintFunctionLibrary
 {
-	GENERATED_UCLASS_BODY()
+	GENERATED_BODY()
+public:
+	ENGINE_API UKismetSystemLibrary(const FObjectInitializer& ObjectInitializer = FObjectInitializer::Get());
 
 	// --- Globally useful functions ------------------------------
 	/** Prints a stack trace to the log, so you can see how a blueprint got to this node */
diff --git a/Engine/Source/Runtime/Engine/Classes/Kismet/KismetTextLibrary.h b/Engine/Source/Runtime/Engine/Classes/Kismet/KismetTextLibrary.h
--- a/Engine/Source/Runtime/Engine/Classes/Kismet/KismetTextLibrary.h
+++ b/Engine/Source/Runtime/Engine/Classes/Kismet/KismetTextLibrary.h
@@ -31,7 +31,7 @@ enum ERoundingMode
 USTRUCT()
 struct FFormatTextArgument
 {
-	GENERATED_USTRUCT_BODY()
+	GENERATED_BODY()
 
 	UPROPERTY(EditInstanceOnly, BlueprintReadWrite, Category=ArgumentValue)
 	FText ArgumentName;
@@ -43,7 +43,9 @@ struct FFormatTextArgument
 UCLASS(MinimalAPI)
 class UKismetTextLibrary : public UBlueprintFunctionLibrary
 {
-	GENERATED_UCLASS_BODY()
+	GENERATED_BODY()
+public:
+	ENGINE_API UKismetTextLibrary(const FObjectInitializer& ObjectInitializer = FObjectInitializer::Get());
 
 	/** Converts localizable text to the string */
 	UFUNCTION(BlueprintPure, meta=(FriendlyName = "ToString (text)", CompactNodeTitle = "->"), Category="Utilities|String")
diff --git a/Engine/Source/Runtime/Engine/Classes/Layers/Layer.h b/Engine/Source/Runtime/Engine/Classes/Layers/Layer.h
--- a/Engine/Source/Runtime/Engine/Classes/Layers/Layer.h
+++ b/Engine/Source/Runtime/Engine/Classes/Layers/Layer.h
@@ -6,7 +6,7 @@
 USTRUCT()
 struct FLayerActorStats
 {
-	GENERATED_USTRUCT_BODY()
+	GENERATED_BODY()
 
 	/** A Type of Actor currently associated with the Layer */
 	UPROPERTY()
@@ -26,7 +26,9 @@ struct FLayerActorStats
 UCLASS( MinimalAPI )
 class ULayer : public UObject
 {
-	GENERATED_UCLASS_BODY()
+	GENERATED_BODY()
+public:
+	ENGINE_API ULayer(const FObjectInitializer& ObjectInitializer = FObjectInitializer::Get());
 
 	/** The display name of the layer */
 	UPROPERTY()
@@ -41,4 +43,4 @@ class ULayer : public UObject
 	 */
 	UPROPERTY(transient)
 	TArray< FLayerActorStats > ActorStats;
-};
\ No newline at end of file
+};
diff --git a/Engine/Source/Runtime/Engine/Classes/Lightmass/LightmappedSurfaceCollection.h b/Engine/Source/Runtime/Engine/Classes/Lightmass/LightmappedSurfaceCollection.h
--- a/Engine/Source/Runtime/Engine/Classes/Lightmass/LightmappedSurfaceCollection.h
+++ b/Engine/Source/Runtime/Engine/Classes/Lightmass/LightmappedSurfaceCollection.h
@@ -11,7 +11,9 @@
 UCLASS(hidecategories=Object, editinlinenew, MinimalAPI)
 class ULightmappedSurfaceCollection : public UObject
 {
-	GENERATED_UCLASS_BODY()
+	GENERATED_BODY()
+public:
+	ENGINE_API ULightmappedSurfaceCollection(const FObjectInitializer& ObjectInitializer = FObjectInitializer::Get());
 
 	/** The UModel these surfaces come from. */
 	UPROPERTY(EditAnywhere, Category=LightmappedSurfaceCollection)
diff --git a/Engine/Source/Runtime/Engine/Classes/Lightmass/LightmassCharacterIndirectDetailVolume.h b/Engine/Source/Runtime/Engine/Classes/Lightmass/LightmassCharacterIndirectDetailVolume.h
--- a/Engine/Source/Runtime/Engine/Classes/Lightmass/LightmassCharacterIndirectDetailVolume.h
+++ b/Engine/Source/Runtime/Engine/Classes/Lightmass/LightmassCharacterIndirectDetailVolume.h
@@ -11,7 +11,9 @@
 UCLASS(hidecategories=(Collision, Brush, Attachment, Physics, Volume), MinimalAPI)
 class ALightmassCharacterIndirectDetailVolume : public AVolume
 {
-	GENERATED_UCLASS_BODY()
+	GENERATED_BODY()
+public:
+	ENGINE_API ALightmassCharacterIndirectDetailVolume(const FObjectInitializer& ObjectInitializer = FObjectInitializer::Get());
 
 };
 
diff --git a/Engine/Source/Runtime/Engine/Classes/Lightmass/LightmassImportanceVolume.h b/Engine/Source/Runtime/Engine/Classes/Lightmass/LightmassImportanceVolume.h
--- a/Engine/Source/Runtime/Engine/Classes/Lightmass/LightmassImportanceVolume.h
+++ b/Engine/Source/Runtime/Engine/Classes/Lightmass/LightmassImportanceVolume.h
@@ -12,7 +12,9 @@
 UCLASS(hidecategories=(Collision, Brush, Attachment, Physics, Volume), MinimalAPI)
 class ALightmassImportanceVolume : public AVolume
 {
-	GENERATED_UCLASS_BODY()
+	GENERATED_BODY()
+public:
+	ENGINE_API ALightmassImportanceVolume(const FObjectInitializer& ObjectInitializer = FObjectInitializer::Get());
 
 };
 
diff --git a/Engine/Source/Runtime/Engine/Classes/Lightmass/LightmassPrimitiveSettingsObject.h b/Engine/Source/Runtime/Engine/Classes/Lightmass/LightmassPrimitiveSettingsObject.h
--- a/Engine/Source/Runtime/Engine/Classes/Lightmass/LightmassPrimitiveSettingsObject.h
+++ b/Engine/Source/Runtime/Engine/Classes/Lightmass/LightmassPrimitiveSettingsObject.h
@@ -13,7 +13,9 @@
 UCLASS(hidecategories=Object, editinlinenew, MinimalAPI,collapsecategories)
 class ULightmassPrimitiveSettingsObject : public UObject
 {
-	GENERATED_UCLASS_BODY()
+	GENERATED_BODY()
+public:
+	ENGINE_API ULightmassPrimitiveSettingsObject(const FObjectInitializer& ObjectInitializer = FObjectInitializer::Get());
 
 	UPROPERTY(EditAnywhere, Category=Lightmass)
 	struct FLightmassPrimitiveSettings LightmassSettings;
diff --git a/Engine/Source/Runtime/Engine/Classes/Lightmass/PrecomputedVisibilityOverrideVolume.h b/Engine/Source/Runtime/Engine/Classes/Lightmass/PrecomputedVisibilityOverrideVolume.h
--- a/Engine/Source/Runtime/Engine/Classes/Lightmass/PrecomputedVisibilityOverrideVolume.h
+++ b/Engine/Source/Runtime/Engine/Classes/Lightmass/PrecomputedVisibilityOverrideVolume.h
@@ -11,7 +11,9 @@
 UCLASS(hidecategories=(Collision, Brush, Attachment, Physics, Volume), MinimalAPI)
 class APrecomputedVisibilityOverrideVolume : public AVolume
 {
-	GENERATED_UCLASS_BODY()
+	GENERATED_BODY()
+public:
+	ENGINE_API APrecomputedVisibilityOverrideVolume(const FObjectInitializer& ObjectInitializer = FObjectInitializer::Get());
 
 	/** Array of actors that will always be considered visible by Precomputed Visibility when viewed from inside this volume. */
 	UPROPERTY(EditAnywhere, BlueprintReadWrite, Category=PrecomputedVisibilityOverrideVolume)
diff --git a/Engine/Source/Runtime/Engine/Classes/Lightmass/PrecomputedVisibilityVolume.h b/Engine/Source/Runtime/Engine/Classes/Lightmass/PrecomputedVisibilityVolume.h
--- a/Engine/Source/Runtime/Engine/Classes/Lightmass/PrecomputedVisibilityVolume.h
+++ b/Engine/Source/Runtime/Engine/Classes/Lightmass/PrecomputedVisibilityVolume.h
@@ -11,7 +11,9 @@
 UCLASS(hidecategories=(Collision, Brush, Attachment, Physics, Volume), MinimalAPI)
 class APrecomputedVisibilityVolume : public AVolume
 {
-	GENERATED_UCLASS_BODY()
+	GENERATED_BODY()
+public:
+	ENGINE_API APrecomputedVisibilityVolume(const FObjectInitializer& ObjectInitializer = FObjectInitializer::Get());
 
 };
 
diff --git a/Engine/Source/Runtime/Engine/Classes/Materials/Material.h b/Engine/Source/Runtime/Engine/Classes/Materials/Material.h
--- a/Engine/Source/Runtime/Engine/Classes/Materials/Material.h
+++ b/Engine/Source/Runtime/Engine/Classes/Materials/Material.h
@@ -217,7 +217,7 @@ struct FVector2MaterialInput : public FMaterialInput
 USTRUCT()
 struct FMaterialFunctionInfo
 {
-	GENERATED_USTRUCT_BODY()
+	GENERATED_BODY()
 
 	/** Id that the function had when this material was last compiled. */
 	UPROPERTY()
@@ -236,7 +236,7 @@ struct FMaterialFunctionInfo
 USTRUCT()
 struct FMaterialParameterCollectionInfo
 {
-	GENERATED_USTRUCT_BODY()
+	GENERATED_BODY()
 
 	/** Id that the collection had when this material was last compiled. */
 	UPROPERTY()
@@ -264,7 +264,9 @@ struct FMaterialParameterCollectionInfo
 UCLASS(hidecategories=Object, MinimalAPI, BlueprintType)
 class UMaterial : public UMaterialInterface
 {
-	GENERATED_UCLASS_BODY()
+	GENERATED_BODY()
+public:
+	ENGINE_API UMaterial(const FObjectInitializer& ObjectInitializer = FObjectInitializer::Get());
 
 	// Physics.
 	
diff --git a/Engine/Source/Runtime/Engine/Classes/Materials/MaterialExpression.h b/Engine/Source/Runtime/Engine/Classes/Materials/MaterialExpression.h
--- a/Engine/Source/Runtime/Engine/Classes/Materials/MaterialExpression.h
+++ b/Engine/Source/Runtime/Engine/Classes/Materials/MaterialExpression.h
@@ -85,7 +85,9 @@ struct FExpressionOutput
 UCLASS(abstract, hidecategories=Object)
 class ENGINE_API UMaterialExpression : public UObject
 {
-	GENERATED_UCLASS_BODY()
+	GENERATED_BODY()
+public:
+	UMaterialExpression(const FObjectInitializer& ObjectInitializer = FObjectInitializer::Get());
 
 #if WITH_EDITORONLY_DATA
 	UPROPERTY()
diff --git a/Engine/Source/Runtime/Engine/Classes/Materials/MaterialExpressionAbs.h b/Engine/Source/Runtime/Engine/Classes/Materials/MaterialExpressionAbs.h
--- a/Engine/Source/Runtime/Engine/Classes/Materials/MaterialExpressionAbs.h
+++ b/Engine/Source/Runtime/Engine/Classes/Materials/MaterialExpressionAbs.h
@@ -12,7 +12,9 @@
 UCLASS(collapsecategories, hidecategories=Object)
 class UMaterialExpressionAbs : public UMaterialExpression
 {
-	GENERATED_UCLASS_BODY()
+	GENERATED_BODY()
+public:
+	UMaterialExpressionAbs(const FObjectInitializer& ObjectInitializer = FObjectInitializer::Get());
 
 	/** Link to the input expression to be evaluated */
 	UPROPERTY()
diff --git a/Engine/Source/Runtime/Engine/Classes/Materials/MaterialExpressionActorPositionWS.h b/Engine/Source/Runtime/Engine/Classes/Materials/MaterialExpressionActorPositionWS.h
--- a/Engine/Source/Runtime/Engine/Classes/Materials/MaterialExpressionActorPositionWS.h
+++ b/Engine/Source/Runtime/Engine/Classes/Materials/MaterialExpressionActorPositionWS.h
@@ -8,7 +8,9 @@
 UCLASS(collapsecategories, hidecategories=Object)
 class UMaterialExpressionActorPositionWS : public UMaterialExpression
 {
-	GENERATED_UCLASS_BODY()
+	GENERATED_BODY()
+public:
+	UMaterialExpressionActorPositionWS(const FObjectInitializer& ObjectInitializer = FObjectInitializer::Get());
 
 
 	// Begin UMaterialExpression Interface
diff --git a/Engine/Source/Runtime/Engine/Classes/Materials/MaterialExpressionAdd.h b/Engine/Source/Runtime/Engine/Classes/Materials/MaterialExpressionAdd.h
--- a/Engine/Source/Runtime/Engine/Classes/Materials/MaterialExpressionAdd.h
+++ b/Engine/Source/Runtime/Engine/Classes/Materials/MaterialExpressionAdd.h
@@ -8,7 +8,9 @@
 UCLASS(MinimalAPI)
 class UMaterialExpressionAdd : public UMaterialExpression
 {
-	GENERATED_UCLASS_BODY()
+	GENERATED_BODY()
+public:
+	ENGINE_API UMaterialExpressionAdd(const FObjectInitializer& ObjectInitializer = FObjectInitializer::Get());
 
 	UPROPERTY(meta = (RequiredInput = "false", ToolTip = "Defaults to 'ConstA' if not specified"))
 	FExpressionInput A;
diff --git a/Engine/Source/Runtime/Engine/Classes/Materials/MaterialExpressionAntialiasedTextureMask.h b/Engine/Source/Runtime/Engine/Classes/Materials/MaterialExpressionAntialiasedTextureMask.h
--- a/Engine/Source/Runtime/Engine/Classes/Materials/MaterialExpressionAntialiasedTextureMask.h
+++ b/Engine/Source/Runtime/Engine/Classes/Materials/MaterialExpressionAntialiasedTextureMask.h
@@ -18,7 +18,9 @@ enum ETextureColorChannel
 UCLASS(MinimalAPI)
 class UMaterialExpressionAntialiasedTextureMask : public UMaterialExpressionTextureSampleParameter2D
 {
-	GENERATED_UCLASS_BODY()
+	GENERATED_BODY()
+public:
+	ENGINE_API UMaterialExpressionAntialiasedTextureMask(const FObjectInitializer& ObjectInitializer = FObjectInitializer::Get());
 
 	UPROPERTY(EditAnywhere, Category=MaterialExpressionAntialiasedTextureMask, meta=(UIMin = "0.0", UIMax = "1.0", ClampMin = "0.0", ClampMax = "1.0"))
 	float Threshold;
diff --git a/Engine/Source/Runtime/Engine/Classes/Materials/MaterialExpressionAppendVector.h b/Engine/Source/Runtime/Engine/Classes/Materials/MaterialExpressionAppendVector.h
--- a/Engine/Source/Runtime/Engine/Classes/Materials/MaterialExpressionAppendVector.h
+++ b/Engine/Source/Runtime/Engine/Classes/Materials/MaterialExpressionAppendVector.h
@@ -8,7 +8,9 @@
 UCLASS(collapsecategories, hidecategories=Object)
 class UMaterialExpressionAppendVector : public UMaterialExpression
 {
-	GENERATED_UCLASS_BODY()
+	GENERATED_BODY()
+public:
+	UMaterialExpressionAppendVector(const FObjectInitializer& ObjectInitializer = FObjectInitializer::Get());
 
 	UPROPERTY()
 	FExpressionInput A;
diff --git a/Engine/Source/Runtime/Engine/Classes/Materials/MaterialExpressionAtmosphericFogColor.h b/Engine/Source/Runtime/Engine/Classes/Materials/MaterialExpressionAtmosphericFogColor.h
--- a/Engine/Source/Runtime/Engine/Classes/Materials/MaterialExpressionAtmosphericFogColor.h
+++ b/Engine/Source/Runtime/Engine/Classes/Materials/MaterialExpressionAtmosphericFogColor.h
@@ -7,7 +7,9 @@
 UCLASS(collapsecategories, hidecategories=Object)
 class UMaterialExpressionAtmosphericFogColor : public UMaterialExpression
 {
-	GENERATED_UCLASS_BODY()
+	GENERATED_BODY()
+public:
+	UMaterialExpressionAtmosphericFogColor(const FObjectInitializer& ObjectInitializer = FObjectInitializer::Get());
 
 	UPROPERTY()
 	FExpressionInput WorldPosition;
diff --git a/Engine/Source/Runtime/Engine/Classes/Materials/MaterialExpressionBlackBody.h b/Engine/Source/Runtime/Engine/Classes/Materials/MaterialExpressionBlackBody.h
--- a/Engine/Source/Runtime/Engine/Classes/Materials/MaterialExpressionBlackBody.h
+++ b/Engine/Source/Runtime/Engine/Classes/Materials/MaterialExpressionBlackBody.h
@@ -8,7 +8,9 @@
 UCLASS()
 class UMaterialExpressionBlackBody : public UMaterialExpression
 {
-	GENERATED_UCLASS_BODY()
+	GENERATED_BODY()
+public:
+	UMaterialExpressionBlackBody(const FObjectInitializer& ObjectInitializer = FObjectInitializer::Get());
 
 	/** Temperature */
 	UPROPERTY()
diff --git a/Engine/Source/Runtime/Engine/Classes/Materials/MaterialExpressionBreakMaterialAttributes.h b/Engine/Source/Runtime/Engine/Classes/Materials/MaterialExpressionBreakMaterialAttributes.h
--- a/Engine/Source/Runtime/Engine/Classes/Materials/MaterialExpressionBreakMaterialAttributes.h
+++ b/Engine/Source/Runtime/Engine/Classes/Materials/MaterialExpressionBreakMaterialAttributes.h
@@ -8,7 +8,9 @@
 UCLASS(collapsecategories, hidecategories=Object, MinimalAPI)
 class UMaterialExpressionBreakMaterialAttributes : public UMaterialExpression
 {
-	GENERATED_UCLASS_BODY()
+	GENERATED_BODY()
+public:
+	ENGINE_API UMaterialExpressionBreakMaterialAttributes(const FObjectInitializer& ObjectInitializer = FObjectInitializer::Get());
 
 	UPROPERTY()
 	FExpressionInput Struct;
diff --git a/Engine/Source/Runtime/Engine/Classes/Materials/MaterialExpressionBumpOffset.h b/Engine/Source/Runtime/Engine/Classes/Materials/MaterialExpressionBumpOffset.h
--- a/Engine/Source/Runtime/Engine/Classes/Materials/MaterialExpressionBumpOffset.h
+++ b/Engine/Source/Runtime/Engine/Classes/Materials/MaterialExpressionBumpOffset.h
@@ -8,7 +8,9 @@
 UCLASS(collapsecategories, hidecategories=Object, MinimalAPI)
 class UMaterialExpressionBumpOffset : public UMaterialExpression
 {
-	GENERATED_UCLASS_BODY()
+	GENERATED_BODY()
+public:
+	ENGINE_API UMaterialExpressionBumpOffset(const FObjectInitializer& ObjectInitializer = FObjectInitializer::Get());
 
 	// Outputs: Coordinate + Eye.xy * (Height - ReferencePlane) * HeightRatio
 	UPROPERTY(meta = (RequiredInput = "false", ToolTip = "Defaults to 'ConstCoordinate' if not specified"))
diff --git a/Engine/Source/Runtime/Engine/Classes/Materials/MaterialExpressionCameraPositionWS.h b/Engine/Source/Runtime/Engine/Classes/Materials/MaterialExpressionCameraPositionWS.h
--- a/Engine/Source/Runtime/Engine/Classes/Materials/MaterialExpressionCameraPositionWS.h
+++ b/Engine/Source/Runtime/Engine/Classes/Materials/MaterialExpressionCameraPositionWS.h
@@ -8,7 +8,9 @@
 UCLASS(collapsecategories, hidecategories=Object)
 class UMaterialExpressionCameraPositionWS : public UMaterialExpression
 {
-	GENERATED_UCLASS_BODY()
+	GENERATED_BODY()
+public:
+	UMaterialExpressionCameraPositionWS(const FObjectInitializer& ObjectInitializer = FObjectInitializer::Get());
 
 
 	// Begin UMaterialExpression Interface
diff --git a/Engine/Source/Runtime/Engine/Classes/Materials/MaterialExpressionCameraVectorWS.h b/Engine/Source/Runtime/Engine/Classes/Materials/MaterialExpressionCameraVectorWS.h
--- a/Engine/Source/Runtime/Engine/Classes/Materials/MaterialExpressionCameraVectorWS.h
+++ b/Engine/Source/Runtime/Engine/Classes/Materials/MaterialExpressionCameraVectorWS.h
@@ -8,7 +8,9 @@
 UCLASS(collapsecategories, hidecategories=Object)
 class UMaterialExpressionCameraVectorWS : public UMaterialExpression
 {
-	GENERATED_UCLASS_BODY()
+	GENERATED_BODY()
+public:
+	UMaterialExpressionCameraVectorWS(const FObjectInitializer& ObjectInitializer = FObjectInitializer::Get());
 
 
 	// Begin UMaterialExpression Interface
diff --git a/Engine/Source/Runtime/Engine/Classes/Materials/MaterialExpressionCeil.h b/Engine/Source/Runtime/Engine/Classes/Materials/MaterialExpressionCeil.h
--- a/Engine/Source/Runtime/Engine/Classes/Materials/MaterialExpressionCeil.h
+++ b/Engine/Source/Runtime/Engine/Classes/Materials/MaterialExpressionCeil.h
@@ -9,7 +9,9 @@
 UCLASS(collapsecategories, hidecategories=Object)
 class UMaterialExpressionCeil : public UMaterialExpression
 {
-	GENERATED_UCLASS_BODY()
+	GENERATED_BODY()
+public:
+	UMaterialExpressionCeil(const FObjectInitializer& ObjectInitializer = FObjectInitializer::Get());
 
 	UPROPERTY()
 	FExpressionInput Input;
diff --git a/Engine/Source/Runtime/Engine/Classes/Materials/MaterialExpressionClamp.h b/Engine/Source/Runtime/Engine/Classes/Materials/MaterialExpressionClamp.h
--- a/Engine/Source/Runtime/Engine/Classes/Materials/MaterialExpressionClamp.h
+++ b/Engine/Source/Runtime/Engine/Classes/Materials/MaterialExpressionClamp.h
@@ -17,7 +17,9 @@ enum EClampMode
 UCLASS(MinimalAPI)
 class UMaterialExpressionClamp : public UMaterialExpression
 {
-	GENERATED_UCLASS_BODY()
+	GENERATED_BODY()
+public:
+	ENGINE_API UMaterialExpressionClamp(const FObjectInitializer& ObjectInitializer = FObjectInitializer::Get());
 
 	UPROPERTY()
 	FExpressionInput Input;
diff --git a/Engine/Source/Runtime/Engine/Classes/Materials/MaterialExpressionCollectionParameter.h b/Engine/Source/Runtime/Engine/Classes/Materials/MaterialExpressionCollectionParameter.h
--- a/Engine/Source/Runtime/Engine/Classes/Materials/MaterialExpressionCollectionParameter.h
+++ b/Engine/Source/Runtime/Engine/Classes/Materials/MaterialExpressionCollectionParameter.h
@@ -11,7 +11,9 @@
 UCLASS(hidecategories=object, MinimalAPI)
 class UMaterialExpressionCollectionParameter : public UMaterialExpression
 {
-	GENERATED_UCLASS_BODY()
+	GENERATED_BODY()
+public:
+	ENGINE_API UMaterialExpressionCollectionParameter(const FObjectInitializer& ObjectInitializer = FObjectInitializer::Get());
 
 	/** The Parameter Collection to use. */
 	UPROPERTY(EditAnywhere, Category=MaterialExpressionCollectionParameter)
diff --git a/Engine/Source/Runtime/Engine/Classes/Materials/MaterialExpressionComment.h b/Engine/Source/Runtime/Engine/Classes/Materials/MaterialExpressionComment.h
--- a/Engine/Source/Runtime/Engine/Classes/Materials/MaterialExpressionComment.h
+++ b/Engine/Source/Runtime/Engine/Classes/Materials/MaterialExpressionComment.h
@@ -8,7 +8,9 @@
 UCLASS(MinimalAPI)
 class UMaterialExpressionComment : public UMaterialExpression
 {
-	GENERATED_UCLASS_BODY()
+	GENERATED_BODY()
+public:
+	ENGINE_API UMaterialExpressionComment(const FObjectInitializer& ObjectInitializer = FObjectInitializer::Get());
 
 	UPROPERTY()
 	int32 SizeX;
diff --git a/Engine/Source/Runtime/Engine/Classes/Materials/MaterialExpressionComponentMask.h b/Engine/Source/Runtime/Engine/Classes/Materials/MaterialExpressionComponentMask.h
--- a/Engine/Source/Runtime/Engine/Classes/Materials/MaterialExpressionComponentMask.h
+++ b/Engine/Source/Runtime/Engine/Classes/Materials/MaterialExpressionComponentMask.h
@@ -8,7 +8,9 @@
 UCLASS(collapsecategories, hidecategories=Object, MinimalAPI)
 class UMaterialExpressionComponentMask : public UMaterialExpression
 {
-	GENERATED_UCLASS_BODY()
+	GENERATED_BODY()
+public:
+	ENGINE_API UMaterialExpressionComponentMask(const FObjectInitializer& ObjectInitializer = FObjectInitializer::Get());
 
 	UPROPERTY()
 	FExpressionInput Input;
diff --git a/Engine/Source/Runtime/Engine/Classes/Materials/MaterialExpressionConstant.h b/Engine/Source/Runtime/Engine/Classes/Materials/MaterialExpressionConstant.h
--- a/Engine/Source/Runtime/Engine/Classes/Materials/MaterialExpressionConstant.h
+++ b/Engine/Source/Runtime/Engine/Classes/Materials/MaterialExpressionConstant.h
@@ -8,7 +8,9 @@
 UCLASS(collapsecategories, hidecategories=Object, MinimalAPI)
 class UMaterialExpressionConstant : public UMaterialExpression
 {
-	GENERATED_UCLASS_BODY()
+	GENERATED_BODY()
+public:
+	ENGINE_API UMaterialExpressionConstant(const FObjectInitializer& ObjectInitializer = FObjectInitializer::Get());
 
 	UPROPERTY(EditAnywhere, Category=MaterialExpressionConstant, DisplayName="Value")
 	float R;
diff --git a/Engine/Source/Runtime/Engine/Classes/Materials/MaterialExpressionConstant2Vector.h b/Engine/Source/Runtime/Engine/Classes/Materials/MaterialExpressionConstant2Vector.h
--- a/Engine/Source/Runtime/Engine/Classes/Materials/MaterialExpressionConstant2Vector.h
+++ b/Engine/Source/Runtime/Engine/Classes/Materials/MaterialExpressionConstant2Vector.h
@@ -8,7 +8,9 @@
 UCLASS(collapsecategories, hidecategories=Object, MinimalAPI)
 class UMaterialExpressionConstant2Vector : public UMaterialExpression
 {
-	GENERATED_UCLASS_BODY()
+	GENERATED_BODY()
+public:
+	ENGINE_API UMaterialExpressionConstant2Vector(const FObjectInitializer& ObjectInitializer = FObjectInitializer::Get());
 
 	UPROPERTY(EditAnywhere, Category=MaterialExpressionConstant2Vector)
 	float R;
diff --git a/Engine/Source/Runtime/Engine/Classes/Materials/MaterialExpressionConstant3Vector.h b/Engine/Source/Runtime/Engine/Classes/Materials/MaterialExpressionConstant3Vector.h
--- a/Engine/Source/Runtime/Engine/Classes/Materials/MaterialExpressionConstant3Vector.h
+++ b/Engine/Source/Runtime/Engine/Classes/Materials/MaterialExpressionConstant3Vector.h
@@ -8,7 +8,9 @@
 UCLASS(collapsecategories, hidecategories=Object, MinimalAPI)
 class UMaterialExpressionConstant3Vector : public UMaterialExpression
 {
-	GENERATED_UCLASS_BODY()
+	GENERATED_BODY()
+public:
+	ENGINE_API UMaterialExpressionConstant3Vector(const FObjectInitializer& ObjectInitializer = FObjectInitializer::Get());
 
  	UPROPERTY(EditAnywhere, Category=MaterialExpressionConstant3Vector, meta=(HideAlphaChannel))
 	FLinearColor Constant;
diff --git a/Engine/Source/Runtime/Engine/Classes/Materials/MaterialExpressionConstant4Vector.h b/Engine/Source/Runtime/Engine/Classes/Materials/MaterialExpressionConstant4Vector.h
--- a/Engine/Source/Runtime/Engine/Classes/Materials/MaterialExpressionConstant4Vector.h
+++ b/Engine/Source/Runtime/Engine/Classes/Materials/MaterialExpressionConstant4Vector.h
@@ -8,7 +8,9 @@
 UCLASS(collapsecategories, hidecategories=Object, MinimalAPI)
 class UMaterialExpressionConstant4Vector : public UMaterialExpression
 {
-	GENERATED_UCLASS_BODY()
+	GENERATED_BODY()
+public:
+	ENGINE_API UMaterialExpressionConstant4Vector(const FObjectInitializer& ObjectInitializer = FObjectInitializer::Get());
 
 	UPROPERTY(EditAnywhere, Category=MaterialExpressionConstant4Vector)
 	FLinearColor Constant;
diff --git a/Engine/Source/Runtime/Engine/Classes/Materials/MaterialExpressionConstantBiasScale.h b/Engine/Source/Runtime/Engine/Classes/Materials/MaterialExpressionConstantBiasScale.h
--- a/Engine/Source/Runtime/Engine/Classes/Materials/MaterialExpressionConstantBiasScale.h
+++ b/Engine/Source/Runtime/Engine/Classes/Materials/MaterialExpressionConstantBiasScale.h
@@ -8,7 +8,9 @@
 UCLASS(MinimalAPI)
 class UMaterialExpressionConstantBiasScale : public UMaterialExpression
 {
-	GENERATED_UCLASS_BODY()
+	GENERATED_BODY()
+public:
+	ENGINE_API UMaterialExpressionConstantBiasScale(const FObjectInitializer& ObjectInitializer = FObjectInitializer::Get());
 
 	UPROPERTY()
 	FExpressionInput Input;
diff --git a/Engine/Source/Runtime/Engine/Classes/Materials/MaterialExpressionCosine.h b/Engine/Source/Runtime/Engine/Classes/Materials/MaterialExpressionCosine.h
--- a/Engine/Source/Runtime/Engine/Classes/Materials/MaterialExpressionCosine.h
+++ b/Engine/Source/Runtime/Engine/Classes/Materials/MaterialExpressionCosine.h
@@ -8,7 +8,9 @@
 UCLASS(collapsecategories, hidecategories=Object)
 class UMaterialExpressionCosine : public UMaterialExpression
 {
-	GENERATED_UCLASS_BODY()
+	GENERATED_BODY()
+public:
+	UMaterialExpressionCosine(const FObjectInitializer& ObjectInitializer = FObjectInitializer::Get());
 
 	UPROPERTY()
 	FExpressionInput Input;
diff --git a/Engine/Source/Runtime/Engine/Classes/Materials/MaterialExpressionCrossProduct.h b/Engine/Source/Runtime/Engine/Classes/Materials/MaterialExpressionCrossProduct.h
--- a/Engine/Source/Runtime/Engine/Classes/Materials/MaterialExpressionCrossProduct.h
+++ b/Engine/Source/Runtime/Engine/Classes/Materials/MaterialExpressionCrossProduct.h
@@ -8,7 +8,9 @@
 UCLASS()
 class UMaterialExpressionCrossProduct : public UMaterialExpression
 {
-	GENERATED_UCLASS_BODY()
+	GENERATED_BODY()
+public:
+	UMaterialExpressionCrossProduct(const FObjectInitializer& ObjectInitializer = FObjectInitializer::Get());
 
 	UPROPERTY()
 	FExpressionInput A;
diff --git a/Engine/Source/Runtime/Engine/Classes/Materials/MaterialExpressionCustom.h b/Engine/Source/Runtime/Engine/Classes/Materials/MaterialExpressionCustom.h
--- a/Engine/Source/Runtime/Engine/Classes/Materials/MaterialExpressionCustom.h
+++ b/Engine/Source/Runtime/Engine/Classes/Materials/MaterialExpressionCustom.h
@@ -18,7 +18,7 @@ enum ECustomMaterialOutputType
 USTRUCT()
 struct FCustomInput
 {
-	GENERATED_USTRUCT_BODY()
+	GENERATED_BODY()
 
 	UPROPERTY(EditAnywhere, Category=CustomInput)
 	FString InputName;
@@ -31,7 +31,9 @@ struct FCustomInput
 UCLASS(collapsecategories, hidecategories=Object, MinimalAPI)
 class UMaterialExpressionCustom : public UMaterialExpression
 {
-	GENERATED_UCLASS_BODY()
+	GENERATED_BODY()
+public:
+	ENGINE_API UMaterialExpressionCustom(const FObjectInitializer& ObjectInitializer = FObjectInitializer::Get());
 
 	UPROPERTY(EditAnywhere, Category=MaterialExpressionCustom, meta=(MultiLine=true))
 	FString Code;
diff --git a/Engine/Source/Runtime/Engine/Classes/Materials/MaterialExpressionCustomOutput.h b/Engine/Source/Runtime/Engine/Classes/Materials/MaterialExpressionCustomOutput.h
--- a/Engine/Source/Runtime/Engine/Classes/Materials/MaterialExpressionCustomOutput.h
+++ b/Engine/Source/Runtime/Engine/Classes/Materials/MaterialExpressionCustomOutput.h
@@ -8,7 +8,9 @@
 UCLASS(collapsecategories, hidecategories=Object, MinimalAPI)
 class UMaterialExpressionCustomOutput : public UMaterialExpression
 {
-	GENERATED_UCLASS_BODY()
+	GENERATED_BODY()
+public:
+	ENGINE_API UMaterialExpressionCustomOutput(const FObjectInitializer& ObjectInitializer = FObjectInitializer::Get());
 
 	// Override to enable multiple outputs
 	virtual int32 GetNumOutputs() const { return 1; };
diff --git a/Engine/Source/Runtime/Engine/Classes/Materials/MaterialExpressionDDX.h b/Engine/Source/Runtime/Engine/Classes/Materials/MaterialExpressionDDX.h
--- a/Engine/Source/Runtime/Engine/Classes/Materials/MaterialExpressionDDX.h
+++ b/Engine/Source/Runtime/Engine/Classes/Materials/MaterialExpressionDDX.h
@@ -9,7 +9,9 @@
 UCLASS()
 class UMaterialExpressionDDX : public UMaterialExpression
 {
-	GENERATED_UCLASS_BODY()
+	GENERATED_BODY()
+public:
+	UMaterialExpressionDDX(const FObjectInitializer& ObjectInitializer = FObjectInitializer::Get());
 
 	/** The value we want to compute ddx/ddy from */
 	UPROPERTY()
diff --git a/Engine/Source/Runtime/Engine/Classes/Materials/MaterialExpressionDDY.h b/Engine/Source/Runtime/Engine/Classes/Materials/MaterialExpressionDDY.h
--- a/Engine/Source/Runtime/Engine/Classes/Materials/MaterialExpressionDDY.h
+++ b/Engine/Source/Runtime/Engine/Classes/Materials/MaterialExpressionDDY.h
@@ -9,7 +9,9 @@
 UCLASS()
 class UMaterialExpressionDDY : public UMaterialExpression
 {
-	GENERATED_UCLASS_BODY()
+	GENERATED_BODY()
+public:
+	UMaterialExpressionDDY(const FObjectInitializer& ObjectInitializer = FObjectInitializer::Get());
 
 	/** The value we want to compute ddx/ddy from */
 	UPROPERTY()
diff --git a/Engine/Source/Runtime/Engine/Classes/Materials/MaterialExpressionDepthFade.h b/Engine/Source/Runtime/Engine/Classes/Materials/MaterialExpressionDepthFade.h
--- a/Engine/Source/Runtime/Engine/Classes/Materials/MaterialExpressionDepthFade.h
+++ b/Engine/Source/Runtime/Engine/Classes/Materials/MaterialExpressionDepthFade.h
@@ -11,7 +11,9 @@
 UCLASS(collapsecategories, hidecategories=Object)
 class UMaterialExpressionDepthFade : public UMaterialExpression
 {
-	GENERATED_UCLASS_BODY()
+	GENERATED_BODY()
+public:
+	UMaterialExpressionDepthFade(const FObjectInitializer& ObjectInitializer = FObjectInitializer::Get());
 
 	/** Input opacity which will be scaled by the result of the fade. */
 	UPROPERTY(meta = (RequiredInput = "false", ToolTip = "Defaults to 'OpacityDefault' if not specified"))
diff --git a/Engine/Source/Runtime/Engine/Classes/Materials/MaterialExpressionDepthOfFieldFunction.h b/Engine/Source/Runtime/Engine/Classes/Materials/MaterialExpressionDepthOfFieldFunction.h
--- a/Engine/Source/Runtime/Engine/Classes/Materials/MaterialExpressionDepthOfFieldFunction.h
+++ b/Engine/Source/Runtime/Engine/Classes/Materials/MaterialExpressionDepthOfFieldFunction.h
@@ -20,7 +20,9 @@ enum EDepthOfFieldFunctionValue
 UCLASS()
 class UMaterialExpressionDepthOfFieldFunction : public UMaterialExpression
 {
-	GENERATED_UCLASS_BODY()
+	GENERATED_BODY()
+public:
+	UMaterialExpressionDepthOfFieldFunction(const FObjectInitializer& ObjectInitializer = FObjectInitializer::Get());
 
 	/** Determines the mapping place to use on the terrain. */
 	UPROPERTY(EditAnywhere, Category=MaterialExpressionDepthOfFieldFunction)
diff --git a/Engine/Source/Runtime/Engine/Classes/Materials/MaterialExpressionDeriveNormalZ.h b/Engine/Source/Runtime/Engine/Classes/Materials/MaterialExpressionDeriveNormalZ.h
--- a/Engine/Source/Runtime/Engine/Classes/Materials/MaterialExpressionDeriveNormalZ.h
+++ b/Engine/Source/Runtime/Engine/Classes/Materials/MaterialExpressionDeriveNormalZ.h
@@ -9,7 +9,9 @@
 UCLASS(collapsecategories, hidecategories=Object)
 class UMaterialExpressionDeriveNormalZ : public UMaterialExpression
 {
-	GENERATED_UCLASS_BODY()
+	GENERATED_BODY()
+public:
+	UMaterialExpressionDeriveNormalZ(const FObjectInitializer& ObjectInitializer = FObjectInitializer::Get());
 
 	UPROPERTY()
 	FExpressionInput InXY;
diff --git a/Engine/Source/Runtime/Engine/Classes/Materials/MaterialExpressionDesaturation.h b/Engine/Source/Runtime/Engine/Classes/Materials/MaterialExpressionDesaturation.h
--- a/Engine/Source/Runtime/Engine/Classes/Materials/MaterialExpressionDesaturation.h
+++ b/Engine/Source/Runtime/Engine/Classes/Materials/MaterialExpressionDesaturation.h
@@ -8,7 +8,9 @@
 UCLASS(collapsecategories, hidecategories=Object)
 class UMaterialExpressionDesaturation : public UMaterialExpression
 {
-	GENERATED_UCLASS_BODY()
+	GENERATED_BODY()
+public:
+	UMaterialExpressionDesaturation(const FObjectInitializer& ObjectInitializer = FObjectInitializer::Get());
 
 	// Outputs: Lerp(Input,dot(Input,LuminanceFactors)),Fraction)
 	UPROPERTY()
diff --git a/Engine/Source/Runtime/Engine/Classes/Materials/MaterialExpressionDistance.h b/Engine/Source/Runtime/Engine/Classes/Materials/MaterialExpressionDistance.h
--- a/Engine/Source/Runtime/Engine/Classes/Materials/MaterialExpressionDistance.h
+++ b/Engine/Source/Runtime/Engine/Classes/Materials/MaterialExpressionDistance.h
@@ -8,7 +8,9 @@
 UCLASS()
 class UMaterialExpressionDistance : public UMaterialExpression
 {
-	GENERATED_UCLASS_BODY()
+	GENERATED_BODY()
+public:
+	UMaterialExpressionDistance(const FObjectInitializer& ObjectInitializer = FObjectInitializer::Get());
 
 	UPROPERTY()
 	FExpressionInput A;
diff --git a/Engine/Source/Runtime/Engine/Classes/Materials/MaterialExpressionDistanceCullFade.h b/Engine/Source/Runtime/Engine/Classes/Materials/MaterialExpressionDistanceCullFade.h
--- a/Engine/Source/Runtime/Engine/Classes/Materials/MaterialExpressionDistanceCullFade.h
+++ b/Engine/Source/Runtime/Engine/Classes/Materials/MaterialExpressionDistanceCullFade.h
@@ -8,7 +8,9 @@
 UCLASS(collapsecategories, hidecategories=Object)
 class UMaterialExpressionDistanceCullFade : public UMaterialExpression
 {
-	GENERATED_UCLASS_BODY()
+	GENERATED_BODY()
+public:
+	UMaterialExpressionDistanceCullFade(const FObjectInitializer& ObjectInitializer = FObjectInitializer::Get());
 
 
 	// Begin UMaterialExpression Interface
diff --git a/Engine/Source/Runtime/Engine/Classes/Materials/MaterialExpressionDivide.h b/Engine/Source/Runtime/Engine/Classes/Materials/MaterialExpressionDivide.h
--- a/Engine/Source/Runtime/Engine/Classes/Materials/MaterialExpressionDivide.h
+++ b/Engine/Source/Runtime/Engine/Classes/Materials/MaterialExpressionDivide.h
@@ -8,7 +8,9 @@
 UCLASS(MinimalAPI)
 class UMaterialExpressionDivide : public UMaterialExpression
 {
-	GENERATED_UCLASS_BODY()
+	GENERATED_BODY()
+public:
+	ENGINE_API UMaterialExpressionDivide(const FObjectInitializer& ObjectInitializer = FObjectInitializer::Get());
 
 	UPROPERTY(meta = (RequiredInput = "false", ToolTip = "Defaults to 'ConstA' if not specified"))
 	FExpressionInput A;
diff --git a/Engine/Source/Runtime/Engine/Classes/Materials/MaterialExpressionDotProduct.h b/Engine/Source/Runtime/Engine/Classes/Materials/MaterialExpressionDotProduct.h
--- a/Engine/Source/Runtime/Engine/Classes/Materials/MaterialExpressionDotProduct.h
+++ b/Engine/Source/Runtime/Engine/Classes/Materials/MaterialExpressionDotProduct.h
@@ -8,7 +8,9 @@
 UCLASS()
 class UMaterialExpressionDotProduct : public UMaterialExpression
 {
-	GENERATED_UCLASS_BODY()
+	GENERATED_BODY()
+public:
+	UMaterialExpressionDotProduct(const FObjectInitializer& ObjectInitializer = FObjectInitializer::Get());
 
 	UPROPERTY()
 	FExpressionInput A;
diff --git a/Engine/Source/Runtime/Engine/Classes/Materials/MaterialExpressionDynamicParameter.h b/Engine/Source/Runtime/Engine/Classes/Materials/MaterialExpressionDynamicParameter.h
--- a/Engine/Source/Runtime/Engine/Classes/Materials/MaterialExpressionDynamicParameter.h
+++ b/Engine/Source/Runtime/Engine/Classes/Materials/MaterialExpressionDynamicParameter.h
@@ -13,7 +13,9 @@
 UCLASS(collapsecategories, hidecategories=Object, MinimalAPI)
 class UMaterialExpressionDynamicParameter : public UMaterialExpression
 {
-	GENERATED_UCLASS_BODY()
+	GENERATED_BODY()
+public:
+	ENGINE_API UMaterialExpressionDynamicParameter(const FObjectInitializer& ObjectInitializer = FObjectInitializer::Get());
 
 	/** 
 	 *	The names of the parameters.
diff --git a/Engine/Source/Runtime/Engine/Classes/Materials/MaterialExpressionEyeAdaptation.h b/Engine/Source/Runtime/Engine/Classes/Materials/MaterialExpressionEyeAdaptation.h
--- a/Engine/Source/Runtime/Engine/Classes/Materials/MaterialExpressionEyeAdaptation.h
+++ b/Engine/Source/Runtime/Engine/Classes/Materials/MaterialExpressionEyeAdaptation.h
@@ -12,7 +12,9 @@
 UCLASS(collapsecategories, hidecategories=Object)
 class UMaterialExpressionEyeAdaptation : public UMaterialExpression
 {
-	GENERATED_UCLASS_BODY()
+	GENERATED_BODY()
+public:
+	UMaterialExpressionEyeAdaptation(const FObjectInitializer& ObjectInitializer = FObjectInitializer::Get());
 
 	// Begin UMaterialExpression Interface
 	virtual int32 Compile(class FMaterialCompiler* Compiler, int32 OutputIndex, int32 MultiplexIndex) override;
diff --git a/Engine/Source/Runtime/Engine/Classes/Materials/MaterialExpressionFeatureLevelSwitch.h b/Engine/Source/Runtime/Engine/Classes/Materials/MaterialExpressionFeatureLevelSwitch.h
--- a/Engine/Source/Runtime/Engine/Classes/Materials/MaterialExpressionFeatureLevelSwitch.h
+++ b/Engine/Source/Runtime/Engine/Classes/Materials/MaterialExpressionFeatureLevelSwitch.h
@@ -8,7 +8,9 @@
 UCLASS(collapsecategories, hidecategories=Object, MinimalAPI)
 class UMaterialExpressionFeatureLevelSwitch : public UMaterialExpression
 {
-	GENERATED_UCLASS_BODY()
+	GENERATED_BODY()
+public:
+	ENGINE_API UMaterialExpressionFeatureLevelSwitch(const FObjectInitializer& ObjectInitializer = FObjectInitializer::Get());
 
 	/** Default connection, used when a certain feature level doesn't have an override. */
 	UPROPERTY()
diff --git a/Engine/Source/Runtime/Engine/Classes/Materials/MaterialExpressionFloor.h b/Engine/Source/Runtime/Engine/Classes/Materials/MaterialExpressionFloor.h
--- a/Engine/Source/Runtime/Engine/Classes/Materials/MaterialExpressionFloor.h
+++ b/Engine/Source/Runtime/Engine/Classes/Materials/MaterialExpressionFloor.h
@@ -8,7 +8,9 @@
 UCLASS(collapsecategories, hidecategories=Object)
 class UMaterialExpressionFloor : public UMaterialExpression
 {
-	GENERATED_UCLASS_BODY()
+	GENERATED_BODY()
+public:
+	UMaterialExpressionFloor(const FObjectInitializer& ObjectInitializer = FObjectInitializer::Get());
 
 	UPROPERTY()
 	FExpressionInput Input;
diff --git a/Engine/Source/Runtime/Engine/Classes/Materials/MaterialExpressionFmod.h b/Engine/Source/Runtime/Engine/Classes/Materials/MaterialExpressionFmod.h
--- a/Engine/Source/Runtime/Engine/Classes/Materials/MaterialExpressionFmod.h
+++ b/Engine/Source/Runtime/Engine/Classes/Materials/MaterialExpressionFmod.h
@@ -8,7 +8,9 @@
 UCLASS(collapsecategories, hidecategories=Object)
 class UMaterialExpressionFmod : public UMaterialExpression
 {
-	GENERATED_UCLASS_BODY()
+	GENERATED_BODY()
+public:
+	UMaterialExpressionFmod(const FObjectInitializer& ObjectInitializer = FObjectInitializer::Get());
 
 	UPROPERTY()
 	FExpressionInput A;
diff --git a/Engine/Source/Runtime/Engine/Classes/Materials/MaterialExpressionFontSample.h b/Engine/Source/Runtime/Engine/Classes/Materials/MaterialExpressionFontSample.h
--- a/Engine/Source/Runtime/Engine/Classes/Materials/MaterialExpressionFontSample.h
+++ b/Engine/Source/Runtime/Engine/Classes/Materials/MaterialExpressionFontSample.h
@@ -8,7 +8,9 @@
 UCLASS(collapsecategories, hidecategories=Object, MinimalAPI)
 class UMaterialExpressionFontSample : public UMaterialExpression
 {
-	GENERATED_UCLASS_BODY()
+	GENERATED_BODY()
+public:
+	ENGINE_API UMaterialExpressionFontSample(const FObjectInitializer& ObjectInitializer = FObjectInitializer::Get());
 
 	/** font resource that will be sampled */
 	UPROPERTY(EditAnywhere, Category=MaterialExpressionFontSample)
diff --git a/Engine/Source/Runtime/Engine/Classes/Materials/MaterialExpressionFontSampleParameter.h b/Engine/Source/Runtime/Engine/Classes/Materials/MaterialExpressionFontSampleParameter.h
--- a/Engine/Source/Runtime/Engine/Classes/Materials/MaterialExpressionFontSampleParameter.h
+++ b/Engine/Source/Runtime/Engine/Classes/Materials/MaterialExpressionFontSampleParameter.h
@@ -8,7 +8,9 @@
 UCLASS(collapsecategories, hidecategories=Object, MinimalAPI)
 class UMaterialExpressionFontSampleParameter : public UMaterialExpressionFontSample
 {
-	GENERATED_UCLASS_BODY()
+	GENERATED_BODY()
+public:
+	ENGINE_API UMaterialExpressionFontSampleParameter(const FObjectInitializer& ObjectInitializer = FObjectInitializer::Get());
 
 	/** name to be referenced when we want to find and set thsi parameter */
 	UPROPERTY(EditAnywhere, Category=MaterialExpressionFontSampleParameter)
diff --git a/Engine/Source/Runtime/Engine/Classes/Materials/MaterialExpressionFrac.h b/Engine/Source/Runtime/Engine/Classes/Materials/MaterialExpressionFrac.h
--- a/Engine/Source/Runtime/Engine/Classes/Materials/MaterialExpressionFrac.h
+++ b/Engine/Source/Runtime/Engine/Classes/Materials/MaterialExpressionFrac.h
@@ -8,7 +8,9 @@
 UCLASS(collapsecategories, hidecategories=Object)
 class UMaterialExpressionFrac : public UMaterialExpression
 {
-	GENERATED_UCLASS_BODY()
+	GENERATED_BODY()
+public:
+	UMaterialExpressionFrac(const FObjectInitializer& ObjectInitializer = FObjectInitializer::Get());
 
 	UPROPERTY()
 	FExpressionInput Input;
diff --git a/Engine/Source/Runtime/Engine/Classes/Materials/MaterialExpressionFresnel.h b/Engine/Source/Runtime/Engine/Classes/Materials/MaterialExpressionFresnel.h
--- a/Engine/Source/Runtime/Engine/Classes/Materials/MaterialExpressionFresnel.h
+++ b/Engine/Source/Runtime/Engine/Classes/Materials/MaterialExpressionFresnel.h
@@ -14,7 +14,9 @@
 UCLASS(collapsecategories, hidecategories=Object)
 class UMaterialExpressionFresnel : public UMaterialExpression
 {
-	GENERATED_UCLASS_BODY()
+	GENERATED_BODY()
+public:
+	UMaterialExpressionFresnel(const FObjectInitializer& ObjectInitializer = FObjectInitializer::Get());
 
 	UPROPERTY(meta = (RequiredInput = "false", ToolTip = "Defaults to 'Exponent' if not specified"))
 	FExpressionInput ExponentIn;
diff --git a/Engine/Source/Runtime/Engine/Classes/Materials/MaterialExpressionFunctionInput.h b/Engine/Source/Runtime/Engine/Classes/Materials/MaterialExpressionFunctionInput.h
--- a/Engine/Source/Runtime/Engine/Classes/Materials/MaterialExpressionFunctionInput.h
+++ b/Engine/Source/Runtime/Engine/Classes/Materials/MaterialExpressionFunctionInput.h
@@ -23,7 +23,9 @@ enum EFunctionInputType
 UCLASS(hidecategories=object, MinimalAPI)
 class UMaterialExpressionFunctionInput : public UMaterialExpression
 {
-	GENERATED_UCLASS_BODY()
+	GENERATED_BODY()
+public:
+	ENGINE_API UMaterialExpressionFunctionInput(const FObjectInitializer& ObjectInitializer = FObjectInitializer::Get());
 
 	/** Used for previewing when editing the function, also temporarily used to connect to the outside material when compiling that material. */
 	UPROPERTY(meta=(RequiredInput = "false"))
diff --git a/Engine/Source/Runtime/Engine/Classes/Materials/MaterialExpressionFunctionOutput.h b/Engine/Source/Runtime/Engine/Classes/Materials/MaterialExpressionFunctionOutput.h
--- a/Engine/Source/Runtime/Engine/Classes/Materials/MaterialExpressionFunctionOutput.h
+++ b/Engine/Source/Runtime/Engine/Classes/Materials/MaterialExpressionFunctionOutput.h
@@ -8,7 +8,9 @@
 UCLASS(hidecategories=object, MinimalAPI)
 class UMaterialExpressionFunctionOutput : public UMaterialExpression
 {
-	GENERATED_UCLASS_BODY()
+	GENERATED_BODY()
+public:
+	ENGINE_API UMaterialExpressionFunctionOutput(const FObjectInitializer& ObjectInitializer = FObjectInitializer::Get());
 
 	/** The output's name, which will be drawn on the connector in function call expressions that use this function. */
 	UPROPERTY(EditAnywhere, Category=MaterialExpressionFunctionOutput)
diff --git a/Engine/Source/Runtime/Engine/Classes/Materials/MaterialExpressionGIReplace.h b/Engine/Source/Runtime/Engine/Classes/Materials/MaterialExpressionGIReplace.h
--- a/Engine/Source/Runtime/Engine/Classes/Materials/MaterialExpressionGIReplace.h
+++ b/Engine/Source/Runtime/Engine/Classes/Materials/MaterialExpressionGIReplace.h
@@ -8,7 +8,9 @@
 UCLASS()
 class UMaterialExpressionGIReplace : public UMaterialExpression
 {
-	GENERATED_UCLASS_BODY()
+	GENERATED_BODY()
+public:
+	UMaterialExpressionGIReplace(const FObjectInitializer& ObjectInitializer = FObjectInitializer::Get());
 
 	/** Used for direct lighting computations e.g. real-time shaders */
 	UPROPERTY()
diff --git a/Engine/Source/Runtime/Engine/Classes/Materials/MaterialExpressionIf.h b/Engine/Source/Runtime/Engine/Classes/Materials/MaterialExpressionIf.h
--- a/Engine/Source/Runtime/Engine/Classes/Materials/MaterialExpressionIf.h
+++ b/Engine/Source/Runtime/Engine/Classes/Materials/MaterialExpressionIf.h
@@ -8,7 +8,9 @@
 UCLASS(collapsecategories, hidecategories=Object, MinimalAPI)
 class UMaterialExpressionIf : public UMaterialExpression
 {
-	GENERATED_UCLASS_BODY()
+	GENERATED_BODY()
+public:
+	ENGINE_API UMaterialExpressionIf(const FObjectInitializer& ObjectInitializer = FObjectInitializer::Get());
 
 	UPROPERTY()
 	FExpressionInput A;
diff --git a/Engine/Source/Runtime/Engine/Classes/Materials/MaterialExpressionLightVector.h b/Engine/Source/Runtime/Engine/Classes/Materials/MaterialExpressionLightVector.h
--- a/Engine/Source/Runtime/Engine/Classes/Materials/MaterialExpressionLightVector.h
+++ b/Engine/Source/Runtime/Engine/Classes/Materials/MaterialExpressionLightVector.h
@@ -8,7 +8,9 @@
 UCLASS(collapsecategories, hidecategories=Object)
 class UMaterialExpressionLightVector : public UMaterialExpression
 {
-	GENERATED_UCLASS_BODY()
+	GENERATED_BODY()
+public:
+	UMaterialExpressionLightVector(const FObjectInitializer& ObjectInitializer = FObjectInitializer::Get());
 
 
 	// Begin UMaterialExpression Interface
diff --git a/Engine/Source/Runtime/Engine/Classes/Materials/MaterialExpressionLightmapUVs.h b/Engine/Source/Runtime/Engine/Classes/Materials/MaterialExpressionLightmapUVs.h
--- a/Engine/Source/Runtime/Engine/Classes/Materials/MaterialExpressionLightmapUVs.h
+++ b/Engine/Source/Runtime/Engine/Classes/Materials/MaterialExpressionLightmapUVs.h
@@ -12,7 +12,9 @@
 UCLASS(collapsecategories, hidecategories=Object)
 class UMaterialExpressionLightmapUVs : public UMaterialExpression
 {
-	GENERATED_UCLASS_BODY()
+	GENERATED_BODY()
+public:
+	UMaterialExpressionLightmapUVs(const FObjectInitializer& ObjectInitializer = FObjectInitializer::Get());
 
 
 	// Begin UMaterialExpression Interface
diff --git a/Engine/Source/Runtime/Engine/Classes/Materials/MaterialExpressionLightmassReplace.h b/Engine/Source/Runtime/Engine/Classes/Materials/MaterialExpressionLightmassReplace.h
--- a/Engine/Source/Runtime/Engine/Classes/Materials/MaterialExpressionLightmassReplace.h
+++ b/Engine/Source/Runtime/Engine/Classes/Materials/MaterialExpressionLightmassReplace.h
@@ -8,7 +8,9 @@
 UCLASS()
 class UMaterialExpressionLightmassReplace : public UMaterialExpression
 {
-	GENERATED_UCLASS_BODY()
+	GENERATED_BODY()
+public:
+	UMaterialExpressionLightmassReplace(const FObjectInitializer& ObjectInitializer = FObjectInitializer::Get());
 
 	UPROPERTY()
 	FExpressionInput Realtime;
diff --git a/Engine/Source/Runtime/Engine/Classes/Materials/MaterialExpressionLinearInterpolate.h b/Engine/Source/Runtime/Engine/Classes/Materials/MaterialExpressionLinearInterpolate.h
--- a/Engine/Source/Runtime/Engine/Classes/Materials/MaterialExpressionLinearInterpolate.h
+++ b/Engine/Source/Runtime/Engine/Classes/Materials/MaterialExpressionLinearInterpolate.h
@@ -8,7 +8,9 @@
 UCLASS(MinimalAPI)
 class UMaterialExpressionLinearInterpolate : public UMaterialExpression
 {
-	GENERATED_UCLASS_BODY()
+	GENERATED_BODY()
+public:
+	ENGINE_API UMaterialExpressionLinearInterpolate(const FObjectInitializer& ObjectInitializer = FObjectInitializer::Get());
 
 	UPROPERTY(meta = (RequiredInput = "false", ToolTip = "Defaults to 'ConstA' if not specified"))
 	FExpressionInput A;
diff --git a/Engine/Source/Runtime/Engine/Classes/Materials/MaterialExpressionMakeMaterialAttributes.h b/Engine/Source/Runtime/Engine/Classes/Materials/MaterialExpressionMakeMaterialAttributes.h
--- a/Engine/Source/Runtime/Engine/Classes/Materials/MaterialExpressionMakeMaterialAttributes.h
+++ b/Engine/Source/Runtime/Engine/Classes/Materials/MaterialExpressionMakeMaterialAttributes.h
@@ -8,7 +8,9 @@
 UCLASS(collapsecategories, hidecategories=Object)
 class UMaterialExpressionMakeMaterialAttributes : public UMaterialExpression
 {
-	GENERATED_UCLASS_BODY()
+	GENERATED_BODY()
+public:
+	UMaterialExpressionMakeMaterialAttributes(const FObjectInitializer& ObjectInitializer = FObjectInitializer::Get());
 
 	UPROPERTY()
 	FExpressionInput BaseColor;
diff --git a/Engine/Source/Runtime/Engine/Classes/Materials/MaterialExpressionMaterialFunctionCall.h b/Engine/Source/Runtime/Engine/Classes/Materials/MaterialExpressionMaterialFunctionCall.h
--- a/Engine/Source/Runtime/Engine/Classes/Materials/MaterialExpressionMaterialFunctionCall.h
+++ b/Engine/Source/Runtime/Engine/Classes/Materials/MaterialExpressionMaterialFunctionCall.h
@@ -12,7 +12,7 @@
 USTRUCT()
 struct FFunctionExpressionInput
 {
-	GENERATED_USTRUCT_BODY()
+	GENERATED_BODY()
 
 	/** 
 	 * Reference to the FunctionInput in the material function.  
@@ -41,7 +41,7 @@ struct FFunctionExpressionInput
 USTRUCT()
 struct FFunctionExpressionOutput
 {
-	GENERATED_USTRUCT_BODY()
+	GENERATED_BODY()
 
 	/** 
 	 * Reference to the FunctionOutput in the material function.  
@@ -69,7 +69,9 @@ struct FFunctionExpressionOutput
 UCLASS(hidecategories=object, MinimalAPI)
 class UMaterialExpressionMaterialFunctionCall : public UMaterialExpression
 {
-	GENERATED_UCLASS_BODY()
+	GENERATED_BODY()
+public:
+	ENGINE_API UMaterialExpressionMaterialFunctionCall(const FObjectInitializer& ObjectInitializer = FObjectInitializer::Get());
 
 	/** The function to call. */
 	UPROPERTY(EditAnywhere, Category=MaterialExpressionMaterialFunctionCall)
diff --git a/Engine/Source/Runtime/Engine/Classes/Materials/MaterialExpressionMax.h b/Engine/Source/Runtime/Engine/Classes/Materials/MaterialExpressionMax.h
--- a/Engine/Source/Runtime/Engine/Classes/Materials/MaterialExpressionMax.h
+++ b/Engine/Source/Runtime/Engine/Classes/Materials/MaterialExpressionMax.h
@@ -8,7 +8,9 @@
 UCLASS(MinimalAPI)
 class UMaterialExpressionMax : public UMaterialExpression
 {
-	GENERATED_UCLASS_BODY()
+	GENERATED_BODY()
+public:
+	ENGINE_API UMaterialExpressionMax(const FObjectInitializer& ObjectInitializer = FObjectInitializer::Get());
 
 	UPROPERTY(meta = (RequiredInput = "false", ToolTip = "Defaults to 'ConstA' if not specified"))
 	FExpressionInput A;
diff --git a/Engine/Source/Runtime/Engine/Classes/Materials/MaterialExpressionMin.h b/Engine/Source/Runtime/Engine/Classes/Materials/MaterialExpressionMin.h
--- a/Engine/Source/Runtime/Engine/Classes/Materials/MaterialExpressionMin.h
+++ b/Engine/Source/Runtime/Engine/Classes/Materials/MaterialExpressionMin.h
@@ -8,7 +8,9 @@
 UCLASS(MinimalAPI)
 class UMaterialExpressionMin : public UMaterialExpression
 {
-	GENERATED_UCLASS_BODY()
+	GENERATED_BODY()
+public:
+	ENGINE_API UMaterialExpressionMin(const FObjectInitializer& ObjectInitializer = FObjectInitializer::Get());
 
 	UPROPERTY(meta = (RequiredInput = "false", ToolTip = "Defaults to 'ConstA' if not specified"))
 	FExpressionInput A;
diff --git a/Engine/Source/Runtime/Engine/Classes/Materials/MaterialExpressionMultiply.h b/Engine/Source/Runtime/Engine/Classes/Materials/MaterialExpressionMultiply.h
--- a/Engine/Source/Runtime/Engine/Classes/Materials/MaterialExpressionMultiply.h
+++ b/Engine/Source/Runtime/Engine/Classes/Materials/MaterialExpressionMultiply.h
@@ -8,7 +8,9 @@
 UCLASS(MinimalAPI)
 class UMaterialExpressionMultiply : public UMaterialExpression
 {
-	GENERATED_UCLASS_BODY()
+	GENERATED_BODY()
+public:
+	ENGINE_API UMaterialExpressionMultiply(const FObjectInitializer& ObjectInitializer = FObjectInitializer::Get());
 
 	UPROPERTY(meta = (RequiredInput = "false", ToolTip = "Defaults to 'ConstA' if not specified"))
 	FExpressionInput A;
diff --git a/Engine/Source/Runtime/Engine/Classes/Materials/MaterialExpressionNoise.h b/Engine/Source/Runtime/Engine/Classes/Materials/MaterialExpressionNoise.h
--- a/Engine/Source/Runtime/Engine/Classes/Materials/MaterialExpressionNoise.h
+++ b/Engine/Source/Runtime/Engine/Classes/Materials/MaterialExpressionNoise.h
@@ -22,7 +22,9 @@ enum ENoiseFunction
 UCLASS()
 class UMaterialExpressionNoise : public UMaterialExpression
 {
-	GENERATED_UCLASS_BODY()
+	GENERATED_BODY()
+public:
+	UMaterialExpressionNoise(const FObjectInitializer& ObjectInitializer = FObjectInitializer::Get());
 
 	/** 2 to 3 dimensional vector */
 	UPROPERTY()
diff --git a/Engine/Source/Runtime/Engine/Classes/Materials/MaterialExpressionNormalize.h b/Engine/Source/Runtime/Engine/Classes/Materials/MaterialExpressionNormalize.h
--- a/Engine/Source/Runtime/Engine/Classes/Materials/MaterialExpressionNormalize.h
+++ b/Engine/Source/Runtime/Engine/Classes/Materials/MaterialExpressionNormalize.h
@@ -8,7 +8,9 @@
 UCLASS(MinimalAPI)
 class UMaterialExpressionNormalize : public UMaterialExpression
 {
-	GENERATED_UCLASS_BODY()
+	GENERATED_BODY()
+public:
+	ENGINE_API UMaterialExpressionNormalize(const FObjectInitializer& ObjectInitializer = FObjectInitializer::Get());
 
 	UPROPERTY()
 	FExpressionInput VectorInput;
diff --git a/Engine/Source/Runtime/Engine/Classes/Materials/MaterialExpressionObjectBounds.h b/Engine/Source/Runtime/Engine/Classes/Materials/MaterialExpressionObjectBounds.h
--- a/Engine/Source/Runtime/Engine/Classes/Materials/MaterialExpressionObjectBounds.h
+++ b/Engine/Source/Runtime/Engine/Classes/Materials/MaterialExpressionObjectBounds.h
@@ -8,7 +8,9 @@
 UCLASS(collapsecategories, hidecategories=Object)
 class UMaterialExpressionObjectBounds : public UMaterialExpression
 {
-	GENERATED_UCLASS_BODY()
+	GENERATED_BODY()
+public:
+	UMaterialExpressionObjectBounds(const FObjectInitializer& ObjectInitializer = FObjectInitializer::Get());
 
 
 	// Begin UMaterialExpression Interface
diff --git a/Engine/Source/Runtime/Engine/Classes/Materials/MaterialExpressionObjectOrientation.h b/Engine/Source/Runtime/Engine/Classes/Materials/MaterialExpressionObjectOrientation.h
--- a/Engine/Source/Runtime/Engine/Classes/Materials/MaterialExpressionObjectOrientation.h
+++ b/Engine/Source/Runtime/Engine/Classes/Materials/MaterialExpressionObjectOrientation.h
@@ -8,7 +8,9 @@
 UCLASS(collapsecategories, hidecategories=Object)
 class UMaterialExpressionObjectOrientation : public UMaterialExpression
 {
-	GENERATED_UCLASS_BODY()
+	GENERATED_BODY()
+public:
+	UMaterialExpressionObjectOrientation(const FObjectInitializer& ObjectInitializer = FObjectInitializer::Get());
 
 
 	// Begin UMaterialExpression Interface
diff --git a/Engine/Source/Runtime/Engine/Classes/Materials/MaterialExpressionObjectPositionWS.h b/Engine/Source/Runtime/Engine/Classes/Materials/MaterialExpressionObjectPositionWS.h
--- a/Engine/Source/Runtime/Engine/Classes/Materials/MaterialExpressionObjectPositionWS.h
+++ b/Engine/Source/Runtime/Engine/Classes/Materials/MaterialExpressionObjectPositionWS.h
@@ -8,7 +8,9 @@
 UCLASS(collapsecategories, hidecategories=Object)
 class UMaterialExpressionObjectPositionWS : public UMaterialExpression
 {
-	GENERATED_UCLASS_BODY()
+	GENERATED_BODY()
+public:
+	UMaterialExpressionObjectPositionWS(const FObjectInitializer& ObjectInitializer = FObjectInitializer::Get());
 
 
 	// Begin UMaterialExpression Interface
diff --git a/Engine/Source/Runtime/Engine/Classes/Materials/MaterialExpressionObjectRadius.h b/Engine/Source/Runtime/Engine/Classes/Materials/MaterialExpressionObjectRadius.h
--- a/Engine/Source/Runtime/Engine/Classes/Materials/MaterialExpressionObjectRadius.h
+++ b/Engine/Source/Runtime/Engine/Classes/Materials/MaterialExpressionObjectRadius.h
@@ -8,7 +8,9 @@
 UCLASS(collapsecategories, hidecategories=Object)
 class UMaterialExpressionObjectRadius : public UMaterialExpression
 {
-	GENERATED_UCLASS_BODY()
+	GENERATED_BODY()
+public:
+	UMaterialExpressionObjectRadius(const FObjectInitializer& ObjectInitializer = FObjectInitializer::Get());
 
 
 	// Begin UMaterialExpression Interface
diff --git a/Engine/Source/Runtime/Engine/Classes/Materials/MaterialExpressionOneMinus.h b/Engine/Source/Runtime/Engine/Classes/Materials/MaterialExpressionOneMinus.h
--- a/Engine/Source/Runtime/Engine/Classes/Materials/MaterialExpressionOneMinus.h
+++ b/Engine/Source/Runtime/Engine/Classes/Materials/MaterialExpressionOneMinus.h
@@ -8,7 +8,9 @@
 UCLASS(collapsecategories, hidecategories=Object, MinimalAPI)
 class UMaterialExpressionOneMinus : public UMaterialExpression
 {
-	GENERATED_UCLASS_BODY()
+	GENERATED_BODY()
+public:
+	ENGINE_API UMaterialExpressionOneMinus(const FObjectInitializer& ObjectInitializer = FObjectInitializer::Get());
 
 	UPROPERTY()
 	FExpressionInput Input;
diff --git a/Engine/Source/Runtime/Engine/Classes/Materials/MaterialExpressionPanner.h b/Engine/Source/Runtime/Engine/Classes/Materials/MaterialExpressionPanner.h
--- a/Engine/Source/Runtime/Engine/Classes/Materials/MaterialExpressionPanner.h
+++ b/Engine/Source/Runtime/Engine/Classes/Materials/MaterialExpressionPanner.h
@@ -8,7 +8,9 @@
 UCLASS(collapsecategories, hidecategories=Object, MinimalAPI)
 class UMaterialExpressionPanner : public UMaterialExpression
 {
-	GENERATED_UCLASS_BODY()
+	GENERATED_BODY()
+public:
+	ENGINE_API UMaterialExpressionPanner(const FObjectInitializer& ObjectInitializer = FObjectInitializer::Get());
 
 	UPROPERTY(meta = (RequiredInput = "false", ToolTip = "Defaults to 'ConstCoordinate' if not specified"))
 	FExpressionInput Coordinate;
diff --git a/Engine/Source/Runtime/Engine/Classes/Materials/MaterialExpressionParameter.h b/Engine/Source/Runtime/Engine/Classes/Materials/MaterialExpressionParameter.h
--- a/Engine/Source/Runtime/Engine/Classes/Materials/MaterialExpressionParameter.h
+++ b/Engine/Source/Runtime/Engine/Classes/Materials/MaterialExpressionParameter.h
@@ -8,7 +8,9 @@
 UCLASS(collapsecategories, hidecategories=Object, MinimalAPI)
 class UMaterialExpressionParameter : public UMaterialExpression
 {
-	GENERATED_UCLASS_BODY()
+	GENERATED_BODY()
+public:
+	ENGINE_API UMaterialExpressionParameter(const FObjectInitializer& ObjectInitializer = FObjectInitializer::Get());
 
 	/** The name of the parameter */
 	UPROPERTY(EditAnywhere, Category=MaterialExpressionParameter)
diff --git a/Engine/Source/Runtime/Engine/Classes/Materials/MaterialExpressionParticleColor.h b/Engine/Source/Runtime/Engine/Classes/Materials/MaterialExpressionParticleColor.h
--- a/Engine/Source/Runtime/Engine/Classes/Materials/MaterialExpressionParticleColor.h
+++ b/Engine/Source/Runtime/Engine/Classes/Materials/MaterialExpressionParticleColor.h
@@ -8,7 +8,9 @@
 UCLASS(collapsecategories, hidecategories=Object)
 class UMaterialExpressionParticleColor : public UMaterialExpression
 {
-	GENERATED_UCLASS_BODY()
+	GENERATED_BODY()
+public:
+	UMaterialExpressionParticleColor(const FObjectInitializer& ObjectInitializer = FObjectInitializer::Get());
 
 
 	// Begin UMaterialExpression Interface
diff --git a/Engine/Source/Runtime/Engine/Classes/Materials/MaterialExpressionParticleDirection.h b/Engine/Source/Runtime/Engine/Classes/Materials/MaterialExpressionParticleDirection.h
--- a/Engine/Source/Runtime/Engine/Classes/Materials/MaterialExpressionParticleDirection.h
+++ b/Engine/Source/Runtime/Engine/Classes/Materials/MaterialExpressionParticleDirection.h
@@ -12,7 +12,9 @@
 UCLASS(collapsecategories, hidecategories=Object)
 class UMaterialExpressionParticleDirection : public UMaterialExpression
 {
-	GENERATED_UCLASS_BODY()
+	GENERATED_BODY()
+public:
+	UMaterialExpressionParticleDirection(const FObjectInitializer& ObjectInitializer = FObjectInitializer::Get());
 
 
 	// Begin UMaterialExpression Interface
diff --git a/Engine/Source/Runtime/Engine/Classes/Materials/MaterialExpressionParticleMacroUV.h b/Engine/Source/Runtime/Engine/Classes/Materials/MaterialExpressionParticleMacroUV.h
--- a/Engine/Source/Runtime/Engine/Classes/Materials/MaterialExpressionParticleMacroUV.h
+++ b/Engine/Source/Runtime/Engine/Classes/Materials/MaterialExpressionParticleMacroUV.h
@@ -12,7 +12,9 @@
 UCLASS()
 class UMaterialExpressionParticleMacroUV : public UMaterialExpression
 {
-	GENERATED_UCLASS_BODY()
+	GENERATED_BODY()
+public:
+	UMaterialExpressionParticleMacroUV(const FObjectInitializer& ObjectInitializer = FObjectInitializer::Get());
 
 
 	// Begin UMaterialExpression Interface
diff --git a/Engine/Source/Runtime/Engine/Classes/Materials/MaterialExpressionParticleMotionBlurFade.h b/Engine/Source/Runtime/Engine/Classes/Materials/MaterialExpressionParticleMotionBlurFade.h
--- a/Engine/Source/Runtime/Engine/Classes/Materials/MaterialExpressionParticleMotionBlurFade.h
+++ b/Engine/Source/Runtime/Engine/Classes/Materials/MaterialExpressionParticleMotionBlurFade.h
@@ -12,7 +12,9 @@
 UCLASS(collapsecategories, hidecategories=Object)
 class UMaterialExpressionParticleMotionBlurFade : public UMaterialExpression
 {
-	GENERATED_UCLASS_BODY()
+	GENERATED_BODY()
+public:
+	UMaterialExpressionParticleMotionBlurFade(const FObjectInitializer& ObjectInitializer = FObjectInitializer::Get());
 
 
 	// Begin UMaterialExpression Interface
diff --git a/Engine/Source/Runtime/Engine/Classes/Materials/MaterialExpressionParticlePositionWS.h b/Engine/Source/Runtime/Engine/Classes/Materials/MaterialExpressionParticlePositionWS.h
--- a/Engine/Source/Runtime/Engine/Classes/Materials/MaterialExpressionParticlePositionWS.h
+++ b/Engine/Source/Runtime/Engine/Classes/Materials/MaterialExpressionParticlePositionWS.h
@@ -8,7 +8,9 @@
 UCLASS(collapsecategories, hidecategories=Object)
 class UMaterialExpressionParticlePositionWS : public UMaterialExpression
 {
-	GENERATED_UCLASS_BODY()
+	GENERATED_BODY()
+public:
+	UMaterialExpressionParticlePositionWS(const FObjectInitializer& ObjectInitializer = FObjectInitializer::Get());
 
 
 	// Begin UMaterialExpression Interface
diff --git a/Engine/Source/Runtime/Engine/Classes/Materials/MaterialExpressionParticleRadius.h b/Engine/Source/Runtime/Engine/Classes/Materials/MaterialExpressionParticleRadius.h
--- a/Engine/Source/Runtime/Engine/Classes/Materials/MaterialExpressionParticleRadius.h
+++ b/Engine/Source/Runtime/Engine/Classes/Materials/MaterialExpressionParticleRadius.h
@@ -8,7 +8,9 @@
 UCLASS(collapsecategories, hidecategories=Object)
 class UMaterialExpressionParticleRadius : public UMaterialExpression
 {
-	GENERATED_UCLASS_BODY()
+	GENERATED_BODY()
+public:
+	UMaterialExpressionParticleRadius(const FObjectInitializer& ObjectInitializer = FObjectInitializer::Get());
 
 
 	// Begin UMaterialExpression Interface
diff --git a/Engine/Source/Runtime/Engine/Classes/Materials/MaterialExpressionParticleRelativeTime.h b/Engine/Source/Runtime/Engine/Classes/Materials/MaterialExpressionParticleRelativeTime.h
--- a/Engine/Source/Runtime/Engine/Classes/Materials/MaterialExpressionParticleRelativeTime.h
+++ b/Engine/Source/Runtime/Engine/Classes/Materials/MaterialExpressionParticleRelativeTime.h
@@ -12,7 +12,9 @@
 UCLASS(collapsecategories, hidecategories=Object)
 class UMaterialExpressionParticleRelativeTime : public UMaterialExpression
 {
-	GENERATED_UCLASS_BODY()
+	GENERATED_BODY()
+public:
+	UMaterialExpressionParticleRelativeTime(const FObjectInitializer& ObjectInitializer = FObjectInitializer::Get());
 
 
 	// Begin UMaterialExpression Interface
diff --git a/Engine/Source/Runtime/Engine/Classes/Materials/MaterialExpressionParticleSize.h b/Engine/Source/Runtime/Engine/Classes/Materials/MaterialExpressionParticleSize.h
--- a/Engine/Source/Runtime/Engine/Classes/Materials/MaterialExpressionParticleSize.h
+++ b/Engine/Source/Runtime/Engine/Classes/Materials/MaterialExpressionParticleSize.h
@@ -12,7 +12,9 @@
 UCLASS(collapsecategories, hidecategories=Object)
 class UMaterialExpressionParticleSize : public UMaterialExpression
 {
-	GENERATED_UCLASS_BODY()
+	GENERATED_BODY()
+public:
+	UMaterialExpressionParticleSize(const FObjectInitializer& ObjectInitializer = FObjectInitializer::Get());
 
 
 	// Begin UMaterialExpression Interface
diff --git a/Engine/Source/Runtime/Engine/Classes/Materials/MaterialExpressionParticleSpeed.h b/Engine/Source/Runtime/Engine/Classes/Materials/MaterialExpressionParticleSpeed.h
--- a/Engine/Source/Runtime/Engine/Classes/Materials/MaterialExpressionParticleSpeed.h
+++ b/Engine/Source/Runtime/Engine/Classes/Materials/MaterialExpressionParticleSpeed.h
@@ -12,7 +12,9 @@
 UCLASS(collapsecategories, hidecategories=Object)
 class UMaterialExpressionParticleSpeed : public UMaterialExpression
 {
-	GENERATED_UCLASS_BODY()
+	GENERATED_BODY()
+public:
+	UMaterialExpressionParticleSpeed(const FObjectInitializer& ObjectInitializer = FObjectInitializer::Get());
 
 
 	// Begin UMaterialExpression Interface
diff --git a/Engine/Source/Runtime/Engine/Classes/Materials/MaterialExpressionParticleSubUV.h b/Engine/Source/Runtime/Engine/Classes/Materials/MaterialExpressionParticleSubUV.h
--- a/Engine/Source/Runtime/Engine/Classes/Materials/MaterialExpressionParticleSubUV.h
+++ b/Engine/Source/Runtime/Engine/Classes/Materials/MaterialExpressionParticleSubUV.h
@@ -8,7 +8,9 @@
 UCLASS(MinimalAPI)
 class UMaterialExpressionParticleSubUV : public UMaterialExpressionTextureSample
 {
-	GENERATED_UCLASS_BODY()
+	GENERATED_BODY()
+public:
+	ENGINE_API UMaterialExpressionParticleSubUV(const FObjectInitializer& ObjectInitializer = FObjectInitializer::Get());
 
 	UPROPERTY(EditAnywhere, Category=MaterialExpressionParticleSubUV)
 	uint32 bBlend:1;
diff --git a/Engine/Source/Runtime/Engine/Classes/Materials/MaterialExpressionPerInstanceFadeAmount.h b/Engine/Source/Runtime/Engine/Classes/Materials/MaterialExpressionPerInstanceFadeAmount.h
--- a/Engine/Source/Runtime/Engine/Classes/Materials/MaterialExpressionPerInstanceFadeAmount.h
+++ b/Engine/Source/Runtime/Engine/Classes/Materials/MaterialExpressionPerInstanceFadeAmount.h
@@ -8,7 +8,9 @@
 UCLASS(collapsecategories, hidecategories=Object)
 class UMaterialExpressionPerInstanceFadeAmount : public UMaterialExpression
 {
-	GENERATED_UCLASS_BODY()
+	GENERATED_BODY()
+public:
+	UMaterialExpressionPerInstanceFadeAmount(const FObjectInitializer& ObjectInitializer = FObjectInitializer::Get());
 
 
 	// Begin UMaterialExpression Interface
diff --git a/Engine/Source/Runtime/Engine/Classes/Materials/MaterialExpressionPerInstanceRandom.h b/Engine/Source/Runtime/Engine/Classes/Materials/MaterialExpressionPerInstanceRandom.h
--- a/Engine/Source/Runtime/Engine/Classes/Materials/MaterialExpressionPerInstanceRandom.h
+++ b/Engine/Source/Runtime/Engine/Classes/Materials/MaterialExpressionPerInstanceRandom.h
@@ -8,7 +8,9 @@
 UCLASS(collapsecategories, hidecategories=Object)
 class UMaterialExpressionPerInstanceRandom : public UMaterialExpression
 {
-	GENERATED_UCLASS_BODY()
+	GENERATED_BODY()
+public:
+	UMaterialExpressionPerInstanceRandom(const FObjectInitializer& ObjectInitializer = FObjectInitializer::Get());
 
 
 	// Begin UMaterialExpression Interface
diff --git a/Engine/Source/Runtime/Engine/Classes/Materials/MaterialExpressionPixelDepth.h b/Engine/Source/Runtime/Engine/Classes/Materials/MaterialExpressionPixelDepth.h
--- a/Engine/Source/Runtime/Engine/Classes/Materials/MaterialExpressionPixelDepth.h
+++ b/Engine/Source/Runtime/Engine/Classes/Materials/MaterialExpressionPixelDepth.h
@@ -11,7 +11,9 @@
 UCLASS(collapsecategories, hidecategories=Object)
 class UMaterialExpressionPixelDepth : public UMaterialExpression
 {
-	GENERATED_UCLASS_BODY()
+	GENERATED_BODY()
+public:
+	UMaterialExpressionPixelDepth(const FObjectInitializer& ObjectInitializer = FObjectInitializer::Get());
 
 
 	// Begin UMaterialExpression Interface
diff --git a/Engine/Source/Runtime/Engine/Classes/Materials/MaterialExpressionPixelNormalWS.h b/Engine/Source/Runtime/Engine/Classes/Materials/MaterialExpressionPixelNormalWS.h
--- a/Engine/Source/Runtime/Engine/Classes/Materials/MaterialExpressionPixelNormalWS.h
+++ b/Engine/Source/Runtime/Engine/Classes/Materials/MaterialExpressionPixelNormalWS.h
@@ -8,7 +8,9 @@
 UCLASS(collapsecategories, hidecategories=Object)
 class UMaterialExpressionPixelNormalWS : public UMaterialExpression
 {
-	GENERATED_UCLASS_BODY()
+	GENERATED_BODY()
+public:
+	UMaterialExpressionPixelNormalWS(const FObjectInitializer& ObjectInitializer = FObjectInitializer::Get());
 
 
 	// Begin UMaterialExpression Interface
diff --git a/Engine/Source/Runtime/Engine/Classes/Materials/MaterialExpressionPower.h b/Engine/Source/Runtime/Engine/Classes/Materials/MaterialExpressionPower.h
--- a/Engine/Source/Runtime/Engine/Classes/Materials/MaterialExpressionPower.h
+++ b/Engine/Source/Runtime/Engine/Classes/Materials/MaterialExpressionPower.h
@@ -8,7 +8,9 @@
 UCLASS(collapsecategories, hidecategories=Object, MinimalAPI)
 class UMaterialExpressionPower : public UMaterialExpression
 {
-	GENERATED_UCLASS_BODY()
+	GENERATED_BODY()
+public:
+	ENGINE_API UMaterialExpressionPower(const FObjectInitializer& ObjectInitializer = FObjectInitializer::Get());
 
 	UPROPERTY()
 	FExpressionInput Base;
diff --git a/Engine/Source/Runtime/Engine/Classes/Materials/MaterialExpressionQualitySwitch.h b/Engine/Source/Runtime/Engine/Classes/Materials/MaterialExpressionQualitySwitch.h
--- a/Engine/Source/Runtime/Engine/Classes/Materials/MaterialExpressionQualitySwitch.h
+++ b/Engine/Source/Runtime/Engine/Classes/Materials/MaterialExpressionQualitySwitch.h
@@ -8,7 +8,9 @@
 UCLASS(collapsecategories, hidecategories=Object, MinimalAPI)
 class UMaterialExpressionQualitySwitch : public UMaterialExpression
 {
-	GENERATED_UCLASS_BODY()
+	GENERATED_BODY()
+public:
+	ENGINE_API UMaterialExpressionQualitySwitch(const FObjectInitializer& ObjectInitializer = FObjectInitializer::Get());
 
 	/** Default connection, used when a specific quality level input is missing. */
 	UPROPERTY()
diff --git a/Engine/Source/Runtime/Engine/Classes/Materials/MaterialExpressionReflectionVectorWS.h b/Engine/Source/Runtime/Engine/Classes/Materials/MaterialExpressionReflectionVectorWS.h
--- a/Engine/Source/Runtime/Engine/Classes/Materials/MaterialExpressionReflectionVectorWS.h
+++ b/Engine/Source/Runtime/Engine/Classes/Materials/MaterialExpressionReflectionVectorWS.h
@@ -8,7 +8,9 @@
 UCLASS(collapsecategories, hidecategories=Object, MinimalAPI)
 class UMaterialExpressionReflectionVectorWS : public UMaterialExpression
 {
-	GENERATED_UCLASS_BODY()
+	GENERATED_BODY()
+public:
+	ENGINE_API UMaterialExpressionReflectionVectorWS(const FObjectInitializer& ObjectInitializer = FObjectInitializer::Get());
 
 	/** Optional world normal to reflect the camera view vector about. If unconnected, pixel normal is used */
 	UPROPERTY(meta = (RequiredInput = "false", ToolTip = "Defaults to Reflection Vector if not specified"))
diff --git a/Engine/Source/Runtime/Engine/Classes/Materials/MaterialExpressionRotateAboutAxis.h b/Engine/Source/Runtime/Engine/Classes/Materials/MaterialExpressionRotateAboutAxis.h
--- a/Engine/Source/Runtime/Engine/Classes/Materials/MaterialExpressionRotateAboutAxis.h
+++ b/Engine/Source/Runtime/Engine/Classes/Materials/MaterialExpressionRotateAboutAxis.h
@@ -8,7 +8,9 @@
 UCLASS(MinimalAPI)
 class UMaterialExpressionRotateAboutAxis : public UMaterialExpression
 {
-	GENERATED_UCLASS_BODY()
+	GENERATED_BODY()
+public:
+	ENGINE_API UMaterialExpressionRotateAboutAxis(const FObjectInitializer& ObjectInitializer = FObjectInitializer::Get());
 
 	UPROPERTY()
 	FExpressionInput NormalizedRotationAxis;
diff --git a/Engine/Source/Runtime/Engine/Classes/Materials/MaterialExpressionRotator.h b/Engine/Source/Runtime/Engine/Classes/Materials/MaterialExpressionRotator.h
--- a/Engine/Source/Runtime/Engine/Classes/Materials/MaterialExpressionRotator.h
+++ b/Engine/Source/Runtime/Engine/Classes/Materials/MaterialExpressionRotator.h
@@ -8,7 +8,9 @@
 UCLASS(collapsecategories, hidecategories=Object)
 class UMaterialExpressionRotator : public UMaterialExpression
 {
-	GENERATED_UCLASS_BODY()
+	GENERATED_BODY()
+public:
+	UMaterialExpressionRotator(const FObjectInitializer& ObjectInitializer = FObjectInitializer::Get());
 
 	UPROPERTY(meta = (RequiredInput = "false", ToolTip = "Defaults to 'ConstCoordinate' if not specified"))
 	FExpressionInput Coordinate;
diff --git a/Engine/Source/Runtime/Engine/Classes/Materials/MaterialExpressionScalarParameter.h b/Engine/Source/Runtime/Engine/Classes/Materials/MaterialExpressionScalarParameter.h
--- a/Engine/Source/Runtime/Engine/Classes/Materials/MaterialExpressionScalarParameter.h
+++ b/Engine/Source/Runtime/Engine/Classes/Materials/MaterialExpressionScalarParameter.h
@@ -8,7 +8,9 @@
 UCLASS(collapsecategories, hidecategories=Object, MinimalAPI)
 class UMaterialExpressionScalarParameter : public UMaterialExpressionParameter
 {
-	GENERATED_UCLASS_BODY()
+	GENERATED_BODY()
+public:
+	ENGINE_API UMaterialExpressionScalarParameter(const FObjectInitializer& ObjectInitializer = FObjectInitializer::Get());
 
 	UPROPERTY(EditAnywhere, Category=MaterialExpressionScalarParameter)
 	float DefaultValue;
diff --git a/Engine/Source/Runtime/Engine/Classes/Materials/MaterialExpressionSceneColor.h b/Engine/Source/Runtime/Engine/Classes/Materials/MaterialExpressionSceneColor.h
--- a/Engine/Source/Runtime/Engine/Classes/Materials/MaterialExpressionSceneColor.h
+++ b/Engine/Source/Runtime/Engine/Classes/Materials/MaterialExpressionSceneColor.h
@@ -18,7 +18,9 @@ namespace EMaterialSceneAttributeInputMode
 UCLASS(collapsecategories, hidecategories=Object)
 class UMaterialExpressionSceneColor : public UMaterialExpression
 {
-	GENERATED_UCLASS_BODY()
+	GENERATED_BODY()
+public:
+	UMaterialExpressionSceneColor(const FObjectInitializer& ObjectInitializer = FObjectInitializer::Get());
 	
 	/**
 	* Coordinates - UV coordinates to apply to the scene color lookup.
diff --git a/Engine/Source/Runtime/Engine/Classes/Materials/MaterialExpressionSceneDepth.h b/Engine/Source/Runtime/Engine/Classes/Materials/MaterialExpressionSceneDepth.h
--- a/Engine/Source/Runtime/Engine/Classes/Materials/MaterialExpressionSceneDepth.h
+++ b/Engine/Source/Runtime/Engine/Classes/Materials/MaterialExpressionSceneDepth.h
@@ -8,7 +8,9 @@
 UCLASS(collapsecategories, hidecategories=Object)
 class UMaterialExpressionSceneDepth : public UMaterialExpression
 {
-	GENERATED_UCLASS_BODY()
+	GENERATED_BODY()
+public:
+	UMaterialExpressionSceneDepth(const FObjectInitializer& ObjectInitializer = FObjectInitializer::Get());
 
 	/** 
 	* Coordinates - UV coordinates to apply to the scene depth lookup.
diff --git a/Engine/Source/Runtime/Engine/Classes/Materials/MaterialExpressionSceneTexelSize.h b/Engine/Source/Runtime/Engine/Classes/Materials/MaterialExpressionSceneTexelSize.h
--- a/Engine/Source/Runtime/Engine/Classes/Materials/MaterialExpressionSceneTexelSize.h
+++ b/Engine/Source/Runtime/Engine/Classes/Materials/MaterialExpressionSceneTexelSize.h
@@ -8,7 +8,9 @@
 UCLASS(collapsecategories, hidecategories=Object)
 class UMaterialExpressionSceneTexelSize : public UMaterialExpression
 {
-	GENERATED_UCLASS_BODY()
+	GENERATED_BODY()
+public:
+	UMaterialExpressionSceneTexelSize(const FObjectInitializer& ObjectInitializer = FObjectInitializer::Get());
 
 	// Begin UMaterialExpression Interface
 	virtual int32 Compile(class FMaterialCompiler* Compiler, int32 OutputIndex, int32 MultiplexIndex) override;
diff --git a/Engine/Source/Runtime/Engine/Classes/Materials/MaterialExpressionSceneTexture.h b/Engine/Source/Runtime/Engine/Classes/Materials/MaterialExpressionSceneTexture.h
--- a/Engine/Source/Runtime/Engine/Classes/Materials/MaterialExpressionSceneTexture.h
+++ b/Engine/Source/Runtime/Engine/Classes/Materials/MaterialExpressionSceneTexture.h
@@ -62,7 +62,9 @@ enum ESceneTextureId
 UCLASS(collapsecategories, hidecategories=Object)
 class UMaterialExpressionSceneTexture : public UMaterialExpression
 {
-	GENERATED_UCLASS_BODY()
+	GENERATED_BODY()
+public:
+	UMaterialExpressionSceneTexture(const FObjectInitializer& ObjectInitializer = FObjectInitializer::Get());
 
 	/** UV in 0..1 range */
 	UPROPERTY(meta = (RequiredInput = "false", ToolTip = "Ignored if not specified"))
diff --git a/Engine/Source/Runtime/Engine/Classes/Materials/MaterialExpressionScreenPosition.h b/Engine/Source/Runtime/Engine/Classes/Materials/MaterialExpressionScreenPosition.h
--- a/Engine/Source/Runtime/Engine/Classes/Materials/MaterialExpressionScreenPosition.h
+++ b/Engine/Source/Runtime/Engine/Classes/Materials/MaterialExpressionScreenPosition.h
@@ -8,7 +8,9 @@
 UCLASS(collapsecategories, hidecategories=Object)
 class UMaterialExpressionScreenPosition : public UMaterialExpression
 {
-	GENERATED_UCLASS_BODY()
+	GENERATED_BODY()
+public:
+	UMaterialExpressionScreenPosition(const FObjectInitializer& ObjectInitializer = FObjectInitializer::Get());
 
 	// Begin UMaterialExpression Interface
 	virtual int32 Compile(class FMaterialCompiler* Compiler, int32 OutputIndex, int32 MultiplexIndex) override;
diff --git a/Engine/Source/Runtime/Engine/Classes/Materials/MaterialExpressionSine.h b/Engine/Source/Runtime/Engine/Classes/Materials/MaterialExpressionSine.h
--- a/Engine/Source/Runtime/Engine/Classes/Materials/MaterialExpressionSine.h
+++ b/Engine/Source/Runtime/Engine/Classes/Materials/MaterialExpressionSine.h
@@ -8,7 +8,9 @@
 UCLASS(collapsecategories, hidecategories=Object)
 class UMaterialExpressionSine : public UMaterialExpression
 {
-	GENERATED_UCLASS_BODY()
+	GENERATED_BODY()
+public:
+	UMaterialExpressionSine(const FObjectInitializer& ObjectInitializer = FObjectInitializer::Get());
 
 	UPROPERTY()
 	FExpressionInput Input;
diff --git a/Engine/Source/Runtime/Engine/Classes/Materials/MaterialExpressionSpeedTree.h b/Engine/Source/Runtime/Engine/Classes/Materials/MaterialExpressionSpeedTree.h
--- a/Engine/Source/Runtime/Engine/Classes/Materials/MaterialExpressionSpeedTree.h
+++ b/Engine/Source/Runtime/Engine/Classes/Materials/MaterialExpressionSpeedTree.h
@@ -37,7 +37,9 @@ enum ESpeedTreeLODType
 UCLASS(collapsecategories, hidecategories=Object, MinimalAPI)
 class UMaterialExpressionSpeedTree : public UMaterialExpression
 {
-	GENERATED_UCLASS_BODY()	
+	GENERATED_BODY()
+public:
+	ENGINE_API UMaterialExpressionSpeedTree(const FObjectInitializer& ObjectInitializer = FObjectInitializer::Get());
 
 	UPROPERTY(EditAnywhere, Category=MaterialExpressionSpeedTree, meta=(DisplayName = "Geometry Type", ToolTip="The type of SpeedTree geometry on which this material will be used"))
 	TEnumAsByte<enum ESpeedTreeGeometryType> GeometryType;
diff --git a/Engine/Source/Runtime/Engine/Classes/Materials/MaterialExpressionSphereMask.h b/Engine/Source/Runtime/Engine/Classes/Materials/MaterialExpressionSphereMask.h
--- a/Engine/Source/Runtime/Engine/Classes/Materials/MaterialExpressionSphereMask.h
+++ b/Engine/Source/Runtime/Engine/Classes/Materials/MaterialExpressionSphereMask.h
@@ -8,7 +8,9 @@
 UCLASS()
 class UMaterialExpressionSphereMask : public UMaterialExpression
 {
-	GENERATED_UCLASS_BODY()
+	GENERATED_BODY()
+public:
+	UMaterialExpressionSphereMask(const FObjectInitializer& ObjectInitializer = FObjectInitializer::Get());
 
 	/** 1 to 4 dimensional vector, should be the same type as B */
 	UPROPERTY()
diff --git a/Engine/Source/Runtime/Engine/Classes/Materials/MaterialExpressionSphericalParticleOpacity.h b/Engine/Source/Runtime/Engine/Classes/Materials/MaterialExpressionSphericalParticleOpacity.h
--- a/Engine/Source/Runtime/Engine/Classes/Materials/MaterialExpressionSphericalParticleOpacity.h
+++ b/Engine/Source/Runtime/Engine/Classes/Materials/MaterialExpressionSphericalParticleOpacity.h
@@ -8,7 +8,9 @@
 UCLASS(collapsecategories, hidecategories=Object)
 class UMaterialExpressionSphericalParticleOpacity : public UMaterialExpression
 {
-	GENERATED_UCLASS_BODY()
+	GENERATED_BODY()
+public:
+	UMaterialExpressionSphericalParticleOpacity(const FObjectInitializer& ObjectInitializer = FObjectInitializer::Get());
 
 	/** Density of the particle sphere. */
 	UPROPERTY(meta = (RequiredInput = "false", ToolTip = "Defaults to 'ConstantDensity' if not specified"))
diff --git a/Engine/Source/Runtime/Engine/Classes/Materials/MaterialExpressionSquareRoot.h b/Engine/Source/Runtime/Engine/Classes/Materials/MaterialExpressionSquareRoot.h
--- a/Engine/Source/Runtime/Engine/Classes/Materials/MaterialExpressionSquareRoot.h
+++ b/Engine/Source/Runtime/Engine/Classes/Materials/MaterialExpressionSquareRoot.h
@@ -8,7 +8,9 @@
 UCLASS(collapsecategories, hidecategories=Object)
 class UMaterialExpressionSquareRoot : public UMaterialExpression
 {
-	GENERATED_UCLASS_BODY()
+	GENERATED_BODY()
+public:
+	UMaterialExpressionSquareRoot(const FObjectInitializer& ObjectInitializer = FObjectInitializer::Get());
 
 	UPROPERTY()
 	FExpressionInput Input;
diff --git a/Engine/Source/Runtime/Engine/Classes/Materials/MaterialExpressionStaticBool.h b/Engine/Source/Runtime/Engine/Classes/Materials/MaterialExpressionStaticBool.h
--- a/Engine/Source/Runtime/Engine/Classes/Materials/MaterialExpressionStaticBool.h
+++ b/Engine/Source/Runtime/Engine/Classes/Materials/MaterialExpressionStaticBool.h
@@ -8,7 +8,9 @@
 UCLASS(collapsecategories, hidecategories=Object, MinimalAPI)
 class UMaterialExpressionStaticBool : public UMaterialExpression
 {
-	GENERATED_UCLASS_BODY()
+	GENERATED_BODY()
+public:
+	ENGINE_API UMaterialExpressionStaticBool(const FObjectInitializer& ObjectInitializer = FObjectInitializer::Get());
 
 	UPROPERTY(EditAnywhere, Category=MaterialExpressionStaticBool)
 	uint32 Value:1;
diff --git a/Engine/Source/Runtime/Engine/Classes/Materials/MaterialExpressionStaticBoolParameter.h b/Engine/Source/Runtime/Engine/Classes/Materials/MaterialExpressionStaticBoolParameter.h
--- a/Engine/Source/Runtime/Engine/Classes/Materials/MaterialExpressionStaticBoolParameter.h
+++ b/Engine/Source/Runtime/Engine/Classes/Materials/MaterialExpressionStaticBoolParameter.h
@@ -8,7 +8,9 @@
 UCLASS(collapsecategories, hidecategories=Object, MinimalAPI)
 class UMaterialExpressionStaticBoolParameter : public UMaterialExpressionParameter
 {
-	GENERATED_UCLASS_BODY()
+	GENERATED_BODY()
+public:
+	ENGINE_API UMaterialExpressionStaticBoolParameter(const FObjectInitializer& ObjectInitializer = FObjectInitializer::Get());
 
 	UPROPERTY(EditAnywhere, Category=MaterialExpressionStaticBoolParameter)
 	uint32 DefaultValue:1;
diff --git a/Engine/Source/Runtime/Engine/Classes/Materials/MaterialExpressionStaticComponentMaskParameter.h b/Engine/Source/Runtime/Engine/Classes/Materials/MaterialExpressionStaticComponentMaskParameter.h
--- a/Engine/Source/Runtime/Engine/Classes/Materials/MaterialExpressionStaticComponentMaskParameter.h
+++ b/Engine/Source/Runtime/Engine/Classes/Materials/MaterialExpressionStaticComponentMaskParameter.h
@@ -8,7 +8,9 @@
 UCLASS(collapsecategories, hidecategories=Object, MinimalAPI)
 class UMaterialExpressionStaticComponentMaskParameter : public UMaterialExpressionParameter
 {
-	GENERATED_UCLASS_BODY()
+	GENERATED_BODY()
+public:
+	ENGINE_API UMaterialExpressionStaticComponentMaskParameter(const FObjectInitializer& ObjectInitializer = FObjectInitializer::Get());
 
 	UPROPERTY()
 	FExpressionInput Input;
diff --git a/Engine/Source/Runtime/Engine/Classes/Materials/MaterialExpressionStaticSwitch.h b/Engine/Source/Runtime/Engine/Classes/Materials/MaterialExpressionStaticSwitch.h
--- a/Engine/Source/Runtime/Engine/Classes/Materials/MaterialExpressionStaticSwitch.h
+++ b/Engine/Source/Runtime/Engine/Classes/Materials/MaterialExpressionStaticSwitch.h
@@ -8,7 +8,9 @@
 UCLASS(collapsecategories, hidecategories=Object, MinimalAPI)
 class UMaterialExpressionStaticSwitch : public UMaterialExpression
 {
-	GENERATED_UCLASS_BODY()
+	GENERATED_BODY()
+public:
+	ENGINE_API UMaterialExpressionStaticSwitch(const FObjectInitializer& ObjectInitializer = FObjectInitializer::Get());
 
 	UPROPERTY(EditAnywhere, Category=MaterialExpressionStaticSwitch, meta=(OverridingInputProperty = "Value"))
 	uint32 DefaultValue:1;
diff --git a/Engine/Source/Runtime/Engine/Classes/Materials/MaterialExpressionStaticSwitchParameter.h b/Engine/Source/Runtime/Engine/Classes/Materials/MaterialExpressionStaticSwitchParameter.h
--- a/Engine/Source/Runtime/Engine/Classes/Materials/MaterialExpressionStaticSwitchParameter.h
+++ b/Engine/Source/Runtime/Engine/Classes/Materials/MaterialExpressionStaticSwitchParameter.h
@@ -8,7 +8,9 @@
 UCLASS(collapsecategories, hidecategories=Object, MinimalAPI)
 class UMaterialExpressionStaticSwitchParameter : public UMaterialExpressionStaticBoolParameter
 {
-	GENERATED_UCLASS_BODY()
+	GENERATED_BODY()
+public:
+	ENGINE_API UMaterialExpressionStaticSwitchParameter(const FObjectInitializer& ObjectInitializer = FObjectInitializer::Get());
 
 	UPROPERTY()
 	FExpressionInput A;
diff --git a/Engine/Source/Runtime/Engine/Classes/Materials/MaterialExpressionSubtract.h b/Engine/Source/Runtime/Engine/Classes/Materials/MaterialExpressionSubtract.h
--- a/Engine/Source/Runtime/Engine/Classes/Materials/MaterialExpressionSubtract.h
+++ b/Engine/Source/Runtime/Engine/Classes/Materials/MaterialExpressionSubtract.h
@@ -8,7 +8,9 @@
 UCLASS()
 class UMaterialExpressionSubtract : public UMaterialExpression
 {
-	GENERATED_UCLASS_BODY()
+	GENERATED_BODY()
+public:
+	UMaterialExpressionSubtract(const FObjectInitializer& ObjectInitializer = FObjectInitializer::Get());
 
 	UPROPERTY(meta = (RequiredInput = "false", ToolTip = "Defaults to 'ConstA' if not specified"))
 	FExpressionInput A;
diff --git a/Engine/Source/Runtime/Engine/Classes/Materials/MaterialExpressionTextureBase.h b/Engine/Source/Runtime/Engine/Classes/Materials/MaterialExpressionTextureBase.h
--- a/Engine/Source/Runtime/Engine/Classes/Materials/MaterialExpressionTextureBase.h
+++ b/Engine/Source/Runtime/Engine/Classes/Materials/MaterialExpressionTextureBase.h
@@ -12,7 +12,9 @@
 UCLASS(abstract, hidecategories=Object)
 class ENGINE_API UMaterialExpressionTextureBase : public UMaterialExpression 
 {
-	GENERATED_UCLASS_BODY()
+	GENERATED_BODY()
+public:
+	UMaterialExpressionTextureBase(const FObjectInitializer& ObjectInitializer = FObjectInitializer::Get());
 
 	UPROPERTY(EditAnywhere, Category=MaterialExpressionTextureBase)
 	class UTexture* Texture;
diff --git a/Engine/Source/Runtime/Engine/Classes/Materials/MaterialExpressionTextureCoordinate.h b/Engine/Source/Runtime/Engine/Classes/Materials/MaterialExpressionTextureCoordinate.h
--- a/Engine/Source/Runtime/Engine/Classes/Materials/MaterialExpressionTextureCoordinate.h
+++ b/Engine/Source/Runtime/Engine/Classes/Materials/MaterialExpressionTextureCoordinate.h
@@ -8,7 +8,9 @@
 UCLASS(collapsecategories, hidecategories=Object, MinimalAPI)
 class UMaterialExpressionTextureCoordinate : public UMaterialExpression
 {
-	GENERATED_UCLASS_BODY()
+	GENERATED_BODY()
+public:
+	ENGINE_API UMaterialExpressionTextureCoordinate(const FObjectInitializer& ObjectInitializer = FObjectInitializer::Get());
 
 	/** Texture coordinate index */
 	UPROPERTY(EditAnywhere, Category=MaterialExpressionTextureCoordinate)
diff --git a/Engine/Source/Runtime/Engine/Classes/Materials/MaterialExpressionTextureObject.h b/Engine/Source/Runtime/Engine/Classes/Materials/MaterialExpressionTextureObject.h
--- a/Engine/Source/Runtime/Engine/Classes/Materials/MaterialExpressionTextureObject.h
+++ b/Engine/Source/Runtime/Engine/Classes/Materials/MaterialExpressionTextureObject.h
@@ -12,7 +12,9 @@
 UCLASS(collapsecategories, hidecategories=Object, MinimalAPI)
 class UMaterialExpressionTextureObject : public UMaterialExpressionTextureBase
 {
-	GENERATED_UCLASS_BODY()
+	GENERATED_BODY()
+public:
+	ENGINE_API UMaterialExpressionTextureObject(const FObjectInitializer& ObjectInitializer = FObjectInitializer::Get());
 
 #if WITH_EDITOR
 	virtual void PostEditChangeProperty(FPropertyChangedEvent& PropertyChangedEvent) override;
diff --git a/Engine/Source/Runtime/Engine/Classes/Materials/MaterialExpressionTextureObjectParameter.h b/Engine/Source/Runtime/Engine/Classes/Materials/MaterialExpressionTextureObjectParameter.h
--- a/Engine/Source/Runtime/Engine/Classes/Materials/MaterialExpressionTextureObjectParameter.h
+++ b/Engine/Source/Runtime/Engine/Classes/Materials/MaterialExpressionTextureObjectParameter.h
@@ -12,7 +12,9 @@
 UCLASS(collapsecategories, hidecategories=Object)
 class UMaterialExpressionTextureObjectParameter : public UMaterialExpressionTextureSampleParameter
 {
-	GENERATED_UCLASS_BODY()
+	GENERATED_BODY()
+public:
+	UMaterialExpressionTextureObjectParameter(const FObjectInitializer& ObjectInitializer = FObjectInitializer::Get());
 
 
 	// Begin UMaterialExpression Interface
